[PCG] Address the ensure in PCG Managed Actors and PCG Managed Component

#rnx

diff --git a/Engine/Plugins/PCG/Source/PCG/Private/Elements/PCGSpawnActor.cpp b/Engine/Plugins/PCG/Source/PCG/Private/Elements/PCGSpawnActor.cpp
--- a/Engine/Plugins/PCG/Source/PCG/Private/Elements/PCGSpawnActor.cpp
+++ b/Engine/Plugins/PCG/Source/PCG/Private/Elements/PCGSpawnActor.cpp
@@ -471,6 +471,12 @@ bool FPCGSpawnActorElement::SpawnAndPrepareSubgraphs(FPCGSubgraphContext* Contex
 
 				for (UPCGManagedISMComponent* MISMC : MISMCs)
 				{
+					if (!MISMC->IsMarkedUnused())
+					{
+						// TODO: Add Context back in with toggles. Revisit if the stack is added to the managed components at creation
+						PCGLog::LogWarningOnGraph(LOCTEXT("IdenticalISMCSpawn", "Identical ISM Component spawn occurred. It may be beneficial to re-check graph logic for identical spawn conditions (same actor at same location, etc) or repeated nodes."), nullptr);
+					}
+
 					MISMC->MarkAsReused();
 				}
 
@@ -814,7 +820,7 @@ void FPCGSpawnActorElement::SpawnActors(FPCGSubgraphContext* Context, AActor* Ta
 #endif
 	UPCGSubsystem* Subsystem = (Context->SourceComponent.Get() ? Context->SourceComponent->GetSubsystem() : nullptr);
 
-	// Try to reuse actors if the are preexisting
+	// Try to reuse actors if they are preexisting
 	TArray<UPCGManagedActors*> ReusedManagedActorsResources;
 	FPCGCrc InputDependenciesCrc;
 	if (CVarAllowActorReuse.GetValueOnAnyThread())
@@ -850,6 +856,13 @@ void FPCGSpawnActorElement::SpawnActors(FPCGSubgraphContext* Context, AActor* Ta
 		for (UPCGManagedActors* ManagedActors : ReusedManagedActorsResources)
 		{
 			check(ManagedActors);
+
+			if (!ManagedActors->IsMarkedUnused())
+			{
+				// TODO: Add Context back in with toggles. Revisit if the stack is added to the managed actors at creation
+				PCGLog::LogWarningOnGraph(LOCTEXT("IdenticalActorSpawn", "Identical actor spawn occurred. It may be beneficial to re-check graph logic for identical spawn conditions (same actor at same location, etc) or repeated nodes."), nullptr);
+			}
+
 			ManagedActors->MarkAsReused();
 
 			// There's no setup to be done, just generation if we're in the no-merge case, so keep track of these actors only in this case
diff --git a/Engine/Plugins/PCG/Source/PCG/Private/Elements/PCGStaticMeshSpawner.cpp b/Engine/Plugins/PCG/Source/PCG/Private/Elements/PCGStaticMeshSpawner.cpp
--- a/Engine/Plugins/PCG/Source/PCG/Private/Elements/PCGStaticMeshSpawner.cpp
+++ b/Engine/Plugins/PCG/Source/PCG/Private/Elements/PCGStaticMeshSpawner.cpp
@@ -126,6 +126,12 @@ bool FPCGStaticMeshSpawnerElement::PrepareDataInternal(FPCGContext* InContext) c
 
 			for (UPCGManagedISMComponent* MISMC : MISMCs)
 			{
+				if (!MISMC->IsMarkedUnused())
+				{
+					// TODO: Add Context back in with toggles. Revisit if the stack is added to the managed components at creation
+					PCGLog::LogWarningOnGraph(LOCTEXT("IdenticalISMCSpawn", "Identical ISM Component spawn occurred. It may be beneficial to re-check graph logic for identical spawn conditions (same mesh descriptor at same location, etc) or repeated nodes."), nullptr);
+				}
+
 				MISMC->MarkAsReused();
 			}
 
diff --git a/Engine/Plugins/PCG/Source/PCG/Public/PCGManagedResource.h b/Engine/Plugins/PCG/Source/PCG/Public/PCGManagedResource.h
--- a/Engine/Plugins/PCG/Source/PCG/Public/PCGManagedResource.h
+++ b/Engine/Plugins/PCG/Source/PCG/Public/PCGManagedResource.h
@@ -41,8 +41,7 @@ public:
 	/** Marks the resources as being kept and changed through generation */
 	virtual void MarkAsUsed() { ensure(CanBeUsed()); bIsMarkedUnused = false; }
 	/** Marks the resource as being reused as-is during the generation */
-	// Ensure may fire if multiple executions of the graph are happening in parallel/overlapping
-	virtual void MarkAsReused() { ensure(bIsMarkedUnused); bIsMarkedUnused = false; }
+	virtual void MarkAsReused() { bIsMarkedUnused = false; }
 	bool IsMarkedUnused() const { return bIsMarkedUnused; }
 
 	/** Move the given resource to a new actor. Return true if it has succeeded */
