[PCG] Fix Editor Context getting in the way of PCG Actor Spawning (Push/Pop a new empty context when spawning actors in PCG)

#rnx

diff --git a/Engine/Plugins/PCG/Source/PCG/Private/Helpers/PCGActorHelpers.cpp b/Engine/Plugins/PCG/Source/PCG/Private/Helpers/PCGActorHelpers.cpp
--- a/Engine/Plugins/PCG/Source/PCG/Private/Helpers/PCGActorHelpers.cpp
+++ b/Engine/Plugins/PCG/Source/PCG/Private/Helpers/PCGActorHelpers.cpp
@@ -19,6 +19,7 @@
 #if WITH_EDITOR
 #include "Engine/Level.h"
 #include "ScopedTransaction.h"
+#include "Subsystems/ActorEditorContextSubsystem.h"
 #include "WorldPartition/DataLayer/DataLayerInstance.h"
 #include "WorldPartition/DataLayer/ExternalDataLayerAsset.h"
 #include "WorldPartition/DataLayer/ExternalDataLayerInstance.h"
@@ -482,6 +483,11 @@ AActor* UPCGActorHelpers::SpawnDefaultActor(const FSpawnDefaultActorParams& Para
 		}
 	}
 
+	// Avoid relying on the Editor Context at all
+	UActorEditorContextSubsystem::Get()->PushContext();
+	ON_SCOPE_EXIT{ UActorEditorContextSubsystem::Get()->PopContext(); };
+
+	// Specify EDL we want to use if any for spawning this actor
 	FScopedOverrideSpawningLevelMountPointObject EDLScope(ExternalDataLayerInstance ? ExternalDataLayerInstance->GetExternalDataLayerAsset() : nullptr);
 #endif
 
diff --git a/Engine/Plugins/PCG/Source/PCG/Private/PCGSubsystem.cpp b/Engine/Plugins/PCG/Source/PCG/Private/PCGSubsystem.cpp
--- a/Engine/Plugins/PCG/Source/PCG/Private/PCGSubsystem.cpp
+++ b/Engine/Plugins/PCG/Source/PCG/Private/PCGSubsystem.cpp
@@ -26,6 +26,7 @@
 #include "PackageSourceControlHelper.h"
 #include "ObjectTools.h"
 #include "Misc/ScopedSlowTask.h"
+#include "Subsystems/ActorEditorContextSubsystem.h"
 #include "WorldPartition/WorldPartition.h"
 #include "WorldPartition/WorldPartitionHelpers.h"
 #include "WorldPartition/DataLayer/DataLayerManager.h"
@@ -1192,6 +1193,11 @@ APCGPartitionActor* UPCGSubsystem::FindOrCreatePCGPartitionActor(const FPCGGridD
 			
 	GridDescriptor.GetDataLayerAssets(DataLayerAssets, ExternalDataLayerAsset);
 
+	// Avoid relying on the Editor Context at all
+	UActorEditorContextSubsystem::Get()->PushContext();
+	ON_SCOPE_EXIT{ UActorEditorContextSubsystem::Get()->PopContext(); };
+
+	// Specify EDL we want to use if any for spawning this actor
 	FScopedOverrideSpawningLevelMountPointObject EDLScope(ExternalDataLayerAsset);
 #endif
 
