Fixed issue in the CopyPoints node when metadata inheritance is used and one of the two data has no attribute

#rnx

diff --git a/Engine/Plugins/PCG/Source/PCG/Private/Elements/PCGCopyPoints.cpp b/Engine/Plugins/PCG/Source/PCG/Private/Elements/PCGCopyPoints.cpp
--- a/Engine/Plugins/PCG/Source/PCG/Private/Elements/PCGCopyPoints.cpp
+++ b/Engine/Plugins/PCG/Source/PCG/Private/Elements/PCGCopyPoints.cpp
@@ -121,37 +121,34 @@ bool FPCGCopyPointsElement::ExecuteInternal(FPCGContext* Context) const
 			{
 				bInheritMetadataFromSource = true;
 				bProcessMetadata = bSourceHasMetadata;
-
-				OutPointData->InitializeFromData(SourcePointData);
 				RootMetadata = SourcePointMetadata;
 				NonRootMetadata = nullptr;
+
+				OutPointData->InitializeFromData(SourcePointData);
 			}
 			else if (AttributeInheritance == EPCGCopyPointsMetadataInheritanceMode::TargetOnly)
 			{
 				bInheritMetadataFromSource = false;
 				bProcessMetadata = bTargetHasMetadata;
-
-				OutPointData->InitializeFromData(TargetPointData);
-
 				RootMetadata = TargetPointMetadata;
 				NonRootMetadata = nullptr;
-			}
-			else if (AttributeInheritance == EPCGCopyPointsMetadataInheritanceMode::SourceFirst)
-			{
-				bInheritMetadataFromSource = bSourceHasMetadata || !bTargetHasMetadata;
 
-				OutPointData->InitializeFromData(SourcePointData);
-				RootMetadata = SourcePointMetadata;
-				NonRootMetadata = TargetPointMetadata;
+				OutPointData->InitializeFromData(TargetPointData);
 			}
-			else if (AttributeInheritance == EPCGCopyPointsMetadataInheritanceMode::TargetFirst)
+			else if (AttributeInheritance == EPCGCopyPointsMetadataInheritanceMode::SourceFirst || AttributeInheritance == EPCGCopyPointsMetadataInheritanceMode::TargetFirst)
 			{
-				bInheritMetadataFromSource = !bTargetHasMetadata;
-
-				OutPointData->InitializeFromData(TargetPointData);
+				if (AttributeInheritance == EPCGCopyPointsMetadataInheritanceMode::SourceFirst)
+				{
+					bInheritMetadataFromSource = bSourceHasMetadata || !bTargetHasMetadata;
+				}
+				else // TargetFirst
+				{
+					bInheritMetadataFromSource = !bTargetHasMetadata && bSourceHasMetadata;
+				}
 
-				RootMetadata = TargetPointMetadata;
-				NonRootMetadata = SourcePointMetadata;
+				RootMetadata = bInheritMetadataFromSource ? SourcePointMetadata : TargetPointMetadata;
+				NonRootMetadata = bInheritMetadataFromSource ? TargetPointMetadata : SourcePointMetadata;
+				OutPointData->InitializeFromData(bInheritMetadataFromSource ? SourcePointData : TargetPointData);
 			}
 			else // None
 			{
