[PCG] Follow up to fix the ensure in Managed Components/Actors

#rnx

diff --git a/Engine/Plugins/PCG/Source/PCG/Private/Elements/PCGSpawnActor.cpp b/Engine/Plugins/PCG/Source/PCG/Private/Elements/PCGSpawnActor.cpp
--- a/Engine/Plugins/PCG/Source/PCG/Private/Elements/PCGSpawnActor.cpp
+++ b/Engine/Plugins/PCG/Source/PCG/Private/Elements/PCGSpawnActor.cpp
@@ -471,10 +471,10 @@ bool FPCGSpawnActorElement::SpawnAndPrepareSubgraphs(FPCGSubgraphContext* Contex
 
 				for (UPCGManagedISMComponent* MISMC : MISMCs)
 				{
-					if (!MISMC->IsMarkedUnused())
+					if (!MISMC->IsMarkedUnused() && Settings->bWarnOnIdenticalSpawn)
 					{
-						// TODO: Add Context back in with toggles. Revisit if the stack is added to the managed components at creation
-						PCGLog::LogWarningOnGraph(LOCTEXT("IdenticalISMCSpawn", "Identical ISM Component spawn occurred. It may be beneficial to re-check graph logic for identical spawn conditions (same actor at same location, etc) or repeated nodes."), nullptr);
+						// TODO: Revisit if the stack is added to the managed components at creation
+						PCGLog::LogWarningOnGraph(LOCTEXT("IdenticalISMCSpawn", "Identical ISM Component spawn occurred. It may be beneficial to re-check graph logic for identical spawn conditions (same actor at same location, etc) or repeated nodes."), Context);
 					}
 
 					MISMC->MarkAsReused();
@@ -856,11 +856,10 @@ void FPCGSpawnActorElement::SpawnActors(FPCGSubgraphContext* Context, AActor* Ta
 		for (UPCGManagedActors* ManagedActors : ReusedManagedActorsResources)
 		{
 			check(ManagedActors);
-
-			if (!ManagedActors->IsMarkedUnused())
+			if (!ManagedActors->IsMarkedUnused() && Settings->bWarnOnIdenticalSpawn)
 			{
-				// TODO: Add Context back in with toggles. Revisit if the stack is added to the managed actors at creation
-				PCGLog::LogWarningOnGraph(LOCTEXT("IdenticalActorSpawn", "Identical actor spawn occurred. It may be beneficial to re-check graph logic for identical spawn conditions (same actor at same location, etc) or repeated nodes."), nullptr);
+				// TODO: Revisit if the stack is added to the managed actors at creation
+				PCGLog::LogWarningOnGraph(LOCTEXT("IdenticalActorSpawn", "Identical actor spawn occurred. It may be beneficial to re-check graph logic for identical spawn conditions (same actor at same location, etc) or repeated nodes."), Context);
 			}
 
 			ManagedActors->MarkAsReused();
diff --git a/Engine/Plugins/PCG/Source/PCG/Private/Elements/PCGStaticMeshSpawner.cpp b/Engine/Plugins/PCG/Source/PCG/Private/Elements/PCGStaticMeshSpawner.cpp
--- a/Engine/Plugins/PCG/Source/PCG/Private/Elements/PCGStaticMeshSpawner.cpp
+++ b/Engine/Plugins/PCG/Source/PCG/Private/Elements/PCGStaticMeshSpawner.cpp
@@ -126,10 +126,10 @@ bool FPCGStaticMeshSpawnerElement::PrepareDataInternal(FPCGContext* InContext) c
 
 			for (UPCGManagedISMComponent* MISMC : MISMCs)
 			{
-				if (!MISMC->IsMarkedUnused())
+				if (!MISMC->IsMarkedUnused() && Settings->bWarnOnIdenticalSpawn)
 				{
-					// TODO: Add Context back in with toggles. Revisit if the stack is added to the managed components at creation
-					PCGLog::LogWarningOnGraph(LOCTEXT("IdenticalISMCSpawn", "Identical ISM Component spawn occurred. It may be beneficial to re-check graph logic for identical spawn conditions (same mesh descriptor at same location, etc) or repeated nodes."), nullptr);
+					// TODO: Revisit if the stack is added to the managed components at creation
+					PCGLog::LogWarningOnGraph(LOCTEXT("IdenticalISMCSpawn", "Identical ISM Component spawn occurred. It may be beneficial to re-check graph logic for identical spawn conditions (same mesh descriptor at same location, etc) or repeated nodes."), Context);
 				}
 
 				MISMC->MarkAsReused();
diff --git a/Engine/Plugins/PCG/Source/PCG/Public/Elements/PCGSpawnActor.h b/Engine/Plugins/PCG/Source/PCG/Public/Elements/PCGSpawnActor.h
--- a/Engine/Plugins/PCG/Source/PCG/Public/Elements/PCGSpawnActor.h
+++ b/Engine/Plugins/PCG/Source/PCG/Public/Elements/PCGSpawnActor.h
@@ -84,6 +84,10 @@ public:
 	UPROPERTY(BlueprintReadWrite, EditAnywhere, Category = Settings, meta = (EditCondition = "bSpawnByAttribute"))
 	FName SpawnAttribute = NAME_None;
 
+	/** Adds a warning to the node on repeated spawning with identical conditions (ie. same actor at same spawn location, etc). */
+	UPROPERTY(BlueprintReadWrite, EditAnywhere, Category = "Settings|Advanced")
+	bool bWarnOnIdenticalSpawn = true;
+
 protected:
 	UPROPERTY(BlueprintReadOnly, EditAnywhere, Category = Settings, meta = (OnlyPlaceable, DisallowCreateNew))
 	TSubclassOf<AActor> TemplateActorClass = nullptr;
@@ -183,4 +187,4 @@ private:
 
 	void CollapseIntoTargetActor(FPCGSubgraphContext* Context, AActor* TargetActor, TSubclassOf<AActor> TemplateActorClass, const UPCGPointData* PointData) const;
 	void SpawnActors(FPCGSubgraphContext* Context, AActor* TargetActor, TSubclassOf<AActor> TemplateActorClass, AActor* TemplateActor, FPCGTaggedData& Output, const UPCGPointData* PointData, UPCGPointData* OutPointData) const;
-};
\ No newline at end of file
+};
diff --git a/Engine/Plugins/PCG/Source/PCG/Public/Elements/PCGStaticMeshSpawner.h b/Engine/Plugins/PCG/Source/PCG/Public/Elements/PCGStaticMeshSpawner.h
--- a/Engine/Plugins/PCG/Source/PCG/Public/Elements/PCGStaticMeshSpawner.h
+++ b/Engine/Plugins/PCG/Source/PCG/Public/Elements/PCGStaticMeshSpawner.h
@@ -102,6 +102,10 @@ public:
 	UPROPERTY(BlueprintReadWrite, EditAnywhere, Category = "Settings|Advanced")
 	bool bSilenceOverrideAttributeNotFoundErrors = false;
 
+	/** Adds a warning to the node on repeated spawning with identical conditions (ie. same mesh descriptor at same spawn location, etc). */
+	UPROPERTY(BlueprintReadWrite, EditAnywhere, Category = "Settings|Advanced")
+	bool bWarnOnIdenticalSpawn = true;
+
 protected:
 	void RefreshMeshSelector();
 	void RefreshInstancePacker();
