[PCG] Clarified the setup stage of ready tasks prior to being promoted to active tasks. Also removed some deadcode

#rnx

diff --git a/Engine/Plugins/PCG/Source/PCG/Private/Graph/PCGGraphExecutor.cpp b/Engine/Plugins/PCG/Source/PCG/Private/Graph/PCGGraphExecutor.cpp
--- a/Engine/Plugins/PCG/Source/PCG/Private/Graph/PCGGraphExecutor.cpp
+++ b/Engine/Plugins/PCG/Source/PCG/Private/Graph/PCGGraphExecutor.cpp
@@ -774,102 +774,115 @@ void FPCGGraphExecutor::Execute()
 
 				FPCGGraphTask& Task = ReadyTasks[ReadyTaskIndex];
 
-				// Build input
-				FPCGDataCollection TaskInput;
-				BuildTaskInput(Task, TaskInput);
-
-				// Initialize the element if needed (required to know whether it will run on the main thread or not)
-				if (!Task.Element)
+				if (!Task.bHasDoneSetup)
 				{
-					// Get appropriate settings
-					check(Task.Node);
-					const UPCGSettings* Settings = TaskInput.GetSettings(Task.Node->GetSettings());
+					Task.bHasDoneSetup = true;
+
+					// Build input
+					FPCGDataCollection TaskInput;
+					BuildTaskInput(Task, TaskInput);
 
-					if (Settings)
+					// Initialize the element if needed (required to know whether it will run on the main thread or not)
+					if (!Task.Element)
 					{
-						Task.Element = Settings->GetElement();
+						// Get appropriate settings
+						check(Task.Node);
+						const UPCGSettings* Settings = TaskInput.GetSettings(Task.Node->GetSettings());
+
+						if (Settings)
+						{
+							Task.Element = Settings->GetElement();
+						}
 					}
-				}
 
-				// At this point, if the task doesn't have an element, we will never be able to execute it, so we can drop it.
-				if (!Task.Element)
-				{
-					check(!Task.Context);
-					ReadyTasks.RemoveAtSwap(ReadyTaskIndex);
-					continue;
-				}
+					// At this point, if the task doesn't have an element, we will never be able to execute it, so we can drop it.
+					if (!Task.Element)
+					{
+						check(!Task.Context);
+						ReadyTasks.RemoveAtSwap(ReadyTaskIndex);
+						continue;
+					}
 
-				PCGGraphExecutionLogging::LogTaskExecute(Task);
+					PCGGraphExecutionLogging::LogTaskExecute(Task);
 
-				// If a task is cacheable and has been cached, then we don't need to create an active task for it unless
-				// there is an execution mode that would prevent us from doing so.
-				const UPCGSettingsInterface* TaskSettingsInterface = TaskInput.GetSettingsInterface(Task.Node ? Task.Node->GetSettingsInterface() : nullptr);
-				const UPCGSettings* TaskSettings = TaskSettingsInterface ? TaskSettingsInterface->GetSettings() : nullptr;
-				const bool bCacheable = Task.Element->IsCacheableInstance(TaskSettingsInterface);
+					// If a task is cacheable and has been cached, then we don't need to create an active task for it unless
+					// there is an execution mode that would prevent us from doing so.
+					const UPCGSettingsInterface* TaskSettingsInterface = TaskInput.GetSettingsInterface(Task.Node ? Task.Node->GetSettingsInterface() : nullptr);
+					const UPCGSettings* TaskSettings = TaskSettingsInterface ? TaskSettingsInterface->GetSettings() : nullptr;
+					const bool bCacheable = Task.Element->IsCacheableInstance(TaskSettingsInterface);
 
-				// Calculate Crc of dependencies (input data Crcs, settings) and use this as the key in the cache lookup
-				FPCGCrc DependenciesCrc;
-				if (TaskSettings && bCacheable)
-				{
-					Task.Element->GetDependenciesCrc(TaskInput, TaskSettings, Task.SourceComponent.Get(), DependenciesCrc);
-				}
+					// Calculate Crc of dependencies (input data Crcs, settings) and use this as the key in the cache lookup
+					FPCGCrc DependenciesCrc;
+					if (TaskSettings && bCacheable)
+					{
+						Task.Element->GetDependenciesCrc(TaskInput, TaskSettings, Task.SourceComponent.Get(), DependenciesCrc);
+					}
 
-				if (!bCacheable)
-				{
-					PCGGraphExecutionLogging::LogTaskExecuteCachingDisabled(Task);
-				}
+					if (!bCacheable)
+					{
+						PCGGraphExecutionLogging::LogTaskExecuteCachingDisabled(Task);
+					}
 
-				FPCGDataCollection CachedOutput;
-				const bool bResultAlreadyInCache = bCacheable && DependenciesCrc.IsValid() && GraphCache.GetFromCache(Task.Node, Task.Element.Get(), DependenciesCrc, Task.SourceComponent.Get(), CachedOutput);
+					FPCGDataCollection CachedOutput;
+					const bool bResultAlreadyInCache = bCacheable && DependenciesCrc.IsValid() && GraphCache.GetFromCache(Task.Node, Task.Element.Get(), DependenciesCrc, Task.SourceComponent.Get(), CachedOutput);
 #if WITH_EDITOR
-				const bool bNeedsToCreateActiveTask = !bResultAlreadyInCache || TaskSettingsInterface->bDebug;
+					const bool bNeedsToCreateActiveTask = !bResultAlreadyInCache || TaskSettingsInterface->bDebug;
 #else
-				const bool bNeedsToCreateActiveTask = !bResultAlreadyInCache;
+					const bool bNeedsToCreateActiveTask = !bResultAlreadyInCache;
 #endif
 
-				if (!bNeedsToCreateActiveTask)
-				{
-#if WITH_EDITOR
-					// Doing this now since we're about to modify ReadyTasks potentially reallocating while Task is a reference. 
-					if (UPCGComponent* SourceComponent = Task.SourceComponent.Get())
+					if (!bNeedsToCreateActiveTask)
 					{
-						if (Task.StackIndex != INDEX_NONE)
+#if WITH_EDITOR
+						// Doing this now since we're about to modify ReadyTasks potentially reallocating while Task is a reference. 
+						if (UPCGComponent* SourceComponent = Task.SourceComponent.Get())
 						{
-							const FPCGStack* Stack = Task.GetStack();
-							SourceComponent->StoreInspectionData(Stack, Task.Node, nullptr, TaskInput, CachedOutput, /*bUsedCache=*/true);
+							if (Task.StackIndex != INDEX_NONE)
+							{
+								const FPCGStack* Stack = Task.GetStack();
+								SourceComponent->StoreInspectionData(Stack, Task.Node, nullptr, TaskInput, CachedOutput, /*bUsedCache=*/true);
+							}
 						}
-					}
 #endif
 
-					if (bDynamicTaskCulling && TaskSettings && TaskSettings->OutputPinsCanBeDeactivated() && CachedOutput.InactiveOutputPinBitmask != 0)
-					{
-						CullInactiveDownstreamNodes(Task.NodeId, CachedOutput.InactiveOutputPinBitmask);
+						if (bDynamicTaskCulling && TaskSettings && TaskSettings->OutputPinsCanBeDeactivated() && CachedOutput.InactiveOutputPinBitmask != 0)
+						{
+							CullInactiveDownstreamNodes(Task.NodeId, CachedOutput.InactiveOutputPinBitmask);
 
 #if WITH_EDITOR
-						SendInactivePinNotification(Task.Node, Task.GetStack(), CachedOutput.InactiveOutputPinBitmask);
+							SendInactivePinNotification(Task.Node, Task.GetStack(), CachedOutput.InactiveOutputPinBitmask);
 #endif
-					}
+						}
 
-					// Fast-forward cached result to stored results
-					FPCGTaskId SkippedTaskId = Task.NodeId;
-					StoreResults(SkippedTaskId, CachedOutput);
-					delete Task.Context;
-					ReadyTasks.RemoveAtSwap(ReadyTaskIndex);
-					QueueNextTasks(SkippedTaskId);
-					bAnyTaskEnded = true;
+						// Fast-forward cached result to stored results
+						FPCGTaskId SkippedTaskId = Task.NodeId;
+						StoreResults(SkippedTaskId, CachedOutput);
+						delete Task.Context;
+						ReadyTasks.RemoveAtSwap(ReadyTaskIndex);
+						QueueNextTasks(SkippedTaskId);
+						bAnyTaskEnded = true;
 
-					continue;
-				}
+						continue;
+					}
 
-				// Allocate context if not previously done
-				if (!Task.Context)
-				{
-					Task.Context = Task.Element->Initialize(TaskInput, Task.SourceComponent, Task.Node);
-					Task.Context->InitializeSettings();
-					Task.Context->TaskId = Task.NodeId;
-					Task.Context->CompiledTaskId = Task.CompiledTaskId;
-					Task.Context->DependenciesCrc = DependenciesCrc;
-					Task.Context->Stack = Task.GetStack();
+					// Allocate context if not previously done
+					if (!Task.Context)
+					{
+						Task.Context = Task.Element->Initialize(TaskInput, Task.SourceComponent, Task.Node);
+						Task.Context->InitializeSettings();
+						Task.Context->TaskId = Task.NodeId;
+						Task.Context->CompiledTaskId = Task.CompiledTaskId;
+						Task.Context->DependenciesCrc = DependenciesCrc;
+						Task.Context->Stack = Task.GetStack();
+					}
+
+#if WITH_EDITOR
+					if (bResultAlreadyInCache)
+					{
+						Task.bIsBypassed = true;
+						Task.Context->OutputData = CachedOutput;
+					}
+#endif
 				}
 
 				// Validate that we can start this task now
@@ -884,13 +897,8 @@ void FPCGGraphExecutor::Execute()
 					ActiveTask.Context = TUniquePtr<FPCGContext>(Task.Context);
 					ActiveTask.StackIndex = Task.StackIndex;
 					ActiveTask.StackContext = Task.StackContext;
-
 #if WITH_EDITOR
-					if (bResultAlreadyInCache)
-					{
-						ActiveTask.bIsBypassed = true;
-						ActiveTask.Context->OutputData = CachedOutput;
-					}
+					ActiveTask.bIsBypassed = Task.bIsBypassed;
 #endif
 
 					// Move the task up front if it needs to run on the main thread
@@ -1150,12 +1158,6 @@ void FPCGGraphExecutor::Execute()
 			}
 		}
 #endif
-
-#if WITH_EDITOR
-		// Save & release resources when running in-editor
-		SaveDirtyActors();
-		ReleaseUnusedActors();
-#endif
 	}
 }
 
@@ -1695,59 +1697,6 @@ FPCGTaskId FPCGGraphExecutor::ScheduleDebugWithTaskCallback(UPCGComponent* InCom
 	return ScheduleGeneric([] { return true; }, InComponent, FinalDependencies);
 }
 
-void FPCGGraphExecutor::AddToDirtyActors(AActor* Actor)
-{
-	ActorsListLock.Lock();
-	ActorsToSave.Add(Actor);
-	ActorsListLock.Unlock();
-}
-
-void FPCGGraphExecutor::AddToUnusedActors(const TSet<FWorldPartitionReference>& UnusedActors)
-{
-	ActorsListLock.Lock();
-	ActorsToRelease.Append(UnusedActors);
-	ActorsListLock.Unlock();
-}
-
-void FPCGGraphExecutor::SaveDirtyActors()
-{
-	ActorsListLock.Lock();
-	TSet<AActor*> ToSave = MoveTemp(ActorsToSave);
-	ActorsToSave.Reset();
-	ActorsListLock.Unlock();
-
-	TSet<UPackage*> PackagesToSave;
-	for (AActor* Actor : ToSave)
-	{
-		PackagesToSave.Add(Actor->GetExternalPackage());
-	}
-
-	if (PackagesToSave.Num() > 0)
-	{
-		UEditorLoadingAndSavingUtils::SavePackages(PackagesToSave.Array(), true);
-	}
-}
-
-void FPCGGraphExecutor::ReleaseUnusedActors()
-{
-	ActorsListLock.Lock();
-	bool bRunGC = ActorsToRelease.Num() > 0;
-	ActorsToRelease.Reset();
-	ActorsListLock.Unlock();
-
-#if WITH_EDITOR
-	if (bRunGC && !PCGHelpers::IsRuntimeOrPIE())
-	{
-		--ReleaseActorsCountUntilGC;
-		if (ReleaseActorsCountUntilGC <= 0)
-		{
-			ReleaseActorsCountUntilGC = 30;
-			CollectGarbage(GARBAGE_COLLECTION_KEEPFLAGS, true);
-		}
-	}
-#endif
-}
-
 void FPCGGraphExecutor::NotifyGraphChanged(UPCGGraph* InGraph, EPCGChangeType ChangeType)
 {
 	GraphCompiler.NotifyGraphChanged(InGraph, ChangeType);
diff --git a/Engine/Plugins/PCG/Source/PCG/Private/Graph/PCGGraphExecutor.h b/Engine/Plugins/PCG/Source/PCG/Private/Graph/PCGGraphExecutor.h
--- a/Engine/Plugins/PCG/Source/PCG/Private/Graph/PCGGraphExecutor.h
+++ b/Engine/Plugins/PCG/Source/PCG/Private/Graph/PCGGraphExecutor.h
@@ -90,6 +90,14 @@ struct FPCGGraphTask
 
 	int32 StackIndex = INDEX_NONE;
 	TSharedPtr<const FPCGStackContext> StackContext;
+
+	// Contains info around whether the element, context, cache check and more was already done.
+	bool bHasDoneSetup = false;
+
+#if WITH_EDITOR
+	// Can be true when we want to have debug display on a task but have taken the results from the cache
+	bool bIsBypassed = false;
+#endif
 };
 
 struct FPCGGraphScheduleTask
@@ -187,9 +195,6 @@ public:
 #if WITH_EDITOR
 	FPCGTaskId ScheduleDebugWithTaskCallback(UPCGComponent* InComponent, TFunction<void(FPCGTaskId, const UPCGNode*, const FPCGDataCollection&)> TaskCompleteCallback);
 
-	void AddToDirtyActors(AActor* Actor);
-	void AddToUnusedActors(const TSet<FWorldPartitionReference>& UnusedActors);
-
 	/** Notify compiler that graph has changed so it'll be removed from the cache */
 	void NotifyGraphChanged(UPCGGraph* InGraph, EPCGChangeType ChangeType);
 
@@ -239,9 +244,6 @@ private:
 	/** Notify the component that the given pins were deactivated during execution. */
 	void SendInactivePinNotification(const UPCGNode* InNode, const FPCGStack* InStack, uint64 InactiveOutputPinBitmask);
 
-	void SaveDirtyActors();
-	void ReleaseUnusedActors();
-
 	void UpdateGenerationNotification();
 	void ReleaseGenerationNotification();
 	void OnNotificationCancel();
@@ -275,11 +277,6 @@ private:
 	int32 CurrentlyUsedThreads = 0;
 
 #if WITH_EDITOR
-	FCriticalSection ActorsListLock;
-	TSet<AActor*> ActorsToSave;
-	TSet<FWorldPartitionReference> ActorsToRelease;
-
-	int32 ReleaseActorsCountUntilGC = 30;
 	TWeakPtr<IPCGEditorProgressNotification> GenerationProgressNotification;
 	double GenerationProgressNotificationStartTime = 0.0;
 	int32 GenerationProgressLastTaskNum = 0;
