[PCG] Fixed garbage collection & multithreading issues related to some elements

#rnx

diff --git a/Engine/Plugins/PCG/Source/PCG/Private/Elements/Metadata/PCGMetadataOpElementBase.cpp b/Engine/Plugins/PCG/Source/PCG/Private/Elements/Metadata/PCGMetadataOpElementBase.cpp
--- a/Engine/Plugins/PCG/Source/PCG/Private/Elements/Metadata/PCGMetadataOpElementBase.cpp
+++ b/Engine/Plugins/PCG/Source/PCG/Private/Elements/Metadata/PCGMetadataOpElementBase.cpp
@@ -432,6 +432,7 @@ bool FPCGMetadataElementBase::PrepareDataInternal(FPCGContext* Context) const
 	// Set up the iterations on the multiple inputs of the primary pin
 	TimeSlicedContext->InitializePerIterationStates(OperandInputNumMax, [this, Context, OperandNum, Settings, OperandInputNumMax](IterStateType& OutState, const ExecStateType& ExecState, const uint32 IterationIndex)
 	{
+		FPCGMetadataElementBase::ContextType* TimeSlicedContext = static_cast<FPCGMetadataElementBase::ContextType*>(Context);
 		TArray<FPCGTaggedData>& Outputs = Context->OutputData.TaggedData;
 		const uint32 NumberOfResults = Settings->GetResultNum();
 
@@ -472,6 +473,9 @@ bool FPCGMetadataElementBase::PrepareDataInternal(FPCGContext* Context) const
 				}
 				else
 				{
+					// Need to make sure the param data is properly tracked by the context to prevent garbage collection
+					TimeSlicedContext->TrackObject(DefaultData.Data);
+
 					// Need to make sure the param data has at least one entry
 					UPCGMetadata* DefaultParamMetadata = CastChecked<UPCGParamData>(DefaultData.Data)->Metadata;
 					if (DefaultParamMetadata->GetLocalItemCount() == 0)
diff --git a/Engine/Plugins/PCG/Source/PCG/Private/Elements/PCGLoopElement.cpp b/Engine/Plugins/PCG/Source/PCG/Private/Elements/PCGLoopElement.cpp
--- a/Engine/Plugins/PCG/Source/PCG/Private/Elements/PCGLoopElement.cpp
+++ b/Engine/Plugins/PCG/Source/PCG/Private/Elements/PCGLoopElement.cpp
@@ -241,6 +241,7 @@ bool FPCGLoopElement::ExecuteInternal(FPCGContext* InContext) const
 
 		FPCGDataCollection PreSubgraphDataCollection;
 		PrepareSubgraphUserParameters(Settings, Context, PreSubgraphDataCollection);
+		Context->AddToReferencedObjects(PreSubgraphDataCollection);
 
 		FPCGElementPtr PreGraphElement = MakeShared<FPCGInputForwardingElement>(PreSubgraphDataCollection);
 
@@ -281,6 +282,8 @@ bool FPCGLoopElement::ExecuteInternal(FPCGContext* InContext) const
 				Dependencies.Add(PreviousTaskId);
 			}
 
+			Context->AddToReferencedObjects(InputDataCollection);
+
 			FPCGElementPtr InputElement = MakeShared<FPCGLoopInputForwardingElement>(InputDataCollection, PreviousTaskId, FeedbackPinNames);
 			FPCGTaskId SubgraphTaskId = Subsystem->ScheduleGraph(
 				Subgraph,
diff --git a/Engine/Plugins/PCG/Source/PCG/Private/Elements/PCGSpawnActor.cpp b/Engine/Plugins/PCG/Source/PCG/Private/Elements/PCGSpawnActor.cpp
--- a/Engine/Plugins/PCG/Source/PCG/Private/Elements/PCGSpawnActor.cpp
+++ b/Engine/Plugins/PCG/Source/PCG/Private/Elements/PCGSpawnActor.cpp
@@ -586,6 +586,8 @@ bool FPCGSpawnActorElement::SpawnAndPrepareSubgraphs(FPCGSubgraphContext* Contex
 							FPCGDataCollection SubgraphInputData;
 							SubgraphInputData.TaggedData.Add(PartialInput);
 
+							Context->AddToReferencedObjects(SubgraphInputData);
+
 							// Prepare the invocation stack - which is the stack up to this node, and then this node, then a loop index
 							FPCGStack InvocationStack = ensure(Context->Stack) ? *Context->Stack : FPCGStack();
 
diff --git a/Engine/Plugins/PCG/Source/PCG/Private/PCGSubgraph.cpp b/Engine/Plugins/PCG/Source/PCG/Private/PCGSubgraph.cpp
--- a/Engine/Plugins/PCG/Source/PCG/Private/PCGSubgraph.cpp
+++ b/Engine/Plugins/PCG/Source/PCG/Private/PCGSubgraph.cpp
@@ -509,6 +509,22 @@ void FPCGSubgraphContext::UpdateOverridesWithOverriddenGraph()
 	}
 }
 
+void FPCGSubgraphContext::AddToReferencedObjects(const FPCGDataCollection& InDataCollection)
+{
+	for (const FPCGTaggedData& TaggedData : InDataCollection.TaggedData)
+	{
+		if (TaggedData.Data)
+		{
+			ReferencedObjects.Add(TaggedData.Data);
+		}
+	}
+}
+
+void FPCGSubgraphContext::AddExtraStructReferencedObjects(FReferenceCollector& Collector)
+{
+	Collector.AddReferencedObjects(ReferencedObjects);
+}
+
 FPCGContext* FPCGSubgraphElement::Initialize(const FPCGDataCollection& InputData, TWeakObjectPtr<UPCGComponent> SourceComponent, const UPCGNode* Node)
 {
 	FPCGSubgraphContext* Context = new FPCGSubgraphContext();
@@ -642,9 +658,11 @@ bool FPCGSubgraphElement::ExecuteInternal(FPCGContext* InContext) const
 				// using this node's task id as additional inputs
 				FPCGDataCollection PreSubgraphInputData;
 				PrepareSubgraphUserParameters(Settings, Context, PreSubgraphInputData);
+				Context->AddToReferencedObjects(PreSubgraphInputData);
 
 				FPCGDataCollection SubgraphInputData;
 				PrepareSubgraphData(Settings, Context, Context->InputData, SubgraphInputData);
+				Context->AddToReferencedObjects(SubgraphInputData);
 
 				// At this point, if we're in a recursive context and we have no input, we must terminate execution
 				if (bIsRecursive && SubgraphInputData.TaggedData.IsEmpty())
@@ -753,29 +771,10 @@ bool FPCGSubgraphElement::ExecuteInternal(FPCGContext* InContext) const
 FPCGInputForwardingElement::FPCGInputForwardingElement(const FPCGDataCollection& InputToForward)
 	: Input(InputToForward)
 {
-	// Root any previously unrooted data, needed here because the context does not exist yet and we need to ensure that the input is not garbage collected
-	for (const FPCGTaggedData& TaggedData : Input.TaggedData)
-	{
-		if (TaggedData.Data && !TaggedData.Data->IsRooted())
-		{
-			UPCGData* DataToRoot = const_cast<UPCGData*>(TaggedData.Data.Get());
-			DataToRoot->AddToRoot();
-			RootedData.Add(DataToRoot);
-		}
-	}
 }
 
 bool FPCGInputForwardingElement::ExecuteInternal(FPCGContext* Context) const
 {
-	// Remove from rootset during the execution if we had previously done so. After execution, the graph cache will keep track of these references
-	for (UPCGData* DataToUnroot : RootedData)
-	{
-		ensure(DataToUnroot->IsRooted());
-		DataToUnroot->RemoveFromRoot();
-	}
-
-	const_cast<FPCGInputForwardingElement*>(this)->RootedData.Reset();
-
 	Context->OutputData = Input;
 	return true;
 }
diff --git a/Engine/Plugins/PCG/Source/PCG/Public/PCGSubgraph.h b/Engine/Plugins/PCG/Source/PCG/Public/PCGSubgraph.h
--- a/Engine/Plugins/PCG/Source/PCG/Public/PCGSubgraph.h
+++ b/Engine/Plugins/PCG/Source/PCG/Public/PCGSubgraph.h
@@ -164,6 +164,7 @@ struct FPCGSubgraphContext : public FPCGContext
 	TArray<FPCGTaskId> SubgraphTaskIds;
 	bool bScheduledSubgraph = false;
 	FInstancedStruct GraphInstanceParametersOverride;
+	TSet<TObjectPtr<const UPCGData>> ReferencedObjects;
 
 	// Analyze input data to detect if there is any override for the user parameters. If so will duplicate it to gather overrides.
 	void InitializeUserParametersStruct();
@@ -171,8 +172,11 @@ struct FPCGSubgraphContext : public FPCGContext
 	// If we have a subgraph override, update the underlying duplicated parameters with the overrides from the subgraph.
 	void UpdateOverridesWithOverriddenGraph();
 
+	void AddToReferencedObjects(const FPCGDataCollection& InDataCollection);
+
 protected:
 	virtual void* GetUnsafeExternalContainerForOverridableParam(const FPCGSettingsOverridableParam& InParam) override;
+	virtual void AddExtraStructReferencedObjects(FReferenceCollector& Collector);
 };
 
 class FPCGSubgraphElement : public IPCGElement
@@ -188,15 +192,14 @@ protected:
 	void PrepareSubgraphUserParameters(const UPCGSubgraphSettings* Settings, FPCGSubgraphContext* Context, FPCGDataCollection& OutputData) const;
 };
 
+// Implementation note: this node forwards data, but does not keep that data alive. This is the responsibility of the corresponding FPCGSubgraphContext
 class FPCGInputForwardingElement : public IPCGElement
 {
 public:
-	FPCGInputForwardingElement(const FPCGDataCollection& InputToForward);
+	explicit FPCGInputForwardingElement(const FPCGDataCollection& InputToForward);
 
 protected:
 	virtual bool ExecuteInternal(FPCGContext* Context) const override;
 	virtual bool IsPassthrough(const UPCGSettings* InSettings) const override { return true; }
 	FPCGDataCollection Input;
-
-	TArray<UPCGData*> RootedData;
 };
