[PCG] Improved execution results map memory consumption by better tracking. More improvements to come.

#rnx

diff --git a/Engine/Plugins/PCG/Source/PCG/Private/Elements/PCGLoopElement.cpp b/Engine/Plugins/PCG/Source/PCG/Private/Elements/PCGLoopElement.cpp
--- a/Engine/Plugins/PCG/Source/PCG/Private/Elements/PCGLoopElement.cpp
+++ b/Engine/Plugins/PCG/Source/PCG/Private/Elements/PCGLoopElement.cpp
@@ -359,7 +359,7 @@ bool FPCGLoopElement::ExecuteInternal(FPCGContext* InContext) const
 
 				FPCGDataCollection SubgraphOutput;
 				// While this should be always return true, if a scheduled task was cancelled, this can still happen because the dependency has been removed but this code isn't aware of this
-				if (Subsystem->GetOutputData(SubgraphTaskId, SubgraphOutput))
+				if (Subsystem->GetOutputData(SubgraphTaskId, SubgraphOutput, /*bClearDataOnGet=*/true))
 				{
 					if (FeedbackPinNames.IsEmpty() || bIsLastTask)
 					{
@@ -432,7 +432,7 @@ bool FPCGLoopInputForwardingElement::ExecuteInternal(FPCGContext* Context) const
 		if (UPCGSubsystem* Subsystem = Context->SourceComponent->GetSubsystem())
 		{
 			FPCGDataCollection PreviousTaskOutput;
-			if (Subsystem->GetOutputData(PreviousIterationTaskId, PreviousTaskOutput))
+			if (Subsystem->GetOutputData(PreviousIterationTaskId, PreviousTaskOutput, /*bClearDataOnGet=*/false))
 			{
 				Algo::CopyIf(PreviousTaskOutput.TaggedData, Context->OutputData.TaggedData, [this](const FPCGTaggedData& InTaggedData) { return FeedbackPinNames.Contains(InTaggedData.Pin); });
 			}
diff --git a/Engine/Plugins/PCG/Source/PCG/Private/Graph/PCGGraphCache.cpp b/Engine/Plugins/PCG/Source/PCG/Private/Graph/PCGGraphCache.cpp
--- a/Engine/Plugins/PCG/Source/PCG/Private/Graph/PCGGraphCache.cpp
+++ b/Engine/Plugins/PCG/Source/PCG/Private/Graph/PCGGraphCache.cpp
@@ -9,6 +9,11 @@
 #include "HAL/IConsoleManager.h"
 #include "Misc/ScopeRWLock.h"
 
+static TAutoConsoleVariable<bool> CVarCacheEnabled(
+	TEXT("pcg.Cache.Enabled"),
+	true,
+	TEXT("Enables the cache system."));
+
 static TAutoConsoleVariable<bool> CVarCacheDebugging(
 	TEXT("pcg.Cache.EnableDebugging"),
 	false,
@@ -39,6 +44,11 @@ FPCGGraphCache::~FPCGGraphCache()
 
 bool FPCGGraphCache::GetFromCache(const UPCGNode* InNode, const IPCGElement* InElement, const FPCGCrc& InDependenciesCrc, const UPCGComponent* InComponent, FPCGDataCollection& OutOutput) const
 {
+	if (!CVarCacheEnabled.GetValueOnAnyThread())
+	{
+		return false;
+	}
+
 	if(!InDependenciesCrc.IsValid())
 	{
 		UE_LOG(LogPCG, Warning, TEXT("Invalid dependencies passed to FPCGGraphCache::GetFromCache(), lookup aborted."));
@@ -78,6 +88,11 @@ bool FPCGGraphCache::GetFromCache(const UPCGNode* InNode, const IPCGElement* InE
 
 void FPCGGraphCache::StoreInCache(const IPCGElement* InElement, const FPCGCrc& InDependenciesCrc, const FPCGDataCollection& InOutput)
 {
+	if (!CVarCacheEnabled.GetValueOnAnyThread())
+	{
+		return;
+	}
+
 	if (!ensure(InDependenciesCrc.IsValid()))
 	{
 		return;
@@ -113,6 +128,10 @@ void FPCGGraphCache::ClearCache()
 bool FPCGGraphCache::EnforceMemoryBudget()
 {
 	TRACE_CPUPROFILER_EVENT_SCOPE(FPCGGraphCache::FPCGGraphCache::EnforceMemoryBudget);
+	if (!CVarCacheEnabled.GetValueOnAnyThread())
+	{
+		return true; // Always return true so we can report that a GC might be needed
+	}
 
 	if (!CVarCacheMemoryBudgetEnabled.GetValueOnAnyThread())
 	{
diff --git a/Engine/Plugins/PCG/Source/PCG/Private/Graph/PCGGraphCompiler.cpp b/Engine/Plugins/PCG/Source/PCG/Private/Graph/PCGGraphCompiler.cpp
--- a/Engine/Plugins/PCG/Source/PCG/Private/Graph/PCGGraphCompiler.cpp
+++ b/Engine/Plugins/PCG/Source/PCG/Private/Graph/PCGGraphCompiler.cpp
@@ -1077,7 +1077,7 @@ void FPCGGraphCompiler::CompileTopGraph(UPCGGraph* InGraph, uint32 GenerationGri
 	}
 
 	FPCGGraphTask& PostExecuteTask = CompiledTasks.Emplace_GetRef();
-	PostExecuteTask.Element = GetSharedTrivialElement();
+	PostExecuteTask.Element = GetSharedTrivialPostGraphElement();
 	PostExecuteTask.NodeId = PostExecuteTaskId;
 
 	// Find end nodes, e.g. all nodes that have no successors.
@@ -1142,6 +1142,27 @@ FPCGElementPtr FPCGGraphCompiler::GetSharedTrivialElement()
 	return SharedTrivialElement;
 }
 
+FPCGElementPtr FPCGGraphCompiler::GetSharedTrivialPostGraphElement()
+{
+	{
+		FReadScopeLock Lock(SharedTrivialElementLock);
+
+		if (SharedTrivialPostGraphElement)
+		{
+			return SharedTrivialPostGraphElement;
+		}
+	}
+
+	FWriteScopeLock Lock(SharedTrivialElementLock);
+
+	if (!SharedTrivialPostGraphElement)
+	{
+		SharedTrivialPostGraphElement = MakeShared<FPCGTrivialElement>();
+	}
+
+	return SharedTrivialPostGraphElement;
+}
+
 void FPCGGraphCompiler::ClearCache()
 {
 	FWriteScopeLock Lock(GraphToTaskMapLock);
diff --git a/Engine/Plugins/PCG/Source/PCG/Private/Graph/PCGGraphCompiler.h b/Engine/Plugins/PCG/Source/PCG/Private/Graph/PCGGraphCompiler.h
--- a/Engine/Plugins/PCG/Source/PCG/Private/Graph/PCGGraphCompiler.h
+++ b/Engine/Plugins/PCG/Source/PCG/Private/Graph/PCGGraphCompiler.h
@@ -38,6 +38,9 @@ public:
 	/** Flush all cached compiled graphs. */
 	void ClearCache();
 
+	/** Returns the post graph element, which is used to determine results caching behavior */
+	FPCGElementPtr GetSharedTrivialPostGraphElement();
+
 private:
 	TArray<FPCGGraphTask> CompileGraph(UPCGGraph* InGraph, FPCGTaskId& NextId, FPCGStackContext& InOutStackContext);
 
@@ -94,6 +97,7 @@ private:
 	TMap<UPCGGraph*, TMap<uint32, FPCGStackContext>> TopGraphToStackContextMap;
 
 	FPCGElementPtr SharedTrivialElement;
+	FPCGElementPtr SharedTrivialPostGraphElement;
 	mutable FRWLock SharedTrivialElementLock;
 
 #if WITH_EDITOR
diff --git a/Engine/Plugins/PCG/Source/PCG/Private/Graph/PCGGraphExecutor.cpp b/Engine/Plugins/PCG/Source/PCG/Private/Graph/PCGGraphExecutor.cpp
--- a/Engine/Plugins/PCG/Source/PCG/Private/Graph/PCGGraphExecutor.cpp
+++ b/Engine/Plugins/PCG/Source/PCG/Private/Graph/PCGGraphExecutor.cpp
@@ -587,12 +587,31 @@ FPCGTaskId FPCGGraphExecutor::ScheduleGenericWithContext(TFunction<bool(FPCGCont
 	return Task.NodeId;
 }
 
-bool FPCGGraphExecutor::GetOutputData(FPCGTaskId TaskId, FPCGDataCollection& OutData)
+void FPCGGraphExecutor::MarkInputResults(TArrayView<const FPCGTaskId> InInputResults)
+{
+	for (FPCGTaskId TaskId : InInputResults)
+	{
+		FOutputDataInfo& OutputDataInfo = OutputData[TaskId];
+		if (!OutputDataInfo.bNeedsManualClear && --OutputDataInfo.RemainingSuccessorCount == 0)
+		{
+			OutputData.Remove(TaskId);
+		}
+	}
+}
+
+bool FPCGGraphExecutor::GetOutputData(FPCGTaskId TaskId, FPCGDataCollection& OutData, bool bClearDataOnGet)
 {
 	// TODO: this is not threadsafe - make threadsafe once we multithread execution
 	if (OutputData.Contains(TaskId))
 	{
-		OutData = OutputData[TaskId];
+		FOutputDataInfo& OutputDataInfo = OutputData[TaskId];
+		OutData = OutputDataInfo.DataCollection;
+
+		if (bClearDataOnGet && ensure(OutputDataInfo.bNeedsManualClear))
+		{
+			OutputData.Remove(TaskId);
+		}
+
 		return true;
 	}
 	else
@@ -854,9 +873,12 @@ void FPCGGraphExecutor::Execute()
 #endif
 						}
 
+						// If the task is a post execute, then we can safely clear the data after getting it from the results.
+						const bool bTaskIsPostExecute = (Task.Element == GraphCompiler.GetSharedTrivialPostGraphElement());
+
 						// Fast-forward cached result to stored results
 						FPCGTaskId SkippedTaskId = Task.NodeId;
-						StoreResults(SkippedTaskId, CachedOutput);
+						StoreResults(SkippedTaskId, CachedOutput, bTaskIsPostExecute);
 						delete Task.Context;
 						ReadyTasks.RemoveAtSwap(ReadyTaskIndex);
 						QueueNextTasks(SkippedTaskId);
@@ -1029,9 +1051,12 @@ void FPCGGraphExecutor::Execute()
 			}
 #endif
 
+			// If the task is a post execute, then we can safely clear the data after getting it from the results.
+			const bool bTaskIsPostExecute = (ActiveTask.Element == GraphCompiler.GetSharedTrivialPostGraphElement());
+
 			// Store output in data map.
 			// TODO - investigate if we should avoid doing this if the task was cancelled.
-			StoreResults(ActiveTask.NodeId, ActiveTask.Context->OutputData);
+			StoreResults(ActiveTask.NodeId, ActiveTask.Context->OutputData, bTaskIsPostExecute);
 
 			// Book-keeping
 			QueueNextTasks(ActiveTask.NodeId);
@@ -1145,19 +1170,10 @@ void FPCGGraphExecutor::Execute()
 		}
 
 		// Purge things from cache if memory usage is too high
-		const bool bSomethingTidied = GraphCache.EnforceMemoryBudget();
-
-#if WITH_EDITOR
-		if (bSomethingTidied && !PCGHelpers::IsRuntimeOrPIE())
+		if (GraphCache.EnforceMemoryBudget())
 		{
-			--TidyCacheCountUntilGC;
-			if (TidyCacheCountUntilGC <= 0)
-			{
-				TidyCacheCountUntilGC = 100;
-				CollectGarbage(GARBAGE_COLLECTION_KEEPFLAGS, true);
-			}
+			CollectGarbage(GARBAGE_COLLECTION_KEEPFLAGS, true);
 		}
-#endif
 	}
 }
 
@@ -1312,10 +1328,12 @@ void FPCGGraphExecutor::BuildTaskInput(const FPCGGraphTask& Task, FPCGDataCollec
 	// Hoisted out of loop for performance reasons.
 	TArray<FPCGTaggedData, TInlineAllocator<16>> InputDataOnPin;
 	TArray<FPCGCrc, TInlineAllocator<16>> InputDataCrcsOnPin;
+	TArray<FPCGTaskId, TInlineAllocator<16>> ResultsToMarkAsRead;
 
 	for (const FPCGGraphTaskInput& Input : Task.Inputs)
 	{
 		check(OutputData.Contains(Input.TaskId));
+		ResultsToMarkAsRead.AddUnique(Input.TaskId);
 
 		// If the input does not provide any data, don't add it to the task input.
 		if (!Input.bProvideData)
@@ -1333,7 +1351,7 @@ void FPCGGraphExecutor::BuildTaskInput(const FPCGGraphTask& Task, FPCGDataCollec
 			continue;
 		}
 
-		const FPCGDataCollection& InputCollection = OutputData[Input.TaskId];
+		const FPCGDataCollection& InputCollection = OutputData[Input.TaskId].DataCollection;
 
 		TaskInput.bCancelExecution |= InputCollection.bCancelExecution;
 
@@ -1387,6 +1405,9 @@ void FPCGGraphExecutor::BuildTaskInput(const FPCGGraphTask& Task, FPCGDataCollec
 
 	// Then combine params if needed
 	CombineParams(Task.NodeId, TaskInput);
+
+	// Finally, mark inputs we read from so we can release them if we were the last remaining consumer.
+	MarkInputResults(ResultsToMarkAsRead);
 }
 
 void FPCGGraphExecutor::CombineParams(FPCGTaskId InTaskId, FPCGDataCollection& InTaskInput)
@@ -1440,12 +1461,20 @@ void FPCGGraphExecutor::CombineParams(FPCGTaskId InTaskId, FPCGDataCollection& I
 	}
 }
 
-void FPCGGraphExecutor::StoreResults(FPCGTaskId InTaskId, const FPCGDataCollection& InTaskOutput)
+void FPCGGraphExecutor::StoreResults(FPCGTaskId InTaskId, const FPCGDataCollection& InTaskOutput, bool bNeedsManualClear)
 {
 	TRACE_CPUPROFILER_EVENT_SCOPE(FPCGGraphExecutor::StoreResults);
 
+	FOutputDataInfo OutputDataInfo;
+	OutputDataInfo.DataCollection = InTaskOutput;
+	OutputDataInfo.bNeedsManualClear = bNeedsManualClear;
+	if (TSet<FPCGTaskId>* Successors = TaskSuccessors.Find(InTaskId))
+	{
+		OutputDataInfo.RemainingSuccessorCount = Successors->Num();
+	}
+
 	// Store output in map
-	OutputData.Add(InTaskId, InTaskOutput);
+	OutputData.Add(InTaskId, MoveTemp(OutputDataInfo));
 }
 
 void FPCGGraphExecutor::ClearResults()
@@ -1615,7 +1644,7 @@ void FPCGGraphExecutor::AddReferencedObjects(FReferenceCollector& Collector)
 	// Go through all data in the cached output map
 	for (auto& OutputDataEntry : OutputData)
 	{
-		OutputDataEntry.Value.AddReferences(Collector);
+		OutputDataEntry.Value.DataCollection.AddReferences(Collector);
 	}
 
 	// Go through ready tasks, active tasks and sleeping tasks contexts
@@ -1681,7 +1710,7 @@ FPCGTaskId FPCGGraphExecutor::ScheduleDebugWithTaskCallback(UPCGComponent* InCom
 		FPCGTaskId CaptureTaskId = ScheduleGeneric([this, TaskCompleteCallback, CompiledTask]
 		{
 			FPCGDataCollection TaskOutputData;
-			if (CompiledTask.Node && GetOutputData(CompiledTask.NodeId, TaskOutputData))
+			if (CompiledTask.Node && GetOutputData(CompiledTask.NodeId, TaskOutputData, /*bClearDataOnGet=*/false))
 			{
 				TaskCompleteCallback(CompiledTask.NodeId, CompiledTask.Node, TaskOutputData);
 			}
diff --git a/Engine/Plugins/PCG/Source/PCG/Private/Graph/PCGGraphExecutor.h b/Engine/Plugins/PCG/Source/PCG/Private/Graph/PCGGraphExecutor.h
--- a/Engine/Plugins/PCG/Source/PCG/Private/Graph/PCGGraphExecutor.h
+++ b/Engine/Plugins/PCG/Source/PCG/Private/Graph/PCGGraphExecutor.h
@@ -184,7 +184,7 @@ public:
 	FPCGTaskId ScheduleGenericWithContext(TFunction<bool(FPCGContext*)> InOperation, TFunction<void(FPCGContext*)> InAbortOperation, UPCGComponent* InSourceComponent, const TArray<FPCGTaskId>& TaskExecutionDependencies, const TArray<FPCGTaskId>& TaskDataDependencies);
 
 	/** Gets data in the output results. Returns false if data is not ready. */
-	bool GetOutputData(FPCGTaskId InTaskId, FPCGDataCollection& OutData);
+	bool GetOutputData(FPCGTaskId InTaskId, FPCGDataCollection& OutData, bool bClearDataOnGet);
 
 	/** Accessor so PCG tools (e.g. profiler) can easily decode graph task ids **/
 	FPCGGraphCompiler& GetCompiler() { return GraphCompiler; }
@@ -225,8 +225,9 @@ private:
 	void BuildTaskInput(const FPCGGraphTask& Task, FPCGDataCollection& TaskInput);
 	/** Combine all param data into one on the Params pin, if any.*/
 	void CombineParams(FPCGTaskId InTaskId, FPCGDataCollection& InTaskInput);
-	void StoreResults(FPCGTaskId InTaskId, const FPCGDataCollection& InTaskOutput);
+	void StoreResults(FPCGTaskId InTaskId, const FPCGDataCollection& InTaskOutput, bool bNeedsManualClear);
 	void ClearResults();
+	void MarkInputResults(TArrayView<const FPCGTaskId> InInputResults);
 
 	/** If the completed task has one or more deactivated pins, delete any downstream tasks that are inactive as a result. */
 	void CullInactiveDownstreamNodes(FPCGTaskId CompletedTaskId, uint64 InInactiveOutputPinBitmask);
@@ -269,7 +270,17 @@ private:
 	TMap<FPCGTaskId, TSet<FPCGTaskId>> TaskSuccessors;
 	/** Map of node instances to their output, could be cleared once execution is done */
 	/** Note: this should at some point unload based on loaded/unloaded proxies, otherwise memory cost will be unbounded */
-	TMap<FPCGTaskId, FPCGDataCollection> OutputData;
+	struct FOutputDataInfo
+	{
+		FPCGDataCollection DataCollection;
+		// Controls whether the results will be expunged from the OutputData as soon as the successor count reaches 0 or not.
+		bool bNeedsManualClear = false;
+		// Successor count, updated after a successor is done executing (MarkInputResults).
+		int32 RemainingSuccessorCount = 0;
+	};
+
+	TMap<FPCGTaskId, FOutputDataInfo> OutputData;
+
 	/** Monotonically increasing id. Should be reset once all tasks are executed, should be protected by the ScheduleLock */
 	FPCGTaskId NextTaskId = 0;
 
@@ -280,8 +291,6 @@ private:
 	TWeakPtr<IPCGEditorProgressNotification> GenerationProgressNotification;
 	double GenerationProgressNotificationStartTime = 0.0;
 	int32 GenerationProgressLastTaskNum = 0;
-
-	int32 TidyCacheCountUntilGC = 100;
 #endif
 
 	TObjectPtr<UWorld> World = nullptr;
diff --git a/Engine/Plugins/PCG/Source/PCG/Private/PCGSubgraph.cpp b/Engine/Plugins/PCG/Source/PCG/Private/PCGSubgraph.cpp
--- a/Engine/Plugins/PCG/Source/PCG/Private/PCGSubgraph.cpp
+++ b/Engine/Plugins/PCG/Source/PCG/Private/PCGSubgraph.cpp
@@ -747,7 +747,7 @@ bool FPCGSubgraphElement::ExecuteInternal(FPCGContext* InContext) const
 					// as we just can call the GetOutputData on a fresh data collection and merge it afterwards in the output data,
 					// but this is not needed here so we will keep the full assignment & benefit from the crc as well.
 					ensure(Context->SubgraphTaskIds.Num() == 1);
-					Subsystem->GetOutputData(Context->SubgraphTaskIds[0], Context->OutputData);
+					Subsystem->GetOutputData(Context->SubgraphTaskIds[0], Context->OutputData, /*bClearDataOnGet=*/true);
 				}
 			}
 			else
diff --git a/Engine/Plugins/PCG/Source/PCG/Private/PCGSubsystem.cpp b/Engine/Plugins/PCG/Source/PCG/Private/PCGSubsystem.cpp
--- a/Engine/Plugins/PCG/Source/PCG/Private/PCGSubsystem.cpp
+++ b/Engine/Plugins/PCG/Source/PCG/Private/PCGSubsystem.cpp
@@ -1230,10 +1230,10 @@ FPCGGenSourceManager* UPCGSubsystem::GetGenSourceManager() const
 	return RuntimeGenScheduler ? RuntimeGenScheduler->GenSourceManager : nullptr;
 }
 
-bool UPCGSubsystem::GetOutputData(FPCGTaskId TaskId, FPCGDataCollection& OutData)
+bool UPCGSubsystem::GetOutputData(FPCGTaskId TaskId, FPCGDataCollection& OutData, bool bClearDataOnGet)
 {
 	check(GraphExecutor);
-	return GraphExecutor->GetOutputData(TaskId, OutData);
+	return GraphExecutor->GetOutputData(TaskId, OutData, bClearDataOnGet);
 }
 
 #if WITH_EDITOR
diff --git a/Engine/Plugins/PCG/Source/PCG/Private/Tests/Subgraphs/PCGCollapseSubgraphTest.cpp b/Engine/Plugins/PCG/Source/PCG/Private/Tests/Subgraphs/PCGCollapseSubgraphTest.cpp
--- a/Engine/Plugins/PCG/Source/PCG/Private/Tests/Subgraphs/PCGCollapseSubgraphTest.cpp
+++ b/Engine/Plugins/PCG/Source/PCG/Private/Tests/Subgraphs/PCGCollapseSubgraphTest.cpp
@@ -283,7 +283,7 @@ bool FPCGCollapseSubgraphCombinationOfCollapses::RunTest(const FString& Paramete
 		const FPCGTaskId FinalTaskId = GraphExecutor.ScheduleGeneric([&OutputPoints, &GraphExecutor, TaskId, &bDone]() -> bool
 		{
 			FPCGDataCollection OutputData{};
-			GraphExecutor.GetOutputData(TaskId, OutputData);
+			GraphExecutor.GetOutputData(TaskId, OutputData, /*bClearDataOnGet=*/false);
 
 			if (OutputData.TaggedData.Num() == 1)
 			{
diff --git a/Engine/Plugins/PCG/Source/PCG/Public/PCGSubsystem.h b/Engine/Plugins/PCG/Source/PCG/Public/PCGSubsystem.h
--- a/Engine/Plugins/PCG/Source/PCG/Public/PCGSubsystem.h
+++ b/Engine/Plugins/PCG/Source/PCG/Public/PCGSubsystem.h
@@ -180,7 +180,7 @@ public:
 	void CancelAllGeneration();
 
 	/** Gets the output data for a given task */
-	bool GetOutputData(FPCGTaskId InTaskId, FPCGDataCollection& OutData);
+	bool GetOutputData(FPCGTaskId InTaskId, FPCGDataCollection& OutData, bool bClearDataOnGet = false);
 
 	/** Register a new PCG Component or update it, will be added to the octree if it doesn't exists yet. Returns true if it was added/updated. Thread safe */
 	bool RegisterOrUpdatePCGComponent(UPCGComponent* InComponent, bool bDoActorMapping = true) { return ActorAndComponentMapping.RegisterOrUpdatePCGComponent(InComponent, bDoActorMapping); }
