[PCG] Temporarily forced cleanup in a single frame to prevent issues with filling update pipeline on other systems - will introduce a managed components list resource to enforce synchronous release soon.

#rnx

diff --git a/Engine/Plugins/PCG/Source/PCG/Private/PCGComponent.cpp b/Engine/Plugins/PCG/Source/PCG/Private/PCGComponent.cpp
--- a/Engine/Plugins/PCG/Source/PCG/Private/PCGComponent.cpp
+++ b/Engine/Plugins/PCG/Source/PCG/Private/PCGComponent.cpp
@@ -1362,7 +1362,22 @@ FPCGTaskId UPCGComponent::CreateCleanupTask(bool bRemoveComponents, const TArray
 		return true;
 	};
 
-	return GetSubsystem()->ScheduleGeneric(CleanupTask, this, *AllDependencies);
+#if WITH_EDITOR
+	if (!PCGHelpers::IsRuntimeOrPIE())
+	{
+		auto SuperCleanupTask = [CleanupTask]()
+		{
+			while (!CleanupTask()) {}
+			return true;
+		};
+
+		return GetSubsystem()->ScheduleGeneric(SuperCleanupTask, this, *AllDependencies);
+	}
+	else
+#endif
+	{
+		return GetSubsystem()->ScheduleGeneric(CleanupTask, this, *AllDependencies);
+	}
 }
 
 void UPCGComponent::CleanupUnusedManagedResources()
