[PCG] Refactoring of texture data so color channel used to drive density is clearer, and added density merge operations to the sample texture node

diff --git a/Engine/Plugins/PCG/Source/PCG/Private/Data/PCGSurfaceData.cpp b/Engine/Plugins/PCG/Source/PCG/Private/Data/PCGSurfaceData.cpp
--- a/Engine/Plugins/PCG/Source/PCG/Private/Data/PCGSurfaceData.cpp
+++ b/Engine/Plugins/PCG/Source/PCG/Private/Data/PCGSurfaceData.cpp
@@ -8,4 +8,5 @@
 void UPCGSurfaceData::CopyBaseSurfaceData(UPCGSurfaceData* NewSurfaceData) const
 {
 	NewSurfaceData->Transform = Transform;
+	NewSurfaceData->bKeepZeroDensityPoints = bKeepZeroDensityPoints;
 }
diff --git a/Engine/Plugins/PCG/Source/PCG/Private/Data/PCGTextureData.cpp b/Engine/Plugins/PCG/Source/PCG/Private/Data/PCGTextureData.cpp
--- a/Engine/Plugins/PCG/Source/PCG/Private/Data/PCGTextureData.cpp
+++ b/Engine/Plugins/PCG/Source/PCG/Private/Data/PCGTextureData.cpp
@@ -160,6 +160,32 @@ namespace PCGTextureSamplingHelpers
 	}
 }
 
+void UPCGBaseTextureData::PostLoad()
+{
+	Super::PostLoad();
+
+#if WITH_EDITOR
+	PRAGMA_DISABLE_DEPRECATION_WARNINGS
+	if (DensityFunction != EPCGTextureDensityFunction::Multiply)
+	{
+		bUseDensitySourceChannel = false;
+	}
+	PRAGMA_ENABLE_DEPRECATION_WARNINGS
+#endif
+}
+
+PRAGMA_DISABLE_DEPRECATION_WARNINGS
+EPCGTextureDensityFunction UPCGBaseTextureData::GetDensityFunctionEquivalent() const
+{
+	return bUseDensitySourceChannel ? EPCGTextureDensityFunction::Multiply : EPCGTextureDensityFunction::Ignore;
+}
+
+void UPCGBaseTextureData::SetDensityFunctionEquivalent(EPCGTextureDensityFunction InDensityFunction)
+{
+	bUseDensitySourceChannel = (InDensityFunction != EPCGTextureDensityFunction::Ignore);
+}
+PRAGMA_ENABLE_DEPRECATION_WARNINGS
+
 FBox UPCGBaseTextureData::GetBounds() const
 {
 	return Bounds;
@@ -201,7 +227,7 @@ bool UPCGBaseTextureData::SamplePoint(const FTransform& InTransform, const FBox&
 	if (PCGTextureSamplingHelpers::Sample<FLinearColor>(Position2D, Surface, this, Width, Height, Color, [this](int32 Index) { return ColorData[Index]; }))
 	{
 		OutPoint.Color = Color;
-		OutPoint.Density = ((DensityFunction == EPCGTextureDensityFunction::Ignore) ? 1.0f : PCGTextureSamplingHelpers::SampleFloatChannel(Color, ColorChannel));
+		OutPoint.Density = bUseDensitySourceChannel ? PCGTextureSamplingHelpers::SampleFloatChannel(Color, ColorChannel) : 1.0f;
 		return OutPoint.Density > 0 || bKeepZeroDensityPoints;
 	}
 	else
@@ -257,7 +283,7 @@ const UPCGPointData* UPCGBaseTextureData::CreatePointData(FPCGContext* Context)
 
 		if (PCGTextureSamplingHelpers::Sample<FLinearColor>(LocalCoordinate, Surface, this, Width, Height, Color, [this](int32 Index) { return ColorData[Index]; }))
 		{
-			const float Density = ((DensityFunction == EPCGTextureDensityFunction::Ignore) ? 1.0f : PCGTextureSamplingHelpers::SampleFloatChannel(Color, ColorChannel));
+			const float Density = bUseDensitySourceChannel ? PCGTextureSamplingHelpers::SampleFloatChannel(Color, ColorChannel) : 1.0f;
 			if (Density > 0 || bKeepZeroDensityPoints)
 			{
 				FVector LocalPosition(LocalCoordinate, 0);
@@ -297,7 +323,7 @@ bool UPCGBaseTextureData::SamplePointLocal(const FVector2D& LocalPosition, FVect
 	const FLinearColor OutSample = PCGTextureSamplingHelpers::SampleInternal<FLinearColor>(Pos, Width, Height, Filter, [this](int32 Index) { return ColorData[Index]; });
 
 	OutColor = OutSample;
-	OutDensity = (DensityFunction == EPCGTextureDensityFunction::Ignore) ? 1.0f : PCGTextureSamplingHelpers::SampleFloatChannel(OutSample, ColorChannel);
+	OutDensity = bUseDensitySourceChannel ? PCGTextureSamplingHelpers::SampleFloatChannel(OutSample, ColorChannel) : 1.0f;
 	
 	return OutDensity > 0.0 || bKeepZeroDensityPoints;
 }
@@ -306,7 +332,7 @@ void UPCGBaseTextureData::CopyBaseTextureData(UPCGBaseTextureData* NewTextureDat
 {
 	CopyBaseSurfaceData(NewTextureData);
 
-	NewTextureData->DensityFunction = DensityFunction;
+	NewTextureData->bUseDensitySourceChannel = bUseDensitySourceChannel;
 	NewTextureData->ColorChannel = ColorChannel;
 	NewTextureData->TexelSize = TexelSize;
 	NewTextureData->bUseAdvancedTiling = bUseAdvancedTiling;
diff --git a/Engine/Plugins/PCG/Source/PCG/Private/Elements/PCGSampleTexture.cpp b/Engine/Plugins/PCG/Source/PCG/Private/Elements/PCGSampleTexture.cpp
--- a/Engine/Plugins/PCG/Source/PCG/Private/Elements/PCGSampleTexture.cpp
+++ b/Engine/Plugins/PCG/Source/PCG/Private/Elements/PCGSampleTexture.cpp
@@ -7,6 +7,7 @@
 #include "Data/PCGPointData.h"
 #include "Data/PCGTextureData.h"
 #include "Helpers/PCGAsync.h"
+#include "Helpers/PCGHelpers.h"
 #include "Metadata/Accessors/PCGAttributeAccessorHelpers.h"
 
 #define LOCTEXT_NAMESPACE "PCGSampleTextureElement"
@@ -56,6 +57,8 @@ bool FPCGSampleTextureElement::ExecuteInternal(FPCGContext* Context) const
 		return true;
 	}
 
+	auto DensityMergeFunc = PCGHelpers::GetDensityMergeFunction(Settings->DensityMergeFunction);
+
 	for (int i = 0; i < PointInputs.Num(); ++i)
 	{
 		const UPCGPointData* PointData = Cast<UPCGPointData>(PointInputs[i].Data);
@@ -92,7 +95,7 @@ bool FPCGSampleTextureElement::ExecuteInternal(FPCGContext* Context) const
 		UPCGPointData* OutPointData = NewObject<UPCGPointData>();
 		OutPointData->InitializeFromData(PointData);
 
-		auto ProcessPoint = [Settings, BaseTextureData, &InputAccessor, &InputKeys, &InputPoints, OutPointData](int32 Index, FPCGPoint& OutPoint)
+		auto ProcessPoint = [Settings, BaseTextureData, &InputAccessor, &InputKeys, &InputPoints, &DensityMergeFunc, OutPointData](int32 Index, FPCGPoint& OutPoint)
 		{
 			OutPoint = InputPoints[Index];
 
@@ -107,11 +110,42 @@ bool FPCGSampleTextureElement::ExecuteInternal(FPCGContext* Context) const
 					OutSamplePosition.Y = FMath::Clamp(OutSamplePosition.Y, 0.0, 1.0);
 				}
 
-				return BaseTextureData->SamplePointLocal(FVector2D(OutSamplePosition), OutPoint.Color, OutPoint.Density);
+				float SampleDensity = 1.0f;
+				if (BaseTextureData->SamplePointLocal(FVector2D(OutSamplePosition), OutPoint.Color, SampleDensity))
+				{
+					float ComputedDensity = DensityMergeFunc(OutPoint.Density, SampleDensity);
+
+					if (Settings->bClampOutputDensity)
+					{
+						ComputedDensity = FMath::Clamp(ComputedDensity, 0.0f, 1.0f);
+					}
+
+					OutPoint.Density = ComputedDensity;
+					return true;
+				}
+				else
+				{
+					return false;
+				}
 			}
 			else
 			{
-				return BaseTextureData->SamplePoint(OutPoint.Transform, OutPoint.GetLocalBounds(), OutPoint, OutPointData->Metadata);
+				if (BaseTextureData->SamplePoint(OutPoint.Transform, OutPoint.GetLocalBounds(), OutPoint, OutPointData->Metadata))
+				{
+					float ComputedDensity = DensityMergeFunc(InputPoints[Index].Density, OutPoint.Density);
+
+					if (Settings->bClampOutputDensity)
+					{
+						ComputedDensity = FMath::Clamp(ComputedDensity, 0.0f, 1.0f);
+					}
+
+					OutPoint.Density = ComputedDensity;
+					return true;
+				}
+				else
+				{
+					return false;
+				}
 			}
 		};
 
diff --git a/Engine/Plugins/PCG/Source/PCG/Private/Elements/PCGTextureSampler.cpp b/Engine/Plugins/PCG/Source/PCG/Private/Elements/PCGTextureSampler.cpp
--- a/Engine/Plugins/PCG/Source/PCG/Private/Elements/PCGTextureSampler.cpp
+++ b/Engine/Plugins/PCG/Source/PCG/Private/Elements/PCGTextureSampler.cpp
@@ -52,6 +52,15 @@ void UPCGTextureSamplerSettings::PostLoad()
 {
 	Super::PostLoad();
 
+#if WITH_EDITOR
+	PRAGMA_DISABLE_DEPRECATION_WARNINGS
+	if (DensityFunction != EPCGTextureDensityFunction::Multiply)
+	{
+		bUseDensitySourceChannel = false;
+	}
+	PRAGMA_ENABLE_DEPRECATION_WARNINGS
+#endif
+
 	UpdateDisplayTextureArrayIndex();
 }
 #endif
@@ -78,6 +87,18 @@ FPCGElementPtr UPCGTextureSamplerSettings::CreateElement() const
 	return MakeShared<FPCGTextureSamplerElement>();
 }
 
+PRAGMA_DISABLE_DEPRECATION_WARNINGS
+EPCGTextureDensityFunction UPCGTextureSamplerSettings::GetDensityFunctionEquivalent() const
+{
+	return bUseDensitySourceChannel ? EPCGTextureDensityFunction::Multiply : EPCGTextureDensityFunction::Ignore;
+}
+
+void UPCGTextureSamplerSettings::SetDensityFunctionEquivalent(EPCGTextureDensityFunction InDensityFunction)
+{
+	bUseDensitySourceChannel = (InDensityFunction != EPCGTextureDensityFunction::Ignore);
+}
+PRAGMA_ENABLE_DEPRECATION_WARNINGS
+
 #if WITH_EDITOR
 void UPCGTextureSamplerSettings::UpdateDisplayTextureArrayIndex()
 {
@@ -170,7 +191,7 @@ bool FPCGTextureSamplerElement::ExecuteInternal(FPCGContext* InContext) const
 
 	const FTransform& Transform = Settings->Transform;
 	const bool bUseAbsoluteTransform = Settings->bUseAbsoluteTransform;
-	const EPCGTextureDensityFunction DensityFunction = Settings->DensityFunction;
+	const bool bUseDensitySourceChannel = Settings->bUseDensitySourceChannel;
 	const EPCGTextureColorChannel ColorChannel = Settings->ColorChannel;
 	const EPCGTextureFilter Filter = Settings->Filter;
 	const float TexelSize = Settings->TexelSize;
@@ -220,7 +241,7 @@ bool FPCGTextureSamplerElement::ExecuteInternal(FPCGContext* InContext) const
 
 	TextureData->Initialize(Texture, TextureArrayIndex, FinalTransform, PostInitializeCallback, bForceEditorOnlyCPUSampling);
 
-	TextureData->DensityFunction = DensityFunction;
+	TextureData->bUseDensitySourceChannel = bUseDensitySourceChannel;
 	TextureData->ColorChannel = ColorChannel;
 	TextureData->Filter = Filter;
 	TextureData->TexelSize = TexelSize;
diff --git a/Engine/Plugins/PCG/Source/PCG/Private/Helpers/PCGHelpers.cpp b/Engine/Plugins/PCG/Source/PCG/Private/Helpers/PCGHelpers.cpp
--- a/Engine/Plugins/PCG/Source/PCG/Private/Helpers/PCGHelpers.cpp
+++ b/Engine/Plugins/PCG/Source/PCG/Private/Helpers/PCGHelpers.cpp
@@ -572,6 +572,22 @@ namespace PCGHelpers
 
 		return Functions;
 	}
+
+	TFunction<float(float, float)> GetDensityMergeFunction(EPCGDensityMergeOperation InOperation)
+	{
+		switch (InOperation)
+		{
+		case EPCGDensityMergeOperation::Set: return [](float A, float B) { return B; };
+		case EPCGDensityMergeOperation::Ignore: return [](float A, float B) { return A; };
+		case EPCGDensityMergeOperation::Minimum: return [](float A, float B) { return FMath::Min(A, B); };
+		case EPCGDensityMergeOperation::Maximum: return [](float A, float B) { return FMath::Max(A, B); };
+		case EPCGDensityMergeOperation::Add: return [](float A, float B) { return A + B; };
+		case EPCGDensityMergeOperation::Subtract: return [](float A, float B) { return A - B; };
+		case EPCGDensityMergeOperation::Multiply: return [](float A, float B) { return A * B; };
+		case EPCGDensityMergeOperation::Divide: return [](float A, float B) { return B != 0.0f ? (A / B) : 0.0f; };
+		default: checkNoEntry(); return [](float, float) { return 0.0f; };
+		}
+	}
 }
 
 #undef LOCTEXT_NAMESPACE
diff --git a/Engine/Plugins/PCG/Source/PCG/Public/Data/PCGTextureData.h b/Engine/Plugins/PCG/Source/PCG/Public/Data/PCGTextureData.h
--- a/Engine/Plugins/PCG/Source/PCG/Public/Data/PCGTextureData.h
+++ b/Engine/Plugins/PCG/Source/PCG/Public/Data/PCGTextureData.h
@@ -20,7 +20,7 @@ enum class EPCGTextureColorChannel : uint8
 };
 
 UENUM(BlueprintType)
-enum class EPCGTextureDensityFunction : uint8
+enum class UE_DEPRECATED(5.5, "EPCGTextureDensityFunction has been deprecated.") EPCGTextureDensityFunction : uint8
 {
 	Ignore,
 	Multiply
@@ -56,6 +56,10 @@ class PCG_API UPCGBaseTextureData : public UPCGSurfaceData
 	GENERATED_BODY()
 
 public:
+	// ~Being UObject interface
+	virtual void PostLoad() override;
+	// ~End UObject interface
+
 	// ~Begin UPCGData interface
 	virtual EPCGDataType GetDataType() const override { return EPCGDataType::BaseTexture; }
 	// ~End UPCGData interface
@@ -75,11 +79,28 @@ public:
 
 	virtual bool IsValid() const;
 
+protected:
+	PRAGMA_DISABLE_DEPRECATION_WARNINGS
+	UFUNCTION(BlueprintGetter, meta = (BlueprintInternalUseOnly = "true"))
+	EPCGTextureDensityFunction GetDensityFunctionEquivalent() const;
+
+	UFUNCTION(BlueprintSetter, meta = (BlueprintInternalUseOnly = "true"))
+	void SetDensityFunctionEquivalent(EPCGTextureDensityFunction DensityFunction);
+	PRAGMA_ENABLE_DEPRECATION_WARNINGS
+
 public:
-	UPROPERTY(BlueprintReadWrite, VisibleAnywhere, Category = SpatialData)
+#if WITH_EDITORONLY_DATA
+	UE_DEPRECATED(5.5, "DensityFunction has been deprecated in favor of bUseDensitySourceChannel.")
+	PRAGMA_DISABLE_DEPRECATION_WARNINGS
+	UPROPERTY(BlueprintGetter = GetDensityFunctionEquivalent, BlueprintSetter = SetDensityFunctionEquivalent, Category = SpatialData, meta = (DeprecatedProperty, DeprecatedMessage = "Density function on GetTextureData is deprecated in favor of bUseDensitySourceChannel."))
 	EPCGTextureDensityFunction DensityFunction = EPCGTextureDensityFunction::Multiply; 
+	PRAGMA_ENABLE_DEPRECATION_WARNINGS
+#endif
 
-	UPROPERTY(BlueprintReadWrite, EditAnywhere, Category = Settings)
+	UPROPERTY(BlueprintReadWrite, EditAnywhere, Category = Settings, meta = (InlineEditConditionToggle))
+	bool bUseDensitySourceChannel = false;
+
+	UPROPERTY(BlueprintReadWrite, EditAnywhere, Category = Settings, meta = (DisplayName = "Density Source Channel", EditCondition = "bUseDensitySourceChannel"))
 	EPCGTextureColorChannel ColorChannel = EPCGTextureColorChannel::Alpha;
 
 	/** Method used to determine the value for a sample based on the value of nearby texels. */
diff --git a/Engine/Plugins/PCG/Source/PCG/Public/Elements/PCGSampleTexture.h b/Engine/Plugins/PCG/Source/PCG/Public/Elements/PCGSampleTexture.h
--- a/Engine/Plugins/PCG/Source/PCG/Public/Elements/PCGSampleTexture.h
+++ b/Engine/Plugins/PCG/Source/PCG/Public/Elements/PCGSampleTexture.h
@@ -53,6 +53,14 @@ public:
 	/** Overrides the texture's tiling to wrap or clamp its UVs. */
 	UPROPERTY(BlueprintReadWrite, EditAnywhere, Category = Settings, meta = (PCG_Overridable, EditCondition = "TextureMappingMethod == EPCGTextureMappingMethod::UVCoordinates", EditConditionHides))
 	EPCGTextureAddressMode TilingMode = EPCGTextureAddressMode::Wrap;
+
+	/** Controls the behavior of density computation with respect to initial data. */
+	UPROPERTY(BlueprintReadWrite, EditAnywhere, Category = Settings, meta = (PCG_Overridable))
+	EPCGDensityMergeOperation DensityMergeFunction = EPCGDensityMergeOperation::Set;
+
+	/** Controls whether the output density should be clamped or not. */
+	UPROPERTY(BlueprintReadWrite, EditAnywhere, Category = Settings, meta = (PCG_Overridable))
+	bool bClampOutputDensity = true;
 };
 
 class FPCGSampleTextureElement : public IPCGElement
diff --git a/Engine/Plugins/PCG/Source/PCG/Public/Elements/PCGTextureSampler.h b/Engine/Plugins/PCG/Source/PCG/Public/Elements/PCGTextureSampler.h
--- a/Engine/Plugins/PCG/Source/PCG/Public/Elements/PCGTextureSampler.h
+++ b/Engine/Plugins/PCG/Source/PCG/Public/Elements/PCGTextureSampler.h
@@ -44,6 +44,14 @@ protected:
 	void UpdateDisplayTextureArrayIndex();
 #endif
 
+	PRAGMA_DISABLE_DEPRECATION_WARNINGS
+	UFUNCTION(BlueprintGetter, meta = (BlueprintInternalUseOnly = "true"))
+	EPCGTextureDensityFunction GetDensityFunctionEquivalent() const;
+
+	UFUNCTION(BlueprintSetter, meta = (BlueprintInternalUseOnly = "true"))
+	void SetDensityFunctionEquivalent(EPCGTextureDensityFunction DensityFunction);
+	PRAGMA_ENABLE_DEPRECATION_WARNINGS
+
 public:
 	PCG_API void SetTexture(TSoftObjectPtr<UTexture> InTexture);
 	TSoftObjectPtr<UTexture> GetTexture() const { return Texture; }
@@ -59,10 +67,18 @@ public:
 	UPROPERTY(BlueprintReadWrite, EditAnywhere, Category = Settings, meta = (EditCondition = bDisplayTextureArrayIndex, EditConditionHides, HideEditConditionToggle, ClampMin = '0', PCG_Overridable))
 	int TextureArrayIndex = 0;
 
-	UPROPERTY(BlueprintReadWrite, EditAnywhere, Category = SpatialData, meta = (PCG_Overridable))
+#if WITH_EDITORONLY_DATA
+	UE_DEPRECATED(5.5, "DensityFunction has been deprecated in favor of bUseDensitySourceChannel.")
+	PRAGMA_DISABLE_DEPRECATION_WARNINGS
+	UPROPERTY(BlueprintGetter=GetDensityFunctionEquivalent, BlueprintSetter=SetDensityFunctionEquivalent, Category = SpatialData, meta = (DeprecatedProperty, DeprecatedMessage = "Density function on GetTextureData is deprecated in favor of bUseDensitySourceChannel."))
 	EPCGTextureDensityFunction DensityFunction = EPCGTextureDensityFunction::Multiply;
+	PRAGMA_ENABLE_DEPRECATION_WARNINGS
+#endif
 
-	UPROPERTY(BlueprintReadWrite, EditAnywhere, Category = Settings, meta = (PCG_Overridable))
+	UPROPERTY(BlueprintReadWrite, EditAnywhere, Category = SpatialData, meta = (InlineEditConditionToggle, PCG_Overridable, PCG_OverrideAliases="DensityFunction"))
+	bool bUseDensitySourceChannel = true;
+
+	UPROPERTY(BlueprintReadWrite, EditAnywhere, Category = Settings, meta = (EditCondition="bUseDensitySourceChannel", DisplayName = "Density Source Channel", PCG_Overridable))
 	EPCGTextureColorChannel ColorChannel = EPCGTextureColorChannel::Alpha;
 
 	/** Method used to determine the value for a sample based on the value of nearby texels. */
diff --git a/Engine/Plugins/PCG/Source/PCG/Public/Helpers/PCGHelpers.h b/Engine/Plugins/PCG/Source/PCG/Public/Helpers/PCGHelpers.h
--- a/Engine/Plugins/PCG/Source/PCG/Public/Helpers/PCGHelpers.h
+++ b/Engine/Plugins/PCG/Source/PCG/Public/Helpers/PCGHelpers.h
@@ -85,6 +85,8 @@ namespace PCGHelpers
 	 * and have parameters matching one of the provided prototypes. Some prototypes are provided in UPCGFunctionPrototypes.
 	 */
 	PCG_API TArray<UFunction*> FindUserFunctions(TSubclassOf<AActor> ActorClass, const TArray<FName>& FunctionNames, const TArray<const UFunction*>& FunctionPrototypes, const FPCGContext* InContext = nullptr);
+
+	PCG_API TFunction<float(float, float)> GetDensityMergeFunction(EPCGDensityMergeOperation InOperation);
 };
 
 /** Holds function prototypes used to match against actor function signatures. */
diff --git a/Engine/Plugins/PCG/Source/PCG/Public/PCGCommon.h b/Engine/Plugins/PCG/Source/PCG/Public/PCGCommon.h
--- a/Engine/Plugins/PCG/Source/PCG/Public/PCGCommon.h
+++ b/Engine/Plugins/PCG/Source/PCG/Public/PCGCommon.h
@@ -379,3 +379,24 @@ struct UE_DEPRECATED(5.5, "FPCGPartitionActorRecord is deprecated.") FPCGPartiti
 	bool operator==(const FPCGPartitionActorRecord& InOther) const;
 	friend uint32 GetTypeHash(const FPCGPartitionActorRecord& In);
 };
+
+UENUM(BlueprintType)
+enum class EPCGDensityMergeOperation : uint8
+{
+	/** D = B */
+	Set,
+	/** D = A */
+	Ignore,
+	/** D = min(A, B) */
+	Minimum,
+	/** D = max(A, B) */
+	Maximum,
+	/** D = A + B */
+	Add,
+	/** D = A - B */
+	Subtract,
+	/** D = A * B */
+	Multiply,
+	/** D = A / B */
+	Divide
+};
\ No newline at end of file
