Flipped defaults around the pause button & ctrl+alt behaviors.
Improved behavior on the generate when PCG is paused too.

diff --git a/Engine/Plugins/PCG/Source/PCG/Private/PCGComponent.cpp b/Engine/Plugins/PCG/Source/PCG/Private/PCGComponent.cpp
--- a/Engine/Plugins/PCG/Source/PCG/Private/PCGComponent.cpp
+++ b/Engine/Plugins/PCG/Source/PCG/Private/PCGComponent.cpp
@@ -470,6 +470,15 @@ FPCGTaskId UPCGComponent::CreateGenerateTask(bool bForce, const TArray<FPCGTaskI
 	TArray<FPCGTaskId> AdditionalDependencies;
 	const TArray<FPCGTaskId>* AllDependencies = &Dependencies;
 
+#if WITH_EDITOR
+	// If PCG is currently paused, we don't want to do immediate cleanup here otherwise it looks wrong and can't be cancelled.
+	// In this case, we'll "pre-queue" a cleanup before so it behaves as-if we had done this exactly.
+	if (bGenerated && PCGSystemSwitches::CVarPausePCGExecution.GetValueOnAnyThread())
+	{
+		CleanupInternal(/*bRemoveComponents=*/false, Dependencies);
+	}
+#endif
+
 	if (IsCleaningUp())
 	{
 		AdditionalDependencies = Dependencies;
diff --git a/Engine/Plugins/PCG/Source/PCGEditor/Private/PCGEditorModule.cpp b/Engine/Plugins/PCG/Source/PCGEditor/Private/PCGEditorModule.cpp
--- a/Engine/Plugins/PCG/Source/PCGEditor/Private/PCGEditorModule.cpp
+++ b/Engine/Plugins/PCG/Source/PCGEditor/Private/PCGEditorModule.cpp
@@ -319,15 +319,15 @@ void FPCGEditorModule::RegisterMenuExtensions()
 
 							if (bWasPaused)
 							{
-								bool bShouldCancelAll = true;
+								bool bShouldCancelAll = false;
 
 								if (FSlateApplication::Get().GetModifierKeys().IsControlDown())
 								{
-									bShouldCancelAll = false;
+									bShouldCancelAll = true;
 								}
 								else if (FSlateApplication::Get().GetModifierKeys().IsAltDown())
 								{
-									bShouldCancelAll = true;
+									bShouldCancelAll = false;
 								}
 								else if (const UPCGEditorSettings* EditorSettings = GetDefault<UPCGEditorSettings>())
 								{
@@ -362,7 +362,7 @@ void FPCGEditorModule::RegisterMenuExtensions()
 							return LOCTEXT("PCGPauseButton_On", "PCG");
 						}
 					}),
-					LOCTEXT("PCGPauseButton_Tooltip", "Toggles PCG processing on/off and will cancel tasks depending on settings.\nUse Ctrl to unpause without cancelling all tasks.\nUse Alt to unpause and cancel all tasks."),
+					LOCTEXT("PCGPauseButton_Tooltip", "Toggles PCG processing on/off and will cancel tasks depending on settings.\nUse Ctrl to unpause and cancel all tasks.\nUse Alt to unpause without cancelling tasks."),
 					FSlateIcon(FPCGEditorStyle::Get().GetStyleSetName(), "PCG.EditorIcon"),
 					EUserInterfaceActionType::ToggleButton
 				);
diff --git a/Engine/Plugins/PCG/Source/PCGEditor/Public/PCGEditorSettings.h b/Engine/Plugins/PCG/Source/PCGEditor/Public/PCGEditorSettings.h
--- a/Engine/Plugins/PCG/Source/PCGEditor/Public/PCGEditorSettings.h
+++ b/Engine/Plugins/PCG/Source/PCGEditor/Public/PCGEditorSettings.h
@@ -177,7 +177,7 @@ public:
 	bool bShowPauseButton = false;
 
 	UPROPERTY(EditAnywhere, config, Category = "Editor Performance", meta = (DisplayName = "Unpausing cancels all PCG tasks"))
-	bool bUnpauseCancelsAll = true;
+	bool bUnpauseCancelsAll = false;
 };
 
 #if UE_ENABLE_INCLUDE_ORDER_DEPRECATED_IN_5_2
