Added managed component list resource to tie lifetime of multiple components together.
Refactored some code to prevent code some duplication
Added "default" managed components list resource so that BP-created components automatically get attached to the default list.

diff --git a/Engine/Plugins/PCG/Source/PCG/Private/PCGComponent.cpp b/Engine/Plugins/PCG/Source/PCG/Private/PCGComponent.cpp
--- a/Engine/Plugins/PCG/Source/PCG/Private/PCGComponent.cpp
+++ b/Engine/Plugins/PCG/Source/PCG/Private/PCGComponent.cpp
@@ -218,25 +218,48 @@ void UPCGComponent::AddToManagedResources(UPCGManagedResource* InResource)
 
 void UPCGComponent::AddComponentsToManagedResources(const TArray<UActorComponent*>& InComponents)
 {
+	if (InComponents.IsEmpty())
+	{
+		return;
+	}
+
 	for (UActorComponent* Component : InComponents)
 	{
-		if (!Component)
+		if (Component && !Component->ComponentHasTag(PCGHelpers::DefaultPCGTag))
 		{
-			continue;
+			Component->Modify();
+			Component->ComponentTags.Add(PCGHelpers::DefaultPCGTag);
 		}
+	}
+
+	FScopeLock ResourcesLock(&GeneratedResourcesLock);
+	check(!GeneratedResourcesInaccessible);
 
-		if (!Component->ComponentHasTag(PCGHelpers::DefaultPCGTag))
+	UPCGManagedComponentDefaultList* DefaultList = nullptr;
+	for (const TObjectPtr<UPCGManagedResource>& ManagedResource : GeneratedResources)
+	{
+		if (UPCGManagedComponentDefaultList* ExistingList = Cast<UPCGManagedComponentDefaultList>(ManagedResource))
 		{
-			Component->Modify();
-			Component->ComponentTags.Add(PCGHelpers::DefaultPCGTag);
+			DefaultList = ExistingList;
+			break;
 		}
+	}
 
-		UPCGManagedComponent* ManagedResource = NewObject<UPCGManagedComponent>(this);
-		// Implementation note: we call the setter to make sure that if this is done from BP, the construction method is properly updated
-		ManagedResource->SetGeneratedComponentFromBP(Component);
+	if (!DefaultList)
+	{
+		DefaultList = NewObject<UPCGManagedComponentDefaultList>(this);
 
-		AddToManagedResources(ManagedResource);
+		PCGGeneratedResourcesLogging::LogAddToManagedResources(this, DefaultList);
+		GeneratedResources.Add(DefaultList);
 	}
+
+	check(DefaultList);
+	// Implementation note: we call the AddGeneratedComponentsFromBP method to make sure that if this is done from BP, the construction method is properly updated
+
+	TArray<TSoftObjectPtr<UActorComponent>> Components;
+	Algo::Transform(InComponents, Components, [](UActorComponent* Component) { return TSoftObjectPtr<UActorComponent>(Component); });
+
+	DefaultList->AddGeneratedComponentsFromBP(Components);
 }
 
 void UPCGComponent::AddActorsToManagedResources(const TArray<AActor*>& InActors)
diff --git a/Engine/Plugins/PCG/Source/PCG/Private/PCGManagedResource.cpp b/Engine/Plugins/PCG/Source/PCG/Private/PCGManagedResource.cpp
--- a/Engine/Plugins/PCG/Source/PCG/Private/PCGManagedResource.cpp
+++ b/Engine/Plugins/PCG/Source/PCG/Private/PCGManagedResource.cpp
@@ -278,7 +278,7 @@ void UPCGManagedActors::ChangeTransientState(EPCGEditorDirtyMode NewEditingMode)
 }
 #endif // WITH_EDITOR
 
-void UPCGManagedComponent::PostEditImport()
+void UPCGManagedComponentBase::PostEditImport()
 {
 	Super::PostEditImport();
 
@@ -286,38 +286,61 @@ void UPCGManagedComponent::PostEditImport()
 	UPCGComponent* OwningComponent = Cast<UPCGComponent>(GetOuter());
 	AActor* Actor = OwningComponent ? OwningComponent->GetOwner() : nullptr;
 
-	bool bFoundMatch = false;
-
-	if (Actor && GeneratedComponent.IsValid())
+	if (!Actor)
+	{
+		// Somewhat irrelevant case, if we don't have an actor or a component, there's not a lot we can do.
+		ForgetComponents();
+	}
+	else if(GetComponentsCount() > 0)
 	{
-		TInlineComponentArray<UActorComponent*, 16> Components;
-		Actor->GetComponents(Components);
+		TInlineComponentArray<UActorComponent*, 64> ActorComponents;
+		Actor->GetComponents(ActorComponents);
 
-		for (UActorComponent* Component : Components)
+		TArrayView<TSoftObjectPtr<UActorComponent>> GeneratedComponents = GetComponentsArray();
+		for (int ComponentIndex = GeneratedComponents.Num() - 1; ComponentIndex >= 0; --ComponentIndex)
 		{
-			if (Component && Component->GetFName() == GeneratedComponent->GetFName())
+			TSoftObjectPtr<UActorComponent> GeneratedComponent = GeneratedComponents[ComponentIndex];
+
+			UActorComponent* MatchingComponent = nullptr;
+
+			if (GeneratedComponent.IsValid())
 			{
-				GeneratedComponent = Component;
-				bFoundMatch = true;
-				break;
+				if (UActorComponent** FoundMatchingComponent = ActorComponents.FindByPredicate([&GeneratedComponent](UActorComponent* ActorComponent) { return ActorComponent && ActorComponent->GetFName() == GeneratedComponent->GetFName(); }))
+				{
+					MatchingComponent = *FoundMatchingComponent;
+				}
 			}
-		}
 
-		if (!bFoundMatch)
-		{
-			// Not quite clear what to do when we have a component that cannot be remapped.
-			// Maybe we should check against guids instead?
-			ForgetComponent();
+			if (MatchingComponent)
+			{
+				GeneratedComponents[ComponentIndex] = MatchingComponent;
+			}
+			else
+			{
+				// Not quite clear what to do when we have a component that cannot be remapped.
+				// Maybe we should check against guids instead?
+				ForgetComponent(ComponentIndex);
+			}
 		}
 	}
-	else
+}
+
+#if WITH_EDITOR
+void UPCGManagedComponentBase::HideComponents()
+{
+	const int32 ComponentCount = GetComponentsCount();
+	for (int32 ComponentIndex = 0; ComponentIndex < ComponentCount; ++ComponentIndex)
 	{
-		// Somewhat irrelevant case, if we don't have an actor or a component, there's not a lot we can do.
-		ForgetComponent();
+		HideComponent(ComponentIndex);
 	}
 }
 
-#if WITH_EDITOR
+void UPCGManagedComponentBase::HideComponent(int32 ComponentIndex)
+{
+	// Default implementation to be backward compatible
+	HideComponent();
+}
+
 void UPCGManagedComponent::HideComponent()
 {
 	if (GeneratedComponent.IsValid())
@@ -325,8 +348,51 @@ void UPCGManagedComponent::HideComponent()
 		GeneratedComponent->UnregisterComponent();
 	}
 }
+
+void UPCGManagedComponentList::HideComponent(int32 ComponentIndex)
+{
+	if (GeneratedComponents[ComponentIndex].IsValid())
+	{
+		GeneratedComponents[ComponentIndex]->UnregisterComponent();
+	}
+}
 #endif // WITH_EDITOR
 
+void UPCGManagedComponentBase::ForgetComponents()
+{
+	const int32 ComponentCount = GetComponentsCount();
+	for (int32 ComponentIndex = ComponentCount - 1; ComponentIndex >= 0; --ComponentIndex)
+	{
+		ForgetComponent(ComponentIndex);
+	}
+}
+
+void UPCGManagedComponentBase::ForgetComponent(int32 ComponentIndex)
+{
+	// Default implementation to be backward compatible
+	ForgetComponent();
+}
+
+void UPCGManagedComponentList::ForgetComponent(int32 ComponentIndex)
+{
+	GeneratedComponents.RemoveAtSwap(ComponentIndex);
+}
+
+void UPCGManagedComponentBase::ResetComponents()
+{
+	const int32 ComponentCount = GetComponentsCount();
+	for (int32 ComponentIndex = 0; ComponentIndex < ComponentCount; ++ComponentIndex)
+	{
+		ResetComponent(ComponentIndex);
+	}
+}
+
+void UPCGManagedComponentBase::ResetComponent(int32 ComponentIndex)
+{
+	// Default implementation to be backward compatible
+	ResetComponent();
+}
+
 bool UPCGManagedComponent::Release(bool bHardRelease, TSet<TSoftObjectPtr<AActor>>& /*OutActorsToDelete*/)
 {
 	TRACE_CPUPROFILER_EVENT_SCOPE(UPCGManagedComponent::Release);
@@ -375,64 +441,149 @@ bool UPCGManagedComponent::Release(bool bHardRelease, TSet<TSoftObjectPtr<AActor
 	return bDeleteComponent;
 }
 
-bool UPCGManagedComponent::ReleaseIfUnused(TSet<TSoftObjectPtr<AActor>>& OutActorsToDelete)
+bool UPCGManagedComponentList::Release(bool bHardRelease, TSet<TSoftObjectPtr<AActor>>& /*OutActorsToDelete*/)
 {
-	return Super::ReleaseIfUnused(OutActorsToDelete) || !GeneratedComponent.IsValid();
-}
+	TRACE_CPUPROFILER_EVENT_SCOPE(UPCGManagedComponentList::Release);
 
-bool UPCGManagedComponent::MoveResourceToNewActor(AActor* NewActor, const AActor* ExpectedPreviousOwner)
-{
-	check(NewActor);
+	const bool bSupportsComponentReset = SupportsComponentReset();
+	bool bDeleteComponent = bHardRelease || !bSupportsComponentReset;
 
-	if (!GeneratedComponent.IsValid())
+	// Start by removing all dead components from the array
+	for (int32 ComponentIndex = GeneratedComponents.Num() - 1; ComponentIndex >= 0; --ComponentIndex)
 	{
-		return false;
+		if (!GeneratedComponents[ComponentIndex].IsValid())
+		{
+			GeneratedComponents.RemoveAtSwap(ComponentIndex);
+		}
 	}
 
-	TObjectPtr<AActor> OldOwner = GeneratedComponent->GetOwner();
-	check(OldOwner);
+	// Nothing left - this resource can be released
+	if (GeneratedComponents.IsEmpty())
+	{
+		return true;
+	}
 
-	// Prevent moving of components on external (or spawned) actors
-	if (ExpectedPreviousOwner && OldOwner != ExpectedPreviousOwner)
+#if WITH_EDITOR
+	if (bMarkedTransientOnLoad)
 	{
-		return false;
+		PCGGeneratedResourcesLogging::LogManagedComponentHidden(this);
+		HideComponents();
+		bIsMarkedUnused = true;
 	}
+	else
+#endif
+	{
+		if (bDeleteComponent)
+		{
+			PCGGeneratedResourcesLogging::LogManagedResourceHardRelease(this);
+			for (TSoftObjectPtr<UActorComponent> GeneratedComponent : GeneratedComponents)
+			{
+				GeneratedComponent->DestroyComponent();
+			}
 
-	bool bDetached = false;
-	bool bAttached = false;
+			ForgetComponents();
+		}
+		else
+		{
+			PCGGeneratedResourcesLogging::LogManagedResourceSoftRelease(this);
+			bIsMarkedUnused = true;
 
-	GeneratedComponent->UnregisterComponent();
+			for (TSoftObjectPtr<UActorComponent> GeneratedComponent : GeneratedComponents)
+			{
+				GeneratedComponent->ComponentTags.Add(PCGHelpers::MarkedForCleanupPCGTag);
+			}
+		}
+	}
 
-	// Need to change owner first to avoid that the PCG Component will react to this component changes.
-	GeneratedComponent->Rename(nullptr, NewActor);
+	return bDeleteComponent;
+}
 
-	// Check if it is a scene component, and if so, use its method to attach/detach to root component
-	if (TObjectPtr<USceneComponent> GeneratedSceneComponent = Cast<USceneComponent>(GeneratedComponent.Get()))
+bool UPCGManagedComponent::ReleaseIfUnused(TSet<TSoftObjectPtr<AActor>>& OutActorsToDelete)
+{
+	return Super::ReleaseIfUnused(OutActorsToDelete) || !GeneratedComponent.IsValid();
+}
+
+bool UPCGManagedComponentList::ReleaseIfUnused(TSet<TSoftObjectPtr<AActor>>& OutActorsToDelete)
+{
+	if (Super::ReleaseIfUnused(OutActorsToDelete))
 	{
-		GeneratedSceneComponent->DetachFromComponent(FDetachmentTransformRules::KeepWorldTransform);
-		bDetached = true;
-		bAttached = GeneratedSceneComponent->AttachToComponent(NewActor->GetRootComponent(), FAttachmentTransformRules::KeepWorldTransform);
+		return true;
 	}
 
-	// Otherwise use the default one.
-	if (!bAttached)
+	// Start by removing all dead components from the array
+	for (int32 ComponentIndex = GeneratedComponents.Num() - 1; ComponentIndex >= 0; --ComponentIndex)
 	{
-		if (!bDetached)
+		if (!GeneratedComponents[ComponentIndex].IsValid())
 		{
-			OldOwner->RemoveInstanceComponent(GeneratedComponent.Get());
+			GeneratedComponents.RemoveAtSwap(ComponentIndex);
 		}
-
-		NewActor->AddInstanceComponent(GeneratedComponent.Get());
 	}
 
-	GeneratedComponent->RegisterComponent();
+	// Nothing left - this resource can be released
+	return GeneratedComponents.IsEmpty();
+}
 
-	ForgetComponent();
+bool UPCGManagedComponentBase::MoveResourceToNewActor(AActor* NewActor, const AActor* ExpectedPreviousOwner)
+{
+	check(NewActor);
 
-	return true;
+	bool bMovedResources = false;
+
+	TArrayView<TSoftObjectPtr<UActorComponent>> GeneratedComponents = GetComponentsArray();
+	for (int32 ComponentIndex = GeneratedComponents.Num() - 1; ComponentIndex >= 0; --ComponentIndex)
+	{
+		TSoftObjectPtr<UActorComponent> GeneratedComponent = GeneratedComponents[ComponentIndex];
+
+		if (!GeneratedComponent.IsValid())
+		{
+			continue;
+		}
+
+		TObjectPtr<AActor> OldOwner = GeneratedComponent->GetOwner();
+		check(OldOwner);
+
+		// Prevent moving of components on external (or spawned) actors
+		if (ExpectedPreviousOwner && OldOwner != ExpectedPreviousOwner)
+		{
+			continue;
+		}
+
+		bool bDetached = false;
+		bool bAttached = false;
+
+		GeneratedComponent->UnregisterComponent();
+
+		// Need to change owner first to avoid that the PCG Component will react to this component changes.
+		GeneratedComponent->Rename(nullptr, NewActor);
+
+		// Check if it is a scene component, and if so, use its method to attach/detach to root component
+		if (TObjectPtr<USceneComponent> GeneratedSceneComponent = Cast<USceneComponent>(GeneratedComponent.Get()))
+		{
+			GeneratedSceneComponent->DetachFromComponent(FDetachmentTransformRules::KeepWorldTransform);
+			bDetached = true;
+			bAttached = GeneratedSceneComponent->AttachToComponent(NewActor->GetRootComponent(), FAttachmentTransformRules::KeepWorldTransform);
+		}
+
+		// Otherwise use the default one.
+		if (!bAttached)
+		{
+			if (!bDetached)
+			{
+				OldOwner->RemoveInstanceComponent(GeneratedComponent.Get());
+			}
+
+			NewActor->AddInstanceComponent(GeneratedComponent.Get());
+		}
+
+		GeneratedComponent->RegisterComponent();
+		ForgetComponent(ComponentIndex);
+		bMovedResources = true;
+	}
+
+	return bMovedResources;
 }
 
-void UPCGManagedComponent::MarkAsUsed()
+void UPCGManagedComponentBase::MarkAsUsed()
 {
 	if (!bIsMarkedUnused)
 	{
@@ -444,32 +595,38 @@ void UPCGManagedComponent::MarkAsUsed()
 	// Can't reuse a resource if we can't reset it. Make sure we never take this path in this case.
 	check(SupportsComponentReset());
 
-	ResetComponent();
+	ResetComponents();
 
-	if (GeneratedComponent.Get())
+	TArrayView<TSoftObjectPtr<UActorComponent>> GeneratedComponents = GetComponentsArray();
+	for (TSoftObjectPtr<UActorComponent> GeneratedComponent : GeneratedComponents)
 	{
-		GeneratedComponent->ComponentTags.Remove(PCGHelpers::MarkedForCleanupPCGTag);
+		if (GeneratedComponent.Get())
+		{
+			GeneratedComponent->ComponentTags.Remove(PCGHelpers::MarkedForCleanupPCGTag);
+		}
 	}
 }
 
-void UPCGManagedComponent::MarkAsReused()
+void UPCGManagedComponentBase::MarkAsReused()
 {
 	Super::MarkAsReused();
 
-	if (GeneratedComponent.Get())
+	TArrayView<TSoftObjectPtr<UActorComponent>> GeneratedComponents = GetComponentsArray();
+	for (TSoftObjectPtr<UActorComponent> GeneratedComponent : GeneratedComponents)
 	{
-		GeneratedComponent->ComponentTags.Remove(PCGHelpers::MarkedForCleanupPCGTag);
+		if (GeneratedComponent.Get())
+		{
+			GeneratedComponent->ComponentTags.Remove(PCGHelpers::MarkedForCleanupPCGTag);
+		}
 	}
 }
 
-void UPCGManagedComponent::SetGeneratedComponentFromBP(TSoftObjectPtr<UActorComponent> InGeneratedComponent)
+void UPCGManagedComponentBase::SetupGeneratedComponentFromBP(TSoftObjectPtr<UActorComponent> InGeneratedComponent)
 {
-	GeneratedComponent = InGeneratedComponent;
-
 	// Components that are created from blueprint are automatically tagged as "created by construction script",
 	// regardless of whether that is true. This makes sure that the flags on the component are correct and considered an instance component
 	// and will then be properly serialized and managed by PCG.
-	if (UActorComponent* Component = GeneratedComponent.Get())
+	if (UActorComponent* Component = InGeneratedComponent.Get())
 	{
 		if (AActor* ComponentOwner = Component->GetOwner())
 		{
@@ -483,38 +640,69 @@ void UPCGManagedComponent::SetGeneratedComponentFromBP(TSoftObjectPtr<UActorComp
 	}
 }
 
+void UPCGManagedComponent::SetGeneratedComponentFromBP(TSoftObjectPtr<UActorComponent> InGeneratedComponent)
+{
+	GeneratedComponent = InGeneratedComponent;
+	SetupGeneratedComponentFromBP(InGeneratedComponent);
+}
+
+void UPCGManagedComponentList::SetGeneratedComponentsFromBP(const TArray<TSoftObjectPtr<UActorComponent>>& InGeneratedComponents)
+{
+	GeneratedComponents = InGeneratedComponents;
+
+	for (TSoftObjectPtr<UActorComponent> GeneratedComponent : GeneratedComponents)
+	{
+		SetupGeneratedComponentFromBP(GeneratedComponent);
+	}
+}
+
+void UPCGManagedComponentDefaultList::AddGeneratedComponentsFromBP(const TArray<TSoftObjectPtr<UActorComponent>>& InGeneratedComponents)
+{
+	GeneratedComponents.Append(InGeneratedComponents);
+
+	for (TSoftObjectPtr<UActorComponent> GeneratedComponent : GeneratedComponents)
+	{
+		SetupGeneratedComponentFromBP(GeneratedComponent);
+	}
+}
+
 #if WITH_EDITOR
-void UPCGManagedComponent::ChangeTransientState(EPCGEditorDirtyMode NewEditingMode)
+void UPCGManagedComponentBase::ChangeTransientState(EPCGEditorDirtyMode NewEditingMode)
 {
 	const bool bNowTransient = (NewEditingMode == EPCGEditorDirtyMode::Preview);
-	if (GeneratedComponent.Get())
-	{
-		const bool bWasTransient = GeneratedComponent->HasAnyFlags(RF_Transient);
 
-		if (bWasTransient != bNowTransient)
+	TArrayView<TSoftObjectPtr<UActorComponent>> GeneratedComponents = GetComponentsArray();
+	for (TSoftObjectPtr<UActorComponent> GeneratedComponent : GeneratedComponents)
+	{
+		if (GeneratedComponent.Get())
 		{
-			if (bNowTransient)
-			{
-				GeneratedComponent->SetFlags(RF_Transient);
-			}
-			else
-			{
-				GeneratedComponent->ClearFlags(RF_Transient);
-			}
+			const bool bWasTransient = GeneratedComponent->HasAnyFlags(RF_Transient);
 
-			ForEachObjectWithOuter(GeneratedComponent.Get(), [bNowTransient](UObject* Object)
+			if (bWasTransient != bNowTransient)
 			{
 				if (bNowTransient)
 				{
-					Object->SetFlags(RF_Transient);
+					GeneratedComponent->SetFlags(RF_Transient);
 				}
 				else
 				{
-					Object->ClearFlags(RF_Transient);
+					GeneratedComponent->ClearFlags(RF_Transient);
 				}
-			});
 
-			GeneratedComponent->MarkPackageDirty(); // should dirty actor this component is attached to
+				ForEachObjectWithOuter(GeneratedComponent.Get(), [bNowTransient](UObject* Object)
+				{
+					if (bNowTransient)
+					{
+						Object->SetFlags(RF_Transient);
+					}
+					else
+					{
+						Object->ClearFlags(RF_Transient);
+					}
+				});
+
+				GeneratedComponent->MarkPackageDirty(); // should dirty actor this component is attached to
+			}
 		}
 	}
 
diff --git a/Engine/Plugins/PCG/Source/PCG/Private/Utils/PCGGeneratedResourcesLogging.cpp b/Engine/Plugins/PCG/Source/PCG/Private/Utils/PCGGeneratedResourcesLogging.cpp
--- a/Engine/Plugins/PCG/Source/PCG/Private/Utils/PCGGeneratedResourcesLogging.cpp
+++ b/Engine/Plugins/PCG/Source/PCG/Private/Utils/PCGGeneratedResourcesLogging.cpp
@@ -387,7 +387,7 @@ namespace PCGGeneratedResourcesLogging
 #endif
 	}
 
-	void LogManagedComponentHidden(UPCGManagedComponent* InResource)
+	void LogManagedComponentHidden(UPCGManagedResource* InResource)
 	{
 #if WITH_EDITOR
 		if (!LogEnabled())
@@ -405,7 +405,7 @@ namespace PCGGeneratedResourcesLogging
 #endif
 	}
 
-	void LogManagedComponentDeleteNull(UPCGManagedComponent* InResource)
+	void LogManagedComponentDeleteNull(UPCGManagedResource* InResource)
 	{
 #if WITH_EDITOR
 		if (!LogEnabled())
diff --git a/Engine/Plugins/PCG/Source/PCG/Private/Utils/PCGGeneratedResourcesLogging.h b/Engine/Plugins/PCG/Source/PCG/Private/Utils/PCGGeneratedResourcesLogging.h
--- a/Engine/Plugins/PCG/Source/PCG/Private/Utils/PCGGeneratedResourcesLogging.h
+++ b/Engine/Plugins/PCG/Source/PCG/Private/Utils/PCGGeneratedResourcesLogging.h
@@ -33,6 +33,6 @@ namespace PCGGeneratedResourcesLogging
 
 	void LogManagedResourceSoftRelease(UPCGManagedResource* InResource);
 	void LogManagedResourceHardRelease(UPCGManagedResource* InResource);
-	void LogManagedComponentHidden(UPCGManagedComponent* InResource);
-	void LogManagedComponentDeleteNull(UPCGManagedComponent* InResource);
+	void LogManagedComponentHidden(UPCGManagedResource* InResource);
+	void LogManagedComponentDeleteNull(UPCGManagedResource* InResource);
 }
diff --git a/Engine/Plugins/PCG/Source/PCG/Public/PCGManagedResource.h b/Engine/Plugins/PCG/Source/PCG/Public/PCGManagedResource.h
--- a/Engine/Plugins/PCG/Source/PCG/Public/PCGManagedResource.h
+++ b/Engine/Plugins/PCG/Source/PCG/Public/PCGManagedResource.h
@@ -104,36 +104,76 @@ public:
 	TSet<TSoftObjectPtr<AActor>> GeneratedActors;
 };
 
-UCLASS(BlueprintType)
-class PCG_API UPCGManagedComponent : public UPCGManagedResource
+UCLASS(Abstract)
+class PCG_API UPCGManagedComponentBase : public UPCGManagedResource
 {
 	GENERATED_BODY()
 
 	friend class UPCGComponent;
-
 public:
 	//~Begin UObject interface
 	virtual void PostEditImport() override;
 	//~End UObject interface
 
 	//~Begin UPCGManagedResource interface
-	virtual bool Release(bool bHardRelease, TSet<TSoftObjectPtr<AActor>>& OutActorsToDelete) override;
-	virtual bool ReleaseIfUnused(TSet<TSoftObjectPtr<AActor>>& OutActorsToDelete) override;
+	virtual void MarkAsUsed() override;
+	virtual void MarkAsReused() override;
 	virtual bool MoveResourceToNewActor(AActor* NewActor, const AActor* ExpectedPreviousOwner) override;
 
 #if WITH_EDITOR
 	virtual void ChangeTransientState(EPCGEditorDirtyMode NewEditingMode) override;
-	/** Hides the content of the component in a transient way (such as unregistering) */
-	virtual void HideComponent();
 #endif
 	//~End UPCGManagedResource interface
 
+public:
+#if WITH_EDITOR
+	/** Hides the content of the component in a transient way (such as unregistering) */
+	void HideComponents();
+	virtual void HideComponent(int32 ComponentIndex);
+	virtual void HideComponent() {};
+#endif
+
+	void ForgetComponents();
+	virtual void ForgetComponent(int32 ComponentIndex);
+	virtual void ForgetComponent() {}
+
+	void ResetComponents();
+	virtual void ResetComponent(int32 ComponentIndex);
 	virtual void ResetComponent() { check(0); }
 	virtual bool SupportsComponentReset() const { return false; }
-	virtual void MarkAsUsed() override;
-	virtual void MarkAsReused() override;
+
+protected:
+	virtual TArrayView<TSoftObjectPtr<UActorComponent>> GetComponentsArray() PURE_VIRTUAL(UPCGManagedComponentBase::GetComponentsArray, return TArrayView<TSoftObjectPtr<UActorComponent>>{};);
+	virtual int32 GetComponentsCount() const PURE_VIRTUAL(UPCGManagedComponentBase::GetComponentsCount, return 0;);
+	void SetupGeneratedComponentFromBP(TSoftObjectPtr<UActorComponent> InGeneratedComponent);
+};
+
+UCLASS(BlueprintType)
+class PCG_API UPCGManagedComponent : public UPCGManagedComponentBase
+{
+	GENERATED_BODY()
+
+	friend class UPCGComponent;
+
+public:
+	//~Begin UPCGManagedResource interface
+	virtual bool Release(bool bHardRelease, TSet<TSoftObjectPtr<AActor>>& OutActorsToDelete) override;
+	virtual bool ReleaseIfUnused(TSet<TSoftObjectPtr<AActor>>& OutActorsToDelete) override;
+	//~End UPCGManagedResource interface
+
+	//~Begin UPCGManagedComponentBase interface
+#if WITH_EDITOR
+	/** Hides the content of the component in a transient way (such as unregistering) */
+	virtual void HideComponent();
+#endif
+	
 	virtual void ForgetComponent() { GeneratedComponent.Reset(); }
 
+protected:
+	virtual TArrayView<TSoftObjectPtr<UActorComponent>> GetComponentsArray() override { return TArrayView<TSoftObjectPtr<UActorComponent>>(&GeneratedComponent, 1); }
+	virtual int32 GetComponentsCount() const override { return 1; }
+	//~End UPCGManagedComponentBase interface
+
 public:
 	UPROPERTY(BlueprintReadWrite, EditAnywhere, Category = GeneratedData, BlueprintSetter = SetGeneratedComponentFromBP)
 	TSoftObjectPtr<UActorComponent> GeneratedComponent;
@@ -144,6 +184,53 @@ private:
 	void SetGeneratedComponentFromBP(TSoftObjectPtr<UActorComponent> InGeneratedComponent);
 };
 
+/** This managed resource class is used to tie multiple components in the same resource so that they are cleaned up all at the same time. */
+UCLASS(BlueprintType)
+class PCG_API UPCGManagedComponentList : public UPCGManagedComponentBase
+{
+	GENERATED_BODY()
+	friend class UPCGComponent;
+
+public:
+	//~Begin UPCGManagedResource interface
+	virtual bool Release(bool bHardRelease, TSet<TSoftObjectPtr<AActor>>& OutActorsToDelete) override;
+	virtual bool ReleaseIfUnused(TSet<TSoftObjectPtr<AActor>>& OutActorsToDelete) override;
+	//~End UPCGManagedResource interface
+
+	//~Begin UPCGManagedComponentBase interface
+#if WITH_EDITOR
+	/** Hides the content of the component in a transient way (such as unregistering) */
+	virtual void HideComponent(int32 ComponentIndex) override;
+#endif
+
+	virtual void ForgetComponent(int32 ComponentIndex) override;
+
+protected:
+	virtual TArrayView<TSoftObjectPtr<UActorComponent>> GetComponentsArray() override { return GeneratedComponents; }
+	virtual int32 GetComponentsCount() const override { return GeneratedComponents.Num(); }
+	//~End UPCGManagedComponentBase interface
+
+public:
+	UPROPERTY(BlueprintReadWrite, EditAnywhere, Category = GeneratedData, BlueprintSetter = SetGeneratedComponentsFromBP)
+	TArray<TSoftObjectPtr<UActorComponent>> GeneratedComponents;
+
+private:
+	// When creating components from BP they will be tagged automatically as created from construction script, which makes them transient and isn't compatible with the PCG workflow.
+	UFUNCTION(BlueprintSetter, meta = (BlueprintInternalUseOnly = "true"))
+	void SetGeneratedComponentsFromBP(const TArray<TSoftObjectPtr<UActorComponent>>& InGeneratedComponent);
+};
+
+/** Stub default list to hold resources pushed from BP in a single place on the PCG component. */
+UCLASS()
+class UPCGManagedComponentDefaultList final : public UPCGManagedComponentList
+{
+	GENERATED_BODY()
+	friend class UPCGComponent;
+
+private:
+	void AddGeneratedComponentsFromBP(const TArray<TSoftObjectPtr<UActorComponent>>& InGeneratedComponents);
+};
+
 UCLASS(BlueprintType)
 class PCG_API UPCGManagedISMComponent : public UPCGManagedComponent
 {
