Removed temporary cleanup change since the new managed component default list.
#rnx

diff --git a/Engine/Plugins/PCG/Source/PCG/Private/PCGComponent.cpp b/Engine/Plugins/PCG/Source/PCG/Private/PCGComponent.cpp
--- a/Engine/Plugins/PCG/Source/PCG/Private/PCGComponent.cpp
+++ b/Engine/Plugins/PCG/Source/PCG/Private/PCGComponent.cpp
@@ -1394,22 +1394,7 @@ FPCGTaskId UPCGComponent::CreateCleanupTask(bool bRemoveComponents, const TArray
 		return true;
 	};
 
-#if WITH_EDITOR
-	if (!PCGHelpers::IsRuntimeOrPIE())
-	{
-		auto SuperCleanupTask = [CleanupTask]()
-		{
-			while (!CleanupTask()) {}
-			return true;
-		};
-
-		return GetSubsystem()->ScheduleGeneric(SuperCleanupTask, this, *AllDependencies);
-	}
-	else
-#endif
-	{
-		return GetSubsystem()->ScheduleGeneric(CleanupTask, this, *AllDependencies);
-	}
+	return GetSubsystem()->ScheduleGeneric(CleanupTask, this, *AllDependencies);
 }
 
 void UPCGComponent::CleanupUnusedManagedResources()
