Added more flexibility around keeping/removing attributes and attribute filtering

diff --git a/Engine/Plugins/PCG/Source/PCG/Private/Elements/PCGDeleteAttributesElement.cpp b/Engine/Plugins/PCG/Source/PCG/Private/Elements/PCGDeleteAttributesElement.cpp
--- a/Engine/Plugins/PCG/Source/PCG/Private/Elements/PCGDeleteAttributesElement.cpp
+++ b/Engine/Plugins/PCG/Source/PCG/Private/Elements/PCGDeleteAttributesElement.cpp
@@ -170,7 +170,7 @@ bool FPCGDeleteAttributesElement::ExecuteInternal(FPCGContext* Context) const
 
 			UPCGSpatialData* NewSpatialData = InputSpatialData->DuplicateData(/*bInitializeFromThisData=*/false);
 			Metadata = NewSpatialData->Metadata;
-			NewSpatialData->Metadata->InitializeWithAttributeFilter(ParentMetadata, AttributesToFilter, FilterMode);
+			NewSpatialData->Metadata->InitializeWithAttributeFilter(ParentMetadata, AttributesToFilter, FilterMode, Settings->Operator);
 
 			// No need to inherit metadata since we already initialized it.
 			NewSpatialData->InitializeFromData(InputSpatialData, /*InMetadataParentOverride=*/ nullptr, /*bInheritMetadata=*/ false);
diff --git a/Engine/Plugins/PCG/Source/PCG/Private/Elements/PCGFilterByAttribute.cpp b/Engine/Plugins/PCG/Source/PCG/Private/Elements/PCGFilterByAttribute.cpp
--- a/Engine/Plugins/PCG/Source/PCG/Private/Elements/PCGFilterByAttribute.cpp
+++ b/Engine/Plugins/PCG/Source/PCG/Private/Elements/PCGFilterByAttribute.cpp
@@ -4,8 +4,10 @@
 
 #include "PCGContext.h"
 #include "PCGData.h"
-#include "Metadata/Accessors/IPCGAttributeAccessor.h"
-#include "Metadata/Accessors/PCGAttributeAccessorHelpers.h"
+#include "Helpers/PCGHelpers.h"
+#include "Metadata/PCGMetadata.h"
+
+#include "Algo/AnyOf.h"
 
 #define LOCTEXT_NAMESPACE "PCGFilterByAttributeElement"
 
@@ -43,25 +45,53 @@ bool FPCGFilterByAttributeElement::ExecuteInternal(FPCGContext* Context) const
 	const UPCGFilterByAttributeSettings* Settings = Context->GetInputSettings<UPCGFilterByAttributeSettings>();
 	check(Settings);
 
-	FPCGAttributePropertySelector Selector;
-	Selector.Update(Settings->Attribute.ToString());
+	TArray<FString> Attributes = PCGHelpers::GetStringArrayFromCommaSeparatedList(Settings->Attribute.ToString());
 
 	TArray<FPCGTaggedData> Inputs = Context->InputData.GetInputsByPin(PCGPinConstants::DefaultInputLabel);
 	TArray<FPCGTaggedData>& Outputs = Context->OutputData.TaggedData;
 
+	TArray<FName> DataAttributes;
+	TArray<EPCGMetadataTypes> DataAttributeTypes;
+	TArray<FString> DataAttributeStrings;
+
 	for (const FPCGTaggedData& Input : Inputs)
 	{
 		FPCGTaggedData& Output = Outputs.Add_GetRef(Input);
 		Output.Pin = PCGPinConstants::DefaultOutFilterLabel;
 
-		if (!Input.Data)
+		if (!Input.Data || !Input.Data->ConstMetadata())
 		{
 			continue;
 		}
 
-		TUniquePtr<const IPCGAttributeAccessor> Accessor = PCGAttributeAccessorHelpers::CreateConstAccessor(Input.Data, Selector, /*bQuiet=*/true);
+		// All attributes from the list must have a match in order to put the data in the In Filter pin.
+		const UPCGMetadata* Metadata = Input.Data->ConstMetadata();
+
+		DataAttributes.Reset();
+		DataAttributeTypes.Reset();
+		Metadata->GetAttributes(DataAttributes, DataAttributeTypes);
+
+		DataAttributeStrings.Reset();
+		Algo::Transform(DataAttributes, DataAttributeStrings, [](const FName& InAttribute) { return InAttribute.ToString(); });
+
+		bool bInFilter = true;
+
+		for (const FString& Attribute : Attributes)
+		{
+			if ((Settings->Operator == EPCGStringMatchingOperator::Equal && DataAttributeStrings.Contains(Attribute)) ||
+				(Settings->Operator == EPCGStringMatchingOperator::Substring && Algo::AnyOf(DataAttributeStrings, [&Attribute](const FString& DataAttribute) { return DataAttribute.Contains(Attribute); })) ||
+				(Settings->Operator == EPCGStringMatchingOperator::Matches && Algo::AnyOf(DataAttributeStrings, [&Attribute](const FString& DataAttribute) { return DataAttribute.MatchesWildcard(Attribute); })))
+			{
+				// This attribute has found a match, carry on
+			}
+			else
+			{
+				bInFilter = false;
+				break;
+			}
+		}
 
-		if (Accessor && Accessor.IsValid())
+		if (bInFilter)
 		{
 			Output.Pin = PCGPinConstants::DefaultInFilterLabel;
 		}
diff --git a/Engine/Plugins/PCG/Source/PCG/Private/Metadata/PCGMetadata.cpp b/Engine/Plugins/PCG/Source/PCG/Private/Metadata/PCGMetadata.cpp
--- a/Engine/Plugins/PCG/Source/PCG/Private/Metadata/PCGMetadata.cpp
+++ b/Engine/Plugins/PCG/Source/PCG/Private/Metadata/PCGMetadata.cpp
@@ -9,6 +9,7 @@
 #include "Helpers/PCGPropertyHelpers.h"
 #include "Metadata/PCGAttributePropertySelector.h"
 
+#include "Algo/AnyOf.h"
 #include "Algo/Transform.h"
 #include "Async/ParallelFor.h"
 
@@ -91,7 +92,7 @@ void UPCGMetadata::Initialize(const UPCGMetadata* InParent, bool bAddAttributesF
 	InitializeWithAttributeFilter(InParent, TSet<FName>(), bFilter);
 }
 
-void UPCGMetadata::InitializeWithAttributeFilter(const UPCGMetadata* InParent, const TSet<FName>& InFilteredAttributes, EPCGMetadataFilterMode InFilterMode)
+void UPCGMetadata::InitializeWithAttributeFilter(const UPCGMetadata* InParent, const TSet<FName>& InFilteredAttributes, EPCGMetadataFilterMode InFilterMode, EPCGStringMatchingOperator InMatchOperator)
 {
 	if (Parent || Attributes.Num() != 0)
 	{
@@ -106,7 +107,7 @@ void UPCGMetadata::InitializeWithAttributeFilter(const UPCGMetadata* InParent, c
 	const bool bSkipAddingAttributesFromParent = (InFilterMode == EPCGMetadataFilterMode::IncludeAttributes) && (InFilteredAttributes.Num() == 0);
 	if (!bSkipAddingAttributesFromParent)
 	{
-		AddAttributesFiltered(InParent, InFilteredAttributes, InFilterMode);
+		AddAttributesFiltered(InParent, InFilteredAttributes, InFilterMode, InMatchOperator);
 	}
 }
 
@@ -178,19 +179,50 @@ void UPCGMetadata::InitializeAsCopyWithAttributeFilter(const UPCGMetadata* InMet
 	}
 }
 
-void UPCGMetadata::AddAttributesFiltered(const UPCGMetadata* InOther, const TSet<FName>& InFilteredAttributes, EPCGMetadataFilterMode InFilterMode)
+void UPCGMetadata::AddAttributesFiltered(const UPCGMetadata* InOther, const TSet<FName>& InFilteredAttributes, EPCGMetadataFilterMode InFilterMode, EPCGStringMatchingOperator InMatchOperator)
 {
 	if (!InOther)
 	{
 		return;
 	}
 
+	TArray<FString> InFilteredAttributesStrings;
+	if (InMatchOperator != EPCGStringMatchingOperator::Equal)
+	{
+		Algo::Transform(InFilteredAttributes, InFilteredAttributesStrings, [](const FName& InFilteredAttribute) {return InFilteredAttribute.ToString(); });
+	}
+
+	auto IsAttributeInFilterList = [&InFilteredAttributes, &InFilteredAttributesStrings, InMatchOperator](FName OtherAttributeName) -> bool
+	{
+		if (InMatchOperator == EPCGStringMatchingOperator::Equal)
+		{
+			return InFilteredAttributes.Contains(OtherAttributeName);
+		}
+		else
+		{
+			const FString OtherAttributeString = OtherAttributeName.ToString();
+			if (InMatchOperator == EPCGStringMatchingOperator::Substring)
+			{
+				return Algo::AnyOf(InFilteredAttributesStrings, [&OtherAttributeString](const FString& InAttribute) { return OtherAttributeString.Contains(InAttribute); });
+			}
+			else if (InMatchOperator == EPCGStringMatchingOperator::Matches)
+			{
+				return Algo::AnyOf(InFilteredAttributesStrings, [&OtherAttributeString](const FString& InAttribute) { return OtherAttributeString.MatchesWildcard(InAttribute); });
+			}
+			else
+			{
+				checkNoEntry();
+				return false;
+			}
+		}
+	};
+
 	bool bAttributeAdded = false;
 
-	for (const TPair<FName, FPCGMetadataAttributeBase*> OtherAttribute : InOther->Attributes)
+	for (const TPair<FName, FPCGMetadataAttributeBase*>& OtherAttribute : InOther->Attributes)
 	{
 		// Skip this attribute if it is in an exclude list, or if it is not in an include list
-		const bool bAttributeInFilterList = InFilteredAttributes.Contains(OtherAttribute.Key);
+		const bool bAttributeInFilterList = IsAttributeInFilterList(OtherAttribute.Key);
 		const bool bSkipAttributesInFilterList = InFilterMode == EPCGMetadataFilterMode::ExcludeAttributes;
 		const bool bSkipThisAttribute = bSkipAttributesInFilterList == bAttributeInFilterList;
 
diff --git a/Engine/Plugins/PCG/Source/PCG/Public/Elements/PCGDeleteAttributesElement.h b/Engine/Plugins/PCG/Source/PCG/Public/Elements/PCGDeleteAttributesElement.h
--- a/Engine/Plugins/PCG/Source/PCG/Public/Elements/PCGDeleteAttributesElement.h
+++ b/Engine/Plugins/PCG/Source/PCG/Public/Elements/PCGDeleteAttributesElement.h
@@ -57,6 +57,9 @@ public:
 	UPROPERTY(BlueprintReadWrite, EditAnywhere, Category = Settings, meta = (PCG_Overridable))
 	EPCGAttributeFilterOperation Operation = EPCGAttributeFilterOperation::KeepSelectedAttributes;
 
+	UPROPERTY(BlueprintReadWrite, EditAnywhere, Category = Settings, meta = (PCG_Overridable))
+	EPCGStringMatchingOperator Operator = EPCGStringMatchingOperator::Equal;
+
 	/** Comma-separated list of attributes to keep or remove from the input data. */
 	UPROPERTY(BlueprintReadWrite, EditAnywhere, Category = Settings, meta = (PCG_Overridable))
 	FString SelectedAttributes;
diff --git a/Engine/Plugins/PCG/Source/PCG/Public/Elements/PCGFilterByAttribute.h b/Engine/Plugins/PCG/Source/PCG/Public/Elements/PCGFilterByAttribute.h
--- a/Engine/Plugins/PCG/Source/PCG/Public/Elements/PCGFilterByAttribute.h
+++ b/Engine/Plugins/PCG/Source/PCG/Public/Elements/PCGFilterByAttribute.h
@@ -30,9 +30,12 @@ protected:
 	//~End UPCGSettings interface
 
 public:
-	/** Attribute to look for */
-	UPROPERTY(BlueprintReadWrite, EditAnywhere, Category = Settings, meta = (PCG_Overridable))
+	/** Comma-separated list of attributes to look for */
+	UPROPERTY(BlueprintReadWrite, EditAnywhere, Category = Settings, meta = (DisplayName = "Attributes", PCG_Overridable))
 	FName Attribute;
+
+	UPROPERTY(BlueprintReadWrite, EditAnywhere, Category = Settings, meta = (PCG_Overridable))
+	EPCGStringMatchingOperator Operator = EPCGStringMatchingOperator::Equal;
 };
 
 class FPCGFilterByAttributeElement : public IPCGElement
diff --git a/Engine/Plugins/PCG/Source/PCG/Public/Metadata/PCGMetadata.h b/Engine/Plugins/PCG/Source/PCG/Public/Metadata/PCGMetadata.h
--- a/Engine/Plugins/PCG/Source/PCG/Public/Metadata/PCGMetadata.h
+++ b/Engine/Plugins/PCG/Source/PCG/Public/Metadata/PCGMetadata.h
@@ -2,6 +2,7 @@
 
 #pragma once
 
+#include "PCGCommon.h"
 #include "PCGMetadataCommon.h" // IWYU pragma: keep
 #include "PCGMetadataAttributeTraits.h"
 #include "Metadata/PCGMetadataAttributeTpl.h"
@@ -37,7 +38,7 @@ public:
 	 * @param InFilterMode Defines attribute filter operation.
 	 */
 	UFUNCTION(BlueprintCallable, Category = "PCG|Metadata")
-	void InitializeWithAttributeFilter(const UPCGMetadata* InParent, const TSet<FName>& InFilteredAttributes, EPCGMetadataFilterMode InFilterMode = EPCGMetadataFilterMode::ExcludeAttributes);
+	void InitializeWithAttributeFilter(const UPCGMetadata* InParent, const TSet<FName>& InFilteredAttributes, EPCGMetadataFilterMode InFilterMode = EPCGMetadataFilterMode::ExcludeAttributes, EPCGStringMatchingOperator InMatchOperator = EPCGStringMatchingOperator::Equal);
 
 	/** Initializes the metadata from a parent metadata by copying all attributes to it.
 	* @param InMetadataToCopy Metadata to copy from
@@ -64,7 +65,7 @@ public:
 	 * @param InFilterMode Defines attribute filter operation.
 	 */
 	UFUNCTION(BlueprintCallable, Category = "PCG|Metadata")
-	void AddAttributesFiltered(const UPCGMetadata* InOther, const TSet<FName>& InFilteredAttributes, EPCGMetadataFilterMode InFilterMode = EPCGMetadataFilterMode::ExcludeAttributes);
+	void AddAttributesFiltered(const UPCGMetadata* InOther, const TSet<FName>& InFilteredAttributes, EPCGMetadataFilterMode InFilterMode = EPCGMetadataFilterMode::ExcludeAttributes, EPCGStringMatchingOperator InMatchOperator = EPCGStringMatchingOperator::Equal);
 
 	/** Creates missing attribute from another metadata if it is not currently present - note that this does not copy values */
 	UFUNCTION(BlueprintCallable, Category = "PCG|Metadata")
diff --git a/Engine/Plugins/PCG/Source/PCG/Public/PCGCommon.h b/Engine/Plugins/PCG/Source/PCG/Public/PCGCommon.h
--- a/Engine/Plugins/PCG/Source/PCG/Public/PCGCommon.h
+++ b/Engine/Plugins/PCG/Source/PCG/Public/PCGCommon.h
@@ -206,8 +206,11 @@ enum class EPCGCoordinateSpace : uint8
 UENUM(BlueprintType)
 enum class EPCGStringMatchingOperator : uint8
 {
+	/** Will return a match only if the two strings compared are the same */
 	Equal,
+	/** Will return a match if the first string contains the second */
 	Substring,
+	/** Will return a match if the first string matches the pattern defined by the second (including wildcards) */
 	Matches
 };
 
