[PCG] Fix check failure in PCGLandscapeSplineData when sampling on landscape splines with 0 segments

#rnx

diff --git a/Engine/Plugins/PCG/Source/PCG/Private/Data/PCGLandscapeSplineData.cpp b/Engine/Plugins/PCG/Source/PCG/Private/Data/PCGLandscapeSplineData.cpp
--- a/Engine/Plugins/PCG/Source/PCG/Private/Data/PCGLandscapeSplineData.cpp
+++ b/Engine/Plugins/PCG/Source/PCG/Private/Data/PCGLandscapeSplineData.cpp
@@ -78,7 +78,14 @@ FVector::FReal UPCGLandscapeSplineData::GetSegmentLength(int SegmentIndex) const
 
 FTransform UPCGLandscapeSplineData::GetTransformAtDistance(int SegmentIndex, FVector::FReal Distance, bool bWorldSpace, FBox* OutBounds) const
 {
-	check(Spline.IsValid() && Spline->GetSegments().IsValidIndex(SegmentIndex));
+	check(Spline.IsValid());
+
+	if (GetNumSegments() == 0)
+	{
+		return GetTransform();
+	}
+
+	check(Spline->GetSegments().IsValidIndex(SegmentIndex));
 
 	const ULandscapeSplineSegment* Segment = Spline->GetSegments()[SegmentIndex];
 	check(Segment);
@@ -122,7 +129,14 @@ FTransform UPCGLandscapeSplineData::GetTransformAtDistance(int SegmentIndex, FVe
 
 FVector::FReal UPCGLandscapeSplineData::GetCurvatureAtDistance(int SegmentIndex, FVector::FReal Distance) const
 {
-	check(Spline.IsValid() && Spline->GetSegments().IsValidIndex(SegmentIndex));
+	check(Spline.IsValid());
+
+	if (GetNumSegments() == 0)
+	{
+		return 0.0;
+	}
+
+	check(Spline->GetSegments().IsValidIndex(SegmentIndex));
 
 	const ULandscapeSplineSegment* Segment = Spline->GetSegments()[SegmentIndex];
 	check(Segment);
@@ -211,8 +225,17 @@ void UPCGLandscapeSplineData::GetTangentsAtSegmentStart(int SegmentIndex, FVecto
 
 FVector::FReal UPCGLandscapeSplineData::GetDistanceAtSegmentStart(int SegmentIndex) const
 {
+	check(Spline.IsValid());
+
+	const int NumSegments = GetNumSegments();
+
+	if (NumSegments == 0)
+	{
+		return 0.0;
+	}
+	
 	// Allow SegmentIndex == NumSegments, which indicates we want the distance to the final control point, which is like saying "Start of the Nth segment".
-	check(Spline.IsValid() && SegmentIndex >= 0 && SegmentIndex <= Spline->GetSegments().Num());
+	check(SegmentIndex >= 0 && SegmentIndex <= NumSegments);
 
 	const TArray<TObjectPtr<ULandscapeSplineSegment>>& Segments = Spline->GetSegments();
 	int32 ReparamIndex = 0;
