[PCG] Fix ForAllIntersectingPartitionActors to only iterate over non-zero overlaps

#rnx

diff --git a/Engine/Plugins/PCG/Source/PCG/Private/PCGActorAndComponentMapping.cpp b/Engine/Plugins/PCG/Source/PCG/Private/PCGActorAndComponentMapping.cpp
--- a/Engine/Plugins/PCG/Source/PCG/Private/PCGActorAndComponentMapping.cpp
+++ b/Engine/Plugins/PCG/Source/PCG/Private/PCGActorAndComponentMapping.cpp
@@ -734,7 +734,12 @@ void FPCGActorAndComponentMapping::ForAllIntersectingPartitionActors(const FBox&
 					{
 						if (APCGPartitionActor* Actor = ActorPtr->Get())
 						{
-							InFunc(Actor);
+							// Exclude any surrounding cells which are touching this cell but not meaningfully overlapping.
+							const FBox Overlap = Actor->GetFixedBounds().Overlap(InBounds);
+							if (Overlap.GetVolume() > 0)
+							{
+								InFunc(Actor);
+							}
 						}
 					}
 				}
