[PCG] Readded support for extractors & properties in the filter by attribute node

#rnx

diff --git a/Engine/Plugins/PCG/Source/PCG/Private/Elements/PCGFilterByAttribute.cpp b/Engine/Plugins/PCG/Source/PCG/Private/Elements/PCGFilterByAttribute.cpp
--- a/Engine/Plugins/PCG/Source/PCG/Private/Elements/PCGFilterByAttribute.cpp
+++ b/Engine/Plugins/PCG/Source/PCG/Private/Elements/PCGFilterByAttribute.cpp
@@ -4,8 +4,11 @@
 
 #include "PCGContext.h"
 #include "PCGData.h"
+#include "Data/PCGPointData.h"
 #include "Helpers/PCGHelpers.h"
 #include "Metadata/PCGMetadata.h"
+#include "Metadata/Accessors/IPCGAttributeAccessor.h"
+#include "Metadata/Accessors/PCGAttributeAccessorHelpers.h"
 
 #include "Algo/AnyOf.h"
 
@@ -69,25 +72,97 @@ bool FPCGFilterByAttributeElement::ExecuteInternal(FPCGContext* Context) const
 
 		DataAttributes.Reset();
 		DataAttributeTypes.Reset();
-		Metadata->GetAttributes(DataAttributes, DataAttributeTypes);
-
 		DataAttributeStrings.Reset();
-		Algo::Transform(DataAttributes, DataAttributeStrings, [](const FName& InAttribute) { return InAttribute.ToString(); });
+
+		if (Settings->Operator != EPCGStringMatchingOperator::Equal)
+		{
+			// Gather all attributes & properties based on data type
+			Metadata->GetAttributes(DataAttributes, DataAttributeTypes);
+			Algo::Transform(DataAttributes, DataAttributeStrings, [](const FName& InAttribute) { return InAttribute.ToString(); });
+
+			if (Cast<UPCGPointData>(Input.Data))
+			{
+				if (const UEnum* PointProperties = StaticEnum<EPCGPointProperties>())
+				{
+					DataAttributeStrings.Reserve(DataAttributeStrings.Num() + PointProperties->NumEnums());
+					for (int32 EnumIndex = 0; EnumIndex < PointProperties->NumEnums(); ++EnumIndex)
+					{
+						DataAttributeStrings.Add(FString(TEXT("$")) + PointProperties->GetNameStringByIndex(EnumIndex));
+					}
+				}
+			}
+
+			if (const UEnum* ExtraProperties = StaticEnum<EPCGExtraProperties>())
+			{
+				DataAttributeStrings.Reserve(DataAttributeStrings.Num() + ExtraProperties->NumEnums());
+				for (int32 EnumIndex = 0; EnumIndex < ExtraProperties->NumEnums(); ++EnumIndex)
+				{
+					DataAttributeStrings.Add(FString(TEXT("$")) + ExtraProperties->GetNameStringByIndex(EnumIndex));
+				}
+			}
+		}
 
 		bool bInFilter = true;
 
 		for (const FString& Attribute : Attributes)
 		{
-			if ((Settings->Operator == EPCGStringMatchingOperator::Equal && DataAttributeStrings.Contains(Attribute)) ||
-				(Settings->Operator == EPCGStringMatchingOperator::Substring && Algo::AnyOf(DataAttributeStrings, [&Attribute](const FString& DataAttribute) { return DataAttribute.Contains(Attribute); })) ||
-				(Settings->Operator == EPCGStringMatchingOperator::Matches && Algo::AnyOf(DataAttributeStrings, [&Attribute](const FString& DataAttribute) { return DataAttribute.MatchesWildcard(Attribute); })))
+			// In the case of the equal test, we can test directly if the selector would yield something valid
+			if (Settings->Operator == EPCGStringMatchingOperator::Equal)
 			{
-				// This attribute has found a match, carry on
+				FPCGAttributePropertySelector Selector;
+				Selector.Update(Attribute);
+
+				TUniquePtr<const IPCGAttributeAccessor> Accessor = PCGAttributeAccessorHelpers::CreateConstAccessor(Input.Data, Selector, /*bQuiet=*/true);
+				if (!Accessor || !Accessor.IsValid())
+				{
+					bInFilter = false;
+					break;
+				}
 			}
 			else
 			{
-				bInFilter = false;
-				break;
+				// Otherwise, it's going to be a bit more complex -
+				// First, reconstruct the main property/attribute name from the selector, because it might have removed the $ character.
+				FString AttributeWithNoAccessor;
+				FString AttributeAccessors;
+				int32 FirstAccessorCharIndex = INDEX_NONE;
+				if (Attribute.FindChar('.', FirstAccessorCharIndex))
+				{
+					AttributeWithNoAccessor = Attribute.Left(FirstAccessorCharIndex);
+					AttributeAccessors = Attribute.RightChop(FirstAccessorCharIndex); // includes the '.' as we'll add it back after
+				}
+				else
+				{
+					AttributeWithNoAccessor = Attribute;
+				}
+
+				// Try to find a valid match of AttributeWithNoAccessor against the DataAttributeStrings.
+				bool bFoundValidMatch = false;
+				for (const FString& DataAttribute : DataAttributeStrings)
+				{
+					if ((Settings->Operator == EPCGStringMatchingOperator::Substring && !DataAttribute.Contains(AttributeWithNoAccessor)) ||
+						(Settings->Operator == EPCGStringMatchingOperator::Matches && !DataAttribute.MatchesWildcard(AttributeWithNoAccessor)))
+					{
+						continue;
+					}
+
+					// We have a valid name-based match, now check if the full attribute can be used as a valid extractor.
+					FPCGAttributePropertySelector Selector;
+					Selector.Update(DataAttribute + AttributeAccessors);
+
+					TUniquePtr<const IPCGAttributeAccessor> Accessor = PCGAttributeAccessorHelpers::CreateConstAccessor(Input.Data, Selector, /*bQuiet=*/true);
+					if (Accessor && Accessor.IsValid())
+					{
+						bFoundValidMatch = true;
+						break;
+					}
+				}
+
+				if (!bFoundValidMatch)
+				{
+					bInFilter = false;
+					break;
+				}
 			}
 		}
 
