[PCG] Make sure PCG Tasks running outside of game thread can safely create UObjects
- Add a FPCGContext::NewObject_AnyThread method
- FPCGContext tracks non game thread created objects to remove EInternalObjectFlags::Async when FPCGContext gets deleted on game thread (when ActiveTask is done)
- FGCScopeGuard is used when allocating UObjects in non game thread tasks to make sure we don't allocate while GC is running

diff --git a/Engine/Plugins/Experimental/PCGWaterInterop/Source/PCGWaterInterop/Private/Data/PCGWaterSplineData.cpp b/Engine/Plugins/Experimental/PCGWaterInterop/Source/PCGWaterInterop/Private/Data/PCGWaterSplineData.cpp
--- a/Engine/Plugins/Experimental/PCGWaterInterop/Source/PCGWaterInterop/Private/Data/PCGWaterSplineData.cpp
+++ b/Engine/Plugins/Experimental/PCGWaterInterop/Source/PCGWaterInterop/Private/Data/PCGWaterSplineData.cpp
@@ -2,6 +2,7 @@
 
 #include "Data/PCGWaterSplineData.h"
 
+#include "PCGContext.h"
 #include "Data/PCGPointData.h"
 #include "Metadata/PCGMetadata.h"
 
@@ -69,9 +70,9 @@ bool UPCGWaterSplineData::SamplePoint(const FTransform& Transform, const FBox& B
 	}
 }
 
-UPCGSpatialData* UPCGWaterSplineData::CopyInternal() const
+UPCGSpatialData* UPCGWaterSplineData::CopyInternal(FPCGContext* Context) const
 {
-	UPCGWaterSplineData* NewSplineData = NewObject<UPCGWaterSplineData>();
+	UPCGWaterSplineData* NewSplineData = FPCGContext::NewObject_AnyThread<UPCGWaterSplineData>(Context);
 	NewSplineData->WaterSplineMetadataStruct = WaterSplineMetadataStruct;
 
 	CopySplineData(NewSplineData);
diff --git a/Engine/Plugins/Experimental/PCGWaterInterop/Source/PCGWaterInterop/Public/Data/PCGWaterSplineData.h b/Engine/Plugins/Experimental/PCGWaterInterop/Source/PCGWaterInterop/Public/Data/PCGWaterSplineData.h
--- a/Engine/Plugins/Experimental/PCGWaterInterop/Source/PCGWaterInterop/Public/Data/PCGWaterSplineData.h
+++ b/Engine/Plugins/Experimental/PCGWaterInterop/Source/PCGWaterInterop/Public/Data/PCGWaterSplineData.h
@@ -42,7 +42,7 @@ public:
 	//~Begin UPCGSpatialData interface
 	virtual bool SamplePoint(const FTransform& Transform, const FBox& Bounds, FPCGPoint& OutPoint, UPCGMetadata* OutMetadata) const override;
 protected:
-	virtual UPCGSpatialData* CopyInternal() const override;
+	virtual UPCGSpatialData* CopyInternal(FPCGContext* Context) const override;
 	//~End UPCGSpatialData interface
 
 	//~Begin UPCGPolyLineData interface
diff --git a/Engine/Plugins/PCG/Config/BasePCG.ini b/Engine/Plugins/PCG/Config/BasePCG.ini
--- a/Engine/Plugins/PCG/Config/BasePCG.ini
+++ b/Engine/Plugins/PCG/Config/BasePCG.ini
@@ -16,6 +16,13 @@
 +FunctionRedirects=(OldName="PCGMetadata.InitializeAsCopyWithAttributeFilter",NewName="PCGMetadata.K2_InitializeAsCopyWithAttributeFilter")
 +FunctionRedirects=(OldName="PCGSpatialData.SamplePoint",NewName="PCGSpatialData.K2_SamplePoint")
 +FunctionRedirects=(OldName="PCGSpatialData.ProjectPoint",NewName="PCGSpatialData.K2_ProjectPoint")
++FunctionRedirects=(OldName="PCGSpatialData.IntersectWith",NewName="PCGSpatialData.K2_IntersectWith")
++FunctionRedirects=(OldName="PCGSpatialData.ProjectOn",NewName="PCGSpatialData.K2_ProjectOn")
++FunctionRedirects=(OldName="PCGSpatialData.UnionWith",NewName="PCGSpatialData.K2_UnionWith")
++FunctionRedirects=(OldName="PCGSpatialData.Subtract",NewName="PCGSpatialData.K2_Subtract")
++FunctionRedirects=(OldName="PCGDifferenceData.AddDifference",NewName="PCGDifferenceData.K2_AddDifference")
++FunctionRedirects=(OldName="PCGParamData.FilterParamsByName",NewName="PCGParamData.K2_FilterParamsByName")
++FunctionRedirects=(OldName="PCGParamData.FilterParamsByKey",NewName="PCGParamData.K2_FilterParamsByKey")
 
 +PropertyRedirects=(OldName="PCGSplineSamplerSettings.Params",NewName="PCGSplineSamplerSettings.SamplerParams")
 +PropertyRedirects=(OldName="PCGProjectionSettings.Params",NewName="PCGProjectionSettings.ProjectionParams")
diff --git a/Engine/Plugins/PCG/Source/PCG/Private/Data/PCGCollisionShapeData.cpp b/Engine/Plugins/PCG/Source/PCG/Private/Data/PCGCollisionShapeData.cpp
--- a/Engine/Plugins/PCG/Source/PCG/Private/Data/PCGCollisionShapeData.cpp
+++ b/Engine/Plugins/PCG/Source/PCG/Private/Data/PCGCollisionShapeData.cpp
@@ -106,9 +106,9 @@ const UPCGPointData* UPCGCollisionShapeData::CreatePointData(FPCGContext* Contex
 	return Data;
 }
 
-UPCGSpatialData* UPCGCollisionShapeData::CopyInternal() const
+UPCGSpatialData* UPCGCollisionShapeData::CopyInternal(FPCGContext* Context) const
 {
-	UPCGCollisionShapeData* NewShapeData = NewObject<UPCGCollisionShapeData>();
+	UPCGCollisionShapeData* NewShapeData = FPCGContext::NewObject_AnyThread<UPCGCollisionShapeData>(Context);
 
 	NewShapeData->Transform = Transform;
 	NewShapeData->Shape = Shape;
diff --git a/Engine/Plugins/PCG/Source/PCG/Private/Data/PCGCollisionWrapperData.cpp b/Engine/Plugins/PCG/Source/PCG/Private/Data/PCGCollisionWrapperData.cpp
--- a/Engine/Plugins/PCG/Source/PCG/Private/Data/PCGCollisionWrapperData.cpp
+++ b/Engine/Plugins/PCG/Source/PCG/Private/Data/PCGCollisionWrapperData.cpp
@@ -2,6 +2,7 @@
 
 #include "Data/PCGCollisionWrapperData.h"
 
+#include "PCGContext.h"
 #include "Data/PCGPointData.h"
 #include "Elements/Metadata/PCGMetadataElementCommon.h"
 
@@ -337,9 +338,9 @@ bool UPCGCollisionWrapperData::SamplePoint(const FTransform& InTransform, const
 	}
 }
 
-UPCGSpatialData* UPCGCollisionWrapperData::CopyInternal() const
+UPCGSpatialData* UPCGCollisionWrapperData::CopyInternal(FPCGContext* Context) const
 {
-	UPCGCollisionWrapperData* NewCollisionWrapperData = NewObject<UPCGCollisionWrapperData>();
+	UPCGCollisionWrapperData* NewCollisionWrapperData = FPCGContext::NewObject_AnyThread<UPCGCollisionWrapperData>(Context);
 	NewCollisionWrapperData->Initialize(PointData, CollisionSelector, bUseComplexCollision, bUseCollisionAccurateOctree);
 
 	return NewCollisionWrapperData;
diff --git a/Engine/Plugins/PCG/Source/PCG/Private/Data/PCGDifferenceData.cpp b/Engine/Plugins/PCG/Source/PCG/Private/Data/PCGDifferenceData.cpp
--- a/Engine/Plugins/PCG/Source/PCG/Private/Data/PCGDifferenceData.cpp
+++ b/Engine/Plugins/PCG/Source/PCG/Private/Data/PCGDifferenceData.cpp
@@ -7,6 +7,7 @@
 #include "Data/PCGSpatialData.h"
 #include "Data/PCGSpatialDataTpl.h"
 #include "Data/PCGUnionData.h"
+#include "Elements/PCGExecuteBlueprint.h"
 #include "Helpers/PCGAsync.h"
 
 #include "Serialization/ArchiveCrc32.h"
@@ -46,7 +47,12 @@ void UPCGDifferenceData::Initialize(const UPCGSpatialData* InData)
 	Metadata->Initialize(Source->Metadata);
 }
 
-void UPCGDifferenceData::AddDifference(const UPCGSpatialData* InDifference)
+void UPCGDifferenceData::K2_AddDifference(const UPCGSpatialData* InDifference)
+{
+	return AddDifference(UPCGBlueprintElement::ResolveContext(), InDifference);
+}
+
+void UPCGDifferenceData::AddDifference(FPCGContext* InContext, const UPCGSpatialData* InDifference)
 {
 	check(InDifference);
 
@@ -69,7 +75,7 @@ void UPCGDifferenceData::AddDifference(const UPCGSpatialData* InDifference)
 	{
 		if (!DifferencesUnion)
 		{
-			DifferencesUnion = NewObject<UPCGUnionData>();
+			DifferencesUnion = FPCGContext::NewObject_AnyThread<UPCGUnionData>(InContext);
 			DifferencesUnion->AddData(Difference);
 			DifferencesUnion->SetDensityFunction(PCGDifferenceDataUtils::ToUnionDensityFunction(DensityFunction));
 			Difference = DifferencesUnion;
@@ -244,7 +250,7 @@ const UPCGPointData* UPCGDifferenceData::CreatePointData(FPCGContext* Context) c
 	const UPCGMetadata* SourceMetadata = SourcePointData->Metadata;
 	const TArray<FPCGPoint>& SourcePoints = SourcePointData->GetPoints();
 
-	UPCGPointData* Data = NewObject<UPCGPointData>();
+	UPCGPointData* Data = FPCGContext::NewObject_AnyThread<UPCGPointData>(Context);
 	Data->InitializeFromData(this, SourceMetadata);
 	
 	UPCGMetadata* OutMetadata = Data->Metadata;
@@ -254,7 +260,7 @@ const UPCGPointData* UPCGDifferenceData::CreatePointData(FPCGContext* Context) c
 	UPCGMetadata* TempDiffMetadata = nullptr;
 	if (bDiffMetadata && OutMetadata && DifferenceMetadata)
 	{
-		TempDiffMetadata = NewObject<UPCGMetadata>();
+		TempDiffMetadata = FPCGContext::NewObject_AnyThread<UPCGMetadata>(Context);
 		TempDiffMetadata->Initialize(DifferenceMetadata);
 	}
 
@@ -310,16 +316,16 @@ const UPCGPointData* UPCGDifferenceData::CreatePointData(FPCGContext* Context) c
 	return Data;
 }
 
-UPCGSpatialData* UPCGDifferenceData::CopyInternal() const
+UPCGSpatialData* UPCGDifferenceData::CopyInternal(FPCGContext* Context) const
 {
-	UPCGDifferenceData* NewDifferenceData = NewObject<UPCGDifferenceData>();
+	UPCGDifferenceData* NewDifferenceData = FPCGContext::NewObject_AnyThread<UPCGDifferenceData>(Context);
 
 	NewDifferenceData->Source = Source;
 	NewDifferenceData->Difference = Difference;
 	NewDifferenceData->DensityFunction = DensityFunction;
 	if (DifferencesUnion)
 	{
-		NewDifferenceData->DifferencesUnion = static_cast<UPCGUnionData*>(DifferencesUnion->DuplicateData());
+		NewDifferenceData->DifferencesUnion = static_cast<UPCGUnionData*>(DifferencesUnion->DuplicateData(Context));
 
 #if WITH_EDITOR
 		NewDifferenceData->RawPointerDifferencesUnion = NewDifferenceData->DifferencesUnion;
diff --git a/Engine/Plugins/PCG/Source/PCG/Private/Data/PCGIntersectionData.cpp b/Engine/Plugins/PCG/Source/PCG/Private/Data/PCGIntersectionData.cpp
--- a/Engine/Plugins/PCG/Source/PCG/Private/Data/PCGIntersectionData.cpp
+++ b/Engine/Plugins/PCG/Source/PCG/Private/Data/PCGIntersectionData.cpp
@@ -220,14 +220,14 @@ UPCGPointData* UPCGIntersectionData::CreateAndFilterPointData(FPCGContext* Conte
 	const TArray<FPCGPoint>& SourcePoints = SourcePointData->GetPoints();
 	const UPCGMetadata* SourceMetadata = SourcePointData->Metadata;
 
-	UPCGPointData* Data = NewObject<UPCGPointData>();
+	UPCGPointData* Data = FPCGContext::NewObject_AnyThread<UPCGPointData>(Context);
 	Data->InitializeFromData(this, SourceMetadata);
 	Data->Metadata->AddAttributes(Y->Metadata);
 
 	UPCGMetadata* TempYMetadata = nullptr;
 	if (Y->Metadata)
 	{
-		TempYMetadata = Y->Metadata ? NewObject<UPCGMetadata>() : nullptr;
+		TempYMetadata = Y->Metadata ? FPCGContext::NewObject_AnyThread<UPCGMetadata>(Context) : nullptr;
 		TempYMetadata->Initialize(Y->Metadata);
 	}
 
@@ -286,9 +286,9 @@ UPCGPointData* UPCGIntersectionData::CreateAndFilterPointData(FPCGContext* Conte
 	return Data;
 }
 
-UPCGSpatialData* UPCGIntersectionData::CopyInternal() const
+UPCGSpatialData* UPCGIntersectionData::CopyInternal(FPCGContext* Context) const
 {
-	UPCGIntersectionData* NewIntersectionData = NewObject<UPCGIntersectionData>();
+	UPCGIntersectionData* NewIntersectionData = FPCGContext::NewObject_AnyThread<UPCGIntersectionData>(Context);
 
 	NewIntersectionData->DensityFunction = DensityFunction;
 	NewIntersectionData->A = A;
diff --git a/Engine/Plugins/PCG/Source/PCG/Private/Data/PCGLandscapeData.cpp b/Engine/Plugins/PCG/Source/PCG/Private/Data/PCGLandscapeData.cpp
--- a/Engine/Plugins/PCG/Source/PCG/Private/Data/PCGLandscapeData.cpp
+++ b/Engine/Plugins/PCG/Source/PCG/Private/Data/PCGLandscapeData.cpp
@@ -2,6 +2,7 @@
 
 #include "Data/PCGLandscapeData.h"
 
+#include "PCGContext.h"
 #include "PCGSubsystem.h"
 #include "Data/PCGPointData.h"
 #include "Data/PCGSpatialDataTpl.h"
@@ -476,7 +477,7 @@ const UPCGPointData* UPCGLandscapeData::CreatePointData(FPCGContext* Context, co
 		return nullptr;
 	}
 
-	UPCGPointData* Data = NewObject<UPCGPointData>();
+	UPCGPointData* Data = FPCGContext::NewObject_AnyThread<UPCGPointData>(Context);
 	Data->InitializeFromData(this);
 	TArray<FPCGPoint>& Points = Data->GetMutablePoints();
 
@@ -612,9 +613,9 @@ const ULandscapeInfo* UPCGLandscapeData::GetLandscapeInfo(const FVector& InPosit
 	return nullptr;
 }
 
-UPCGSpatialData* UPCGLandscapeData::CopyInternal() const
+UPCGSpatialData* UPCGLandscapeData::CopyInternal(FPCGContext* Context) const
 {
-	UPCGLandscapeData* NewLandscapeData = NewObject<UPCGLandscapeData>();
+	UPCGLandscapeData* NewLandscapeData = FPCGContext::NewObject_AnyThread<UPCGLandscapeData>(Context);
 
 	CopyBaseSurfaceData(NewLandscapeData);
 
diff --git a/Engine/Plugins/PCG/Source/PCG/Private/Data/PCGLandscapeSplineData.cpp b/Engine/Plugins/PCG/Source/PCG/Private/Data/PCGLandscapeSplineData.cpp
--- a/Engine/Plugins/PCG/Source/PCG/Private/Data/PCGLandscapeSplineData.cpp
+++ b/Engine/Plugins/PCG/Source/PCG/Private/Data/PCGLandscapeSplineData.cpp
@@ -2,6 +2,7 @@
 
 #include "Data/PCGLandscapeSplineData.h"
 
+#include "PCGContext.h"
 #include "Data/PCGPolyLineData.h"
 #include "Data/PCGPointData.h"
 #include "Elements/PCGSplineSampler.h"
@@ -255,7 +256,7 @@ const UPCGPointData* UPCGLandscapeSplineData::CreatePointData(FPCGContext* Conte
 	check(Spline.IsValid());
 	TRACE_CPUPROFILER_EVENT_SCOPE(UPCGLandscapeSplineData::CreatePointData);
 
-	UPCGPointData* Data = NewObject<UPCGPointData>();
+	UPCGPointData* Data = FPCGContext::NewObject_AnyThread<UPCGPointData>(Context);
 	Data->InitializeFromData(this);
 	TArray<FPCGPoint>& Points = Data->GetMutablePoints();
 
@@ -352,9 +353,9 @@ bool UPCGLandscapeSplineData::SamplePoint(const FTransform& InTransform, const F
 	return OutPoint.Density > 0;
 }
 
-UPCGSpatialData* UPCGLandscapeSplineData::CopyInternal() const
+UPCGSpatialData* UPCGLandscapeSplineData::CopyInternal(FPCGContext* Context) const
 {
-	UPCGLandscapeSplineData* NewLandscapeSplineData = NewObject<UPCGLandscapeSplineData>();
+	UPCGLandscapeSplineData* NewLandscapeSplineData = FPCGContext::NewObject_AnyThread<UPCGLandscapeSplineData>(Context);
 
 	NewLandscapeSplineData->Spline = Spline;
 	NewLandscapeSplineData->ReparamTable = ReparamTable;
diff --git a/Engine/Plugins/PCG/Source/PCG/Private/Data/PCGPointData.cpp b/Engine/Plugins/PCG/Source/PCG/Private/Data/PCGPointData.cpp
--- a/Engine/Plugins/PCG/Source/PCG/Private/Data/PCGPointData.cpp
+++ b/Engine/Plugins/PCG/Source/PCG/Private/Data/PCGPointData.cpp
@@ -2,6 +2,7 @@
 
 #include "Data/PCGPointData.h"
 
+#include "PCGContext.h"
 #include "Helpers/PCGHelpers.h"
 #include "Helpers/PCGTagHelpers.h"
 #include "Metadata/PCGMetadataAccessor.h"
@@ -650,9 +651,9 @@ void UPCGPointData::RebuildOctree() const
 	bOctreeIsDirty = false;
 }
 
-UPCGSpatialData* UPCGPointData::CopyInternal() const
+UPCGSpatialData* UPCGPointData::CopyInternal(FPCGContext* Context) const
 {
-	UPCGPointData* NewPointData = NewObject<UPCGPointData>();
+	UPCGPointData* NewPointData = FPCGContext::NewObject_AnyThread<UPCGPointData>(Context);
 	NewPointData->GetMutablePoints() = GetPoints();
 
 	return NewPointData;
diff --git a/Engine/Plugins/PCG/Source/PCG/Private/Data/PCGPrimitiveData.cpp b/Engine/Plugins/PCG/Source/PCG/Private/Data/PCGPrimitiveData.cpp
--- a/Engine/Plugins/PCG/Source/PCG/Private/Data/PCGPrimitiveData.cpp
+++ b/Engine/Plugins/PCG/Source/PCG/Private/Data/PCGPrimitiveData.cpp
@@ -66,9 +66,9 @@ const UPCGPointData* UPCGPrimitiveData::CreatePointData(FPCGContext* Context) co
 	return Data;
 }
 
-UPCGSpatialData* UPCGPrimitiveData::CopyInternal() const
+UPCGSpatialData* UPCGPrimitiveData::CopyInternal(FPCGContext* Context) const
 {
-	UPCGPrimitiveData* NewPrimitiveData = NewObject<UPCGPrimitiveData>();
+	UPCGPrimitiveData* NewPrimitiveData = FPCGContext::NewObject_AnyThread<UPCGPrimitiveData>(Context);
 
 	NewPrimitiveData->VoxelSize = VoxelSize;
 	NewPrimitiveData->Primitive = Primitive;
diff --git a/Engine/Plugins/PCG/Source/PCG/Private/Data/PCGProjectionData.cpp b/Engine/Plugins/PCG/Source/PCG/Private/Data/PCGProjectionData.cpp
--- a/Engine/Plugins/PCG/Source/PCG/Private/Data/PCGProjectionData.cpp
+++ b/Engine/Plugins/PCG/Source/PCG/Private/Data/PCGProjectionData.cpp
@@ -2,6 +2,7 @@
 
 #include "Data/PCGProjectionData.h"
 
+#include "PCGContext.h"
 #include "Data/PCGLandscapeSplineData.h"
 #include "Data/PCGPointData.h"
 #include "Data/PCGSplineData.h"
@@ -235,7 +236,7 @@ const UPCGPointData* UPCGProjectionData::CreatePointData(FPCGContext* Context) c
 	const UPCGMetadata* SourceMetadata = SourcePointData->Metadata;
 	const TArray<FPCGPoint>& SourcePoints = SourcePointData->GetPoints();
 
-	UPCGPointData* PointData = NewObject<UPCGPointData>();
+	UPCGPointData* PointData = FPCGContext::NewObject_AnyThread<UPCGPointData>(Context);
 
 	// Copy metadata attributes from source point including values
 	PointData->InitializeFromData(this, SourceMetadata);
@@ -246,7 +247,7 @@ const UPCGPointData* UPCGProjectionData::CreatePointData(FPCGContext* Context) c
 	UPCGMetadata* TempTargetMetadata = nullptr;
 	if (Target->Metadata)
 	{
-		TempTargetMetadata = NewObject<UPCGMetadata>();
+		TempTargetMetadata = FPCGContext::NewObject_AnyThread<UPCGMetadata>(Context);
 
 		// We achieve filtering of metadata attributes by manipulating this temporary metadata, which works because
 		// the projection operation operates on the attributes in this metadata.
@@ -367,9 +368,9 @@ void UPCGProjectionData::CopyBaseProjectionClass(UPCGProjectionData* NewProjecti
 	NewProjectionData->ProjectionParams = ProjectionParams;
 }
 
-UPCGSpatialData* UPCGProjectionData::CopyInternal() const
+UPCGSpatialData* UPCGProjectionData::CopyInternal(FPCGContext* Context) const
 {
-	UPCGProjectionData* NewProjectionData = NewObject<UPCGProjectionData>();
+	UPCGProjectionData* NewProjectionData = FPCGContext::NewObject_AnyThread<UPCGProjectionData>(Context);
 
 	CopyBaseProjectionClass(NewProjectionData);
 
diff --git a/Engine/Plugins/PCG/Source/PCG/Private/Data/PCGRenderTargetData.cpp b/Engine/Plugins/PCG/Source/PCG/Private/Data/PCGRenderTargetData.cpp
--- a/Engine/Plugins/PCG/Source/PCG/Private/Data/PCGRenderTargetData.cpp
+++ b/Engine/Plugins/PCG/Source/PCG/Private/Data/PCGRenderTargetData.cpp
@@ -1,6 +1,8 @@
 // Copyright Epic Games, Inc. All Rights Reserved.
 
 #include "Data/PCGRenderTargetData.h"
+
+#include "PCGContext.h"
 #include "Engine/TextureRenderTarget2D.h"
 #include "TextureResource.h"
 
@@ -50,9 +52,9 @@ void UPCGRenderTargetData::AddToCrc(FArchiveCrc32& Ar, bool bFullDataCrc) const
 	AddUIDToCrc(Ar);
 }
 
-UPCGSpatialData* UPCGRenderTargetData::CopyInternal() const
+UPCGSpatialData* UPCGRenderTargetData::CopyInternal(FPCGContext* Context) const
 {
-	UPCGRenderTargetData* NewRenderTargetData = NewObject<UPCGRenderTargetData>();
+	UPCGRenderTargetData* NewRenderTargetData = FPCGContext::NewObject_AnyThread<UPCGRenderTargetData>(Context);
 
 	CopyBaseTextureData(NewRenderTargetData);
 
diff --git a/Engine/Plugins/PCG/Source/PCG/Private/Data/PCGSpatialData.cpp b/Engine/Plugins/PCG/Source/PCG/Private/Data/PCGSpatialData.cpp
--- a/Engine/Plugins/PCG/Source/PCG/Private/Data/PCGSpatialData.cpp
+++ b/Engine/Plugins/PCG/Source/PCG/Private/Data/PCGSpatialData.cpp
@@ -1,11 +1,14 @@
 // Copyright Epic Games, Inc. All Rights Reserved.
 
 #include "Data/PCGSpatialData.h"
+
+#include "PCGContext.h"
 #include "Data/PCGDifferenceData.h"
 #include "Data/PCGIntersectionData.h"
 #include "Data/PCGPointData.h"
 #include "Data/PCGProjectionData.h"
 #include "Data/PCGUnionData.h"
+#include "Elements/PCGExecuteBlueprint.h"
 
 #include "Serialization/ArchiveCrc32.h"
 
@@ -174,55 +177,75 @@ void UPCGSpatialData::ProjectPoints(const TArrayView<const TPair<FTransform, FBo
 	}
 }
 
-UPCGIntersectionData* UPCGSpatialData::IntersectWith(const UPCGSpatialData* InOther) const
+UPCGIntersectionData* UPCGSpatialData::K2_IntersectWith(const UPCGSpatialData* InOther) const
+{
+	return IntersectWith(UPCGBlueprintElement::ResolveContext(), InOther);
+}
+
+UPCGIntersectionData* UPCGSpatialData::IntersectWith(FPCGContext* InContext, const UPCGSpatialData* InOther) const
 {
-	UPCGIntersectionData* IntersectionData = NewObject<UPCGIntersectionData>();
+	UPCGIntersectionData* IntersectionData = FPCGContext::NewObject_AnyThread<UPCGIntersectionData>(InContext);
 	IntersectionData->Initialize(this, InOther);
 
 	return IntersectionData;
 }
 
-UPCGSpatialData* UPCGSpatialData::ProjectOn(const UPCGSpatialData* InOther, const FPCGProjectionParams& InParams) const
+UPCGSpatialData* UPCGSpatialData::K2_ProjectOn(const UPCGSpatialData* InOther, const FPCGProjectionParams& InParams) const
+{
+	return ProjectOn(UPCGBlueprintElement::ResolveContext(), InOther, InParams);
+}
+
+UPCGSpatialData* UPCGSpatialData::ProjectOn(FPCGContext* InContext, const UPCGSpatialData* InOther, const FPCGProjectionParams& InParams) const
 {
 	// Check necessary conditions. Fail to project -> return copy of projection source, i.e. projection not performed.
 	if (!InOther)
 	{
 		UE_LOG(LogPCG, Warning, TEXT("No projection target specified, no projection will occur"));
-		return DuplicateData();
+		return DuplicateData(InContext);
 	}
 
 	if (GetDimension() > InOther->GetDimension())
 	{
 		UE_LOG(LogPCG, Error, TEXT("Dimension of projection source (%d) must be less than or equal to that of the projection target (%d)"), GetDimension(), InOther->GetDimension());
-		return DuplicateData();
+		return DuplicateData(InContext);
 	}
 
 	const UPCGSpatialData* ConcreteTarget = InOther->FindFirstConcreteShapeFromNetwork();
 	if (!ConcreteTarget)
 	{
 		UE_LOG(LogPCG, Error, TEXT("Could not find a concrete shape in the target data to project onto."));
-		return DuplicateData();
+		return DuplicateData(InContext);
 	}
 
-	UPCGProjectionData* ProjectionData = NewObject<UPCGProjectionData>();
+	UPCGProjectionData* ProjectionData = FPCGContext::NewObject_AnyThread<UPCGProjectionData>(InContext);
 	ProjectionData->Initialize(this, ConcreteTarget, InParams);
 
 	return ProjectionData;
 }
 
-UPCGUnionData* UPCGSpatialData::UnionWith(const UPCGSpatialData* InOther) const
+UPCGUnionData* UPCGSpatialData::K2_UnionWith(const UPCGSpatialData* InOther) const
 {
-	UPCGUnionData* UnionData = NewObject<UPCGUnionData>();
+	return UnionWith(UPCGBlueprintElement::ResolveContext(), InOther);
+}
+
+UPCGUnionData* UPCGSpatialData::UnionWith(FPCGContext* InContext, const UPCGSpatialData* InOther) const
+{
+	UPCGUnionData* UnionData = FPCGContext::NewObject_AnyThread<UPCGUnionData>(InContext);
 	UnionData->Initialize(this, InOther);
 
 	return UnionData;
 }
 
-UPCGDifferenceData* UPCGSpatialData::Subtract(const UPCGSpatialData* InOther) const
+UPCGDifferenceData* UPCGSpatialData::K2_Subtract(const UPCGSpatialData* InOther) const
+{
+	return Subtract(UPCGBlueprintElement::ResolveContext(), InOther);
+}
+
+UPCGDifferenceData* UPCGSpatialData::Subtract(FPCGContext* InContext, const UPCGSpatialData* InOther) const
 {
-	UPCGDifferenceData* DifferenceData = NewObject<UPCGDifferenceData>();
+	UPCGDifferenceData* DifferenceData = FPCGContext::NewObject_AnyThread<UPCGDifferenceData>(InContext);
 	DifferenceData->Initialize(this);
-	DifferenceData->AddDifference(InOther);
+	DifferenceData->AddDifference(InContext, InOther);
 
 	return DifferenceData;
 }
@@ -261,9 +284,9 @@ void UPCGSpatialData::InitializeFromData(const UPCGSpatialData* InSource, const
 	}
 }
 
-UPCGSpatialData* UPCGSpatialData::DuplicateData(bool bInitializeMetadata) const
+UPCGSpatialData* UPCGSpatialData::DuplicateData(FPCGContext* Context, bool bInitializeMetadata) const
 {
-	UPCGSpatialData* NewSpatialData = CopyInternal();
+	UPCGSpatialData* NewSpatialData = CopyInternal(Context);
 	check(NewSpatialData);
 
 	if (bInitializeMetadata)
diff --git a/Engine/Plugins/PCG/Source/PCG/Private/Data/PCGSplineData.cpp b/Engine/Plugins/PCG/Source/PCG/Private/Data/PCGSplineData.cpp
--- a/Engine/Plugins/PCG/Source/PCG/Private/Data/PCGSplineData.cpp
+++ b/Engine/Plugins/PCG/Source/PCG/Private/Data/PCGSplineData.cpp
@@ -2,6 +2,7 @@
 
 #include "Data/PCGSplineData.h"
 
+#include "PCGContext.h"
 #include "Data/PCGPointData.h"
 #include "Data/PCGPolyLineData.h"
 #include "Data/PCGProjectionData.h"
@@ -155,7 +156,7 @@ FVector::FReal UPCGSplineData::GetDistanceAtSegmentStart(int SegmentIndex) const
 const UPCGPointData* UPCGSplineData::CreatePointData(FPCGContext* Context) const
 {
 	TRACE_CPUPROFILER_EVENT_SCOPE(UPCGSplineData::CreatePointData);
-	UPCGPointData* Data = NewObject<UPCGPointData>();
+	UPCGPointData* Data = FPCGContext::NewObject_AnyThread<UPCGPointData>(Context);
 	Data->InitializeFromData(this);
 
 	FPCGSplineSamplerParams SamplerParams;
@@ -202,23 +203,23 @@ bool UPCGSplineData::SamplePoint(const FTransform& InTransform, const FBox& InBo
 	}
 }
 
-UPCGSpatialData* UPCGSplineData::ProjectOn(const UPCGSpatialData* InOther, const FPCGProjectionParams& InParams) const
+UPCGSpatialData* UPCGSplineData::ProjectOn(FPCGContext* InContext, const UPCGSpatialData* InOther, const FPCGProjectionParams& InParams) const
 {
 	if (InOther->GetDimension() == 2)
 	{
-		UPCGSplineProjectionData* SplineProjectionData = NewObject<UPCGSplineProjectionData>();
+		UPCGSplineProjectionData* SplineProjectionData = FPCGContext::NewObject_AnyThread<UPCGSplineProjectionData>(InContext);
 		SplineProjectionData->Initialize(this, InOther, InParams);
 		return SplineProjectionData;
 	}
 	else
 	{
-		return Super::ProjectOn(InOther, InParams);
+		return Super::ProjectOn(InContext, InOther, InParams);
 	}
 }
 
-UPCGSpatialData* UPCGSplineData::CopyInternal() const
+UPCGSpatialData* UPCGSplineData::CopyInternal(FPCGContext* Context) const
 {
-	UPCGSplineData* NewSplineData = NewObject<UPCGSplineData>();
+	UPCGSplineData* NewSplineData = FPCGContext::NewObject_AnyThread<UPCGSplineData>(Context);
 
 	CopySplineData(NewSplineData);
 
@@ -380,9 +381,9 @@ const UPCGSpatialData* UPCGSplineProjectionData::GetSurface() const
 	return Target;
 }
 
-UPCGSpatialData* UPCGSplineProjectionData::CopyInternal() const
+UPCGSpatialData* UPCGSplineProjectionData::CopyInternal(FPCGContext* Context) const
 {
-	UPCGSplineProjectionData* NewProjectionData = NewObject<UPCGSplineProjectionData>();
+	UPCGSplineProjectionData* NewProjectionData = FPCGContext::NewObject_AnyThread<UPCGSplineProjectionData>(Context);
 
 	CopyBaseProjectionClass(NewProjectionData);
 
diff --git a/Engine/Plugins/PCG/Source/PCG/Private/Data/PCGSplineInteriorSurfaceData.cpp b/Engine/Plugins/PCG/Source/PCG/Private/Data/PCGSplineInteriorSurfaceData.cpp
--- a/Engine/Plugins/PCG/Source/PCG/Private/Data/PCGSplineInteriorSurfaceData.cpp
+++ b/Engine/Plugins/PCG/Source/PCG/Private/Data/PCGSplineInteriorSurfaceData.cpp
@@ -2,26 +2,27 @@
 
 #include "Data/PCGSplineInteriorSurfaceData.h"
 
+#include "PCGContext.h"
 #include "Data/PCGPointData.h"
 #include "Data/PCGSplineData.h"
 #include "Elements/PCGSplineSampler.h"
 
 #include UE_INLINE_GENERATED_CPP_BY_NAME(PCGSplineInteriorSurfaceData)
 
-void UPCGSplineInteriorSurfaceData::Initialize(const UPCGSplineData* InSplineData)
+void UPCGSplineInteriorSurfaceData::Initialize(FPCGContext* Context, const UPCGSplineData* InSplineData)
 {
 	check(InSplineData);
 
 	InitializeFromData(InSplineData);
 	SplineStruct = InSplineData->SplineStruct;
 
-	CacheData();
+	CacheData(Context);
 }
 
 void UPCGSplineInteriorSurfaceData::PostLoad()
 {
 	Super::PostLoad();
-	CacheData();
+	CacheData(nullptr);
 }
 
 void UPCGSplineInteriorSurfaceData::AddToCrc(FArchiveCrc32& Ar, bool bFullDataCrc) const
@@ -93,14 +94,14 @@ const UPCGPointData* UPCGSplineInteriorSurfaceData::CreatePointData(FPCGContext*
 {
 	TRACE_CPUPROFILER_EVENT_SCOPE(UPCGSplineData::CreatePointData);
 
-	UPCGPointData* Data = NewObject<UPCGPointData>();
+	UPCGPointData* Data = FPCGContext::NewObject_AnyThread<UPCGPointData>(Context);
 	Data->InitializeFromData(this);
 
 	FPCGSplineSamplerParams SamplerParams;
 	SamplerParams.Dimension = EPCGSplineSamplingDimension::OnInterior;
 	SamplerParams.bProjectOntoSurface = true;
 
-	UPCGSplineData* SplineData = NewObject<UPCGSplineData>();
+	UPCGSplineData* SplineData = FPCGContext::NewObject_AnyThread<UPCGSplineData>(Context);
 	SplineData->Initialize(SplineStruct);
 
 	// TODO: Ideally SampleInteriorData could just consume a UPCGSplineInteriorSurfaceData or FPCGSplineStruct, to avoid the extra copy on the spline struct.
@@ -112,9 +113,9 @@ const UPCGPointData* UPCGSplineInteriorSurfaceData::CreatePointData(FPCGContext*
 	return Data;
 }
 
-UPCGSpatialData* UPCGSplineInteriorSurfaceData::CopyInternal() const
+UPCGSpatialData* UPCGSplineInteriorSurfaceData::CopyInternal(FPCGContext* Context) const
 {
-	UPCGSplineInteriorSurfaceData* NewData = NewObject<UPCGSplineInteriorSurfaceData>();
+	UPCGSplineInteriorSurfaceData* NewData = FPCGContext::NewObject_AnyThread<UPCGSplineInteriorSurfaceData>(Context);
 
 	CopyBaseSurfaceData(NewData);
 	NewData->SplineStruct = SplineStruct;
@@ -128,7 +129,7 @@ UPCGSpatialData* UPCGSplineInteriorSurfaceData::CopyInternal() const
 	return NewData;
 }
 
-void UPCGSplineInteriorSurfaceData::CacheData()
+void UPCGSplineInteriorSurfaceData::CacheData(FPCGContext* Context)
 {
 #if WITH_EDITOR
 	bNeedsToCache = false;
@@ -152,13 +153,13 @@ void UPCGSplineInteriorSurfaceData::CacheData()
 	SamplerParams.Mode = EPCGSplineSamplingMode::Subdivision;
 	SamplerParams.SubdivisionsPerSegment = 5;
 
-	UPCGPointData* PointData = NewObject<UPCGPointData>();
-	UPCGSplineData* SplineData = NewObject<UPCGSplineData>();
+	UPCGPointData* PointData = FPCGContext::NewObject_AnyThread<UPCGPointData>(Context);
+	UPCGSplineData* SplineData = FPCGContext::NewObject_AnyThread<UPCGSplineData>(Context);
 	SplineData->Initialize(SplineStruct);
 
 	// TODO: Ideally SampleLineData could just consume a UPCGSplineInteriorSurfaceData or FPCGSplineStruct, to avoid the extra copy on the spline struct.
 	// TODO: It might be preferable to directly sample the spline than use the spline sampler, since it does a lot of work we don't care about.
-	PCGSplineSamplerHelpers::SampleLineData(/*Context=*/nullptr, SplineData, /*InBoundingShape=*/nullptr, /*InProjectionTarget=*/nullptr, /*InProjectionParams=*/{}, SamplerParams, PointData);
+	PCGSplineSamplerHelpers::SampleLineData(Context, SplineData, /*InBoundingShape=*/nullptr, /*InProjectionTarget=*/nullptr, /*InProjectionParams=*/{}, SamplerParams, PointData);
 
 	// Cache the points which describe the polygon of our spline.
 	for (const FPCGPoint& Point : PointData->GetPoints())
diff --git a/Engine/Plugins/PCG/Source/PCG/Private/Data/PCGTextureData.cpp b/Engine/Plugins/PCG/Source/PCG/Private/Data/PCGTextureData.cpp
--- a/Engine/Plugins/PCG/Source/PCG/Private/Data/PCGTextureData.cpp
+++ b/Engine/Plugins/PCG/Source/PCG/Private/Data/PCGTextureData.cpp
@@ -2,6 +2,7 @@
 
 #include "Data/PCGTextureData.h"
 
+#include "PCGContext.h"
 #include "PCGTextureReadback.h"
 #include "Data/PCGPointData.h"
 #include "Helpers/PCGAsync.h"
@@ -244,7 +245,7 @@ const UPCGPointData* UPCGBaseTextureData::CreatePointData(FPCGContext* Context)
 	// or based on a given texel size
 	FBox2D LocalSurfaceBounds(FVector2D(-1.0f, -1.0f), FVector2D(1.0f, 1.0f));
 
-	UPCGPointData* Data = NewObject<UPCGPointData>();
+	UPCGPointData* Data = FPCGContext::NewObject_AnyThread<UPCGPointData>(Context);
 	Data->InitializeFromData(this);
 
 	// Early out for invalid data
@@ -424,9 +425,9 @@ void UPCGTextureData::AddToCrc(FArchiveCrc32& Ar, bool bFullDataCrc) const
 	AddUIDToCrc(Ar);
 }
 
-UPCGSpatialData* UPCGTextureData::CopyInternal() const
+UPCGSpatialData* UPCGTextureData::CopyInternal(FPCGContext* Context) const
 {
-	UPCGTextureData* NewTextureData = NewObject<UPCGTextureData>();
+	UPCGTextureData* NewTextureData = FPCGContext::NewObject_AnyThread<UPCGTextureData>(Context);
 
 	CopyBaseTextureData(NewTextureData);
 
diff --git a/Engine/Plugins/PCG/Source/PCG/Private/Data/PCGUnionData.cpp b/Engine/Plugins/PCG/Source/PCG/Private/Data/PCGUnionData.cpp
--- a/Engine/Plugins/PCG/Source/PCG/Private/Data/PCGUnionData.cpp
+++ b/Engine/Plugins/PCG/Source/PCG/Private/Data/PCGUnionData.cpp
@@ -1,6 +1,8 @@
 // Copyright Epic Games, Inc. All Rights Reserved.
 
 #include "Data/PCGUnionData.h"
+
+#include "PCGContext.h"
 #include "Data/PCGPointData.h"
 #include "Data/PCGSpatialData.h"
 #include "Helpers/PCGAsync.h"
@@ -262,7 +264,7 @@ const UPCGPointData* UPCGUnionData::CreatePointData(FPCGContext* Context) const
 		InputMetadatas[i] = DataRawPtr[i]->ToPointData(Context)->Metadata;
 	}
 
-	UPCGPointData* PointData = NewObject<UPCGPointData>();
+	UPCGPointData* PointData = FPCGContext::NewObject_AnyThread<UPCGPointData>(Context);
 	PointData->InitializeFromData(this, InputMetadatas[0]);
 
 	UPCGMetadata* OutMetadata = PointData->Metadata;
@@ -430,9 +432,9 @@ void UPCGUnionData::CreateSequentialPointData(FPCGContext* Context, TArray<const
 	}
 }
 
-UPCGSpatialData* UPCGUnionData::CopyInternal() const
+UPCGSpatialData* UPCGUnionData::CopyInternal(FPCGContext* Context) const
 {
-	UPCGUnionData* NewUnionData = NewObject<UPCGUnionData>();
+	UPCGUnionData* NewUnionData = FPCGContext::NewObject_AnyThread<UPCGUnionData>(Context);
 
 	NewUnionData->Data = Data;
 	NewUnionData->FirstNonTrivialTransformData = FirstNonTrivialTransformData;
diff --git a/Engine/Plugins/PCG/Source/PCG/Private/Data/PCGVolumeData.cpp b/Engine/Plugins/PCG/Source/PCG/Private/Data/PCGVolumeData.cpp
--- a/Engine/Plugins/PCG/Source/PCG/Private/Data/PCGVolumeData.cpp
+++ b/Engine/Plugins/PCG/Source/PCG/Private/Data/PCGVolumeData.cpp
@@ -165,9 +165,9 @@ void UPCGVolumeData::CopyBaseVolumeData(UPCGVolumeData* NewVolumeData) const
 	NewVolumeData->StrictBounds = StrictBounds;
 }
 
-UPCGSpatialData* UPCGVolumeData::CopyInternal() const
+UPCGSpatialData* UPCGVolumeData::CopyInternal(FPCGContext* Context) const
 {
-	UPCGVolumeData* NewVolumeData = NewObject<UPCGVolumeData>();
+	UPCGVolumeData* NewVolumeData = FPCGContext::NewObject_AnyThread<UPCGVolumeData>(Context);
 
 	CopyBaseVolumeData(NewVolumeData);
 
diff --git a/Engine/Plugins/PCG/Source/PCG/Private/Data/PCGWorldData.cpp b/Engine/Plugins/PCG/Source/PCG/Private/Data/PCGWorldData.cpp
--- a/Engine/Plugins/PCG/Source/PCG/Private/Data/PCGWorldData.cpp
+++ b/Engine/Plugins/PCG/Source/PCG/Private/Data/PCGWorldData.cpp
@@ -144,7 +144,7 @@ const UPCGPointData* UPCGWorldVolumetricData::CreatePointData(FPCGContext* Conte
 {
 	TRACE_CPUPROFILER_EVENT_SCOPE(UPCGWorldVolumetricData::CreatePointData);
 
-	UPCGPointData* Data = NewObject<UPCGPointData>();
+	UPCGPointData* Data = FPCGContext::NewObject_AnyThread<UPCGPointData>(Context);
 	Data->InitializeFromData(this);
 	TArray<FPCGPoint>& Points = Data->GetMutablePoints();
 
@@ -182,9 +182,9 @@ const UPCGPointData* UPCGWorldVolumetricData::CreatePointData(FPCGContext* Conte
 	return Data;
 }
 
-UPCGSpatialData* UPCGWorldVolumetricData::CopyInternal() const
+UPCGSpatialData* UPCGWorldVolumetricData::CopyInternal(FPCGContext* Context) const
 {
-	UPCGWorldVolumetricData* NewVolumetricData = NewObject<UPCGWorldVolumetricData>();
+	UPCGWorldVolumetricData* NewVolumetricData = FPCGContext::NewObject_AnyThread<UPCGWorldVolumetricData>(Context);
 
 	CopyBaseVolumeData(NewVolumetricData);
 
@@ -248,7 +248,7 @@ bool UPCGWorldRayHitData::SamplePoint(const FTransform& InTransform, const FBox&
 
 const UPCGPointData* UPCGWorldRayHitData::CreatePointData(FPCGContext* Context, const FBox& InBounds) const
 {
-	UPCGPointData* Data = NewObject<UPCGPointData>();
+	UPCGPointData* Data = FPCGContext::NewObject_AnyThread<UPCGPointData>(Context);
 	Data->InitializeFromData(this);
 
 	FBox EffectiveBounds = Bounds;
@@ -282,9 +282,9 @@ const UPCGPointData* UPCGWorldRayHitData::CreatePointData(FPCGContext* Context,
 	return Data;
 }
 
-UPCGSpatialData* UPCGWorldRayHitData::CopyInternal() const
+UPCGSpatialData* UPCGWorldRayHitData::CopyInternal(FPCGContext* Context) const
 {
-	UPCGWorldRayHitData* NewData = NewObject<UPCGWorldRayHitData>();
+	UPCGWorldRayHitData* NewData = FPCGContext::NewObject_AnyThread<UPCGWorldRayHitData>(Context);
 
 	CopyBaseSurfaceData(NewData);
 
diff --git a/Engine/Plugins/PCG/Source/PCG/Private/Elements/IO/PCGDataTableElement.cpp b/Engine/Plugins/PCG/Source/PCG/Private/Elements/IO/PCGDataTableElement.cpp
--- a/Engine/Plugins/PCG/Source/PCG/Private/Elements/IO/PCGDataTableElement.cpp
+++ b/Engine/Plugins/PCG/Source/PCG/Private/Elements/IO/PCGDataTableElement.cpp
@@ -127,14 +127,14 @@ bool FPCGLoadDataTableElement::PrepareLoad(FPCGExternalDataContext* Context) con
 
 	if (Settings->OutputType == EPCGExclusiveDataType::Point)
 	{
-		UPCGPointData* PointData = NewObject<UPCGPointData>();
+		UPCGPointData* PointData = FPCGContext::NewObject_AnyThread<UPCGPointData>(Context);
 		check(PointData);
 		OutData = PointData;
 		OutMetadata = PointData->MutableMetadata();
 	}
 	else if(Settings->OutputType == EPCGExclusiveDataType::Param)
 	{
-		UPCGParamData* ParamData = NewObject<UPCGParamData>();
+		UPCGParamData* ParamData = FPCGContext::NewObject_AnyThread<UPCGParamData>(Context);
 		check(ParamData);
 		OutData = ParamData;
 		OutMetadata = ParamData->MutableMetadata();
diff --git a/Engine/Plugins/PCG/Source/PCG/Private/Elements/Metadata/PCGMetadataElement.cpp b/Engine/Plugins/PCG/Source/PCG/Private/Elements/Metadata/PCGMetadataElement.cpp
--- a/Engine/Plugins/PCG/Source/PCG/Private/Elements/Metadata/PCGMetadataElement.cpp
+++ b/Engine/Plugins/PCG/Source/PCG/Private/Elements/Metadata/PCGMetadataElement.cpp
@@ -160,7 +160,7 @@ bool FPCGMetadataOperationElement::ExecuteInternal(FPCGContext* Context) const
 			continue;
 		}
 
-		UPCGData* OutputData = OriginalData->DuplicateData();
+		UPCGData* OutputData = OriginalData->DuplicateData(Context);
 		bool bSuccess = false;
 		if (Settings->bCopyAllAttributes)
 		{
diff --git a/Engine/Plugins/PCG/Source/PCG/Private/Elements/Metadata/PCGMetadataElementCommon.cpp b/Engine/Plugins/PCG/Source/PCG/Private/Elements/Metadata/PCGMetadataElementCommon.cpp
--- a/Engine/Plugins/PCG/Source/PCG/Private/Elements/Metadata/PCGMetadataElementCommon.cpp
+++ b/Engine/Plugins/PCG/Source/PCG/Private/Elements/Metadata/PCGMetadataElementCommon.cpp
@@ -8,11 +8,16 @@
 namespace PCGMetadataElementCommon
 {
 	void DuplicateTaggedData(const FPCGTaggedData& InTaggedData, FPCGTaggedData& OutTaggedData, UPCGMetadata*& OutMetadata)
+	{
+		DuplicateTaggedData(nullptr, InTaggedData, OutTaggedData, OutMetadata);
+	}
+
+	void DuplicateTaggedData(FPCGContext* InContext, const FPCGTaggedData& InTaggedData, FPCGTaggedData& OutTaggedData, UPCGMetadata*& OutMetadata)
 	{
 		TRACE_CPUPROFILER_EVENT_SCOPE(PCGMetadataElementCommon::DuplicateTaggedData);
 		if (InTaggedData.Data)
 		{
-			UPCGData* NewData = InTaggedData.Data->DuplicateData();
+			UPCGData* NewData = InTaggedData.Data->DuplicateData(InContext);
 			check(NewData);
 			OutTaggedData.Data = NewData;
 			OutMetadata = NewData->MutableMetadata();
diff --git a/Engine/Plugins/PCG/Source/PCG/Private/Elements/Metadata/PCGMetadataMakeRotator.cpp b/Engine/Plugins/PCG/Source/PCG/Private/Elements/Metadata/PCGMetadataMakeRotator.cpp
--- a/Engine/Plugins/PCG/Source/PCG/Private/Elements/Metadata/PCGMetadataMakeRotator.cpp
+++ b/Engine/Plugins/PCG/Source/PCG/Private/Elements/Metadata/PCGMetadataMakeRotator.cpp
@@ -132,14 +132,14 @@ bool UPCGMetadataMakeRotatorSettings::DoesInputSupportDefaultValue(uint32 Index)
 	return Operation == EPCGMetadataMakeRotatorOp::MakeRotFromAngles;
 }
 
-UPCGParamData* UPCGMetadataMakeRotatorSettings::CreateDefaultValueParam(uint32 Index) const
+UPCGParamData* UPCGMetadataMakeRotatorSettings::CreateDefaultValueParam(FPCGContext* Context, uint32 Index) const
 {
 	if (Operation != EPCGMetadataMakeRotatorOp::MakeRotFromAngles)
 	{
 		return nullptr;
 	}
 
-	UPCGParamData* NewParamData = NewObject<UPCGParamData>();
+	UPCGParamData* NewParamData = FPCGContext::NewObject_AnyThread<UPCGParamData>(Context);
 	NewParamData->Metadata->CreateAttribute<double>(NAME_None, 0, /*bAllowsInterpolation=*/ true, /*bOverrideParent=*/ false);
 	return NewParamData;
 }
diff --git a/Engine/Plugins/PCG/Source/PCG/Private/Elements/Metadata/PCGMetadataMakeTransform.cpp b/Engine/Plugins/PCG/Source/PCG/Private/Elements/Metadata/PCGMetadataMakeTransform.cpp
--- a/Engine/Plugins/PCG/Source/PCG/Private/Elements/Metadata/PCGMetadataMakeTransform.cpp
+++ b/Engine/Plugins/PCG/Source/PCG/Private/Elements/Metadata/PCGMetadataMakeTransform.cpp
@@ -126,28 +126,28 @@ bool UPCGMetadataMakeTransformSettings::DoesInputSupportDefaultValue(uint32 Inde
 	return true;
 }
 
-UPCGParamData* UPCGMetadataMakeTransformSettings::CreateDefaultValueParam(uint32 Index) const
+UPCGParamData* UPCGMetadataMakeTransformSettings::CreateDefaultValueParam(FPCGContext* Context, uint32 Index) const
 {
 	switch (Index)
 	{
 	case 0:
 	{
 		// Location -> Default is Zero vector
-		UPCGParamData* NewParamData = NewObject<UPCGParamData>();
+		UPCGParamData* NewParamData = FPCGContext::NewObject_AnyThread<UPCGParamData>(Context);
 		NewParamData->Metadata->CreateAttribute<FVector>(NAME_None, FVector::ZeroVector, /*bAllowsInterpolation=*/ true, /*bOverrideParent=*/ false);
 		return NewParamData;
 	}
 	case 1:
 	{
 		// Rotation -> Default is Zero rotator
-		UPCGParamData* NewParamData = NewObject<UPCGParamData>();
+		UPCGParamData* NewParamData = FPCGContext::NewObject_AnyThread<UPCGParamData>(Context);
 		NewParamData->Metadata->CreateAttribute<FRotator>(NAME_None, FRotator::ZeroRotator, /*bAllowsInterpolation=*/ true, /*bOverrideParent=*/ false);
 		return NewParamData;
 	}
 	case 2:
 	{
 		// Scale -> Default is Vector (1, 1, 1)
-		UPCGParamData* NewParamData = NewObject<UPCGParamData>();
+		UPCGParamData* NewParamData = FPCGContext::NewObject_AnyThread<UPCGParamData>(Context);
 		NewParamData->Metadata->CreateAttribute<FVector>(NAME_None, FVector::OneVector, /*bAllowsInterpolation=*/ true, /*bOverrideParent=*/ false);
 		return NewParamData;
 	}
diff --git a/Engine/Plugins/PCG/Source/PCG/Private/Elements/Metadata/PCGMetadataMakeVector.cpp b/Engine/Plugins/PCG/Source/PCG/Private/Elements/Metadata/PCGMetadataMakeVector.cpp
--- a/Engine/Plugins/PCG/Source/PCG/Private/Elements/Metadata/PCGMetadataMakeVector.cpp
+++ b/Engine/Plugins/PCG/Source/PCG/Private/Elements/Metadata/PCGMetadataMakeVector.cpp
@@ -201,11 +201,12 @@ bool UPCGMetadataMakeVectorSettings::DoesInputSupportDefaultValue(uint32 Index)
 	return true;
 }
 
-UPCGParamData* UPCGMetadataMakeVectorSettings::CreateDefaultValueParam(uint32 Index) const
+UPCGParamData* UPCGMetadataMakeVectorSettings::CreateDefaultValueParam(FPCGContext* Context, uint32 Index) const
 {
 	// Use labels since the logic is already done there.
 	FName Label = GetInputPinLabel(Index);
-	UPCGParamData* NewParamData = NewObject<UPCGParamData>();
+
+	UPCGParamData* NewParamData = FPCGContext::NewObject_AnyThread<UPCGParamData>(Context);
 
 	if (Label == PCGMetadataMakeVectorConstants::XYZLabel)
 	{
diff --git a/Engine/Plugins/PCG/Source/PCG/Private/Elements/Metadata/PCGMetadataOpElementBase.cpp b/Engine/Plugins/PCG/Source/PCG/Private/Elements/Metadata/PCGMetadataOpElementBase.cpp
--- a/Engine/Plugins/PCG/Source/PCG/Private/Elements/Metadata/PCGMetadataOpElementBase.cpp
+++ b/Engine/Plugins/PCG/Source/PCG/Private/Elements/Metadata/PCGMetadataOpElementBase.cpp
@@ -464,7 +464,7 @@ bool FPCGMetadataElementBase::PrepareDataInternal(FPCGContext* Context) const
 			{
 				FPCGTaggedData& DefaultData = CurrentPinInputData.Emplace_GetRef();
 				DefaultData.Pin = CurrentPinLabel;
-				DefaultData.Data = Settings->CreateDefaultValueParam(OperandPinIndex);
+				DefaultData.Data = Settings->CreateDefaultValueParam(Context, OperandPinIndex);
 
 				if (!DefaultData.Data)
 				{
@@ -630,7 +630,7 @@ bool FPCGMetadataElementBase::PrepareDataInternal(FPCGContext* Context) const
 			}
 
 			check(InputTaggedData[PrimaryPinIndex].Data);
-			UPCGData* OutputData = InputTaggedData[PrimaryPinIndex].Data->DuplicateData();
+			UPCGData* OutputData = InputTaggedData[PrimaryPinIndex].Data->DuplicateData(Context);
 			check(OutputData);
 			OutputTaggedData.Data = OutputData;
 
diff --git a/Engine/Plugins/PCG/Source/PCG/Private/Elements/Metadata/PCGMetadataRenameElement.cpp b/Engine/Plugins/PCG/Source/PCG/Private/Elements/Metadata/PCGMetadataRenameElement.cpp
--- a/Engine/Plugins/PCG/Source/PCG/Private/Elements/Metadata/PCGMetadataRenameElement.cpp
+++ b/Engine/Plugins/PCG/Source/PCG/Private/Elements/Metadata/PCGMetadataRenameElement.cpp
@@ -92,7 +92,7 @@ bool FPCGMetadataRenameElement::ExecuteInternal(FPCGContext* Context) const
 		}
 
 		UPCGMetadata* NewMetadata = nullptr;
-		PCGMetadataElementCommon::DuplicateTaggedData(Input, Output, NewMetadata);
+		PCGMetadataElementCommon::DuplicateTaggedData(Context, Input, Output, NewMetadata);
 
 		if (!NewMetadata || !NewMetadata->RenameAttribute(LocalAttributeToRename, NewAttributeName))
 		{
diff --git a/Engine/Plugins/PCG/Source/PCG/Private/Elements/PCGAttributeCast.cpp b/Engine/Plugins/PCG/Source/PCG/Private/Elements/PCGAttributeCast.cpp
--- a/Engine/Plugins/PCG/Source/PCG/Private/Elements/PCGAttributeCast.cpp
+++ b/Engine/Plugins/PCG/Source/PCG/Private/Elements/PCGAttributeCast.cpp
@@ -92,7 +92,7 @@ bool FPCGAttributeCastElement::ExecuteInternal(FPCGContext* Context) const
 
 		FPCGTaggedData& Output = Context->OutputData.TaggedData.Add_GetRef(Input);
 
-		UPCGData* OutputData = Input.Data->DuplicateData();
+		UPCGData* OutputData = Input.Data->DuplicateData(Context);
 		check(OutputData);
 
 		PCGMetadataHelpers::FPCGCopyAttributeParams Params{};
diff --git a/Engine/Plugins/PCG/Source/PCG/Private/Elements/PCGAttributeFilter.cpp b/Engine/Plugins/PCG/Source/PCG/Private/Elements/PCGAttributeFilter.cpp
--- a/Engine/Plugins/PCG/Source/PCG/Private/Elements/PCGAttributeFilter.cpp
+++ b/Engine/Plugins/PCG/Source/PCG/Private/Elements/PCGAttributeFilter.cpp
@@ -175,7 +175,7 @@ namespace PCGAttributeFilterHelpers
 				// Reset the point data and reserving some points
 				// No need to reserve the full number of points, since we'll go by chunk
 				// Only allocate the chunk size
-				InOutThresholdInfo.ThresholdPointData = NewObject<UPCGPointData>();
+				InOutThresholdInfo.ThresholdPointData = FPCGContext::NewObject_AnyThread<UPCGPointData>(InContext);
 				InOutThresholdInfo.ThresholdPointData->InitializeFromData(InOutThresholdInfo.ThresholdSpatialData);
 				InOutThresholdInfo.ThresholdPointData->GetMutablePoints().SetNum(PCGAttributeFilterConstants::ChunkSize);
 			}
@@ -664,8 +664,8 @@ bool FPCGAttributeFilterElementBase::DoFiltering(FPCGContext* Context, EPCGAttri
 		if (OperationData.bIsInputPointData)
 		{
 			const UPCGPointData* OriginalPointData = CastChecked<UPCGPointData>(OriginalData);
-			UPCGPointData* InFilterPointData = NewObject<UPCGPointData>();
-			UPCGPointData* OutFilterPointData = NewObject<UPCGPointData>();
+			UPCGPointData* InFilterPointData = FPCGContext::NewObject_AnyThread<UPCGPointData>(Context);
+			UPCGPointData* OutFilterPointData = FPCGContext::NewObject_AnyThread<UPCGPointData>(Context);
 
 			OperationData.OriginalPoints = &OriginalPointData->GetPoints();
 
@@ -684,8 +684,8 @@ bool FPCGAttributeFilterElementBase::DoFiltering(FPCGContext* Context, EPCGAttri
 		{
 			// Param data
 			const UPCGParamData* OriginalParamData = CastChecked<UPCGParamData>(OriginalData);
-			UPCGParamData* InFilterParamData = NewObject<UPCGParamData>();
-			UPCGParamData* OutFilterParamData = NewObject<UPCGParamData>();
+			UPCGParamData* InFilterParamData = FPCGContext::NewObject_AnyThread<UPCGParamData>(Context);
+			UPCGParamData* OutFilterParamData = FPCGContext::NewObject_AnyThread<UPCGParamData>(Context);
 
 			OperationData.OriginalMetadata = OriginalParamData->Metadata;
 
diff --git a/Engine/Plugins/PCG/Source/PCG/Private/Elements/PCGAttributeGetFromIndexElement.cpp b/Engine/Plugins/PCG/Source/PCG/Private/Elements/PCGAttributeGetFromIndexElement.cpp
--- a/Engine/Plugins/PCG/Source/PCG/Private/Elements/PCGAttributeGetFromIndexElement.cpp
+++ b/Engine/Plugins/PCG/Source/PCG/Private/Elements/PCGAttributeGetFromIndexElement.cpp
@@ -81,7 +81,7 @@ bool FPCGAttributeGetFromIndexElement::ExecuteInternal(FPCGContext* Context) con
 			continue;
 		}
 
-		UPCGParamData* SubParam = NewObject<UPCGParamData>();
+		UPCGParamData* SubParam = FPCGContext::NewObject_AnyThread<UPCGParamData>(Context);
 		check(SubParam->Metadata);
 		SubParam->Metadata->AddAttributes(ParamMetadata);
 
diff --git a/Engine/Plugins/PCG/Source/PCG/Private/Elements/PCGAttributeGetFromPointIndexElement.cpp b/Engine/Plugins/PCG/Source/PCG/Private/Elements/PCGAttributeGetFromPointIndexElement.cpp
--- a/Engine/Plugins/PCG/Source/PCG/Private/Elements/PCGAttributeGetFromPointIndexElement.cpp
+++ b/Engine/Plugins/PCG/Source/PCG/Private/Elements/PCGAttributeGetFromPointIndexElement.cpp
@@ -118,7 +118,7 @@ bool FPCGAttributeGetFromPointIndexElement::ExecuteInternal(FPCGContext* Context
 
 		if (Accessor.IsValid())
 		{
-			UPCGParamData* OutputParamData = NewObject<UPCGParamData>();
+			UPCGParamData* OutputParamData = FPCGContext::NewObject_AnyThread<UPCGParamData>(Context);
 
 			auto ExtractAttribute = [this, Context, OutputAttributeName, &OutputParamData, &Accessor, &PointKey, InputIndex](auto DummyValue) -> bool
 			{
@@ -154,7 +154,7 @@ bool FPCGAttributeGetFromPointIndexElement::ExecuteInternal(FPCGContext* Context
 				if (Context->Node && Context->Node->IsOutputPinConnected(PCGAttributeGetFromPointIndexConstants::OutputPointLabel))
 #endif
 				{
-					UPCGPointData* OutputPointData = NewObject<UPCGPointData>();
+					UPCGPointData* OutputPointData = FPCGContext::NewObject_AnyThread<UPCGPointData>(Context);
 					OutputPointData->InitializeFromData(PointData);
 					OutputPointData->GetMutablePoints().Add(Point);
 
diff --git a/Engine/Plugins/PCG/Source/PCG/Private/Elements/PCGAttributeNoise.cpp b/Engine/Plugins/PCG/Source/PCG/Private/Elements/PCGAttributeNoise.cpp
--- a/Engine/Plugins/PCG/Source/PCG/Private/Elements/PCGAttributeNoise.cpp
+++ b/Engine/Plugins/PCG/Source/PCG/Private/Elements/PCGAttributeNoise.cpp
@@ -299,7 +299,7 @@ bool FPCGAttributeNoiseElement::ExecuteInternal(FPCGContext* InContext) const
 			TempInputAccessor.Reset();
 
 			FPCGTaggedData& Output = Outputs.Add_GetRef(Input);
-			UPCGData* OutputData = InputData->DuplicateData();
+			UPCGData* OutputData = InputData->DuplicateData(Context);
 			Output.Data = OutputData;
 
 			// Then create the accessor/keys
diff --git a/Engine/Plugins/PCG/Source/PCG/Private/Elements/PCGAttributeReduceElement.cpp b/Engine/Plugins/PCG/Source/PCG/Private/Elements/PCGAttributeReduceElement.cpp
--- a/Engine/Plugins/PCG/Source/PCG/Private/Elements/PCGAttributeReduceElement.cpp
+++ b/Engine/Plugins/PCG/Source/PCG/Private/Elements/PCGAttributeReduceElement.cpp
@@ -295,7 +295,7 @@ bool FPCGAttributeReduceElement::ExecuteInternal(FPCGContext* Context) const
 			continue;
 		}
 
-		auto DoOperation = [&Accessor, &Keys, Settings, &OutputParams, &NewAttribute, OutputAttributeName](auto DummyValue) -> bool
+		auto DoOperation = [&Accessor, &Keys, Settings, &OutputParams, &NewAttribute, OutputAttributeName, &Context](auto DummyValue) -> bool
 		{
 			using AttributeType = decltype(DummyValue);
 
@@ -305,7 +305,7 @@ bool FPCGAttributeReduceElement::ExecuteInternal(FPCGContext* Context) const
 			bool bCreatedNewParams = false;
 			if (!OutputParams || !bMergeOutputAttributes)
 			{
-				OutputParams = NewObject<UPCGParamData>();
+				OutputParams = FPCGContext::NewObject_AnyThread<UPCGParamData>(Context);
 				bCreatedNewParams = true;
 			}
 
diff --git a/Engine/Plugins/PCG/Source/PCG/Private/Elements/PCGAttributeSelectElement.cpp b/Engine/Plugins/PCG/Source/PCG/Private/Elements/PCGAttributeSelectElement.cpp
--- a/Engine/Plugins/PCG/Source/PCG/Private/Elements/PCGAttributeSelectElement.cpp
+++ b/Engine/Plugins/PCG/Source/PCG/Private/Elements/PCGAttributeSelectElement.cpp
@@ -307,7 +307,7 @@ bool FPCGAttributeSelectElement::ExecuteInternal(FPCGContext* Context) const
 	FPCGAttributePropertyInputSelector InputSource = Settings->InputSource.CopyAndFixLast(SpatialData);
 
 	const FName OutputAttributeName = (Settings->OutputAttributeName == PCGMetadataAttributeConstants::SourceNameAttributeName) ? InputSource.GetName() : Settings->OutputAttributeName;
-	UPCGParamData* OutputParamData = NewObject<UPCGParamData>();
+	UPCGParamData* OutputParamData = FPCGContext::NewObject_AnyThread<UPCGParamData>(Context);
 
 	TUniquePtr<const IPCGAttributeAccessor> Accessor = PCGAttributeAccessorHelpers::CreateConstAccessor(PointData, InputSource);
 	TUniquePtr<const IPCGAttributeAccessorKeys> Keys = PCGAttributeAccessorHelpers::CreateConstKeys(PointData, InputSource);
@@ -433,7 +433,7 @@ bool FPCGAttributeSelectElement::ExecuteInternal(FPCGContext* Context) const
 	if(PointData && Context->Node && Context->Node->IsOutputPinConnected(PCGAttributeSelectConstants::OutputPointLabel))
 #endif
 	{
-		UPCGPointData* OutputPointData = NewObject<UPCGPointData>();
+		UPCGPointData* OutputPointData = FPCGContext::NewObject_AnyThread<UPCGPointData>(Context);
 		OutputPointData->InitializeFromData(PointData);
 		OutputPointData->GetMutablePoints().Add(PointData->GetPoint(OutputIndex));
 
diff --git a/Engine/Plugins/PCG/Source/PCG/Private/Elements/PCGAttributeTransferElement.cpp b/Engine/Plugins/PCG/Source/PCG/Private/Elements/PCGAttributeTransferElement.cpp
--- a/Engine/Plugins/PCG/Source/PCG/Private/Elements/PCGAttributeTransferElement.cpp
+++ b/Engine/Plugins/PCG/Source/PCG/Private/Elements/PCGAttributeTransferElement.cpp
@@ -188,7 +188,7 @@ bool FPCGAttributeTransferElement::ExecuteInternal(FPCGContext* Context) const
 		return true;
 	}
 
-	UPCGSpatialData* OutputData = TargetData->DuplicateData();
+	UPCGSpatialData* OutputData = TargetData->DuplicateData(Context);
 	check(OutputData->Metadata);
 
 	auto AppendOutputData = [Context, OutputData, &TargetInputs]()
diff --git a/Engine/Plugins/PCG/Source/PCG/Private/Elements/PCGBlurElement.cpp b/Engine/Plugins/PCG/Source/PCG/Private/Elements/PCGBlurElement.cpp
--- a/Engine/Plugins/PCG/Source/PCG/Private/Elements/PCGBlurElement.cpp
+++ b/Engine/Plugins/PCG/Source/PCG/Private/Elements/PCGBlurElement.cpp
@@ -93,7 +93,7 @@ bool FPCGBlurElement::PrepareDataInternal(FPCGContext* InContext) const
 			return EPCGTimeSliceInitResult::NoOperation;
 		}
 
-		UPCGPointData* OutputPointData = CastChecked<UPCGPointData>(OutState.InputPointData->DuplicateData());
+		UPCGPointData* OutputPointData = CastChecked<UPCGPointData>(OutState.InputPointData->DuplicateData(InContext));
 
 		FPCGAttributePropertyOutputSelector OutputSelector = Settings->OutputTarget.CopyAndFixSource(&InputSelector, OutState.InputPointData);
 
diff --git a/Engine/Plugins/PCG/Source/PCG/Private/Elements/PCGBoundsFromMesh.cpp b/Engine/Plugins/PCG/Source/PCG/Private/Elements/PCGBoundsFromMesh.cpp
--- a/Engine/Plugins/PCG/Source/PCG/Private/Elements/PCGBoundsFromMesh.cpp
+++ b/Engine/Plugins/PCG/Source/PCG/Private/Elements/PCGBoundsFromMesh.cpp
@@ -262,7 +262,7 @@ bool FPCGBoundsFromMeshElement::ExecuteInternal(FPCGContext* InContext) const
 
 			const FBox StaticMeshBounds = SingleMesh->GetBoundingBox();
 			
-			UPCGPointData* OutputData = NewObject<UPCGPointData>();
+			UPCGPointData* OutputData = FPCGContext::NewObject_AnyThread<UPCGPointData>(Context);
 			OutputData->InitializeFromData(PointData);
 
 			const TArray<FPCGPoint>& InputPoints = PointData->GetPoints();
@@ -288,7 +288,7 @@ bool FPCGBoundsFromMeshElement::ExecuteInternal(FPCGContext* InContext) const
 			FPCGBoundsFromMeshContext::InputMeshesData* InputData = Context->PerInputData.FindByPredicate([Index = Context->CurrentExecuteIndex - 1](const FPCGBoundsFromMeshContext::InputMeshesData& InputData) { return InputData.InputIndex == Index; });
 			if (InputData)
 			{
-				UPCGPointData* OutputData = NewObject<UPCGPointData>();
+				UPCGPointData* OutputData = FPCGContext::NewObject_AnyThread<UPCGPointData>(Context);
 				OutputData->InitializeFromData(PointData);
 
 				const TArray<FPCGPoint>& InputPoints = PointData->GetPoints();
diff --git a/Engine/Plugins/PCG/Source/PCG/Private/Elements/PCGCollapseElement.cpp b/Engine/Plugins/PCG/Source/PCG/Private/Elements/PCGCollapseElement.cpp
--- a/Engine/Plugins/PCG/Source/PCG/Private/Elements/PCGCollapseElement.cpp
+++ b/Engine/Plugins/PCG/Source/PCG/Private/Elements/PCGCollapseElement.cpp
@@ -72,7 +72,7 @@ bool FPCGCollapseElement::ExecuteInternal(FPCGContext* Context) const
 				continue;
 			}
 
-			UPCGPointData* PointData = NewObject<UPCGPointData>();
+			UPCGPointData* PointData = FPCGContext::NewObject_AnyThread<UPCGPointData>(Context);
 			check(PointData->Metadata);
 			PointData->Metadata->Initialize(ParamMetadata);
 
diff --git a/Engine/Plugins/PCG/Source/PCG/Private/Elements/PCGCombinePoints.cpp b/Engine/Plugins/PCG/Source/PCG/Private/Elements/PCGCombinePoints.cpp
--- a/Engine/Plugins/PCG/Source/PCG/Private/Elements/PCGCombinePoints.cpp
+++ b/Engine/Plugins/PCG/Source/PCG/Private/Elements/PCGCombinePoints.cpp
@@ -48,7 +48,7 @@ bool FPCGCombinePointsElement::ExecuteInternal(FPCGContext* Context) const
 		const FMatrix InversePointTransform = PointTransform.ToInverseMatrixWithScale();
 
 		FPCGTaggedData& Output = Outputs.Add_GetRef(Inputs[i]);
-		UPCGPointData* OutputPointData = NewObject<UPCGPointData>();
+		UPCGPointData* OutputPointData = FPCGContext::NewObject_AnyThread<UPCGPointData>(Context);
 		OutputPointData->InitializeFromData(InputPointData);
 		FPCGPoint& OutputPoint = OutputPointData->GetMutablePoints().Emplace_GetRef();
 		Output.Data = OutputPointData;
diff --git a/Engine/Plugins/PCG/Source/PCG/Private/Elements/PCGConvertToAttributeSet.cpp b/Engine/Plugins/PCG/Source/PCG/Private/Elements/PCGConvertToAttributeSet.cpp
--- a/Engine/Plugins/PCG/Source/PCG/Private/Elements/PCGConvertToAttributeSet.cpp
+++ b/Engine/Plugins/PCG/Source/PCG/Private/Elements/PCGConvertToAttributeSet.cpp
@@ -63,7 +63,7 @@ bool FPCGConvertToAttributeSetElement::ExecuteInternal(FPCGContext* Context) con
 			continue;
 		}
 
-		UPCGParamData* ParamData = NewObject<UPCGParamData>();
+		UPCGParamData* ParamData = FPCGContext::NewObject_AnyThread<UPCGParamData>(Context);
 		check(ParamData->Metadata);
 		ParamData->Metadata->Initialize(SourceMetadata);
 
@@ -112,7 +112,7 @@ bool FPCGTagsToAttributeSetElement::ExecuteInternal(FPCGContext* Context) const
 			continue;
 		}
 
-		UPCGParamData* ParamData = NewObject<UPCGParamData>();
+		UPCGParamData* ParamData = FPCGContext::NewObject_AnyThread<UPCGParamData>(Context);
 		check(ParamData && ParamData->Metadata);
 
 		FPCGTaggedData& Output = Outputs.Emplace_GetRef(Input);
diff --git a/Engine/Plugins/PCG/Source/PCG/Private/Elements/PCGConvexHull2D.cpp b/Engine/Plugins/PCG/Source/PCG/Private/Elements/PCGConvexHull2D.cpp
--- a/Engine/Plugins/PCG/Source/PCG/Private/Elements/PCGConvexHull2D.cpp
+++ b/Engine/Plugins/PCG/Source/PCG/Private/Elements/PCGConvexHull2D.cpp
@@ -45,7 +45,7 @@ bool FPCGConvexHull2DElement::ExecuteInternal(FPCGContext* Context) const
 		ConvexHull2D::ComputeConvexHull(PointsPositions, ConvexHullIndices);
 
 		FPCGTaggedData& Output = Outputs.Add_GetRef(Input);
-		UPCGPointData* OutputPointData = NewObject<UPCGPointData>();
+		UPCGPointData* OutputPointData = FPCGContext::NewObject_AnyThread<UPCGPointData>(Context);
 		OutputPointData->InitializeFromData(PointData);
 		TArray<FPCGPoint>& OutPoints = OutputPointData->GetMutablePoints();
 		OutPoints.Reserve(ConvexHullIndices.Num());
diff --git a/Engine/Plugins/PCG/Source/PCG/Private/Elements/PCGCopyPoints.cpp b/Engine/Plugins/PCG/Source/PCG/Private/Elements/PCGCopyPoints.cpp
--- a/Engine/Plugins/PCG/Source/PCG/Private/Elements/PCGCopyPoints.cpp
+++ b/Engine/Plugins/PCG/Source/PCG/Private/Elements/PCGCopyPoints.cpp
@@ -92,7 +92,7 @@ bool FPCGCopyPointsElement::ExecuteInternal(FPCGContext* Context) const
 			const TArray<FPCGPoint>& TargetPoints = TargetPointData->GetPoints();
 
 
-			UPCGPointData* OutPointData = NewObject<UPCGPointData>();
+			UPCGPointData* OutPointData = FPCGContext::NewObject_AnyThread<UPCGPointData>(Context);
 			TArray<FPCGPoint>& OutPoints = OutPointData->GetMutablePoints();
 			Output.Data = OutPointData;
 
diff --git a/Engine/Plugins/PCG/Source/PCG/Private/Elements/PCGCreateAttribute.cpp b/Engine/Plugins/PCG/Source/PCG/Private/Elements/PCGCreateAttribute.cpp
--- a/Engine/Plugins/PCG/Source/PCG/Private/Elements/PCGCreateAttribute.cpp
+++ b/Engine/Plugins/PCG/Source/PCG/Private/Elements/PCGCreateAttribute.cpp
@@ -297,7 +297,7 @@ bool FPCGAddAttributeElement::ExecuteInternal(FPCGContext* Context) const
 				continue;
 			}
 
-			UPCGData* OutputData = InData->DuplicateData();
+			UPCGData* OutputData = InData->DuplicateData(Context);
 			check(OutputData);
 			UPCGMetadata* OutputMetadata = OutputData->MutableMetadata();
 			if (!PCGCreateAttribute::ClearOrCreateAttribute(Settings->AttributeTypes, OutputMetadata, OutputAttributeName))
@@ -342,7 +342,7 @@ bool FPCGAddAttributeElement::ExecuteInternal(FPCGContext* Context) const
 			continue;
 		}
 
-		UPCGData* TargetData = InputData->DuplicateData();
+		UPCGData* TargetData = InputData->DuplicateData(Context);
 		FPCGTaggedData& Output = Context->OutputData.TaggedData.Add_GetRef(Inputs[i]);
 
 		bool bSuccess = true;
@@ -389,7 +389,7 @@ bool FPCGCreateAttributeElement::ExecuteInternal(FPCGContext* Context) const
 
 	FName OutputAttributeName = Settings->OutputTarget.GetName();
 
-	UPCGParamData* OutputData = NewObject<UPCGParamData>();
+	UPCGParamData* OutputData = FPCGContext::NewObject_AnyThread<UPCGParamData>(Context);
 	check(OutputData && OutputData->Metadata);
 	OutputData->Metadata->AddEntry();
 
diff --git a/Engine/Plugins/PCG/Source/PCG/Private/Elements/PCGCreateCollisionData.cpp b/Engine/Plugins/PCG/Source/PCG/Private/Elements/PCGCreateCollisionData.cpp
--- a/Engine/Plugins/PCG/Source/PCG/Private/Elements/PCGCreateCollisionData.cpp
+++ b/Engine/Plugins/PCG/Source/PCG/Private/Elements/PCGCreateCollisionData.cpp
@@ -70,7 +70,7 @@ bool FPCGCreateCollisionDataElement::PrepareDataInternal(FPCGContext* InContext)
 			{
 				FPCGCreateCollisionContext::InputMeshData InputMeshData;
 				InputMeshData.InputIndex = InputIndex;
-				InputMeshData.Data = NewObject<UPCGCollisionWrapperData>();
+				InputMeshData.Data = FPCGContext::NewObject_AnyThread<UPCGCollisionWrapperData>(Context);
 
 				if (!InputMeshData.Data->PreInitializeAndGatherMeshesEx(PointData, Settings->CollisionAttribute, Settings->bUseComplexCollision, InputMeshData.MeshPaths))
 				{
diff --git a/Engine/Plugins/PCG/Source/PCG/Private/Elements/PCGCreatePoints.cpp b/Engine/Plugins/PCG/Source/PCG/Private/Elements/PCGCreatePoints.cpp
--- a/Engine/Plugins/PCG/Source/PCG/Private/Elements/PCGCreatePoints.cpp
+++ b/Engine/Plugins/PCG/Source/PCG/Private/Elements/PCGCreatePoints.cpp
@@ -84,7 +84,7 @@ bool FPCGCreatePointsElement::ExecuteInternal(FPCGContext* Context) const
 	TArray<FPCGTaggedData>& Outputs = Context->OutputData.TaggedData;
 	FPCGTaggedData& Output = Outputs.Emplace_GetRef();
 
-	UPCGPointData* PointData = NewObject<UPCGPointData>();
+	UPCGPointData* PointData = FPCGContext::NewObject_AnyThread<UPCGPointData>(Context);
 	check(PointData); 
 
 	TArray<FPCGPoint>& OutputPoints = PointData->GetMutablePoints();
diff --git a/Engine/Plugins/PCG/Source/PCG/Private/Elements/PCGCreatePointsGrid.cpp b/Engine/Plugins/PCG/Source/PCG/Private/Elements/PCGCreatePointsGrid.cpp
--- a/Engine/Plugins/PCG/Source/PCG/Private/Elements/PCGCreatePointsGrid.cpp
+++ b/Engine/Plugins/PCG/Source/PCG/Private/Elements/PCGCreatePointsGrid.cpp
@@ -203,7 +203,7 @@ bool FPCGCreatePointsGridElement::ExecuteInternal(FPCGContext* Context) const
 	TArray<FPCGTaggedData>& Outputs = Context->OutputData.TaggedData;
 	FPCGTaggedData& Output = Outputs.Emplace_GetRef();
 
-	UPCGPointData* PtData = NewObject<UPCGPointData>();
+	UPCGPointData* PtData = FPCGContext::NewObject_AnyThread<UPCGPointData>(Context);
 	check(PtData);
 
 	TArray<FPCGPoint>& OutputPoints = PtData->GetMutablePoints();
diff --git a/Engine/Plugins/PCG/Source/PCG/Private/Elements/PCGCreatePointsSphere.cpp b/Engine/Plugins/PCG/Source/PCG/Private/Elements/PCGCreatePointsSphere.cpp
--- a/Engine/Plugins/PCG/Source/PCG/Private/Elements/PCGCreatePointsSphere.cpp
+++ b/Engine/Plugins/PCG/Source/PCG/Private/Elements/PCGCreatePointsSphere.cpp
@@ -429,7 +429,7 @@ bool FPCGCreatePointsSphereElement::ExecuteInternal(FPCGContext* InContext) cons
 	TArray<FPCGTaggedData>& Outputs = InContext->OutputData.TaggedData;
 	FPCGTaggedData& Output = Outputs.Emplace_GetRef();
 
-	UPCGPointData* PointData = NewObject<UPCGPointData>();
+	UPCGPointData* PointData = FPCGContext::NewObject_AnyThread<UPCGPointData>(InContext);
 	TArray<FPCGPoint>& OutputPoints = PointData->GetMutablePoints();
 	Output.Data = PointData;
 
diff --git a/Engine/Plugins/PCG/Source/PCG/Private/Elements/PCGCreateSpline.cpp b/Engine/Plugins/PCG/Source/PCG/Private/Elements/PCGCreateSpline.cpp
--- a/Engine/Plugins/PCG/Source/PCG/Private/Elements/PCGCreateSpline.cpp
+++ b/Engine/Plugins/PCG/Source/PCG/Private/Elements/PCGCreateSpline.cpp
@@ -115,6 +115,7 @@ bool FPCGCreateSplineElement::ExecuteInternal(FPCGContext* Context) const
 			LeaveTangentAttribute = static_cast<const FPCGMetadataAttribute<FVector>*>(LeaveTangentBaseAttribute);
 		}
 
+		check(IsInGameThread());
 		UPCGSplineData* SplineData = NewObject<UPCGSplineData>();
 		AActor* SplineActor = TargetActor;
 
diff --git a/Engine/Plugins/PCG/Source/PCG/Private/Elements/PCGCreateSurfaceFromSpline.cpp b/Engine/Plugins/PCG/Source/PCG/Private/Elements/PCGCreateSurfaceFromSpline.cpp
--- a/Engine/Plugins/PCG/Source/PCG/Private/Elements/PCGCreateSurfaceFromSpline.cpp
+++ b/Engine/Plugins/PCG/Source/PCG/Private/Elements/PCGCreateSurfaceFromSpline.cpp
@@ -63,8 +63,8 @@ bool FPCGCreateSurfaceFromSplineElement::ExecuteInternal(FPCGContext* Context) c
 			continue;
 		}
 
-		UPCGSplineInteriorSurfaceData* SurfaceData = NewObject<UPCGSplineInteriorSurfaceData>();
-		SurfaceData->Initialize(SplineData);
+		UPCGSplineInteriorSurfaceData* SurfaceData = FPCGContext::NewObject_AnyThread<UPCGSplineInteriorSurfaceData>(Context);
+		SurfaceData->Initialize(Context, SplineData);
 
 		FPCGTaggedData& Output = Outputs.Add_GetRef(Input);
 		Output.Data = SurfaceData;
diff --git a/Engine/Plugins/PCG/Source/PCG/Private/Elements/PCGCreateTargetActor.cpp b/Engine/Plugins/PCG/Source/PCG/Private/Elements/PCGCreateTargetActor.cpp
--- a/Engine/Plugins/PCG/Source/PCG/Private/Elements/PCGCreateTargetActor.cpp
+++ b/Engine/Plugins/PCG/Source/PCG/Private/Elements/PCGCreateTargetActor.cpp
@@ -253,6 +253,8 @@ bool FPCGCreateTargetActorElement::ExecuteInternal(FPCGContext* Context) const
 {
 	TRACE_CPUPROFILER_EVENT_SCOPE(FPCGCreateTargetActorElement::Execute);
 
+	check(IsInGameThread());
+
 	// Early out if the actor isn't going to be consumed by something else
 	if (Context->Node && !Context->Node->IsOutputPinConnected(PCGPinConstants::DefaultOutputLabel))
 	{
diff --git a/Engine/Plugins/PCG/Source/PCG/Private/Elements/PCGCullPointsOutsideActorBounds.cpp b/Engine/Plugins/PCG/Source/PCG/Private/Elements/PCGCullPointsOutsideActorBounds.cpp
--- a/Engine/Plugins/PCG/Source/PCG/Private/Elements/PCGCullPointsOutsideActorBounds.cpp
+++ b/Engine/Plugins/PCG/Source/PCG/Private/Elements/PCGCullPointsOutsideActorBounds.cpp
@@ -47,7 +47,7 @@ bool FPCGCullPointsOutsideActorBoundsElement::ExecuteInternal(FPCGContext* Conte
 		return true;
 	}
 
-	UPCGVolumeData* VolumeData = NewObject<UPCGVolumeData>();
+	UPCGVolumeData* VolumeData = FPCGContext::NewObject_AnyThread<UPCGVolumeData>(Context);
 	check(VolumeData);
 
 	const UPCGCullPointsOutsideActorBoundsSettings* Settings = Context->GetInputSettings<UPCGCullPointsOutsideActorBoundsSettings>();
@@ -67,7 +67,7 @@ bool FPCGCullPointsOutsideActorBoundsElement::ExecuteInternal(FPCGContext* Conte
 			continue;
 		}
 
-		UPCGIntersectionData* Intersection = InputPointData->IntersectWith(VolumeData);
+		UPCGIntersectionData* Intersection = InputPointData->IntersectWith(Context, VolumeData);
 		check(Intersection);
 
 		const UPCGPointData* IntersectedPoints = Intersection->ToPointData(Context);
diff --git a/Engine/Plugins/PCG/Source/PCG/Private/Elements/PCGDataFromActor.cpp b/Engine/Plugins/PCG/Source/PCG/Private/Elements/PCGDataFromActor.cpp
--- a/Engine/Plugins/PCG/Source/PCG/Private/Elements/PCGDataFromActor.cpp
+++ b/Engine/Plugins/PCG/Source/PCG/Private/Elements/PCGDataFromActor.cpp
@@ -541,7 +541,7 @@ void FPCGDataFromActorElement::MergeActorsIntoPointData(FPCGContext* Context, co
 
 	if (!bContainsPartitionActors)
 	{
-		UPCGPointData* Data = NewObject<UPCGPointData>();
+		UPCGPointData* Data = FPCGContext::NewObject_AnyThread<UPCGPointData>(Context);
 		bool bHasData = false;
 		bool bAnyAttributeNameWasSanitized = false;
 
@@ -596,7 +596,7 @@ void FPCGDataFromActorElement::MergeActorsIntoPointData(FPCGContext* Context, co
 		// Perform point data-to-point data merge
 		if (DataToMerge.TaggedData.Num() > 1)
 		{
-			UPCGMergeSettings* MergeSettings = NewObject<UPCGMergeSettings>();
+			UPCGMergeSettings* MergeSettings = FPCGContext::NewObject_AnyThread<UPCGMergeSettings>(Context);
 			FPCGMergeElement MergeElement;
 			FPCGContext* MergeContext = MergeElement.Initialize(DataToMerge, Context->SourceComponent, nullptr);
 			MergeContext->AsyncState.NumAvailableTasks = Context->AsyncState.NumAvailableTasks;
diff --git a/Engine/Plugins/PCG/Source/PCG/Private/Elements/PCGDataNum.cpp b/Engine/Plugins/PCG/Source/PCG/Private/Elements/PCGDataNum.cpp
--- a/Engine/Plugins/PCG/Source/PCG/Private/Elements/PCGDataNum.cpp
+++ b/Engine/Plugins/PCG/Source/PCG/Private/Elements/PCGDataNum.cpp
@@ -56,7 +56,7 @@ bool FPCGDataNumElement::ExecuteInternal(FPCGContext* Context) const
 
 	const int32 InputDataCount = Context->InputData.GetInputsByPin(PCGPinConstants::DefaultInputLabel).Num();
 
-	UPCGParamData* OutputParamData = NewObject<UPCGParamData>();
+	UPCGParamData* OutputParamData = FPCGContext::NewObject_AnyThread<UPCGParamData>(Context);
 	FPCGMetadataAttribute<int32>* NumAttribute = OutputParamData->Metadata->CreateAttribute<int32>(Settings->OutputAttributeName, InputDataCount, /*bAllowInterpolation=*/false, /*bOverrideParent=*/false);
 
 	if (!NumAttribute)
diff --git a/Engine/Plugins/PCG/Source/PCG/Private/Elements/PCGDataTableRowToParamData.cpp b/Engine/Plugins/PCG/Source/PCG/Private/Elements/PCGDataTableRowToParamData.cpp
--- a/Engine/Plugins/PCG/Source/PCG/Private/Elements/PCGDataTableRowToParamData.cpp
+++ b/Engine/Plugins/PCG/Source/PCG/Private/Elements/PCGDataTableRowToParamData.cpp
@@ -126,7 +126,7 @@ bool FPCGDataTableRowToParamData::ExecuteInternal(FPCGContext* Context) const
 	}
 
 	// From there, we should be able to create the data.
-	UPCGParamData* ParamData = NewObject<UPCGParamData>();
+	UPCGParamData* ParamData = FPCGContext::NewObject_AnyThread<UPCGParamData>(Context);
 	UPCGMetadata* Metadata = ParamData->MutableMetadata();
 	check(Metadata);
 	PCGMetadataEntryKey EntryKey = Metadata->AddEntry();
diff --git a/Engine/Plugins/PCG/Source/PCG/Private/Elements/PCGDeleteAttributesElement.cpp b/Engine/Plugins/PCG/Source/PCG/Private/Elements/PCGDeleteAttributesElement.cpp
--- a/Engine/Plugins/PCG/Source/PCG/Private/Elements/PCGDeleteAttributesElement.cpp
+++ b/Engine/Plugins/PCG/Source/PCG/Private/Elements/PCGDeleteAttributesElement.cpp
@@ -168,7 +168,7 @@ bool FPCGDeleteAttributesElement::ExecuteInternal(FPCGContext* Context) const
 		{
 			ParentMetadata = InputSpatialData->Metadata;
 
-			UPCGSpatialData* NewSpatialData = InputSpatialData->DuplicateData(/*bInitializeFromThisData=*/false);
+			UPCGSpatialData* NewSpatialData = InputSpatialData->DuplicateData(Context, /*bInitializeFromThisData=*/false);
 			Metadata = NewSpatialData->Metadata;
 			NewSpatialData->Metadata->InitializeWithAttributeFilter(ParentMetadata, AttributesToFilter, FilterMode, Settings->Operator);
 
@@ -181,7 +181,7 @@ bool FPCGDeleteAttributesElement::ExecuteInternal(FPCGContext* Context) const
 		{
 			ParentMetadata = InputParamData->Metadata;
 
-			UPCGParamData* NewParamData = NewObject<UPCGParamData>();
+			UPCGParamData* NewParamData = FPCGContext::NewObject_AnyThread<UPCGParamData>(Context);
 			Metadata = NewParamData->Metadata;
 			Metadata->InitializeAsCopyWithAttributeFilter(ParentMetadata, AttributesToFilter, FilterMode);
 
diff --git a/Engine/Plugins/PCG/Source/PCG/Private/Elements/PCGDensityFilter.cpp b/Engine/Plugins/PCG/Source/PCG/Private/Elements/PCGDensityFilter.cpp
--- a/Engine/Plugins/PCG/Source/PCG/Private/Elements/PCGDensityFilter.cpp
+++ b/Engine/Plugins/PCG/Source/PCG/Private/Elements/PCGDensityFilter.cpp
@@ -76,7 +76,7 @@ bool FPCGDensityFilterElement::ExecuteInternal(FPCGContext* Context) const
 
 		const TArray<FPCGPoint>& Points = OriginalData->GetPoints();
 		
-		UPCGPointData* FilteredData = NewObject<UPCGPointData>();
+		UPCGPointData* FilteredData = FPCGContext::NewObject_AnyThread<UPCGPointData>(Context);
 		FilteredData->InitializeFromData(OriginalData);
 		TArray<FPCGPoint>& FilteredPoints = FilteredData->GetMutablePoints();
 
diff --git a/Engine/Plugins/PCG/Source/PCG/Private/Elements/PCGDifferenceElement.cpp b/Engine/Plugins/PCG/Source/PCG/Private/Elements/PCGDifferenceElement.cpp
--- a/Engine/Plugins/PCG/Source/PCG/Private/Elements/PCGDifferenceElement.cpp
+++ b/Engine/Plugins/PCG/Source/PCG/Private/Elements/PCGDifferenceElement.cpp
@@ -190,14 +190,14 @@ bool FPCGDifferenceElement::ExecuteInternal(FPCGContext* Context) const
 
 		bHasPointsInSource |= SourceSpatialData->IsA<UPCGPointData>();
 
-		UPCGDifferenceData* DifferenceData = NewObject<UPCGDifferenceData>();
+		UPCGDifferenceData* DifferenceData = FPCGContext::NewObject_AnyThread<UPCGDifferenceData>(Context);
 		DifferenceData->Initialize(SourceSpatialData);
 		
 		for (FPCGTaggedData& Difference : Differences)
 		{
 			if (const UPCGSpatialData* DifferenceSpatialData = Cast<const UPCGSpatialData>(Difference.Data))
 			{
-				DifferenceData->AddDifference(DifferenceSpatialData);
+				DifferenceData->AddDifference(Context, DifferenceSpatialData);
 			}
 		}
 
diff --git a/Engine/Plugins/PCG/Source/PCG/Private/Elements/PCGDistance.cpp b/Engine/Plugins/PCG/Source/PCG/Private/Elements/PCGDistance.cpp
--- a/Engine/Plugins/PCG/Source/PCG/Private/Elements/PCGDistance.cpp
+++ b/Engine/Plugins/PCG/Source/PCG/Private/Elements/PCGDistance.cpp
@@ -178,7 +178,7 @@ bool FPCGDistanceElement::ExecuteInternal(FPCGContext* Context) const
 			continue;			
 		}
 
-		UPCGPointData* OutputData = NewObject<UPCGPointData>();
+		UPCGPointData* OutputData = FPCGContext::NewObject_AnyThread<UPCGPointData>(Context);
 		OutputData->InitializeFromData(SourcePointData);
 		Outputs.Add_GetRef(Source).Data = OutputData;
 
diff --git a/Engine/Plugins/PCG/Source/PCG/Private/Elements/PCGDuplicatePoint.cpp b/Engine/Plugins/PCG/Source/PCG/Private/Elements/PCGDuplicatePoint.cpp
--- a/Engine/Plugins/PCG/Source/PCG/Private/Elements/PCGDuplicatePoint.cpp
+++ b/Engine/Plugins/PCG/Source/PCG/Private/Elements/PCGDuplicatePoint.cpp
@@ -61,7 +61,7 @@ bool FPCGDuplicatePointElement::ExecuteInternal(FPCGContext* Context) const
 			return true;
 		}
 
-		UPCGPointData* OutPointData = NewObject<UPCGPointData>();
+		UPCGPointData* OutPointData = FPCGContext::NewObject_AnyThread<UPCGPointData>(Context);
 		OutPointData->InitializeFromData(InputPointData);
 		TArray<FPCGPoint>& OutputPoints = OutPointData->GetMutablePoints();
 		OutputPoints.SetNumUninitialized(NumIterations);
diff --git a/Engine/Plugins/PCG/Source/PCG/Private/Elements/PCGElevationIsolines.cpp b/Engine/Plugins/PCG/Source/PCG/Private/Elements/PCGElevationIsolines.cpp
--- a/Engine/Plugins/PCG/Source/PCG/Private/Elements/PCGElevationIsolines.cpp
+++ b/Engine/Plugins/PCG/Source/PCG/Private/Elements/PCGElevationIsolines.cpp
@@ -280,13 +280,13 @@ bool FPCGElevationIsolinesElement::ExecuteInternal(FPCGContext* InContext) const
 							Settings->bLinearSpline ? ESplinePointType::Linear : ESplinePointType::Curve);
 					}
 
-					UPCGSplineData* OutSplineData = NewObject<UPCGSplineData>();
+					UPCGSplineData* OutSplineData = FPCGContext::NewObject_AnyThread<UPCGSplineData>(InContext);
 					OutSplineData->Initialize(SplinePoints, Result.bClosed, FTransform(FVector::UpVector * Elevation));
 					OutputData.Data = OutSplineData;
 				}
 				else
 				{
-					UPCGPointData* OutPointData = NewObject<UPCGPointData>();
+					UPCGPointData* OutPointData = FPCGContext::NewObject_AnyThread<UPCGPointData>(InContext);
 					TArray<FPCGPoint>& OutPoints = OutPointData->GetMutablePoints();
 					OutPoints.Reserve(Result.LinkedGridCoordinates.Num());
 					OutputData.Data = OutPointData;
diff --git a/Engine/Plugins/PCG/Source/PCG/Private/Elements/PCGExecuteBlueprint.cpp b/Engine/Plugins/PCG/Source/PCG/Private/Elements/PCGExecuteBlueprint.cpp
--- a/Engine/Plugins/PCG/Source/PCG/Private/Elements/PCGExecuteBlueprint.cpp
+++ b/Engine/Plugins/PCG/Source/PCG/Private/Elements/PCGExecuteBlueprint.cpp
@@ -12,6 +12,7 @@
 #include "Engine/World.h"
 #include "PCGPin.h"
 #include "UObject/Package.h"
+#include "UObject/Stack.h"
 #include "UObject/UObjectThreadContext.h"
 
 #if WITH_EDITOR
@@ -133,6 +134,19 @@ void UPCGBlueprintElement::SetCurrentContext(FPCGContext* InCurrentContext)
 	CurrentContext = InCurrentContext;
 }
 
+FPCGContext* UPCGBlueprintElement::ResolveContext()
+{
+	if (FFrame::GetThreadLocalTopStackFrame() && FFrame::GetThreadLocalTopStackFrame()->Object)
+	{
+		if (UPCGBlueprintElement* Caller = Cast<UPCGBlueprintElement>(FFrame::GetThreadLocalTopStackFrame()->Object))
+		{
+			return &Caller->GetContext();
+		}
+	}
+
+	return nullptr;
+}
+
 #if WITH_EDITOR
 void UPCGBlueprintElement::PostEditChangeProperty(FPropertyChangedEvent& PropertyChangedEvent)
 {
@@ -786,7 +800,20 @@ bool FPCGExecuteBlueprintElement::ExecuteInternal(FPCGContext* InContext) const
 		// Note that the context is actually sliced so there should never be any members in the BP element context that are visible/accessible from blueprint
 		/** Finally, execute the actual blueprint */
 		Context->BlueprintElementInstance->SetCurrentContext(Context);
-		Context->BlueprintElementInstance->ExecuteWithContext(*Context, Context->InputData, Context->OutputData);
+		
+		const bool bIsInGameThread = IsInGameThread();
+		// When running outside of main thread make sure GC can't run (BP Nodes might create objects which can't happen while GC runs)
+		if(!bIsInGameThread)
+		{
+			ensure(!Context->AsyncState.bIsRunningOnMainThread);
+			FGCScopeGuard Scope;
+			Context->BlueprintElementInstance->ExecuteWithContext(*Context, Context->InputData, Context->OutputData);
+		}
+		else
+		{
+			Context->BlueprintElementInstance->ExecuteWithContext(*Context, Context->InputData, Context->OutputData);
+		}
+
 		Context->BlueprintElementInstance->SetCurrentContext(nullptr);
 
 		// Log info on outputs
@@ -804,6 +831,16 @@ bool FPCGExecuteBlueprintElement::ExecuteInternal(FPCGContext* InContext) const
 			// Note that we will recurse up the outer tree to make sure we catch every case.
 			if(Output.Data)
 			{
+				if (!bIsInGameThread)
+				{
+					// Clear Async flags on objects created outside of the main thread and not part of the Context known async objects
+					if (Output.Data->HasAnyInternalFlags(EInternalObjectFlags::Async) && !Context->ContainsAsyncObject(Output.Data))
+					{
+						Output.Data->AtomicallyClearInternalFlags(EInternalObjectFlags::Async);
+						ForEachObjectWithOuter(Output.Data, [](UObject* SubObject) { SubObject->AtomicallyClearInternalFlags(EInternalObjectFlags::Async); }, true);
+					}
+				}
+
 				auto ReOuterToTransientPackageIfCreatedFromThis = [Context](UObject* InObject)
 				{
 					bool bHasInstanceAsOuter = false;
@@ -869,7 +906,7 @@ void UPCGBlueprintElement::PointLoop(FPCGContext& InContext, const UPCGPointData
 	}
 	else
 	{
-		OutData = NewObject<UPCGPointData>();
+		OutData = FPCGContext::NewObject_AnyThread<UPCGPointData>(&InContext);
 		OutData->InitializeFromData(InData);
 	}
 
@@ -896,7 +933,7 @@ void UPCGBlueprintElement::VariableLoop(FPCGContext& InContext, const UPCGPointD
 	}
 	else
 	{
-		OutData = NewObject<UPCGPointData>();
+		OutData = FPCGContext::NewObject_AnyThread<UPCGPointData>(&InContext);
 		OutData->InitializeFromData(InData);
 	}
 
@@ -923,7 +960,7 @@ void UPCGBlueprintElement::NestedLoop(FPCGContext& InContext, const UPCGPointDat
 	}
 	else
 	{
-		OutData = NewObject<UPCGPointData>();
+		OutData = FPCGContext::NewObject_AnyThread<UPCGPointData>(&InContext);
 		OutData->InitializeFromData(InOuterData);
 		OutData->Metadata->AddAttributes(InInnerData->Metadata);
 	}
@@ -955,7 +992,7 @@ void UPCGBlueprintElement::IterationLoop(FPCGContext& InContext, int64 NumIterat
 	else
 	{
 		const UPCGSpatialData* Owner = (InA ? InA : InB);
-		OutData = NewObject<UPCGPointData>();
+		OutData = FPCGContext::NewObject_AnyThread<UPCGPointData>(&InContext);
 
 		if (Owner)
 		{
diff --git a/Engine/Plugins/PCG/Source/PCG/Private/Elements/PCGGetBounds.cpp b/Engine/Plugins/PCG/Source/PCG/Private/Elements/PCGGetBounds.cpp
--- a/Engine/Plugins/PCG/Source/PCG/Private/Elements/PCGGetBounds.cpp
+++ b/Engine/Plugins/PCG/Source/PCG/Private/Elements/PCGGetBounds.cpp
@@ -77,7 +77,7 @@ bool FPCGGetBoundsElement::ExecuteInternal(FPCGContext* Context) const
 			continue;
 		}
 
-		UPCGParamData* ParamData = NewObject<UPCGParamData>();
+		UPCGParamData* ParamData = FPCGContext::NewObject_AnyThread<UPCGParamData>(Context);
 		check(ParamData && ParamData->Metadata);
 		
 		ParamData->Metadata->CreateVectorAttribute("BoundsMin", InputBounds.Min, /*bAllowsInterpolation=*/true, /*bOverrideParent=*/false);
diff --git a/Engine/Plugins/PCG/Source/PCG/Private/Elements/PCGGetDataInfo.cpp b/Engine/Plugins/PCG/Source/PCG/Private/Elements/PCGGetDataInfo.cpp
--- a/Engine/Plugins/PCG/Source/PCG/Private/Elements/PCGGetDataInfo.cpp
+++ b/Engine/Plugins/PCG/Source/PCG/Private/Elements/PCGGetDataInfo.cpp
@@ -41,7 +41,7 @@ bool FPCGGetTagsElement::ExecuteInternal(FPCGContext* Context) const
 			continue;
 		}
 
-		UPCGParamData* ParamData = NewObject<UPCGParamData>();
+		UPCGParamData* ParamData = FPCGContext::NewObject_AnyThread<UPCGParamData>(Context);
 		check(ParamData && ParamData->Metadata);
 
 		FPCGTaggedData& Output = Outputs.Emplace_GetRef(Input);
@@ -117,7 +117,7 @@ bool FPCGGetAttributesElement::ExecuteInternal(FPCGContext* Context) const
 			continue;
 		}
 
-		UPCGParamData* ParamData = NewObject<UPCGParamData>();
+		UPCGParamData* ParamData = FPCGContext::NewObject_AnyThread<UPCGParamData>(Context);
 		check(ParamData && ParamData->Metadata);
 
 		FPCGTaggedData& Output = Outputs.Emplace_GetRef(Input);
diff --git a/Engine/Plugins/PCG/Source/PCG/Private/Elements/PCGGetLoopIndex.cpp b/Engine/Plugins/PCG/Source/PCG/Private/Elements/PCGGetLoopIndex.cpp
--- a/Engine/Plugins/PCG/Source/PCG/Private/Elements/PCGGetLoopIndex.cpp
+++ b/Engine/Plugins/PCG/Source/PCG/Private/Elements/PCGGetLoopIndex.cpp
@@ -74,7 +74,7 @@ bool FPCGGetLoopIndexElement::ExecuteInternal(FPCGContext* Context) const
 		return true;
 	}
 
-	UPCGParamData* LoopIndexParamData = NewObject<UPCGParamData>();
+	UPCGParamData* LoopIndexParamData = FPCGContext::NewObject_AnyThread<UPCGParamData>(Context);
 	check(LoopIndexParamData && LoopIndexParamData->Metadata);
 
 	FPCGTaggedData& OutputData = Context->OutputData.TaggedData.Emplace_GetRef();
diff --git a/Engine/Plugins/PCG/Source/PCG/Private/Elements/PCGGetPropertyFromObjectPath.cpp b/Engine/Plugins/PCG/Source/PCG/Private/Elements/PCGGetPropertyFromObjectPath.cpp
--- a/Engine/Plugins/PCG/Source/PCG/Private/Elements/PCGGetPropertyFromObjectPath.cpp
+++ b/Engine/Plugins/PCG/Source/PCG/Private/Elements/PCGGetPropertyFromObjectPath.cpp
@@ -309,7 +309,7 @@ bool FPCGGetPropertyFromObjectPathElement::ExecuteInternal(FPCGContext* Context)
 
 			if (!EmptyParam)
 			{
-				EmptyParam = NewObject<UPCGParamData>();
+				EmptyParam = FPCGContext::NewObject_AnyThread<UPCGParamData>(Context);
 			}
 
 			AddToOutput(EmptyParam, Index);
diff --git a/Engine/Plugins/PCG/Source/PCG/Private/Elements/PCGHiGenGridSize.cpp b/Engine/Plugins/PCG/Source/PCG/Private/Elements/PCGHiGenGridSize.cpp
--- a/Engine/Plugins/PCG/Source/PCG/Private/Elements/PCGHiGenGridSize.cpp
+++ b/Engine/Plugins/PCG/Source/PCG/Private/Elements/PCGHiGenGridSize.cpp
@@ -136,7 +136,7 @@ bool FPCGHiGenGridSizeElement::ExecuteInternal(FPCGContext* Context) const
 		Data.Pin = PCGPinConstants::DefaultOutputLabel;
 	}
 
-	UPCGVolumeData* VolumeData = NewObject<UPCGVolumeData>();
+	UPCGVolumeData* VolumeData = FPCGContext::NewObject_AnyThread<UPCGVolumeData>(Context);
 	check(VolumeData);
 	VolumeData->Initialize(Context->SourceComponent->GetGridBounds());
 
diff --git a/Engine/Plugins/PCG/Source/PCG/Private/Elements/PCGInnerIntersectionElement.cpp b/Engine/Plugins/PCG/Source/PCG/Private/Elements/PCGInnerIntersectionElement.cpp
--- a/Engine/Plugins/PCG/Source/PCG/Private/Elements/PCGInnerIntersectionElement.cpp
+++ b/Engine/Plugins/PCG/Source/PCG/Private/Elements/PCGInnerIntersectionElement.cpp
@@ -84,7 +84,7 @@ bool FPCGInnerIntersectionElement::ExecuteInternal(FPCGContext* Context) const
 		}
 
 		// Create a new intersection
-		IntersectionData = (IntersectionData ? IntersectionData : FirstSpatialData)->IntersectWith(SpatialData);
+		IntersectionData = (IntersectionData ? IntersectionData : FirstSpatialData)->IntersectWith(Context, SpatialData);
 		// Propagate settings
 		IntersectionData->DensityFunction = DensityFunction;
 		IntersectionData->bKeepZeroDensityPoints = bKeepZeroDensityPoints;
diff --git a/Engine/Plugins/PCG/Source/PCG/Private/Elements/PCGMatchAndSetAttributes.cpp b/Engine/Plugins/PCG/Source/PCG/Private/Elements/PCGMatchAndSetAttributes.cpp
--- a/Engine/Plugins/PCG/Source/PCG/Private/Elements/PCGMatchAndSetAttributes.cpp
+++ b/Engine/Plugins/PCG/Source/PCG/Private/Elements/PCGMatchAndSetAttributes.cpp
@@ -926,14 +926,14 @@ bool FPCGMatchAndSetAttributesElement::PrepareDataInternal(FPCGContext* InContex
 
 		if (const UPCGPointData* InPointData = Cast<const UPCGPointData>(OutState.InData))
 		{
-			UPCGPointData* OutPointData = NewObject<UPCGPointData>();
+			UPCGPointData* OutPointData = FPCGContext::NewObject_AnyThread<UPCGPointData>(InContext);
 			OutPointData->InitializeFromData(InPointData);
 			OutPointData->GetMutablePoints().Reserve(InPointData->GetPoints().Num());
 			OutState.OutData = OutPointData;
 		}
 		else if (const UPCGParamData* InParamData = Cast<const UPCGParamData>(OutState.InData))
 		{
-			UPCGParamData* OutParamData = NewObject<UPCGParamData>();
+			UPCGParamData* OutParamData = FPCGContext::NewObject_AnyThread<UPCGParamData>(InContext);
 			OutParamData->Metadata->AddAttributes(InParamData->Metadata);
 
 			OutState.OutData = OutParamData;
diff --git a/Engine/Plugins/PCG/Source/PCG/Private/Elements/PCGMergeAttributes.cpp b/Engine/Plugins/PCG/Source/PCG/Private/Elements/PCGMergeAttributes.cpp
--- a/Engine/Plugins/PCG/Source/PCG/Private/Elements/PCGMergeAttributes.cpp
+++ b/Engine/Plugins/PCG/Source/PCG/Private/Elements/PCGMergeAttributes.cpp
@@ -83,7 +83,7 @@ bool FPCGMergeAttributesElement::ExecuteInternal(FPCGContext* Context) const
 		// When we're merging the 2nd element, create the actual merged attribute set
 		if (!MergedAttributeSet)
 		{
-			MergedAttributeSet = NewObject<UPCGParamData>();
+			MergedAttributeSet = FPCGContext::NewObject_AnyThread<UPCGParamData>(Context);
 			check(MergedOutput);
 			MergedAttributeSet->Metadata->InitializeAsCopy(CastChecked<const UPCGParamData>(MergedOutput->Data)->Metadata);
 			MergedOutput->Data = MergedAttributeSet;
diff --git a/Engine/Plugins/PCG/Source/PCG/Private/Elements/PCGMergeElement.cpp b/Engine/Plugins/PCG/Source/PCG/Private/Elements/PCGMergeElement.cpp
--- a/Engine/Plugins/PCG/Source/PCG/Private/Elements/PCGMergeElement.cpp
+++ b/Engine/Plugins/PCG/Source/PCG/Private/Elements/PCGMergeElement.cpp
@@ -82,7 +82,7 @@ bool FPCGMergeElement::ExecuteInternal(FPCGContext* Context) const
 			// Second valid data - we'll create the actual merged data at this point
 			check(TargetTaggedData);
 
-			TargetPointData = NewObject<UPCGPointData>();
+			TargetPointData = FPCGContext::NewObject_AnyThread<UPCGPointData>(Context);
 			TargetPointData->InitializeFromData(CastChecked<const UPCGPointData>(TargetTaggedData->Data), nullptr, bMergeMetadata);
 			TargetTaggedData->Data = TargetPointData;
 		}
diff --git a/Engine/Plugins/PCG/Source/PCG/Private/Elements/PCGNumberOfElements.cpp b/Engine/Plugins/PCG/Source/PCG/Private/Elements/PCGNumberOfElements.cpp
--- a/Engine/Plugins/PCG/Source/PCG/Private/Elements/PCGNumberOfElements.cpp
+++ b/Engine/Plugins/PCG/Source/PCG/Private/Elements/PCGNumberOfElements.cpp
@@ -88,7 +88,7 @@ bool FPCGNumberOfElementsBaseElement<DataType>::ExecuteInternal(FPCGContext* Con
 
 		if (!OutputParamData)
 		{
-			OutputParamData = NewObject<UPCGParamData>();
+			OutputParamData = FPCGContext::NewObject_AnyThread<UPCGParamData>(Context);
 			NewAttribute = OutputParamData->Metadata->CreateAttribute<int32>(Settings->OutputAttributeName, 0, /*bAllowInterpolation=*/true, /*bOverrideParent=*/false);
 
 			if (!NewAttribute)
diff --git a/Engine/Plugins/PCG/Source/PCG/Private/Elements/PCGOuterIntersectionElement.cpp b/Engine/Plugins/PCG/Source/PCG/Private/Elements/PCGOuterIntersectionElement.cpp
--- a/Engine/Plugins/PCG/Source/PCG/Private/Elements/PCGOuterIntersectionElement.cpp
+++ b/Engine/Plugins/PCG/Source/PCG/Private/Elements/PCGOuterIntersectionElement.cpp
@@ -146,7 +146,7 @@ bool FPCGOuterIntersectionElement::ExecuteInternal(FPCGContext* Context) const
 		const UPCGSpatialData* IntersectionData = PrimarySpatialData;
 		for (const UPCGSpatialData* SecondarySourceUnionData : SecondarySourceUnionArray)
 		{
-			UPCGIntersectionData* TempIntersectionData = IntersectionData->IntersectWith(SecondarySourceUnionData);
+			UPCGIntersectionData* TempIntersectionData = IntersectionData->IntersectWith(Context, SecondarySourceUnionData);
 			// Propagate settings
 			TempIntersectionData->DensityFunction = Settings->DensityFunction;
 			TempIntersectionData->bKeepZeroDensityPoints = Settings->bKeepZeroDensityPoints;
diff --git a/Engine/Plugins/PCG/Source/PCG/Private/Elements/PCGPointFromMeshElement.cpp b/Engine/Plugins/PCG/Source/PCG/Private/Elements/PCGPointFromMeshElement.cpp
--- a/Engine/Plugins/PCG/Source/PCG/Private/Elements/PCGPointFromMeshElement.cpp
+++ b/Engine/Plugins/PCG/Source/PCG/Private/Elements/PCGPointFromMeshElement.cpp
@@ -75,7 +75,7 @@ bool FPCGPointFromMeshElement::ExecuteInternal(FPCGContext* Context) const
 	}
 
 	TArray<FPCGTaggedData>& Outputs = Context->OutputData.TaggedData;
-	UPCGPointData* OutPointData = NewObject<UPCGPointData>();
+	UPCGPointData* OutPointData = FPCGContext::NewObject_AnyThread<UPCGPointData>(Context);
 	Outputs.Emplace_GetRef().Data = OutPointData;
 	
 	TArray<FPCGPoint>& Points = OutPointData->GetMutablePoints();
diff --git a/Engine/Plugins/PCG/Source/PCG/Private/Elements/PCGPointMatchAndSet.cpp b/Engine/Plugins/PCG/Source/PCG/Private/Elements/PCGPointMatchAndSet.cpp
--- a/Engine/Plugins/PCG/Source/PCG/Private/Elements/PCGPointMatchAndSet.cpp
+++ b/Engine/Plugins/PCG/Source/PCG/Private/Elements/PCGPointMatchAndSet.cpp
@@ -157,6 +157,8 @@ void UPCGPointMatchAndSetSettings::RefreshMatchAndSet()
 {
 	if (MatchAndSetType)
 	{
+		ensure(IsInGameThread());
+
 		// Forget previous instance
 		if (MatchAndSetInstance)
 		{
@@ -217,7 +219,7 @@ bool FPCGPointMatchAndSetElement::ExecuteInternal(FPCGContext* Context) const
 			continue;
 		}
 
-		UPCGPointData* OutPointData = NewObject<UPCGPointData>();
+		UPCGPointData* OutPointData = FPCGContext::NewObject_AnyThread<UPCGPointData>(Context);
 		Output.Data = OutPointData;
 
 		OutPointData->InitializeFromData(InPointData);
diff --git a/Engine/Plugins/PCG/Source/PCG/Private/Elements/PCGPointNeighborhood.cpp b/Engine/Plugins/PCG/Source/PCG/Private/Elements/PCGPointNeighborhood.cpp
--- a/Engine/Plugins/PCG/Source/PCG/Private/Elements/PCGPointNeighborhood.cpp
+++ b/Engine/Plugins/PCG/Source/PCG/Private/Elements/PCGPointNeighborhood.cpp
@@ -79,7 +79,7 @@ bool FPCGPointNeighborhoodElement::ExecuteInternal(FPCGContext* Context) const
 
 		const TArray<FPCGPoint>& SrcPoints = InputPointData->GetPoints();
 		FPCGTaggedData& Output = Outputs.Add_GetRef(Inputs[i]);
-		UPCGPointData* OutputPointData = NewObject<UPCGPointData>();
+		UPCGPointData* OutputPointData = FPCGContext::NewObject_AnyThread<UPCGPointData>(Context);
 		OutputPointData->InitializeFromData(InputPointData);
 		TArray<FPCGPoint>& OutputPoints = OutputPointData->GetMutablePoints();
 		OutputPoints.SetNumUninitialized(SrcPoints.Num());
diff --git a/Engine/Plugins/PCG/Source/PCG/Private/Elements/PCGPointOperationElementBase.cpp b/Engine/Plugins/PCG/Source/PCG/Private/Elements/PCGPointOperationElementBase.cpp
--- a/Engine/Plugins/PCG/Source/PCG/Private/Elements/PCGPointOperationElementBase.cpp
+++ b/Engine/Plugins/PCG/Source/PCG/Private/Elements/PCGPointOperationElementBase.cpp
@@ -55,7 +55,7 @@ bool FPCGPointOperationElementBase::PreparePointOperationData(ContextType* InCon
 			OutState.NumPoints = OutState.InputPointData->GetPoints().Num();
 
 			// Create and initialize the output points
-			OutState.OutputPointData = NewObject<UPCGPointData>();
+			OutState.OutputPointData = FPCGContext::NewObject_AnyThread<UPCGPointData>(InContext);
 			OutState.OutputPointData->InitializeFromData(OutState.InputPointData);
 			OutState.OutputPointData->GetMutablePoints().SetNumUninitialized(OutState.NumPoints);
 
diff --git a/Engine/Plugins/PCG/Source/PCG/Private/Elements/PCGPointProcessingElementBase.cpp b/Engine/Plugins/PCG/Source/PCG/Private/Elements/PCGPointProcessingElementBase.cpp
--- a/Engine/Plugins/PCG/Source/PCG/Private/Elements/PCGPointProcessingElementBase.cpp
+++ b/Engine/Plugins/PCG/Source/PCG/Private/Elements/PCGPointProcessingElementBase.cpp
@@ -35,7 +35,7 @@ void FPCGPointProcessingElementBase::ProcessPoints(FPCGContext* Context, const T
 
 		const TArray<FPCGPoint>& Points = PointData->GetPoints();
 
-		UPCGPointData* OutputData = NewObject<UPCGPointData>();
+		UPCGPointData* OutputData = FPCGContext::NewObject_AnyThread<UPCGPointData>(Context);
 		OutputData->InitializeFromData(PointData);
 		TArray<FPCGPoint>& OutputPoints = OutputData->GetMutablePoints();
 		Output.Data = OutputData;
diff --git a/Engine/Plugins/PCG/Source/PCG/Private/Elements/PCGPrintElement.cpp b/Engine/Plugins/PCG/Source/PCG/Private/Elements/PCGPrintElement.cpp
--- a/Engine/Plugins/PCG/Source/PCG/Private/Elements/PCGPrintElement.cpp
+++ b/Engine/Plugins/PCG/Source/PCG/Private/Elements/PCGPrintElement.cpp
@@ -63,7 +63,6 @@ FPCGElementPtr UPCGPrintElementSettings::CreateElement() const
 bool FPCGPrintElement::ExecuteInternal(FPCGContext* Context) const
 {
 	TRACE_CPUPROFILER_EVENT_SCOPE(FPCGPrintElement::Execute);
-
 	Context->OutputData.TaggedData = Context->InputData.GetInputsByPin(PCGPinConstants::DefaultInputLabel);
 
 #if !(UE_BUILD_SHIPPING || UE_BUILD_TEST) || USE_LOGGING_IN_SHIPPING
@@ -145,7 +144,7 @@ bool FPCGPrintElement::ExecuteInternal(FPCGContext* Context) const
 			HashKey32 = HashCombine(HashKey32, GetTypeHash(Context->Stack));
 
 			check(Context->SourceComponent.IsValid());
-			UPCGManagedDebugStringMessageKey* ManagedMessageKey = NewObject<UPCGManagedDebugStringMessageKey>(Context->SourceComponent.Get());
+			UPCGManagedDebugStringMessageKey* ManagedMessageKey = FPCGContext::NewObject_AnyThread<UPCGManagedDebugStringMessageKey>(Context, Context->SourceComponent.Get());
 			ManagedMessageKey->HashKey = static_cast<uint64>(HashKey32);
 
 			Context->SourceComponent->AddToManagedResources(ManagedMessageKey);
diff --git a/Engine/Plugins/PCG/Source/PCG/Private/Elements/PCGProjectionElement.cpp b/Engine/Plugins/PCG/Source/PCG/Private/Elements/PCGProjectionElement.cpp
--- a/Engine/Plugins/PCG/Source/PCG/Private/Elements/PCGProjectionElement.cpp
+++ b/Engine/Plugins/PCG/Source/PCG/Private/Elements/PCGProjectionElement.cpp
@@ -98,7 +98,7 @@ bool FPCGProjectionElement::ExecuteInternal(FPCGContext* Context) const
 			continue;
 		}
 
-		UPCGSpatialData* ProjectionData = ProjectionSource->ProjectOn(ProjectionTarget, ProjectionParams);
+		UPCGSpatialData* ProjectionData = ProjectionSource->ProjectOn(Context, ProjectionTarget, ProjectionParams);
 		ProjectionData->bKeepZeroDensityPoints = bKeepZeroDensityPoints;
 
 		const UPCGSpatialData* OutputData = ProjectionData;
diff --git a/Engine/Plugins/PCG/Source/PCG/Private/Elements/PCGRandomChoice.cpp b/Engine/Plugins/PCG/Source/PCG/Private/Elements/PCGRandomChoice.cpp
--- a/Engine/Plugins/PCG/Source/PCG/Private/Elements/PCGRandomChoice.cpp
+++ b/Engine/Plugins/PCG/Source/PCG/Private/Elements/PCGRandomChoice.cpp
@@ -106,10 +106,10 @@ bool FPCGRandomChoiceElement::ExecuteInternal(FPCGContext* Context) const
 			NumOfElementsToKeep = FMath::CeilToInt(InPoints.Num() * FMath::Clamp(Settings->Ratio, 0, 1));
 		}
 
-		auto CreateData = [InputPointData, &Outputs, &CurrentInput](const int NumPoints, const FName PinLabel) -> TArray<FPCGPoint>*
+		auto CreateData = [InputPointData, &Outputs, &CurrentInput, Context](const int NumPoints, const FName PinLabel) -> TArray<FPCGPoint>*
 		{
 			FPCGTaggedData& Output = Outputs.Add_GetRef(CurrentInput);
-			UPCGPointData* OutPointData = NewObject<UPCGPointData>();
+			UPCGPointData* OutPointData = FPCGContext::NewObject_AnyThread<UPCGPointData>(Context);
 			OutPointData->InitializeFromData(InputPointData);
 
 			TArray<FPCGPoint>& Points = OutPointData->GetMutablePoints();
diff --git a/Engine/Plugins/PCG/Source/PCG/Private/Elements/PCGSampleTexture.cpp b/Engine/Plugins/PCG/Source/PCG/Private/Elements/PCGSampleTexture.cpp
--- a/Engine/Plugins/PCG/Source/PCG/Private/Elements/PCGSampleTexture.cpp
+++ b/Engine/Plugins/PCG/Source/PCG/Private/Elements/PCGSampleTexture.cpp
@@ -92,7 +92,7 @@ bool FPCGSampleTextureElement::ExecuteInternal(FPCGContext* Context) const
 			}
 		}
 
-		UPCGPointData* OutPointData = NewObject<UPCGPointData>();
+		UPCGPointData* OutPointData = FPCGContext::NewObject_AnyThread<UPCGPointData>(Context);
 		OutPointData->InitializeFromData(PointData);
 
 		auto ProcessPoint = [Settings, BaseTextureData, &InputAccessor, &InputKeys, &InputPoints, &DensityMergeFunc, OutPointData](int32 Index, FPCGPoint& OutPoint)
diff --git a/Engine/Plugins/PCG/Source/PCG/Private/Elements/PCGSelectPoints.cpp b/Engine/Plugins/PCG/Source/PCG/Private/Elements/PCGSelectPoints.cpp
--- a/Engine/Plugins/PCG/Source/PCG/Private/Elements/PCGSelectPoints.cpp
+++ b/Engine/Plugins/PCG/Source/PCG/Private/Elements/PCGSelectPoints.cpp
@@ -90,7 +90,7 @@ bool FPCGSelectPointsElement::ExecuteInternal(FPCGContext* Context) const
 		const TArray<FPCGPoint>& Points = OriginalData->GetPoints();
 		const int OriginalPointCount = Points.Num();
 
-		UPCGPointData* SampledData = NewObject<UPCGPointData>();
+		UPCGPointData* SampledData = FPCGContext::NewObject_AnyThread<UPCGPointData>(Context);
 		SampledData->InitializeFromData(OriginalData);
 		TArray<FPCGPoint>& SampledPoints = SampledData->GetMutablePoints();
 
diff --git a/Engine/Plugins/PCG/Source/PCG/Private/Elements/PCGSelfPruning.cpp b/Engine/Plugins/PCG/Source/PCG/Private/Elements/PCGSelfPruning.cpp
--- a/Engine/Plugins/PCG/Source/PCG/Private/Elements/PCGSelfPruning.cpp
+++ b/Engine/Plugins/PCG/Source/PCG/Private/Elements/PCGSelfPruning.cpp
@@ -594,7 +594,7 @@ namespace PCGSelfPruningElement
 		// Finally, output all points that are present in the ExclusionPoints. This part is not time sliced, should it be too?
 		if (bIsDone)
 		{
-			UPCGPointData* PrunedData = NewObject<UPCGPointData>();
+			UPCGPointData* PrunedData = FPCGContext::NewObject_AnyThread<UPCGPointData>(InOptionalContext);
 			PrunedData->InitializeFromData(InState.InputPointData);
 			InState.OutputPointData = PrunedData;
 
diff --git a/Engine/Plugins/PCG/Source/PCG/Private/Elements/PCGSortAttributes.cpp b/Engine/Plugins/PCG/Source/PCG/Private/Elements/PCGSortAttributes.cpp
--- a/Engine/Plugins/PCG/Source/PCG/Private/Elements/PCGSortAttributes.cpp
+++ b/Engine/Plugins/PCG/Source/PCG/Private/Elements/PCGSortAttributes.cpp
@@ -85,14 +85,14 @@ bool FPCGSortAttributesElement::ExecuteInternal(FPCGContext* Context) const
 			UPCGData* OutputData = nullptr;
 			if (const UPCGPointData* InputPointData = Cast<const UPCGPointData>(InputData))
 			{
-				UPCGPointData* OutputPointData = static_cast<UPCGPointData*>(InputData->DuplicateData());
+				UPCGPointData* OutputPointData = static_cast<UPCGPointData*>(InputData->DuplicateData(Context));
 				PCGAttributeAccessorHelpers::SortByAttribute(*Accessor, *Keys, OutputPointData->GetMutablePoints(), Settings->SortMethod == EPCGSortMethod::Ascending);
 				OutputData = OutputPointData;
 			}
 			else if (const UPCGMetadata* InMetadata = InputData->ConstMetadata())
 			{
 				// Duplicate data without metadata
-				OutputData = InputData->DuplicateData(/*bInitializeMetdata=*/false);
+				OutputData = InputData->DuplicateData(Context, /*bInitializeMetdata=*/false);
 				UPCGMetadata* OutMetadata = OutputData->MutableMetadata();
 
 				if (!ensure(OutMetadata))
diff --git a/Engine/Plugins/PCG/Source/PCG/Private/Elements/PCGSpatialNoise.cpp b/Engine/Plugins/PCG/Source/PCG/Private/Elements/PCGSpatialNoise.cpp
--- a/Engine/Plugins/PCG/Source/PCG/Private/Elements/PCGSpatialNoise.cpp
+++ b/Engine/Plugins/PCG/Source/PCG/Private/Elements/PCGSpatialNoise.cpp
@@ -744,7 +744,7 @@ bool FPCGSpatialNoise::ExecuteInternal(FPCGContext* Context) const
 			continue;
 		}
 
-		BufferParams.OutputPointData = NewObject<UPCGPointData>();
+		BufferParams.OutputPointData = FPCGContext::NewObject_AnyThread<UPCGPointData>(Context);
 		BufferParams.OutputPointData->InitializeFromData(BufferParams.InputPointData);
 		Context->OutputData.TaggedData.Add_GetRef(Input).Data = BufferParams.OutputPointData;
 
diff --git a/Engine/Plugins/PCG/Source/PCG/Private/Elements/PCGSpawnActor.cpp b/Engine/Plugins/PCG/Source/PCG/Private/Elements/PCGSpawnActor.cpp
--- a/Engine/Plugins/PCG/Source/PCG/Private/Elements/PCGSpawnActor.cpp
+++ b/Engine/Plugins/PCG/Source/PCG/Private/Elements/PCGSpawnActor.cpp
@@ -746,6 +746,7 @@ void FPCGSpawnActorElement::SpawnActors(FPCGSubgraphContext* Context, AActor* Ta
 {
 	TRACE_CPUPROFILER_EVENT_SCOPE(FPCGSpawnActorElement::ExecuteInternal::SpawnActors);
 	check(Context && TargetActor && PointData);
+	check(IsInGameThread());
 
 	const TArray<FPCGPoint>& Points = PointData->GetPoints();
 	if (Points.IsEmpty())
diff --git a/Engine/Plugins/PCG/Source/PCG/Private/Elements/PCGSplineSampler.cpp b/Engine/Plugins/PCG/Source/PCG/Private/Elements/PCGSplineSampler.cpp
--- a/Engine/Plugins/PCG/Source/PCG/Private/Elements/PCGSplineSampler.cpp
+++ b/Engine/Plugins/PCG/Source/PCG/Private/Elements/PCGSplineSampler.cpp
@@ -1554,7 +1554,7 @@ bool FPCGSplineSamplerElement::ExecuteInternal(FPCGContext* Context) const
 		FPCGTaggedData& Output = Outputs.Emplace_GetRef();
 		Output = Input;
 
-		UPCGPointData* SampledPointData = NewObject<UPCGPointData>();
+		UPCGPointData* SampledPointData = FPCGContext::NewObject_AnyThread<UPCGPointData>(Context);
 		SampledPointData->InitializeFromData(SpatialData);
 		Output.Data = SampledPointData;
 
diff --git a/Engine/Plugins/PCG/Source/PCG/Private/Elements/PCGSplitPoints.cpp b/Engine/Plugins/PCG/Source/PCG/Private/Elements/PCGSplitPoints.cpp
--- a/Engine/Plugins/PCG/Source/PCG/Private/Elements/PCGSplitPoints.cpp
+++ b/Engine/Plugins/PCG/Source/PCG/Private/Elements/PCGSplitPoints.cpp
@@ -64,7 +64,7 @@ bool FPCGSplitPointsElement::ExecuteInternal(FPCGContext* Context) const
 		const TArray<FPCGPoint>& InputPoints = InputPointData->GetPoints();
 
 		FPCGTaggedData& OutputA = Outputs.Add_GetRef(Inputs[i]);
-		UPCGPointData* OutPointDataA = NewObject<UPCGPointData>();
+		UPCGPointData* OutPointDataA = FPCGContext::NewObject_AnyThread<UPCGPointData>(Context);
 		OutPointDataA->InitializeFromData(InputPointData);
 		TArray<FPCGPoint>& PointsA = OutPointDataA->GetMutablePoints();
 		PointsA.SetNumUninitialized(InputPoints.Num());
@@ -72,7 +72,7 @@ bool FPCGSplitPointsElement::ExecuteInternal(FPCGContext* Context) const
 		OutputA.Pin = PCGSplitPointsConstants::OutputALabel;
 
 		FPCGTaggedData& OutputB = Outputs.Add_GetRef(Inputs[i]);
-		UPCGPointData* OutPointDataB = NewObject<UPCGPointData>();
+		UPCGPointData* OutPointDataB = FPCGContext::NewObject_AnyThread<UPCGPointData>(Context);
 		OutPointDataB->InitializeFromData(InputPointData);
 		TArray<FPCGPoint>& PointsB = OutPointDataB->GetMutablePoints();
 		PointsB.SetNumUninitialized(InputPoints.Num());
diff --git a/Engine/Plugins/PCG/Source/PCG/Private/Elements/PCGStaticMeshSpawner.cpp b/Engine/Plugins/PCG/Source/PCG/Private/Elements/PCGStaticMeshSpawner.cpp
--- a/Engine/Plugins/PCG/Source/PCG/Private/Elements/PCGStaticMeshSpawner.cpp
+++ b/Engine/Plugins/PCG/Source/PCG/Private/Elements/PCGStaticMeshSpawner.cpp
@@ -188,7 +188,7 @@ bool FPCGStaticMeshSpawnerElement::PrepareDataInternal(FPCGContext* InContext) c
 			{
 				FPCGTaggedData& Output = Outputs.Add_GetRef(Input);
 
-				UPCGPointData* OutputPointData = NewObject<UPCGPointData>();
+				UPCGPointData* OutputPointData = FPCGContext::NewObject_AnyThread<UPCGPointData>(Context);
 				OutputPointData->InitializeFromData(PointData);
 
 				if (OutputPointData->Metadata->HasAttribute(Settings->OutAttributeName))
@@ -584,6 +584,8 @@ void UPCGStaticMeshSpawnerSettings::RefreshMeshSelector()
 {
 	if (MeshSelectorType)
 	{
+		ensure(IsInGameThread());
+
 		if (MeshSelectorParameters)
 		{
 #if WITH_EDITOR
@@ -606,6 +608,8 @@ void UPCGStaticMeshSpawnerSettings::RefreshInstancePacker()
 {
 	if (InstanceDataPackerType)
 	{
+		ensure(IsInGameThread());
+
 		if (InstanceDataPackerParameters)
 		{
 #if WITH_EDITOR
diff --git a/Engine/Plugins/PCG/Source/PCG/Private/Elements/PCGSurfaceSampler.cpp b/Engine/Plugins/PCG/Source/PCG/Private/Elements/PCGSurfaceSampler.cpp
--- a/Engine/Plugins/PCG/Source/PCG/Private/Elements/PCGSurfaceSampler.cpp
+++ b/Engine/Plugins/PCG/Source/PCG/Private/Elements/PCGSurfaceSampler.cpp
@@ -150,7 +150,7 @@ namespace PCGSurfaceSampler
 
 	UPCGPointData* SampleSurface(FPCGContext* Context, const UPCGSurfaceData* InSurface, const UPCGSpatialData* InBoundingShape, const FBox& EffectiveBounds, const FSurfaceSamplerParams& ExecutionParams)
 	{
-		UPCGPointData* SampledData = NewObject<UPCGPointData>();
+		UPCGPointData* SampledData = FPCGContext::NewObject_AnyThread<UPCGPointData>(Context);
 		SampledData->InitializeFromData(InSurface);
 
 		FSurfaceSamplerData SamplerData;
@@ -437,7 +437,7 @@ bool FPCGSurfaceSamplerElement::PrepareDataInternal(FPCGContext* InContext) cons
 			const UPCGSurfaceData* GeneratingShape = GeneratingShapes[IterationIndex];
 			check(GeneratingShape);
 
-			OutState.OutputPoints = NewObject<UPCGPointData>();
+			OutState.OutputPoints = FPCGContext::NewObject_AnyThread<UPCGPointData>(Context);
 			OutState.OutputPoints->InitializeFromData(GeneratingShape);
 
 			// This bounds will be used to generate the pre-projected grid
@@ -519,7 +519,7 @@ bool FPCGSurfaceSamplerElement::ExecuteInternal(FPCGContext* InContext) const
 		for (FPCGTaggedData& Input : TimeSlicedContext->InputData.GetInputs())
 		{
 			FPCGTaggedData& Output = TimeSlicedContext->OutputData.TaggedData.Emplace_GetRef();
-			UPCGPointData* PointData = NewObject<UPCGPointData>();
+			UPCGPointData* PointData = FPCGContext::NewObject_AnyThread<UPCGPointData>(InContext);
 			PointData->InitializeFromData(Cast<UPCGSpatialData>(Input.Data));
 			Output.Data = PointData;
 		}
@@ -535,7 +535,7 @@ bool FPCGSurfaceSamplerElement::ExecuteInternal(FPCGContext* InContext) const
 		// This iteration resulted in an early out for no sampling operation. Early out with empty point data.
 		if (InitResult == EPCGTimeSliceInitResult::NoOperation)
 		{
-			Context->OutputData.TaggedData[IterationIndex].Data = NewObject<UPCGPointData>();
+			Context->OutputData.TaggedData[IterationIndex].Data = FPCGContext::NewObject_AnyThread<UPCGPointData>(Context);
 
 			return true;
 		}
diff --git a/Engine/Plugins/PCG/Source/PCG/Private/Elements/PCGTextureSampler.cpp b/Engine/Plugins/PCG/Source/PCG/Private/Elements/PCGTextureSampler.cpp
--- a/Engine/Plugins/PCG/Source/PCG/Private/Elements/PCGTextureSampler.cpp
+++ b/Engine/Plugins/PCG/Source/PCG/Private/Elements/PCGTextureSampler.cpp
@@ -211,7 +211,7 @@ bool FPCGTextureSamplerElement::ExecuteInternal(FPCGContext* InContext) const
 	TArray<FPCGTaggedData>& Outputs = Context->OutputData.TaggedData;
 	FPCGTaggedData& Output = Outputs.Emplace_GetRef();
 
-	UPCGTextureData* TextureData = NewObject<UPCGTextureData>();
+	UPCGTextureData* TextureData = FPCGContext::NewObject_AnyThread<UPCGTextureData>(Context);
 	Output.Data = TextureData;
 
 	AActor* OriginalActor = UPCGBlueprintHelpers::GetOriginalComponent(*Context)->GetOwner();
diff --git a/Engine/Plugins/PCG/Source/PCG/Private/Elements/PCGTransformPoints.cpp b/Engine/Plugins/PCG/Source/PCG/Private/Elements/PCGTransformPoints.cpp
--- a/Engine/Plugins/PCG/Source/PCG/Private/Elements/PCGTransformPoints.cpp
+++ b/Engine/Plugins/PCG/Source/PCG/Private/Elements/PCGTransformPoints.cpp
@@ -101,7 +101,7 @@ bool FPCGTransformPointsElement::ExecuteInternal(FPCGContext* Context) const
 
 		const TArray<FPCGPoint>& Points = PointData->GetPoints();
 
-		UPCGPointData* OutputData = NewObject<UPCGPointData>();
+		UPCGPointData* OutputData = FPCGContext::NewObject_AnyThread<UPCGPointData>(Context);
 		OutputData->InitializeFromData(PointData);
 		TArray<FPCGPoint>& OutputPoints = OutputData->GetMutablePoints();
 		Output.Data = OutputData;
diff --git a/Engine/Plugins/PCG/Source/PCG/Private/Elements/PCGTypedGetter.cpp b/Engine/Plugins/PCG/Source/PCG/Private/Elements/PCGTypedGetter.cpp
--- a/Engine/Plugins/PCG/Source/PCG/Private/Elements/PCGTypedGetter.cpp
+++ b/Engine/Plugins/PCG/Source/PCG/Private/Elements/PCGTypedGetter.cpp
@@ -120,7 +120,7 @@ void FPCGGetLandscapeDataElement::ProcessActors(FPCGContext* Context, const UPCG
 
 	if (!Landscapes.IsEmpty())
 	{
-		UPCGLandscapeData* LandscapeData = NewObject<UPCGLandscapeData>();
+		UPCGLandscapeData* LandscapeData = FPCGContext::NewObject_AnyThread<UPCGLandscapeData>(Context);
 		LandscapeData->Initialize(Landscapes, LandscapeBounds, Settings->SamplingProperties);
 		
 		FPCGTaggedData& TaggedData = Context->OutputData.TaggedData.Emplace_GetRef();
diff --git a/Engine/Plugins/PCG/Source/PCG/Private/Elements/PCGUnionElement.cpp b/Engine/Plugins/PCG/Source/PCG/Private/Elements/PCGUnionElement.cpp
--- a/Engine/Plugins/PCG/Source/PCG/Private/Elements/PCGUnionElement.cpp
+++ b/Engine/Plugins/PCG/Source/PCG/Private/Elements/PCGUnionElement.cpp
@@ -71,7 +71,7 @@ bool FPCGUnionElement::ExecuteInternal(FPCGContext* Context) const
 		// Create union or add to it
 		if (!UnionData)
 		{
-			UnionData = FirstSpatialData->UnionWith(SpatialData);
+			UnionData = FirstSpatialData->UnionWith(Context, SpatialData);
 			UnionData->SetType(Type);
 			UnionData->SetDensityFunction(DensityFunction);
 
diff --git a/Engine/Plugins/PCG/Source/PCG/Private/Elements/PCGVisualizeAttribute.cpp b/Engine/Plugins/PCG/Source/PCG/Private/Elements/PCGVisualizeAttribute.cpp
--- a/Engine/Plugins/PCG/Source/PCG/Private/Elements/PCGVisualizeAttribute.cpp
+++ b/Engine/Plugins/PCG/Source/PCG/Private/Elements/PCGVisualizeAttribute.cpp
@@ -18,6 +18,7 @@ FPCGElementPtr UPCGVisualizeAttributeSettings::CreateElement() const
 bool FPCGVisualizeAttribute::ExecuteInternal(FPCGContext* Context) const
 {
 	TRACE_CPUPROFILER_EVENT_SCOPE(FPCGVisualizeAttribute::ExecuteInternal);
+	check(IsInGameThread());
 
 	Context->OutputData = Context->InputData;
 
diff --git a/Engine/Plugins/PCG/Source/PCG/Private/Elements/PCGVolumeSampler.cpp b/Engine/Plugins/PCG/Source/PCG/Private/Elements/PCGVolumeSampler.cpp
--- a/Engine/Plugins/PCG/Source/PCG/Private/Elements/PCGVolumeSampler.cpp
+++ b/Engine/Plugins/PCG/Source/PCG/Private/Elements/PCGVolumeSampler.cpp
@@ -23,7 +23,7 @@ namespace PCGVolumeSampler
 {
 	UPCGPointData* SampleVolume(FPCGContext* Context, const FVolumeSamplerParams& SamplerSettings, const UPCGSpatialData* Volume, const UPCGSpatialData* BoundingShape)
 	{
-		UPCGPointData* Data = NewObject<UPCGPointData>();
+		UPCGPointData* Data = FPCGContext::NewObject_AnyThread<UPCGPointData>(Context);
 		Data->InitializeFromData(Volume);
 
 		const bool bTimeSliceIsEnabled = Context ? Context->TimeSliceIsEnabled() : false;
@@ -272,7 +272,7 @@ bool FPCGVolumeSamplerElement::PrepareDataInternal(FPCGContext* Context) const
 			check(GeneratingShape);
 
 			OutState.Volume = GeneratingShape;
-			OutState.OutputData = NewObject<UPCGPointData>();
+			OutState.OutputData = FPCGContext::NewObject_AnyThread<UPCGPointData>(Context);
 			OutState.OutputData->InitializeFromData(OutState.Volume);
 			Outputs[IterationIndex].Data = OutState.OutputData;
 
@@ -344,7 +344,7 @@ bool FPCGVolumeSamplerElement::ExecuteInternal(FPCGContext* Context) const
 		{
 			// TODO: Empty point data (to preserve previous behavior). Eventually, should be replaced with no output at all
 			FPCGTaggedData& Output = TimeSlicedContext->OutputData.TaggedData.Emplace_GetRef();
-			UPCGPointData* PointData = NewObject<UPCGPointData>();
+			UPCGPointData* PointData = FPCGContext::NewObject_AnyThread<UPCGPointData>(Context);
 			PointData->InitializeFromData(Cast<UPCGSpatialData>(Input.Data));
 			Output.Data = PointData;
 		}
@@ -358,7 +358,7 @@ bool FPCGVolumeSamplerElement::ExecuteInternal(FPCGContext* Context) const
 
 		if (InitResult == EPCGTimeSliceInitResult::NoOperation)
 		{
-			Context->OutputData.TaggedData[IterationIndex].Data = NewObject<UPCGPointData>();
+			Context->OutputData.TaggedData[IterationIndex].Data = FPCGContext::NewObject_AnyThread<UPCGPointData>(Context);
 
 			return true;
 		}
diff --git a/Engine/Plugins/PCG/Source/PCG/Private/Elements/PCGWorldQuery.cpp b/Engine/Plugins/PCG/Source/PCG/Private/Elements/PCGWorldQuery.cpp
--- a/Engine/Plugins/PCG/Source/PCG/Private/Elements/PCGWorldQuery.cpp
+++ b/Engine/Plugins/PCG/Source/PCG/Private/Elements/PCGWorldQuery.cpp
@@ -51,7 +51,7 @@ bool FPCGWorldVolumetricQueryElement::ExecuteInternal(FPCGContext* Context) cons
 	check(Context->SourceComponent.IsValid());
 	UWorld* World = Context->SourceComponent->GetWorld();
 
-	UPCGWorldVolumetricData* Data = NewObject<UPCGWorldVolumetricData>();
+	UPCGWorldVolumetricData* Data = FPCGContext::NewObject_AnyThread<UPCGWorldVolumetricData>(Context);
 	Data->Initialize(World);
 	Data->QueryParams = QueryParams;
 	Data->QueryParams.Initialize();
@@ -133,7 +133,7 @@ bool FPCGWorldRayHitQueryElement::ExecuteInternal(FPCGContext* Context) const
 	check(Context->SourceComponent.IsValid());
 	UWorld* World = Context->SourceComponent->GetWorld();
 
-	UPCGWorldRayHitData* Data = NewObject<UPCGWorldRayHitData>();
+	UPCGWorldRayHitData* Data = FPCGContext::NewObject_AnyThread<UPCGWorldRayHitData>(Context);
 	Data->Initialize(World, Transform, /*InBounds=*/FBox(EForceInit::ForceInit), LocalBounds);
 	Data->QueryParams = QueryParams;
 	Data->QueryParams.Initialize();
diff --git a/Engine/Plugins/PCG/Source/PCG/Private/Elements/PCGWorldRaycast.cpp b/Engine/Plugins/PCG/Source/PCG/Private/Elements/PCGWorldRaycast.cpp
--- a/Engine/Plugins/PCG/Source/PCG/Private/Elements/PCGWorldRaycast.cpp
+++ b/Engine/Plugins/PCG/Source/PCG/Private/Elements/PCGWorldRaycast.cpp
@@ -283,7 +283,7 @@ bool FPCGWorldRaycastElement::ExecuteInternal(FPCGContext* InContext) const
 		const TArray<FPCGPoint>* OriginInputPoints = OriginInputPointData ? &OriginInputPointData->GetPoints() : nullptr;
 
 		TArray<FPCGTaggedData>& OutputData = Context->OutputData.TaggedData;
-		UPCGPointData* OutputPointData = NewObject<UPCGPointData>();
+		UPCGPointData* OutputPointData = FPCGContext::NewObject_AnyThread<UPCGPointData>(Context);
 		// TODO: This should initialize from the End Points point data, if Origin point data doesn't exist
 		OutputPointData->InitializeFromData(OriginInputPointData);
 		FPCGTaggedData& Output = OutputData.Emplace_GetRef(OriginInputData);
diff --git a/Engine/Plugins/PCG/Source/PCG/Private/Graph/PCGGraphExecutor.cpp b/Engine/Plugins/PCG/Source/PCG/Private/Graph/PCGGraphExecutor.cpp
--- a/Engine/Plugins/PCG/Source/PCG/Private/Graph/PCGGraphExecutor.cpp
+++ b/Engine/Plugins/PCG/Source/PCG/Private/Graph/PCGGraphExecutor.cpp
@@ -138,6 +138,24 @@ bool FPCGGraphTask::IsApproximatelyEqual(const FPCGGraphTask& Other) const
 }
 #endif // WITH_EDITOR
 
+FPCGGraphActiveTask::~FPCGGraphActiveTask()
+{
+	if (Context)
+	{
+		// Remove Async Flags from Objects created in Async PCG Tasks so that they can get tracked by the GarbageCollector properly
+		// The reason to do this here and not in the FPCGContext destructor is that FPCGContext might get copied in Blueprint node execution causing multiple instances of the same AsyncObjects list
+		// The FPCGGraphActiveTask is the true owner of the original context (throughg TUniquePtr) so we let it do the cleanup
+		for (TObjectPtr<UObject>& AsyncObject : Context->AsyncObjects)
+		{
+			if (ensure(AsyncObject->HasAnyInternalFlags(EInternalObjectFlags::Async)))
+			{
+				AsyncObject->ClearInternalFlags(EInternalObjectFlags::Async);
+				ForEachObjectWithOuter(AsyncObject, [](UObject* SubObject) { SubObject->ClearInternalFlags(EInternalObjectFlags::Async); }, true);
+			}
+		}
+	}
+}
+
 FPCGGraphExecutor::FPCGGraphExecutor(UWorld* InWorld)
 	: World(InWorld)
 {
@@ -1413,6 +1431,7 @@ void FPCGGraphExecutor::BuildTaskInput(const FPCGGraphTask& Task, FPCGDataCollec
 void FPCGGraphExecutor::CombineParams(FPCGTaskId InTaskId, FPCGDataCollection& InTaskInput)
 {
 	TRACE_CPUPROFILER_EVENT_SCOPE(FPCGGraphExecutor::CombineParams);
+	ensure(IsInGameThread());
 
 	TArray<FPCGTaggedData> AllParamsData = InTaskInput.GetParamsByPin(PCGPinConstants::DefaultParamsLabel);
 	if (AllParamsData.Num() > 1)
@@ -1425,7 +1444,7 @@ void FPCGGraphExecutor::CombineParams(FPCGTaskId InTaskId, FPCGDataCollection& I
 			const UPCGParamData* ParamData = CastChecked<UPCGParamData>(TaggedDatum.Data);
 			if (!CombinedParamData)
 			{
-				CombinedParamData = ParamData->DuplicateData();
+				CombinedParamData = ParamData->DuplicateData(nullptr);
 			}
 			else
 			{
diff --git a/Engine/Plugins/PCG/Source/PCG/Private/Graph/PCGGraphExecutor.h b/Engine/Plugins/PCG/Source/PCG/Private/Graph/PCGGraphExecutor.h
--- a/Engine/Plugins/PCG/Source/PCG/Private/Graph/PCGGraphExecutor.h
+++ b/Engine/Plugins/PCG/Source/PCG/Private/Graph/PCGGraphExecutor.h
@@ -111,6 +111,15 @@ struct FPCGGraphScheduleTask
 
 struct FPCGGraphActiveTask
 {
+	FPCGGraphActiveTask() = default;
+	~FPCGGraphActiveTask();
+	
+	FPCGGraphActiveTask(const FPCGGraphActiveTask&) = delete;
+	FPCGGraphActiveTask& operator=(const FPCGGraphActiveTask&) = delete;
+
+	FPCGGraphActiveTask(FPCGGraphActiveTask&&) = default;
+	FPCGGraphActiveTask& operator=(FPCGGraphActiveTask&&) = default;
+
 	TArray<FPCGGraphTaskInput> Inputs;
 	FPCGElementPtr Element;
 	TUniquePtr<FPCGContext> Context;
diff --git a/Engine/Plugins/PCG/Source/PCG/Private/Metadata/PCGMetadataPartitionCommon.cpp b/Engine/Plugins/PCG/Source/PCG/Private/Metadata/PCGMetadataPartitionCommon.cpp
--- a/Engine/Plugins/PCG/Source/PCG/Private/Metadata/PCGMetadataPartitionCommon.cpp
+++ b/Engine/Plugins/PCG/Source/PCG/Private/Metadata/PCGMetadataPartitionCommon.cpp
@@ -2,6 +2,7 @@
 
 #include "Metadata/PCGMetadataPartitionCommon.h"
 
+#include "PCGContext.h"
 #include "PCGModule.h"
 #include "PCGParamData.h"
 #include "Data/PCGPointData.h"
@@ -437,7 +438,7 @@ namespace PCGMetadataPartitionCommon
 				continue;
 			}
 
-			UPCGPointData* CurrentPointData = NewObject<UPCGPointData>();
+			UPCGPointData* CurrentPointData = FPCGContext::NewObject_AnyThread<UPCGPointData>(InOptionalContext);
 			PartitionedData.Add(CurrentPointData);
 			CurrentPointData->InitializeFromData(InData);
 
@@ -488,14 +489,14 @@ namespace PCGMetadataPartitionCommon
 
 			if (InSpatialData)
 			{
-				UPCGSpatialData* NewData = NewObject<UPCGSpatialData>();
+				UPCGSpatialData* NewData = FPCGContext::NewObject_AnyThread<UPCGSpatialData>(InOptionalContext);
 				NewData->InitializeFromData(InSpatialData);
 				NewMetadata = NewData->Metadata;
 				PartitionedData.Add(NewData);
 			}
 			else
 			{
-				UPCGParamData* NewData = NewObject<UPCGParamData>();
+				UPCGParamData* NewData = FPCGContext::NewObject_AnyThread<UPCGParamData>(InOptionalContext);
 				NewData->Metadata->AddAttributes(OriginalMetadata);
 				NewMetadata = NewData->Metadata;
 				PartitionedData.Add(NewData);
diff --git a/Engine/Plugins/PCG/Source/PCG/Private/PCGComponent.cpp b/Engine/Plugins/PCG/Source/PCG/Private/PCGComponent.cpp
--- a/Engine/Plugins/PCG/Source/PCG/Private/PCGComponent.cpp
+++ b/Engine/Plugins/PCG/Source/PCG/Private/PCGComponent.cpp
@@ -567,7 +567,7 @@ void UPCGComponent::PostProcessGraph(const FBox& InNewBounds, bool bInGenerated,
 				{
 					// TODO: outering the first layer might not be sufficient here - might need to expose
 					// some methods in the data to traverse all the data to outer everything for serialization
-					if (UPCGData* DuplicatedData = TaggedData.Data->DuplicateData())
+					if (UPCGData* DuplicatedData = TaggedData.Data->DuplicateData(Context))
 					{
 						FPCGTaggedData& DuplicatedTaggedData = GeneratedGraphOutput.TaggedData.Add_GetRef(TaggedData);
 						DuplicatedTaggedData.Data = DuplicatedData;
@@ -2849,7 +2849,7 @@ FPCGDataCollection UPCGComponent::CreateActorPCGDataCollection(AActor* Actor, co
 
 			if (OriginalComponentSpatialData)
 			{
-				Result = Result->IntersectWith(OriginalComponentSpatialData);
+				Result = Result->IntersectWith(nullptr, OriginalComponentSpatialData);
 			}
 		}
 
@@ -3160,11 +3160,11 @@ UPCGData* UPCGComponent::CreateInputPCGData()
 		// so intersections (such as volume X partition actor) get picked up properly
 		if (ActorSpatialData->GetDimension() >= 3)
 		{
-			return LandscapeData->IntersectWith(ActorSpatialData);
+			return LandscapeData->IntersectWith(nullptr, ActorSpatialData);
 		}
 		else
 		{
-			return ActorSpatialData->ProjectOn(LandscapeData);
+			return ActorSpatialData->ProjectOn(nullptr, LandscapeData);
 		}
 	}
 	else
diff --git a/Engine/Plugins/PCG/Source/PCG/Private/PCGData.cpp b/Engine/Plugins/PCG/Source/PCG/Private/PCGData.cpp
--- a/Engine/Plugins/PCG/Source/PCG/Private/PCGData.cpp
+++ b/Engine/Plugins/PCG/Source/PCG/Private/PCGData.cpp
@@ -89,9 +89,18 @@ void UPCGData::Flatten()
 
 UPCGData* UPCGData::DuplicateData(bool bInitializeMetadata) const
 {
+	//@todo_pcg: Not sure StaticDuplicateObject is thread safe, most UPCGDAta subclasses override DuplicateData
+	ensure(IsInGameThread());
 	return Cast<UPCGData>(StaticDuplicateObject(this, GetTransientPackage()));
 }
 
+UPCGData* UPCGData::DuplicateData(FPCGContext* Context, bool bInitializeMetadata) const
+{
+PRAGMA_DISABLE_DEPRECATION_WARNINGS
+	return DuplicateData(bInitializeMetadata);
+PRAGMA_ENABLE_DEPRECATION_WARNINGS
+}
+
 bool FPCGTaggedData::operator==(const FPCGTaggedData& Other) const
 {
 	return Data == Other.Data &&
diff --git a/Engine/Plugins/PCG/Source/PCG/Private/PCGParamData.cpp b/Engine/Plugins/PCG/Source/PCG/Private/PCGParamData.cpp
--- a/Engine/Plugins/PCG/Source/PCG/Private/PCGParamData.cpp
+++ b/Engine/Plugins/PCG/Source/PCG/Private/PCGParamData.cpp
@@ -1,7 +1,9 @@
 // Copyright Epic Games, Inc. All Rights Reserved.
 
 #include "PCGParamData.h"
+#include "PCGContext.h"
 
+#include "Elements/PCGExecuteBlueprint.h"
 #include "Metadata/PCGAttributePropertySelector.h"
 #include "Metadata/Accessors/IPCGAttributeAccessor.h"
 #include "Metadata/Accessors/PCGAttributeAccessorHelpers.h"
@@ -91,10 +93,15 @@ int64 UPCGParamData::FindOrAddMetadataKey(const FName& InName)
 	}
 }
 
-UPCGParamData* UPCGParamData::FilterParamsByName(const FName& InName) const
+UPCGParamData* UPCGParamData::K2_FilterParamsByName(const FName& InName) const
+{
+	return FilterParamsByName(UPCGBlueprintElement::ResolveContext(), InName);
+}
+
+UPCGParamData* UPCGParamData::FilterParamsByName(FPCGContext* Context, const FName& InName) const
 {
 	PCGMetadataEntryKey EntryKey = FindMetadataKey(InName);
-	UPCGParamData* NewParams = FilterParamsByKey(EntryKey);
+	UPCGParamData* NewParams = FilterParamsByKey(Context, EntryKey);
 
 	if (EntryKey != PCGInvalidEntryKey)
 	{
@@ -105,9 +112,14 @@ UPCGParamData* UPCGParamData::FilterParamsByName(const FName& InName) const
 	return NewParams;
 }
 
-UPCGParamData* UPCGParamData::FilterParamsByKey(int64 InKey) const
+UPCGParamData* UPCGParamData::K2_FilterParamsByKey(int64 InKey) const
+{
+	return FilterParamsByKey(UPCGBlueprintElement::ResolveContext(), InKey);
+}
+
+UPCGParamData* UPCGParamData::FilterParamsByKey(FPCGContext* Context, int64 InKey) const
 {
-	UPCGParamData* NewParams = NewObject<UPCGParamData>();
+	UPCGParamData* NewParams = FPCGContext::NewObject_AnyThread<UPCGParamData>(Context);
 
 	// Here instead of parenting the metadata, we will create a copy
 	// so that the only entry in the metadata (if any) will have the 0 key.
@@ -163,9 +175,9 @@ void UPCGParamData::SetLastSelector(const FPCGAttributePropertySelector& InSelec
 	CachedLastSelector.ImportFromOtherSelector(InSelector);
 }
 
-UPCGParamData* UPCGParamData::DuplicateData(bool bInitializeMetadata) const
+UPCGParamData* UPCGParamData::DuplicateData(FPCGContext* Context, bool bInitializeMetadata) const
 {
-	UPCGParamData* NewParamData = NewObject<UPCGParamData>();
+	UPCGParamData* NewParamData = FPCGContext::NewObject_AnyThread<UPCGParamData>(Context);
 	if (bInitializeMetadata)
 	{
 		check(NewParamData && NewParamData->Metadata);
diff --git a/Engine/Plugins/PCG/Source/PCG/Private/Tests/Data/PCGDifferenceDataTest.cpp b/Engine/Plugins/PCG/Source/PCG/Private/Tests/Data/PCGDifferenceDataTest.cpp
--- a/Engine/Plugins/PCG/Source/PCG/Private/Tests/Data/PCGDifferenceDataTest.cpp
+++ b/Engine/Plugins/PCG/Source/PCG/Private/Tests/Data/PCGDifferenceDataTest.cpp
@@ -34,7 +34,7 @@ bool FPCGDifferenceDataTest::RunTest(const FString& Parameters)
 		}
 
 		// Create difference data with Subtract
-		const UPCGDifferenceData* DifferenceData = OriginalVolume->Subtract(DifferenceVolume);
+		const UPCGDifferenceData* DifferenceData = OriginalVolume->Subtract(nullptr, DifferenceVolume);
 
 		if (!TestNotNull("DifferenceData", DifferenceData))
 		{
diff --git a/Engine/Plugins/PCG/Source/PCG/Private/Tests/Data/PCGIntersectionDataTest.cpp b/Engine/Plugins/PCG/Source/PCG/Private/Tests/Data/PCGIntersectionDataTest.cpp
--- a/Engine/Plugins/PCG/Source/PCG/Private/Tests/Data/PCGIntersectionDataTest.cpp
+++ b/Engine/Plugins/PCG/Source/PCG/Private/Tests/Data/PCGIntersectionDataTest.cpp
@@ -18,10 +18,10 @@ bool FPCGIntersectionDataTest::RunTest(const FString& Parameters)
 	UPCGVolumeData* Volume = PCGTestsCommon::CreateVolumeData(FBox::BuildAABB(FVector::ZeroVector, FVector::OneVector * 100));
 
 	// Create intersections
-	UPCGIntersectionData* InsideVolume = InsidePoint->IntersectWith(Volume);
-	UPCGIntersectionData* VolumeInside = Volume->IntersectWith(InsidePoint);
-	UPCGIntersectionData* OutsideVolume = OutsidePoint->IntersectWith(Volume);
-	UPCGIntersectionData* VolumeOutside = Volume->IntersectWith(OutsidePoint);
+	UPCGIntersectionData* InsideVolume = InsidePoint->IntersectWith(nullptr, Volume);
+	UPCGIntersectionData* VolumeInside = Volume->IntersectWith(nullptr, InsidePoint);
+	UPCGIntersectionData* OutsideVolume = OutsidePoint->IntersectWith(nullptr, Volume);
+	UPCGIntersectionData* VolumeOutside = Volume->IntersectWith(nullptr, OutsidePoint);
 
 	auto ValidateInsideIntersection = [this, InsidePoint](UPCGIntersectionData* Intersection)
 	{
diff --git a/Engine/Plugins/PCG/Source/PCG/Private/Tests/Metadata/PCGMetadataFlattenTest.cpp b/Engine/Plugins/PCG/Source/PCG/Private/Tests/Metadata/PCGMetadataFlattenTest.cpp
--- a/Engine/Plugins/PCG/Source/PCG/Private/Tests/Metadata/PCGMetadataFlattenTest.cpp
+++ b/Engine/Plugins/PCG/Source/PCG/Private/Tests/Metadata/PCGMetadataFlattenTest.cpp
@@ -44,7 +44,7 @@ bool FPCGMetadataFlatten::RunTest(const FString& Parameters)
 	// for Attribute 2: ["0", "Default", "2", "Default", "0", ...]
 	UTEST_EQUAL("RootMetadata has 10 entries", RootPointData->Metadata->GetItemCountForChild(), 10);
 
-	UPCGPointData* FirstChildPointData = Cast<UPCGPointData>(RootPointData->DuplicateData());
+	UPCGPointData* FirstChildPointData = Cast<UPCGPointData>(RootPointData->DuplicateData(nullptr));
 	Attribute1 = FirstChildPointData->Metadata->GetMutableTypedAttribute<float>(Attribute1Name);
 	Attribute2 = FirstChildPointData->Metadata->GetMutableTypedAttribute<FString>(Attribute2Name);
 
@@ -109,7 +109,7 @@ bool FPCGMetadataFlatten::RunTest(const FString& Parameters)
 	UTEST_EQUAL("SecondChildMetadata has 25 entries", SecondChildPointData->Metadata->GetItemCountForChild(), 25);
 
 	// For final set, duplicate the data and flatten it
-	UPCGPointData* FinalPointData = Cast<UPCGPointData>(SecondChildPointData->DuplicateData());
+	UPCGPointData* FinalPointData = Cast<UPCGPointData>(SecondChildPointData->DuplicateData(nullptr));
 	FinalPointData->Flatten();
 
 	Attribute1 = FinalPointData->Metadata->GetMutableTypedAttribute<float>(Attribute1Name);
diff --git a/Engine/Plugins/PCG/Source/PCG/Public/Data/PCGCollisionShapeData.h b/Engine/Plugins/PCG/Source/PCG/Public/Data/PCGCollisionShapeData.h
--- a/Engine/Plugins/PCG/Source/PCG/Public/Data/PCGCollisionShapeData.h
+++ b/Engine/Plugins/PCG/Source/PCG/Public/Data/PCGCollisionShapeData.h
@@ -34,7 +34,7 @@ public:
 	// TODO needs an implementation to support projection
 	//virtual bool ProjectPoint(const FTransform& InTransform, const FBox& InBounds, const FPCGProjectionParams& InParams, FPCGPoint& OutPoint, UPCGMetadata* OutMetadata) const;
 protected:
-	virtual UPCGSpatialData* CopyInternal() const override;
+	virtual UPCGSpatialData* CopyInternal(FPCGContext* Context) const override;
 	//~End UPCGSpatialData interface
 
 public:
diff --git a/Engine/Plugins/PCG/Source/PCG/Public/Data/PCGCollisionWrapperData.h b/Engine/Plugins/PCG/Source/PCG/Public/Data/PCGCollisionWrapperData.h
--- a/Engine/Plugins/PCG/Source/PCG/Public/Data/PCGCollisionWrapperData.h
+++ b/Engine/Plugins/PCG/Source/PCG/Public/Data/PCGCollisionWrapperData.h
@@ -63,7 +63,7 @@ public:
 	virtual const UPCGPointData* ToPointData(FPCGContext* Context, const FBox& InBounds = FBox(EForceInit::ForceInit)) const override;
 	virtual bool SamplePoint(const FTransform& Transform, const FBox& Bounds, FPCGPoint& OutPoint, UPCGMetadata* OutMetadata) const override;
 protected:
-	virtual UPCGSpatialData* CopyInternal() const override;
+	virtual UPCGSpatialData* CopyInternal(FPCGContext* Context) const override;
 	// ~End UPCGSpatialData interface
 
 	// For performance reasons, we keep a raw pointer to the point data in editor.
diff --git a/Engine/Plugins/PCG/Source/PCG/Public/Data/PCGDifferenceData.h b/Engine/Plugins/PCG/Source/PCG/Public/Data/PCGDifferenceData.h
--- a/Engine/Plugins/PCG/Source/PCG/Public/Data/PCGDifferenceData.h
+++ b/Engine/Plugins/PCG/Source/PCG/Public/Data/PCGDifferenceData.h
@@ -34,8 +34,13 @@ public:
 	UFUNCTION(BlueprintCallable, Category = SpatialData)
 	void Initialize(const UPCGSpatialData* InData);
 
-	UFUNCTION(BlueprintCallable, Category = SpatialData)
-	void AddDifference(const UPCGSpatialData* InDifference);
+	UFUNCTION(BlueprintCallable, Category = SpatialData, meta = (DisplayName = "Add Difference"))
+	void K2_AddDifference(const UPCGSpatialData* InDifference);
+
+	UE_DEPRECATED(5.5, "Call/Implement version with FPCGContext parameter")
+	void AddDifference(const UPCGSpatialData* InDifference) { AddDifference(nullptr, InDifference); }
+
+	void AddDifference(FPCGContext* InContext, const UPCGSpatialData* InDifference);
 
 	UFUNCTION(BlueprintCallable, Category = Settings)
 	void SetDensityFunction(EPCGDifferenceDensityFunction InDensityFunction);
@@ -72,7 +77,7 @@ public:
 	virtual const UPCGSpatialData* FindFirstConcreteShapeFromNetwork() const override { return GetSource() ? GetSource()->FindFirstConcreteShapeFromNetwork() : nullptr; }
 
 protected:
-	virtual UPCGSpatialData* CopyInternal() const override;
+	virtual UPCGSpatialData* CopyInternal(FPCGContext* Context) const override;
 	//~End UPCGSpatialData interface
 
 public:
diff --git a/Engine/Plugins/PCG/Source/PCG/Public/Data/PCGIntersectionData.h b/Engine/Plugins/PCG/Source/PCG/Public/Data/PCGIntersectionData.h
--- a/Engine/Plugins/PCG/Source/PCG/Public/Data/PCGIntersectionData.h
+++ b/Engine/Plugins/PCG/Source/PCG/Public/Data/PCGIntersectionData.h
@@ -49,7 +49,7 @@ public:
 	virtual const UPCGSpatialData* FindFirstConcreteShapeFromNetwork() const override;
 
 protected:
-	virtual UPCGSpatialData* CopyInternal() const override;
+	virtual UPCGSpatialData* CopyInternal(FPCGContext* Context) const override;
 	//~End UPCGSpatialData interface
 
 public:
diff --git a/Engine/Plugins/PCG/Source/PCG/Public/Data/PCGLandscapeData.h b/Engine/Plugins/PCG/Source/PCG/Public/Data/PCGLandscapeData.h
--- a/Engine/Plugins/PCG/Source/PCG/Public/Data/PCGLandscapeData.h
+++ b/Engine/Plugins/PCG/Source/PCG/Public/Data/PCGLandscapeData.h
@@ -67,7 +67,7 @@ public:
 	virtual bool ProjectPoint(const FTransform& InTransform, const FBox& InBounds, const FPCGProjectionParams& InParams, FPCGPoint& OutPoint, UPCGMetadata* OutMetadata) const override;
 	virtual bool HasNonTrivialTransform() const override { return true; }
 protected:
-	virtual UPCGSpatialData* CopyInternal() const override;
+	virtual UPCGSpatialData* CopyInternal(FPCGContext* Context) const override;
 	//~End UPCGSpatialData interface
 
 public:
diff --git a/Engine/Plugins/PCG/Source/PCG/Public/Data/PCGLandscapeSplineData.h b/Engine/Plugins/PCG/Source/PCG/Public/Data/PCGLandscapeSplineData.h
--- a/Engine/Plugins/PCG/Source/PCG/Public/Data/PCGLandscapeSplineData.h
+++ b/Engine/Plugins/PCG/Source/PCG/Public/Data/PCGLandscapeSplineData.h
@@ -45,7 +45,7 @@ public:
 	virtual FBox GetBounds() const override;
 	virtual bool SamplePoint(const FTransform& Transform, const FBox& Bounds, FPCGPoint& OutPoint, UPCGMetadata* OutMetadata) const override;
 protected:
-	virtual UPCGSpatialData* CopyInternal() const override;
+	virtual UPCGSpatialData* CopyInternal(FPCGContext* Context) const override;
 	//~End UPCGSpatialData interface
 
 	/** Recompute the reparameterization of the spline by distance. */
diff --git a/Engine/Plugins/PCG/Source/PCG/Public/Data/PCGPointData.h b/Engine/Plugins/PCG/Source/PCG/Public/Data/PCGPointData.h
--- a/Engine/Plugins/PCG/Source/PCG/Public/Data/PCGPointData.h
+++ b/Engine/Plugins/PCG/Source/PCG/Public/Data/PCGPointData.h
@@ -89,7 +89,7 @@ public:
 	virtual bool ProjectPoint(const FTransform& InTransform, const FBox& InBounds, const FPCGProjectionParams& InParams, FPCGPoint& OutPoint, UPCGMetadata* OutMetadata) const override;
 
 protected:
-	virtual UPCGSpatialData* CopyInternal() const override;
+	virtual UPCGSpatialData* CopyInternal(FPCGContext* Context) const override;
 	// ~End UPCGSpatialData interface
 
 public:
diff --git a/Engine/Plugins/PCG/Source/PCG/Public/Data/PCGPrimitiveData.h b/Engine/Plugins/PCG/Source/PCG/Public/Data/PCGPrimitiveData.h
--- a/Engine/Plugins/PCG/Source/PCG/Public/Data/PCGPrimitiveData.h
+++ b/Engine/Plugins/PCG/Source/PCG/Public/Data/PCGPrimitiveData.h
@@ -27,7 +27,7 @@ public:
 	// TODO needs an implementation to support projection
 	//virtual bool ProjectPoint(const FTransform& InTransform, const FBox& InBounds, const FPCGProjectionParams& InParams, FPCGPoint& OutPoint, UPCGMetadata* OutMetadata) const;
 protected:
-	virtual UPCGSpatialData* CopyInternal() const override;
+	virtual UPCGSpatialData* CopyInternal(FPCGContext* Context) const override;
 	//~End UPCGSpatialData interface
 
 public:
diff --git a/Engine/Plugins/PCG/Source/PCG/Public/Data/PCGProjectionData.h b/Engine/Plugins/PCG/Source/PCG/Public/Data/PCGProjectionData.h
--- a/Engine/Plugins/PCG/Source/PCG/Public/Data/PCGProjectionData.h
+++ b/Engine/Plugins/PCG/Source/PCG/Public/Data/PCGProjectionData.h
@@ -35,7 +35,7 @@ public:
 	virtual bool HasNonTrivialTransform() const override;
 	virtual bool RequiresCollapseToSample() const override;
 protected:
-	virtual UPCGSpatialData* CopyInternal() const override;
+	virtual UPCGSpatialData* CopyInternal(FPCGContext* Context) const override;
 	//~End UPCGSpatialData interface
 
 public:
diff --git a/Engine/Plugins/PCG/Source/PCG/Public/Data/PCGRenderTargetData.h b/Engine/Plugins/PCG/Source/PCG/Public/Data/PCGRenderTargetData.h
--- a/Engine/Plugins/PCG/Source/PCG/Public/Data/PCGRenderTargetData.h
+++ b/Engine/Plugins/PCG/Source/PCG/Public/Data/PCGRenderTargetData.h
@@ -24,7 +24,7 @@ public:
 
 	//~Begin UPCGSpatialData interface
 protected:
-	virtual UPCGSpatialData* CopyInternal() const override;
+	virtual UPCGSpatialData* CopyInternal(FPCGContext* Context) const override;
 	//~End UPCGSpatialData interface
 
 public:
diff --git a/Engine/Plugins/PCG/Source/PCG/Public/Data/PCGSpatialData.h b/Engine/Plugins/PCG/Source/PCG/Public/Data/PCGSpatialData.h
--- a/Engine/Plugins/PCG/Source/PCG/Public/Data/PCGSpatialData.h
+++ b/Engine/Plugins/PCG/Source/PCG/Public/Data/PCGSpatialData.h
@@ -55,7 +55,7 @@ public:
 	*   They are mainly cached values (and octree for points).
 	*   TODO: If we want to also copy those values (can be an optimization), we need to guard the copy.
 	*/
-	virtual UPCGSpatialData* DuplicateData(bool bInitializeMetadata = true) const override;
+	virtual UPCGSpatialData* DuplicateData(FPCGContext* Context, bool bInitializeMetadata = true) const override;
 	// ~End UPCGData interface
 
 	/** Returns the dimension of the data type, which has nothing to do with the dimension of its points */
@@ -100,7 +100,7 @@ public:
 	virtual void SamplePoints(const TArrayView<const TPair<FTransform, FBox>>& Samples, const TArrayView<FPCGPoint>& OutPoints, UPCGMetadata* OutMetadata) const;
 
 	/** Sample rotation, scale and other attributes from this data at the query position. Returns true if Transform location and Bounds overlaps this data. */
-	UFUNCTION(BlueprintCallable, Category = SpatialData, meta = (DisplayName="SamplePoint"))
+	UFUNCTION(BlueprintCallable, Category = SpatialData, meta = (DisplayName = "Sample Point"))
 	bool K2_SamplePoint(const FTransform& Transform, const FBox& Bounds, FPCGPoint& OutPoint, UPCGMetadata* OutMetadata) const;
 
 	/** Project the query point onto this data, and sample point and metadata information at the projected position. Returns true if successful. */
@@ -112,7 +112,7 @@ public:
 	*/
 	virtual void ProjectPoints(const TArrayView<const TPair<FTransform, FBox>>& Samples, const FPCGProjectionParams& InParams, const TArrayView<FPCGPoint>& OutPoints, UPCGMetadata* OutMetadata) const;
 
-	UFUNCTION(BlueprintCallable, Category = SpatialData, meta = (DisplayName="ProjectPoint"))
+	UFUNCTION(BlueprintCallable, Category = SpatialData, meta = (DisplayName = "Project Point"))
 	bool K2_ProjectPoint(const FTransform& InTransform, const FBox& InBounds, const FPCGProjectionParams& InParams, FPCGPoint& OutPoint, UPCGMetadata* OutMetadata) const;
 
 	/** Returns true if the data has a non-trivial transform */
@@ -120,19 +120,39 @@ public:
 	virtual bool HasNonTrivialTransform() const { return false; }
 
 	/** Returns a specialized data to intersect with another data */
-	UFUNCTION(BlueprintCallable, Category = SpatialData)
-	virtual UPCGIntersectionData* IntersectWith(const UPCGSpatialData* InOther) const;
+	UFUNCTION(BlueprintCallable, Category = SpatialData, meta = (DisplayName = "Intersect With"))
+	UPCGIntersectionData* K2_IntersectWith(const UPCGSpatialData* InOther) const;
+
+	UE_DEPRECATED(5.5, "Call/Implement version with FPCGContext parameter")
+	virtual UPCGIntersectionData* IntersectWith(const UPCGSpatialData* InOther) const { return IntersectWith(nullptr, InOther); }
+
+	virtual UPCGIntersectionData* IntersectWith(FPCGContext* InContext, const UPCGSpatialData* InOther) const;
 
 	/** Returns a specialized data to project this on another data of equal or higher dimension. Returns copy of this data if projection fails. */
-	UFUNCTION(BlueprintCallable, Category = SpatialData, meta=(AutoCreateRefTerm="InParams"))
-	virtual UPCGSpatialData* ProjectOn(const UPCGSpatialData* InOther, const FPCGProjectionParams& InParams = FPCGProjectionParams()) const;
+	UFUNCTION(BlueprintCallable, Category = SpatialData, meta = (AutoCreateRefTerm = "InParams", DisplayName  ="Project On"))
+	UPCGSpatialData* K2_ProjectOn(const UPCGSpatialData* InOther, const FPCGProjectionParams& InParams = FPCGProjectionParams()) const;
+
+	UE_DEPRECATED(5.5, "Call/Implement version with FPCGContext parameter")
+	virtual UPCGSpatialData* ProjectOn(const UPCGSpatialData* InOther, const FPCGProjectionParams& InParams = FPCGProjectionParams()) const { return ProjectOn(nullptr, InOther, InParams); }
+
+	virtual UPCGSpatialData* ProjectOn(FPCGContext* InContext, const UPCGSpatialData* InOther, const FPCGProjectionParams& InParams = FPCGProjectionParams()) const;
 
 	/** Returns a specialized data to union this with another data */
-	UFUNCTION(BlueprintCallable, Category = SpatialData)
-	virtual UPCGUnionData* UnionWith(const UPCGSpatialData* InOther) const;
+	UFUNCTION(BlueprintCallable, Category = SpatialData, meta = (DisplayName = "Union With"))
+	UPCGUnionData* K2_UnionWith(const UPCGSpatialData* InOther) const;
 
-	UFUNCTION(BlueprintCallable, Category = SpatialData)
-	virtual UPCGDifferenceData* Subtract(const UPCGSpatialData* InOther) const;
+	UE_DEPRECATED(5.5, "Call/Implement version with FPCGContext parameter")
+	virtual UPCGUnionData* UnionWith(const UPCGSpatialData* InOther) const { return UnionWith(nullptr, InOther); }
+
+	virtual UPCGUnionData* UnionWith(FPCGContext* InContext, const UPCGSpatialData* InOther) const;
+
+	UFUNCTION(BlueprintCallable, Category = SpatialData, meta = (DisplayName = "Subtract"))
+	UPCGDifferenceData* K2_Subtract(const UPCGSpatialData* InOther) const;
+
+	UE_DEPRECATED(5.5, "Call/Implement version with FPCGContext parameter")
+	virtual UPCGDifferenceData* Subtract(const UPCGSpatialData* InOther) const { return Subtract(nullptr, InOther); }
+
+	virtual UPCGDifferenceData* Subtract(FPCGContext* InContext, const UPCGSpatialData* InOther) const;
 
 	UFUNCTION(BlueprintCallable, Category = Metadata)
 	virtual const UPCGMetadata* ConstMetadata() const override { return Metadata; }
@@ -164,7 +184,10 @@ public:
 	TObjectPtr<UPCGMetadata> Metadata = nullptr;
 
 protected:
-	virtual UPCGSpatialData* CopyInternal() const PURE_VIRTUAL(UPCGSpatialData::CopyInternal, return nullptr;);
+	UE_DEPRECATED(5.5, "Call/Implement version with FPCGContext parameter")
+	virtual UPCGSpatialData* CopyInternal() const { return nullptr; }
+
+	virtual UPCGSpatialData* CopyInternal(FPCGContext* Context) const PURE_VIRTUAL(UPCGSpatialData::CopyInternal, return nullptr;);
 
 private:
 	/** Cache to keep track of the latest attribute manipulated on this data. */
diff --git a/Engine/Plugins/PCG/Source/PCG/Public/Data/PCGSplineData.h b/Engine/Plugins/PCG/Source/PCG/Public/Data/PCGSplineData.h
--- a/Engine/Plugins/PCG/Source/PCG/Public/Data/PCGSplineData.h
+++ b/Engine/Plugins/PCG/Source/PCG/Public/Data/PCGSplineData.h
@@ -51,9 +51,9 @@ public:
 	//~Begin UPCGSpatialData interface
 	virtual FBox GetBounds() const override;
 	virtual bool SamplePoint(const FTransform& Transform, const FBox& Bounds, FPCGPoint& OutPoint, UPCGMetadata* OutMetadata) const override;
-	virtual UPCGSpatialData* ProjectOn(const UPCGSpatialData* InOther, const FPCGProjectionParams& InParams = FPCGProjectionParams()) const override;
+	virtual UPCGSpatialData* ProjectOn(FPCGContext* InContext, const UPCGSpatialData* InOther, const FPCGProjectionParams& InParams = FPCGProjectionParams()) const override;
 protected:
-	virtual UPCGSpatialData* CopyInternal() const override;
+	virtual UPCGSpatialData* CopyInternal(FPCGContext* Context) const override;
 	//~End UPCGSpatialData interface
 
 	virtual void CopySplineData(UPCGSplineData* InCopy) const;
@@ -96,7 +96,7 @@ protected:
 	FVector2D Project(const FVector& InVector) const;
 
 	//~Begin UPCGSpatialData interface
-	virtual UPCGSpatialData* CopyInternal() const override;
+	virtual UPCGSpatialData* CopyInternal(FPCGContext* Context) const override;
 	//~End UPCGSpatialData interface
 };
 
diff --git a/Engine/Plugins/PCG/Source/PCG/Public/Data/PCGSplineInteriorSurfaceData.h b/Engine/Plugins/PCG/Source/PCG/Public/Data/PCGSplineInteriorSurfaceData.h
--- a/Engine/Plugins/PCG/Source/PCG/Public/Data/PCGSplineInteriorSurfaceData.h
+++ b/Engine/Plugins/PCG/Source/PCG/Public/Data/PCGSplineInteriorSurfaceData.h
@@ -20,7 +20,10 @@ class UPCGSplineInteriorSurfaceData : public UPCGSurfaceData
 	GENERATED_BODY()
 
 public:
-	PCG_API void Initialize(const UPCGSplineData* InSplineData);
+	UE_DEPRECATED(5.5, "Call version with FPCGContext parameter")
+	PCG_API void Initialize(const UPCGSplineData* InSplineData) { Initialize(nullptr, InSplineData); }
+
+	PCG_API void Initialize(FPCGContext* Context, const UPCGSplineData* InSplineData);
 
 	// ~Begin UObject interface
 	virtual void PostLoad() override;
@@ -38,7 +41,7 @@ public:
 	virtual bool ProjectPoint(const FTransform& InTransform, const FBox& InBounds, const FPCGProjectionParams& InParams, FPCGPoint& OutPoint, UPCGMetadata* OutMetadata) const override;
 	virtual bool HasNonTrivialTransform() const override { return true; }
 protected:
-	virtual UPCGSpatialData* CopyInternal() const override;
+	virtual UPCGSpatialData* CopyInternal(FPCGContext* Context) const override;
 	//~End UPCGSpatialData interface
 
 public:
@@ -48,7 +51,7 @@ public:
 
 protected:
 	/** Recompute cached data used for sampling. */
-	void CacheData();
+	void CacheData(FPCGContext* Context);
 
 	/** True if the given location falls in the top-down projection of the polygon given by our cached spline points. */
 	bool PointInsidePolygon(const FTransform& InTransform, const FBox& InBounds) const;
diff --git a/Engine/Plugins/PCG/Source/PCG/Public/Data/PCGTextureData.h b/Engine/Plugins/PCG/Source/PCG/Public/Data/PCGTextureData.h
--- a/Engine/Plugins/PCG/Source/PCG/Public/Data/PCGTextureData.h
+++ b/Engine/Plugins/PCG/Source/PCG/Public/Data/PCGTextureData.h
@@ -162,7 +162,7 @@ public:
 
 	//~Begin UPCGSpatialData interface
 protected:
-	virtual UPCGSpatialData* CopyInternal() const override;
+	virtual UPCGSpatialData* CopyInternal(FPCGContext* Context) const override;
 	//~End UPCGSpatialData interface
 
 private:
diff --git a/Engine/Plugins/PCG/Source/PCG/Public/Data/PCGUnionData.h b/Engine/Plugins/PCG/Source/PCG/Public/Data/PCGUnionData.h
--- a/Engine/Plugins/PCG/Source/PCG/Public/Data/PCGUnionData.h
+++ b/Engine/Plugins/PCG/Source/PCG/Public/Data/PCGUnionData.h
@@ -55,7 +55,7 @@ public:
 	virtual const UPCGSpatialData* FindFirstConcreteShapeFromNetwork() const override;
 
 protected:
-	virtual UPCGSpatialData* CopyInternal() const override;
+	virtual UPCGSpatialData* CopyInternal(FPCGContext* Context) const override;
 	//~End UPCGSpatialData interface
 
 public:
diff --git a/Engine/Plugins/PCG/Source/PCG/Public/Data/PCGVolumeData.h b/Engine/Plugins/PCG/Source/PCG/Public/Data/PCGVolumeData.h
--- a/Engine/Plugins/PCG/Source/PCG/Public/Data/PCGVolumeData.h
+++ b/Engine/Plugins/PCG/Source/PCG/Public/Data/PCGVolumeData.h
@@ -32,7 +32,7 @@ public:
 	// TODO what should this do - closest point on volume?
 	//virtual bool ProjectPoint(const FTransform& InTransform, const FBox& InBounds, const FPCGProjectionParams& InParams, FPCGPoint& OutPoint, UPCGMetadata* OutMetadata) const;
 protected:
-	virtual UPCGSpatialData* CopyInternal() const override;
+	virtual UPCGSpatialData* CopyInternal(FPCGContext* Context) const override;
 	//~End UPCGSpatialData interface
 
 public:
diff --git a/Engine/Plugins/PCG/Source/PCG/Public/Data/PCGWorldData.h b/Engine/Plugins/PCG/Source/PCG/Public/Data/PCGWorldData.h
--- a/Engine/Plugins/PCG/Source/PCG/Public/Data/PCGWorldData.h
+++ b/Engine/Plugins/PCG/Source/PCG/Public/Data/PCGWorldData.h
@@ -235,7 +235,7 @@ public:
 	// TODO not sure what this would mean. Without a direction, this means perhaps finding closest point on any collision surface? Should we implement this disabled?
 	//virtual bool ProjectPoint(const FTransform& InTransform, const FBox& InBounds, const FPCGProjectionParams& InParams, FPCGPoint& OutPoint, UPCGMetadata* OutMetadata) const;
 protected:
-	virtual UPCGSpatialData* CopyInternal() const override;
+	virtual UPCGSpatialData* CopyInternal(FPCGContext* Context) const override;
 	//~End UPCGSpatialData interface
 
 public:
@@ -276,7 +276,7 @@ public:
 	virtual bool HasNonTrivialTransform() const override { return true; }
 	virtual FVector GetNormal() const override { return Transform.GetRotation().GetUpVector(); }
 protected:
-	virtual UPCGSpatialData* CopyInternal() const override;
+	virtual UPCGSpatialData* CopyInternal(FPCGContext* Context) const override;
 	//~End UPCGSpatialData interface
 
 public:
diff --git a/Engine/Plugins/PCG/Source/PCG/Public/Elements/Metadata/PCGMetadataElementCommon.h b/Engine/Plugins/PCG/Source/PCG/Public/Elements/Metadata/PCGMetadataElementCommon.h
--- a/Engine/Plugins/PCG/Source/PCG/Public/Elements/Metadata/PCGMetadataElementCommon.h
+++ b/Engine/Plugins/PCG/Source/PCG/Public/Elements/Metadata/PCGMetadataElementCommon.h
@@ -15,8 +15,11 @@ struct FPCGTaggedData;
 
 namespace PCGMetadataElementCommon
 {
+	UE_DEPRECATED(5.5, "Call/Implement version with FPCGContext parameter")
 	PCG_API void DuplicateTaggedData(const FPCGTaggedData& InTaggedData, FPCGTaggedData& OutTaggedData, UPCGMetadata*& OutMetadata);
 
+	PCG_API void DuplicateTaggedData(FPCGContext* InContext, const FPCGTaggedData& InTaggedData, FPCGTaggedData& OutTaggedData, UPCGMetadata*& OutMetadata);
+
 	/** Copies the entry to value key relationship stored in the given Metadata, including its parents */
 	PCG_API void CopyEntryToValueKeyMap(const UPCGMetadata* MetadataToCopy, const FPCGMetadataAttributeBase* AttributeToCopy, FPCGMetadataAttributeBase* OutAttribute);
 
diff --git a/Engine/Plugins/PCG/Source/PCG/Public/Elements/Metadata/PCGMetadataMakeRotator.h b/Engine/Plugins/PCG/Source/PCG/Public/Elements/Metadata/PCGMetadataMakeRotator.h
--- a/Engine/Plugins/PCG/Source/PCG/Public/Elements/Metadata/PCGMetadataMakeRotator.h
+++ b/Engine/Plugins/PCG/Source/PCG/Public/Elements/Metadata/PCGMetadataMakeRotator.h
@@ -51,7 +51,7 @@ public:
 #endif
 	virtual void ApplyPreconfiguredSettings(const FPCGPreConfiguredSettingsInfo& PreconfiguredInfo) override;
 	virtual bool DoesInputSupportDefaultValue(uint32 Index) const override;
-	virtual UPCGParamData* CreateDefaultValueParam(uint32 Index) const override;
+	virtual UPCGParamData* CreateDefaultValueParam(FPCGContext* Context, uint32 Index) const override;
 #if WITH_EDITOR
 	virtual FString GetDefaultValueString(uint32 Index) const override;
 #endif // WITH_EDITOR
diff --git a/Engine/Plugins/PCG/Source/PCG/Public/Elements/Metadata/PCGMetadataMakeTransform.h b/Engine/Plugins/PCG/Source/PCG/Public/Elements/Metadata/PCGMetadataMakeTransform.h
--- a/Engine/Plugins/PCG/Source/PCG/Public/Elements/Metadata/PCGMetadataMakeTransform.h
+++ b/Engine/Plugins/PCG/Source/PCG/Public/Elements/Metadata/PCGMetadataMakeTransform.h
@@ -30,7 +30,7 @@ public:
 	virtual FText GetDefaultNodeTitle() const override;
 #endif
 	virtual bool DoesInputSupportDefaultValue(uint32 Index) const override;
-	virtual UPCGParamData* CreateDefaultValueParam(uint32 Index) const override;
+	virtual UPCGParamData* CreateDefaultValueParam(FPCGContext* Context, uint32 Index) const override;
 #if WITH_EDITOR
 	virtual FString GetDefaultValueString(uint32 Index) const override;
 #endif // WITH_EDITORs
diff --git a/Engine/Plugins/PCG/Source/PCG/Public/Elements/Metadata/PCGMetadataMakeVector.h b/Engine/Plugins/PCG/Source/PCG/Public/Elements/Metadata/PCGMetadataMakeVector.h
--- a/Engine/Plugins/PCG/Source/PCG/Public/Elements/Metadata/PCGMetadataMakeVector.h
+++ b/Engine/Plugins/PCG/Source/PCG/Public/Elements/Metadata/PCGMetadataMakeVector.h
@@ -52,7 +52,7 @@ public:
 	virtual FText GetDefaultNodeTitle() const override;
 #endif
 	virtual bool DoesInputSupportDefaultValue(uint32 Index) const override;
-	virtual UPCGParamData* CreateDefaultValueParam(uint32 Index) const override;
+	virtual UPCGParamData* CreateDefaultValueParam(FPCGContext* Context, uint32 Index) const override;
 #if WITH_EDITOR
 	virtual FString GetDefaultValueString(uint32 Index) const override;
 #endif // WITH_EDITOR
diff --git a/Engine/Plugins/PCG/Source/PCG/Public/Elements/Metadata/PCGMetadataOpElementBase.h b/Engine/Plugins/PCG/Source/PCG/Public/Elements/Metadata/PCGMetadataOpElementBase.h
--- a/Engine/Plugins/PCG/Source/PCG/Public/Elements/Metadata/PCGMetadataOpElementBase.h
+++ b/Engine/Plugins/PCG/Source/PCG/Public/Elements/Metadata/PCGMetadataOpElementBase.h
@@ -128,7 +128,11 @@ public:
 
 	/* Can be overriden by child class to support default values on unplugged pins. */
 	virtual bool DoesInputSupportDefaultValue(uint32 Index) const { return false; }
+
+	UE_DEPRECATED(5.5, "Call/Implement version with FPCGContext parameter")
 	virtual UPCGParamData* CreateDefaultValueParam(uint32 Index) const { return nullptr; }
+
+	virtual UPCGParamData* CreateDefaultValueParam(FPCGContext* Context, uint32 Index) const { return nullptr; }
 #if WITH_EDITOR
 	virtual FString GetDefaultValueString(uint32 Index) const { return FString(); }
 #endif // WITH_EDITOR
diff --git a/Engine/Plugins/PCG/Source/PCG/Public/Elements/PCGExecuteBlueprint.h b/Engine/Plugins/PCG/Source/PCG/Public/Elements/PCGExecuteBlueprint.h
--- a/Engine/Plugins/PCG/Source/PCG/Public/Elements/PCGExecuteBlueprint.h
+++ b/Engine/Plugins/PCG/Source/PCG/Public/Elements/PCGExecuteBlueprint.h
@@ -191,6 +191,9 @@ public:
 	/** Called after the element duplication during execution to be able to get the context easily - internal call only */
 	void SetCurrentContext(FPCGContext* InCurrentContext);
 
+	/** Tries to resolve current Context from thread local BP stack */
+	static FPCGContext* ResolveContext();
+
 #if WITH_EDITOR
 	// ~Begin UObject interface
 	virtual void PostEditChangeProperty(FPropertyChangedEvent& PropertyChangedEvent) override;
diff --git a/Engine/Plugins/PCG/Source/PCG/Public/PCGContext.h b/Engine/Plugins/PCG/Source/PCG/Public/PCGContext.h
--- a/Engine/Plugins/PCG/Source/PCG/Public/PCGContext.h
+++ b/Engine/Plugins/PCG/Source/PCG/Public/PCGContext.h
@@ -8,6 +8,7 @@
 #include "Utils/PCGExtraCapture.h"
 
 #include "UObject/GCObject.h"
+#include "UObject/GarbageCollection.h"
 
 #include "PCGContext.generated.h"
 
@@ -89,6 +90,16 @@ struct PCG_API FPCGContext
 	// Return the seed, possibly overriden by params, and combined with the source component (if any).
 	int GetSeed() const;
 
+	// Allows creating a new object safely inside the execution of a PCG Element, this object will also get tracked properly by the context
+	template<class T, typename... Args>
+	static T* NewObject_AnyThread(FPCGContext* Context, Args&&... InArgs)
+	{
+		ensure(Context != nullptr || IsInGameThread());
+		return Context ? Context->NewObject_AnyThread_Impl<T>(std::forward<Args>(InArgs)...) : ::NewObject<T>(std::forward<Args>(InArgs)...);
+	}
+
+	bool ContainsAsyncObject(const UObject* InAsyncObject) { return AsyncObjects.Contains(InAsyncObject); }
+
 	// Return the settings casted in the wanted type.
 	// If there is any override, those settings will already contains all the overriden values.
 	template<typename SettingsType>
@@ -132,12 +143,36 @@ protected:
 	virtual void AddExtraStructReferencedObjects(FReferenceCollector& Collector) {}
 
 private:
+	template<class T, typename... Args>
+	T* NewObject_AnyThread_Impl(Args&&... InArgs)
+	{
+		if (!IsInGameThread())
+		{
+			ensure(!AsyncState.bIsRunningOnMainThread);
+			T* Object = nullptr;
+			{
+				FGCScopeGuard Scope;
+				Object = ::NewObject<T>(std::forward<Args>(InArgs)...);
+			}
+			check(Object);
+			AsyncObjects.Add(Object);
+			return Object;
+		}
+
+		return ::NewObject<T>(std::forward<Args>(InArgs)...);
+	}
+
 	// Copy of the settings that will be used to apply overrides.
 	TObjectPtr<UPCGSettings> SettingsWithOverride = nullptr;
 
 	// List of params that were in effect overriden
 	TArray<const FPCGSettingsOverridableParam*> OverriddenParams;
 
+	// List of objects created by the PCG Elements, we need to track them so we can remove their Async flags when storing results on main thread
+	// so that they can be considered as existing on the main thread (and get properly GCed)
+	TSet<TObjectPtr<UObject>> AsyncObjects;
+
+	friend struct FPCGGraphActiveTask;
 #if WITH_EDITOR
 	friend class PCGUtils::FExtraCapture;
 	PCGUtils::FCallTime Timer;
diff --git a/Engine/Plugins/PCG/Source/PCG/Public/PCGData.h b/Engine/Plugins/PCG/Source/PCG/Public/PCGData.h
--- a/Engine/Plugins/PCG/Source/PCG/Public/PCGData.h
+++ b/Engine/Plugins/PCG/Source/PCG/Public/PCGData.h
@@ -18,6 +18,7 @@ class UPCGParamData;
 class UPCGSettings;
 class UPCGSettingsInterface;
 class UPCGSpatialData;
+struct FPCGContext;
 
 /**
 * Base class for any "data" class in the PCG framework.
@@ -50,9 +51,12 @@ public:
 	virtual FPCGAttributePropertyInputSelector GetCachedLastSelector() const { return FPCGAttributePropertyInputSelector{}; }
 	virtual void SetLastSelector(const FPCGAttributePropertySelector& InSelector) {};
 
-	/** Return a copy of the data, with Metadata inheritence for spatial data. */
+	UE_DEPRECATED(5.5, "Call/Implement version with FPCGContext parameter")
 	virtual UPCGData* DuplicateData(bool bInitializeMetadata = true) const;
 
+	/** Return a copy of the data, with Metadata inheritence for spatial data. */
+	virtual UPCGData* DuplicateData(FPCGContext* Context, bool bInitializeMetadata = true) const;
+
 	// ~Begin UObject interface
 	virtual void PostDuplicate(bool bDuplicateForPIE) override { InitUID(); }
 	// ~End UObject interface
diff --git a/Engine/Plugins/PCG/Source/PCG/Public/PCGParamData.h b/Engine/Plugins/PCG/Source/PCG/Public/PCGParamData.h
--- a/Engine/Plugins/PCG/Source/PCG/Public/PCGParamData.h
+++ b/Engine/Plugins/PCG/Source/PCG/Public/PCGParamData.h
@@ -26,7 +26,7 @@ public:
 	virtual FPCGAttributePropertyInputSelector GetCachedLastSelector() const override;
 	virtual void SetLastSelector(const FPCGAttributePropertySelector& InSelector) override;
 
-	virtual UPCGParamData* DuplicateData(bool bInitializeMetadata = true) const override;
+	virtual UPCGParamData* DuplicateData(FPCGContext* Context, bool bInitializeMetadata = true) const override;
 	// ~End UPCGData interface
 
 	UFUNCTION(BlueprintCallable, Category = Metadata)
@@ -44,11 +44,21 @@ public:
 	int64 FindOrAddMetadataKey(const FName& InName);
 
 	/** Creates a new params that keeps only a given key/name */
-	UFUNCTION(BlueprintCallable, Category = Params)
-	UPCGParamData* FilterParamsByName(const FName& InName) const;
+	UFUNCTION(BlueprintCallable, Category = Params, meta = (DisplayName = "Filter Params By Name"))
+	UPCGParamData* K2_FilterParamsByName(const FName& InName) const;
 
-	UFUNCTION(BlueprintCallable, Category = Params)
-	UPCGParamData* FilterParamsByKey(int64 InKey) const;
+	UE_DEPRECATED(5.5, "Call version with FPCGContext parameter")
+	UPCGParamData* FilterParamsByName(const FName& InName) const { return FilterParamsByName(nullptr, InName); }
+
+	UPCGParamData* FilterParamsByName(FPCGContext* Context, const FName& InName) const;
+
+	UFUNCTION(BlueprintCallable, Category = Params, meta = (DisplayName = "Filter Params By Key"))
+	UPCGParamData* K2_FilterParamsByKey(int64 InKey) const;
+
+	UE_DEPRECATED(5.5, "Call version with FPCGContext parameter")
+	UPCGParamData* FilterParamsByKey(int64 InKey) const { return FilterParamsByKey(nullptr, InKey); }
+
+	UPCGParamData* FilterParamsByKey(FPCGContext* Context, int64 InKey) const;
 
 	// Not accessible through blueprint to make sure the constness is preserved
 	UPROPERTY(VisibleAnywhere, Category = Metadata)
