[PCG] Protect TaskOutputs with RWLock

diff --git a/Engine/Plugins/PCG/Source/PCG/Private/Graph/PCGGraphExecutor.cpp b/Engine/Plugins/PCG/Source/PCG/Private/Graph/PCGGraphExecutor.cpp
--- a/Engine/Plugins/PCG/Source/PCG/Private/Graph/PCGGraphExecutor.cpp
+++ b/Engine/Plugins/PCG/Source/PCG/Private/Graph/PCGGraphExecutor.cpp
@@ -607,27 +607,34 @@ FPCGTaskId FPCGGraphExecutor::ScheduleGenericWithContext(TFunction<bool(FPCGCont
 
 void FPCGGraphExecutor::MarkInputResults(TArrayView<const FPCGTaskId> InInputResults)
 {
-	for (FPCGTaskId TaskId : InInputResults)
+	if (InInputResults.Num() > 0)
 	{
-		FOutputDataInfo& OutputDataInfo = OutputData[TaskId];
-		if (!OutputDataInfo.bNeedsManualClear && --OutputDataInfo.RemainingSuccessorCount == 0)
+		FWriteScopeLock WriteLock(TaskOutputsRWLock);
+		for (FPCGTaskId TaskId : InInputResults)
 		{
-			OutputData.Remove(TaskId);
+			FOutputDataInfo& OutputDataInfo = TaskOutputs[TaskId];
+			if (!OutputDataInfo.bNeedsManualClear && --OutputDataInfo.RemainingSuccessorCount == 0)
+			{
+				TaskOutputs.Remove(TaskId);
+			}
 		}
 	}
 }
 
 bool FPCGGraphExecutor::GetOutputData(FPCGTaskId TaskId, FPCGDataCollection& OutData, bool bClearDataOnGet)
 {
-	// TODO: this is not threadsafe - make threadsafe once we multithread execution
-	if (OutputData.Contains(TaskId))
+	FRWScopeLock ReadWriteLock(TaskOutputsRWLock, SLT_ReadOnly);
+	if (TaskOutputs.Contains(TaskId))
 	{
-		FOutputDataInfo& OutputDataInfo = OutputData[TaskId];
+		FOutputDataInfo& OutputDataInfo = TaskOutputs[TaskId];
 		OutData = OutputDataInfo.DataCollection;
 
 		if (bClearDataOnGet && ensure(OutputDataInfo.bNeedsManualClear))
 		{
-			OutputData.Remove(TaskId);
+			// Safe here as we are just removing from key, if release of read lock causes TaskOutputs to change
+			// before we acquire the write lock this will result in a no op.
+			ReadWriteLock.ReleaseReadOnlyLockAndAcquireWriteLock_USE_WITH_CAUTION();
+			TaskOutputs.Remove(TaskId);
 		}
 
 		return true;
@@ -664,12 +671,16 @@ void FPCGGraphExecutor::Execute()
 
 			// TODO: review if it's actually possible for a task with inputs to be ready at this point - it seems very unlikely
 			bool bPushToReady = true;
-			for (const FPCGGraphTaskInput& Input : Task.Inputs)
+			if (Task.Inputs.Num() > 0)
 			{
-				if (!OutputData.Contains(Input.TaskId))
+				FReadScopeLock ReadLock(TaskOutputsRWLock);
+				for (const FPCGGraphTaskInput& Input : Task.Inputs)
 				{
-					TaskSuccessors.FindOrAdd(Input.TaskId).Add(TaskId);
-					bPushToReady = false;
+					if (!TaskOutputs.Contains(Input.TaskId))
+					{
+						TaskSuccessors.FindOrAdd(Input.TaskId).Add(TaskId);
+						bPushToReady = false;
+					}
 				}
 			}
 
@@ -1227,9 +1238,13 @@ void FPCGGraphExecutor::QueueNextTasks(FPCGTaskId FinishedTask, bool bIgnoreMiss
 			{
 				FPCGGraphTask& SuccessorTask = *SuccessorTaskPtr;
 
-				for (const FPCGGraphTaskInput& Input : SuccessorTask.Inputs)
+				if (SuccessorTask.Inputs.Num() > 0)
 				{
-					bAllPrerequisitesMet &= OutputData.Contains(Input.TaskId);
+					FReadScopeLock ReadLock(TaskOutputsRWLock);
+					for (const FPCGGraphTaskInput& Input : SuccessorTask.Inputs)
+					{
+						bAllPrerequisitesMet &= TaskOutputs.Contains(Input.TaskId);
+					}
 				}
 
 				if (bAllPrerequisitesMet)
@@ -1350,7 +1365,9 @@ void FPCGGraphExecutor::BuildTaskInput(const FPCGGraphTask& Task, FPCGDataCollec
 
 	for (const FPCGGraphTaskInput& Input : Task.Inputs)
 	{
-		check(OutputData.Contains(Input.TaskId));
+		FReadScopeLock ReadLock(TaskOutputsRWLock);
+		check(TaskOutputs.Contains(Input.TaskId));
+		
 		ResultsToMarkAsRead.AddUnique(Input.TaskId);
 
 		// If the input does not provide any data, don't add it to the task input.
@@ -1369,7 +1386,7 @@ void FPCGGraphExecutor::BuildTaskInput(const FPCGGraphTask& Task, FPCGDataCollec
 			continue;
 		}
 
-		const FPCGDataCollection& InputCollection = OutputData[Input.TaskId].DataCollection;
+		const FPCGDataCollection& InputCollection = TaskOutputs[Input.TaskId].DataCollection;
 
 		TaskInput.bCancelExecution |= InputCollection.bCancelExecution;
 
@@ -1493,7 +1510,8 @@ void FPCGGraphExecutor::StoreResults(FPCGTaskId InTaskId, const FPCGDataCollecti
 	}
 
 	// Store output in map
-	OutputData.Add(InTaskId, MoveTemp(OutputDataInfo));
+	FWriteScopeLock WriteLock(TaskOutputsRWLock);
+	TaskOutputs.Add(InTaskId, MoveTemp(OutputDataInfo));
 }
 
 void FPCGGraphExecutor::ClearResults()
@@ -1506,7 +1524,8 @@ void FPCGGraphExecutor::ClearResults()
 		NextTaskId = 0;
 	}
 
-	OutputData.Reset();
+	FWriteScopeLock WriteLock(TaskOutputsRWLock);
+	TaskOutputs.Reset();
 
 	ScheduleLock.Unlock();
 }
@@ -1661,9 +1680,12 @@ void FPCGGraphExecutor::AddReferencedObjects(FReferenceCollector& Collector)
 	TRACE_CPUPROFILER_EVENT_SCOPE(FPCGGraphExecutor::AddReferencedObjects);
 
 	// Go through all data in the cached output map
-	for (auto& OutputDataEntry : OutputData)
 	{
-		OutputDataEntry.Value.DataCollection.AddReferences(Collector);
+		FReadScopeLock ReadLock(TaskOutputsRWLock);
+		for (auto& OutputDataEntry : TaskOutputs)
+		{
+			OutputDataEntry.Value.DataCollection.AddReferences(Collector);
+		}
 	}
 
 	// Go through ready tasks, active tasks and sleeping tasks contexts
diff --git a/Engine/Plugins/PCG/Source/PCG/Private/Graph/PCGGraphExecutor.h b/Engine/Plugins/PCG/Source/PCG/Private/Graph/PCGGraphExecutor.h
--- a/Engine/Plugins/PCG/Source/PCG/Private/Graph/PCGGraphExecutor.h
+++ b/Engine/Plugins/PCG/Source/PCG/Private/Graph/PCGGraphExecutor.h
@@ -288,7 +288,8 @@ private:
 		int32 RemainingSuccessorCount = 0;
 	};
 
-	TMap<FPCGTaskId, FOutputDataInfo> OutputData;
+	FRWLock TaskOutputsRWLock;
+	TMap<FPCGTaskId, FOutputDataInfo> TaskOutputs;
 
 	/** Monotonically increasing id. Should be reset once all tasks are executed, should be protected by the ScheduleLock */
 	FPCGTaskId NextTaskId = 0;
