[PCG] Add PostExecuteInternal virtual
- PCGExecuteBlueprint implements this and runs its PostExecute on the main thread always

diff --git a/Engine/Plugins/PCG/Source/PCG/Private/Elements/PCGExecuteBlueprint.cpp b/Engine/Plugins/PCG/Source/PCG/Private/Elements/PCGExecuteBlueprint.cpp
--- a/Engine/Plugins/PCG/Source/PCG/Private/Elements/PCGExecuteBlueprint.cpp
+++ b/Engine/Plugins/PCG/Source/PCG/Private/Elements/PCGExecuteBlueprint.cpp
@@ -766,9 +766,10 @@ void UPCGBlueprintSettings::FixingOverridableParamPropertyClass(FPCGSettingsOver
 
 bool FPCGExecuteBlueprintElement::ExecuteInternal(FPCGContext* InContext) const
 {
+	check(InContext);
 	FPCGBlueprintExecutionContext* Context = static_cast<FPCGBlueprintExecutionContext*>(InContext);
 
-	if (Context && Context->BlueprintElementInstance)
+	if (Context->BlueprintElementInstance)
 	{
 		UClass* BPClass = Context->BlueprintElementInstance->GetClass();
 
@@ -801,9 +802,8 @@ bool FPCGExecuteBlueprintElement::ExecuteInternal(FPCGContext* InContext) const
 		/** Finally, execute the actual blueprint */
 		Context->BlueprintElementInstance->SetCurrentContext(Context);
 		
-		const bool bIsInGameThread = IsInGameThread();
 		// When running outside of main thread make sure GC can't run (BP Nodes might create objects which can't happen while GC runs)
-		if(!bIsInGameThread)
+		if(!IsInGameThread())
 		{
 			ensure(!Context->AsyncState.bIsRunningOnMainThread);
 			FGCScopeGuard Scope;
@@ -815,7 +815,24 @@ bool FPCGExecuteBlueprintElement::ExecuteInternal(FPCGContext* InContext) const
 		}
 
 		Context->BlueprintElementInstance->SetCurrentContext(nullptr);
+	}
+	else
+	{
+		// Nothing to do but forward data
+		Context->OutputData = Context->InputData;
+	}
+	
+	return true;
+}
+
+void FPCGExecuteBlueprintElement::PostExecuteInternal(FPCGContext* InContext) const
+{
+	check(InContext);
+	FPCGBlueprintExecutionContext* Context = static_cast<FPCGBlueprintExecutionContext*>(InContext);
 
+	check(IsInGameThread());
+	if (Context->BlueprintElementInstance)
+	{
 		// Log info on outputs
 		for (int32 OutputIndex = 0; OutputIndex < Context->OutputData.TaggedData.Num(); ++OutputIndex)
 		{
@@ -829,16 +846,13 @@ bool FPCGExecuteBlueprintElement::ExecuteInternal(FPCGContext* InContext) const
 			// Any data that was created by the user in the blueprint will have that data parented to this blueprint element instance
 			// Which will cause issues wrt to reference leaks. We need to fix this here.
 			// Note that we will recurse up the outer tree to make sure we catch every case.
-			if(Output.Data)
+			if (Output.Data)
 			{
-				if (!bIsInGameThread)
+				// Clear Async flags on objects created outside of the main thread and not part of the Context known async objects
+				if (Output.Data->HasAnyInternalFlags(EInternalObjectFlags::Async) && !Context->ContainsAsyncObject(Output.Data))
 				{
-					// Clear Async flags on objects created outside of the main thread and not part of the Context known async objects
-					if (Output.Data->HasAnyInternalFlags(EInternalObjectFlags::Async) && !Context->ContainsAsyncObject(Output.Data))
-					{
-						Output.Data->AtomicallyClearInternalFlags(EInternalObjectFlags::Async);
-						ForEachObjectWithOuter(Output.Data, [](UObject* SubObject) { SubObject->AtomicallyClearInternalFlags(EInternalObjectFlags::Async); }, true);
-					}
+					Output.Data->ClearInternalFlags(EInternalObjectFlags::Async);
+					ForEachObjectWithOuter(Output.Data, [](UObject* SubObject) { SubObject->ClearInternalFlags(EInternalObjectFlags::Async); }, true);
 				}
 
 				auto ReOuterToTransientPackageIfCreatedFromThis = [Context](UObject* InObject)
@@ -870,7 +884,7 @@ bool FPCGExecuteBlueprintElement::ExecuteInternal(FPCGContext* InContext) const
 				{
 					Metadata = ParamData->Metadata;
 				}
-				
+
 				if (Metadata)
 				{
 					while (Metadata->GetParent())
@@ -883,13 +897,6 @@ bool FPCGExecuteBlueprintElement::ExecuteInternal(FPCGContext* InContext) const
 			}
 		}
 	}
-	else if(Context)
-	{
-		// Nothing to do but forward data
-		Context->OutputData = Context->InputData;
-	}
-	
-	return true;
 }
 
 void UPCGBlueprintElement::PointLoop(FPCGContext& InContext, const UPCGPointData* InData, UPCGPointData*& OutData, UPCGPointData* OptionalOutData) const
@@ -1074,6 +1081,14 @@ bool FPCGExecuteBlueprintElement::ShouldComputeFullOutputDataCrc(FPCGContext* Co
 bool FPCGExecuteBlueprintElement::CanExecuteOnlyOnMainThread(FPCGContext* Context) const
 {
 	check(Context);
+	FPCGBlueprintExecutionContext* BPContext = static_cast<FPCGBlueprintExecutionContext*>(Context);
+
+	// Always execute PostExecute on main thread
+	if (Context->CurrentPhase == EPCGExecutionPhase::PostExecute && BPContext->BlueprintElementInstance)
+	{
+		return true;
+	}
+
 	const UPCGBlueprintSettings* BPSettings = Context->GetInputSettings<UPCGBlueprintSettings>();
 	if (BPSettings && BPSettings->BlueprintElementInstance)
 	{
diff --git a/Engine/Plugins/PCG/Source/PCG/Private/PCGElement.cpp b/Engine/Plugins/PCG/Source/PCG/Private/PCGElement.cpp
--- a/Engine/Plugins/PCG/Source/PCG/Private/PCGElement.cpp
+++ b/Engine/Plugins/PCG/Source/PCG/Private/PCGElement.cpp
@@ -158,6 +158,9 @@ bool IPCGElement::PrepareDataInternal(FPCGContext* Context) const
 
 void IPCGElement::PostExecute(FPCGContext* Context) const
 {
+	// Allow sub class to do some processing here
+	PostExecuteInternal(Context);
+
 	// Cleanup and validate output
 	CleanupAndValidateOutput(Context);
 
diff --git a/Engine/Plugins/PCG/Source/PCG/Public/Elements/PCGExecuteBlueprint.h b/Engine/Plugins/PCG/Source/PCG/Public/Elements/PCGExecuteBlueprint.h
--- a/Engine/Plugins/PCG/Source/PCG/Public/Elements/PCGExecuteBlueprint.h
+++ b/Engine/Plugins/PCG/Source/PCG/Public/Elements/PCGExecuteBlueprint.h
@@ -416,6 +416,7 @@ public:
 
 protected:
 	virtual bool ExecuteInternal(FPCGContext* Context) const override;
+	virtual void PostExecuteInternal(FPCGContext* Context) const override;
 	virtual FPCGContext* Initialize(const FPCGDataCollection& InputData, TWeakObjectPtr<UPCGComponent> SourceComponent, const UPCGNode* Node) override;	
 };
 
diff --git a/Engine/Plugins/PCG/Source/PCG/Public/PCGElement.h b/Engine/Plugins/PCG/Source/PCG/Public/PCGElement.h
--- a/Engine/Plugins/PCG/Source/PCG/Public/PCGElement.h
+++ b/Engine/Plugins/PCG/Source/PCG/Public/PCGElement.h
@@ -116,6 +116,8 @@ protected:
 	virtual bool ExecuteInternal(FPCGContext* Context) const = 0;
 	/** This function will be called once and once only, at the end of an execution */
 	void PostExecute(FPCGContext* Context) const;
+	/** Core post execute method for the given element. */
+	virtual void PostExecuteInternal(FPCGContext* Context) const {}
 	/** This function will be called once and only once if the element is aborted. The Context can be used to retrieve the current phase if needed. */
 	virtual void AbortInternal(FPCGContext* Context) const {};
 
