Added option to ignore properties in filter data by attributes

diff --git a/Engine/Plugins/PCG/Source/PCG/Private/Elements/PCGFilterByAttribute.cpp b/Engine/Plugins/PCG/Source/PCG/Private/Elements/PCGFilterByAttribute.cpp
--- a/Engine/Plugins/PCG/Source/PCG/Private/Elements/PCGFilterByAttribute.cpp
+++ b/Engine/Plugins/PCG/Source/PCG/Private/Elements/PCGFilterByAttribute.cpp
@@ -80,24 +80,27 @@ bool FPCGFilterByAttributeElement::ExecuteInternal(FPCGContext* Context) const
 			Metadata->GetAttributes(DataAttributes, DataAttributeTypes);
 			Algo::Transform(DataAttributes, DataAttributeStrings, [](const FName& InAttribute) { return InAttribute.ToString(); });
 
-			if (Cast<UPCGPointData>(Input.Data))
+			if (!Settings->bIgnoreProperties)
 			{
-				if (const UEnum* PointProperties = StaticEnum<EPCGPointProperties>())
+				if (Cast<UPCGPointData>(Input.Data))
 				{
-					DataAttributeStrings.Reserve(DataAttributeStrings.Num() + PointProperties->NumEnums());
-					for (int32 EnumIndex = 0; EnumIndex < PointProperties->NumEnums(); ++EnumIndex)
+					if (const UEnum* PointProperties = StaticEnum<EPCGPointProperties>())
 					{
-						DataAttributeStrings.Add(FString(TEXT("$")) + PointProperties->GetNameStringByIndex(EnumIndex));
+						DataAttributeStrings.Reserve(DataAttributeStrings.Num() + PointProperties->NumEnums());
+						for (int32 EnumIndex = 0; EnumIndex < PointProperties->NumEnums(); ++EnumIndex)
+						{
+							DataAttributeStrings.Add(FString(TEXT("$")) + PointProperties->GetNameStringByIndex(EnumIndex));
+						}
 					}
 				}
-			}
 
-			if (const UEnum* ExtraProperties = StaticEnum<EPCGExtraProperties>())
-			{
-				DataAttributeStrings.Reserve(DataAttributeStrings.Num() + ExtraProperties->NumEnums());
-				for (int32 EnumIndex = 0; EnumIndex < ExtraProperties->NumEnums(); ++EnumIndex)
+				if (const UEnum* ExtraProperties = StaticEnum<EPCGExtraProperties>())
 				{
-					DataAttributeStrings.Add(FString(TEXT("$")) + ExtraProperties->GetNameStringByIndex(EnumIndex));
+					DataAttributeStrings.Reserve(DataAttributeStrings.Num() + ExtraProperties->NumEnums());
+					for (int32 EnumIndex = 0; EnumIndex < ExtraProperties->NumEnums(); ++EnumIndex)
+					{
+						DataAttributeStrings.Add(FString(TEXT("$")) + ExtraProperties->GetNameStringByIndex(EnumIndex));
+					}
 				}
 			}
 		}
@@ -106,12 +109,12 @@ bool FPCGFilterByAttributeElement::ExecuteInternal(FPCGContext* Context) const
 
 		for (const FString& Attribute : Attributes)
 		{
+			FPCGAttributePropertySelector Selector;
+			Selector.Update(Attribute);
+
 			// In the case of the equal test, we can test directly if the selector would yield something valid
 			if (Settings->Operator == EPCGStringMatchingOperator::Equal)
 			{
-				FPCGAttributePropertySelector Selector;
-				Selector.Update(Attribute);
-
 				TUniquePtr<const IPCGAttributeAccessor> Accessor = PCGAttributeAccessorHelpers::CreateConstAccessor(Input.Data, Selector, /*bQuiet=*/true);
 				if (!Accessor || !Accessor.IsValid())
 				{
@@ -123,19 +126,9 @@ bool FPCGFilterByAttributeElement::ExecuteInternal(FPCGContext* Context) const
 			{
 				// Otherwise, it's going to be a bit more complex -
 				// First, reconstruct the main property/attribute name from the selector, because it might have removed the $ character.
-				FString AttributeWithNoAccessor;
-				FString AttributeAccessors;
-				int32 FirstAccessorCharIndex = INDEX_NONE;
-				if (Attribute.FindChar('.', FirstAccessorCharIndex))
-				{
-					AttributeWithNoAccessor = Attribute.Left(FirstAccessorCharIndex);
-					AttributeAccessors = Attribute.RightChop(FirstAccessorCharIndex); // includes the '.' as we'll add it back after
-				}
-				else
-				{
-					AttributeWithNoAccessor = Attribute;
-				}
-
+				const FString AttributeWithNoAccessor = Selector.GetAttributePropertyString(/*bAddQualifier=*/true);
+				const FString AttributeAccessors = Selector.GetAttributePropertyAccessorsString(/*bAddLeadingSeparator=*/true);
+				
 				// Try to find a valid match of AttributeWithNoAccessor against the DataAttributeStrings.
 				bool bFoundValidMatch = false;
 				for (const FString& DataAttribute : DataAttributeStrings)
@@ -147,10 +140,10 @@ bool FPCGFilterByAttributeElement::ExecuteInternal(FPCGContext* Context) const
 					}
 
 					// We have a valid name-based match, now check if the full attribute can be used as a valid extractor.
-					FPCGAttributePropertySelector Selector;
-					Selector.Update(DataAttribute + AttributeAccessors);
+					FPCGAttributePropertySelector DataSelector;
+					DataSelector.Update(DataAttribute + AttributeAccessors);
 
-					TUniquePtr<const IPCGAttributeAccessor> Accessor = PCGAttributeAccessorHelpers::CreateConstAccessor(Input.Data, Selector, /*bQuiet=*/true);
+					TUniquePtr<const IPCGAttributeAccessor> Accessor = PCGAttributeAccessorHelpers::CreateConstAccessor(Input.Data, DataSelector, /*bQuiet=*/true);
 					if (Accessor && Accessor.IsValid())
 					{
 						bFoundValidMatch = true;
diff --git a/Engine/Plugins/PCG/Source/PCG/Private/Metadata/PCGAttributePropertySelector.cpp b/Engine/Plugins/PCG/Source/PCG/Private/Metadata/PCGAttributePropertySelector.cpp
--- a/Engine/Plugins/PCG/Source/PCG/Private/Metadata/PCGAttributePropertySelector.cpp
+++ b/Engine/Plugins/PCG/Source/PCG/Private/Metadata/PCGAttributePropertySelector.cpp
@@ -152,32 +152,45 @@ bool FPCGAttributePropertySelector::SetExtraProperty(EPCGExtraProperties InExtra
 	}
 }
 
-FString FPCGAttributePropertySelector::ToString() const
+FString FPCGAttributePropertySelector::GetAttributePropertyString(bool bAddPropertyQualifier) const
 {
-	FString Res;
 	const FName Name = GetName();
-
 	// Add a '$' if it is a property
-	if (Selection != EPCGAttributePropertySelection::Attribute && (Name != NAME_None))
+	if (bAddPropertyQualifier && Selection != EPCGAttributePropertySelection::Attribute && Name != NAME_None)
 	{
-		Res = FString(PCGAttributePropertySelectorConstants::PropertyPrefix) + Name.ToString();
+		return FString(PCGAttributePropertySelectorConstants::PropertyPrefix) + Name.ToString();
 	}
 	else
 	{
-		Res = Name.ToString();
+		return Name.ToString();
 	}
+}
 
+FString FPCGAttributePropertySelector::GetAttributePropertyAccessorsString(bool bAddLeadingSeparator) const
+{
 	if (!ExtraNames.IsEmpty())
 	{
-		TArray<FString> AllNames;
-		AllNames.Add(Res);
-		AllNames.Append(ExtraNames);
-		Res = FString::Join(AllNames, PCGAttributePropertySelectorConstants::ExtraSeparator);
-	}
+		FString LeadingSeparatorString;
+		if (bAddLeadingSeparator)
+		{
+			LeadingSeparatorString = FString(PCGAttributePropertySelectorConstants::ExtraSeparator);
+		}
 
-	return Res;
+		return LeadingSeparatorString + FString::Join(ExtraNames, PCGAttributePropertySelectorConstants::ExtraSeparator);
+	}
+	else
+	{
+		return FString();
+	}
 }
 
+FString FPCGAttributePropertySelector::ToString() const
+{
+	const FString Attribute = GetAttributePropertyString(/*bAddPropertyQualifier=*/true);
+	const FString Accessors = GetAttributePropertyAccessorsString(/*bAddLeadingSeparator*/true);
+
+	return Attribute + Accessors;
+}
 
 bool FPCGAttributePropertySelector::operator==(const FPCGAttributePropertySelector& Other) const
 {
diff --git a/Engine/Plugins/PCG/Source/PCG/Public/Elements/PCGFilterByAttribute.h b/Engine/Plugins/PCG/Source/PCG/Public/Elements/PCGFilterByAttribute.h
--- a/Engine/Plugins/PCG/Source/PCG/Public/Elements/PCGFilterByAttribute.h
+++ b/Engine/Plugins/PCG/Source/PCG/Public/Elements/PCGFilterByAttribute.h
@@ -36,6 +36,10 @@ public:
 
 	UPROPERTY(BlueprintReadWrite, EditAnywhere, Category = Settings, meta = (PCG_Overridable))
 	EPCGStringMatchingOperator Operator = EPCGStringMatchingOperator::Equal;
+
+	/** Controls whether properties (denoted by $) will be considered in the filter or not. */
+	UPROPERTY(BlueprintReadWrite, EditAnywhere, Category = Settings, meta = (PCG_Overridable))
+	bool bIgnoreProperties = false;
 };
 
 class FPCGFilterByAttributeElement : public IPCGElement
diff --git a/Engine/Plugins/PCG/Source/PCG/Public/Metadata/PCGAttributePropertySelector.h b/Engine/Plugins/PCG/Source/PCG/Public/Metadata/PCGAttributePropertySelector.h
--- a/Engine/Plugins/PCG/Source/PCG/Public/Metadata/PCGAttributePropertySelector.h
+++ b/Engine/Plugins/PCG/Source/PCG/Public/Metadata/PCGAttributePropertySelector.h
@@ -65,7 +65,14 @@ public:
 	// Return the name of the selector.
 	FName GetName() const;
 
+	// Returns qualified attribute/property name with the accessors.
 	FString ToString() const;
+	
+	// Returns attribute/property name only, with optional '$' qualifier for properties.
+	FString GetAttributePropertyString(bool bAddPropertyQualifier = true) const;
+
+	// Returns the accessor part of the selector, with optional leading '.' separator.
+	FString GetAttributePropertyAccessorsString(bool bAddLeadingSeparator) const;
 
 	// Returns the text to display in the widget.
 	FText GetDisplayText() const { return FText::FromString(ToString()); }
