[PCG] Fixed two issues in collision-based self-pruning related to scales and index matching

#rnx

diff --git a/Engine/Plugins/PCG/Source/PCG/Private/Elements/PCGSelfPruning.cpp b/Engine/Plugins/PCG/Source/PCG/Private/Elements/PCGSelfPruning.cpp
--- a/Engine/Plugins/PCG/Source/PCG/Private/Elements/PCGSelfPruning.cpp
+++ b/Engine/Plugins/PCG/Source/PCG/Private/Elements/PCGSelfPruning.cpp
@@ -208,8 +208,8 @@ namespace PCGSelfPruningElement
 			}
 
 			// Implementation note: this is a deconstruction of FBodyInstance::OverlapTestForBodiesImpl
-			// Get this point's body instance, if any. Note the -1 because we already incremented the value.
-			FBodyInstance* ThisInstance = IterationState.CollisionWrapper.GetBodyInstance(IterationState.CurrentPointIndex - 1);
+			// Get this point's body instance, if any.
+			FBodyInstance* ThisInstance = IterationState.CollisionWrapper.GetBodyInstance(PointRef.Point - IterationState.InputPointData->GetPoints().GetData());
 			FCollisionShape ThisCollisionShape;
 			FTransform TransformNoScale = FTransform(PointRef.Point->Transform.GetRotation(), PointRef.Point->Transform.GetLocation());
 
@@ -250,7 +250,10 @@ namespace PCGSelfPruningElement
 
 							IterationState.TemporaryBodyInstances.Add(ThisInstance, OtherInstance);
 						}
+					}
 
+					if (ThisInstance != OtherInstance)
+					{
 						OtherInstance->UpdateBodyScale(OtherPointRef->Point->Transform.GetScale3D());
 					}
 
