Delays regeneration of a graph to obtain inspection information until it is needed. If the graph has already been generated in this session, will not rerun unless needed.
Also improved unselection of objects in the tree so that we have less chances of triggering a refresh that wouldn't be needed.

#rnx

diff --git a/Engine/Plugins/PCG/Source/PCG/Private/PCGComponent.cpp b/Engine/Plugins/PCG/Source/PCG/Private/PCGComponent.cpp
--- a/Engine/Plugins/PCG/Source/PCG/Private/PCGComponent.cpp
+++ b/Engine/Plugins/PCG/Source/PCG/Private/PCGComponent.cpp
@@ -549,6 +549,9 @@ void UPCGComponent::PostProcessGraph(const FBox& InNewBounds, bool bInGenerated,
 	if (bInGenerated)
 	{
 		bGenerated = true;
+#if WITH_EDITOR
+		bWasGeneratedThisSession = true;
+#endif
 
 		CurrentGenerationTask = InvalidPCGTaskId;
 
diff --git a/Engine/Plugins/PCG/Source/PCG/Public/PCGComponent.h b/Engine/Plugins/PCG/Source/PCG/Public/PCGComponent.h
--- a/Engine/Plugins/PCG/Source/PCG/Public/PCGComponent.h
+++ b/Engine/Plugins/PCG/Source/PCG/Public/PCGComponent.h
@@ -340,6 +340,7 @@ public:
 	void ResetLastGeneratedBounds();
 
 	/** Functions for managing the node inspection cache */
+	bool WasGeneratedThisSession() const { return bWasGeneratedThisSession; }
 	bool IsInspecting() const;
 	void EnableInspection();
 	void DisableInspection();
@@ -623,6 +624,7 @@ private:
 	bool bIsComponentLocal = false;
 
 #if WITH_EDITOR
+	bool bWasGeneratedThisSession = false;
 	int32 InspectionCounter = 0;
 	FBox LastGeneratedBoundsPriorToUndo = FBox(EForceInit::ForceInit);
 #endif
diff --git a/Engine/Plugins/PCG/Source/PCGEditor/Private/PCGEditor.cpp b/Engine/Plugins/PCG/Source/PCGEditor/Private/PCGEditor.cpp
--- a/Engine/Plugins/PCG/Source/PCGEditor/Private/PCGEditor.cpp
+++ b/Engine/Plugins/PCG/Source/PCGEditor/Private/PCGEditor.cpp
@@ -244,45 +244,56 @@ void FPCGEditor::SetStackBeingInspected(const FPCGStack& FullStack)
 		return;
 	}
 
-	UPCGComponent* OldComponent = PCGComponentBeingInspected.Get();
+	UPCGComponent* LastComponent = LastValidPCGComponentBeingInspected.Get();
 	UPCGComponent* NewComponent = const_cast<UPCGComponent*>(FullStack.GetRootComponent());
-	const bool bComponentChanged = (NewComponent != OldComponent);
 
-	if (OldComponent)
+	if (NewComponent && NewComponent != LastComponent)
 	{
-		if (bComponentChanged)
+		if (LastComponent && LastComponent->IsInspecting())
 		{
-			OldComponent->DisableInspection();
+			LastComponent->DisableInspection();
 		}
 
-		if (PCGGraphBeingEdited)
+		LastValidPCGComponentBeingInspected = NewComponent;
+	}
+
+	if (PCGGraphBeingEdited)
+	{
+		if (PCGGraphBeingEdited->IsInspecting())
 		{
 			PCGGraphBeingEdited->DisableInspection();
 		}
-	}
 
-	const bool bNewComponentStartedInspecting = NewComponent && !NewComponent->IsInspecting();
+		PCGGraphBeingEdited->EnableInspection(StackBeingInspected);
+	}
 
 	PCGComponentBeingInspected = NewComponent;
 
 	StackBeingInspected = FullStack;
 	OnInspectedStackChangedDelegate.Broadcast(StackBeingInspected);
 
-	if (NewComponent)
+	UpdateAfterInspectedStackChanged();
+}
+
+void FPCGEditor::UpdateAfterInspectedStackChanged()
+{
+	UPCGComponent* Component = PCGComponentBeingInspected.Get();
+
+	if (Component)
 	{
-		if (bComponentChanged)
-		{
-			PCGComponentBeingInspected->EnableInspection();
-		}
+		// Implementation note: if we're inspecting and have not pre-run the graph, then it probably makes sense to enable inspection by default. 
+		// TODO This could be selected with a cvar though.
+		const bool bHasBeenGeneratedThisSession = Component->bGenerated && Component->WasGeneratedThisSession();
+		const bool bWasInspecting = Component->IsInspecting();
+		const bool bNeedsInspection = Algo::AnyOf(AttributesWidgets, [](const TSharedPtr<SPCGEditorGraphAttributeListView>& ALV) { return ALV->GetNodeBeingInspected() != nullptr; });
 
-		if (PCGGraphBeingEdited)
+		if (!bHasBeenGeneratedThisSession || (bNeedsInspection && !bWasInspecting))
 		{
-			PCGGraphBeingEdited->EnableInspection(StackBeingInspected);
+			Component->EnableInspection();
+			UpdateDebugAfterComponentSelection(Component, Component, true);
 		}
 	}
 
-	UpdateDebugAfterComponentSelection(OldComponent, NewComponent, bNewComponentStartedInspecting);
-
 	check(PCGEditorGraph);
 	for (UEdGraphNode* Node : PCGEditorGraph->Nodes)
 	{
@@ -290,7 +301,7 @@ void FPCGEditor::SetStackBeingInspected(const FPCGStack& FullStack)
 		{
 			// Update now that component has changed. Will fire OnNodeChanged if necessary.
 			EPCGChangeType ChangeType = PCGNode->UpdateErrorsAndWarnings();
-			ChangeType |= PCGNode->UpdateStructuralVisualization(NewComponent, &StackBeingInspected);
+			ChangeType |= PCGNode->UpdateStructuralVisualization(Component, &StackBeingInspected);
 
 			if (ChangeType != EPCGChangeType::None)
 			{
@@ -310,7 +321,7 @@ void FPCGEditor::ClearStackBeingInspected()
 
 void FPCGEditor::UpdateDebugAfterComponentSelection(UPCGComponent* InOldComponent, UPCGComponent* InNewComponent, bool bInNewComponentStartedInspecting)
 {
-	if (!ensure(PCGGraphBeingEdited) || (InOldComponent == InNewComponent))
+	if (!ensure(PCGGraphBeingEdited))
 	{
 		return;
 	}
@@ -366,7 +377,7 @@ void FPCGEditor::UpdateDebugAfterComponentSelection(UPCGComponent* InOldComponen
 		});
 
 		// Regenerate to clear debug info if switching components, or if changing from a component to null.
-		if (InNewComponent || bDebugFlagSetOnAnyNode)
+		if (InNewComponent != InOldComponent && (InNewComponent || bDebugFlagSetOnAnyNode))
 		{
 			// Use original component - debug can be displayed both by the local component and parent local components.
 			RefreshComponent(InOldComponent->GetOriginalComponent());
@@ -1944,6 +1955,9 @@ void FPCGEditor::OnToggleInspected()
 	{
 		DebugObjectTreeWidget->SetNodeBeingInspected(nullptr);
 	}
+
+	// Turn on "inspecting" on graph if we now have at least one inspected node and had none before
+	UpdateAfterInspectedStackChanged();
 }
 
 bool FPCGEditor::CanToggleInspected() const
@@ -2821,15 +2835,23 @@ void FPCGEditor::OnClose()
 		{
 			PCGComponentBeingInspected->DisableInspection();
 		}
+	}
 
-		if (PCGGraphBeingEdited && PCGGraphBeingEdited->IsInspecting())
+	if (LastValidPCGComponentBeingInspected.IsValid())
+	{
+		if (LastValidPCGComponentBeingInspected->IsInspecting())
 		{
-			PCGGraphBeingEdited->DisableInspection();
+			LastValidPCGComponentBeingInspected->DisableInspection();
 		}
 	}
 
 	if (PCGGraphBeingEdited)
 	{
+		if (PCGGraphBeingEdited->IsInspecting())
+		{
+			PCGGraphBeingEdited->DisableInspection();
+		}
+
 		if (PCGGraphBeingEdited->NotificationsForEditorArePausedByUser())
 		{
 			PCGGraphBeingEdited->ToggleUserPausedNotificationsForEditor();
diff --git a/Engine/Plugins/PCG/Source/PCGEditor/Private/PCGEditor.h b/Engine/Plugins/PCG/Source/PCGEditor/Private/PCGEditor.h
--- a/Engine/Plugins/PCG/Source/PCGEditor/Private/PCGEditor.h
+++ b/Engine/Plugins/PCG/Source/PCGEditor/Private/PCGEditor.h
@@ -134,6 +134,8 @@ private:
 	bool CanToggleInspected() const;
 	/** Whether selected nodes are inspected or not */
 	ECheckBoxState GetInspectedCheckState() const;
+
+	void UpdateAfterInspectedStackChanged();
 	
 	/** Toggle node enabled state for selected nodes */
 	void OnToggleEnabled();
@@ -372,5 +374,7 @@ private:
 	UPCGEditorGraph* PCGEditorGraph = nullptr;
 
 	TWeakObjectPtr<UPCGComponent> PCGComponentBeingInspected;
+	// Implementation note: we'll keep the last valid component inspected so we don't un-inspect on spurious selection changes
+	TWeakObjectPtr<UPCGComponent> LastValidPCGComponentBeingInspected;
 	FPCGStack StackBeingInspected;
 };
diff --git a/Engine/Plugins/PCG/Source/PCGEditor/Private/PCGEditorGraphSchema.cpp b/Engine/Plugins/PCG/Source/PCGEditor/Private/PCGEditorGraphSchema.cpp
--- a/Engine/Plugins/PCG/Source/PCGEditor/Private/PCGEditorGraphSchema.cpp
+++ b/Engine/Plugins/PCG/Source/PCGEditor/Private/PCGEditorGraphSchema.cpp
@@ -801,7 +801,7 @@ bool FPCGEditorConnectionDrawingPolicy::UpdateParamsIfDebugging(UEdGraphPin* Out
 	const FPCGStack* PCGStack = Editor ? Editor->GetStackBeingInspected() : nullptr;
 	UPCGComponent* PCGComponent = Editor ? Editor->GetPCGComponentBeingInspected() : nullptr;
 
-	if (!Editor || !PCGStack || !PCGComponent)
+	if (!Editor || !PCGStack || !PCGComponent || !PCGComponent->IsInspecting())
 	{
 		return false;
 	}
