Fixed issue where the spawn actor would not respect the actor mobility, following changes for data layers.

#rnx

diff --git a/Engine/Plugins/PCG/Source/PCG/Private/Elements/PCGSpawnActor.cpp b/Engine/Plugins/PCG/Source/PCG/Private/Elements/PCGSpawnActor.cpp
--- a/Engine/Plugins/PCG/Source/PCG/Private/Elements/PCGSpawnActor.cpp
+++ b/Engine/Plugins/PCG/Source/PCG/Private/Elements/PCGSpawnActor.cpp
@@ -809,6 +809,7 @@ void FPCGSpawnActorElement::SpawnActors(FPCGSubgraphContext* Context, AActor* Ta
 	}
 
 	UPCGActorHelpers::FSpawnDefaultActorParams SpawnDefaultActorParams(TargetActor->GetWorld(), InTemplateActorClass, FTransform::Identity, SpawnParams);
+	SpawnDefaultActorParams.bForceStaticMobility = false; // Always respect the actor's mobility
 
 #if WITH_EDITOR
 	SpawnDefaultActorParams.DataLayerInstances = TargetActor->GetDataLayerInstances();
diff --git a/Engine/Plugins/PCG/Source/PCG/Private/Helpers/PCGActorHelpers.cpp b/Engine/Plugins/PCG/Source/PCG/Private/Helpers/PCGActorHelpers.cpp
--- a/Engine/Plugins/PCG/Source/PCG/Private/Helpers/PCGActorHelpers.cpp
+++ b/Engine/Plugins/PCG/Source/PCG/Private/Helpers/PCGActorHelpers.cpp
@@ -529,7 +529,10 @@ AActor* UPCGActorHelpers::SpawnDefaultActor(const FSpawnDefaultActorParams& Para
 		RootComponent->RegisterComponent();
 	}
 
-	RootComponent->Mobility = EComponentMobility::Static;
+	if (Params.bForceStaticMobility)
+	{
+		RootComponent->Mobility = EComponentMobility::Static;
+	}
 
 #if WITH_EDITOR
 	RootComponent->bVisualizeComponent = true;
diff --git a/Engine/Plugins/PCG/Source/PCG/Public/Helpers/PCGActorHelpers.h b/Engine/Plugins/PCG/Source/PCG/Public/Helpers/PCGActorHelpers.h
--- a/Engine/Plugins/PCG/Source/PCG/Public/Helpers/PCGActorHelpers.h
+++ b/Engine/Plugins/PCG/Source/PCG/Public/Helpers/PCGActorHelpers.h
@@ -143,6 +143,7 @@ public:
 		FTransform Transform;
 		FActorSpawnParameters SpawnParams;
 		AActor* Parent = nullptr;
+		bool bForceStaticMobility = true;
 #if WITH_EDITOR
 		TArray<const UDataLayerInstance*> DataLayerInstances;
 #endif
