[PCG] Fix landscape spline data crash when there are overlapping InterpPoints

#rnx

diff --git a/Engine/Plugins/PCG/Source/PCG/Private/Data/PCGLandscapeSplineData.cpp b/Engine/Plugins/PCG/Source/PCG/Private/Data/PCGLandscapeSplineData.cpp
--- a/Engine/Plugins/PCG/Source/PCG/Private/Data/PCGLandscapeSplineData.cpp
+++ b/Engine/Plugins/PCG/Source/PCG/Private/Data/PCGLandscapeSplineData.cpp
@@ -226,29 +226,24 @@ void UPCGLandscapeSplineData::GetTangentsAtSegmentStart(int SegmentIndex, FVecto
 
 FVector::FReal UPCGLandscapeSplineData::GetDistanceAtSegmentStart(int SegmentIndex) const
 {
-	check(Spline.IsValid());
-
-	const int NumSegments = GetNumSegments();
+	FVector::FReal Distance = 0.0;
 
-	if (NumSegments == 0)
+	// Implementation note: It would be cheaper to loop over segments instead of ReparamPoints, but that is not robust since
+	// InterpPoints on the segments may not be 1:1 with the ReparamPoints. This can happen for poorly formed Landscape Splines where
+	// the adjacent control points are identical points.
+	for (int32 ReparamIndex = 0; ReparamIndex < ReparamTable.Points.Num(); ++ReparamIndex)
 	{
-		return 0.0;
-	}
-	
-	// Allow SegmentIndex == NumSegments, which indicates we want the distance to the final control point, which is like saying "Start of the Nth segment".
-	check(SegmentIndex >= 0 && SegmentIndex <= NumSegments);
+		const FVector::FReal CurrentDistance = ReparamTable.Points[ReparamIndex].InVal;
 
-	const TArray<TObjectPtr<ULandscapeSplineSegment>>& Segments = Spline->GetSegments();
-	int32 ReparamIndex = 0;
+		if (ReparamTable.Points[ReparamIndex].OutVal > SegmentIndex)
+		{
+			break;
+		}
 
-	for (int32 Index = 0; Index < SegmentIndex; ++Index)
-	{
-		// NumPoints - 1 to avoid double-counting the control points, which overlap at the start + end of each segment.
-		ReparamIndex += ensure(Segments[Index]) ? Segments[Index]->GetPoints().Num() - 1 : 0;
+		Distance = CurrentDistance;
 	}
 
-	check(ReparamTable.Points.IsValidIndex(ReparamIndex));
-	return ReparamTable.Points[ReparamIndex].InVal;
+	return Distance;
 }
 
 const UPCGPointData* UPCGLandscapeSplineData::CreatePointData(FPCGContext* Context) const
@@ -389,7 +384,16 @@ void UPCGLandscapeSplineData::UpdateReparamTable()
 		{
 			const FLandscapeSplineInterpPoint& Start = InterpPoints[PointIndex - 1];
 			const FLandscapeSplineInterpPoint& End = InterpPoints[PointIndex];
-			AccumulatedDistance += FVector::Distance(Start.Center, End.Center);
+			const FVector::FReal Distance = FVector::Distance(Start.Center, End.Center);
+
+			// Skip points that overlap the previous point. We should not have duplicate distance entries in the ReparamTable.
+			// Note: This means InterpPoints are not 1:1 with ReparamTable entries.
+			if (FMath::IsNearlyZero(Distance, /*ErrorTolerance=*/UE_KINDA_SMALL_NUMBER))
+			{
+				continue;
+			}
+
+			AccumulatedDistance += Distance;
 
 			const float Param = PointIndex / (NumPoints - 1.0f);
 			ReparamTable.Points.Emplace(AccumulatedDistance, SegmentIndex + Param, /*ArriveTangent=*/0.0f, /*LeaveTangent=*/0.0f, CIM_Linear);
