[PCG] Fix infinite task when frametime is 0

#rnx

diff --git a/Engine/Plugins/PCG/Source/PCG/Public/Helpers/PCGAsync.h b/Engine/Plugins/PCG/Source/PCG/Public/Helpers/PCGAsync.h
--- a/Engine/Plugins/PCG/Source/PCG/Public/Helpers/PCGAsync.h
+++ b/Engine/Plugins/PCG/Source/PCG/Public/Helpers/PCGAsync.h
@@ -261,10 +261,12 @@ namespace FPCGAsync
 
 			// Collapsing needs to be done in order, so we have a map between chunk index and the number of elements written for this chunk.
 			// Note that we need to collpase because points can be discarded, so we can have less points "kept" than the chunk size.
+			// Do at least one run, as we could fall into an infinite loop if we always have to stop before doing anything.
+			bool bHasRunOnce = false;
 			while (true)
 			{
 				// Either we should stop because the time has elapsed or all workloads have been processed
-				if (ShouldStop() || SynchroStruct.bQuit)
+				if (bHasRunOnce && (ShouldStop() || SynchroStruct.bQuit))
 				{
 					SynchroStruct.bQuit = true;
 
@@ -278,6 +280,8 @@ namespace FPCGAsync
 					break;
 				}
 
+				bHasRunOnce = true;
+
 				// Try to flush queue before trying anything.
 				FlushQueue();
 				
