[PCG] Update Switch and MultiSelect node to use enum value names instead of their display text as pin label, because it was not resilient to localization

#rnx

diff --git a/Engine/Plugins/PCG/Source/PCG/Private/Elements/ControlFlow/PCGMultiSelect.cpp b/Engine/Plugins/PCG/Source/PCG/Private/Elements/ControlFlow/PCGMultiSelect.cpp
--- a/Engine/Plugins/PCG/Source/PCG/Private/Elements/ControlFlow/PCGMultiSelect.cpp
+++ b/Engine/Plugins/PCG/Source/PCG/Private/Elements/ControlFlow/PCGMultiSelect.cpp
@@ -16,7 +16,34 @@ namespace PCGMultiSelectConstants
 void UPCGMultiSelectSettings::PostLoad()
 {
 	Super::PostLoad();
+
 	CachePinLabels();
+
+#if WITH_EDITOR
+	// @todo_pcg To be behind a version bump in 5.5. Cannot do that in an hotfix
+	// Make sure we rename the pin properties that were serialized with a localized text. Since we can't exactly match the text
+	// with the enum value, we'll go with index.
+	if (SelectionMode == EPCGControlFlowSelectionMode::Enum)
+	{
+		UPCGNode* OuterNode = Cast<UPCGNode>(GetOuter());
+		if (OuterNode)
+		{
+			TArray<UPCGPin*> SerializedInputPins = OuterNode->GetInputPins();
+			// It we have a num mismatch, we can't recover
+			if (SerializedInputPins.Num() == CachedPinLabels.Num())
+			{
+				// -1 since we don't need to check "Default"
+				for (int32 i = 0; i < CachedPinLabels.Num() - 1; ++i)
+				{
+					if (SerializedInputPins[i] && SerializedInputPins[i]->Properties.Label != CachedPinLabels[i])
+					{
+						OuterNode->RenameInputPin(SerializedInputPins[i]->Properties.Label, CachedPinLabels[i], /*bBroadcastUpdate=*/false);
+					}
+				}
+			}
+		}
+	}
+#endif // WITH_EDITOR
 }
 
 #if WITH_EDITOR
@@ -140,16 +167,15 @@ TArray<FPCGPinProperties> UPCGMultiSelectSettings::InputPinProperties() const
 			// -1 to bypass the MAX value
 			for (int32 Index = 0; EnumSelection.Class && Index < EnumSelection.Class->NumEnums() - 1; ++Index)
 			{
+				bool bHidden = false;
 #if WITH_EDITOR
-				bool const bHidden = EnumSelection.Class->HasMetaData(TEXT("Hidden"), Index) || EnumSelection.Class->HasMetaData(TEXT("Spacer"), Index);
+				bHidden = EnumSelection.Class->HasMetaData(TEXT("Hidden"), Index) || EnumSelection.Class->HasMetaData(TEXT("Spacer"), Index);
+#endif // WITH_EDITOR
+
 				if (!bHidden)
 				{
-					PinProperties.Emplace(FName(EnumSelection.Class->GetDisplayNameTextByIndex(Index).ToString()));
+					PinProperties.Emplace(EnumSelection.Class->GetNameByIndex(Index));
 				}
-#else // WITH_EDITOR
-				// HasMetaData is editor only, so there will be extra pins at runtime, but that should be okay
-				PinProperties.Emplace(FName(EnumSelection.Class->GetDisplayNameTextByIndex(Index).ToString()));
-#endif // WITH_EDITOR
 			}
 			break;
 		default:
@@ -218,7 +244,7 @@ bool UPCGMultiSelectSettings::GetSelectedPinLabel(FName& OutSelectedPinLabel) co
 	else if (SelectionMode == EPCGControlFlowSelectionMode::Enum && IsValuePresent(EnumSelection.Value))
 	{
 		// To account for hidden enums missing from the pin properties, find the pin by label instead of index
-		const FName PinLabel(EnumSelection.Class->GetDisplayNameTextByValue(EnumSelection.Value).ToString());
+		const FName PinLabel = EnumSelection.Class->GetNameByValue(EnumSelection.Value);
 		for (int i = 0; i < CachedPinLabels.Num(); ++i)
 		{
 			if (CachedPinLabels[i] == PinLabel)
diff --git a/Engine/Plugins/PCG/Source/PCG/Private/Elements/ControlFlow/PCGSwitch.cpp b/Engine/Plugins/PCG/Source/PCG/Private/Elements/ControlFlow/PCGSwitch.cpp
--- a/Engine/Plugins/PCG/Source/PCG/Private/Elements/ControlFlow/PCGSwitch.cpp
+++ b/Engine/Plugins/PCG/Source/PCG/Private/Elements/ControlFlow/PCGSwitch.cpp
@@ -16,7 +16,34 @@ namespace PCGSwitchConstants
 void UPCGSwitchSettings::PostLoad()
 {
 	Super::PostLoad();
+
 	CachePinLabels();
+
+#if WITH_EDITOR
+	// @todo_pcg To be behind a version bump in 5.5. Cannot do that in an hotfix
+	// Make sure we rename the pin properties that were serialized with a localized text. Since we can't exactly match the text
+	// with the enum value, we'll go with index.
+	if (SelectionMode == EPCGControlFlowSelectionMode::Enum)
+	{
+		UPCGNode* OuterNode = Cast<UPCGNode>(GetOuter());
+		if (OuterNode)
+		{
+			TArray<UPCGPin*> SerializedOutputPins = OuterNode->GetOutputPins();
+			// It we have a num mismatch, we can't recover
+			if (SerializedOutputPins.Num() == CachedPinLabels.Num())
+			{
+				// -1 since we don't need to check "Default"
+				for (int32 i = 0; i < CachedPinLabels.Num() - 1; ++i)
+				{
+					if (SerializedOutputPins[i] && SerializedOutputPins[i]->Properties.Label != CachedPinLabels[i])
+					{
+						OuterNode->RenameOutputPin(SerializedOutputPins[i]->Properties.Label, CachedPinLabels[i], /*bBroadcastUpdate=*/false);
+					}
+				}
+			}
+		}
+	}
+#endif // WITH_EDITOR
 }
 
 #if WITH_EDITOR
@@ -175,16 +202,16 @@ TArray<FPCGPinProperties> UPCGSwitchSettings::OutputPinProperties() const
 			// -1 to bypass the MAX value
 			for (int32 Index = 0; EnumSelection.Class && Index < EnumSelection.Class->NumEnums() - 1; ++Index)
 			{
+				bool bHidden = false;
 #if WITH_EDITOR
-				bool const bHidden = EnumSelection.Class->HasMetaData(TEXT("Hidden"), Index) || EnumSelection.Class->HasMetaData(TEXT("Spacer"), Index);
+				// HasMetaData is editor only, so there will be extra pins at runtime, but that should be okay
+				bHidden = EnumSelection.Class->HasMetaData(TEXT("Hidden"), Index) || EnumSelection.Class->HasMetaData(TEXT("Spacer"), Index);
+#endif // WITH_EDITOR
+
 				if (!bHidden)
 				{
-					PinProperties.Emplace(FName(EnumSelection.Class->GetDisplayNameTextByIndex(Index).ToString()));
+					PinProperties.Emplace(EnumSelection.Class->GetNameByIndex(Index));
 				}
-#else // WITH_EDITOR
-				// HasMetaData is editor only, so there will be extra pins at runtime, but that should be okay
-				PinProperties.Emplace(FName(EnumSelection.Class->GetDisplayNameTextByIndex(Index).ToString()));
-#endif // WITH_EDITOR
 			}
 			break;
 
@@ -251,7 +278,7 @@ int UPCGSwitchSettings::GetSelectedOutputPinIndex() const
 	else if (SelectionMode == EPCGControlFlowSelectionMode::Enum)
 	{
 		// A "hidden" value could be selected that wasn't cached, so do a name-wise comparison
-		const FName PinLabel(EnumSelection.Class->GetDisplayNameTextByValue(EnumSelection.Value).ToString());
+		const FName PinLabel = EnumSelection.Class->GetNameByValue(EnumSelection.Value);
 
 		// Return index if found, otherwise fallback to the index after the options which will be "Default" pin.
 		const int FoundIndex = CachedPinLabels.IndexOfByKey(PinLabel);
@@ -280,7 +307,7 @@ bool UPCGSwitchSettings::GetSelectedPinLabel(FName& OutSelectedPinLabel) const
 	else if (SelectionMode == EPCGControlFlowSelectionMode::Enum && IsValuePresent(EnumSelection.Value))
 	{
 		// A "hidden" value could be selected that wasn't cached, so do a name-wise comparison
-		const FName PinLabel(EnumSelection.Class->GetDisplayNameTextByValue(EnumSelection.Value).ToString());
+		const FName PinLabel = EnumSelection.Class->GetNameByValue(EnumSelection.Value);
 		for (int i = 0; i < CachedPinLabels.Num(); ++i)
 		{
 			if (CachedPinLabels[i] == PinLabel)
