[PCG] Expose PrimitveData's primitive component to the API

#rnx

diff --git a/Engine/Plugins/PCG/Source/PCG/Public/Data/PCGPrimitiveData.h b/Engine/Plugins/PCG/Source/PCG/Public/Data/PCGPrimitiveData.h
--- a/Engine/Plugins/PCG/Source/PCG/Public/Data/PCGPrimitiveData.h
+++ b/Engine/Plugins/PCG/Source/PCG/Public/Data/PCGPrimitiveData.h
@@ -26,6 +26,7 @@ public:
 	virtual bool SamplePoint(const FTransform& Transform, const FBox& Bounds, FPCGPoint& OutPoint, UPCGMetadata* OutMetadata) const override;
 	// TODO needs an implementation to support projection
 	//virtual bool ProjectPoint(const FTransform& InTransform, const FBox& InBounds, const FPCGProjectionParams& InParams, FPCGPoint& OutPoint, UPCGMetadata* OutMetadata) const;
+
 protected:
 	virtual UPCGSpatialData* CopyInternal(FPCGContext* Context) const override;
 	//~End UPCGSpatialData interface
@@ -35,6 +36,8 @@ public:
 	virtual const UPCGPointData* CreatePointData(FPCGContext* Context) const override;
 	// ~End UPCGSpatialDataWithPointCache implementation
 
+	TWeakObjectPtr<UPrimitiveComponent> GetComponent() const { return Primitive; }
+
 	UPROPERTY(BlueprintReadWrite, EditAnywhere, Category = "Data")
 	FVector VoxelSize = FVector(100.0, 100.0, 100.0);
 
@@ -47,4 +50,4 @@ protected:
 
 	UPROPERTY()
 	FBox CachedStrictBounds = FBox(EForceInit::ForceInit);
-};
\ No newline at end of file
+};
