[PCG] Fix Switch and Multi-Select nodes with Enums are breaking connections when language changes

* Previous iteration was not working for Multi-Select, because of the overrides pins.
* Previous iteration was using names which didn't work well with custom Enums made in BP.
* This iteration now uses the Culture Invariant display name.
* Also fix breaking unit test (because pin names changed)

#rnx

diff --git a/Engine/Plugins/PCG/Source/PCG/Private/Elements/ControlFlow/PCGControlFlow.cpp b/Engine/Plugins/PCG/Source/PCG/Private/Elements/ControlFlow/PCGControlFlow.cpp
new file mode 100644
--- /dev/null
+++ b/Engine/Plugins/PCG/Source/PCG/Private/Elements/ControlFlow/PCGControlFlow.cpp
@@ -0,0 +1,16 @@
+// Copyright Epic Games, Inc. All Rights Reserved.
+
+#include "Elements/ControlFlow/PCGControlFlow.h"
+
+#include "Internationalization/Text.h"
+#include "UObject/Class.h"
+
+FText FEnumSelector::GetDisplayName() const 
+{
+	return Class ? Class->GetDisplayNameTextByValue(Value) : FText{};
+}
+
+FString FEnumSelector::GetCultureInvariantDisplayName() const 
+{
+	return GetDisplayName().BuildSourceString();
+}
diff --git a/Engine/Plugins/PCG/Source/PCG/Private/Elements/ControlFlow/PCGMultiSelect.cpp b/Engine/Plugins/PCG/Source/PCG/Private/Elements/ControlFlow/PCGMultiSelect.cpp
--- a/Engine/Plugins/PCG/Source/PCG/Private/Elements/ControlFlow/PCGMultiSelect.cpp
+++ b/Engine/Plugins/PCG/Source/PCG/Private/Elements/ControlFlow/PCGMultiSelect.cpp
@@ -29,6 +29,14 @@ void UPCGMultiSelectSettings::PostLoad()
 		if (OuterNode)
 		{
 			TArray<UPCGPin*> SerializedInputPins = OuterNode->GetInputPins();
+
+			// We need to remove the override pins, find the overrides pin index and remove all following
+			const int32 OverridesPinIndex = SerializedInputPins.IndexOfByPredicate([](const UPCGPin* Pin) { return Pin && Pin->Properties.Label == PCGPinConstants::DefaultParamsLabel; });
+			if (OverridesPinIndex != INDEX_NONE)
+			{
+				SerializedInputPins.SetNum(OverridesPinIndex);
+			}
+
 			// It we have a num mismatch, we can't recover
 			if (SerializedInputPins.Num() == CachedPinLabels.Num())
 			{
@@ -116,7 +124,7 @@ FString UPCGMultiSelectSettings::GetAdditionalTitleInformation() const
 				FString Subtitle = EnumSelection.Class->GetName();
 				if (!IsPropertyOverriddenByPin({GET_MEMBER_NAME_CHECKED(UPCGMultiSelectSettings, EnumSelection), GET_MEMBER_NAME_CHECKED(FEnumSelector, Value)}))
 				{
-					Subtitle += FString::Format(TEXT(": {0}"), {EnumSelection.Class->GetNameStringByValue(EnumSelection.Value)});
+					Subtitle += FString::Format(TEXT(": {0}"), { EnumSelection.GetCultureInvariantDisplayName() });
 				}
 
 				return Subtitle;
@@ -174,7 +182,8 @@ TArray<FPCGPinProperties> UPCGMultiSelectSettings::InputPinProperties() const
 
 				if (!bHidden)
 				{
-					PinProperties.Emplace(EnumSelection.Class->GetNameByIndex(Index));
+					const FString EnumDisplayName = EnumSelection.Class->GetDisplayNameTextByIndex(Index).BuildSourceString();
+					PinProperties.Emplace(FName(EnumDisplayName));
 				}
 			}
 			break;
@@ -244,7 +253,7 @@ bool UPCGMultiSelectSettings::GetSelectedPinLabel(FName& OutSelectedPinLabel) co
 	else if (SelectionMode == EPCGControlFlowSelectionMode::Enum && IsValuePresent(EnumSelection.Value))
 	{
 		// To account for hidden enums missing from the pin properties, find the pin by label instead of index
-		const FName PinLabel = EnumSelection.Class->GetNameByValue(EnumSelection.Value);
+		const FName PinLabel(EnumSelection.GetCultureInvariantDisplayName());
 		for (int i = 0; i < CachedPinLabels.Num(); ++i)
 		{
 			if (CachedPinLabels[i] == PinLabel)
diff --git a/Engine/Plugins/PCG/Source/PCG/Private/Elements/ControlFlow/PCGSwitch.cpp b/Engine/Plugins/PCG/Source/PCG/Private/Elements/ControlFlow/PCGSwitch.cpp
--- a/Engine/Plugins/PCG/Source/PCG/Private/Elements/ControlFlow/PCGSwitch.cpp
+++ b/Engine/Plugins/PCG/Source/PCG/Private/Elements/ControlFlow/PCGSwitch.cpp
@@ -116,7 +116,7 @@ FString UPCGSwitchSettings::GetAdditionalTitleInformation() const
 				FString Subtitle = EnumSelection.Class->GetName();
 				if (!IsPropertyOverriddenByPin({GET_MEMBER_NAME_CHECKED(UPCGSwitchSettings, EnumSelection), GET_MEMBER_NAME_CHECKED(FEnumSelector, Value)}))
 				{
-					Subtitle += FString::Format(TEXT(": {0}"), {EnumSelection.Class->GetNameStringByValue(EnumSelection.Value)});
+					Subtitle += FString::Format(TEXT(": {0}"), { EnumSelection.GetCultureInvariantDisplayName() });
 				}
 
 				return Subtitle;
@@ -210,7 +210,8 @@ TArray<FPCGPinProperties> UPCGSwitchSettings::OutputPinProperties() const
 
 				if (!bHidden)
 				{
-					PinProperties.Emplace(EnumSelection.Class->GetNameByIndex(Index));
+					const FString EnumDisplayName = EnumSelection.Class->GetDisplayNameTextByIndex(Index).BuildSourceString();
+					PinProperties.Emplace(FName(EnumDisplayName));
 				}
 			}
 			break;
@@ -278,7 +279,7 @@ int UPCGSwitchSettings::GetSelectedOutputPinIndex() const
 	else if (SelectionMode == EPCGControlFlowSelectionMode::Enum)
 	{
 		// A "hidden" value could be selected that wasn't cached, so do a name-wise comparison
-		const FName PinLabel = EnumSelection.Class->GetNameByValue(EnumSelection.Value);
+		const FName PinLabel(EnumSelection.GetCultureInvariantDisplayName());
 
 		// Return index if found, otherwise fallback to the index after the options which will be "Default" pin.
 		const int FoundIndex = CachedPinLabels.IndexOfByKey(PinLabel);
@@ -307,7 +308,7 @@ bool UPCGSwitchSettings::GetSelectedPinLabel(FName& OutSelectedPinLabel) const
 	else if (SelectionMode == EPCGControlFlowSelectionMode::Enum && IsValuePresent(EnumSelection.Value))
 	{
 		// A "hidden" value could be selected that wasn't cached, so do a name-wise comparison
-		const FName PinLabel = EnumSelection.Class->GetNameByValue(EnumSelection.Value);
+		const FName PinLabel(EnumSelection.GetCultureInvariantDisplayName());
 		for (int i = 0; i < CachedPinLabels.Num(); ++i)
 		{
 			if (CachedPinLabels[i] == PinLabel)
diff --git a/Engine/Plugins/PCG/Source/PCG/Private/Tests/Elements/PCGSwitchTest.cpp b/Engine/Plugins/PCG/Source/PCG/Private/Tests/Elements/PCGSwitchTest.cpp
--- a/Engine/Plugins/PCG/Source/PCG/Private/Tests/Elements/PCGSwitchTest.cpp
+++ b/Engine/Plugins/PCG/Source/PCG/Private/Tests/Elements/PCGSwitchTest.cpp
@@ -169,7 +169,7 @@ bool FPCGSwitchTest_Enum::RunTest(const FString& Parameters)
 		}
 		else
 		{
-			UTEST_TRUE("Enum selection is valid", ValidateOutput(OutputData, FName(TestSettings->EnumSelection.Class->GetDisplayNameTextByIndex(i).ToString())));
+			UTEST_TRUE("Enum selection is valid", ValidateOutput(OutputData, FName(TestSettings->EnumSelection.GetCultureInvariantDisplayName())));
 		}
 	}
 
diff --git a/Engine/Plugins/PCG/Source/PCG/Public/Elements/ControlFlow/PCGControlFlow.h b/Engine/Plugins/PCG/Source/PCG/Public/Elements/ControlFlow/PCGControlFlow.h
--- a/Engine/Plugins/PCG/Source/PCG/Public/Elements/ControlFlow/PCGControlFlow.h
+++ b/Engine/Plugins/PCG/Source/PCG/Public/Elements/ControlFlow/PCGControlFlow.h
@@ -14,6 +14,9 @@ struct PCG_API FEnumSelector
 
 	UPROPERTY(DisplayName="Enum Value")
 	int64 Value = 0;
+
+	FText GetDisplayName() const;
+	FString GetCultureInvariantDisplayName() const;
 };
 
 UENUM()
diff --git a/Engine/Plugins/PCG/Source/PCGEditor/Private/Details/EnumSelectorDetails.cpp b/Engine/Plugins/PCG/Source/PCGEditor/Private/Details/EnumSelectorDetails.cpp
--- a/Engine/Plugins/PCG/Source/PCGEditor/Private/Details/EnumSelectorDetails.cpp
+++ b/Engine/Plugins/PCG/Source/PCGEditor/Private/Details/EnumSelectorDetails.cpp
@@ -330,7 +330,7 @@ FText FEnumSelectorDetails::GetSelectedEnumClassText() const
 FText FEnumSelectorDetails::GetSelectedEnumValueText() const
 {
 	const FEnumSelector* Selector = GetEnumSelector();
-	return Selector->Class ? StripEnumClassName(Selector->Class->GetNameByValue(Selector->Value)) : FText();
+	return FText::AsCultureInvariant(Selector->GetCultureInvariantDisplayName());
 }
 
 #undef LOCTEXT_NAMESPACE
