[PCG] Use pin properties at execution time for graph task inputs rather than holding raw pointers to graph pins.

#rnx

diff --git a/Engine/Plugins/PCG/Source/PCG/Private/Graph/PCGGraphCompiler.cpp b/Engine/Plugins/PCG/Source/PCG/Private/Graph/PCGGraphCompiler.cpp
--- a/Engine/Plugins/PCG/Source/PCG/Private/Graph/PCGGraphCompiler.cpp
+++ b/Engine/Plugins/PCG/Source/PCG/Private/Graph/PCGGraphCompiler.cpp
@@ -118,7 +118,7 @@ TArray<FPCGGraphTask> FPCGGraphCompiler::CompileGraph(UPCGGraph* InGraph, FPCGTa
 				{
 					if (InboundEdge->IsValid())
 					{
-						PreTask.Inputs.Emplace(IdMapping[InboundEdge->InputPin->Node], InboundEdge->InputPin, InboundEdge->OutputPin);
+						PreTask.Inputs.Emplace(IdMapping[InboundEdge->InputPin->Node], InboundEdge->InputPin->Properties, InboundEdge->OutputPin->Properties);
 					}
 					else
 					{
@@ -130,7 +130,7 @@ TArray<FPCGGraphTask> FPCGGraphCompiler::CompileGraph(UPCGGraph* InGraph, FPCGTa
 			// Add pre-task as input to subgraph input node task
 			if (InputNodeTask)
 			{
-				InputNodeTask->Inputs.Emplace(PreId, nullptr, nullptr);
+				InputNodeTask->Inputs.Emplace(PreId);
 			}
 
 			// Hook nodes to the PreTask if they require so.
@@ -148,7 +148,7 @@ TArray<FPCGGraphTask> FPCGGraphCompiler::CompileGraph(UPCGGraph* InGraph, FPCGTa
 				const bool bRequiresDataFromPreTask = Settings && Settings->RequiresDataFromPreTask();
 				if (bRequiresDataFromPreTask || Subtask.Inputs.IsEmpty())
 				{
-					Subtask.Inputs.Emplace(PreId, /*InInboundPin=*/nullptr, /*InOutboundPin=*/nullptr, bRequiresDataFromPreTask);
+					Subtask.Inputs.Emplace(PreId, /*InUpstreamPin=*/FPCGGraphTaskInput::NoPin, /*InDownstreamPin=*/FPCGGraphTaskInput::NoPin, bRequiresDataFromPreTask);
 				}
 			}
 
@@ -169,12 +169,12 @@ TArray<FPCGGraphTask> FPCGGraphCompiler::CompileGraph(UPCGGraph* InGraph, FPCGTa
 			// with something that might become inactive and would then fail to dynamically cull this already-scheduled task.
 			// Additional implementation note: this first depedencency is critical in our ability to do static culling (see CalculateStaticallyActiveRecursive)
 			// and should not be changed here without changing the other.
-			PostTask.Inputs.Emplace(PreId, /*InInboundPin=*/nullptr, /*InOutboundPin=*/nullptr, /*bInProvideData=*/false);
+			PostTask.Inputs.Emplace(PreId, /*InUpstreamPin=*/FPCGGraphTaskInput::NoPin, /*InDownstreamPin=*/FPCGGraphTaskInput::NoPin, /*bInProvideData=*/false);
 
 			// Add subgraph output node task as input to the post-task
 			if (OutputNodeTask)
 			{
-				PostTask.Inputs.Emplace(OutputNodeTask->NodeId, /*InInboundPin=*/nullptr, /*InOutboundPin=*/nullptr);
+				PostTask.Inputs.Emplace(OutputNodeTask->NodeId);
 			}
 
 			check(!IdMapping.Contains(Node));
@@ -200,7 +200,7 @@ TArray<FPCGGraphTask> FPCGGraphCompiler::CompileGraph(UPCGGraph* InGraph, FPCGTa
 
 					if (FPCGTaskId* InboundId = IdMapping.Find(InboundEdge->InputPin->Node))
 					{
-						Task.Inputs.Emplace(*InboundId, InboundEdge->InputPin, InboundEdge->OutputPin); 
+						Task.Inputs.Emplace(*InboundId, InboundEdge->InputPin->Properties, InboundEdge->OutputPin->Properties);
 					}
 					else
 					{
@@ -364,7 +364,16 @@ void FPCGGraphCompiler::CreateGridLinkages(
 		const EPCGHiGenGrid GraphGenerationGrid = InOutTaskGenerationGrid[TaskId];
 		for (FPCGGraphTaskInput& TaskInput : InOutCompiledTasks[TaskId].Inputs)
 		{
-			if (!TaskInput.InPin)
+			const UPCGPin* UpstreamPin = nullptr;
+			if (TaskInput.UpstreamPin.IsSet())
+			{
+				if (const UPCGNode* Node = InOutCompiledTasks[TaskInput.TaskId].Node)
+				{
+					UpstreamPin = Node->GetOutputPin(TaskInput.UpstreamPin.GetValue().Label);
+				}
+			}
+
+			if (!UpstreamPin)
 			{
 				// Don't link if we don't have a upstream pin to retrieve data from
 				continue;
@@ -377,7 +386,7 @@ void FPCGGraphCompiler::CreateGridLinkages(
 			{
 				// Build a string identifier for the data
 				FString ResourceKey;
-				if (!ensure(CurrentStack->CreateStackFramePath(ResourceKey, TaskInput.InPin->Node, TaskInput.InPin)))
+				if (!ensure(CurrentStack->CreateStackFramePath(ResourceKey, UpstreamPin->Node, UpstreamPin)))
 				{
 					continue;
 				}
@@ -387,13 +396,13 @@ void FPCGGraphCompiler::CreateGridLinkages(
 				LinkTask.NodeId = InOutCompiledTasks.Num() - 1;
 				LinkTask.StackIndex = InOutCompiledTasks[TaskId].StackIndex;
 
-				LinkTask.Inputs.Emplace(TaskInput.TaskId, TaskInput.InPin, nullptr, /*bConsumeInputData=*/true);
+				LinkTask.Inputs.Emplace(TaskInput.TaskId, UpstreamPin->Properties);
 
 				const EPCGHiGenGrid FromGrid = InOutTaskGenerationGrid[TaskInput.TaskId];
 				const EPCGHiGenGrid ToGrid = InOutTaskGenerationGrid[TaskId];
 
 				// This lambda runs at execution time and attempts to retrieve the data from a larger grid. Capture by value is intentional.
-				auto GridLinkageOperation = [FromGrid, ToGrid, ResourceKey, OutputPinLabel = TaskInput.InPin->Properties.Label,
+				auto GridLinkageOperation = [FromGrid, ToGrid, ResourceKey, OutputPinLabel = TaskInput.UpstreamPin.GetValue().Label,
 					DownstreamNode = InOutCompiledTasks[TaskId].Node, InGenerationGrid](FPCGContext* InContext)
 				{
 					return PCGGraphExecutor::ExecuteGridLinkage(
@@ -559,10 +568,8 @@ bool FPCGGraphCompiler::CalculateStaticallyActiveRecursive(FPCGTaskId InTaskId,
 
 	for (const FPCGGraphTaskInput& Input : InCompiledTasks[InTaskId].Inputs)
 	{
-		const UPCGPin* InputPin = Input.OutPin;
-
 		// Only non-advanced input pins play a part in determining active/inactive state.
-		if (InputPin && InputPin->Properties.IsAdvancedPin())
+		if (Input.DownstreamPin.IsSet() && Input.DownstreamPin.GetValue().IsAdvancedPin())
 		{
 			continue;
 		}
@@ -573,12 +580,12 @@ bool FPCGGraphCompiler::CalculateStaticallyActiveRecursive(FPCGTaskId InTaskId,
 		bool bInputActive = true;
 
 		// If we are connected to an upstream node, evaluate if the output pin is active.
-		if (Input.InPin)
+		if (Input.UpstreamPin.IsSet())
 		{
 			const UPCGNode* UpstreamNode = InCompiledTasks[Input.TaskId].Node;
 			if (const UPCGSettings* UpstreamSettings = UpstreamNode ? UpstreamNode->GetSettings() : nullptr)
 			{
-				bInputActive &= UpstreamSettings->IsPinStaticallyActive(Input.InPin->Properties.Label);
+				bInputActive &= UpstreamSettings->IsPinStaticallyActive(Input.UpstreamPin.GetValue().Label);
 			}
 		}
 
@@ -591,10 +598,10 @@ bool FPCGGraphCompiler::CalculateStaticallyActiveRecursive(FPCGTaskId InTaskId,
 		{
 			bHasAnyActiveNonAdvancedInput = true;
 
-			if (InputPin)
+			if (Input.DownstreamPin.IsSet())
 			{
 				// Register received input on this pin.
-				PinsRequiringActiveConnection.Remove(InputPin->Properties.Label);
+				PinsRequiringActiveConnection.Remove(Input.DownstreamPin.GetValue().Label);
 			}
 		}
 	}
@@ -703,7 +710,7 @@ void FPCGGraphCompiler::CullTasks(TArray<FPCGGraphTask>& InOutCompiledTasks, boo
 					const int NewWireCount = CulledInputTask.Inputs.Num();
 					for (int I = 0; I < NewWireCount; ++I)
 					{
-						Task.Inputs[InputIndex + 1 + I].OutPin = Task.Inputs[InputIndex].OutPin;
+						Task.Inputs[InputIndex + 1 + I].DownstreamPin = Task.Inputs[InputIndex].DownstreamPin;
 					}
 
 					// Skip evaluating the newly wired inputs and increment
@@ -823,15 +830,15 @@ void FPCGGraphCompiler::CalculateDynamicActivePinDependencies(FPCGTaskId InTaskI
 
 	for (const FPCGGraphTaskInput& Input : InOutCompiledTasks[InTaskId].Inputs)
 	{
-		if (!Node || !Input.OutPin)
+		if (!Node || !Input.DownstreamPin.IsSet())
 		{
 			continue;
 		}
 
-		const UPCGPin* InputPin = Input.OutPin;
+		const UPCGPin* InputPin = Node->GetInputPin(Input.DownstreamPin.GetValue().Label);
 
 		// Consider only primary input pins in this pass.
-		if (!Node->IsInputPinRequiredByExecution(InputPin))
+		if (!InputPin || !Node->IsInputPinRequiredByExecution(InputPin))
 		{
 			continue;
 		}
@@ -844,14 +851,14 @@ void FPCGGraphCompiler::CalculateDynamicActivePinDependencies(FPCGTaskId InTaskI
 
 		const int PinIndex = UpstreamNode->GetOutputPins().IndexOfByPredicate([&Input](const UPCGPin* InPin)
 		{
-			return InPin == Input.InPin;
+			return InPin->Properties == Input.UpstreamPin;
 		});
 
 		if (PinIndex != INDEX_NONE)
 		{
 			check(PinIndex < PCGPinIdHelpers::MaxOutputPins);
 
-			FPCGPinDependencyExpression& PinDependency = InputPinLabelToPinDependency.FindOrAdd(Input.OutPin->Properties.Label);
+			FPCGPinDependencyExpression& PinDependency = InputPinLabelToPinDependency.FindOrAdd(Input.DownstreamPin.GetValue().Label);
 			PinDependency.AddPinDependency(PCGPinIdHelpers::NodeIdAndPinIndexToPinId(Input.TaskId, PinIndex));
 		}
 	}
@@ -879,13 +886,13 @@ void FPCGGraphCompiler::CalculateDynamicActivePinDependencies(FPCGTaskId InTaskI
 		// is active. We build a disjunction that expresses this.
 		for (const FPCGGraphTaskInput& Input : InOutCompiledTasks[InTaskId].Inputs)
 		{
-			if (Input.OutPin && Input.OutPin->Properties.IsAdvancedPin() && !bTreatAdvancedPinsAsNormal)
+			if (Input.DownstreamPin.IsSet() && Input.DownstreamPin.GetValue().IsAdvancedPin() && !bTreatAdvancedPinsAsNormal)
 			{
 				// Advanced input pins never participate in keeping node active.
 				continue;
 			}
-
-			if (const UPCGPin* UpstreamOutputPin = Input.InPin)
+		
+			if (Input.UpstreamPin.IsSet())
 			{
 				// Input connection is via node pins.
 				const UPCGNode* UpstreamNode = InOutCompiledTasks[Input.TaskId].Node;
@@ -894,9 +901,9 @@ void FPCGGraphCompiler::CalculateDynamicActivePinDependencies(FPCGTaskId InTaskI
 					continue;
 				}
 
-				const int PinIndex = UpstreamNode->GetOutputPins().IndexOfByPredicate([UpstreamOutputPin](const UPCGPin* InPin)
+				const int PinIndex = UpstreamNode->GetOutputPins().IndexOfByPredicate([&Input](const UPCGPin* InPin)
 				{
-					return InPin == UpstreamOutputPin;
+					return InPin->Properties == Input.UpstreamPin.GetValue();
 				});
 
 				if (PinIndex != INDEX_NONE)
@@ -1070,7 +1077,7 @@ void FPCGGraphCompiler::CompileTopGraph(UPCGGraph* InGraph, uint32 GenerationGri
 		FPCGGraphTask& Task = CompiledTasks[TaskIndex];
 		if (Task.Inputs.IsEmpty())
 		{
-			Task.Inputs.Emplace(PreExecuteTaskId, nullptr, nullptr);
+			Task.Inputs.Emplace(PreExecuteTaskId);
 		}
 
 		Task.CompiledTaskId = Task.NodeId;
@@ -1096,7 +1103,7 @@ void FPCGGraphCompiler::CompileTopGraph(UPCGGraph* InGraph, uint32 GenerationGri
 			// It is necessary for any post generation task to get the content of the output node
 			// and only this content.
 			const bool bProvideData = Task.Node == GraphOutputNode;
-			PostExecuteTask.Inputs.Emplace(Task.NodeId, nullptr, nullptr, bProvideData);
+			PostExecuteTask.Inputs.Emplace(Task.NodeId, /*InUpstreamPin=*/FPCGGraphTaskInput::NoPin, /*InDownstreamPin=*/FPCGGraphTaskInput::NoPin, bProvideData);
 		}
 
 		for (const FPCGGraphTaskInput& Input : Task.Inputs)
diff --git a/Engine/Plugins/PCG/Source/PCG/Private/Graph/PCGGraphExecutor.cpp b/Engine/Plugins/PCG/Source/PCG/Private/Graph/PCGGraphExecutor.cpp
--- a/Engine/Plugins/PCG/Source/PCG/Private/Graph/PCGGraphExecutor.cpp
+++ b/Engine/Plugins/PCG/Source/PCG/Private/Graph/PCGGraphExecutor.cpp
@@ -100,8 +100,8 @@ void FPCGGraphTask::LogVisual(ELogVerbosity::Type InVerbosity, const FText& InMe
 bool FPCGGraphTaskInput::operator==(const FPCGGraphTaskInput& Other) const
 {
 	return (TaskId == Other.TaskId)
-		&& (InPin == Other.InPin)
-		&& (OutPin == Other.OutPin)
+		&& (UpstreamPin == Other.UpstreamPin)
+		&& (DownstreamPin == Other.DownstreamPin)
 		&& (bProvideData == Other.bProvideData);
 }
 
@@ -273,7 +273,7 @@ FPCGTaskId FPCGGraphExecutor::Schedule(
 		for (FPCGTaskId ExternalDependency : ExternalDependencies)
 		{
 			// For the pre-task, we don't consume any input
-			PreGraphTask.Inputs.Emplace(ExternalDependency, nullptr, nullptr, /*bConsumeData=*/false);
+			PreGraphTask.Inputs.Emplace(ExternalDependency, /*InUpstreamPin=*/FPCGGraphTaskInput::NoPin, /*InDownstreamPin=*/FPCGGraphTaskInput::NoPin, /*bInProvideData=*/false);
 		}
 
 		ScheduledTask.FirstTaskIndex = ScheduledTask.Tasks.Num() - 2;
@@ -578,13 +578,13 @@ FPCGTaskId FPCGGraphExecutor::ScheduleGenericWithContext(TFunction<bool(FPCGCont
 	for (FPCGTaskId TaskDependency : TaskExecutionDependencies)
 	{
 		ensure(TaskDependency != InvalidPCGTaskId);
-		Task.Inputs.Emplace(TaskDependency, /*InPin=*/nullptr, /*OutPin=*/nullptr, /*bConsumeInputData=*/false);
+		Task.Inputs.Emplace(TaskDependency, /*InUpstreamPin=*/FPCGGraphTaskInput::NoPin, /*InDownstreamPin=*/FPCGGraphTaskInput::NoPin, /*bInProvideData=*/false);
 	}
 
 	for (FPCGTaskId TaskDependency : TaskDataDependencies)
 	{
 		ensure(TaskDependency != InvalidPCGTaskId);
-		Task.Inputs.Emplace(TaskDependency, /*InPin=*/nullptr, /*OutPin=*/nullptr, /*bConsumeInputData=*/true);
+		Task.Inputs.Emplace(TaskDependency);
 	}
 
 	Task.SourceComponent = InSourceComponent;
@@ -1326,7 +1326,7 @@ void FPCGGraphExecutor::BuildTaskInput(const FPCGGraphTask& Task, FPCGDataCollec
 {
 	TRACE_CPUPROFILER_EVENT_SCOPE(FPCGGraphExecutor::BuildTaskInput);
 
-	auto LogDiscardedData = [&Task](const UPCGPin* InPin)
+	auto LogDiscardedData = [&Task](const TOptional<FPCGPinProperties>& DownstreamPin)
 	{
 #if !(UE_BUILD_SHIPPING || UE_BUILD_TEST) || USE_LOGGING_IN_SHIPPING
 		// Turn off eventual errors/warnings when the node is disabled, as this is irrelevant.
@@ -1340,7 +1340,7 @@ void FPCGGraphExecutor::BuildTaskInput(const FPCGGraphTask& Task, FPCGDataCollec
 			TEXT("[%s] %s - BuildTaskInput - too many data items arriving on single data pin '%s', only first data item will be used"),
 			(Task.SourceComponent.Get() && Task.SourceComponent->GetOwner()) ? *Task.SourceComponent->GetOwner()->GetName() : TEXT("MissingComponent"),
 			Task.Node ? *Task.Node->GetNodeTitle(EPCGNodeTitleType::ListView).ToString() : TEXT("MissingNode"),
-			InPin ? *InPin->Properties.Label.ToString() : TEXT("MissingPin"));
+			DownstreamPin.IsSet() ? *DownstreamPin.GetValue().Label.ToString() : TEXT("NoPin"));
 
 #if WITH_EDITOR
 		Task.LogVisual(ELogVerbosity::Warning, FText::FromString(Message));
@@ -1376,13 +1376,13 @@ void FPCGGraphExecutor::BuildTaskInput(const FPCGGraphTask& Task, FPCGDataCollec
 			continue;
 		}
 
-		const bool bAllowMultipleData = Input.OutPin ? Input.OutPin->Properties.bAllowMultipleData : true;
-		const uint32 InputPinLabelCrc = Input.OutPin ? GetTypeHash(Input.OutPin->Properties.Label) : DefaultHashForNoOutputPin;
+		const bool bAllowMultipleData = Input.DownstreamPin.IsSet() ? Input.DownstreamPin.GetValue().bAllowMultipleData : true;
+		const uint32 InputPinLabelCrc = Input.DownstreamPin.IsSet() ? GetTypeHash(Input.DownstreamPin.GetValue().Label) : DefaultHashForNoOutputPin;
 
 		// Enforce single data - if already have input for this pin, don't add more. Early check before other side effects below.
-		if (Input.OutPin && !bAllowMultipleData && TaskInput.GetInputCountByPin(Input.OutPin->Properties.Label) > 0)
+		if (Input.DownstreamPin.IsSet() && !bAllowMultipleData && TaskInput.GetInputCountByPin(Input.DownstreamPin.GetValue().Label) > 0)
 		{
-			LogDiscardedData(Input.OutPin);
+			LogDiscardedData(Input.DownstreamPin);
 			continue;
 		}
 
@@ -1395,11 +1395,11 @@ void FPCGGraphExecutor::BuildTaskInput(const FPCGGraphTask& Task, FPCGDataCollec
 		// Get input data at the given pin (or everything). This will add the data and include the input pin Crc to uniquely identify
 		// inputs per-pin, or use a placeholder for symmetry.
 		// Note: The input data CRC will already contain the output pin (calculated in element post execute).
-		if (Input.InPin)
+		if (Input.UpstreamPin.IsSet())
 		{
 			InputDataOnPin.Reset();
 			InputDataCrcsOnPin.Reset();
-			InputCollection.GetInputsAndCrcsByPin(Input.InPin->Properties.Label, InputDataOnPin, InputDataCrcsOnPin);
+			InputCollection.GetInputsAndCrcsByPin(Input.UpstreamPin.GetValue().Label, InputDataOnPin, InputDataCrcsOnPin);
 
 			if (!InputDataOnPin.IsEmpty())
 			{
@@ -1418,7 +1418,7 @@ void FPCGGraphExecutor::BuildTaskInput(const FPCGGraphTask& Task, FPCGDataCollec
 
 				if (NumberDataItemsToTake < InputDataOnPin.Num())
 				{
-					LogDiscardedData(Input.OutPin);
+					LogDiscardedData(Input.DownstreamPin);
 				}
 			}
 		}
@@ -1429,11 +1429,11 @@ void FPCGGraphExecutor::BuildTaskInput(const FPCGGraphTask& Task, FPCGDataCollec
 
 		// Apply labelling on data; technically, we should ensure that we do this only for pass-through nodes,
 		// Otherwise we could also null out the label on the input...
-		if (Input.OutPin)
+		if (Input.DownstreamPin.IsSet())
 		{
 			for (int32 TaggedDataIndex = TaggedDataOffset; TaggedDataIndex < TaskInput.TaggedData.Num(); ++TaggedDataIndex)
 			{
-				TaskInput.TaggedData[TaggedDataIndex].Pin = Input.OutPin->Properties.Label;
+				TaskInput.TaggedData[TaggedDataIndex].Pin = Input.DownstreamPin.GetValue().Label;
 			}
 		}
 	}
@@ -2068,10 +2068,11 @@ namespace PCGGraphExecutor
 
 		if (!!(InFromGrid & InGenerationGrid) && FromGridSize != ToGridSize)
 		{
-			PCGGraphExecutionLogging::LogGridLinkageTaskExecuteStore(InContext, InGenerationGrid, FromGridSize, ToGridSize, InResourceKey);
-
 			FPCGDataCollection Data;
 			Data.TaggedData = InContext->InputData.GetInputsByPin(InOutputPinLabel);
+
+			PCGGraphExecutionLogging::LogGridLinkageTaskExecuteStore(InContext, InGenerationGrid, FromGridSize, ToGridSize, InResourceKey, Data.TaggedData.Num());
+			
 			InContext->SourceComponent->StoreOutputDataForPin(InResourceKey, Data);
 		}
 		else if (InToGrid == InGenerationGrid && FromGridSize != ToGridSize)
diff --git a/Engine/Plugins/PCG/Source/PCG/Private/Graph/PCGGraphExecutor.h b/Engine/Plugins/PCG/Source/PCG/Private/Graph/PCGGraphExecutor.h
--- a/Engine/Plugins/PCG/Source/PCG/Private/Graph/PCGGraphExecutor.h
+++ b/Engine/Plugins/PCG/Source/PCG/Private/Graph/PCGGraphExecutor.h
@@ -4,6 +4,7 @@
 
 #include "PCGContext.h"
 #include "PCGElement.h"
+#include "PCGPin.h"
 #include "PCGSubsystem.h"
 #include "Graph/PCGGraphCache.h"
 #include "Graph/PCGGraphCompiler.h"
@@ -17,7 +18,6 @@
 #include "WorldPartition/WorldPartitionHandle.h" // Needed for FWorldPartitionReference
 #endif
 
-class UPCGPin;
 class UPCGGraph;
 class UPCGNode;
 class UPCGComponent;
@@ -38,10 +38,10 @@ namespace PCGGraphExecutor
 
 struct FPCGGraphTaskInput
 {
-	FPCGGraphTaskInput(FPCGTaskId InTaskId, const UPCGPin* InInboundPin, const UPCGPin* InOutboundPin, bool bInProvideData = true)
+	FPCGGraphTaskInput(FPCGTaskId InTaskId, const TOptional<FPCGPinProperties>& InUpstreamPin = NoPin, const TOptional<FPCGPinProperties>& InDownstreamPin = NoPin, bool bInProvideData = true)
 		: TaskId(InTaskId)
-		, InPin(InInboundPin)
-		, OutPin(InOutboundPin)
+		, UpstreamPin(InUpstreamPin)
+		, DownstreamPin(InDownstreamPin)
 		, bProvideData(bInProvideData)
 	{
 	}
@@ -53,13 +53,15 @@ struct FPCGGraphTaskInput
 	FPCGTaskId TaskId;
 
 	/** The upstream output pin from which the input data comes. */
-	const UPCGPin* InPin;
+	TOptional<FPCGPinProperties> UpstreamPin;
 
 	/** The input pin on the task element. */
-	const UPCGPin* OutPin;
+	TOptional<FPCGPinProperties> DownstreamPin;
 
 	/** Whether the input provides any data. For the post execute task, only the output node will provide data. */
 	bool bProvideData;
+
+	static inline const TOptional<FPCGPinProperties> NoPin = TOptional<FPCGPinProperties>();
 };
 
 struct FPCGGraphTask
diff --git a/Engine/Plugins/PCG/Source/PCG/Private/Utils/PCGGraphExecutionLogging.cpp b/Engine/Plugins/PCG/Source/PCG/Private/Utils/PCGGraphExecutionLogging.cpp
--- a/Engine/Plugins/PCG/Source/PCG/Private/Utils/PCGGraphExecutionLogging.cpp
+++ b/Engine/Plugins/PCG/Source/PCG/Private/Utils/PCGGraphExecutionLogging.cpp
@@ -77,7 +77,10 @@ namespace PCGGraphExecutionLogging
 				}
 				bFirstInput = false;
 
-				InputString += FString::Printf(TEXT("%" UINT64_FMT "->'%s'"), Input.TaskId, Input.OutPin ? *Input.OutPin->Properties.Label.ToString() : TEXT(""));
+				InputString += FString::Printf(
+					TEXT("%" UINT64_FMT "->'%s'"),
+					Input.TaskId,
+					Input.DownstreamPin.IsSet() ? *Input.DownstreamPin.GetValue().Label.ToString() : TEXT("NoPin"));
 			}
 
 			return InputString;
@@ -368,7 +371,7 @@ namespace PCGGraphExecutionLogging
 #endif // !(UE_BUILD_SHIPPING || UE_BUILD_TEST) || USE_LOGGING_IN_SHIPPING
 	}
 
-	void LogGridLinkageTaskExecuteStore(const FPCGContext* InContext, EPCGHiGenGrid InGenerationGrid, int32 InFromGridSize, int32 InToGridSize, const FString& InResourcePath)
+	void LogGridLinkageTaskExecuteStore(const FPCGContext* InContext, EPCGHiGenGrid InGenerationGrid, int32 InFromGridSize, int32 InToGridSize, const FString& InResourcePath, int32 InDataItemCount)
 	{
 #if !(UE_BUILD_SHIPPING || UE_BUILD_TEST) || USE_LOGGING_IN_SHIPPING
 		if (!LogEnabled())
@@ -377,12 +380,13 @@ namespace PCGGraphExecutionLogging
 		}
 		check(InContext);
 
-		UE_LOG(LogPCG, Log, TEXT("[GRIDLINKING] [%s] STORE. GenerationGridSize=%d, FromGridSize=%d, ToGridSize=%d, Path=%s"),
+		UE_LOG(LogPCG, Log, TEXT("[GRIDLINKING] [%s] STORE. GenerationGridSize=%u, FromGridSize=%u, ToGridSize=%u, Path=%s, DataItems=%d"),
 			*GetOwnerName(InContext->SourceComponent.Get()),
-			PCGHiGenGrid::GridToGridSize(InGenerationGrid),
+			PCGHiGenGrid::IsValidGrid(InGenerationGrid) ? PCGHiGenGrid::GridToGridSize(InGenerationGrid) : PCGHiGenGrid::UnboundedGridSize(),
 			InFromGridSize,
 			InToGridSize,
-			*InResourcePath);
+			*InResourcePath,
+			InDataItemCount);
 #endif // !(UE_BUILD_SHIPPING || UE_BUILD_TEST) || USE_LOGGING_IN_SHIPPING
 	}
 
@@ -395,9 +399,9 @@ namespace PCGGraphExecutionLogging
 		}
 		check(InContext);
 
-		UE_LOG(LogPCG, Log, TEXT("[GRIDLINKING] [%s] RETRIEVE. GenerationGridSize=%d, FromGridSize=%d, ToGridSize=%d, Path=%s"),
+		UE_LOG(LogPCG, Log, TEXT("[GRIDLINKING] [%s] RETRIEVE. GenerationGridSize=%u, FromGridSize=%u, ToGridSize=%u, Path=%s"),
 			*GetOwnerName(InContext->SourceComponent.Get()),
-			PCGHiGenGrid::GridToGridSize(InGenerationGrid),
+			PCGHiGenGrid::IsValidGrid(InGenerationGrid) ? PCGHiGenGrid::GridToGridSize(InGenerationGrid) : PCGHiGenGrid::UnboundedGridSize(),
 			InFromGridSize,
 			InToGridSize,
 			*InResourcePath);
@@ -413,7 +417,7 @@ namespace PCGGraphExecutionLogging
 		}
 		check(InContext);
 
-		UE_LOG(LogPCG, Log, TEXT("[GRIDLINKING] [%s] RETRIEVE: SUCCESS. Path=%s DataItems=%d"),
+		UE_LOG(LogPCG, Log, TEXT("[GRIDLINKING] [%s] RETRIEVE: SUCCESS. Path=%s, DataItems=%d"),
 			*GetOwnerName(InContext->SourceComponent.Get()),
 			*InResourcePath,
 			InDataItemCount);
@@ -429,7 +433,7 @@ namespace PCGGraphExecutionLogging
 		}
 		check(InContext);
 
-		UE_LOG(LogPCG, Log, TEXT("[GRIDLINKING] [%s] RETRIEVE: SCHEDULE GRAPH. Component=%s Path=%s"),
+		UE_LOG(LogPCG, Log, TEXT("[GRIDLINKING] [%s] RETRIEVE: SCHEDULE GRAPH. Component=%s, Path=%s"),
 			*GetOwnerName(InContext->SourceComponent.Get()),
 			*GetOwnerName(InScheduledComponent),
 			*InResourcePath);
@@ -445,7 +449,7 @@ namespace PCGGraphExecutionLogging
 		}
 		check(InContext);
 
-		UE_LOG(LogPCG, Log, TEXT("[GRIDLINKING] [%s] RETRIEVE: WAIT FOR SCHEDULED GRAPH. Component=%s Path=%s"),
+		UE_LOG(LogPCG, Log, TEXT("[GRIDLINKING] [%s] RETRIEVE: WAIT FOR SCHEDULED GRAPH. Component=%s, Path=%s"),
 			*GetOwnerName(InContext->SourceComponent.Get()),
 			*GetOwnerName(InWaitOnComponent),
 			*InResourcePath);
diff --git a/Engine/Plugins/PCG/Source/PCG/Private/Utils/PCGGraphExecutionLogging.h b/Engine/Plugins/PCG/Source/PCG/Private/Utils/PCGGraphExecutionLogging.h
--- a/Engine/Plugins/PCG/Source/PCG/Private/Utils/PCGGraphExecutionLogging.h
+++ b/Engine/Plugins/PCG/Source/PCG/Private/Utils/PCGGraphExecutionLogging.h
@@ -45,7 +45,7 @@ namespace PCGGraphExecutionLogging
 	void LogTaskCullingBeginLoop(FPCGTaskId PinTaskId, uint64 PinIndex, const TArray<FPCGPinId>& PinIdsToDeactivate);
 	void LogTaskCullingUpdatedPinDeps(FPCGTaskId TaskId, const FPCGPinDependencyExpression& PinDependency, bool bDependencyExpressionBecameFalse);
 
-	void LogGridLinkageTaskExecuteStore(const FPCGContext* InContext, EPCGHiGenGrid InGenerationGrid, int32 InFromLOD, int32 InToLOD, const FString& InResourcePath);
+	void LogGridLinkageTaskExecuteStore(const FPCGContext* InContext, EPCGHiGenGrid InGenerationGrid, int32 InFromLOD, int32 InToLOD, const FString& InResourcePath, int32 InDataItemCount);
 	void LogGridLinkageTaskExecuteRetrieve(const FPCGContext* InContext, EPCGHiGenGrid InGenerationGrid, int32 InFromLOD, int32 InToLOD, const FString& InResourcePath);
 	void LogGridLinkageTaskExecuteRetrieveSuccess(const FPCGContext* InContext, const UPCGComponent* InComponent, const FString& InResourcePath, int32 InDataItemCount);
 	void LogGridLinkageTaskExecuteRetrieveScheduleGraph(const FPCGContext* InContext, const UPCGComponent* InScheduledComponent, const FString& InResourcePath);
