Robustness fix for collision wrapper data & self-pruning when there are meshes with no collision

#rnx

diff --git a/Engine/Plugins/PCG/Source/PCG/Private/Data/PCGCollisionWrapperData.cpp b/Engine/Plugins/PCG/Source/PCG/Private/Data/PCGCollisionWrapperData.cpp
--- a/Engine/Plugins/PCG/Source/PCG/Private/Data/PCGCollisionWrapperData.cpp
+++ b/Engine/Plugins/PCG/Source/PCG/Private/Data/PCGCollisionWrapperData.cpp
@@ -90,7 +90,7 @@ void FPCGCollisionWrapper::CreateBodyInstances(const TArray<FSoftObjectPath>& Me
 			MeshPtr = Mesh.LoadSynchronous();
 		}
 
-		if(MeshPtr)
+		if(MeshPtr && MeshPtr->GetBodySetup())
 		{
 			//TRACE_CPUPROFILER_EVENT_SCOPE(FPCGCollisionWrapper::CreateBodyInstances::CreateBodyInstance);
 			BodyInstance = new FBodyInstance();
diff --git a/Engine/Plugins/PCG/Source/PCG/Private/Elements/PCGSelfPruning.cpp b/Engine/Plugins/PCG/Source/PCG/Private/Elements/PCGSelfPruning.cpp
--- a/Engine/Plugins/PCG/Source/PCG/Private/Elements/PCGSelfPruning.cpp
+++ b/Engine/Plugins/PCG/Source/PCG/Private/Elements/PCGSelfPruning.cpp
@@ -243,6 +243,7 @@ namespace PCGSelfPruningElement
 						}
 						else
 						{
+							check(ThisInstance->GetBodySetup());
 							OtherInstance = new FBodyInstance();
 							OtherInstance->bAutoWeld = false;
 							OtherInstance->bSimulatePhysics = false;
