Added PCG node to wait for landscape changes to be done, with also some options around not triggering updates from landscape

diff --git a/Engine/Plugins/PCG/Source/PCG/Private/Elements/Landscape/PCGWaitLandscapeReady.cpp b/Engine/Plugins/PCG/Source/PCG/Private/Elements/Landscape/PCGWaitLandscapeReady.cpp
new file mode 100644
--- /dev/null
+++ b/Engine/Plugins/PCG/Source/PCG/Private/Elements/Landscape/PCGWaitLandscapeReady.cpp
@@ -0,0 +1,126 @@
+// Copyright Epic Games, Inc. All Rights Reserved.
+
+#include "Elements/Landscape/PCGWaitLandscapeReady.h"
+
+#include "PCGComponent.h"
+#include "PCGSubsystem.h"
+#include "Elements/PCGActorSelector.h"
+
+#include "Landscape.h"
+
+TArray<FPCGPinProperties> UPCGWaitLandscapeReadySettings::InputPinProperties() const
+{
+	TArray<FPCGPinProperties> PinProperties;
+	FPCGPinProperties& InputPinProperty = PinProperties.Emplace_GetRef(PCGPinConstants::DefaultInputLabel, EPCGDataType::Any);
+	InputPinProperty.SetRequiredPin();
+
+	return PinProperties;
+}
+
+TArray<FPCGPinProperties> UPCGWaitLandscapeReadySettings::OutputPinProperties() const
+{
+	TArray<FPCGPinProperties> PinProperties;
+	PinProperties.Emplace(PCGPinConstants::DefaultOutputLabel, EPCGDataType::Any);
+
+	return PinProperties;
+}
+
+FPCGElementPtr UPCGWaitLandscapeReadySettings::CreateElement() const
+{
+	return MakeShared<FPCGWaitLandscapeReadyElement>();
+}
+
+bool FPCGWaitLandscapeReadyElement::ExecuteInternal(FPCGContext* InContext) const
+{
+	TRACE_CPUPROFILER_EVENT_SCOPE(FPCGWaitLandscapeReadyElement::Execute);
+	check(InContext);
+	FPCGWaitLandscapeReadyElementContext* Context = static_cast<FPCGWaitLandscapeReadyElementContext*>(InContext);
+	const UPCGWaitLandscapeReadySettings* Settings = Context->GetInputSettings<UPCGWaitLandscapeReadySettings>();
+	check(Settings);
+
+	if (!Context->bLandscapeQueryDone)
+	{
+		TRACE_CPUPROFILER_EVENT_SCOPE(FPCGWaitLandscapeReadyElement::Execute::LandscapeQuery);
+		check(Context->Landscapes.IsEmpty());
+		Context->bLandscapeQueryDone = true;
+
+		// Get landscape(s)
+		FPCGActorSelectorSettings ActorSelector;
+		ActorSelector.ActorFilter = EPCGActorFilter::AllWorldActors;
+		ActorSelector.ActorSelection = EPCGActorSelection::ByClass;
+		ActorSelector.ActorSelectionClass = ALandscapeProxy::StaticClass();
+		ActorSelector.bSelectMultiple = true;
+
+		auto BoundsCheck = [](const AActor*) -> bool { return true; };
+		auto SelfIgnoreCheck = [](const AActor*) -> bool { return true; };
+
+		TArray<AActor*> FoundLandscapeProxies = PCGActorSelector::FindActors(ActorSelector, Context->SourceComponent.Get(), BoundsCheck, SelfIgnoreCheck);
+
+		for (AActor* Proxy : FoundLandscapeProxies)
+		{
+			if (ALandscapeProxy* LandscapeProxy = Cast<ALandscapeProxy>(Proxy))
+			{
+				TWeakObjectPtr<ALandscape> Landscape(LandscapeProxy->GetLandscapeActor());
+				Context->Landscapes.AddUnique(Landscape);
+			}
+		}
+	}
+
+	if (!Context->bLandscapeReady)
+	{
+		TRACE_CPUPROFILER_EVENT_SCOPE(FPCGWaitLandscapeReadyElement::Execute::CheckingIfLandscapeAreReady);
+		check(Context->bLandscapeQueryDone);
+
+		bool bAllLandscapesAreReady = true;
+		for (const TWeakObjectPtr<ALandscape>& Landscape : Context->Landscapes)
+		{
+			if (Landscape.IsValid() && !Landscape->IsUpToDate())
+			{
+				bAllLandscapesAreReady = false;
+				break;
+			}
+		}
+
+		if (bAllLandscapesAreReady)
+		{
+			Context->bLandscapeReady = true;
+		}
+	}
+
+	// We've validated the landscape(s) are ready - we`re done.
+	if (Context->bLandscapeReady || Context->OutputData.bCancelExecution)
+	{
+		TRACE_CPUPROFILER_EVENT_SCOPE(FPCGWaitLandscapeReadyElement::Execute::FinalizeExecution);
+		Context->OutputData = Context->InputData;
+		return true;
+	}
+
+	// At this point, we need to go to sleep for at least the remainder of the frame
+	UPCGSubsystem* Subsystem = Context->SourceComponent.IsValid() ? Context->SourceComponent->GetSubsystem() : nullptr;
+	if (ensure(Subsystem))
+	{
+		Context->bIsPaused = true;
+
+		Subsystem->ScheduleGeneric(
+			[Context]() // Normal execution: wake up the current task
+			{
+				Context->bIsPaused = false;
+				return true;
+			},
+			[Context]() // On abort: wakeup and cancel
+			{
+				Context->bIsPaused = false;
+				Context->OutputData.bCancelExecution = true;
+				return true;
+			},
+			Context->SourceComponent.Get(),
+			{});
+
+		return false;
+	}
+	else
+	{
+		return true;
+	}
+}
+
diff --git a/Engine/Plugins/PCG/Source/PCG/Private/PCGComponent.cpp b/Engine/Plugins/PCG/Source/PCG/Private/PCGComponent.cpp
--- a/Engine/Plugins/PCG/Source/PCG/Private/PCGComponent.cpp
+++ b/Engine/Plugins/PCG/Source/PCG/Private/PCGComponent.cpp
@@ -3452,6 +3452,11 @@ void UPCGComponent::StopIgnoringChangeOriginDuringGeneration(UObject* InChangeOr
 
 bool UPCGComponent::IsIgnoringChangeOrigin(UObject* InChangeOrigin)
 {
+	if ((bIgnoreLandscapeTracking || (GetGraph() && GetGraph()->bIgnoreLandscapeTracking)) && Cast<ALandscapeProxy>(InChangeOrigin))
+	{
+		return true;
+	}
+
 	FReadScopeLock Lock(IgnoredChangeOriginsLock);
 	const int32* Counter = IgnoredChangeOriginsToCounters.Find(InChangeOrigin);
 	return Counter && ensure(*Counter > 0);
diff --git a/Engine/Plugins/PCG/Source/PCG/Public/Elements/Landscape/PCGWaitLandscapeReady.h b/Engine/Plugins/PCG/Source/PCG/Public/Elements/Landscape/PCGWaitLandscapeReady.h
new file mode 100644
--- /dev/null
+++ b/Engine/Plugins/PCG/Source/PCG/Public/Elements/Landscape/PCGWaitLandscapeReady.h
@@ -0,0 +1,51 @@
+// Copyright Epic Games, Inc. All Rights Reserved.
+#pragma once
+
+#include "PCGContext.h"
+#include "PCGSettings.h"	
+
+#include "PCGWaitLandscapeReady.generated.h"
+
+class ALandscape;
+
+/**
+* Used to wait until the landscape is done updating anything.
+*/
+UCLASS(BlueprintType, ClassGroup = (Procedural))
+class UPCGWaitLandscapeReadySettings : public UPCGSettings
+{
+	GENERATED_BODY()
+
+public:
+	//~Begin UPCGSettings interface
+#if WITH_EDITOR
+	virtual FName GetDefaultNodeName() const override { return FName(TEXT("WaitUntilLandscapeReady")); }
+	virtual FText GetDefaultNodeTitle() const override { return NSLOCTEXT("PCGWaitLandscapeReadySettings", "NodeTitle", "Wait Until Landscape Is Ready"); }
+	virtual FText GetNodeTooltipText() const override { return NSLOCTEXT("PCGWaitLandscapeReadySettings", "NodeTooltip", "Waits until landscape is ready, then passes data downstream."); }
+	virtual EPCGSettingsType GetType() const override { return EPCGSettingsType::Generic; }
+	virtual bool HasDynamicPins() const override { return true; }
+#endif
+
+protected:
+	virtual TArray<FPCGPinProperties> InputPinProperties() const override;
+	virtual TArray<FPCGPinProperties> OutputPinProperties() const override;
+	virtual FPCGElementPtr CreateElement() const override;
+	//~End UPCGSettings interface
+};
+
+struct FPCGWaitLandscapeReadyElementContext : public FPCGContext
+{
+	bool bLandscapeQueryDone = false;
+	bool bLandscapeReady = false;
+	TArray<TWeakObjectPtr<ALandscape>> Landscapes;
+};
+
+class FPCGWaitLandscapeReadyElement : public IPCGElementWithCustomContext<FPCGWaitLandscapeReadyElementContext>
+{
+public:
+	virtual bool CanExecuteOnlyOnMainThread(FPCGContext* Context) const override { return true; }
+	virtual bool IsCacheable(const UPCGSettings* InSettings) const override { return false; }
+
+protected:
+	virtual bool ExecuteInternal(FPCGContext* Context) const override;
+};
diff --git a/Engine/Plugins/PCG/Source/PCG/Public/PCGComponent.h b/Engine/Plugins/PCG/Source/PCG/Public/PCGComponent.h
--- a/Engine/Plugins/PCG/Source/PCG/Public/PCGComponent.h
+++ b/Engine/Plugins/PCG/Source/PCG/Public/PCGComponent.h
@@ -270,6 +270,10 @@ public:
 	UPROPERTY(BlueprintReadWrite, EditAnywhere, Category = "Editing Settings", meta = (DisplayPriority = 450))
 	bool bOnlyTrackItself = false;
 
+	/** Marks the component to be not refreshed automatically when the landscape changes, even if it is used. */
+	UPROPERTY(BlueprintReadWrite, EditAnywhere, Category = "Editing Settings", meta = (DisplayPriority = 451))
+	bool bIgnoreLandscapeTracking = false;
+
 	UPROPERTY(BlueprintReadOnly, VisibleAnywhere, Transient, Category = Debug, meta = (NoResetToDefault))
 	bool bDirtyGenerated = false;
 
diff --git a/Engine/Plugins/PCG/Source/PCG/Public/PCGGraph.h b/Engine/Plugins/PCG/Source/PCG/Public/PCGGraph.h
--- a/Engine/Plugins/PCG/Source/PCG/Public/PCGGraph.h
+++ b/Engine/Plugins/PCG/Source/PCG/Public/PCGGraph.h
@@ -217,6 +217,10 @@ public:
 
 	UPROPERTY(BlueprintReadWrite, EditAnywhere, Category = AssetInfo, AssetRegistrySearchable)
 	FText Description;
+
+	/** Marks the graph to be not refreshed automatically when the landscape changes, even if it is used. */
+	UPROPERTY(BlueprintReadWrite, EditAnywhere, Category = "Settings|Advanced")
+	bool bIgnoreLandscapeTracking = false;
 #endif
 
 	UPROPERTY(EditAnywhere, Category = Settings)
