Adding support to Apply on any uobject

diff --git a/Engine/Plugins/PCG/Source/PCG/Private/Elements/PCGApplyOnActor.cpp b/Engine/Plugins/PCG/Source/PCG/Private/Elements/PCGApplyOnActor.cpp
--- a/Engine/Plugins/PCG/Source/PCG/Private/Elements/PCGApplyOnActor.cpp
+++ b/Engine/Plugins/PCG/Source/PCG/Private/Elements/PCGApplyOnActor.cpp
@@ -3,56 +3,153 @@
 #include "Elements/PCGApplyOnActor.h"
 
 #include "PCGContext.h"
+#include "PCGEdge.h"
 #include "PCGModule.h"
+#include "PCGNode.h"
+#include "PCGPin.h"
 #include "Helpers/PCGHelpers.h"
 
 #include "GameFramework/Actor.h"
+#include "UObject/FortniteMainBranchObjectVersion.h"
 
 #define LOCTEXT_NAMESPACE "PCGApplyOnActorElement"
 
 namespace PCGApplyOnActorConstants
 {
 	const FText DependencyTooltip = LOCTEXT("DependencyPinTooltip", "Data passed to this pin will be used to order execution but will otherwise not contribute to the results of this node.");
-	const FName ActorPropertyOverridesLabel = TEXT("Property Overrides");
-	const FText ActorPropertyOverridesTooltip = LOCTEXT("ActorOverrideToolTip", "Provide property overrides for the target actor. The attribute name must match the InputSource name in the actor property override description.");
+	const FName ObjectPropertyOverridesLabel = TEXT("Property Overrides");
+	const FText ObjectPropertyOverridesTooltip = LOCTEXT("ObjectOverrideToolTip", "Provide property overrides for the target object. The attribute name must match the InputSource name in the object property override description.");
 }
 
+#if WITH_EDITOR
+void UPCGApplyOnActorSettings::ApplyDeprecationBeforeUpdatePins(UPCGNode* InOutNode, TArray<TObjectPtr<UPCGPin>>& InputPins, TArray<TObjectPtr<UPCGPin>>& OutputPins)
+{
+	if(GetLinkerCustomVersion(FFortniteMainBranchObjectVersion::GUID) < FFortniteMainBranchObjectVersion::PCGApplyOnActorNodeMoveTargetActorEdgeToInput && ensure(InOutNode))
+	{
+		// Find in the input pins if we have edge(s) on the "TargetActor" pin
+		TObjectPtr<UPCGPin>* TargetActorPinPtr = InputPins.FindByPredicate([](const TObjectPtr<UPCGPin>& InputPin) { return InputPin && InputPin->Properties.Label == TEXT("TargetActor"); });
+
+		if (TargetActorPinPtr && (*TargetActorPinPtr)->EdgeCount() > 0)
+		{
+			TObjectPtr<UPCGPin> TargetActorPin = *TargetActorPinPtr;
+
+			// Create input pin if needed (should be all the time)
+			TObjectPtr<UPCGPin> InputPin = nullptr;
+			TObjectPtr<UPCGPin>* PreExistingInputPin = InputPins.FindByPredicate([](const TObjectPtr<UPCGPin>& InputPin) { return InputPin && InputPin->Properties.Label == PCGPinConstants::DefaultInputLabel; });
+
+			if (!PreExistingInputPin)
+			{
+				InputPin = NewObject<UPCGPin>(InOutNode);
+				InputPin->Node = InOutNode;
+				InputPin->Properties = FPCGPinProperties(PCGPinConstants::DefaultInputLabel, EPCGDataType::Any, /*bAllowMultipleConnections=*/true, /*bAllowMultipleData=*/true);
+				InputPins.Insert(InputPin, 0);
+			}
+			else
+			{
+				InputPin = *PreExistingInputPin;
+			}
+
+			TArray<UPCGPin*> UpstreamPins;
+			for (const UPCGEdge* Connection : TargetActorPin->Edges)
+			{
+				UpstreamPins.Add(Connection->InputPin);
+			}
+
+			for (UPCGPin* UpstreamPin : UpstreamPins)
+			{
+				UpstreamPin->BreakEdgeTo(TargetActorPin);
+				UpstreamPin->AddEdgeTo(InputPin);
+			}
+		}
+	}
+
+	Super::ApplyDeprecationBeforeUpdatePins(InOutNode, InputPins, OutputPins);
+}
+#endif // WITH_EDITOR
+
 TArray<FPCGPinProperties> UPCGApplyOnActorSettings::InputPinProperties() const
 {
 	TArray<FPCGPinProperties> PinProperties;
+	PinProperties.Emplace(PCGPinConstants::DefaultInputLabel, EPCGDataType::Any, /*bAllowMultipleConnections=*/true, /*bAllowMultipleData=*/true);
+	PinProperties.Emplace(PCGApplyOnActorConstants::ObjectPropertyOverridesLabel, EPCGDataType::Any, /*bAllowMultipleConnections=*/true, /*bAllowMultipleData=*/true, PCGApplyOnActorConstants::ObjectPropertyOverridesTooltip);
 	PinProperties.Emplace(PCGPinConstants::DefaultDependencyOnlyLabel, EPCGDataType::Any, /*bAllowMultipleConnections=*/true, /*bAllowMultipleData=*/true, PCGApplyOnActorConstants::DependencyTooltip);
-	PinProperties.Add(PCGObjectPropertyOverrideHelpers::CreateObjectPropertiesOverridePin(PCGApplyOnActorConstants::ActorPropertyOverridesLabel, PCGApplyOnActorConstants::ActorPropertyOverridesTooltip));
+	return PinProperties;
+}
+
+TArray<FPCGPinProperties> UPCGApplyOnActorSettings::OutputPinProperties() const
+{
+	TArray<FPCGPinProperties> PinProperties;
+	PinProperties.Emplace(PCGPinConstants::DefaultOutputLabel, EPCGDataType::Any);
+
 	return PinProperties;
 }
 
 FPCGElementPtr UPCGApplyOnActorSettings::CreateElement() const
 {
-	return MakeShared<FPCGApplyOnActorElement>();
+	return MakeShared<FPCGApplyOnObjectElement>();
 }
 
-bool FPCGApplyOnActorElement::ExecuteInternal(FPCGContext* Context) const
+bool FPCGApplyOnObjectElement::PrepareDataInternal(FPCGContext* Context) const
+{
+	TRACE_CPUPROFILER_EVENT_SCOPE(FPCGApplyOnActorElement::PrepareData);
+	check(Context);
+
+	const UPCGApplyOnActorSettings* Settings = Context->GetInputSettings<UPCGApplyOnActorSettings>();
+	check(Settings);
+
+	FPCGLoadObjectsFromPathContext* ThisContext = static_cast<FPCGLoadObjectsFromPathContext*>(Context);
+	return ThisContext->InitializeAndRequestLoad(PCGPinConstants::DefaultInputLabel,
+		Settings->ObjectReferenceAttribute,
+		{},
+		/*bPersistAllData=*/false,
+		Settings->bSilenceErrorOnEmptyObjectPath,
+		Settings->bSynchronousLoad);
+}
+
+bool FPCGApplyOnObjectElement::ExecuteInternal(FPCGContext* Context) const
 {
 	TRACE_CPUPROFILER_EVENT_SCOPE(FPCGApplyOnActorElement::Execute);
 
 	check(Context);
+	FPCGLoadObjectsFromPathContext* ThisContext = static_cast<FPCGLoadObjectsFromPathContext*>(Context);
 
 	const UPCGApplyOnActorSettings* Settings = Context->GetInputSettings<UPCGApplyOnActorSettings>();
 	check(Settings);
 
-	AActor* TargetActor = Settings->TargetActor.Get() ? Settings->TargetActor.Get() : Context->GetTargetActor(nullptr);
-	if (!TargetActor)
+	TArray<TPair<UObject*, int32>> TargetObjectsAndIndices;
+
+	int CurrentPathIndex = 0;
+	while (CurrentPathIndex < ThisContext->PathsToObjectsAndDataIndex.Num())
 	{
-		PCGE_LOG(Error, GraphAndLog, LOCTEXT("InvalidTargetActor", "Invalid target actor."));
-		return true;
-	}
+		TargetObjectsAndIndices.Reset();
 
-	PCGObjectPropertyOverrideHelpers::ApplyOverridesFromParams(Settings->PropertyOverrideDescriptions, TargetActor, PCGApplyOnActorConstants::ActorPropertyOverridesLabel, Context);
+		int InputIndex = ThisContext->PathsToObjectsAndDataIndex[CurrentPathIndex].Get<1>();
+		while (CurrentPathIndex < ThisContext->PathsToObjectsAndDataIndex.Num() && InputIndex == ThisContext->PathsToObjectsAndDataIndex[CurrentPathIndex].Get<1>())
+		{
+			TargetObjectsAndIndices.Emplace(ThisContext->PathsToObjectsAndDataIndex[CurrentPathIndex].Get<0>().ResolveObject(), ThisContext->PathsToObjectsAndDataIndex[CurrentPathIndex].Get<2>());
+			++CurrentPathIndex;
+		}
 
-	for (UFunction* Function : PCGHelpers::FindUserFunctions(TargetActor->GetClass(), Settings->PostProcessFunctionNames, { UPCGFunctionPrototypes::GetPrototypeWithNoParams() }, Context))
-	{
-		TargetActor->ProcessEvent(Function, nullptr);
+		PCGObjectPropertyOverrideHelpers::ApplyOverrides(Settings->PropertyOverrideDescriptions, TargetObjectsAndIndices, PCGApplyOnActorConstants::ObjectPropertyOverridesLabel, InputIndex, Context);
+
+		for (const TPair<UObject*, int32>& TargetObjectAndIndex : TargetObjectsAndIndices)
+		{
+			UObject* TargetObject = TargetObjectAndIndex.Key;
+			if (!TargetObject)
+			{
+				continue;
+			}
+
+			for (UFunction* Function : PCGHelpers::FindUserFunctions(TargetObject->GetClass(), Settings->PostProcessFunctionNames, { UPCGFunctionPrototypes::GetPrototypeWithNoParams() }, Context))
+			{
+				TargetObject->ProcessEvent(Function, nullptr);
+			}
+		}
 	}
 
+	// Since we have only one output pin, there's no need to update the pins here.
+	Context->OutputData.TaggedData = Context->InputData.GetInputsByPin(PCGPinConstants::DefaultInputLabel);
+
 	return true;
 }
 
diff --git a/Engine/Plugins/PCG/Source/PCG/Private/Elements/PCGGetPropertyFromObjectPath.cpp b/Engine/Plugins/PCG/Source/PCG/Private/Elements/PCGGetPropertyFromObjectPath.cpp
--- a/Engine/Plugins/PCG/Source/PCG/Private/Elements/PCGGetPropertyFromObjectPath.cpp
+++ b/Engine/Plugins/PCG/Source/PCG/Private/Elements/PCGGetPropertyFromObjectPath.cpp
@@ -6,9 +6,6 @@
 #include "PCGParamData.h"
 #include "Helpers/PCGDynamicTrackingHelpers.h"
 #include "Helpers/PCGPropertyHelpers.h"
-#include "Metadata/Accessors/IPCGAttributeAccessor.h"
-#include "Metadata/Accessors/PCGAttributeAccessorKeys.h"
-#include "Metadata/Accessors/PCGAttributeAccessorHelpers.h"
 
 #define LOCTEXT_NAMESPACE "PCGGetPropertyFromObjectPathElement"
 
@@ -124,11 +121,6 @@ TArray<FPCGPinProperties> UPCGGetPropertyFromObjectPathSettings::OutputPinProper
 	return PinProperties;
 }
 
-FPCGContext* FPCGGetPropertyFromObjectPathElement::CreateContext()
-{
-	return new FPCGGetPropertyFromObjectPathContext();
-}
-
 bool FPCGGetPropertyFromObjectPathElement::PrepareDataInternal(FPCGContext* Context) const
 {
 	TRACE_CPUPROFILER_EVENT_SCOPE(FPCGGetPropertyFromObjectPathElement::PrepareData);
@@ -138,124 +130,13 @@ bool FPCGGetPropertyFromObjectPathElement::PrepareDataInternal(FPCGContext* Cont
 	const UPCGGetPropertyFromObjectPathSettings* Settings = Context->GetInputSettings<UPCGGetPropertyFromObjectPathSettings>();
 	check(Settings);
 
-	FPCGGetPropertyFromObjectPathContext* ThisContext = static_cast<FPCGGetPropertyFromObjectPathContext*>(Context);
-	IPCGAsyncLoadingContext* AsyncLoadContext = static_cast<IPCGAsyncLoadingContext*>(ThisContext);
-
-	bool bIsDone = true;
-
-	if (!AsyncLoadContext->WasLoadRequested())
-	{
-		const UPCGNode* Node = Context->Node;
-		const bool InPinIsConnected = Node ? Node->IsInputPinConnected(PCGPinConstants::DefaultInputLabel) : false;
-
-		// First gather all the soft objects paths to extract
-		TArray<FSoftObjectPath> ObjectsToLoad;
-		if (!InPinIsConnected)
-		{
-			ThisContext->PathsToObjectsToExtractAndIncomingDataIndex.Reserve(Settings->ObjectPathsToExtract.Num());
-			for (const FSoftObjectPath& Path : Settings->ObjectPathsToExtract)
-			{
-				if (!Path.IsNull())
-				{
-					ObjectsToLoad.AddUnique(Path);
-				}
-
-				ThisContext->PathsToObjectsToExtractAndIncomingDataIndex.Emplace(Path, -1);
-			}
-		}
-		else
-		{
-			const TArray<FPCGTaggedData> Inputs = Context->InputData.GetInputsByPin(PCGPinConstants::DefaultInputLabel);
-			for (int32 Index = 0; Index < Inputs.Num(); ++Index)
-			{
-				const FPCGTaggedData& Input = Inputs[Index];
-
-				if (!Input.Data)
-				{
-					PCGE_LOG(Error, GraphAndLog, FText::Format(LOCTEXT("InvalidData", "Invalid data for input {0}"), FText::AsNumber(Index)));
-					continue;
-				}
-
-				const FPCGAttributePropertyInputSelector AttributeSelector = Settings->InputSource.CopyAndFixLast(Input.Data);
-				const TUniquePtr<const IPCGAttributeAccessor> Accessor = PCGAttributeAccessorHelpers::CreateConstAccessor(Input.Data, AttributeSelector);
-				const TUniquePtr<const IPCGAttributeAccessorKeys> Keys = PCGAttributeAccessorHelpers::CreateConstKeys(Input.Data, AttributeSelector);
-
-				if (!Accessor.IsValid() || !Keys.IsValid())
-				{
-					if (Settings->bPersistAllData)
-					{
-						// Special case for empty data. We need this case if we ever chain this node multiple times. An empty param (with no attributes and no entries) will generate another empty param.
-						const UPCGMetadata* Metadata = Input.Data->ConstMetadata();
-						if (Metadata && Metadata->GetAttributeCount() == 0 && Metadata->GetLocalItemCount() == 0)
-						{
-							// Emplace empty path for this input. Will generate an empty param data.
-							ThisContext->PathsToObjectsToExtractAndIncomingDataIndex.Emplace(FSoftObjectPath(), Index);
-						}
-					}
-
-					if (!Settings->bSilenceErrorOnEmptyObjectPath)
-					{
-						PCGE_LOG(Error, GraphAndLog, FText::Format(LOCTEXT("AttributeNotFound", "Attribute/Property '{0}' does not exist on input {1}"), AttributeSelector.GetDisplayText(), FText::AsNumber(Index)));
-					}
-
-					continue;
-				}
-
-				const int32 NumElementsToAdd = Keys->GetNum();
-				if (NumElementsToAdd == 0)
-				{
-					continue;
-				}
-
-				// Extract value as String to validate that a path is empty or ill-formed (because any ill-formed path will be null).
-				TArray<FString> InputValues;
-				InputValues.SetNum(NumElementsToAdd);
-				if (Accessor->GetRange(MakeArrayView(InputValues), 0, *Keys, EPCGAttributeAccessorFlags::AllowConstructible | EPCGAttributeAccessorFlags::AllowBroadcast))
-				{
-					ThisContext->PathsToObjectsToExtractAndIncomingDataIndex.Reserve(ThisContext->PathsToObjectsToExtractAndIncomingDataIndex.Num() + NumElementsToAdd);
-					for (int32 i = 0; i < InputValues.Num(); ++i)
-					{
-						FString& StringPath = InputValues[i];
-						// Empty SoftObjectPath can convert to string to None and is treated as empty, so check that one too.
-						const bool PathIsEmpty = StringPath.IsEmpty() || StringPath.Equals(TEXT("None"), ESearchCase::CaseSensitive);
-						if (PathIsEmpty && Settings->bPersistAllData)
-						{
-							ThisContext->PathsToObjectsToExtractAndIncomingDataIndex.Emplace(FSoftObjectPath(), Index);
-						}
-
-						FSoftObjectPath Path(std::move(StringPath));
-
-						if (!Path.IsNull())
-						{
-							ObjectsToLoad.AddUnique(Path);
-							ThisContext->PathsToObjectsToExtractAndIncomingDataIndex.Emplace(std::move(Path), Index);
-						}
-						else
-						{
-							if (!PathIsEmpty || !Settings->bSilenceErrorOnEmptyObjectPath)
-							{
-								PCGE_LOG(Error, GraphAndLog, FText::Format(LOCTEXT("InvalidPath", "Value number {0} for Attribute/Property '{1}' on input {2} is not a valid path or is null. Will be ignored."), FText::AsNumber(i), AttributeSelector.GetDisplayText(), FText::AsNumber(Index)));
-							}
-
-							continue;
-						}
-					}
-				}
-				else
-				{
-					PCGE_LOG(Error, GraphAndLog, FText::Format(LOCTEXT("InvalidAttribute", "Attribute/Property '{0}'({1}) is not convertible to a SoftObjectPath on input {2}"), AttributeSelector.GetDisplayText(), PCG::Private::GetTypeNameText(Accessor->GetUnderlyingType()), FText::AsNumber(Index)));
-					continue;
-				}
-			}
-		}
-
-		if (!ObjectsToLoad.IsEmpty())
-		{
-			bIsDone = AsyncLoadContext->RequestResourceLoad(ThisContext, std::move(ObjectsToLoad), !Settings->bSynchronousLoad);
-		}
-	}
-
-	return bIsDone;
+	FPCGLoadObjectsFromPathContext* ThisContext = static_cast<FPCGLoadObjectsFromPathContext*>(Context);
+	return ThisContext->InitializeAndRequestLoad(PCGPinConstants::DefaultInputLabel,
+		Settings->InputSource,
+		Settings->ObjectPathsToExtract,
+		Settings->bPersistAllData,
+		Settings->bSilenceErrorOnEmptyObjectPath,
+		Settings->bSynchronousLoad);
 }
 
 bool FPCGGetPropertyFromObjectPathElement::ExecuteInternal(FPCGContext* Context) const
@@ -263,7 +144,7 @@ bool FPCGGetPropertyFromObjectPathElement::ExecuteInternal(FPCGContext* Context)
 	TRACE_CPUPROFILER_EVENT_SCOPE(FPCGGetPropertyFromObjectPathElement::Execute);
 
 	check(Context);
-	FPCGGetPropertyFromObjectPathContext* ThisContext = static_cast<FPCGGetPropertyFromObjectPathContext*>(Context);
+	FPCGLoadObjectsFromPathContext* ThisContext = static_cast<FPCGLoadObjectsFromPathContext*>(Context);
 
 	const UPCGGetPropertyFromObjectPathSettings* Settings = Context->GetInputSettings<UPCGGetPropertyFromObjectPathSettings>();
 	check(Settings);
@@ -291,14 +172,14 @@ bool FPCGGetPropertyFromObjectPathElement::ExecuteInternal(FPCGContext* Context)
 	const bool IsDynamicallyTracking = (SourceComponent && Context->Node) ? Context->Node->IsInputPinConnected(PCGPinConstants::DefaultInputLabel) : false;
 	if (IsDynamicallyTracking)
 	{
-		DynamicTrackingHelper.EnableAndInitialize(ThisContext, ThisContext->PathsToObjectsToExtractAndIncomingDataIndex.Num());
+		DynamicTrackingHelper.EnableAndInitialize(ThisContext, ThisContext->PathsToObjectsAndDataIndex.Num());
 	}
 #endif // WITH_EDITOR
 
-	for (const TTuple<FSoftObjectPath, int32>& SoftPathAndIndex : ThisContext->PathsToObjectsToExtractAndIncomingDataIndex)
+	for (const TTuple<FSoftObjectPath, int32, int32>& SoftPathAndIndex : ThisContext->PathsToObjectsAndDataIndex)
 	{
-		const FSoftObjectPath& SoftPath = SoftPathAndIndex.Get<FSoftObjectPath>();
-		const int32 Index = SoftPathAndIndex.Get<int32>();
+		const FSoftObjectPath& SoftPath = SoftPathAndIndex.Get<0>();
+		const int32 Index = SoftPathAndIndex.Get<1>();
 
 		if (SoftPath.IsNull() && Settings->bPersistAllData)
 		{
diff --git a/Engine/Plugins/PCG/Source/PCG/Private/Elements/PCGLoadObjectsContext.cpp b/Engine/Plugins/PCG/Source/PCG/Private/Elements/PCGLoadObjectsContext.cpp
new file mode 100644
--- /dev/null
+++ b/Engine/Plugins/PCG/Source/PCG/Private/Elements/PCGLoadObjectsContext.cpp
@@ -0,0 +1,137 @@
+// Copyright Epic Games, Inc. All Rights Reserved.
+
+#include "Elements/PCGLoadObjectsContext.h"
+
+#include "PCGElement.h"
+#include "PCGModule.h"
+#include "Metadata/PCGMetadata.h"
+#include "Metadata/Accessors/IPCGAttributeAccessor.h"
+#include "Metadata/Accessors/PCGAttributeAccessorKeys.h"
+#include "Metadata/Accessors/PCGAttributeAccessorHelpers.h"
+
+#define LOCTEXT_NAMESPACE "PCGLoadObjectsFromPathContext"
+
+bool FPCGLoadObjectsFromPathContext::InitializeAndRequestLoad(
+	FName InputPinName,
+	const FPCGAttributePropertyInputSelector& InputAttributeSelector,
+	const TArray<FSoftObjectPath>& StaticObjectPaths,
+	bool bPersistAllData,
+	bool bSilenceErrorOnEmptyObjectPath,
+	bool bSynchronousLoad)
+{
+	if (!WasLoadRequested())
+	{
+		const bool InPinIsConnected = Node && Node->IsInputPinConnected(InputPinName);
+
+		// First gather all the soft objects paths to extract
+		TArray<FSoftObjectPath> ObjectsToLoad;
+		if (!InPinIsConnected)
+		{
+			PathsToObjectsAndDataIndex.Reserve(StaticObjectPaths.Num());
+			for (int32 PathIndex = 0; PathIndex < StaticObjectPaths.Num(); ++PathIndex)
+			{
+				const FSoftObjectPath& Path = StaticObjectPaths[PathIndex];
+				if (!Path.IsNull())
+				{
+					ObjectsToLoad.AddUnique(Path);
+				}
+
+				PathsToObjectsAndDataIndex.Emplace(Path, -1, PathIndex);
+			}
+		}
+		else
+		{
+			const TArray<FPCGTaggedData> Inputs = InputData.GetInputsByPin(InputPinName);
+			for (int32 Index = 0; Index < Inputs.Num(); ++Index)
+			{
+				const FPCGTaggedData& Input = Inputs[Index];
+
+				if (!Input.Data)
+				{
+					PCGE_LOG_C(Error, GraphAndLog, this, FText::Format(LOCTEXT("InvalidData", "Invalid data for input {0}"), FText::AsNumber(Index)));
+					continue;
+				}
+
+				const FPCGAttributePropertyInputSelector AttributeSelector = InputAttributeSelector.CopyAndFixLast(Input.Data);
+				const TUniquePtr<const IPCGAttributeAccessor> Accessor = PCGAttributeAccessorHelpers::CreateConstAccessor(Input.Data, AttributeSelector);
+				const TUniquePtr<const IPCGAttributeAccessorKeys> Keys = PCGAttributeAccessorHelpers::CreateConstKeys(Input.Data, AttributeSelector);
+
+				if (!Accessor.IsValid() || !Keys.IsValid())
+				{
+					if (bPersistAllData)
+					{
+						// Special case for empty data. We need this case if we ever chain this node multiple times. An empty param (with no attributes and no entries) will generate another empty param.
+						const UPCGMetadata* Metadata = Input.Data->ConstMetadata();
+						if (Metadata && Metadata->GetAttributeCount() == 0 && Metadata->GetLocalItemCount() == 0)
+						{
+							// Emplace empty path for this input. Will generate an empty param data.
+							PathsToObjectsAndDataIndex.Emplace(FSoftObjectPath(), Index, -1);
+						}
+					}
+
+					if (!bSilenceErrorOnEmptyObjectPath)
+					{
+						PCGE_LOG_C(Error, GraphAndLog, this, FText::Format(LOCTEXT("AttributeNotFound", "Attribute/Property '{0}' does not exist on input {1}"), AttributeSelector.GetDisplayText(), FText::AsNumber(Index)));
+					}
+
+					continue;
+				}
+
+				const int32 NumElementsToAdd = Keys->GetNum();
+				if (NumElementsToAdd == 0)
+				{
+					continue;
+				}
+
+				// Extract value as String to validate that a path is empty or ill-formed (because any ill-formed path will be null).
+				TArray<FString> InputValues;
+				InputValues.SetNum(NumElementsToAdd);
+				if (Accessor->GetRange(MakeArrayView(InputValues), 0, *Keys, EPCGAttributeAccessorFlags::AllowConstructible | EPCGAttributeAccessorFlags::AllowBroadcast))
+				{
+					PathsToObjectsAndDataIndex.Reserve(PathsToObjectsAndDataIndex.Num() + NumElementsToAdd);
+					for (int32 i = 0; i < InputValues.Num(); ++i)
+					{
+						FString& StringPath = InputValues[i];
+						// Empty SoftObjectPath can convert to string to None and is treated as empty, so check that one too.
+						const bool PathIsEmpty = StringPath.IsEmpty() || StringPath.Equals(TEXT("None"), ESearchCase::CaseSensitive);
+						if (PathIsEmpty && bPersistAllData)
+						{
+							PathsToObjectsAndDataIndex.Emplace(FSoftObjectPath(), Index, i);
+						}
+
+						FSoftObjectPath Path(std::move(StringPath));
+
+						if (!Path.IsNull())
+						{
+							ObjectsToLoad.AddUnique(Path);
+							PathsToObjectsAndDataIndex.Emplace(std::move(Path), Index, i);
+						}
+						else
+						{
+							if (!PathIsEmpty || !bSilenceErrorOnEmptyObjectPath)
+							{
+								PCGE_LOG_C(Error, GraphAndLog, this, FText::Format(LOCTEXT("InvalidPath", "Value number {0} for Attribute/Property '{1}' on input {2} is not a valid path or is null. Will be ignored."), FText::AsNumber(i), AttributeSelector.GetDisplayText(), FText::AsNumber(Index)));
+							}
+
+							continue;
+						}
+					}
+				}
+				else
+				{
+					PCGE_LOG_C(Error, GraphAndLog, this, FText::Format(LOCTEXT("InvalidAttribute", "Attribute/Property '{0}'({1}) is not convertible to a SoftObjectPath on input {2}"), AttributeSelector.GetDisplayText(), PCG::Private::GetTypeNameText(Accessor->GetUnderlyingType()), FText::AsNumber(Index)));
+					continue;
+				}
+			}
+		}
+
+		if (!ObjectsToLoad.IsEmpty())
+		{
+			return RequestResourceLoad(this, std::move(ObjectsToLoad), !bSynchronousLoad);
+		}
+	}
+
+	return true;
+}
+
+#undef LOCTEXT_NAMESPACE
\ No newline at end of file
diff --git a/Engine/Plugins/PCG/Source/PCG/Private/Helpers/PCGHelpers.cpp b/Engine/Plugins/PCG/Source/PCG/Private/Helpers/PCGHelpers.cpp
--- a/Engine/Plugins/PCG/Source/PCG/Private/Helpers/PCGHelpers.cpp
+++ b/Engine/Plugins/PCG/Source/PCG/Private/Helpers/PCGHelpers.cpp
@@ -525,11 +525,11 @@ namespace PCGHelpers
 #endif
 	}
 
-	TArray<UFunction*> FindUserFunctions(TSubclassOf<AActor> ActorClass, const TArray<FName>& FunctionNames, const TArray<const UFunction*>& FunctionPrototypes, const FPCGContext* InContext)
+	TArray<UFunction*> FindUserFunctions(TSubclassOf<UObject> ObjectClass, const TArray<FName>& FunctionNames, const TArray<const UFunction*>& FunctionPrototypes, const FPCGContext* InContext)
 	{
 		TArray<UFunction*> Functions;
 
-		if (!ActorClass)
+		if (!ObjectClass)
 		{
 			return Functions;
 		}
@@ -541,12 +541,14 @@ namespace PCGHelpers
 				continue;
 			}
 
-			if (UFunction* Function = ActorClass->FindFunctionByName(FunctionName))
+			if (UFunction* Function = ObjectClass->FindFunctionByName(FunctionName))
 			{
 #if WITH_EDITOR
-				if (!Function->GetBoolMetaData(TEXT("CallInEditor")))
+				// Implementation note: for AActors, using ProcessEvent requires the function to either be 'CallInEditor' or GAllowActorScriptExecutionInEditor to be true.
+				// It might not be strictly needed in cases where the object is not an actor.
+				if (ObjectClass->GetDefaultObject()->IsA<AActor>() && !Function->GetBoolMetaData(TEXT("CallInEditor")))
 				{
-					PCGLog::LogWarningOnGraph(FText::Format(LOCTEXT("CallInEditorFailed", "Function '{0}' in class '{1}' requires CallInEditor to be true while in-editor."), FText::FromName(FunctionName), FText::FromName(ActorClass->GetFName())), InContext);
+					PCGLog::LogWarningOnGraph(FText::Format(LOCTEXT("CallInEditorFailed", "Function '{0}' in class '{1}' requires CallInEditor to be true while in-editor."), FText::FromName(FunctionName), FText::FromName(ObjectClass->GetFName())), InContext);
 					continue;
 				}
 #endif
@@ -561,12 +563,12 @@ namespace PCGHelpers
 
 				if (Functions.IsEmpty() || Functions.Last() != Function)
 				{
-					PCGLog::LogWarningOnGraph(FText::Format(LOCTEXT("ParametersIncorrect", "Function '{0}' in class '{1}' has incorrect parameters."), FText::FromName(FunctionName), FText::FromName(ActorClass->GetFName())), InContext);
+					PCGLog::LogWarningOnGraph(FText::Format(LOCTEXT("ParametersIncorrect", "Function '{0}' in class '{1}' has incorrect parameters."), FText::FromName(FunctionName), FText::FromName(ObjectClass->GetFName())), InContext);
 				}
 			}
 			else
 			{
-				PCGLog::LogWarningOnGraph(FText::Format(LOCTEXT("FunctionNotFound", "Function '{0}' was not found in class '{1}'."), FText::FromName(FunctionName), FText::FromName(ActorClass->GetFName())), InContext);
+				PCGLog::LogWarningOnGraph(FText::Format(LOCTEXT("FunctionNotFound", "Function '{0}' was not found in class '{1}'."), FText::FromName(FunctionName), FText::FromName(ObjectClass->GetFName())), InContext);
 			}
 		}
 
diff --git a/Engine/Plugins/PCG/Source/PCG/Private/Metadata/PCGObjectPropertyOverride.cpp b/Engine/Plugins/PCG/Source/PCG/Private/Metadata/PCGObjectPropertyOverride.cpp
--- a/Engine/Plugins/PCG/Source/PCG/Private/Metadata/PCGObjectPropertyOverride.cpp
+++ b/Engine/Plugins/PCG/Source/PCG/Private/Metadata/PCGObjectPropertyOverride.cpp
@@ -21,39 +21,60 @@ namespace PCGObjectPropertyOverrideHelpers
 		return ObjectOverridePinProperties;
 	}
 
-	void ApplyOverridesFromParams(const TArray<FPCGObjectPropertyOverrideDescription>& InObjectPropertyOverrideDescriptions, AActor* TargetActor, FName OverridesPinLabel, FPCGContext* Context)
+	void ApplyOverrides(const TArray<FPCGObjectPropertyOverrideDescription>& InObjectPropertyOverrideDescriptions, const TArray<TPair<UObject*, int32>>& TargetObjectToOverrideIndices, FName OverridesPinLabel, int32 InInputDataIndex, FPCGContext* Context)
 	{
-		if (!Context)
+		if (!Context || TargetObjectToOverrideIndices.IsEmpty())
 		{
 			return;
 		}
 
 		const TArray<FPCGTaggedData> OverrideInputs = Context->InputData.GetInputsByPin(OverridesPinLabel);
 
+		int32 InputIndex = InInputDataIndex;
+
 		if (OverrideInputs.Num() == 0)
 		{
 			return;
 		}
-		else if (OverrideInputs.Num() > 1)
+		else if (OverrideInputs.Num() == 1)
 		{
-			PCGLog::LogWarningOnGraph(FText::Format(LOCTEXT("MoreThanOneData", "More than one data was found on pin '{0}'. Only using the first one."), FText::FromName(OverridesPinLabel)), Context);
+			InputIndex = 0;
+		}
+		else if (OverrideInputs.Num() > 1 && !OverrideInputs.IsValidIndex(InputIndex))
+		{
+			PCGLog::LogWarningOnGraph(FText::Format(LOCTEXT("InconsistentDataCount", "The data provided on pin '{0}' does not have a consistent size with the input index '{1}'. Will use the first one."), FText::FromName(OverridesPinLabel), FText::AsNumber(InInputDataIndex)), Context);
+			InputIndex = 0;
 		}
 
-		const UPCGParamData* ParamData = Cast<const UPCGParamData>(OverrideInputs[0].Data);
-
-		if (!ParamData)
+		const UPCGData* OverrideData = OverrideInputs[InputIndex].Data;
+		if (!OverrideData)
 		{
-			PCGLog::LogErrorOnGraph(LOCTEXT("InvalidActorOverrideData", "Invalid input data type for Actor Property Overrides pin, must be of type Param."), Context);
+			PCGLog::LogErrorOnGraph(LOCTEXT("InvalidActorOverrideData", "Invalid input data for Object Property Overrides pin."), Context);
 			return;
 		}
 
-		FPCGObjectOverrides ActorOverrides(TargetActor);
-		ActorOverrides.Initialize(InObjectPropertyOverrideDescriptions, TargetActor, ParamData, Context);
-		if (!ActorOverrides.Apply(/*InputKeyIndex=*/0)) // Use the First Entry of the param data for override (similar to what is done in Parameter Overrides in FPCGContext)
+		for (const TPair<UObject*, int32>& TargetObjectAndIndex : TargetObjectToOverrideIndices)
 		{
-			PCGLog::LogErrorOnGraph(FText::Format(LOCTEXT("ApplyOverrideFailed", "Failed to apply property overrides to actor '%s' from an attribute set."), FText::FromName(TargetActor->GetClass()->GetFName())), Context);
+			UObject* TargetObject = TargetObjectAndIndex.Key;
+			int32 InputKeyIndex = TargetObjectAndIndex.Value;
+
+			FPCGObjectOverrides ObjectOverrides(TargetObject);
+			ObjectOverrides.Initialize(InObjectPropertyOverrideDescriptions, TargetObject, OverrideData, Context);
+			if (!ObjectOverrides.Apply(InputKeyIndex)) // Use the First Entry of the param data for override (similar to what is done in Parameter Overrides in FPCGContext)
+			{
+				PCGLog::LogErrorOnGraph(FText::Format(LOCTEXT("ApplyOverrideFailed", "Failed to apply property overrides to object '%s'."), FText::FromName(TargetObject->GetClass()->GetFName())), Context);
+			}
 		}
 	}
+
+	void ApplyOverridesFromParams(const TArray<FPCGObjectPropertyOverrideDescription>& InObjectPropertyOverrideDescriptions, UObject* TargetObject, FName OverridesPinLabel, FPCGContext* Context)
+	{
+		ApplyOverrides(InObjectPropertyOverrideDescriptions,
+			{ TPair<UObject*, int32>(TargetObject, /*Entry index=*/0) },
+			OverridesPinLabel,
+			/*InInputDataIndex=*/0,
+			Context);
+	}
 }
 
 void FPCGObjectSingleOverride::Initialize(const FPCGAttributePropertySelector& InputSelector, const FString& OutputProperty, const UStruct* TemplateClass, const UPCGData* SourceData, FPCGContext* Context)
diff --git a/Engine/Plugins/PCG/Source/PCG/Private/PCGSettings.cpp b/Engine/Plugins/PCG/Source/PCG/Private/PCGSettings.cpp
--- a/Engine/Plugins/PCG/Source/PCG/Private/PCGSettings.cpp
+++ b/Engine/Plugins/PCG/Source/PCG/Private/PCGSettings.cpp
@@ -19,6 +19,7 @@
 
 #include "Serialization/ArchiveObjectCrc32.h"
 #include "UObject/ObjectSaveContext.h"
+#include "UObject/FortniteMainBranchObjectVersion.h"
 
 #include UE_INLINE_GENERATED_CPP_BY_NAME(PCGSettings)
 
@@ -319,6 +320,7 @@ void UPCGSettings::Serialize(FArchive& Ar)
 	Super::Serialize(Ar);
 
 	Ar.UsingCustomVersion(FPCGCustomVersion::GUID);
+	Ar.UsingCustomVersion(FFortniteMainBranchObjectVersion::GUID);
 
 #if WITH_EDITOR
 	if (Ar.IsLoading())
diff --git a/Engine/Plugins/PCG/Source/PCG/Public/Elements/PCGApplyOnActor.h b/Engine/Plugins/PCG/Source/PCG/Public/Elements/PCGApplyOnActor.h
--- a/Engine/Plugins/PCG/Source/PCG/Public/Elements/PCGApplyOnActor.h
+++ b/Engine/Plugins/PCG/Source/PCG/Public/Elements/PCGApplyOnActor.h
@@ -3,6 +3,7 @@
 #pragma once
 
 #include "PCGSettings.h"
+#include "Elements/PCGLoadObjectsContext.h"
 #include "Metadata/PCGObjectPropertyOverride.h"
 
 #include "PCGApplyOnActor.generated.h"
@@ -21,20 +22,30 @@ public:
 	//~Begin UPCGSettings interface
 #if WITH_EDITOR
 	virtual FName GetDefaultNodeName() const override { return FName(TEXT("ApplyOnActor")); }
-	virtual FText GetDefaultNodeTitle() const override { return NSLOCTEXT("PCGApplyOnActorElement", "NodeTitle", "Apply On Actor"); }
-	virtual FText GetNodeTooltipText() const override { return NSLOCTEXT("PCGApplyOnActorElement", "NodeTooltip", "Applies property overrides and executes functions on a target actor."); }
+	virtual FText GetDefaultNodeTitle() const override { return NSLOCTEXT("PCGApplyOnActorElement", "NodeTitle", "Apply On Object"); }
+	virtual FText GetNodeTooltipText() const override { return NSLOCTEXT("PCGApplyOnActorElement", "NodeTooltip", "Applies property overrides and executes functions on a target object."); }
 	virtual EPCGSettingsType GetType() const override { return EPCGSettingsType::Generic; }
+	virtual bool HasDynamicPins() const override { return true; }
+	virtual void ApplyDeprecationBeforeUpdatePins(UPCGNode* InOutNode, TArray<TObjectPtr<UPCGPin>>& InputPins, TArray<TObjectPtr<UPCGPin>>& OutputPins) override;
 #endif
 
 protected:
 	virtual TArray<FPCGPinProperties> InputPinProperties() const override;
-	virtual TArray<FPCGPinProperties> OutputPinProperties() const override { return TArray<FPCGPinProperties>(); }
+	virtual TArray<FPCGPinProperties> OutputPinProperties() const override;
 	virtual FPCGElementPtr CreateElement() const override;
 	//~End UPCGSettings interface
 
 public:
-	UPROPERTY(BlueprintReadWrite, Category = Settings, meta = (PCG_Overridable))
+	/** If something is connected in the In pin, will look for this attribute values to load, representing the object reference. */
+	UPROPERTY(BlueprintReadWrite, EditAnywhere, Category = "Settings", meta = (PCG_Overridable, PCG_DiscardPropertySelection, PCG_DiscardExtraSelection))
+	FPCGAttributePropertyInputSelector ObjectReferenceAttribute;
+
+#if WITH_EDITORONLY_DATA
+	/** Provides a single target actor to apply on, used only if the 'In' pin is not connected. */
+	UE_DEPRECATED(5.5, "TargetActor has been deprecated; pass data directly to the 'In' pin instead.")
+	UPROPERTY(BlueprintReadWrite, Category = Settings, meta = (DeprecatedProperty, DeprecationMessage = "TargetActor has been deprecated; pass data directly to the 'In' pin instead."))
 	TSoftObjectPtr<AActor> TargetActor;
+#endif
 
 	/** Override the default property values on the target actor. Applied before post-process functions. */
 	UPROPERTY(EditAnywhere, BlueprintReadWrite, Category = Settings)
@@ -43,9 +54,17 @@ public:
 	/** Specify a list of functions to be called on the target actor. Functions need to be parameter-less and with "CallInEditor" flag enabled. */
 	UPROPERTY(BlueprintReadWrite, EditAnywhere, Category = Settings)
 	TArray<FName> PostProcessFunctionNames;
+
+	/** Opt-in option to silence errors when the path is Empty or nothing to extract. */
+	UPROPERTY(BlueprintReadWrite, EditAnywhere, Category = "Settings|Advanced")
+	bool bSilenceErrorOnEmptyObjectPath = false;
+
+	/** By default, object loading is asynchronous, can force it synchronous if needed. */
+	UPROPERTY(BlueprintReadWrite, EditAnywhere, Category = "Settings|Debug")
+	bool bSynchronousLoad = false;
 };
 
-class FPCGApplyOnActorElement : public IPCGElement
+class FPCGApplyOnObjectElement : public IPCGElementWithCustomContext<FPCGLoadObjectsFromPathContext>
 {
 public:
 	// Calling a function on an actor might not be threadsafe, so taking the safe approach.
@@ -53,5 +72,6 @@ public:
 	virtual bool IsCacheable(const UPCGSettings* InSettings) const override { return false; }
 
 protected:
+	virtual bool PrepareDataInternal(FPCGContext* Context) const override;
 	virtual bool ExecuteInternal(FPCGContext* Context) const override;
 };
diff --git a/Engine/Plugins/PCG/Source/PCG/Public/Elements/PCGGetPropertyFromObjectPath.h b/Engine/Plugins/PCG/Source/PCG/Public/Elements/PCGGetPropertyFromObjectPath.h
--- a/Engine/Plugins/PCG/Source/PCG/Public/Elements/PCGGetPropertyFromObjectPath.h
+++ b/Engine/Plugins/PCG/Source/PCG/Public/Elements/PCGGetPropertyFromObjectPath.h
@@ -2,9 +2,8 @@
 
 #pragma once
 
-#include "PCGContext.h"
 #include "PCGSettings.h"
-#include "Async/PCGAsyncLoadingContext.h"
+#include "Elements/PCGLoadObjectsContext.h"
 
 #include "PCGGetPropertyFromObjectPath.generated.h"
 
@@ -73,12 +72,7 @@ public:
 	bool bSilenceErrorOnEmptyObjectPath = false;
 };
 
-struct FPCGGetPropertyFromObjectPathContext : public FPCGContext, public IPCGAsyncLoadingContext
-{
-	TArray<TTuple<FSoftObjectPath, int32>> PathsToObjectsToExtractAndIncomingDataIndex;
-};
-
-class FPCGGetPropertyFromObjectPathElement : public IPCGElement
+class FPCGGetPropertyFromObjectPathElement : public IPCGElementWithCustomContext<FPCGLoadObjectsFromPathContext>
 {
 public:
 	// Loading needs to be done on the main thread and accessing objects outside of PCG might not be thread safe, so taking the safe approach
@@ -90,7 +84,6 @@ public:
 	virtual bool IsCacheable(const UPCGSettings* InSettings) const override { return false; }
 
 protected:
-	virtual FPCGContext* CreateContext() override;
 	virtual bool PrepareDataInternal(FPCGContext* Context) const override;
 	virtual bool ExecuteInternal(FPCGContext* Context) const override;
 };
diff --git a/Engine/Plugins/PCG/Source/PCG/Public/Elements/PCGLoadObjectsContext.h b/Engine/Plugins/PCG/Source/PCG/Public/Elements/PCGLoadObjectsContext.h
new file mode 100644
--- /dev/null
+++ b/Engine/Plugins/PCG/Source/PCG/Public/Elements/PCGLoadObjectsContext.h
@@ -0,0 +1,21 @@
+// Copyright Epic Games, Inc. All Rights Reserved.
+
+#pragma once
+
+#include "PCGContext.h"
+#include "Async/PCGAsyncLoadingContext.h"
+#include "Metadata/PCGAttributePropertySelector.h"
+
+struct FPCGLoadObjectsFromPathContext : public FPCGContext, public IPCGAsyncLoadingContext
+{
+	bool InitializeAndRequestLoad(
+		FName InputPinName,
+		const FPCGAttributePropertyInputSelector& InputAttributeSelector,
+		const TArray<FSoftObjectPath>& StaticObjectPaths,
+		bool bPersistAllData,
+		bool bSilenceErrorOnEmptyObjectPath,
+		bool bSynchronousLoad
+	);
+
+	TArray<TTuple<FSoftObjectPath, int32, int32>> PathsToObjectsAndDataIndex;
+};
\ No newline at end of file
diff --git a/Engine/Plugins/PCG/Source/PCG/Public/Helpers/PCGHelpers.h b/Engine/Plugins/PCG/Source/PCG/Public/Helpers/PCGHelpers.h
--- a/Engine/Plugins/PCG/Source/PCG/Public/Helpers/PCGHelpers.h
+++ b/Engine/Plugins/PCG/Source/PCG/Public/Helpers/PCGHelpers.h
@@ -84,7 +84,7 @@ namespace PCGHelpers
 	 * Finds functions on the actor matching the provided function names. Functions must be marked as CallInEditor
 	 * and have parameters matching one of the provided prototypes. Some prototypes are provided in UPCGFunctionPrototypes.
 	 */
-	PCG_API TArray<UFunction*> FindUserFunctions(TSubclassOf<AActor> ActorClass, const TArray<FName>& FunctionNames, const TArray<const UFunction*>& FunctionPrototypes, const FPCGContext* InContext = nullptr);
+	PCG_API TArray<UFunction*> FindUserFunctions(TSubclassOf<UObject> ActorClass, const TArray<FName>& FunctionNames, const TArray<const UFunction*>& FunctionPrototypes, const FPCGContext* InContext = nullptr);
 
 	PCG_API TFunction<float(float, float)> GetDensityMergeFunction(EPCGDensityMergeOperation InOperation);
 };
diff --git a/Engine/Plugins/PCG/Source/PCG/Public/Metadata/PCGObjectPropertyOverride.h b/Engine/Plugins/PCG/Source/PCG/Public/Metadata/PCGObjectPropertyOverride.h
--- a/Engine/Plugins/PCG/Source/PCG/Public/Metadata/PCGObjectPropertyOverride.h
+++ b/Engine/Plugins/PCG/Source/PCG/Public/Metadata/PCGObjectPropertyOverride.h
@@ -51,7 +51,9 @@ namespace PCGObjectPropertyOverrideHelpers
 	FPCGPinProperties CreateObjectPropertiesOverridePin(FName Label, const FText& Tooltip);
 
 	/** Apply property overrides to the TargetObject directly from the ObjectPropertiesOverride pin. Use CreateObjectPropertiesOverridePin(). */
-	void ApplyOverridesFromParams(const TArray<FPCGObjectPropertyOverrideDescription>& InObjectPropertyOverrideDescriptions, AActor* TargetActor, FName OverridesPinLabel, FPCGContext* Context);
+	void ApplyOverridesFromParams(const TArray<FPCGObjectPropertyOverrideDescription>& InObjectPropertyOverrideDescriptions, UObject* TargetObject, FName OverridesPinLabel, FPCGContext* Context);
+
+	void ApplyOverrides(const TArray<FPCGObjectPropertyOverrideDescription>& InObjectPropertyOverrideDescriptions, const TArray<TPair<UObject*, int32>>& TargetObjectAndIndex, FName OverridesPinLabel, int32 InputDataIndex, FPCGContext* Context);
 }
 
 /**
diff --git a/Engine/Source/Runtime/Core/Public/UObject/FortniteMainBranchObjectVersion.h b/Engine/Source/Runtime/Core/Public/UObject/FortniteMainBranchObjectVersion.h
--- a/Engine/Source/Runtime/Core/Public/UObject/FortniteMainBranchObjectVersion.h
+++ b/Engine/Source/Runtime/Core/Public/UObject/FortniteMainBranchObjectVersion.h
@@ -503,6 +503,9 @@ struct FFortniteMainBranchObjectVersion
 		
 		// Changed the precision for the stored rotation on kinematic targets to match the precision used in particles 
 		ChaosStoreKinematicTargetRotationAsSinglePrecision,
+
+		// PCG changes around the ApplyOnActor node, where we collapsed the TargetActor to the input pin.
+		PCGApplyOnActorNodeMoveTargetActorEdgeToInput,
 		
 		// -----<new versions can be added above this line>-------------------------------------------------
 		VersionPlusOne,
