Add option to only get actor reference in the Get Actor Data node.

diff --git a/Engine/Plugins/PCG/Source/PCG/Private/Elements/PCGDataFromActor.cpp b/Engine/Plugins/PCG/Source/PCG/Private/Elements/PCGDataFromActor.cpp
--- a/Engine/Plugins/PCG/Source/PCG/Private/Elements/PCGDataFromActor.cpp
+++ b/Engine/Plugins/PCG/Source/PCG/Private/Elements/PCGDataFromActor.cpp
@@ -5,11 +5,13 @@
 #include "PCGActorAndComponentMapping.h"
 #include "PCGComponent.h"
 #include "PCGCustomVersion.h"
+#include "PCGParamData.h"
 #include "PCGSubsystem.h"
 #include "Data/PCGPointData.h"
 #include "Data/PCGSpatialData.h"
 #include "Data/PCGVolumeData.h"
 #include "Elements/PCGMergeElement.h"
+#include "Elements/PCGMergeAttributes.h"
 #include "Grid/PCGPartitionActor.h"
 #include "Helpers/PCGHelpers.h"
 #include "Utils/PCGGraphExecutionLogging.h"
@@ -222,7 +224,7 @@ EPCGDataType UPCGDataFromActorSettings::GetCurrentPinTypes(const UPCGPin* InPin)
 		{
 			return EPCGDataType::Point;
 		}
-		else if (Mode == EPCGGetDataFromActorMode::GetDataFromProperty)
+		else if (Mode == EPCGGetDataFromActorMode::GetActorReference)
 		{
 			return EPCGDataType::Param;
 		}
@@ -518,9 +520,9 @@ void FPCGDataFromActorElement::ProcessActors(FPCGContext* Context, const UPCGDat
 {
 	// Special case:
 	// If we're asking for single point with the merge single point data, we can do a more efficient process
-	if (Settings->Mode == EPCGGetDataFromActorMode::GetSinglePoint && Settings->bMergeSinglePointData && FoundActors.Num() > 1)
+	if ((Settings->Mode == EPCGGetDataFromActorMode::GetSinglePoint || Settings->Mode == EPCGGetDataFromActorMode::GetActorReference) && Settings->bMergeSinglePointData && FoundActors.Num() > 1)
 	{
-		MergeActorsIntoPointData(Context, Settings, FoundActors);
+		MergeActorsIntoData(Context, Settings, FoundActors);
 	}
 	else
 	{
@@ -531,17 +533,35 @@ void FPCGDataFromActorElement::ProcessActors(FPCGContext* Context, const UPCGDat
 	}
 }
 
-void FPCGDataFromActorElement::MergeActorsIntoPointData(FPCGContext* Context, const UPCGDataFromActorSettings* Settings, const TArray<AActor*>& FoundActors) const
+void FPCGDataFromActorElement::MergeActorsIntoData(FPCGContext* Context, const UPCGDataFromActorSettings* Settings, const TArray<AActor*>& FoundActors) const
 {
 	check(Context);
+	check(Settings && (Settings->Mode == EPCGGetDataFromActorMode::GetSinglePoint || Settings->Mode == EPCGGetDataFromActorMode::GetActorReference));
 
 	// At this point in time, the partition actors behave slightly differently, so if we are in the case where
 	// we have one or more partition actors, we'll go through the normal process and do post-processing to merge the point data instead.
 	const bool bContainsPartitionActors = Algo::AnyOf(FoundActors, [](const AActor* Actor) { return Cast<APCGPartitionActor>(Actor) != nullptr; });
+	const bool bCreatePointData = (Settings->Mode == EPCGGetDataFromActorMode::GetSinglePoint);
 
 	if (!bContainsPartitionActors)
 	{
-		UPCGPointData* Data = FPCGContext::NewObject_AnyThread<UPCGPointData>(Context);
+		UPCGData* Data = nullptr;
+		UPCGPointData* PointData = nullptr;
+		UPCGParamData* ParamData = nullptr;
+		FPCGMetadataAttribute<FSoftObjectPath>* ActorReference = nullptr;
+
+		if (bCreatePointData)
+		{
+			PointData = FPCGContext::NewObject_AnyThread<UPCGPointData>(Context);
+			Data = PointData;
+		}
+		else
+		{
+			ParamData = FPCGContext::NewObject_AnyThread<UPCGParamData>(Context);
+			ActorReference = ParamData->MutableMetadata()->FindOrCreateAttribute(PCGPointDataConstants::ActorReferenceAttribute, FSoftObjectPath(), /*bAllowsInterpolation=*/false, /*bOverrideParent=*/false, /*bOverwriteIfTypeMismatch=*/true);
+			Data = ParamData;
+		}
+
 		bool bHasData = false;
 		bool bAnyAttributeNameWasSanitized = false;
 
@@ -550,7 +570,14 @@ void FPCGDataFromActorElement::MergeActorsIntoPointData(FPCGContext* Context, co
 			if (Actor)
 			{
 				bool bAttributeNameWasSanitized = false;
-				Data->AddSinglePointFromActor(Actor, &bAttributeNameWasSanitized);
+				if (PointData)
+				{
+					PointData->AddSinglePointFromActor(Actor, &bAttributeNameWasSanitized);
+				}
+				else
+				{
+					ActorReference->SetValue(ParamData->MutableMetadata()->AddEntry(), FSoftObjectPath(Actor));
+				}
 
 				bAnyAttributeNameWasSanitized |= bAttributeNameWasSanitized;
 
@@ -580,7 +607,8 @@ void FPCGDataFromActorElement::MergeActorsIntoPointData(FPCGContext* Context, co
 			if (Actor)
 			{
 				bool bAttributeNameWasSanitized = false;
-				FPCGDataCollection Collection = UPCGComponent::CreateActorPCGDataCollection(Actor, Context->SourceComponent.Get(), EPCGDataType::Any, bParseActor, &bAttributeNameWasSanitized);
+				const EPCGDataType TargetDataType = (bCreatePointData ? EPCGDataType::Point : EPCGDataType::Param);
+				FPCGDataCollection Collection = UPCGComponent::CreateActorPCGDataCollection(Actor, Context->SourceComponent.Get(), TargetDataType, bParseActor, &bAttributeNameWasSanitized);
 
 				bAnyAttributeNameWasSanitized |= bAttributeNameWasSanitized;
 
@@ -596,13 +624,14 @@ void FPCGDataFromActorElement::MergeActorsIntoPointData(FPCGContext* Context, co
 		// Perform point data-to-point data merge
 		if (DataToMerge.TaggedData.Num() > 1)
 		{
-			UPCGMergeSettings* MergeSettings = FPCGContext::NewObject_AnyThread<UPCGMergeSettings>(Context);
 			FPCGMergeElement MergeElement;
-			FPCGContext* MergeContext = MergeElement.Initialize(DataToMerge, Context->SourceComponent, nullptr);
+			FPCGMergeAttributesElement MergeAttributesElement;
+			IPCGElement* Element = (bCreatePointData ? static_cast<IPCGElement*>(&MergeElement) : static_cast<IPCGElement*>(&MergeAttributesElement));
+
+			FPCGContext* MergeContext = Element->Initialize(DataToMerge, Context->SourceComponent, nullptr);
 			MergeContext->AsyncState.NumAvailableTasks = Context->AsyncState.NumAvailableTasks;
-			MergeContext->InputData.TaggedData.Emplace_GetRef().Data = MergeSettings;
 
-			while (!MergeElement.Execute(MergeContext))
+			while (!Element->Execute(MergeContext))
 			{
 			}
 
@@ -734,9 +763,10 @@ void FPCGDataFromActorElement::ProcessActor(FPCGContext* Context, const UPCGData
 	}
 	else
 	{
-		const bool bParseActor = (Settings->Mode != EPCGGetDataFromActorMode::GetSinglePoint);
+		const bool bParseActor = (Settings->Mode != EPCGGetDataFromActorMode::GetSinglePoint && Settings->Mode != EPCGGetDataFromActorMode::GetActorReference);
+		EPCGDataType DataFilter = (Settings->Mode == EPCGGetDataFromActorMode::GetActorReference ? EPCGDataType::Param : Settings->GetDataFilter());
 		bool bAttributeNameWasSanitized = false;
-		FPCGDataCollection Collection = UPCGComponent::CreateActorPCGDataCollection(FoundActor, SourceComponent, Settings->GetDataFilter(), bParseActor, &bAttributeNameWasSanitized);
+		FPCGDataCollection Collection = UPCGComponent::CreateActorPCGDataCollection(FoundActor, SourceComponent, DataFilter, bParseActor, &bAttributeNameWasSanitized);
 
 		if (bAttributeNameWasSanitized && !Settings->bSilenceSanitizedAttributeNameWarnings)
 		{
diff --git a/Engine/Plugins/PCG/Source/PCG/Private/Elements/PCGMergeAttributes.cpp b/Engine/Plugins/PCG/Source/PCG/Private/Elements/PCGMergeAttributes.cpp
--- a/Engine/Plugins/PCG/Source/PCG/Private/Elements/PCGMergeAttributes.cpp
+++ b/Engine/Plugins/PCG/Source/PCG/Private/Elements/PCGMergeAttributes.cpp
@@ -48,9 +48,6 @@ bool FPCGMergeAttributesElement::ExecuteInternal(FPCGContext* Context) const
 	TRACE_CPUPROFILER_EVENT_SCOPE(FPCGMergeAttributesElement::Execute);
 	check(Context);
 
-	const UPCGMergeAttributesSettings* Settings = Context->GetInputSettings<UPCGMergeAttributesSettings>();
-	check(Settings);
-
 	TArray<FPCGTaggedData>& Outputs = Context->OutputData.TaggedData;
 
 	FPCGTaggedData* MergedOutput = nullptr;
diff --git a/Engine/Plugins/PCG/Source/PCG/Private/Elements/PCGMergeElement.cpp b/Engine/Plugins/PCG/Source/PCG/Private/Elements/PCGMergeElement.cpp
--- a/Engine/Plugins/PCG/Source/PCG/Private/Elements/PCGMergeElement.cpp
+++ b/Engine/Plugins/PCG/Source/PCG/Private/Elements/PCGMergeElement.cpp
@@ -41,9 +41,7 @@ bool FPCGMergeElement::ExecuteInternal(FPCGContext* Context) const
 	check(Context);
 
 	const UPCGMergeSettings* Settings = Context->GetInputSettings<UPCGMergeSettings>();
-	check(Settings);
-
-	const bool bMergeMetadata = Settings->bMergeMetadata;
+	const bool bMergeMetadata = !Settings || Settings->bMergeMetadata;
 
 	TArray<FPCGTaggedData> Sources = Context->InputData.GetInputsByPin(PCGPinConstants::DefaultInputLabel);
 	TArray<FPCGTaggedData>& Outputs = Context->OutputData.TaggedData;
diff --git a/Engine/Plugins/PCG/Source/PCG/Private/PCGComponent.cpp b/Engine/Plugins/PCG/Source/PCG/Private/PCGComponent.cpp
--- a/Engine/Plugins/PCG/Source/PCG/Private/PCGComponent.cpp
+++ b/Engine/Plugins/PCG/Source/PCG/Private/PCGComponent.cpp
@@ -2832,6 +2832,20 @@ FPCGDataCollection UPCGComponent::CreateActorPCGDataCollection(AActor* Actor, co
 		TaggedData.Data = Data;
 		TaggedData.Tags = ActorTags;
 	}
+	else if (!bParseActor && !!(InDataFilter & EPCGDataType::Param))
+	{
+		UPCGParamData* Data = NewObject<UPCGParamData>();
+		check(Data && Data->MutableMetadata());
+		
+		if (FPCGMetadataAttribute<FSoftObjectPath>* ActorReferenceAttribute = Data->MutableMetadata()->FindOrCreateAttribute(PCGPointDataConstants::ActorReferenceAttribute, FSoftObjectPath(), /*bAllowsInterpolation=*/false, /*bOverrideParent=*/false, /*bOverwriteIfTypeMismatch=*/true))
+		{
+			ActorReferenceAttribute->SetValue(Data->MutableMetadata()->AddEntry(), FSoftObjectPath(Actor));
+		}
+
+		FPCGTaggedData& TaggedData = Collection.TaggedData.Emplace_GetRef();
+		TaggedData.Data = Data;
+		TaggedData.Tags = ActorTags;
+	}
 	else if (PartitionActor)
 	{
 		check(!Component || Component->GetOwner() == Actor); // Invalid processing otherwise because of the this usage
diff --git a/Engine/Plugins/PCG/Source/PCG/Public/Elements/PCGDataFromActor.h b/Engine/Plugins/PCG/Source/PCG/Public/Elements/PCGDataFromActor.h
--- a/Engine/Plugins/PCG/Source/PCG/Public/Elements/PCGDataFromActor.h
+++ b/Engine/Plugins/PCG/Source/PCG/Public/Elements/PCGDataFromActor.h
@@ -15,9 +15,10 @@ enum class EPCGGetDataFromActorMode : uint8
 {
 	ParseActorComponents UMETA(Tooltip = "Parse the found actor(s) for relevant components such as Primitives, Splines, and Volumes."),
 	GetSinglePoint UMETA(Tooltip = "Produces a single point per actor with the actor transform and bounds."),
-	GetDataFromProperty UMETA(Tooltip = "Gets a data collection from an actor property."),
+	GetDataFromProperty UMETA(DisplayName = "Get PCG Data From Property", Tooltip = "Gets a data collection from an actor property."),
 	GetDataFromPCGComponent UMETA(Tooltip = "Copy generated output from other PCG components on the found actor(s)."),
-	GetDataFromPCGComponentOrParseComponents UMETA(Tooltip = "Attempts to copy generated output from other PCG components on the found actor(s), otherwise, falls back to parsing actor components.")
+	GetDataFromPCGComponentOrParseComponents UMETA(Tooltip = "Attempts to copy generated output from other PCG components on the found actor(s), otherwise, falls back to parsing actor components."),
+	GetActorReference UMETA(Tooltip = "Produces one entry per actor with only the actor reference.")
 };
 
 /** Builds a collection of PCG-compatible data from the selected actors. */
@@ -98,7 +99,7 @@ public:
 	int32 AllowedGrids = int32(EPCGHiGenGrid::Uninitialized);
 
 	/** Merges all the single point data outputs into a single point data. */
-	UPROPERTY(BlueprintReadWrite, EditAnywhere, Category = "Data Retrieval Settings", meta = (EditCondition = "Mode == EPCGGetDataFromActorMode::GetSinglePoint", EditConditionHides))
+	UPROPERTY(BlueprintReadWrite, EditAnywhere, Category = "Data Retrieval Settings", meta = (DisplayName = "Merge Simple Data", EditCondition = "Mode == EPCGGetDataFromActorMode::GetSinglePoint || Mode == EPCGGetDataFromActorMode::GetActorReference", EditConditionHides))
 	bool bMergeSinglePointData = false;
 
 	/** Provide pin names to match against the found component output pins. Data will automatically be wired to the expected pin if the name comparison succeeds. All unmatched pins will go into the standard out pin. */
@@ -149,5 +150,5 @@ protected:
 	virtual void ProcessActors(FPCGContext* Context, const UPCGDataFromActorSettings* Settings, const TArray<AActor*>& FoundActors) const;
 	virtual void ProcessActor(FPCGContext* Context, const UPCGDataFromActorSettings* Settings, AActor* FoundActor) const;
 
-	void MergeActorsIntoPointData(FPCGContext* Context, const UPCGDataFromActorSettings* Settings, const TArray<AActor*>& FoundActors) const;
+	void MergeActorsIntoData(FPCGContext* Context, const UPCGDataFromActorSettings* Settings, const TArray<AActor*>& FoundActors) const;
 };
