[PCG] Metadata operation update

* Operations that has no entries but a default value will be executed (Like metadata on spatial data)
* Spline Attribute List View now show metadata

#rnx

diff --git a/Engine/Plugins/PCG/Source/PCG/Private/Elements/Metadata/PCGMetadataOpElementBase.cpp b/Engine/Plugins/PCG/Source/PCG/Private/Elements/Metadata/PCGMetadataOpElementBase.cpp
--- a/Engine/Plugins/PCG/Source/PCG/Private/Elements/Metadata/PCGMetadataOpElementBase.cpp
+++ b/Engine/Plugins/PCG/Source/PCG/Private/Elements/Metadata/PCGMetadataOpElementBase.cpp
@@ -540,9 +540,9 @@ bool FPCGMetadataElementBase::PrepareDataInternal(FPCGContext* Context) const
 			return EPCGTimeSliceInitResult::AbortExecution;
 		}
 
-		// Update the number of elements to process
+		// Update the number of elements to process, it's OK to be 0 if it is an attribute, as we can do a default value operation.
 		OperationData.NumberOfElementsToProcess = OperationData.InputKeys[PrimaryPinIndex]->GetNum();
-		if (OperationData.NumberOfElementsToProcess == 0)
+		if (OperationData.NumberOfElementsToProcess == 0 && !OperationData.InputAccessors[PrimaryPinIndex]->IsAttribute())
 		{
 			PCGE_LOG(Verbose, LogOnly, FText::Format(LOCTEXT("NoElementsInForwardedInput", "No elements in data from forwarded pin '{0}'."), FText::FromName(PrimaryPinData.Pin)));
 			return NoOperation();
@@ -569,15 +569,15 @@ bool FPCGMetadataElementBase::PrepareDataInternal(FPCGContext* Context) const
 
 				const int32 ElementNum = OperationData.InputKeys[Index]->GetNum();
 
-				// No elements on secondary pin, early out for no operation
-				if (ElementNum == 0)
+				// No elements on secondary pin, early out for no operation, only if it is not an attribute, as we could still do a default value operation
+				if (ElementNum == 0 && !OperationData.InputAccessors[Index]->IsAttribute())
 				{
 					PCGE_LOG(Verbose, LogOnly, FText::Format(LOCTEXT("NoElementsInInput", "No elements in data from secondary pin '{0}'."), FText::FromName(PrimaryPinData.Pin)));
 					return NoOperation();
 				}
 
 				// Verify that the number of elements makes sense
-				if (OperationData.NumberOfElementsToProcess % ElementNum != 0)
+				if (ElementNum != 0 && OperationData.NumberOfElementsToProcess % ElementNum != 0)
 				{
 					PCGE_LOG(Error, GraphAndLog, FText::Format(LOCTEXT("MismatchInNumberOfElements", "Mismatch between the number of elements from pin '{0}' ({1}) and from pin '{2}' ({3})."), FText::FromName(PrimaryPinData.Pin), OperationData.NumberOfElementsToProcess, FText::FromName(InputTaggedData[Index].Pin), ElementNum));
 					return EPCGTimeSliceInitResult::AbortExecution;
@@ -627,6 +627,13 @@ bool FPCGMetadataElementBase::PrepareDataInternal(FPCGContext* Context) const
 						PCG::Private::GetTypeNameText<AttributeType>()));
 					return false;
 				}
+
+				// We have no element to process but we try to write into a property, early out.
+				if (OperationData.NumberOfElementsToProcess == 0 && !TempConstAccessor->IsAttribute())
+				{
+					PCGLog::LogErrorOnGraph(FText::Format(LOCTEXT("NoDefaultValue", "Operation is done on the default value, but output attribute '{0}' does not support default values"), OutputTarget.GetDisplayText()), Context);
+					return false;
+				}
 			}
 
 			check(InputTaggedData[PrimaryPinIndex].Data);
diff --git a/Engine/Plugins/PCG/Source/PCG/Public/Elements/Metadata/PCGMetadataOpElementBase.h b/Engine/Plugins/PCG/Source/PCG/Public/Elements/Metadata/PCGMetadataOpElementBase.h
--- a/Engine/Plugins/PCG/Source/PCG/Public/Elements/Metadata/PCGMetadataOpElementBase.h
+++ b/Engine/Plugins/PCG/Source/PCG/Public/Elements/Metadata/PCGMetadataOpElementBase.h
@@ -422,12 +422,6 @@ void PCGMetadataOps::FOperationData::Validate()
 template <typename... InputTypes, typename... Callbacks>
 inline bool FPCGMetadataElementBase::DoNAryOp(PCGMetadataOps::FOperationData& InOperationData, TTuple<Callbacks...>&& InCallbacks) const
 {
-	// If nothing to do, exit immediately
-	if (InOperationData.NumberOfElementsToProcess == 0)
-	{
-		return true;
-	}
-
 	// Validate that all is good
 	constexpr uint32 NbInputs = (uint32)sizeof...(InputTypes);
 	constexpr uint32 NbOutputs = (uint32)sizeof...(Callbacks);
@@ -446,6 +440,12 @@ inline bool FPCGMetadataElementBase::DoNAryOp(PCGMetadataOps::FOperationData& In
 		PCG::Private::NAryOperation::Operation<InputTypes...>(InOperationData, /*StartIndex=*/0, /*Range=*/1, Options, InCallbacks);
 	}
 
+	// If nothing to do now, we can early out.
+	if (InOperationData.NumberOfElementsToProcess == 0)
+	{
+		return true;
+	}
+
 	// Then iterate over all the values
 	Options.SetFlags = Flags;
 	Options.bUseDefaultKey = false;
diff --git a/Engine/Plugins/PCG/Source/PCG/Public/Metadata/Accessors/IPCGAttributeAccessor.h b/Engine/Plugins/PCG/Source/PCG/Public/Metadata/Accessors/IPCGAttributeAccessor.h
--- a/Engine/Plugins/PCG/Source/PCG/Public/Metadata/Accessors/IPCGAttributeAccessor.h
+++ b/Engine/Plugins/PCG/Source/PCG/Public/Metadata/Accessors/IPCGAttributeAccessor.h
@@ -136,6 +136,9 @@ public:
 	int16 GetUnderlyingType() const { return UnderlyingType; }
 	bool IsReadOnly() const { return bReadOnly; }
 
+	// To know if we can do default value operations
+	virtual bool IsAttribute() const { return false; }
+
 protected:
 	IPCGAttributeAccessor(bool bInReadOnly, int16 InUnderlyingType)
 		: bReadOnly(bInReadOnly)
diff --git a/Engine/Plugins/PCG/Source/PCG/Public/Metadata/Accessors/PCGAttributeAccessor.h b/Engine/Plugins/PCG/Source/PCG/Public/Metadata/Accessors/PCGAttributeAccessor.h
--- a/Engine/Plugins/PCG/Source/PCG/Public/Metadata/Accessors/PCGAttributeAccessor.h
+++ b/Engine/Plugins/PCG/Source/PCG/Public/Metadata/Accessors/PCGAttributeAccessor.h
@@ -111,6 +111,8 @@ public:
 		return true;
 	}
 
+	virtual bool IsAttribute() const override { return true; }
+
 private:
 	FPCGMetadataAttribute<T>* Attribute = nullptr;
 	UPCGMetadata* Metadata = nullptr;
diff --git a/Engine/Plugins/PCG/Source/PCGEditor/Private/DataVisualizations/PCGSplineDataVisualization.cpp b/Engine/Plugins/PCG/Source/PCGEditor/Private/DataVisualizations/PCGSplineDataVisualization.cpp
--- a/Engine/Plugins/PCG/Source/PCGEditor/Private/DataVisualizations/PCGSplineDataVisualization.cpp
+++ b/Engine/Plugins/PCG/Source/PCGEditor/Private/DataVisualizations/PCGSplineDataVisualization.cpp
@@ -11,6 +11,7 @@ const UPCGPointData* IPCGSplineDataVisualization::CollapseToDebugPointData(FPCGC
 	if (const UPCGSplineData* SplineData = Cast<UPCGSplineData>(Data))
 	{
 		UPCGPointData* PointData = NewObject<UPCGPointData>();
+		PointData->InitializeFromData(SplineData);
 		TArray<FPCGPoint>& Points = PointData->GetMutablePoints();
 
 		const int32 NumControlPoints = SplineData->SplineStruct.SplineCurves.Position.Points.Num();
