[PCG] Move delaying of PCG Component unregistering to Editor only. This exists to handle RerunConstructionScript & BP Reinstancing which we don't need on cook builds

#rnx

diff --git a/Engine/Plugins/PCG/Source/PCG/Private/PCGActorAndComponentMapping.cpp b/Engine/Plugins/PCG/Source/PCG/Private/PCGActorAndComponentMapping.cpp
--- a/Engine/Plugins/PCG/Source/PCG/Private/PCGActorAndComponentMapping.cpp
+++ b/Engine/Plugins/PCG/Source/PCG/Private/PCGActorAndComponentMapping.cpp
@@ -196,6 +196,8 @@ void FPCGActorAndComponentMapping::Deinitialize()
 
 void FPCGActorAndComponentMapping::Tick()
 {
+#if WITH_EDITOR
+
 	TSet<UPCGComponent*> ComponentToUnregister;
 	{
 		FScopeLock Lock(&DelayedComponentToUnregisterLock);
@@ -207,7 +209,6 @@ void FPCGActorAndComponentMapping::Tick()
 		UnregisterPCGComponent(Component, /*bForce=*/true);
 	}
 
-#if WITH_EDITOR
 	ProcessDelayedEvents();
 
 	const double CurrentTime = FApp::GetCurrentTime();
@@ -343,6 +344,7 @@ bool FPCGActorAndComponentMapping::RegisterOrUpdatePCGComponent(UPCGComponent* I
 		bHasChanged = RegisterOrUpdateNonPartitionedPCGComponent(InComponent);
 	}
 
+#if WITH_EDITOR
 	// If the component was previously marked as to be unregistered, remove it here.
 	{
 		FScopeLock Lock(&DelayedComponentToUnregisterLock);
@@ -350,7 +352,6 @@ bool FPCGActorAndComponentMapping::RegisterOrUpdatePCGComponent(UPCGComponent* I
 	}
 
 	// And finally handle the tracking. Only do it when the component is registered for the first time.
-#if WITH_EDITOR
 	if (!bWasAlreadyRegistered && bHasChanged)
 	{
 		RegisterTracking(InComponent);
@@ -420,11 +421,13 @@ bool FPCGActorAndComponentMapping::RemapPCGComponent(const UPCGComponent* OldCom
 		}
 	}
 
+#if WITH_EDITOR
 	// Remove it from the delayed
 	{
 		FScopeLock Lock(&DelayedComponentToUnregisterLock);
 		DelayedComponentToUnregister.Remove(OldComponent);
 	}
+#endif
 
 	// Remap all previous instances
 	auto RemapPreviousInstances = [OldComponent, NewComponent](TMap<const UPCGComponent*, TSet<TObjectPtr<APCGPartitionActor>>>& Map, FRWLock& Lock)
@@ -467,6 +470,7 @@ void FPCGActorAndComponentMapping::UnregisterPCGComponent(UPCGComponent* InCompo
 		return;
 	}
 
+#if WITH_EDITOR
 	if ((PartitionedOctree.Contains(InComponent) || NonPartitionedOctree.Contains(InComponent)))
 	{
 		bool bShouldBeDelayed = true;
@@ -487,19 +491,20 @@ void FPCGActorAndComponentMapping::UnregisterPCGComponent(UPCGComponent* InCompo
 			return;
 		}
 
-#if WITH_EDITOR
 		UnregisterTracking(InComponent);
-#endif // WITH_EDITOR
 	}
+#endif // WITH_EDITOR
 
 	UnregisterPartitionedPCGComponent(InComponent);
 	UnregisterNonPartitionedPCGComponent(InComponent);
 
+#if WITH_EDITOR
 	FScopeLock Lock(&DelayedComponentToUnregisterLock);
 	if (DelayedComponentToUnregister.Contains(InComponent))
 	{
 		DelayedComponentToUnregister.Remove(InComponent);
 	}
+#endif
 }
 
 void FPCGActorAndComponentMapping::UnregisterPartitionedPCGComponent(UPCGComponent* InComponent)
