Added cvar to control how much of the graph cache is evicted when it goes over budget

#rnx

diff --git a/Engine/Plugins/PCG/Source/PCG/Private/Graph/PCGGraphCache.cpp b/Engine/Plugins/PCG/Source/PCG/Private/Graph/PCGGraphCache.cpp
--- a/Engine/Plugins/PCG/Source/PCG/Private/Graph/PCGGraphCache.cpp
+++ b/Engine/Plugins/PCG/Source/PCG/Private/Graph/PCGGraphCache.cpp
@@ -24,6 +24,11 @@ static TAutoConsoleVariable<int32> CVarCacheMemoryBudgetMB(
 	6144,
 	TEXT("Memory budget for data in cache (MB)."));
 
+static TAutoConsoleVariable<float> CVarCacheMemoryCleanupRatio(
+	TEXT("pcg.Cache.MemoryCleanupRatio"),
+	0.5f,
+	TEXT("Target cache size ratio after triggering a cleanup (between 0 and 1.)."));
+
 static TAutoConsoleVariable<bool> CVarCacheMemoryBudgetEnabled(
 	TEXT("pcg.Cache.EnableMemoryBudget"),
 	true,
@@ -146,8 +151,10 @@ bool FPCGGraphCache::EnforceMemoryBudget()
 
 	{
 		FWriteScopeLock ScopeWriteLock(CacheLock);
+		const float MemoryCleanupRatio = FMath::Clamp(CVarCacheMemoryCleanupRatio.GetValueOnAnyThread(), 0.0f, 1.0f);
+		const uint64 TargetCacheMemoryUsage = static_cast<uint64>(MemoryCleanupRatio * MemoryBudget);
 
-		while (TotalMemoryUsed > MemoryBudget && CacheData.Num() > 0)
+		while (TotalMemoryUsed > TargetCacheMemoryUsage && CacheData.Num() > 0)
 		{
 			FPCGDataCollection RemovedData = CacheData.RemoveLeastRecent();
 			RemoveFromMemoryTotal(RemovedData);
