Get Property related improvements
- Added tags to output data for GetActorProperty
- Support multiple properties extraction, assuming they have the same cardinality with a CSV for GetActorProperty & GetPropertyFromObjectPath
- Add an option to output the actor reference from GetActorProperty.
- Added some options to override additional fields.

diff --git a/Engine/Plugins/PCG/Source/PCG/Private/Elements/PCGGetActorProperty.cpp b/Engine/Plugins/PCG/Source/PCG/Private/Elements/PCGGetActorProperty.cpp
--- a/Engine/Plugins/PCG/Source/PCG/Private/Elements/PCGGetActorProperty.cpp
+++ b/Engine/Plugins/PCG/Source/PCG/Private/Elements/PCGGetActorProperty.cpp
@@ -6,6 +6,7 @@
 #include "PCGContext.h"
 #include "PCGComponent.h"
 #include "PCGParamData.h"
+#include "Data/PCGPointData.h"
 #include "Helpers/PCGBlueprintHelpers.h"
 #include "Helpers/PCGDynamicTrackingHelpers.h"
 #include "Helpers/PCGPropertyHelpers.h"
@@ -42,6 +43,7 @@ void UPCGGetActorPropertySettings::PostLoad()
 {
 	Super::PostLoad();
 
+#if WITH_EDITOR
 	// Migrate deprecated actor selection settings to struct if needed
 	if (ActorSelection_DEPRECATED != EPCGActorSelection::ByTag ||
 		ActorSelectionTag_DEPRECATED != NAME_None ||
@@ -68,6 +70,7 @@ void UPCGGetActorPropertySettings::PostLoad()
 	{
 		ActorSelector.ActorSelectionClass = TSubclassOf<AActor>();
 	}
+#endif // WITH_EDITOR
 }
 
 FString UPCGGetActorPropertySettings::GetAdditionalTitleInformation() const
@@ -160,9 +163,7 @@ bool FPCGGetActorPropertyElement::ExecuteInternal(FPCGContext* Context) const
 			}
 		}
 
-		const FPCGAttributePropertySelector Selector = FPCGAttributePropertySelector::CreateSelectorFromString(Settings->PropertyName.ToString());
-
-		PCGPropertyHelpers::FExtractorParameters Parameters{ ObjectToInspect, ObjectToInspect->GetClass(), Selector, Settings->OutputAttributeName, Settings->bForceObjectAndStructExtraction, /*bPropertyNeedsToBeVisible=*/true };
+		PCGPropertyHelpers::FExtractorParameters Parameters(ObjectToInspect, ObjectToInspect->GetClass(), Settings->PropertyName.ToString(), Settings->OutputAttributeName, Settings->bForceObjectAndStructExtraction, /*bPropertyNeedsToBeVisible=*/true);
 
 		// Don't care for object traversed in non-editor build, since it is only useful for tracking.
 		TSet<FSoftObjectPath>* ObjectTraversedPtr = nullptr;
@@ -176,16 +177,27 @@ bool FPCGGetActorPropertyElement::ExecuteInternal(FPCGContext* Context) const
 			TArray<FPCGTaggedData>& Outputs = Context->OutputData.TaggedData;
 			FPCGTaggedData& Output = Outputs.Emplace_GetRef();
 			Output.Data = ParamData;
+			
+			for (FName ActorTag : FoundActor->Tags)
+			{
+				Output.Tags.Add(ActorTag.ToString());
+			}
+
+			if (Settings->bOutputActorReference)
+			{
+				check(ParamData->MutableMetadata());
+				ParamData->MutableMetadata()->FindOrCreateAttribute(PCGPointDataConstants::ActorReferenceAttribute, FSoftObjectPath(FoundActor), /*bAllowsInterpolation=*/false, /*bOverrideParent=*/false, /*bOverwriteIfTypeMismatch=*/true);
+			}
 		}
 		else
 		{
-			if (Selector.GetName() == NAME_None)
+			if(Parameters.PropertySelectors.IsEmpty() || Parameters.PropertySelectors[0].GetName() == NAME_None)
 			{
 				PCGE_LOG(Error, GraphAndLog, FText::Format(LOCTEXT("FailedToExtractActor", "Fail to extract actor {0}."), FText::FromString(FoundActor->GetName())));
 			}
 			else
 			{
-				PCGE_LOG(Error, GraphAndLog, FText::Format(LOCTEXT("FailedToExtract", "Fail to extract the property '{0}' on actor {1}."), Selector.GetDisplayText(), FText::FromString(FoundActor->GetName())));
+				PCGE_LOG(Error, GraphAndLog, FText::Format(LOCTEXT("FailedToExtract", "Fail to extract the property '{0}' on actor {1}."), Parameters.PropertySelectors[0].GetDisplayText(), FText::FromString(FoundActor->GetName())));
 			}
 		}
 
diff --git a/Engine/Plugins/PCG/Source/PCG/Private/Elements/PCGGetPropertyFromObjectPath.cpp b/Engine/Plugins/PCG/Source/PCG/Private/Elements/PCGGetPropertyFromObjectPath.cpp
--- a/Engine/Plugins/PCG/Source/PCG/Private/Elements/PCGGetPropertyFromObjectPath.cpp
+++ b/Engine/Plugins/PCG/Source/PCG/Private/Elements/PCGGetPropertyFromObjectPath.cpp
@@ -204,9 +204,8 @@ bool FPCGGetPropertyFromObjectPathElement::ExecuteInternal(FPCGContext* Context)
 			continue;
 		}
 
-		const FPCGAttributePropertySelector Selector = FPCGAttributePropertySelector::CreateSelectorFromString(Settings->PropertyName.ToString());
+		PCGPropertyHelpers::FExtractorParameters Parameters(Object, Object->GetClass(), Settings->PropertyName.ToString(), Settings->OutputAttributeName, Settings->bForceObjectAndStructExtraction, /*bPropertyNeedsToBeVisible=*/true);
 
-		PCGPropertyHelpers::FExtractorParameters Parameters{ Object, Object->GetClass(), Selector, Settings->OutputAttributeName, Settings->bForceObjectAndStructExtraction, /*bPropertyNeedsToBeVisible=*/true };
 		if (UPCGParamData* ParamData = PCGPropertyHelpers::ExtractPropertyAsAttributeSet(Parameters, Context))
 		{
 			AddToOutput(ParamData, Index);
@@ -216,13 +215,13 @@ bool FPCGGetPropertyFromObjectPathElement::ExecuteInternal(FPCGContext* Context)
 		}
 		else
 		{
-			if (Selector.GetName() == NAME_None)
+			if(Parameters.PropertySelectors.IsEmpty() || Parameters.PropertySelectors[0].GetName() == NAME_None)
 			{
 				PCGE_LOG(Error, GraphAndLog, FText::Format(LOCTEXT("FailedToExtractObject", "Fail to extract object {0}."), FText::FromString(Object->GetName())));
 			}
 			else
 			{
-				PCGE_LOG(Error, GraphAndLog, FText::Format(LOCTEXT("FailedToExtract", "Fail to extract the property '{0}' on object {1}."), Selector.GetDisplayText(), FText::FromString(Object->GetName())));
+				PCGE_LOG(Error, GraphAndLog, FText::Format(LOCTEXT("FailedToExtract", "Fail to extract the property '{0}' on object {1}."), Parameters.PropertySelectors[0].GetDisplayText(), FText::FromString(Object->GetName())));
 			}
 		}
 	}
diff --git a/Engine/Plugins/PCG/Source/PCG/Private/Elements/PCGUserParameterGet.cpp b/Engine/Plugins/PCG/Source/PCG/Private/Elements/PCGUserParameterGet.cpp
--- a/Engine/Plugins/PCG/Source/PCG/Private/Elements/PCGUserParameterGet.cpp
+++ b/Engine/Plugins/PCG/Source/PCG/Private/Elements/PCGUserParameterGet.cpp
@@ -166,21 +166,25 @@ bool FPCGUserParameterGetElement::ExecuteInternal(FPCGContext* Context) const
 	FConstStructView Parameters = PCGUserParameterGetSettings::GetFirstValidLayout(*Context);
 	const UScriptStruct* PropertyBag = Parameters.GetScriptStruct();
 
-	bool bForceObjectAndStructExtraction = false;
-	FPCGAttributePropertySelector InputSelector;
-	FName OutputAttributeName;
+	PCGPropertyHelpers::FExtractorParameters ExtractorParameters;
 
 	if (const UPCGUserParameterGetSettings* Settings = Context->GetInputSettings<UPCGUserParameterGetSettings>())
 	{
-		bForceObjectAndStructExtraction = Settings->bForceObjectAndStructExtraction;
-		InputSelector = FPCGAttributePropertySelector::CreateAttributeSelector(Settings->PropertyName);
-		OutputAttributeName = Settings->PropertyName;
+		ExtractorParameters = PCGPropertyHelpers::FExtractorParameters(Parameters.GetMemory(),
+			PropertyBag,
+			FPCGAttributePropertySelector::CreateAttributeSelector(Settings->PropertyName),
+			Settings->PropertyName,
+			Settings->bForceObjectAndStructExtraction,
+			/*bPropertyNeedsToBeVisible=*/false);
 	}
 	else if (const UPCGGenericUserParameterGetSettings* GenericSettings = Context->GetInputSettings<UPCGGenericUserParameterGetSettings>())
 	{
-		bForceObjectAndStructExtraction = GenericSettings->bForceObjectAndStructExtraction;
-		InputSelector = FPCGAttributePropertySelector::CreateSelectorFromString(GenericSettings->PropertyPath);
-		OutputAttributeName = GenericSettings->OutputAttributeName;
+		ExtractorParameters = PCGPropertyHelpers::FExtractorParameters(Parameters.GetMemory(),
+			PropertyBag,
+			GenericSettings->PropertyPath,
+			GenericSettings->OutputAttributeName,
+			GenericSettings->bForceObjectAndStructExtraction,
+			/*bPropertyNeedsToBeVisible=*/false);
 	}
 	else
 	{
@@ -188,8 +192,6 @@ bool FPCGUserParameterGetElement::ExecuteInternal(FPCGContext* Context) const
 		return true;
 	}
 
-	PCGPropertyHelpers::FExtractorParameters ExtractorParameters{ Parameters.GetMemory(), PropertyBag, InputSelector, OutputAttributeName, bForceObjectAndStructExtraction, /*bPropertyNeedsToBeVisible=*/false};
-
 	// Don't care for object traversed in non-editor build, since it is only useful for tracking.
 	TSet<FSoftObjectPath>* ObjectTraversedPtr = nullptr;
 #if WITH_EDITOR
@@ -203,7 +205,15 @@ bool FPCGUserParameterGetElement::ExecuteInternal(FPCGContext* Context) const
 	}
 	else
 	{
-		PCGE_LOG(Error, GraphAndLog, FText::Format(LOCTEXT("InvalidProperty", "Could not find the property '{0}' in the user parameters"), InputSelector.GetDisplayText()));
+		if (!ExtractorParameters.PropertySelectors.IsEmpty())
+		{
+			PCGE_LOG(Error, GraphAndLog, FText::Format(LOCTEXT("InvalidProperty", "Could not find the property '{0}' in the user parameters"), ExtractorParameters.PropertySelectors[0].GetDisplayText()));
+		}
+		else
+		{
+			PCGE_LOG(Error, GraphAndLog, LOCTEXT("InvalidProperty", "Could not find nameless property in the user parameters"));
+		}
+
 		return true;
 	}
 
diff --git a/Engine/Plugins/PCG/Source/PCG/Private/Helpers/PCGPropertyHelpers.cpp b/Engine/Plugins/PCG/Source/PCG/Private/Helpers/PCGPropertyHelpers.cpp
--- a/Engine/Plugins/PCG/Source/PCG/Private/Helpers/PCGPropertyHelpers.cpp
+++ b/Engine/Plugins/PCG/Source/PCG/Private/Helpers/PCGPropertyHelpers.cpp
@@ -6,6 +6,7 @@
 #include "PCGElement.h"
 #include "PCGModule.h"
 #include "PCGParamData.h"
+#include "Helpers/PCGHelpers.h"
 #include "Helpers/PCGSettingsHelpers.h"
 #include "Metadata/PCGMetadata.h"
 
@@ -175,6 +176,28 @@ namespace PCGPropertyHelpers
 			return Property;
 		}
 	}
+
+	FExtractorParameters::FExtractorParameters(const void* InContainer, const UStruct* InClass, const FPCGAttributePropertySelector& InPropertySelector, FName InOutputAttributeName, bool bInShouldExtract, bool bInPropertyNeedsToBeVisible)
+		: Container(InContainer), Class(InClass), OutputAttributeName(InOutputAttributeName), bShouldExtract(bInShouldExtract), bPropertyNeedsToBeVisible(bInPropertyNeedsToBeVisible)
+	{
+		PropertySelectors.Add(InPropertySelector);
+	}
+
+	FExtractorParameters::FExtractorParameters(const void* InContainer, const UStruct* InClass, const TArray<FPCGAttributePropertySelector>& InPropertySelectors, FName InOutputAttributeName, bool bInShouldExtract, bool bInPropertyNeedsToBeVisible)
+		: Container(InContainer), Class(InClass), PropertySelectors(InPropertySelectors), OutputAttributeName(InOutputAttributeName), bShouldExtract(bInShouldExtract), bPropertyNeedsToBeVisible(bInPropertyNeedsToBeVisible)
+	{}
+
+	FExtractorParameters::FExtractorParameters(const void* InContainer, const UStruct* InClass, const FString& InPropertySelectorString, FName InOutputAttributeName, bool bInShouldExtract, bool bInPropertyNeedsToBeVisible)
+		: Container(InContainer), Class(InClass), OutputAttributeName(InOutputAttributeName), bShouldExtract(bInShouldExtract), bPropertyNeedsToBeVisible(bInPropertyNeedsToBeVisible)
+	{
+		const TArray<FString> PropertySelectorStrings = PCGHelpers::GetStringArrayFromCommaSeparatedList(InPropertySelectorString);
+
+		PropertySelectors.Reserve(PropertySelectorStrings.Num());
+		for (int StringIndex = 0; StringIndex < PropertySelectorStrings.Num(); ++StringIndex)
+		{
+			PropertySelectors.Add(FPCGAttributePropertySelector::CreateSelectorFromString(PropertySelectorStrings[StringIndex]));
+		}
+	}
 }
 
 EPCGMetadataTypes PCGPropertyHelpers::GetMetadataTypeFromProperty(const FProperty* InProperty)
@@ -199,154 +222,192 @@ UPCGParamData* PCGPropertyHelpers::ExtractPropertyAsAttributeSet(const PCGProper
 {
 	check(Parameters.Container && Parameters.Class);
 
-	TArray<const void*> Containers = { Parameters.Container };
-	const FProperty* Property = nullptr;
-	const FName PropertyName = Parameters.PropertySelector.GetName();
-	const bool ExtractRoot = (PropertyName == NAME_None);
-	// If Name is none, extract the container as-is, using Parameters.Class, otherwise, extract the chain.
-	if (!ExtractRoot)
+	// TODO: Remove this after we've removed the deprecated PropertySelector.
+	TArrayView<const FPCGAttributePropertySelector> PropertySelectors;
+
+	if (Parameters.PropertySelectors.Num() > 0)
 	{
-		Property = ExtractPropertyChain(Parameters.Class, PropertyName, Parameters.PropertySelector.GetExtraNames(), Parameters.bPropertyNeedsToBeVisible, Containers, OptionalContext, OptionalObjectTraversed);
-		if (!Property)
-		{
-			return nullptr;
-		}
+		PropertySelectors = TArrayView<const FPCGAttributePropertySelector>(Parameters.PropertySelectors);
 	}
-
-	// If the property is an array, we will work on the underlying property, and extract each element as an entry in the param data
-	const FArrayProperty* ArrayProperty = CastField<FArrayProperty>(Property);
-	if (ArrayProperty)
+	else
 	{
-		Property = ArrayProperty->Inner;
+		PRAGMA_DISABLE_DEPRECATION_WARNINGS
+		PropertySelectors = TArrayView<const FPCGAttributePropertySelector>(&Parameters.PropertySelector, 1);
+		PRAGMA_ENABLE_DEPRECATION_WARNINGS
 	}
 
-	using ExtractablePropertyTuple = TTuple<FString, const FProperty*>;
-	TArray<ExtractablePropertyTuple> ExtractableProperties;
+	UPCGParamData* ParamData = nullptr;
+	UPCGMetadata* Metadata = nullptr;
+	TArray<PCGMetadataEntryKey> Entries;
+	bool bValidOperation = true;
+	bool bIgnoreOutputAttributeName = (PropertySelectors.Num() > 1);
+
+	for(const FPCGAttributePropertySelector& PropertySelector : PropertySelectors)
+	{
+		TArray<const void*> Containers = { Parameters.Container };
+		const FProperty* Property = nullptr;
+		const FName PropertyName = PropertySelector.GetName();
+		const bool ExtractRoot = (PropertyName == NAME_None);
+		// If Name is none, extract the container as-is, using Parameters.Class, otherwise, extract the chain.
+		if (!ExtractRoot)
+		{
+			Property = ExtractPropertyChain(Parameters.Class, PropertyName, PropertySelector.GetExtraNames(), Parameters.bPropertyNeedsToBeVisible, Containers, OptionalContext, OptionalObjectTraversed);
+			if (!Property)
+			{
+				return nullptr;
+			}
+		}
+
+		// If the property is an array, we will work on the underlying property, and extract each element as an entry in the param data
+		const FArrayProperty* ArrayProperty = CastField<FArrayProperty>(Property);
+		if (ArrayProperty)
+		{
+			Property = ArrayProperty->Inner;
+		}
 
-	using GetAddressFunc = TFunction<const void* (const void*)>;
-	GetAddressFunc AddressFunc;
+		using ExtractablePropertyTuple = TTuple<FString, const FProperty*>;
+		TArray<ExtractablePropertyTuple> ExtractableProperties;
 
-	// Force extraction if the property is not supported by accessors.
-	const bool bShouldExtract = Parameters.bShouldExtract || !PCGAttributeAccessorHelpers::IsPropertyAccessorSupported(Property);
+		using GetAddressFunc = TFunction<const void* (const void*)>;
+		GetAddressFunc AddressFunc;
 
-	// Keep track if the extracted property is an object or not
-	const FObjectProperty* ObjectProperty = CastField<FObjectProperty>(Property);
+		// Force extraction if the property is not supported by accessors.
+		const bool bShouldExtract = Parameters.bShouldExtract || !PCGAttributeAccessorHelpers::IsPropertyAccessorSupported(Property);
 
-	// Special case where the property is a struct/object, that is not supported by our metadata, we will try to break it down to multiple attributes in the resulting param data, if asked.
-	if (ExtractRoot || ((Property->IsA<FStructProperty>() || Property->IsA<FObjectProperty>()) && bShouldExtract))
-	{
-		const UStruct* UnderlyingClass = nullptr;
+		// Keep track if the extracted property is an object or not
+		const FObjectProperty* ObjectProperty = CastField<FObjectProperty>(Property);
 
-		if (ExtractRoot)
+		// Special case where the property is a struct/object, that is not supported by our metadata, we will try to break it down to multiple attributes in the resulting param data, if asked.
+		if (ExtractRoot || ((Property->IsA<FStructProperty>() || Property->IsA<FObjectProperty>()) && bShouldExtract))
 		{
-			UnderlyingClass = Parameters.Class;
-			// Identity
-			AddressFunc = [](const void* InAddress) { return InAddress; };
+			const UStruct* UnderlyingClass = nullptr;
+
+			if (ExtractRoot)
+			{
+				UnderlyingClass = Parameters.Class;
+				// Identity
+				AddressFunc = [](const void* InAddress) { return InAddress; };
+			}
+			else if (const FStructProperty* StructProperty = CastField<FStructProperty>(Property))
+			{
+				UnderlyingClass = StructProperty->Struct;
+				AddressFunc = [StructProperty](const void* InAddress) { return StructProperty->ContainerPtrToValuePtr<void>(InAddress); };
+			}
+			else if (ObjectProperty)
+			{
+				UnderlyingClass = ObjectProperty->PropertyClass;
+				AddressFunc = [ObjectProperty](const void* InAddress) { return ObjectProperty->GetObjectPropertyValue_InContainer(InAddress); };
+			}
+
+			check(UnderlyingClass);
+			check(!!AddressFunc);
+
+			// Re-use code from overridable params
+			// Limit ourselves to not recurse into more structs.
+			PCGSettingsHelpers::FPCGGetAllOverridableParamsConfig Config;
+			Config.bUseSeed = true;
+			Config.bExcludeSuperProperties = true;
+			Config.MaxStructDepth = 0;
+			// Can only get exposed properties and visible if requested
+			if (Parameters.bPropertyNeedsToBeVisible)
+			{
+				Config.ExcludePropertyFlags = ExcludePropertyFlags;
+				Config.IncludePropertyFlags = IncludePropertyFlags;
+			}
+			TArray<FPCGSettingsOverridableParam> AllChildProperties = PCGSettingsHelpers::GetAllOverridableParams(UnderlyingClass, Config);
+
+			for (const FPCGSettingsOverridableParam& Param : AllChildProperties)
+			{
+				if (ensure(!Param.PropertiesNames.IsEmpty()))
+				{
+					const FName ChildPropertyName = Param.PropertiesNames[0];
+					if (const FProperty* ChildProperty = UnderlyingClass->FindPropertyByName(ChildPropertyName))
+					{
+						// We use authored name as attribute name to avoid issue with noisy property names, like in UUserDefinedStructs, where some random number is appended to the property name.
+						// By default, it will just return the property name anyway.
+						FString AuthoredName = UnderlyingClass->GetAuthoredNameForField(ChildProperty);
+						ExtractableProperties.Emplace(std::move(AuthoredName), ChildProperty);
+					}
+				}
+			}
 		}
-		else if (const FStructProperty* StructProperty = CastField<FStructProperty>(Property))
+		else
 		{
-			UnderlyingClass = StructProperty->Struct;
-			AddressFunc = [StructProperty](const void* InAddress) { return StructProperty->ContainerPtrToValuePtr<void>(InAddress); };
+			// For non struct/object, there is just a single property to extract with no shenanigans for address indirection.
+			const FName AttributeName = (bIgnoreOutputAttributeName || Parameters.OutputAttributeName == PCGMetadataAttributeConstants::SourceNameAttributeName) ? Property->GetFName() : Parameters.OutputAttributeName;
+			ExtractableProperties.Emplace(AttributeName.ToString(), Property);
+			// Identity
+			AddressFunc = [](const void* InAddress) { return InAddress; };
 		}
-		else if (ObjectProperty)
+
+		if (ExtractableProperties.IsEmpty())
 		{
-			UnderlyingClass = ObjectProperty->PropertyClass;
-			AddressFunc = [ObjectProperty](const void* InAddress) { return ObjectProperty->GetObjectPropertyValue_InContainer(InAddress); };
+			LogError(LOCTEXT("NoPropertiesFound", "No properties found to extract"), OptionalContext);
+			return nullptr;
 		}
 
-		check(UnderlyingClass);
-		check(!!AddressFunc);
-
-		// Re-use code from overridable params
-		// Limit ourselves to not recurse into more structs.
-		PCGSettingsHelpers::FPCGGetAllOverridableParamsConfig Config;
-		Config.bUseSeed = true;
-		Config.bExcludeSuperProperties = true;
-		Config.MaxStructDepth = 0;
-		// Can only get exposed properties and visible if requested
-		if (Parameters.bPropertyNeedsToBeVisible)
+		// Before we need to compute all the addresses for each entry in our array (or just a single entry if there is no array)
+		TArray<const void*, TInlineAllocator<16>> ElementAddresses;
+		ExpandContainers(ArrayProperty, Containers, ElementAddresses);
+
+		if (!ParamData)
 		{
-			Config.ExcludePropertyFlags = ExcludePropertyFlags;
-			Config.IncludePropertyFlags = IncludePropertyFlags;
+			// From there, we should be able to create the data.
+			ParamData = NewObject<UPCGParamData>();
+			check(!Metadata);
+			Metadata = ParamData->MutableMetadata();
+			check(Metadata);
+
+			// Allocate entries
+			TArray<PCGMetadataEntryKey, TInlineAllocator<16>> ParentEntries;
+			ParentEntries.Init(PCGInvalidEntryKey, ElementAddresses.Num());
+			Entries = Metadata->AddEntries(ParentEntries);
 		}
-		TArray<FPCGSettingsOverridableParam> AllChildProperties = PCGSettingsHelpers::GetAllOverridableParams(UnderlyingClass, Config);
-
-		for (const FPCGSettingsOverridableParam& Param : AllChildProperties)
+		else // Pre-existing data, we're expecting the same cardinality
 		{
-			if (ensure(!Param.PropertiesNames.IsEmpty()))
+			if (Entries.Num() != ElementAddresses.Num())
 			{
-				const FName ChildPropertyName = Param.PropertiesNames[0];
-				if (const FProperty* ChildProperty = UnderlyingClass->FindPropertyByName(ChildPropertyName))
-				{
-					// We use authored name as attribute name to avoid issue with noisy property names, like in UUserDefinedStructs, where some random number is appended to the property name.
-					// By default, it will just return the property name anyway.
-					FString AuthoredName = UnderlyingClass->GetAuthoredNameForField(ChildProperty);
-					ExtractableProperties.Emplace(std::move(AuthoredName), ChildProperty);
-				}
+				LogError(LOCTEXT("InvalidCardinality", "Unable to extract because some properties are of mismatched sizes"), OptionalContext);
+				return nullptr;
 			}
 		}
-	}
-	else
-	{
-		// For non struct/object, there is just a single property to extract with no shenanigans for address indirection.
-		const FName AttributeName = (Parameters.OutputAttributeName == PCGMetadataAttributeConstants::SourceNameAttributeName) ? Property->GetFName() : Parameters.OutputAttributeName;
-		ExtractableProperties.Emplace(AttributeName.ToString(), Property);
-		// Identity
-		AddressFunc = [](const void* InAddress) { return InAddress; };
-	}
-
-	if (ExtractableProperties.IsEmpty())
-	{
-		LogError(LOCTEXT("NoPropertiesFound", "No properties found to extract"), OptionalContext);
-		return nullptr;
-	}
-
-	// Before we need to compute all the addresses for each entry in our array (or just a single entry if there is no array)
-	TArray<const void*, TInlineAllocator<16>> ElementAddresses;
-	ExpandContainers(ArrayProperty, Containers, ElementAddresses);
 
-	// From there, we should be able to create the data.
-	UPCGParamData* ParamData = NewObject<UPCGParamData>();
-	UPCGMetadata* Metadata = ParamData->MutableMetadata();
-	check(Metadata);
-
-	bool bValidOperation = true;
-
-	for (const void* ElementAddress : ElementAddresses)
-	{
-		if (!bValidOperation)
+		for(int32 ElementAddressIndex = 0; ElementAddressIndex < ElementAddresses.Num(); ++ElementAddressIndex)
 		{
-			break;
-		}
+			const void* ElementAddress = ElementAddresses[ElementAddressIndex];
+			PCGMetadataEntryKey EntryKey = Entries[ElementAddressIndex];
 
-		// Add a new entry for all elements
-		PCGMetadataEntryKey EntryKey = Metadata->AddEntry();
+			// Offset the address if needed
+			const void* ContainerPtr = AddressFunc(ElementAddress);
 
-		// Offset the address if needed
-		const void* ContainerPtr = AddressFunc(ElementAddress);
+			for (ExtractablePropertyTuple& ExtractableProperty : ExtractableProperties)
+			{
+				FString AttributeName = ExtractableProperty.Get<0>();
+				const FProperty* FinalProperty = ExtractableProperty.Get<1>();
 
-		for (ExtractablePropertyTuple& ExtractableProperty : ExtractableProperties)
-		{
-			FString AttributeName = ExtractableProperty.Get<0>();
-			const FProperty* FinalProperty = ExtractableProperty.Get<1>();
+				// Make sure the Attribute name is sanitized, to prevent cases where property names have unsupported characters.
+				FPCGMetadataAttributeBase::SanitizeName(AttributeName);
 
-			// Make sure the Attribute name is sanitized, to prevent cases where property names have unsupported characters.
-			FPCGMetadataAttributeBase::SanitizeName(AttributeName);
+				if (!Metadata->SetAttributeFromDataProperty(FName(AttributeName), EntryKey, ContainerPtr, FinalProperty, /*bCreate=*/ true))
+				{
+					LogError(FText::Format(LOCTEXT("ErrorCreatingAttribute", "Error while creating an attribute for property '{0}'. Either the property type is not supported by PCG or attribute creation failed."), FText::FromString(FinalProperty->GetName())), OptionalContext);
+					bValidOperation = false;
+					break;
+				}
+			}
 
-			if (!Metadata->SetAttributeFromDataProperty(FName(AttributeName), EntryKey, ContainerPtr, FinalProperty, /*bCreate=*/ true))
+			if (!bValidOperation)
 			{
-				LogError(FText::Format(LOCTEXT("ErrorCreatingAttribute", "Error while creating an attribute for property '{0}'. Either the property type is not supported by PCG or attribute creation failed."), FText::FromString(FinalProperty->GetName())), OptionalContext);
-				bValidOperation = false;
 				break;
 			}
-		}
 
-		if (bValidOperation && bShouldExtract && OptionalObjectTraversed && ObjectProperty)
-		{
-			const UObject* Object = ObjectProperty->GetPropertyValue_InContainer(ElementAddress);
-			if (IsValid(Object))
+			if (bShouldExtract && OptionalObjectTraversed && ObjectProperty)
 			{
-				OptionalObjectTraversed->Add(Object);
+				const UObject* Object = ObjectProperty->GetPropertyValue_InContainer(ElementAddress);
+				if (IsValid(Object))
+				{
+					OptionalObjectTraversed->Add(Object);
+				}
 			}
 		}
 	}
diff --git a/Engine/Plugins/PCG/Source/PCG/Public/Elements/PCGGetActorProperty.h b/Engine/Plugins/PCG/Source/PCG/Public/Elements/PCGGetActorProperty.h
--- a/Engine/Plugins/PCG/Source/PCG/Public/Elements/PCGGetActorProperty.h
+++ b/Engine/Plugins/PCG/Source/PCG/Public/Elements/PCGGetActorProperty.h
@@ -58,18 +58,22 @@ public:
 	UPROPERTY(BlueprintReadWrite, EditAnywhere, Category = Settings, meta = (EditCondition = "bSelectComponent", EditConditionHides))
 	TSubclassOf<UActorComponent> ComponentClass;
 
-	/** Property name to extract. Can only extract properties that are compatible with metadata types. If None, extract the actor/component directly.*/
-	UPROPERTY(BlueprintReadWrite, EditAnywhere, Category = Settings)
+	/** Property name to extract. Can only extract properties that are compatible with metadata types. If None, extract the actor/component directly. Can be a comma-separated list, assuming they have the same cardinality. */
+	UPROPERTY(BlueprintReadWrite, EditAnywhere, Category = Settings, meta = (PCG_Overridable))
 	FName PropertyName = NAME_None;
 
 	/** If the property is a struct/object supported by metadata, this option can be toggled to force extracting all (compatible) properties contained in this property. Automatically true if unsupported by metadata. For now, only supports direct child properties (and not deeper). */
-	UPROPERTY(BlueprintReadWrite, EditAnywhere, Category = Settings)
+	UPROPERTY(BlueprintReadWrite, EditAnywhere, Category = Settings, meta = (PCG_Overridable))
 	bool bForceObjectAndStructExtraction = false;
 
-	/** By default, attribute name will be None, but it can be overridden by this name. Use @SourceName to use the property name (only works when not extracting). */
-	UPROPERTY(BlueprintReadWrite, EditAnywhere, Category = Settings, meta = (EditCondition = "!bForceObjectAndStructExtraction", EditConditionHides))
+	/** By default, attribute name will be None, but it can be overridden by this name. Use @SourceName to use the property name (only works when not extracting). Will be ignored if multiple properties are extracted at the same time. */
+	UPROPERTY(BlueprintReadWrite, EditAnywhere, Category = Settings, meta = (PCG_Overridable, EditCondition = "!bForceObjectAndStructExtraction", EditConditionHides))
 	FName OutputAttributeName = NAME_None;
 
+	/** Controls whether an actor reference attribute will be added to the result */
+	UPROPERTY(BlueprintReadWrite, EditAnywhere, Category = Settings, meta = (PCG_Overridable))
+	bool bOutputActorReference = false;
+
 	/** If this is true, we will never put this element in cache, and will always try to re-query the actors and read the latest properties from them. */
 	UPROPERTY(BlueprintReadWrite, EditAnywhere, Category = Settings)
 	bool bAlwaysRequeryActors = true;
@@ -81,6 +85,7 @@ public:
 #endif // WITH_EDITORONLY_DATA
 
 private:
+#if WITH_EDITORONLY_DATA
 	UPROPERTY()
 	EPCGActorSelection ActorSelection_DEPRECATED;
 
@@ -98,6 +103,7 @@ private:
 
 	UPROPERTY()
 	bool bIncludeChildren_DEPRECATED = false;
+#endif
 };
 
 class FPCGGetActorPropertyElement : public IPCGElement
diff --git a/Engine/Plugins/PCG/Source/PCG/Public/Elements/PCGGetPropertyFromObjectPath.h b/Engine/Plugins/PCG/Source/PCG/Public/Elements/PCGGetPropertyFromObjectPath.h
--- a/Engine/Plugins/PCG/Source/PCG/Public/Elements/PCGGetPropertyFromObjectPath.h
+++ b/Engine/Plugins/PCG/Source/PCG/Public/Elements/PCGGetPropertyFromObjectPath.h
@@ -48,15 +48,15 @@ public:
 	FPCGAttributePropertyInputSelector InputSource;
 
 	/** Property name to extract. Can only extract properties that are compatible with metadata types. If None, extract the object. */
-	UPROPERTY(BlueprintReadWrite, EditAnywhere, Category = Settings)
+	UPROPERTY(BlueprintReadWrite, EditAnywhere, Category = Settings, meta = (PCG_Overridable))
 	FName PropertyName = NAME_None;
 
 	/** If the property is a struct/object supported by metadata, this option can be toggled to force extracting all (compatible) properties contained in this property. Automatically true if unsupported by metadata. For now, only supports direct child properties (and not deeper). */
 	UPROPERTY(BlueprintReadWrite, EditAnywhere, Category = Settings)
 	bool bForceObjectAndStructExtraction = false;
 
-	/** By default, attribute name will be None, but it can be overridden by this name. Use @SourceName to use the property name (only works when not extracting). */
-	UPROPERTY(BlueprintReadWrite, EditAnywhere, Category = Settings, meta = (EditCondition = "!bForceObjectAndStructExtraction", EditConditionHides))
+	/** By default, attribute name will be None, but it can be overridden by this name. Use @SourceName to use the property name (only works when not extracting). In the case of multiple properties being extracted, will be ignored. */
+	UPROPERTY(BlueprintReadWrite, EditAnywhere, Category = Settings, meta = (PCG_Overridable, EditCondition = "!bForceObjectAndStructExtraction", EditConditionHides))
 	FName OutputAttributeName = NAME_None;
 
 	/** By default, object loading is asynchronous, can force it synchronous if needed. */
diff --git a/Engine/Plugins/PCG/Source/PCG/Public/Helpers/PCGPropertyHelpers.h b/Engine/Plugins/PCG/Source/PCG/Public/Helpers/PCGPropertyHelpers.h
--- a/Engine/Plugins/PCG/Source/PCG/Public/Helpers/PCGPropertyHelpers.h
+++ b/Engine/Plugins/PCG/Source/PCG/Public/Helpers/PCGPropertyHelpers.h
@@ -47,6 +47,19 @@ namespace PCGPropertyHelpers
 
 	struct FExtractorParameters
 	{
+		// Disable deprecation warning on the rule of 5 because of the PropertySelector member.
+		PRAGMA_DISABLE_DEPRECATION_WARNINGS
+		FExtractorParameters() = default;
+		FExtractorParameters(const FExtractorParameters&) = default;
+		FExtractorParameters(FExtractorParameters&&) = default;
+		FExtractorParameters& operator=(const FExtractorParameters&) = default;
+		FExtractorParameters& operator=(FExtractorParameters&&) = default;
+		PRAGMA_ENABLE_DEPRECATION_WARNINGS
+
+		FExtractorParameters(const void* InContainer, const UStruct* InClass, const FPCGAttributePropertySelector& InPropertySelector, FName InOutputAttributeName, bool bInShouldExtract, bool bInPropertyNeedsToBeVisible);
+		FExtractorParameters(const void* InContainer, const UStruct* InClass, const TArray<FPCGAttributePropertySelector>& InPropertySelectors, FName InOutputAttributeName, bool bInShouldExtract, bool bInPropertyNeedsToBeVisible);
+		FExtractorParameters(const void* InContainer, const UStruct* InClass, const FString& InPropertySelectorString, FName InOutputAttributeName, bool bInShouldExtract, bool bInPropertyNeedsToBeVisible);
+		
 		// Pointer to the container containing the data we want to extract
 		const void* Container = nullptr;
 
@@ -54,8 +67,11 @@ namespace PCGPropertyHelpers
 		const UStruct* Class = nullptr;
 
 		// Selector of the property we want to extract
+		UE_DEPRECATED(5.5, "Use the selector array (PropertySelectors) instead.")
 		FPCGAttributePropertySelector PropertySelector;
 
+		TArray<FPCGAttributePropertySelector> PropertySelectors;
+
 		// Optional name of the attribute that will receive the extracted property. If None, will take the property name. 
 		// Also not used for Structs/Object extraction, as we will create multiple attributes, and they will be the name of all the extracted members.
 		FName OutputAttributeName = NAME_None;
