Allowed specifying list of actors to get data from ("From Input")
Added ability to override parameters of the actor selectors and updated dynamic tracking accordingly
Fixed an issue with soft class path overloading

diff --git a/Engine/Plugins/PCG/Source/PCG/Private/Elements/PCGActorSelector.cpp b/Engine/Plugins/PCG/Source/PCG/Private/Elements/PCGActorSelector.cpp
--- a/Engine/Plugins/PCG/Source/PCG/Private/Elements/PCGActorSelector.cpp
+++ b/Engine/Plugins/PCG/Source/PCG/Private/Elements/PCGActorSelector.cpp
@@ -73,7 +73,7 @@ namespace PCGActorSelector
 		return {};
 	}
 
-	TArray<AActor*> FindActors(const FPCGActorSelectorSettings& Settings, const UPCGComponent* InComponent, const TFunction<bool(const AActor*)>& BoundsCheck, const TFunction<bool(const AActor*)>& SelfIgnoreCheck)
+	TArray<AActor*> FindActors(const FPCGActorSelectorSettings& Settings, const UPCGComponent* InComponent, const TFunction<bool(const AActor*)>& BoundsCheck, const TFunction<bool(const AActor*)>& SelfIgnoreCheck, TArrayView<AActor*> InputActors)
 	{
 		TRACE_CPUPROFILER_EVENT_SCOPE(PCGActorSelector::FindActor);
 
@@ -173,6 +173,11 @@ namespace PCGActorSelector
 			}
 		}
 
+		case EPCGActorFilter::FromInput:
+		{
+			ActorsToCheck.Append(InputActors);
+		}
+
 		default:
 			break;
 		}
@@ -198,20 +203,40 @@ namespace PCGActorSelector
 		return FoundActors;
 	}
 
-	AActor* FindActor(const FPCGActorSelectorSettings& InSettings, UPCGComponent* InComponent, const TFunction<bool(const AActor*)>& BoundsCheck, const TFunction<bool(const AActor*)>& SelfIgnoreCheck)
+	AActor* FindActor(const FPCGActorSelectorSettings& InSettings, UPCGComponent* InComponent, const TFunction<bool(const AActor*)>& BoundsCheck, const TFunction<bool(const AActor*)>& SelfIgnoreCheck, TArrayView<AActor*> InputActors)
 	{
 		// In order to make sure we don't try to select multiple, we'll do a copy of the settings here.
 		FPCGActorSelectorSettings Settings = InSettings;
 		Settings.bSelectMultiple = false;
 
-		TArray<AActor*> Actors = FindActors(Settings, InComponent, BoundsCheck, SelfIgnoreCheck);
+		TArray<AActor*> Actors = FindActors(Settings, InComponent, BoundsCheck, SelfIgnoreCheck, InputActors);
 		return Actors.IsEmpty() ? nullptr : Actors[0];
 	}
 }
 
+FPCGSelectionKey::FPCGSelectionKey(const FPCGActorSelectorSettings& InActorSelector)
+{
+	check(InActorSelector.ActorFilter != EPCGActorFilter::FromInput);
+	ActorFilter = InActorSelector.ActorFilter;
+	if (ActorFilter == EPCGActorFilter::AllWorldActors)
+	{
+		check(InActorSelector.ActorSelection == EPCGActorSelection::ByTag || InActorSelector.ActorSelection == EPCGActorSelection::ByClass);
+		Selection = InActorSelector.ActorSelection;
+
+		if (Selection == EPCGActorSelection::ByTag)
+		{
+			Tag = InActorSelector.ActorSelectionTag;
+		}
+		else
+		{
+			SelectionClass = InActorSelector.ActorSelectionClass;
+		}
+	}
+}
+
 FPCGSelectionKey::FPCGSelectionKey(EPCGActorFilter InFilter)
 {
-	check(InFilter != EPCGActorFilter::AllWorldActors);
+	check(InFilter != EPCGActorFilter::AllWorldActors && InFilter != EPCGActorFilter::FromInput);
 	ActorFilter = InFilter;
 }
 
@@ -468,20 +493,26 @@ FName FPCGActorSelectorSettings::GetTaskName(const FText& Prefix) const
 
 FPCGSelectionKey FPCGActorSelectorSettings::GetAssociatedKey() const
 {
-	if (ActorFilter != EPCGActorFilter::AllWorldActors)
+	if (ActorFilter == EPCGActorFilter::AllWorldActors)
 	{
-		return FPCGSelectionKey(ActorFilter);
+		switch (ActorSelection)
+		{
+		case EPCGActorSelection::ByTag:
+			return FPCGSelectionKey(ActorSelectionTag);
+		case EPCGActorSelection::ByClass:
+			return FPCGSelectionKey(ActorSelectionClass);
+		default:
+			return FPCGSelectionKey();
+		}
 	}
-
-	switch (ActorSelection)
+	else if (ActorFilter == EPCGActorFilter::FromInput)
 	{
-	case EPCGActorSelection::ByTag:
-		return FPCGSelectionKey(ActorSelectionTag);
-	case EPCGActorSelection::ByClass:
-		return FPCGSelectionKey(ActorSelectionClass);
-	default:
 		return FPCGSelectionKey();
 	}
+	else
+	{
+		return FPCGSelectionKey(ActorFilter);
+	}
 }
 
 FPCGActorSelectorSettings FPCGActorSelectorSettings::ReconstructFromKey(const FPCGSelectionKey& InKey)
diff --git a/Engine/Plugins/PCG/Source/PCG/Private/Elements/PCGDataFromActor.cpp b/Engine/Plugins/PCG/Source/PCG/Private/Elements/PCGDataFromActor.cpp
--- a/Engine/Plugins/PCG/Source/PCG/Private/Elements/PCGDataFromActor.cpp
+++ b/Engine/Plugins/PCG/Source/PCG/Private/Elements/PCGDataFromActor.cpp
@@ -13,6 +13,7 @@
 #include "Elements/PCGMergeElement.h"
 #include "Elements/PCGMergeAttributes.h"
 #include "Grid/PCGPartitionActor.h"
+#include "Helpers/PCGDynamicTrackingHelpers.h"
 #include "Helpers/PCGHelpers.h"
 #include "Utils/PCGGraphExecutionLogging.h"
 
@@ -230,7 +231,24 @@ EPCGDataType UPCGDataFromActorSettings::GetCurrentPinTypes(const UPCGPin* InPin)
 		}
 	}
 
-	return Super::GetCurrentPinTypes(InPin);
+	// Implementation note: since we can have an input pin in some instances, we can't rely on the
+	// base class to provide the proper output type, as this will override the values set in the OutputPinProperties.
+	return InPin->Properties.AllowedTypes;
+}
+
+TArray<FPCGPinProperties> UPCGDataFromActorSettings::InputPinProperties() const
+{
+	TArray<FPCGPinProperties> PinProperties;
+
+	if (ActorSelector.ActorFilter == EPCGActorFilter::FromInput)
+	{
+		PinProperties.Emplace(PCGPinConstants::DefaultInputLabel,
+			EPCGDataType::Any,
+			/*bAllowMultipleConnections=*/true,
+			/*bAllowMultipleData=*/true);
+	}
+
+	return PinProperties;
 }
 
 TArray<FPCGPinProperties> UPCGDataFromActorSettings::OutputPinProperties() const
@@ -257,9 +275,28 @@ TArray<FPCGPinProperties> UPCGDataFromActorSettings::OutputPinProperties() const
 	return Pins;
 }
 
-FPCGContext* FPCGDataFromActorElement::CreateContext()
+bool FPCGDataFromActorElement::PrepareDataInternal(FPCGContext* Context) const
 {
-	return new FPCGDataFromActorContext();
+	TRACE_CPUPROFILER_EVENT_SCOPE(FPCGDataFromActorElement::PrepareData);
+	check(Context);
+
+	const UPCGDataFromActorSettings* Settings = Context->GetInputSettings<UPCGDataFromActorSettings>();
+	check(Settings);
+
+	if (Settings->ActorSelector.ActorFilter == EPCGActorFilter::FromInput)
+	{
+		FPCGDataFromActorContext* ThisContext = static_cast<FPCGDataFromActorContext*>(Context);
+		return ThisContext->InitializeAndRequestLoad(PCGPinConstants::DefaultInputLabel,
+			Settings->ActorSelector.ActorReferenceSelector,
+			{},
+			/*bPersistAllData=*/false,
+			/*bSilenceErrorOnEmptyObjectPath=*/true,
+			/*bSynchronousLoad=*/false);
+	}
+	else
+	{
+		return true;
+	}
 }
 
 bool FPCGDataFromActorElement::ExecuteInternal(FPCGContext* InContext) const
@@ -368,8 +405,34 @@ bool FPCGDataFromActorElement::ExecuteInternal(FPCGContext* InContext) const
 		if (!Context->bPerformedQuery)
 		{
 			TRACE_CPUPROFILER_EVENT_SCOPE(FPCGDataFromActorElement::Execute::FindActors);
-			Context->FoundActors = PCGActorSelector::FindActors(Settings->ActorSelector, Context->SourceComponent.Get(), BoundsCheck, SelfIgnoreCheck);
+			TArray<AActor*> ActorsFromInput;
+
+			if (Settings->ActorSelector.ActorFilter == EPCGActorFilter::FromInput)
+			{
+				Algo::Transform(Context->PathsToObjectsAndDataIndex, ActorsFromInput, [](const TTuple<FSoftObjectPath, int32, int32>& InPath) { return Cast<AActor>(InPath.Get<0>().ResolveObject()); });
+			}
+
+			Context->FoundActors = PCGActorSelector::FindActors(Settings->ActorSelector, Context->SourceComponent.Get(), BoundsCheck, SelfIgnoreCheck, ActorsFromInput);
 			Context->bPerformedQuery = true;
+
+#if WITH_EDITOR
+			// Setup dynamic tracking if needed
+			if (Settings->ActorSelector.ActorFilter == EPCGActorFilter::FromInput)
+			{
+				FPCGDynamicTrackingHelper DynamicTracking;
+				DynamicTracking.EnableAndInitialize(Context, Context->PathsToObjectsAndDataIndex.Num());
+				for (const TTuple<FSoftObjectPath, int32, int32>& Path : Context->PathsToObjectsAndDataIndex)
+				{
+					DynamicTracking.AddToTracking(FPCGSelectionKey::CreateFromPath(Path.Get<0>()), Settings->ActorSelector.bMustOverlapSelf);
+				}
+
+				DynamicTracking.Finalize(Context);
+			}
+			else if (Context->IsValueOverriden(GET_MEMBER_NAME_CHECKED(UPCGDataFromActorSettings, ActorSelector)))
+			{
+				FPCGDynamicTrackingHelper::AddSingleDynamicTrackingKey(Context, Settings->ActorSelector);
+			}
+#endif // WITH_EDITOR
 		}
 
 		if (Context->FoundActors.IsEmpty())
diff --git a/Engine/Plugins/PCG/Source/PCG/Private/Elements/PCGGetActorProperty.cpp b/Engine/Plugins/PCG/Source/PCG/Private/Elements/PCGGetActorProperty.cpp
--- a/Engine/Plugins/PCG/Source/PCG/Private/Elements/PCGGetActorProperty.cpp
+++ b/Engine/Plugins/PCG/Source/PCG/Private/Elements/PCGGetActorProperty.cpp
@@ -82,6 +82,21 @@ FString UPCGGetActorPropertySettings::GetAdditionalTitleInformation() const
 #endif
 }
 
+TArray<FPCGPinProperties> UPCGGetActorPropertySettings::InputPinProperties() const
+{
+	TArray<FPCGPinProperties> PinProperties;
+
+	if (ActorSelector.ActorFilter == EPCGActorFilter::FromInput)
+	{
+		PinProperties.Emplace(PCGPinConstants::DefaultInputLabel,
+			EPCGDataType::Any,
+			/*bAllowMultipleConnections=*/true,
+			/*bAllowMultipleData=*/true);
+	}
+
+	return PinProperties;
+}
+
 TArray<FPCGPinProperties> UPCGGetActorPropertySettings::OutputPinProperties() const
 {
 	TArray<FPCGPinProperties> PinProperties;
@@ -90,16 +105,50 @@ TArray<FPCGPinProperties> UPCGGetActorPropertySettings::OutputPinProperties() co
 	return PinProperties;
 }
 
+EPCGDataType UPCGGetActorPropertySettings::GetCurrentPinTypes(const UPCGPin* InPin) const
+{
+	check(InPin);
+	// Implementation note: since we can have an input pin in some instances, we can't rely on the
+	// base class to provide the proper output type, as this will override the values set in the OutputPinProperties.
+	return InPin->Properties.AllowedTypes;
+}
+
 FPCGElementPtr UPCGGetActorPropertySettings::CreateElement() const
 {
 	return MakeShared<FPCGGetActorPropertyElement>();
 }
 
-bool FPCGGetActorPropertyElement::ExecuteInternal(FPCGContext* Context) const
+bool FPCGGetActorPropertyElement::PrepareDataInternal(FPCGContext* Context) const
+{
+	TRACE_CPUPROFILER_EVENT_SCOPE(FPCGGetActorPropertyElement::PrepareData);
+	check(Context);
+
+	const UPCGGetActorPropertySettings* Settings = Context->GetInputSettings<UPCGGetActorPropertySettings>();
+	check(Settings);
+
+	if (Settings->ActorSelector.ActorFilter == EPCGActorFilter::FromInput)
+	{
+		FPCGLoadObjectsFromPathContext* ThisContext = static_cast<FPCGLoadObjectsFromPathContext*>(Context);
+		return ThisContext->InitializeAndRequestLoad(PCGPinConstants::DefaultInputLabel,
+			Settings->ActorSelector.ActorReferenceSelector,
+			{},
+			/*bPersistAllData=*/false,
+			/*bSilenceErrorOnEmptyObjectPath=*/true,
+			/*bSynchronousLoad=*/false);
+	}
+	else
+	{
+		return true;
+	}
+
+}
+
+bool FPCGGetActorPropertyElement::ExecuteInternal(FPCGContext* InContext) const
 {
 	TRACE_CPUPROFILER_EVENT_SCOPE(FPCGGetActorPropertyElement::Execute);
 
-	check(Context);
+	check(InContext);
+	FPCGLoadObjectsFromPathContext* Context = static_cast<FPCGLoadObjectsFromPathContext*>(InContext);
 
 	const UPCGGetActorPropertySettings* Settings = Context->GetInputSettings<UPCGGetActorPropertySettings>();
 	check(Settings);
@@ -136,7 +185,32 @@ bool FPCGGetActorPropertyElement::ExecuteInternal(FPCGContext* Context) const
 		};
 	}
 
-	TArray<AActor*> FoundActors = PCGActorSelector::FindActors(Settings->ActorSelector, OriginalComponent, BoundsCheck, NoSelfIgnoreCheck);
+	TArray<AActor*> ActorsFromInput;
+	if (Settings->ActorSelector.ActorFilter == EPCGActorFilter::FromInput)
+	{
+		Algo::Transform(Context->PathsToObjectsAndDataIndex, ActorsFromInput, [](const TTuple<FSoftObjectPath, int32, int32>& InPath) { return Cast<AActor>(InPath.Get<0>().ResolveObject()); });
+	}
+
+	TArray<AActor*> FoundActors = PCGActorSelector::FindActors(Settings->ActorSelector, OriginalComponent, BoundsCheck, NoSelfIgnoreCheck, ActorsFromInput);
+
+#if WITH_EDITOR
+	// Setup dynamic tracking if needed
+	if (Settings->ActorSelector.ActorFilter == EPCGActorFilter::FromInput)
+	{
+		FPCGDynamicTrackingHelper DynamicTracking;
+		DynamicTracking.EnableAndInitialize(Context, Context->PathsToObjectsAndDataIndex.Num());
+		for (const TTuple<FSoftObjectPath, int32, int32>& Path : Context->PathsToObjectsAndDataIndex)
+		{
+			DynamicTracking.AddToTracking(FPCGSelectionKey::CreateFromPath(Path.Get<0>()), Settings->ActorSelector.bMustOverlapSelf);
+		}
+
+		DynamicTracking.Finalize(Context);
+	}
+	else if (Context->IsValueOverriden(GET_MEMBER_NAME_CHECKED(UPCGGetActorPropertySettings, ActorSelector)))
+	{
+		FPCGDynamicTrackingHelper::AddSingleDynamicTrackingKey(Context, Settings->ActorSelector);
+	}
+#endif // WITH_EDITOR
 
 	if (FoundActors.IsEmpty())
 	{
@@ -207,9 +281,9 @@ bool FPCGGetActorPropertyElement::ExecuteInternal(FPCGContext* Context) const
 		{
 			FPCGDynamicTrackingHelper DynamicTracking;
 			DynamicTracking.EnableAndInitialize(Context, ObjectTraversed.Num());
-			for (FSoftObjectPath& Path : ObjectTraversed)
+			for (const FSoftObjectPath& Path : ObjectTraversed)
 			{
-				DynamicTracking.AddToTracking(FPCGSelectionKey::CreateFromPath(std::move(Path)), /*bCulled=*/false);
+				DynamicTracking.AddToTracking(FPCGSelectionKey::CreateFromPath(Path), /*bCulled=*/false);
 			}
 
 			DynamicTracking.Finalize(Context);
diff --git a/Engine/Plugins/PCG/Source/PCG/Private/Helpers/PCGDynamicTrackingHelpers.cpp b/Engine/Plugins/PCG/Source/PCG/Private/Helpers/PCGDynamicTrackingHelpers.cpp
--- a/Engine/Plugins/PCG/Source/PCG/Private/Helpers/PCGDynamicTrackingHelpers.cpp
+++ b/Engine/Plugins/PCG/Source/PCG/Private/Helpers/PCGDynamicTrackingHelpers.cpp
@@ -52,4 +52,9 @@ void FPCGDynamicTrackingHelper::AddSingleDynamicTrackingKey(FPCGContext* InConte
 	}
 }
 
+void FPCGDynamicTrackingHelper::AddSingleDynamicTrackingKey(FPCGContext* InContext, const FPCGActorSelectorSettings& InSelector)
+{
+	AddSingleDynamicTrackingKey(InContext, FPCGSelectionKey(InSelector), InSelector.bMustOverlapSelf);
+}
+
 #endif // WITH_EDITOR
diff --git a/Engine/Plugins/PCG/Source/PCG/Private/PCGContext.cpp b/Engine/Plugins/PCG/Source/PCG/Private/PCGContext.cpp
--- a/Engine/Plugins/PCG/Source/PCG/Private/PCGContext.cpp
+++ b/Engine/Plugins/PCG/Source/PCG/Private/PCGContext.cpp
@@ -269,9 +269,26 @@ void FPCGContext::OverrideSettings()
 						}
 						else if constexpr (std::is_same_v<FSoftClassPath, PropertyType>)
 						{
-							if (const UClass* Class = Cast<UClass>(Value.ResolveObject()))
+							UClass* Subclass = nullptr;
+							if (const FClassProperty* ClassProp = CastField<const FClassProperty>(Param.Properties.Last()))
 							{
-								bInvalid = ObjectProperty->PropertyClass && !Class->IsChildOf(ObjectProperty->PropertyClass);
+								Subclass = ClassProp->MetaClass;
+							}
+							else if (const FSoftClassProperty* SoftClassProp = CastField<const FSoftClassProperty>(Param.Properties.Last()))
+							{
+								Subclass = SoftClassProp->MetaClass;
+							}
+							else
+							{
+								// TODO: should we use prop -> GetOwnerProperty()->GetClassMetadata(TEXT("MetaClass")) ?
+								bInvalid = true;
+							}
+
+							const UClass* ValueClass = (bInvalid ? nullptr : Cast<UClass>(Value.ResolveObject()));
+
+							if (ValueClass)
+							{
+								bInvalid = !Subclass || !ValueClass->IsChildOf(Subclass);
 							}
 						}
 
diff --git a/Engine/Plugins/PCG/Source/PCG/Public/Elements/PCGActorSelector.h b/Engine/Plugins/PCG/Source/PCG/Public/Elements/PCGActorSelector.h
--- a/Engine/Plugins/PCG/Source/PCG/Public/Elements/PCGActorSelector.h
+++ b/Engine/Plugins/PCG/Source/PCG/Public/Elements/PCGActorSelector.h
@@ -2,6 +2,8 @@
 
 #pragma once
 
+#include "Metadata/PCGAttributePropertySelector.h"
+
 #include "Templates/SubclassOf.h"
 #include "UObject/NameTypes.h"
 
@@ -10,6 +12,7 @@
 class AActor;
 class UPCGComponent;
 class UWorld;
+struct FPCGActorSelectorSettings;
 
 UENUM()
 enum class EPCGActorSelection : uint8
@@ -35,6 +38,8 @@ enum class EPCGActorFilter : uint8
 	AllWorldActors,
 	/** The source PCG actor (rather than the generated partition actor). */
 	Original,
+	/** Consider only the provided list of actors */
+	FromInput
 };
 
 /**
@@ -49,9 +54,11 @@ struct PCG_API FPCGSelectionKey
 
 	FPCGSelectionKey() = default;
 
+	// For all filters except FromInput
+	explicit FPCGSelectionKey(const FPCGActorSelectorSettings& InActorSelector);
+
 	// For all filters others than AllWorldActor. For AllWorldActors Filter, use the other constructors.
 	explicit FPCGSelectionKey(EPCGActorFilter InFilter);
-
 	explicit FPCGSelectionKey(FName InTag);
 	explicit FPCGSelectionKey(TSubclassOf<UObject> InSelectionClass);
 
@@ -109,7 +116,7 @@ struct PCG_API FPCGActorSelectorSettings
 	bool bIncludeChildren = false;
 
 	/** Enables/disables fine-grained actor filtering options. */
-	UPROPERTY(BlueprintReadWrite, EditAnywhere, Category = "Actor Selector Settings", meta = (EditCondition = "ActorFilter!=EPCGActorFilter::AllWorldActors && bIncludeChildren", EditConditionHides))
+	UPROPERTY(BlueprintReadWrite, EditAnywhere, Category = "Actor Selector Settings", meta = (PCG_NotOverridable, EditCondition = "ActorFilter!=EPCGActorFilter::AllWorldActors && bIncludeChildren", EditConditionHides))
 	bool bDisableFilter = false;
 
 	/** How to select when filtering actors. */
@@ -124,6 +131,10 @@ struct PCG_API FPCGActorSelectorSettings
 	UPROPERTY(BlueprintReadWrite, EditAnywhere, Category = "Actor Selector Settings", meta = (EditCondition = "bShowActorSelectionClass && bShowActorSelection && (ActorFilter==EPCGActorFilter::AllWorldActors || (bIncludeChildren && !bDisableFilter)) && ActorSelection==EPCGActorSelection::ByClass", EditConditionHides, AllowAbstract = "true"))
 	TSubclassOf<AActor> ActorSelectionClass;
 
+	/** Controls what attribute to read from when the actor selector uses the "FromInput" actor filter. */
+	UPROPERTY(BlueprintReadWrite, EditAnywhere, Category = "Actor Selector Settings", meta = (EditCondition = "ActorFilter==EPCGActorFilter::FromInput", EditConditionHides))
+	FPCGAttributePropertyInputSelector ActorReferenceSelector;
+
 	/** If true processes all matching actors, otherwise returns data from first match. */
 	UPROPERTY(BlueprintReadWrite, EditAnywhere, Category = "Actor Selector Settings", meta = (EditCondition = "bShowSelectMultiple && ActorFilter==EPCGActorFilter::AllWorldActors && ActorSelection!=EPCGActorSelection::ByName", EditConditionHides))
 	bool bSelectMultiple = false;
@@ -162,6 +173,6 @@ struct PCG_API FPCGActorSelectorSettings
 
 namespace PCGActorSelector
 {
-	PCG_API TArray<AActor*> FindActors(const FPCGActorSelectorSettings& Settings, const UPCGComponent* InComponent, const TFunction<bool(const AActor*)>& BoundsCheck, const TFunction<bool(const AActor*)>& SelfIgnoreCheck);
-	PCG_API AActor* FindActor(const FPCGActorSelectorSettings& InSettings, UPCGComponent* InComponent, const TFunction<bool(const AActor*)>& BoundsCheck, const TFunction<bool(const AActor*)>& SelfIgnoreCheck);
+	PCG_API TArray<AActor*> FindActors(const FPCGActorSelectorSettings& Settings, const UPCGComponent* InComponent, const TFunction<bool(const AActor*)>& BoundsCheck, const TFunction<bool(const AActor*)>& SelfIgnoreCheck, TArrayView<AActor*> InputActors = TArrayView<AActor*>());
+	PCG_API AActor* FindActor(const FPCGActorSelectorSettings& InSettings, UPCGComponent* InComponent, const TFunction<bool(const AActor*)>& BoundsCheck, const TFunction<bool(const AActor*)>& SelfIgnoreCheck, TArrayView<AActor*> InputActors = TArrayView<AActor*>());
 }
diff --git a/Engine/Plugins/PCG/Source/PCG/Public/Elements/PCGDataFromActor.h b/Engine/Plugins/PCG/Source/PCG/Public/Elements/PCGDataFromActor.h
--- a/Engine/Plugins/PCG/Source/PCG/Public/Elements/PCGDataFromActor.h
+++ b/Engine/Plugins/PCG/Source/PCG/Public/Elements/PCGDataFromActor.h
@@ -5,6 +5,7 @@
 #include "PCGContext.h"
 #include "PCGSettings.h"
 #include "Elements/PCGActorSelector.h"
+#include "Elements/PCGLoadObjectsContext.h"
 
 #include "UObject/ObjectKey.h"
 
@@ -35,6 +36,7 @@ public:
 	virtual FText GetNodeTooltipText() const override;
 	virtual EPCGSettingsType GetType() const override { return EPCGSettingsType::Spatial; }
 	virtual void GetStaticTrackedKeys(FPCGSelectionKeyToSettingsMap& OutKeysToSettings, TArray<TObjectPtr<const UPCGGraph>>& OutVisitedGraphs) const override;
+	virtual bool CanDynamicallyTrackKeys() const override { return true; }
 	virtual bool HasDynamicPins() const override { return true; }
 	virtual void ApplyDeprecation(UPCGNode* InOutNode) override;
 #endif
@@ -46,7 +48,7 @@ protected:
 #if WITH_EDITOR
 	virtual EPCGChangeType GetChangeTypeForProperty(const FName& InPropertyName) const override { return Super::GetChangeTypeForProperty(InPropertyName) | EPCGChangeType::Cosmetic; }
 #endif
-	virtual TArray<FPCGPinProperties> InputPinProperties() const override { return TArray<FPCGPinProperties>(); }
+	virtual TArray<FPCGPinProperties> InputPinProperties() const override;
 	virtual TArray<FPCGPinProperties> OutputPinProperties() const override;
 
 	virtual FPCGElementPtr CreateElement() const override;
@@ -75,7 +77,7 @@ protected:
 
 public:
 	/** Describes which actors to select for data collection. */
-	UPROPERTY(BlueprintReadWrite, EditAnywhere, Category = Settings, meta = (ShowOnlyInnerProperties))
+	UPROPERTY(BlueprintReadWrite, EditAnywhere, Category = Settings, meta = (PCG_Overridable, ShowOnlyInnerProperties))
 	FPCGActorSelectorSettings ActorSelector;
 
 	/** Describes what kind of data we will collect from the found actor(s). */
@@ -125,7 +127,7 @@ public:
 	bool bDisplayModeSettings = true;
 };
 
-struct FPCGDataFromActorContext : public FPCGContext
+struct FPCGDataFromActorContext : public FPCGLoadObjectsFromPathContext
 {
 	TArray<AActor*> FoundActors;
 	bool bPerformedQuery = false;
@@ -136,7 +138,7 @@ struct FPCGDataFromActorContext : public FPCGContext
 #endif
 };
 
-class PCG_API FPCGDataFromActorElement : public IPCGElement
+class PCG_API FPCGDataFromActorElement : public IPCGElementWithCustomContext<FPCGDataFromActorContext>
 {
 public:
 	virtual bool CanExecuteOnlyOnMainThread(FPCGContext* Context) const override { return true; }
@@ -144,7 +146,7 @@ public:
 	virtual bool ShouldComputeFullOutputDataCrc(FPCGContext* Context) const override { return true; }
 
 protected:
-	virtual FPCGContext* CreateContext() override;
+	virtual bool PrepareDataInternal(FPCGContext* Context) const override;
 	virtual bool ExecuteInternal(FPCGContext* Context) const override;
 	void GatherWaitTasks(AActor* FoundActor, FPCGContext* InContext, TArray<FPCGTaskId>& OutWaitTasks) const;
 	virtual void ProcessActors(FPCGContext* Context, const UPCGDataFromActorSettings* Settings, const TArray<AActor*>& FoundActors) const;
diff --git a/Engine/Plugins/PCG/Source/PCG/Public/Elements/PCGGetActorProperty.h b/Engine/Plugins/PCG/Source/PCG/Public/Elements/PCGGetActorProperty.h
--- a/Engine/Plugins/PCG/Source/PCG/Public/Elements/PCGGetActorProperty.h
+++ b/Engine/Plugins/PCG/Source/PCG/Public/Elements/PCGGetActorProperty.h
@@ -4,6 +4,7 @@
 
 #include "PCGSettings.h"
 #include "Elements/PCGActorSelector.h"
+#include "Elements/PCGLoadObjectsContext.h"
 #include "Metadata/PCGMetadataAttribute.h"
 
 #include "PCGGetActorProperty.generated.h"
@@ -33,11 +34,13 @@ public:
 	virtual EPCGSettingsType GetType() const override { return EPCGSettingsType::Param; }
 	virtual void GetStaticTrackedKeys(FPCGSelectionKeyToSettingsMap& OutKeysToSettings, TArray<TObjectPtr<const UPCGGraph>>& OutVisitedGraphs) const override;
 	virtual bool CanDynamicallyTrackKeys() const override { return true; }
+	virtual bool HasDynamicPins() const override { return true; }
 #endif
+	virtual EPCGDataType GetCurrentPinTypes(const UPCGPin* InPin) const override;
 
 	virtual FString GetAdditionalTitleInformation() const override;
 	virtual TArray<FPCGPinProperties> OutputPinProperties() const override;
-	virtual TArray<FPCGPinProperties> InputPinProperties() const override { return TArray<FPCGPinProperties>(); }
+	virtual TArray<FPCGPinProperties> InputPinProperties() const override;
 
 protected:
 #if WITH_EDITOR
@@ -47,7 +50,7 @@ protected:
 	//~End UPCGSettings interface
 
 public:
-	UPROPERTY(BlueprintReadWrite, EditAnywhere, Category = Settings, meta = (ShowOnlyInnerProperties))
+	UPROPERTY(BlueprintReadWrite, EditAnywhere, Category = Settings, meta = (PCG_Overridable, ShowOnlyInnerProperties))
 	FPCGActorSelectorSettings ActorSelector;
 
 	/** Allow to look for an actor component instead of an actor. It will need to be attached to the found actor. */
@@ -106,11 +109,12 @@ private:
 #endif
 };
 
-class FPCGGetActorPropertyElement : public IPCGElement
+class FPCGGetActorPropertyElement : public IPCGElementWithCustomContext<FPCGLoadObjectsFromPathContext>
 {
 public:
 	virtual bool CanExecuteOnlyOnMainThread(FPCGContext* Context) const override { return true; }
 	virtual bool IsCacheable(const UPCGSettings* InSettings) const override { return !CastChecked<UPCGGetActorPropertySettings>(InSettings)->bAlwaysRequeryActors; }
 protected:
+	virtual bool PrepareDataInternal(FPCGContext* Context) const override;
 	virtual bool ExecuteInternal(FPCGContext* Context) const override;
 };
diff --git a/Engine/Plugins/PCG/Source/PCG/Public/Helpers/PCGDynamicTrackingHelpers.h b/Engine/Plugins/PCG/Source/PCG/Public/Helpers/PCGDynamicTrackingHelpers.h
--- a/Engine/Plugins/PCG/Source/PCG/Public/Helpers/PCGDynamicTrackingHelpers.h
+++ b/Engine/Plugins/PCG/Source/PCG/Public/Helpers/PCGDynamicTrackingHelpers.h
@@ -29,6 +29,9 @@ public:
 	/** Convenient function to push just a single tracking key to the component. */
 	static void AddSingleDynamicTrackingKey(FPCGContext* InContext, FPCGSelectionKey&& InKey, bool bIsCulled);
 
+	/** Convenience function to push just a single selector as a tracking key to the component. */
+	static void AddSingleDynamicTrackingKey(FPCGContext* InContext, const FPCGActorSelectorSettings& InSelector);
+
 private:
 	bool bDynamicallyTracked = false;
 	TWeakObjectPtr<UPCGComponent> CachedComponent;
