Improved behavior of create collision data / self-pruning with collision, when the static meshes don't have a valid collision setup.

#rnx

diff --git a/Engine/Plugins/PCG/Source/PCG/Private/Data/PCGCollisionWrapperData.cpp b/Engine/Plugins/PCG/Source/PCG/Private/Data/PCGCollisionWrapperData.cpp
--- a/Engine/Plugins/PCG/Source/PCG/Private/Data/PCGCollisionWrapperData.cpp
+++ b/Engine/Plugins/PCG/Source/PCG/Private/Data/PCGCollisionWrapperData.cpp
@@ -9,6 +9,7 @@
 #include "Components/StaticMeshComponent.h"
 #include "Engine/StaticMesh.h"
 #include "PhysicsEngine/BodyInstance.h"
+#include "PhysicsEngine/BodySetup.h"
 #include "Serialization/ArchiveCrc32.h"
 
 #include UE_INLINE_GENERATED_CPP_BY_NAME(PCGCollisionWrapperData)
@@ -23,7 +24,7 @@ void FPCGCollisionWrapper::Uninitialize()
 	// Implementation note: we do the full uninitialize even if we think we're not initialized due to the default move operators
 	for (FBodyInstance* BodyInstance : BodyInstances)
 	{
-		if (BodyInstance->IsValidBodyInstance())
+		if (BodyInstance && BodyInstance->IsValidBodyInstance())
 		{
 			BodyInstance->TermBody();
 		}
@@ -90,13 +91,20 @@ void FPCGCollisionWrapper::CreateBodyInstances(const TArray<FSoftObjectPath>& Me
 			MeshPtr = Mesh.LoadSynchronous();
 		}
 
-		if(MeshPtr && MeshPtr->GetBodySetup())
+		if (MeshPtr && MeshPtr->GetBodySetup() && MeshPtr->GetBodySetup()->AggGeom.GetElementCount() > 0)
 		{
 			//TRACE_CPUPROFILER_EVENT_SCOPE(FPCGCollisionWrapper::CreateBodyInstances::CreateBodyInstance);
 			BodyInstance = new FBodyInstance();
 			BodyInstance->bAutoWeld = false;
 			BodyInstance->bSimulatePhysics = false;
 			BodyInstance->InitBody(MeshPtr->GetBodySetup(), FTransform::Identity, nullptr, nullptr);
+
+			if (!BodyInstance->IsValidBodyInstance() || !ensure(BodyInstance->GetBodySetup()))
+			{
+				// Some part of the initialization process failed.
+				delete BodyInstance;
+				BodyInstance = nullptr;
+			}
 		}
 
 		BodyInstances.Add(BodyInstance);
