[PCG] Implement ApplyOnMultiAccessors

* Generalization of ApplyOnAccessor, for multiple types at the same time.
* Will need it for grammar stuff

#rnx

diff --git a/Engine/Plugins/PCG/Source/PCG/Public/Elements/Metadata/PCGMetadataElementCommon.h b/Engine/Plugins/PCG/Source/PCG/Public/Elements/Metadata/PCGMetadataElementCommon.h
--- a/Engine/Plugins/PCG/Source/PCG/Public/Elements/Metadata/PCGMetadataElementCommon.h
+++ b/Engine/Plugins/PCG/Source/PCG/Public/Elements/Metadata/PCGMetadataElementCommon.h
@@ -8,6 +8,9 @@
 #include "Metadata/Accessors/IPCGAttributeAccessor.h" // IWYU pragma: keep
 #include "Metadata/Accessors/PCGAttributeAccessorKeys.h"
 
+#include "Misc/ScopeExit.h" // For ON_SCOPE_EXIT
+#include "Templates/MemoryOps.h" // For DestructItem
+
 class IPCGAttributeAccessor;
 template <typename T> class FPCGMetadataAttribute;
 
@@ -113,6 +116,128 @@ namespace PCGMetadataElementCommon
 		return ApplyOnAccessorRange<T>(Keys, Accessor, std::move(RangeCallback), Flags, ChunkSize, Count);
 	}
 
+	/**
+	* Iterate over the full range of the keys (if Count is negative, otherwise, as many times as Count), calling the callback with values ranges get from the multiple accessors.
+	* Keys needs to be the same for all accessors.
+	* ChunkSize influence the max number of values to get in one go with GetRange.
+	* Callback should have a signature: void(const TArrayView<T>&... Args, int32 Index)
+	* Return false if it process nothing.
+	*/
+	template <typename... T, typename Func = TFunction<void(const TArrayView<T>&..., int32 Index)>>
+	bool ApplyOnMultiAccessorsRange(const IPCGAttributeAccessorKeys& Keys, const TArray<const IPCGAttributeAccessor*>& Accessors, Func&& InCallback, EPCGAttributeAccessorFlags Flags = EPCGAttributeAccessorFlags::StrictType, const int32 ChunkSize = DefaultChunkSize, const int32 Count = -1)
+	{
+		// We use C++20 functionality here, so make sure we are at this version at least
+		static_assert(__cplusplus >= 202002L);
+
+		const int32 NumberOfEntries = Count < 0 ? Keys.GetNum() : Count;
+
+		if (NumberOfEntries == 0)
+		{
+			return false;
+		}
+
+		constexpr int32 NumArgs = static_cast<int32>(sizeof...(T));
+
+		// Buffers for our get range cannot be stored in a TArray, since we have different T types. So buffers will be raw buffers of bytes.
+		// Since every type can be different, we pre-allocate chunk size times the size of FVector for all, which should cover most usecases.
+		using RawBuffer = TArray<uint8, TInlineAllocator<DefaultChunkSize * sizeof(FVector)>>;
+		RawBuffer RawBuffers[NumArgs]{};
+
+		// First fold expression: Initialize all the buffers. We unfold T to initialize the buffer to 
+		// ChunkSize element of type T. If T is not POD, we need to construct the object, so we do a placement new.
+		int i = 0;
+		([&]
+		{
+			RawBuffers[i].SetNumUninitialized(sizeof(T) * ChunkSize);
+			if constexpr (!std::is_trivially_copyable_v<T>)
+			{
+				for (int32 j = 0; j < ChunkSize; ++j)
+				{
+					new(reinterpret_cast<T*>(RawBuffers[i].GetData()) + j) T;
+				}
+			}
+			++i;
+		}(), ...);
+
+		ON_SCOPE_EXIT
+		{
+			// Second fold expression, to be executed when we exit the function, so that if we constructed our objects while allocating the buffer
+			// we need to explicitly call their destructor. Since we are in a fold expression, we can't use `.~T` as ~T is not recognized as a function name,
+			// so we use DestructItem.
+			i = 0;
+			([&]
+			{
+				if constexpr (!std::is_trivially_copyable_v<T>)
+				{
+					for (int32 j = 0; j < ChunkSize; ++j)
+					{
+						DestructItem(reinterpret_cast<T*>(RawBuffers[i].GetData()) + j);
+					}
+				}
+				++i;
+			}(), ...);
+		};
+
+		const int32 NumberOfIterations = (NumberOfEntries + ChunkSize - 1) / ChunkSize;
+
+		for (int32 j = 0; j < NumberOfIterations; ++j)
+		{
+			const int32 StartIndex = j * ChunkSize;
+			const int32 Range = FMath::Min(NumberOfEntries - StartIndex, ChunkSize);
+
+			// Third fold expression, that will unroll T to get all ranges.
+			// Since we manage our memory ourselves, it is safe to do a reinterpret_cast here.
+			bool bSuccess = true;
+			i = 0;
+			(
+			[&]
+			{
+				TArrayView<T> View(reinterpret_cast<T*>(RawBuffers[i].GetData()), Range);
+				bSuccess &= Accessors[i]->GetRange<T>(View, StartIndex, Keys, Flags);
+				++i;
+			}(), ...);
+
+			if (!bSuccess)
+			{
+				return false;
+			}
+
+			// Finally, we need to pass the views to the callback. We'll have to have an indirection by first unpacking T to create the views and
+			// then unpack again the views to the callback.
+			// We can't do that with a fold expression, so we need to use a Tuple. We create a Tuple of views, and use a lambda unpacker to construct our views
+			// then we unpack it again when calling our callback. ApplyBefore will first unpack the views and then pass the rest of the parameters (here StartIndex and Range)
+			using TupleType = TTuple<TArrayView<T>...>;
+			TupleType TempTuple = [&]<std::size_t ...I>(std::index_sequence<I...>)
+			{
+				return TupleType(TArrayView<T>(reinterpret_cast<T*>(RawBuffers[I].GetData()), Range)...);
+			}(std::make_index_sequence<NumArgs>{});
+
+			TempTuple.ApplyBefore(InCallback, StartIndex, Range);
+		}
+
+		return true;
+	}
+
+	/**
+	* Iterate over the full range of the keys (if Count is negative, otherwise, as many times as Count), calling the callback with values get from multiple accessor.
+	* ChunkSize influence the max number of values to get in one go with GetRange.
+	* Callback should have a signature: void(const T&... Values, int32 Index)
+	* Return false if it process nothing.
+	*/
+	template <typename ...T, typename Func = TFunction<void(const T&..., int32 Index)>>
+	bool ApplyOnMultiAccessors(const IPCGAttributeAccessorKeys& Keys, const TArray<const IPCGAttributeAccessor*>& Accessors, Func&& InCallback, EPCGAttributeAccessorFlags Flags = EPCGAttributeAccessorFlags::StrictType, const int32 ChunkSize = DefaultChunkSize, const int32 Count = -1)
+	{
+		auto RangeCallback = [Callback = std::forward<Func>(InCallback)](const TArrayView<T>& ...Views, int32 Start, int32 Range)
+		{
+			for (int32 j = 0; j < Range; ++j)
+			{
+				Callback(Views[j]..., Start + j);
+			}
+		};
+
+		return ApplyOnMultiAccessorsRange<T...>(Keys, Accessors, std::move(RangeCallback), Flags, ChunkSize, Count);
+	}	
+
 	/**
 	* Read all the values from in accessor and write to out accessor. Iterate as many times than the out accessor keys
 	*/
