[PCG] Fix trace scope name

#rnx

diff --git a/Engine/Plugins/PCG/Source/PCG/Private/Elements/Metadata/PCGMetadataRenameElement.cpp b/Engine/Plugins/PCG/Source/PCG/Private/Elements/Metadata/PCGMetadataRenameElement.cpp
--- a/Engine/Plugins/PCG/Source/PCG/Private/Elements/Metadata/PCGMetadataRenameElement.cpp
+++ b/Engine/Plugins/PCG/Source/PCG/Private/Elements/Metadata/PCGMetadataRenameElement.cpp
@@ -41,7 +41,7 @@ FPCGElementPtr UPCGMetadataRenameSettings::CreateElement() const
 
 bool FPCGMetadataRenameElement::ExecuteInternal(FPCGContext* Context) const
 {
-	TRACE_CPUPROFILER_EVENT_SCOPE(FPCGMetadataOperationElement::Execute);
+	TRACE_CPUPROFILER_EVENT_SCOPE(FPCGMetadataRenameElement::Execute);
 
 	const UPCGMetadataRenameSettings* Settings = Context->GetInputSettings<UPCGMetadataRenameSettings>();
 	check(Settings);
