Removed unneeded data duplication in Get Actor Data since we've improved resource tracking

diff --git a/Engine/Plugins/PCG/Source/PCG/Private/Elements/PCGDataFromActor.cpp b/Engine/Plugins/PCG/Source/PCG/Private/Elements/PCGDataFromActor.cpp
--- a/Engine/Plugins/PCG/Source/PCG/Private/Elements/PCGDataFromActor.cpp
+++ b/Engine/Plugins/PCG/Source/PCG/Private/Elements/PCGDataFromActor.cpp
@@ -782,20 +782,14 @@ void FPCGDataFromActorElement::ProcessActor(FPCGContext* Context, const UPCGData
 	{
 		for (UPCGComponent* Component : PCGComponents)
 		{
-			// TODO - Temporary behavior
-			// At the moment, intersections that reside in the transient package can hold on to a reference on this data
-			// which prevents proper garbage collection on map change, hence why we duplicate here. Normally, we would expect
-			// this not to be a problem, as these intersections should be garbage collected, but this requires more investigation
 			for (const FPCGTaggedData& TaggedData : Component->GetGeneratedGraphOutput().TaggedData)
 			{
 				if (ensure(TaggedData.Data))
 				{
 					FPCGTaggedData& DuplicatedTaggedData = Outputs.Add_GetRef(TaggedData);
-					DuplicatedTaggedData.Data = Cast<UPCGData>(StaticDuplicateObject(TaggedData.Data, GetTransientPackage()));
 					DuplicatedTaggedData.Tags.Add(PCGDataFromActorConstants::PCGComponentDataGridSizeTagPrefix + FString::FromInt(PCGHiGenGrid::GridToGridSize(Component->GetGenerationGrid())));
 				}
 			}
-			//Outputs.Append(Component->GetGeneratedGraphOutput().TaggedData);
 		}
 	}
 	else if (FoundProperty)
