Added node to create components on target actor(s)

diff --git a/Engine/Plugins/PCG/Source/PCG/Private/Elements/PCGAddComponent.cpp b/Engine/Plugins/PCG/Source/PCG/Private/Elements/PCGAddComponent.cpp
new file mode 100644
--- /dev/null
+++ b/Engine/Plugins/PCG/Source/PCG/Private/Elements/PCGAddComponent.cpp
@@ -0,0 +1,524 @@
+// Copyright Epic Games, Inc. All Rights Reserved.
+
+#include "Elements/PCGAddComponent.h"
+
+#include "PCGComponent.h"
+#include "PCGManagedResource.h"
+#include "Elements/Metadata/PCGMetadataElementCommon.h"
+#include "Helpers/PCGHelpers.h"
+#include "Metadata/Accessors/PCGAttributeAccessorHelpers.h"
+#include "Metadata/Accessors/PCGCustomAccessor.h"
+
+#include "Components/SceneComponent.h"
+#include "Engine/Blueprint.h"
+#include "Engine/BlueprintGeneratedClass.h"
+#include "Engine/Engine.h"
+#include "UObject/Package.h"
+
+#define LOCTEXT_NAMESPACE "PCGAddComponentElement"
+
+UPCGAddComponentSettings::UPCGAddComponentSettings(const FObjectInitializer& ObjectInitializer)
+	: Super(ObjectInitializer)
+{
+	ComponentReferenceAttribute.SetAttributeName(PCGAddComponentConstants::ComponentReferenceAttribute);
+}
+
+void UPCGAddComponentSettings::PostLoad()
+{
+	Super::PostLoad();
+
+#if WITH_EDITOR
+	SetupOnBlueprintChangedEvent();
+
+	if (TemplateComponentClass && TemplateComponent)
+	{
+		TemplateComponent->ConditionalPostLoad();
+	}
+
+	RefreshTemplateComponent();
+#endif
+}
+
+void UPCGAddComponentSettings::BeginDestroy()
+{
+#if WITH_EDITOR
+	TeardownOnBlueprintChangedEvent();
+#endif // WITH_EDITOR
+
+	Super::BeginDestroy();
+}
+
+#if WITH_EDITOR
+void UPCGAddComponentSettings::PreEditChange(FProperty* PropertyAboutToChange)
+{
+	Super::PreEditChange(PropertyAboutToChange);
+
+	if (PropertyAboutToChange && PropertyAboutToChange->GetFName() == GET_MEMBER_NAME_CHECKED(UPCGAddComponentSettings, TemplateComponentClass))
+	{
+		TeardownOnBlueprintChangedEvent();
+	}
+}
+
+void UPCGAddComponentSettings::PostEditChangeProperty(FPropertyChangedEvent& PropertyChangedEvent)
+{
+	if (PropertyChangedEvent.Property)
+	{
+		const FName& PropertyName = PropertyChangedEvent.Property->GetFName();
+
+		if (PropertyName == GET_MEMBER_NAME_CHECKED(UPCGAddComponentSettings, TemplateComponentClass))
+		{
+			SetupOnBlueprintChangedEvent();
+			RefreshTemplateComponent();
+		}
+		else if (PropertyName == GET_MEMBER_NAME_CHECKED(UPCGAddComponentSettings, bAllowTemplateComponentEditing))
+		{
+			RefreshTemplateComponent();
+		}
+	}
+
+	Super::PostEditChangeProperty(PropertyChangedEvent);
+}
+
+void UPCGAddComponentSettings::PreEditUndo()
+{
+	TeardownOnBlueprintChangedEvent();
+	Super::PreEditUndo();
+}
+
+void UPCGAddComponentSettings::PostEditUndo()
+{
+	Super::PostEditUndo();
+	SetupOnBlueprintChangedEvent();
+	RefreshTemplateComponent();
+}
+#endif // WITH_EDITOR
+
+TArray<FPCGPinProperties> UPCGAddComponentSettings::InputPinProperties() const
+{
+	TArray<FPCGPinProperties> PinProperties;
+	FPCGPinProperties& InputPin = PinProperties.Emplace_GetRef(PCGPinConstants::DefaultInputLabel, EPCGDataType::Point);
+	InputPin.SetRequiredPin();
+
+	FPCGPinProperties& TargetActorPin = PinProperties.Emplace_GetRef(PCGAddComponentConstants::TargetActorPinLabel, EPCGDataType::Any);
+	TargetActorPin.SetAdvancedPin();
+
+	return PinProperties;
+}
+
+TArray<FPCGPinProperties> UPCGAddComponentSettings::OutputPinProperties() const
+{
+	return Super::DefaultPointOutputPinProperties();
+}
+
+FPCGElementPtr UPCGAddComponentSettings::CreateElement() const
+{
+	return MakeShared<FPCGAddComponentElement>();
+}
+
+#if WITH_EDITOR
+void UPCGAddComponentSettings::OnBlueprintChanged(UBlueprint* InBlueprint)
+{
+	RefreshTemplateComponent();
+	DirtyCache();
+	OnSettingsChangedDelegate.Broadcast(this, EPCGChangeType::Settings);
+}
+
+void UPCGAddComponentSettings::SetupOnBlueprintChangedEvent()
+{
+	if (UBlueprintGeneratedClass* BlueprintClass = Cast<UBlueprintGeneratedClass>(TemplateComponentClass))
+	{
+		if (UBlueprint* Blueprint = Cast<UBlueprint>(BlueprintClass->ClassGeneratedBy))
+		{
+			Blueprint->OnChanged().AddUObject(this, &UPCGAddComponentSettings::OnBlueprintChanged);
+		}
+	}
+}
+
+void UPCGAddComponentSettings::TeardownOnBlueprintChangedEvent()
+{
+	if (UBlueprintGeneratedClass* BlueprintClass = Cast<UBlueprintGeneratedClass>(TemplateComponentClass))
+	{
+		if (UBlueprint* Blueprint = Cast<UBlueprint>(BlueprintClass->ClassGeneratedBy))
+		{
+			Blueprint->OnChanged().RemoveAll(this);
+		}
+	}
+}
+
+void UPCGAddComponentSettings::RefreshTemplateComponent()
+{
+	// Implementation note: this is similar to the child actor component implementation
+	if (TemplateComponentClass && bAllowTemplateComponentEditing)
+	{
+		const bool bCreateNewTemplateComponent = (!TemplateComponent || TemplateComponent->GetClass() != TemplateComponentClass);
+
+		if (bCreateNewTemplateComponent)
+		{
+			UActorComponent* NewTemplateComponent = NewObject<UActorComponent>(GetTransientPackage(), TemplateComponentClass, NAME_None, RF_ArchetypeObject | RF_Transactional | RF_Public);
+
+			if (TemplateComponent)
+			{
+				UEngine::FCopyPropertiesForUnrelatedObjectsParams Options;
+				Options.bNotifyObjectReplacement = true;
+				UEngine::CopyPropertiesForUnrelatedObjects(TemplateComponent, NewTemplateComponent, Options);
+
+				TemplateComponent->Rename(nullptr, GetTransientPackage(), REN_DontCreateRedirectors | REN_ForceNoResetLoaders);
+
+				TMap<UObject*, UObject*> OldToNew;
+				OldToNew.Emplace(TemplateComponent, NewTemplateComponent);
+				GEngine->NotifyToolsOfObjectReplacement(OldToNew);
+
+				TemplateComponent->MarkAsGarbage();
+			}
+
+			TemplateComponent = NewTemplateComponent;
+
+			// Record initial object state in case we're in a transaction context.
+			TemplateComponent->Modify();
+
+			// Outer to this object
+			TemplateComponent->Rename(nullptr, this, REN_DoNotDirty | REN_DontCreateRedirectors | REN_ForceNoResetLoaders);
+		}
+	}
+	else
+	{
+		if (TemplateComponent)
+		{
+			TemplateComponent->MarkAsGarbage();
+		}
+
+		TemplateComponent = nullptr;
+	}
+}
+
+#endif // WITH_EDITOR
+
+void UPCGAddComponentSettings::SetTemplateComponentClass(const TSubclassOf<UActorComponent>& InTemplateComponentClass)
+{
+#if WITH_EDITOR
+	TeardownOnBlueprintChangedEvent();
+#endif
+
+	TemplateComponentClass = InTemplateComponentClass;
+
+#if WITH_EDITOR
+	SetupOnBlueprintChangedEvent();
+	RefreshTemplateComponent();
+#endif
+}
+
+void UPCGAddComponentSettings::SetAllowTemplateComponentEditing(bool bInAllowTemplateComponentEditing)
+{
+	bAllowTemplateComponentEditing = bInAllowTemplateComponentEditing;
+
+#if WITH_EDITOR
+	RefreshTemplateComponent();
+#endif
+}
+
+bool FPCGAddComponentElement::ExecuteInternal(FPCGContext* InContext) const
+{
+	TRACE_CPUPROFILER_EVENT_SCOPE(FPCGAddComponentElement::Execute);
+
+	check(InContext);
+	FPCGAddComponentContext* Context = static_cast<FPCGAddComponentContext*>(InContext);
+
+	const UPCGAddComponentSettings* Settings = Context->GetInputSettings<UPCGAddComponentSettings>();
+	check(Settings);
+
+	// For all entries on the input pin, create one actor component.
+	// Add that component either to a default value (this), or a value from the 'Target Actor' pin, which should be
+	// either the same cardinality or only one entry.
+	TArray<FPCGTaggedData> Inputs = Context->InputData.GetInputsByPin(PCGPinConstants::DefaultInputLabel);
+	TArray<FPCGTaggedData> TargetActors = Context->InputData.GetInputsByPin(PCGAddComponentConstants::TargetActorPinLabel);
+
+	// Early outs & validation
+	if (Inputs.Num() == 0 || !Context->SourceComponent.Get())
+	{
+		return true;
+	}
+
+	if (TargetActors.Num() > 1 && TargetActors.Num() != Inputs.Num())
+	{
+		PCGE_LOG(Error, GraphAndLog, FText::Format(LOCTEXT("MismatchInputsAndTargetActors", "Expected either a single input for target actor or the same number as the inputs ('{0}') - got '{1}'."), FText::AsNumber(Inputs.Num()), FText::AsNumber(TargetActors.Num())));
+		return true;
+	}
+
+	while (Context->CurrentInputIndex < Inputs.Num())
+	{
+		FPCGTaggedData& OutputTaggedData = Context->OutputData.TaggedData.Add_GetRef(Inputs[Context->CurrentInputIndex]);
+		const UPCGData* TargetActorData = (!TargetActors.IsEmpty() ? (TargetActors.Num() == Inputs.Num() ? TargetActors[Context->CurrentInputIndex].Data : TargetActors[0].Data) : nullptr);
+		const UPCGData* OriginalData = Inputs[Context->CurrentInputIndex++].Data;
+
+		if (!OriginalData)
+		{
+			continue;
+		}
+
+		// Create accessors
+		FPCGAttributePropertyInputSelector TransformSelector; // initialize from settings?
+		TransformSelector.Update(FString(TEXT("$Transform")));
+
+		TUniquePtr<const IPCGAttributeAccessor> TransformAccessor = PCGAttributeAccessorHelpers::CreateConstAccessor(OriginalData, TransformSelector);
+		TUniquePtr<const IPCGAttributeAccessorKeys> TransformKeys = PCGAttributeAccessorHelpers::CreateConstKeys(OriginalData, TransformSelector);
+
+		if (!TransformAccessor.IsValid() || !TransformKeys.IsValid())
+		{
+			PCGE_LOG(Warning, GraphAndLog, FText::Format(LOCTEXT("MissingTransformAttribute", "Input data does not have the target attribute/property '{0}'"), FText::FromName(TransformSelector.GetName())));
+			continue;
+		}
+
+		FTransform DummyTransform;
+		if (TransformKeys->GetNum() > 0 && !TransformAccessor->Get<FTransform>(DummyTransform, *TransformKeys, EPCGAttributeAccessorFlags::AllowBroadcastAndConstructible))
+		{
+			PCGE_LOG(Warning, GraphAndLog, FText::Format(LOCTEXT("InvalidTransformAttribute", "Input transform attribute/property '{0}' is not compatible with a transform."), FText::FromName(TransformSelector.GetName())));
+			continue;
+		}
+
+		TUniquePtr<const IPCGAttributeAccessor> ClassOrTemplateAccessor = nullptr;
+		TUniquePtr<const IPCGAttributeAccessorKeys> ClassOrTemplateKeys = nullptr;
+
+		if (Settings->bUseClassAttribute)
+		{
+			FPCGAttributePropertyInputSelector ClassSelector = Settings->ClassAttribute.CopyAndFixLast(OriginalData);
+			ClassOrTemplateAccessor = PCGAttributeAccessorHelpers::CreateConstAccessor(OriginalData, ClassSelector);
+			ClassOrTemplateKeys = PCGAttributeAccessorHelpers::CreateConstKeys(OriginalData, ClassSelector);
+
+			if (!ClassOrTemplateAccessor.IsValid() || !ClassOrTemplateKeys.IsValid())
+			{
+				PCGE_LOG(Warning, GraphAndLog, FText::Format(LOCTEXT("MissingClassAttribute", "Input data does not have the target class attribute/property '{0}'"), FText::FromName(ClassSelector.GetName())));
+				continue;
+			}
+
+			FSoftObjectPath DummyPath;
+			if (ClassOrTemplateKeys->GetNum() > 0 && !ClassOrTemplateAccessor->Get<FSoftObjectPath>(DummyPath, *ClassOrTemplateKeys, EPCGAttributeAccessorFlags::AllowBroadcastAndConstructible))
+			{
+				PCGE_LOG(Warning, GraphAndLog, FText::Format(LOCTEXT("InvalidClassOrTemplateAttribute", "Input class attribute/property '{0}' is not an object or class path."), FText::FromName(ClassSelector.GetName())));
+				continue;
+			}
+		}
+		else if (Settings->bAllowTemplateComponentEditing && Settings->TemplateComponent)
+		{
+			ClassOrTemplateAccessor = MakeUnique<FPCGConstantValueAccessor<FSoftObjectPath>>(FSoftObjectPath(Settings->TemplateComponent));
+			ClassOrTemplateKeys = MakeUnique<FPCGAttributeAccessorKeysSingleObjectPtr<void>>();
+		}
+		else
+		{
+			ClassOrTemplateAccessor = MakeUnique<FPCGConstantValueAccessor<FSoftClassPath>>(FSoftClassPath(Settings->TemplateComponentClass));
+			ClassOrTemplateKeys = MakeUnique<FPCGAttributeAccessorKeysSingleObjectPtr<void>>();
+		}
+
+		TUniquePtr<const IPCGAttributeAccessor> TargetActorAccessor = nullptr;
+		TUniquePtr<const IPCGAttributeAccessorKeys> TargetActorKeys = nullptr;
+
+		if (TargetActorData)
+		{
+			FPCGAttributePropertyInputSelector TargetActorSelector = Settings->ActorReferenceAttribute.CopyAndFixLast(TargetActorData);
+			TargetActorAccessor = PCGAttributeAccessorHelpers::CreateConstAccessor(TargetActorData, TargetActorSelector);
+			TargetActorKeys = PCGAttributeAccessorHelpers::CreateConstKeys(TargetActorData, TargetActorSelector);
+
+			if (!TargetActorAccessor.IsValid() || !TargetActorKeys.IsValid())
+			{
+				PCGE_LOG(Warning, GraphAndLog, FText::Format(LOCTEXT("MissingActorReferenceAttribute", "Target actor data does not have the target attribute '{0}'"), FText::FromName(TargetActorSelector.GetName())));
+				continue;
+			}
+
+			FSoftObjectPath DummyPath;
+			if (TargetActorKeys->GetNum() > 0 && !TargetActorAccessor->Get<FSoftObjectPath>(DummyPath, *TargetActorKeys, EPCGAttributeAccessorFlags::AllowBroadcastAndConstructible))
+			{
+				PCGE_LOG(Warning, GraphAndLog, FText::Format(LOCTEXT("InvalidActorReferenceAttribute", "Target actor attribute '{0}' is not an object path."), FText::FromName(TargetActorSelector.GetName())));
+				continue;
+			}
+
+			// Additional validation: number of target actor keys should be either 1 or the same amount as the input data.
+			if (TargetActorKeys->GetNum() > 1 && TargetActorKeys->GetNum() != TransformKeys->GetNum())
+			{
+				PCGE_LOG(Warning, GraphAndLog, FText::Format(LOCTEXT("TargetActorAndTransformsMismatch", "Target actor data entry count '{0}' does not match the input entry count '{1}'."), FText::AsNumber(TargetActorKeys->GetNum()), FText::AsNumber(TransformKeys->GetNum())));
+				continue;
+			}
+		}
+		else
+		{
+			// Wrap on default actor namely the owner of the PCG component.
+			TargetActorAccessor = MakeUnique<FPCGConstantValueAccessor<FSoftObjectPath>>(FSoftObjectPath(Context->SourceComponent->GetOwner()));
+			TargetActorKeys = MakeUnique<FPCGAttributeAccessorKeysSingleObjectPtr<void>>();
+		}
+
+		// Create output
+		UPCGData* OutputData = OriginalData->DuplicateData(Context);
+
+		// Add component reference attribute
+		FPCGAttributePropertyOutputSelector ComponentReferenceSelector = Settings->ComponentReferenceAttribute; // No copy source here because there's no 1:1 mapping
+
+		// Create attribute if needed
+		if (ComponentReferenceSelector.IsBasicAttribute())
+		{
+			OutputData->MutableMetadata()->FindOrCreateAttribute(ComponentReferenceSelector.GetName(), FSoftObjectPath(), /*bAllowsInterpolation=*/false, /*bOverrideParent=*/false, /*bOverwriteIfTypeMismatch=*/false);
+		}
+
+		TUniquePtr<IPCGAttributeAccessor> OutputComponentReferenceAccessor = PCGAttributeAccessorHelpers::CreateAccessor(OutputData, ComponentReferenceSelector);
+		TUniquePtr<IPCGAttributeAccessorKeys> OutputComponentReferenceKeys = PCGAttributeAccessorHelpers::CreateKeys(OutputData, ComponentReferenceSelector);
+
+		if (!OutputComponentReferenceAccessor.IsValid() || !OutputComponentReferenceKeys.IsValid())
+		{
+			PCGE_LOG(Error, GraphAndLog, FText::Format(LOCTEXT("UnableToCreateOutputReference", "Unable to write component references to '{0}'."), FText::FromName(ComponentReferenceSelector.GetName())));
+			continue;
+		}
+
+		OutputTaggedData.Data = OutputData;
+
+		// Now create and add the components
+		AddComponents(Context,
+			TransformAccessor.Get(),
+			TransformKeys.Get(),
+			ClassOrTemplateAccessor.Get(),
+			ClassOrTemplateKeys.Get(),
+			TargetActorAccessor.Get(),
+			TargetActorKeys.Get(),
+			OutputComponentReferenceAccessor.Get(),
+			OutputComponentReferenceKeys.Get());
+
+		// Finally, break out if we're out of time
+		if (Context->CurrentInputIndex != Inputs.Num() && Context->ShouldStop())
+		{
+			break;
+		}
+	}
+
+	return Context->CurrentInputIndex == Inputs.Num();
+}
+
+void FPCGAddComponentElement::AddComponents(FPCGAddComponentContext* Context,
+	const IPCGAttributeAccessor* TransformAccessor,
+	const IPCGAttributeAccessorKeys* TransformKeys,
+	const IPCGAttributeAccessor* ClassOrTemplateAccessor,
+	const IPCGAttributeAccessorKeys* ClassOrTemplateKeys,
+	const IPCGAttributeAccessor* TargetActorAccessor,
+	const IPCGAttributeAccessorKeys* TargetActorKeys,
+	IPCGAttributeAccessor* OutputComponentAccessor,
+	IPCGAttributeAccessorKeys* OutputComponentKeys) const
+{
+	check(TransformAccessor && TransformKeys);
+	check(ClassOrTemplateAccessor && ClassOrTemplateKeys);
+	check(TargetActorAccessor && TargetActorKeys);
+	check(OutputComponentAccessor && OutputComponentKeys);
+
+	// TODO: implement component reuse
+	UPCGComponent* SourceComponent = Context->SourceComponent.Get();
+	check(SourceComponent);
+	const bool bIsInPreviewMode = SourceComponent->IsInPreviewMode();
+	const EObjectFlags ObjectFlags = bIsInPreviewMode ? RF_Transient : RF_NoFlags;
+
+	// Similar approach to PCGMetadataCommon::ApplyOnAccessorRange
+	const int32 NumberOfEntries = TransformKeys->GetNum();
+	const int32 ChunkSize = PCGMetadataElementCommon::DefaultChunkSize;
+	const int32 NumberOfIterations = (NumberOfEntries + ChunkSize - 1) / ChunkSize;
+
+	if (NumberOfEntries == 0)
+	{
+		return;
+	}
+
+	// TODO: root location? (needed for reuse) + settings UID
+	UPCGManagedComponentList* Resource = NewObject<UPCGManagedComponentList>(SourceComponent);
+	Resource->GeneratedComponents.Reserve(NumberOfEntries);
+
+	TArray<FTransform, TInlineAllocator<ChunkSize>> Transforms;
+	TArray<FSoftObjectPath, TInlineAllocator<ChunkSize>> ComponentsClassOrTemplate;
+	TArray<FSoftObjectPath, TInlineAllocator<ChunkSize>> TargetActors;
+	TArray<FSoftObjectPath, TInlineAllocator<ChunkSize>> ComponentReferences;
+
+	Transforms.SetNumUninitialized(ChunkSize);
+	ComponentsClassOrTemplate.SetNum(ChunkSize);
+	TargetActors.SetNum(ChunkSize);
+	ComponentReferences.SetNum(ChunkSize);
+
+	const EPCGAttributeAccessorFlags Flags = EPCGAttributeAccessorFlags::AllowBroadcast | EPCGAttributeAccessorFlags::AllowConstructible;
+
+	for (int32 i = 0; i < NumberOfIterations; ++i)
+	{
+		const int32 StartIndex = i * ChunkSize;
+		const int32 Range = FMath::Min(NumberOfEntries - StartIndex, ChunkSize);
+		TArrayView<FTransform> TransformsView(Transforms.GetData(), Range);
+		TArrayView<FSoftObjectPath> ComponentsClassOrTemplateView(ComponentsClassOrTemplate.GetData(), Range);
+		TArrayView<FSoftObjectPath> TargetActorsView(TargetActors.GetData(), Range);
+		TArrayView<FSoftObjectPath> ComponentReferencesView(ComponentReferences.GetData(), Range);
+
+		if(!TransformAccessor->GetRange(TransformsView, StartIndex, *TransformKeys, Flags) ||
+			!ClassOrTemplateAccessor->GetRange(ComponentsClassOrTemplateView, StartIndex, *ClassOrTemplateKeys, Flags) ||
+			!TargetActorAccessor->GetRange(TargetActorsView, StartIndex, *TargetActorKeys, Flags))
+		{
+			// Failed getting data...
+			continue;
+		}
+
+		for (int32 j = 0; j < TransformsView.Num(); ++j)
+		{
+			ComponentReferencesView[j] = FSoftObjectPath();
+
+			const FTransform& Transform = TransformsView[j];
+			AActor* TargetActor = TSoftObjectPtr<AActor>(TargetActorsView[j]).LoadSynchronous();
+			const FSoftObjectPath& ClassOrTemplatePath = ComponentsClassOrTemplateView[j];
+
+			TSubclassOf<UActorComponent> ComponentClass;
+			UActorComponent* ComponentTemplate = nullptr;
+
+			TSoftClassPtr<UActorComponent> ComponentClassPtr(ClassOrTemplatePath);
+			ComponentClass = ComponentClassPtr.LoadSynchronous();
+
+			if (!ComponentClass)
+			{
+				TSoftObjectPtr<UActorComponent> ComponentObjectPtr(ClassOrTemplatePath);
+				ComponentTemplate = ComponentObjectPtr.LoadSynchronous();
+				ComponentClass = ComponentTemplate ? ComponentTemplate->GetClass() : nullptr;
+			}
+			
+			if (TargetActor && ComponentClass)
+			{
+				UActorComponent* Component = NewObject<UActorComponent>(TargetActor, ComponentClass, NAME_None, ObjectFlags, ComponentTemplate);
+
+				if (!Component)
+				{
+					continue;
+				}
+
+				Component->RegisterComponent();
+				TargetActor->AddInstanceComponent(Component);
+
+				if (USceneComponent* SceneComponent = Cast<USceneComponent>(Component))
+				{
+					SceneComponent->AttachToComponent(TargetActor->GetRootComponent(), FAttachmentTransformRules(EAttachmentRule::KeepRelative, EAttachmentRule::KeepWorld, EAttachmentRule::KeepWorld, false));
+					SceneComponent->SetComponentToWorld(Transform);
+				}
+
+				Component->ComponentTags.Add(SourceComponent->GetFName());
+				Component->ComponentTags.Add(PCGHelpers::DefaultPCGTag);
+
+				// Add component reference on output data
+				ComponentReferences[j] = FSoftObjectPath(Component);
+
+				// Add component to generated components
+				Resource->GeneratedComponents.Emplace(Component);
+			}
+		}
+
+		// Write component reference to accessor
+		OutputComponentAccessor->SetRange<FSoftObjectPath>(ComponentReferencesView, StartIndex, *OutputComponentKeys, Flags);
+	}
+
+	// Finally add the resource to the component
+	if (Resource)
+	{
+		SourceComponent->AddToManagedResources(Resource);
+	}	
+}
+
+#if WITH_EDITOR
+bool UPCGAddComponentSettings::CanEditActorReference() const
+{
+	const UPCGNode* Node = Cast<UPCGNode>(GetOuter());
+	return !Node || Node->IsInputPinConnected(PCGAddComponentConstants::TargetActorPinLabel);
+}
+#endif // WITH_EDITOR
+
+#undef LOCTEXT_NAMESPACE
\ No newline at end of file
diff --git a/Engine/Plugins/PCG/Source/PCG/Public/Elements/PCGAddComponent.h b/Engine/Plugins/PCG/Source/PCG/Public/Elements/PCGAddComponent.h
new file mode 100644
--- /dev/null
+++ b/Engine/Plugins/PCG/Source/PCG/Public/Elements/PCGAddComponent.h
@@ -0,0 +1,133 @@
+// Copyright Epic Games, Inc. All Rights Reserved.
+
+#pragma once
+
+#include "PCGContext.h"
+#include "PCGSettings.h"
+#include "Metadata/Accessors/IPCGAttributeAccessor.h"
+#include "Metadata/Accessors/PCGAttributeAccessorKeys.h"
+
+#include "Components/ActorComponent.h"
+#include "Templates/SubclassOf.h"
+#include "Templates/UniquePtr.h"
+
+#include "PCGAddComponent.generated.h"
+
+namespace PCGAddComponentConstants
+{
+	const FName ComponentReferenceAttribute = TEXT("ComponentReference");
+	const FName TargetActorPinLabel = TEXT("TargetActor");
+}
+
+/**
+* Creates components and adds them to specified actors.
+*/
+UCLASS(MinimalAPI, BlueprintType, ClassGroup=(Procedural))
+class UPCGAddComponentSettings : public UPCGSettings
+{
+	GENERATED_BODY()
+
+	friend class FPCGAddComponentElement;
+public:
+	UPCGAddComponentSettings(const FObjectInitializer& ObjectInitializer);
+
+	//~Begin UObject interface
+	virtual void PostLoad() override;
+	virtual void BeginDestroy() override;
+#if WITH_EDITOR
+	virtual void PreEditChange(FProperty* PropertyAboutToChange) override;
+	virtual void PostEditChangeProperty(FPropertyChangedEvent& PropertyChangedEvent) override;
+	virtual void PreEditUndo() override;
+	virtual void PostEditUndo() override;
+#endif
+	//~End UObject interface
+
+	//~Begin UPCGSettings interface
+#if WITH_EDITOR
+	virtual FName GetDefaultNodeName() const override { return FName(TEXT("AddComponent")); }
+	virtual FText GetDefaultNodeTitle() const override { return NSLOCTEXT("PCGAddComponentElement", "NodeTitle", "Add Component"); }
+	virtual FText GetNodeTooltipText() const override { return NSLOCTEXT("PCGAddComponentElement", "NodeTooltip", "Adds component(s) to specified target actor(s)."); }
+	virtual EPCGSettingsType GetType() const override { return EPCGSettingsType::Generic; }
+	virtual bool HasDynamicPins() const override { return true; }
+#endif
+
+protected:
+	virtual TArray<FPCGPinProperties> InputPinProperties() const override;
+	virtual TArray<FPCGPinProperties> OutputPinProperties() const override;
+	virtual FPCGElementPtr CreateElement() const override;
+	//~End UPCGSettings interface
+
+public:
+	PCG_API void SetTemplateComponentClass(const TSubclassOf<UActorComponent>& InTemplateComponentClass);
+	PCG_API void SetAllowTemplateComponentEditing(bool bInAllowTemplateComponentEditing);
+	const TSubclassOf<UActorComponent>& GetTemplateComponentClass() const { return TemplateComponentClass; }
+	bool GetAllowTemplateComponentEditing() const { return bAllowTemplateComponentEditing; }
+
+private:
+#if WITH_EDITOR
+	void OnBlueprintChanged(UBlueprint* Blueprint);
+	void SetupOnBlueprintChangedEvent();
+	void TeardownOnBlueprintChangedEvent();
+	void RefreshTemplateComponent();
+#endif
+
+public:
+	/** Controls whether component class selection will be done by attribute or from a constant defined in this node. */
+	UPROPERTY(BlueprintReadWrite, EditAnywhere, Category = Settings)
+	bool bUseClassAttribute = false;
+
+	/** Specifies component class selection */
+	UPROPERTY(BlueprintReadWrite, EditAnywhere, Category = Settings, meta = (PCG_Overridable, EditCondition="bUseClassAttribute", EditConditionHides))
+	FPCGAttributePropertyInputSelector ClassAttribute;
+
+protected:
+	UPROPERTY(BlueprintReadOnly, EditAnywhere, Category = Settings, meta = (EditCondition = "!bUseClassAttribute", EditConditionHides))
+	TSubclassOf<UActorComponent> TemplateComponentClass = nullptr;
+
+	UPROPERTY(BlueprintReadWrite, EditAnywhere, Category = Settings, meta = (EditCondition = "!bUseClassAttribute", EditConditionHides))
+	bool bAllowTemplateComponentEditing = false;
+
+public:
+	UPROPERTY(BlueprintReadWrite, EditAnywhere, Instanced, Category = Settings, meta = (ShowInnerProperties, EditCondition = "!bUseClassAttribute && bAllowTemplateComponentEditing", EditConditionHides))
+	TObjectPtr<UActorComponent> TemplateComponent;
+
+	/** Specifies what attribute is used to derive actor reference */
+	UPROPERTY(BlueprintReadWrite, EditAnywhere, Category = Settings, meta = (EditCondition="CanEditActorReference()", PCG_Overridable))
+	FPCGAttributePropertyInputSelector ActorReferenceAttribute;
+
+	/** Specifies what attribute to write the component reference to. */
+	UPROPERTY(BlueprintReadWrite, EditAnywhere, Category = Settings, meta = (PCG_Overridable, PCG_DiscardPropertySelection, PCG_DiscardExtraSelection))
+	FPCGAttributePropertyOutputSelector ComponentReferenceAttribute;
+
+protected:
+#if WITH_EDITOR
+	UFUNCTION()
+	bool CanEditActorReference() const;
+#endif // WITH_EDITOR
+};
+
+struct FPCGAddComponentContext : public FPCGContext
+{
+	int32 CurrentInputIndex = 0;
+};
+
+class FPCGAddComponentElement : public IPCGElementWithCustomContext<FPCGAddComponentContext>
+{
+public:
+	virtual bool CanExecuteOnlyOnMainThread(FPCGContext* Context) const override { return true; }
+	virtual bool IsCacheable(const UPCGSettings* InSettings) const override { return false; }
+
+protected:
+	virtual bool ExecuteInternal(FPCGContext* Context) const override;
+
+private:
+	void AddComponents(FPCGAddComponentContext* Context,
+		const IPCGAttributeAccessor* TransformAccessor,
+		const IPCGAttributeAccessorKeys* TransformKeys,
+		const IPCGAttributeAccessor* ClassAccessor,
+		const IPCGAttributeAccessorKeys* ClassKeys,
+		const IPCGAttributeAccessor* TargetActorAccessor,
+		const IPCGAttributeAccessorKeys* TargetActorKeys,
+		IPCGAttributeAccessor* OutputComponentAccessor,
+		IPCGAttributeAccessorKeys* OutputComponentKeys) const;
+};
\ No newline at end of file
