Static analysis fixes following the add component changes

#rnx

diff --git a/Engine/Plugins/PCG/Source/PCG/Private/Elements/PCGAddComponent.cpp b/Engine/Plugins/PCG/Source/PCG/Private/Elements/PCGAddComponent.cpp
--- a/Engine/Plugins/PCG/Source/PCG/Private/Elements/PCGAddComponent.cpp
+++ b/Engine/Plugins/PCG/Source/PCG/Private/Elements/PCGAddComponent.cpp
@@ -421,6 +421,7 @@ void FPCGAddComponentElement::AddComponents(FPCGAddComponentContext* Context,
 
 	// TODO: root location? (needed for reuse) + settings UID
 	UPCGManagedComponentList* Resource = NewObject<UPCGManagedComponentList>(SourceComponent);
+	check(Resource);
 	Resource->GeneratedComponents.Reserve(NumberOfEntries);
 
 	TArray<FTransform, TInlineAllocator<ChunkSize>> Transforms;
@@ -507,7 +508,7 @@ void FPCGAddComponentElement::AddComponents(FPCGAddComponentContext* Context,
 	}
 
 	// Finally add the resource to the component
-	if (Resource)
+	if (!Resource->GeneratedComponents.IsEmpty())
 	{
 		SourceComponent->AddToManagedResources(Resource);
 	}	
