[PCG] Fix changing 'Is Partitioned' flag on PCG Volume not causing Partition Actors to be created (since CL 31214740)

- Add a Tool > Create missing Partition Grid Actors menu (in case user deletes all the PAs and wants to recreate them)

diff --git a/Engine/Plugins/PCG/Source/PCG/Private/PCGActorAndComponentMapping.cpp b/Engine/Plugins/PCG/Source/PCG/Private/PCGActorAndComponentMapping.cpp
--- a/Engine/Plugins/PCG/Source/PCG/Private/PCGActorAndComponentMapping.cpp
+++ b/Engine/Plugins/PCG/Source/PCG/Private/PCGActorAndComponentMapping.cpp
@@ -369,6 +369,17 @@ bool FPCGActorAndComponentMapping::RegisterOrUpdatePartitionedPCGComponent(UPCGC
 
 	PartitionedOctree.AddOrUpdateComponent(InComponent, Bounds, bComponentHasChanged, bComponentWasAdded);
 
+#if WITH_EDITOR
+	// In Editor only, we will create new partition actors depending on the new bounds and generation trigger. Runtime managed components should not create PAs here
+	if ((bComponentHasChanged || bComponentWasAdded) && !InComponent->IsManagedByRuntimeGenSystem())
+	{
+		bool bHasUnbounded = false;
+		PCGHiGenGrid::FSizeArray GridSizes;
+		ensure(PCGHelpers::GetGenerationGridSizes(InComponent->GetGraph(), PCGSubsystem->GetPCGWorldActor(), GridSizes, bHasUnbounded));
+		PCGSubsystem->CreatePartitionActorsWithinBounds(InComponent, Bounds, GridSizes);
+	}
+#endif // WITH_EDITOR
+
 	// After adding/updating, try to do the mapping (if we asked for it and the component changed)
 	if (bDoActorMapping)
 	{
diff --git a/Engine/Plugins/PCG/Source/PCG/Private/PCGSubsystem.cpp b/Engine/Plugins/PCG/Source/PCG/Private/PCGSubsystem.cpp
--- a/Engine/Plugins/PCG/Source/PCG/Private/PCGSubsystem.cpp
+++ b/Engine/Plugins/PCG/Source/PCG/Private/PCGSubsystem.cpp
@@ -549,20 +549,19 @@ FPCGTaskId UPCGSubsystem::ScheduleComponent(UPCGComponent* PCGComponent, EPCGHiG
 	PCGHiGenGrid::FSizeArray GridSizes;
 	ensure(PCGHelpers::GetGenerationGridSizes(PCGComponent->GetGraph(), GetPCGWorldActor(), GridSizes, bHasUnbounded));
 
+#if WITH_EDITOR
 	// Create the PartitionActors if necessary. Skip if this is a runtime managed component, PAs are handled manually by the RuntimeGenScheduler.
+	// Editor only because we expect at runtime for PAs to already exist so they can properly be streamed in and out (creating them at runtime would leave them unmanaged and always loaded)
 	if (PCGComponent->IsPartitioned() && !PCGComponent->IsManagedByRuntimeGenSystem())
 	{
 		if (!GridSizes.IsEmpty())
 		{
-			// In this case create the PA and update the mapping
-			// Note: This is an immediate operation, as we need the PA for the generation.
-			auto ScheduleTask = [](APCGPartitionActor* PCGActor, const FBox& InIntersectedBounds) { return InvalidPCGTaskId; };
-
-			ForAllOverlappingCells(PCGComponent, PCGComponent->GetGridBounds(), GridSizes, /*bCanCreateActor=*/true, /*Dependencies=*/{}, ScheduleTask);
+			CreatePartitionActorsWithinBounds(PCGComponent, PCGComponent->GetGridBounds(), GridSizes);
 		}
 
 		ActorAndComponentMapping.UpdateMappingPCGComponentPartitionActor(PCGComponent);
 	}
+#endif // WITH_EDITOR
 
 	// Execution dependencies require a task to finish executing before the dependent task.
 	TArray<FPCGTaskId> ExecutionDependencyTasks;
@@ -1710,6 +1709,41 @@ void UPCGSubsystem::ClearExecutedStacks(const UPCGGraph* InContainingGraph)
 	}
 }
 
+void UPCGSubsystem::CreateMissingPartitionActors()
+{
+	if (!PCGHelpers::IsRuntimeOrPIE())
+	{
+		ActorAndComponentMapping.ForAllOriginalComponents([this](UPCGComponent* PCGComponent)
+		{
+			if (PCGComponent->IsPartitioned() && !PCGComponent->IsManagedByRuntimeGenSystem())
+			{
+				bool bHasUnbounded = false;
+				PCGHiGenGrid::FSizeArray GridSizes;
+				ensure(PCGHelpers::GetGenerationGridSizes(PCGComponent->GetGraph(), GetPCGWorldActor(), GridSizes, bHasUnbounded));
+				if (!GridSizes.IsEmpty())
+				{
+					CreatePartitionActorsWithinBounds(PCGComponent, PCGComponent->GetGridBounds(), GridSizes);
+				}
+				ActorAndComponentMapping.UpdateMappingPCGComponentPartitionActor(PCGComponent);
+			}
+		});
+	}
+}
+
+void UPCGSubsystem::CreatePartitionActorsWithinBounds(UPCGComponent* InComponent, const FBox& InBounds, const PCGHiGenGrid::FSizeArray& InGridSizes)
+{
+	if (!PCGHelpers::IsRuntimeOrPIE())
+	{
+		// We can't spawn actors if we are running constructions scripts, asserting when we try to get the actor with the WP API.
+		// We should never enter this if we are in a construction script. If the ensure is hit, we need to fix it.
+		UWorld* World = GetWorld();
+		if (ensure(World && !World->bIsRunningConstructionScript))
+		{
+			ForAllOverlappingCells(InComponent, InBounds, InGridSizes, true, {}, [](APCGPartitionActor*, const FBox&) { return InvalidPCGTaskId; });
+		}
+	}
+}
+
 #endif // WITH_EDITOR
 
 void UPCGSubsystem::FlushCache()
diff --git a/Engine/Plugins/PCG/Source/PCG/Public/PCGSubsystem.h b/Engine/Plugins/PCG/Source/PCG/Public/PCGSubsystem.h
--- a/Engine/Plugins/PCG/Source/PCG/Public/PCGSubsystem.h
+++ b/Engine/Plugins/PCG/Source/PCG/Public/PCGSubsystem.h
@@ -316,6 +316,7 @@ public:
 	FPCGOnComponentUnregistered OnComponentUnregistered;
 	FPCGOnComponentGenerationCompleteOrCancelled OnComponentGenerationCompleteOrCancelled;
 
+	void CreateMissingPartitionActors();
 private:
 	enum class EOperation : uint32
 	{
@@ -324,7 +325,7 @@ private:
 		Generate
 	};
 
-	void CreatePartitionActorsWithinBounds(const FBox& InBounds, const PCGHiGenGrid::FSizeArray& InGridSizes);
+	void CreatePartitionActorsWithinBounds(UPCGComponent* InComponent, const FBox& InBounds, const PCGHiGenGrid::FSizeArray& InGridSizes);
 
 	FPCGNodeVisualLogs NodeVisualLogs;
 #endif // WITH_EDITOR
diff --git a/Engine/Plugins/PCG/Source/PCGEditor/Private/PCGEditorModule.cpp b/Engine/Plugins/PCG/Source/PCGEditor/Private/PCGEditorModule.cpp
--- a/Engine/Plugins/PCG/Source/PCGEditor/Private/PCGEditorModule.cpp
+++ b/Engine/Plugins/PCG/Source/PCGEditor/Private/PCGEditorModule.cpp
@@ -380,6 +380,24 @@ void FPCGEditorModule::UnregisterMenuExtensions()
 
 void FPCGEditorModule::PopulateMenuActions(FMenuBuilder& MenuBuilder)
 {
+	MenuBuilder.AddMenuEntry(
+		LOCTEXT("CreateMissingPartitionActors", "Create missing Partition Grid Actors"),
+		LOCTEXT("CreateMissingPartitionActors_Tooltip", "Will visit all Partitioned PCG Components and create the missing intersecting Partition Grid Actors"),
+		FSlateIcon(),
+		FUIAction(
+			FExecuteAction::CreateLambda([]() 
+			{
+				if (GEditor)
+				{
+					if (UPCGSubsystem* PCGSubsystem = UPCGSubsystem::GetInstance(GEditor->GetEditorWorldContext().World()))
+					{
+						FScopedTransaction Transaction(LOCTEXT("CreateMissingPartitionActorsTransaction", "Create missing Partition Grid Actors"));
+						PCGSubsystem->CreateMissingPartitionActors();
+					}
+				}
+			})),
+		NAME_None);
+
 	MenuBuilder.AddSubMenu(LOCTEXT("PCGSubMenuDelete", "Delete"), FText(), FNewMenuDelegate::CreateLambda([](FMenuBuilder& SubMenuBuilder)
 	{
 		SubMenuBuilder.AddMenuEntry(
