[PCG] Remove superfluous EditCondition from World Raycast Node

Note: This was originally a toggle, but there was no need for it since the selector can just be used for Properties as well

#rnx

diff --git a/Engine/Plugins/PCG/Source/PCG/Public/Elements/PCGWorldRaycast.h b/Engine/Plugins/PCG/Source/PCG/Public/Elements/PCGWorldRaycast.h
--- a/Engine/Plugins/PCG/Source/PCG/Public/Elements/PCGWorldRaycast.h
+++ b/Engine/Plugins/PCG/Source/PCG/Public/Elements/PCGWorldRaycast.h
@@ -41,7 +41,7 @@ public:
 	EPCGWorldRaycastMode RaycastMode = EPCGWorldRaycastMode::Infinite;
 
 	/** The selected attribute determines the ray origin. */
-	UPROPERTY(BlueprintReadWrite, EditAnywhere, Category = Settings, meta = (PCG_Overridable, EditCondition = "bOverrideOriginPositions", EditConditionHides))
+	UPROPERTY(BlueprintReadWrite, EditAnywhere, Category = Settings, meta = (PCG_Overridable))
 	FPCGAttributePropertyInputSelector OriginInputAttribute;
 
 	/** Use a selected attribute as the ray direction. */
