[PCG] FPCGSubgraphElement::PrepareSubgraphUserParameters use FPCGContext::NewObject_AnyThread

#rnx

diff --git a/Engine/Plugins/PCG/Source/PCG/Private/PCGSubgraph.cpp b/Engine/Plugins/PCG/Source/PCG/Private/PCGSubgraph.cpp
--- a/Engine/Plugins/PCG/Source/PCG/Private/PCGSubgraph.cpp
+++ b/Engine/Plugins/PCG/Source/PCG/Private/PCGSubgraph.cpp
@@ -578,7 +578,7 @@ void FPCGSubgraphElement::PrepareSubgraphUserParameters(const UPCGSubgraphSettin
 	// By construction, there should be one and only one of this data. (Filtering of previous data is done in PrepareSubgraphData)
 	if (const UPCGGraphInterface* SubgraphInterface = Settings->GetSubgraphInterface())
 	{
-		UPCGUserParametersData* UserParamData = NewObject<UPCGUserParametersData>();
+		UPCGUserParametersData* UserParamData = FPCGContext::NewObject_AnyThread<UPCGUserParametersData>(Context);
 
 		if (Context->GraphInstanceParametersOverride.IsValid())
 		{
