[PCG] Add a new node to reverse splines

#rnx

diff --git a/Engine/Plugins/PCG/Source/PCG/Private/Elements/PCGReverseSpline.cpp b/Engine/Plugins/PCG/Source/PCG/Private/Elements/PCGReverseSpline.cpp
new file mode 100644
--- /dev/null
+++ b/Engine/Plugins/PCG/Source/PCG/Private/Elements/PCGReverseSpline.cpp
@@ -0,0 +1,145 @@
+// Copyright Epic Games, Inc. All Rights Reserved.
+
+#include "Elements/PCGReverseSpline.h"
+
+#include "PCGContext.h"
+#include "Data/PCGSplineData.h"
+
+#define LOCTEXT_NAMESPACE "PCGReverseSplineElement"
+
+namespace PCGReverseSpline
+{
+	bool IsClockwiseXY(const UPCGSplineData* InputSplineData)
+	{
+		check(InputSplineData);
+
+		double CumulativeAngle = 0.0;
+		int NumPoints = InputSplineData->SplineStruct.SplineCurves.Position.Points.Num();
+		if (!InputSplineData->IsClosed())
+		{
+			NumPoints--;
+		}
+
+		for (int32 Index = 0; Index < NumPoints; ++Index)
+		{
+			FVector ArriveTangent{}, LeaveTangent{};
+
+			InputSplineData->GetTangentsAtSegmentStart(Index, ArriveTangent, LeaveTangent);
+			const double CrossProduct = (ArriveTangent.X * LeaveTangent.Y - ArriveTangent.Y * LeaveTangent.X);
+			const double AbsAngle = FMath::Acos(ArriveTangent.CosineAngle2D(LeaveTangent));
+			const double Angle = FMath::Sign(CrossProduct) * AbsAngle;
+
+			CumulativeAngle += Angle;
+		}
+
+		return CumulativeAngle <= 0;
+	}
+
+	UPCGSplineData* Reverse(const UPCGSplineData* InputSplineData, FPCGContext* Context)
+	{
+		check(InputSplineData);
+
+		const FInterpCurveVector& ControlPointsPosition = InputSplineData->SplineStruct.SplineCurves.Position;
+		const FInterpCurveQuat& ControlPointsRotation = InputSplineData->SplineStruct.SplineCurves.Rotation;
+		const FInterpCurveVector& ControlPointsScale = InputSplineData->SplineStruct.SplineCurves.Scale;
+
+		TArray<FSplinePoint> NewControlPoints;
+		NewControlPoints.Reserve(ControlPointsPosition.Points.Num());
+
+		for (int i = ControlPointsPosition.Points.Num() - 1; i >= 0; --i)
+		{
+			// Tangents are inverted and swapped.
+			NewControlPoints.Emplace(static_cast<float>(NewControlPoints.Num()),
+				ControlPointsPosition.Points[i].OutVal,
+				-ControlPointsPosition.Points[i].LeaveTangent,
+				-ControlPointsPosition.Points[i].ArriveTangent,
+				ControlPointsRotation.Points[i].OutVal.Rotator(),
+				ControlPointsScale.Points[i].OutVal,
+				ConvertInterpCurveModeToSplinePointType(ControlPointsPosition.Points[i].InterpMode));
+		}
+
+		UPCGSplineData* NewSplineData = FPCGContext::NewObject_AnyThread<UPCGSplineData>(Context);
+		NewSplineData->InitializeFromData(InputSplineData);
+		NewSplineData->Initialize(NewControlPoints, InputSplineData->IsClosed(), InputSplineData->GetTransform());
+
+		return NewSplineData;
+	}
+}
+
+#if WITH_EDITOR
+FName UPCGReverseSplineSettings::GetDefaultNodeName() const
+{
+	return FName(TEXT("ReverseSpline"));
+}
+
+FText UPCGReverseSplineSettings::GetDefaultNodeTitle() const
+{
+	return LOCTEXT("NodeTitle", "Reverse Spline");
+}
+#endif // WITH_EDITOR
+
+FPCGElementPtr UPCGReverseSplineSettings::CreateElement() const
+{
+	return MakeShared<FPCGReverseSplineElement>();
+}
+
+TArray<FPCGPinProperties> UPCGReverseSplineSettings::InputPinProperties() const
+{
+	TArray<FPCGPinProperties> Properties;
+	Properties.Emplace_GetRef(PCGPinConstants::DefaultInputLabel, EPCGDataType::Spline).SetRequiredPin();
+	return Properties;
+}
+
+TArray<FPCGPinProperties> UPCGReverseSplineSettings::OutputPinProperties() const
+{ 
+	TArray<FPCGPinProperties> Properties;
+	Properties.Emplace(PCGPinConstants::DefaultOutputLabel, EPCGDataType::Spline);
+	return Properties;
+}
+
+bool FPCGReverseSplineElement::ExecuteInternal(FPCGContext* InContext) const
+{
+	TRACE_CPUPROFILER_EVENT_SCOPE(FPCGReverseSplineElement::Execute);
+
+	check(InContext);
+
+	const UPCGReverseSplineSettings* Settings = InContext->GetInputSettings<UPCGReverseSplineSettings>();
+	check(Settings);
+
+	for (const FPCGTaggedData& InputData : InContext->InputData.GetInputsByPin(PCGPinConstants::DefaultInputLabel))
+	{
+		FPCGTaggedData& Output = InContext->OutputData.TaggedData.Emplace_GetRef(InputData);
+
+		const UPCGSplineData* InputSplineData = Cast<const UPCGSplineData>(InputData.Data);
+		if (!InputSplineData)
+		{
+			continue;
+		}
+
+		bool bShouldReverse = true;
+		switch (Settings->Operation)
+		{
+		case EPCGReverseSplineOperation::ForceClockwise:
+			bShouldReverse = !PCGReverseSpline::IsClockwiseXY(InputSplineData);
+			break;
+		case EPCGReverseSplineOperation::ForceCounterClockwise:
+			bShouldReverse = PCGReverseSpline::IsClockwiseXY(InputSplineData);
+			break;
+		case EPCGReverseSplineOperation::Reverse:
+			bShouldReverse = true;
+			break;
+		default:
+			PCGLog::LogErrorOnGraph(LOCTEXT("InvalidOperation", "Invalid operation enum value"), InContext);
+			return true;
+		}
+
+		if (bShouldReverse)
+		{
+			Output.Data = PCGReverseSpline::Reverse(InputSplineData, InContext);
+		}
+	}
+
+	return true;
+}
+
+#undef LOCTEXT_NAMESPACE
diff --git a/Engine/Plugins/PCG/Source/PCG/Public/Elements/PCGReverseSpline.h b/Engine/Plugins/PCG/Source/PCG/Public/Elements/PCGReverseSpline.h
new file mode 100644
--- /dev/null
+++ b/Engine/Plugins/PCG/Source/PCG/Public/Elements/PCGReverseSpline.h
@@ -0,0 +1,66 @@
+// Copyright Epic Games, Inc. All Rights Reserved.
+
+#pragma once
+
+#include "PCGSettings.h"
+
+#include "PCGReverseSpline.generated.h"
+
+class UPCGSplineData;
+
+namespace PCGReverseSpline
+{
+	/** Test spline controls points to know if they are clockwise on the XY plane (clockwise around the Z axis). */
+	PCG_API bool IsClockwiseXY(const UPCGSplineData* InputSplineData);
+
+	/** Return the same spline data, but with its control points in reverse order. */
+	PCG_API UPCGSplineData* Reverse(const UPCGSplineData* InputSplineData, FPCGContext* Context);
+}
+
+UENUM()
+enum class EPCGReverseSplineOperation
+{
+	/** Will reverse the spline points regardless of their orientation. */
+	Reverse,
+
+	/** Will reverse the spline points if they are counter clockwise. */
+	ForceClockwise,
+
+	/** Will reverse the spline points if they are clockwise. */
+	ForceCounterClockwise
+};
+
+/**
+* Reverse the order of a spline's control points.
+* This can be conditional to force a given orientation (clockwise or counter clockwise).
+*/
+UCLASS(MinimalAPI, BlueprintType, ClassGroup = (Procedural))
+class UPCGReverseSplineSettings : public UPCGSettings
+{
+	GENERATED_BODY()
+
+public:
+	//~Begin UPCGSettings interface
+#if WITH_EDITOR
+	virtual FName GetDefaultNodeName() const override;
+	virtual FText GetDefaultNodeTitle() const override;
+	virtual EPCGSettingsType GetType() const override { return EPCGSettingsType::Spatial; }
+#endif
+
+protected:
+	virtual TArray<FPCGPinProperties> InputPinProperties() const override;
+	virtual TArray<FPCGPinProperties> OutputPinProperties() const override;
+	virtual FPCGElementPtr CreateElement() const override;
+	//~End UPCGSettings interface
+
+public:
+	UPROPERTY(EditAnywhere, BlueprintReadWrite, Category = "Settings")
+	EPCGReverseSplineOperation Operation = EPCGReverseSplineOperation::Reverse;
+};
+
+class FPCGReverseSplineElement : public IPCGElement
+{
+protected:
+	virtual bool ExecuteInternal(FPCGContext* InContext) const override;
+};
+
