[PCG] Add option for N:N, N:1, 1:N op on Copy Points

#rnx

diff --git a/Engine/Plugins/PCG/Source/PCG/Private/Elements/PCGCopyPoints.cpp b/Engine/Plugins/PCG/Source/PCG/Private/Elements/PCGCopyPoints.cpp
--- a/Engine/Plugins/PCG/Source/PCG/Private/Elements/PCGCopyPoints.cpp
+++ b/Engine/Plugins/PCG/Source/PCG/Private/Elements/PCGCopyPoints.cpp
@@ -56,272 +56,301 @@ bool FPCGCopyPointsElement::ExecuteInternal(FPCGContext* Context) const
 
 	TArray<FPCGTaggedData>& Outputs = Context->OutputData.TaggedData;
 
-	for (const FPCGTaggedData& Source : Sources)
+	const int32 NumSources = Sources.Num();
+	const int32 NumTargets = Targets.Num();
+
+	if (NumSources == 0 || NumTargets == 0)
+	{
+		// Nothing to do
+		return true;
+	}
+
+	int32 NumIterations = 0;
+	if (Settings->bCopyEachSourceOnEveryTarget)
 	{
-		for (const FPCGTaggedData& Target : Targets)
+		NumIterations = NumSources * NumTargets;
+	}
+	else
+	{
+		if (NumSources != NumTargets && NumSources != 1 && NumTargets != 1)
 		{
-			FPCGTaggedData& Output = Outputs.Add_GetRef(Source);
-			if (!Source.Data || !Target.Data) 
-			{
-				PCGE_LOG(Error, GraphAndLog, LOCTEXT("InvalidInputData", "Invalid input data"));
-				return true;
-			}
+			PCGLog::LogErrorOnGraph(FText::Format(LOCTEXT("MismatchNum", "Num Sources ({0}) mismatches with Num Targets ({1}). Only supports N:N, 1:N and N:1 operation."), NumSources, NumTargets), Context);
+			Outputs = Sources;
+			return true;
+		}
 
-			const UPCGSpatialData* SourceSpatialData = Cast<UPCGSpatialData>(Source.Data);
-			const UPCGSpatialData* TargetSpatialData = Cast<UPCGSpatialData>(Target.Data);
+		NumIterations = FMath::Max(NumSources, NumTargets);
+	}
 
-			if (!SourceSpatialData || !TargetSpatialData)
-			{
-				PCGE_LOG(Error, GraphAndLog, LOCTEXT("CouldNotObtainSpatialData", "Unable to get Spatial Data from input"));
-				return true;
-			}
+	for (int32 i = 0; i < NumIterations; ++i)
+	{
+		const int32 SourceIndex = Settings->bCopyEachSourceOnEveryTarget ? (i / NumTargets) : FMath::Min(i, Sources.Num() - 1);
+		const int32 TargetIndex = Settings->bCopyEachSourceOnEveryTarget ? (i % NumTargets) : FMath::Min(i, Targets.Num() - 1);
 
-			const UPCGPointData* SourcePointData = SourceSpatialData->ToPointData(Context);
-			const UPCGPointData* TargetPointData = TargetSpatialData->ToPointData(Context);
+		const FPCGTaggedData& Source = Sources[SourceIndex];
+		const FPCGTaggedData& Target = Targets[TargetIndex];
 
-			if (!SourcePointData || !TargetPointData)
-			{
-				PCGE_LOG(Error, GraphAndLog, LOCTEXT("CouldNotGetPointData", "Unable to get Point Data from input"));
-				return true;
-			}
+		FPCGTaggedData& Output = Outputs.Add_GetRef(Source);
+		if (!Source.Data || !Target.Data) 
+		{
+			PCGE_LOG(Error, GraphAndLog, LOCTEXT("InvalidInputData", "Invalid input data"));
+			return true;
+		}
 
-			const UPCGMetadata* SourcePointMetadata = SourcePointData->Metadata;
-			const UPCGMetadata* TargetPointMetadata = TargetPointData->Metadata;
+		const UPCGSpatialData* SourceSpatialData = Cast<UPCGSpatialData>(Source.Data);
+		const UPCGSpatialData* TargetSpatialData = Cast<UPCGSpatialData>(Target.Data);
 
-			const TArray<FPCGPoint>& SourcePoints = SourcePointData->GetPoints();
-			const TArray<FPCGPoint>& TargetPoints = TargetPointData->GetPoints();
+		if (!SourceSpatialData || !TargetSpatialData)
+		{
+			PCGE_LOG(Error, GraphAndLog, LOCTEXT("CouldNotObtainSpatialData", "Unable to get Spatial Data from input"));
+			return true;
+		}
 
+		const UPCGPointData* SourcePointData = SourceSpatialData->ToPointData(Context);
+		const UPCGPointData* TargetPointData = TargetSpatialData->ToPointData(Context);
 
-			UPCGPointData* OutPointData = FPCGContext::NewObject_AnyThread<UPCGPointData>(Context);
-			TArray<FPCGPoint>& OutPoints = OutPointData->GetMutablePoints();
-			Output.Data = OutPointData;
+		if (!SourcePointData || !TargetPointData)
+		{
+			PCGE_LOG(Error, GraphAndLog, LOCTEXT("CouldNotGetPointData", "Unable to get Point Data from input"));
+			return true;
+		}
 
-			// Make sure that output contains both collection of tags from source and target
-			if (Settings->TagInheritance == EPCGCopyPointsTagInheritanceMode::Target)
-			{
-				Output.Tags = Target.Tags;
-			}
-			else if (Settings->TagInheritance == EPCGCopyPointsTagInheritanceMode::Both)
-			{
-				Output.Tags.Append(Target.Tags);
-			}
+		const UPCGMetadata* SourcePointMetadata = SourcePointData->Metadata;
+		const UPCGMetadata* TargetPointMetadata = TargetPointData->Metadata;
 
-			// RootMetadata will be parent to the ouptut metadata, while NonRootMetadata will carry attributes from the input not selected for inheritance
-			// Note that this is a preference, as we can and should pick more efficiently in the trivial cases
-			const UPCGMetadata* RootMetadata = nullptr;
-			const UPCGMetadata* NonRootMetadata = nullptr;
+		const TArray<FPCGPoint>& SourcePoints = SourcePointData->GetPoints();
+		const TArray<FPCGPoint>& TargetPoints = TargetPointData->GetPoints();
 
-			const bool bSourceHasMetadata = (SourcePointMetadata->GetAttributeCount() > 0 && SourcePointMetadata->GetItemCountForChild() > 0);
-			const bool bTargetHasMetadata = (TargetPointMetadata->GetAttributeCount() > 0 && TargetPointMetadata->GetItemCountForChild() > 0);
 
-			bool bInheritMetadataFromSource = true;
-			bool bProcessMetadata = (bSourceHasMetadata || bTargetHasMetadata);
+		UPCGPointData* OutPointData = FPCGContext::NewObject_AnyThread<UPCGPointData>(Context);
+		TArray<FPCGPoint>& OutPoints = OutPointData->GetMutablePoints();
+		Output.Data = OutPointData;
 
-			if (AttributeInheritance == EPCGCopyPointsMetadataInheritanceMode::SourceOnly)
-			{
-				bInheritMetadataFromSource = true;
-				bProcessMetadata = bSourceHasMetadata;
-				RootMetadata = SourcePointMetadata;
-				NonRootMetadata = nullptr;
+		// Make sure that output contains both collection of tags from source and target
+		if (Settings->TagInheritance == EPCGCopyPointsTagInheritanceMode::Target)
+		{
+			Output.Tags = Target.Tags;
+		}
+		else if (Settings->TagInheritance == EPCGCopyPointsTagInheritanceMode::Both)
+		{
+			Output.Tags.Append(Target.Tags);
+		}
 
-				OutPointData->InitializeFromData(SourcePointData);
-			}
-			else if (AttributeInheritance == EPCGCopyPointsMetadataInheritanceMode::TargetOnly)
-			{
-				bInheritMetadataFromSource = false;
-				bProcessMetadata = bTargetHasMetadata;
-				RootMetadata = TargetPointMetadata;
-				NonRootMetadata = nullptr;
+		// RootMetadata will be parent to the ouptut metadata, while NonRootMetadata will carry attributes from the input not selected for inheritance
+		// Note that this is a preference, as we can and should pick more efficiently in the trivial cases
+		const UPCGMetadata* RootMetadata = nullptr;
+		const UPCGMetadata* NonRootMetadata = nullptr;
 
-				OutPointData->InitializeFromData(TargetPointData);
-			}
-			else if (AttributeInheritance == EPCGCopyPointsMetadataInheritanceMode::SourceFirst || AttributeInheritance == EPCGCopyPointsMetadataInheritanceMode::TargetFirst)
-			{
-				if (AttributeInheritance == EPCGCopyPointsMetadataInheritanceMode::SourceFirst)
-				{
-					bInheritMetadataFromSource = bSourceHasMetadata || !bTargetHasMetadata;
-				}
-				else // TargetFirst
-				{
-					bInheritMetadataFromSource = !bTargetHasMetadata && bSourceHasMetadata;
-				}
+		const bool bSourceHasMetadata = (SourcePointMetadata->GetAttributeCount() > 0 && SourcePointMetadata->GetItemCountForChild() > 0);
+		const bool bTargetHasMetadata = (TargetPointMetadata->GetAttributeCount() > 0 && TargetPointMetadata->GetItemCountForChild() > 0);
 
-				RootMetadata = bInheritMetadataFromSource ? SourcePointMetadata : TargetPointMetadata;
-				NonRootMetadata = bInheritMetadataFromSource ? TargetPointMetadata : SourcePointMetadata;
-				OutPointData->InitializeFromData(bInheritMetadataFromSource ? SourcePointData : TargetPointData);
+		bool bInheritMetadataFromSource = true;
+		bool bProcessMetadata = (bSourceHasMetadata || bTargetHasMetadata);
+
+		if (AttributeInheritance == EPCGCopyPointsMetadataInheritanceMode::SourceOnly)
+		{
+			bInheritMetadataFromSource = true;
+			bProcessMetadata = bSourceHasMetadata;
+			RootMetadata = SourcePointMetadata;
+			NonRootMetadata = nullptr;
+
+			OutPointData->InitializeFromData(SourcePointData);
+		}
+		else if (AttributeInheritance == EPCGCopyPointsMetadataInheritanceMode::TargetOnly)
+		{
+			bInheritMetadataFromSource = false;
+			bProcessMetadata = bTargetHasMetadata;
+			RootMetadata = TargetPointMetadata;
+			NonRootMetadata = nullptr;
+
+			OutPointData->InitializeFromData(TargetPointData);
+		}
+		else if (AttributeInheritance == EPCGCopyPointsMetadataInheritanceMode::SourceFirst || AttributeInheritance == EPCGCopyPointsMetadataInheritanceMode::TargetFirst)
+		{
+			if (AttributeInheritance == EPCGCopyPointsMetadataInheritanceMode::SourceFirst)
+			{
+				bInheritMetadataFromSource = bSourceHasMetadata || !bTargetHasMetadata;
 			}
-			else // None
+			else // TargetFirst
 			{
-				OutPointData->InitializeFromData(SourcePointData, nullptr, /*bInheritMetadata=*/false, /*bInheritAttributes=*/false);
-
-				bProcessMetadata = false;
-				RootMetadata = NonRootMetadata = nullptr;
+				bInheritMetadataFromSource = !bTargetHasMetadata && bSourceHasMetadata;
 			}
 
-			// Always use the target actor from the target, irrespective of the source
-			OutPointData->TargetActor = TargetPointData->TargetActor;
+			RootMetadata = bInheritMetadataFromSource ? SourcePointMetadata : TargetPointMetadata;
+			NonRootMetadata = bInheritMetadataFromSource ? TargetPointMetadata : SourcePointMetadata;
+			OutPointData->InitializeFromData(bInheritMetadataFromSource ? SourcePointData : TargetPointData);
+		}
+		else // None
+		{
+			OutPointData->InitializeFromData(SourcePointData, nullptr, /*bInheritMetadata=*/false, /*bInheritAttributes=*/false);
+
+			bProcessMetadata = false;
+			RootMetadata = NonRootMetadata = nullptr;
+		}
 
-			UPCGMetadata* OutMetadata = OutPointData->Metadata;
-			check(OutMetadata);
+		// Always use the target actor from the target, irrespective of the source
+		OutPointData->TargetActor = TargetPointData->TargetActor;
 
-			TArray<FPCGMetadataAttributeBase*> AttributesToSet;
-			TArray<const FPCGMetadataAttributeBase*> NonRootAttributes;
-			TArray<TTuple<int64, int64>> AllMetadataEntries;
-			TArray<TArray<TTuple<PCGMetadataEntryKey, PCGMetadataValueKey>>> AttributeValuesToSet;
+		UPCGMetadata* OutMetadata = OutPointData->Metadata;
+		check(OutMetadata);
 
-			if (bProcessMetadata)
+		TArray<FPCGMetadataAttributeBase*> AttributesToSet;
+		TArray<const FPCGMetadataAttributeBase*> NonRootAttributes;
+		TArray<TTuple<int64, int64>> AllMetadataEntries;
+		TArray<TArray<TTuple<PCGMetadataEntryKey, PCGMetadataValueKey>>> AttributeValuesToSet;
+
+		if (bProcessMetadata)
+		{
+			TRACE_CPUPROFILER_EVENT_SCOPE(FPCGCopyPointsElement::Execute::SetupMetadata);
+			if (NonRootMetadata)
 			{
-				TRACE_CPUPROFILER_EVENT_SCOPE(FPCGCopyPointsElement::Execute::SetupMetadata);
-				if (NonRootMetadata)
-				{
-					// Prepare the attributes from the non-root that we'll need to use to copy values over
-					TArray<FName> AttributeNames;
-					TArray<EPCGMetadataTypes> AttributeTypes;
-					NonRootMetadata->GetAttributes(AttributeNames, AttributeTypes);
+				// Prepare the attributes from the non-root that we'll need to use to copy values over
+				TArray<FName> AttributeNames;
+				TArray<EPCGMetadataTypes> AttributeTypes;
+				NonRootMetadata->GetAttributes(AttributeNames, AttributeTypes);
 
-					for (const FName& AttributeName : AttributeNames)
+				for (const FName& AttributeName : AttributeNames)
+				{
+					if (!OutMetadata->HasAttribute(AttributeName))
 					{
-						if (!OutMetadata->HasAttribute(AttributeName))
+						const FPCGMetadataAttributeBase* Attribute = NonRootMetadata->GetConstAttribute(AttributeName);
+						if (FPCGMetadataAttributeBase* NewAttribute = OutMetadata->CopyAttribute(Attribute, AttributeName, /*bKeepRoot=*/false, /*bCopyEntries=*/false, /*bCopyValues=*/true))
 						{
-							const FPCGMetadataAttributeBase* Attribute = NonRootMetadata->GetConstAttribute(AttributeName);
-							if (FPCGMetadataAttributeBase* NewAttribute = OutMetadata->CopyAttribute(Attribute, AttributeName, /*bKeepRoot=*/false, /*bCopyEntries=*/false, /*bCopyValues=*/true))
-							{
-								AttributesToSet.Add(NewAttribute);
-								NonRootAttributes.Add(Attribute);
-							}
+							AttributesToSet.Add(NewAttribute);
+							NonRootAttributes.Add(Attribute);
 						}
 					}
+				}
 
-					// Considering writing to the attribute value requires a lock, we'll gather the value keys to write
-					// and do it on a 1-thread-per-attribute basis at the end
-					AttributeValuesToSet.SetNum(AttributesToSet.Num());
+				// Considering writing to the attribute value requires a lock, we'll gather the value keys to write
+				// and do it on a 1-thread-per-attribute basis at the end
+				AttributeValuesToSet.SetNum(AttributesToSet.Num());
 
-					for (TArray<TTuple<PCGMetadataEntryKey, PCGMetadataValueKey>>& AttributeValues : AttributeValuesToSet)
-					{
-						AttributeValues.SetNumUninitialized(SourcePoints.Num() * TargetPoints.Num());
-					}
+				for (TArray<TTuple<PCGMetadataEntryKey, PCGMetadataValueKey>>& AttributeValues : AttributeValuesToSet)
+				{
+					AttributeValues.SetNumUninitialized(SourcePoints.Num() * TargetPoints.Num());
 				}
-
-				// Preallocate the metadata entries array if we're going to use it
-				AllMetadataEntries.SetNumUninitialized(SourcePoints.Num() * TargetPoints.Num());
 			}
 
-			// Use implicit capture, since we capture a lot
-			FPCGAsync::AsyncPointProcessing(Context, SourcePoints.Num() * TargetPoints.Num(), OutPoints, [&](int32 Index, FPCGPoint& OutPoint)
-			{
-				const FPCGPoint& SourcePoint = SourcePoints[Index / TargetPoints.Num()];
-				const FPCGPoint& TargetPoint = TargetPoints[Index % TargetPoints.Num()];
+			// Preallocate the metadata entries array if we're going to use it
+			AllMetadataEntries.SetNumUninitialized(SourcePoints.Num() * TargetPoints.Num());
+		}
 
-				OutPoint = SourcePoint;
+		// Use implicit capture, since we capture a lot
+		FPCGAsync::AsyncPointProcessing(Context, SourcePoints.Num() * TargetPoints.Num(), OutPoints, [&](int32 Index, FPCGPoint& OutPoint)
+		{
+			const FPCGPoint& SourcePoint = SourcePoints[Index / TargetPoints.Num()];
+			const FPCGPoint& TargetPoint = TargetPoints[Index % TargetPoints.Num()];
 
-				// Set the position, rotation and scale as relative by default.
-				OutPoint.Transform = SourcePoint.Transform * TargetPoint.Transform;
+			OutPoint = SourcePoint;
 
-				// Set Rotation, Scale, and Color based on inheritance mode
-				if (RotationInheritance == EPCGCopyPointsInheritanceMode::Source)
-				{
-					OutPoint.Transform.SetRotation(SourcePoint.Transform.GetRotation());
-				}
-				else if (RotationInheritance == EPCGCopyPointsInheritanceMode::Target)
-				{
-					OutPoint.Transform.SetRotation(TargetPoint.Transform.GetRotation());
-				}
+			// Set the position, rotation and scale as relative by default.
+			OutPoint.Transform = SourcePoint.Transform * TargetPoint.Transform;
 
-				if (ScaleInheritance == EPCGCopyPointsInheritanceMode::Source)
-				{ 
-					OutPoint.Transform.SetScale3D(SourcePoint.Transform.GetScale3D());
-				}
-				else if (ScaleInheritance == EPCGCopyPointsInheritanceMode::Target)
-				{
-					OutPoint.Transform.SetScale3D(TargetPoint.Transform.GetScale3D());
-				}
+			// Set Rotation, Scale, and Color based on inheritance mode
+			if (RotationInheritance == EPCGCopyPointsInheritanceMode::Source)
+			{
+				OutPoint.Transform.SetRotation(SourcePoint.Transform.GetRotation());
+			}
+			else if (RotationInheritance == EPCGCopyPointsInheritanceMode::Target)
+			{
+				OutPoint.Transform.SetRotation(TargetPoint.Transform.GetRotation());
+			}
 
-				if (ColorInheritance == EPCGCopyPointsInheritanceMode::Relative)
-				{
-					OutPoint.Color = SourcePoint.Color * TargetPoint.Color;
-				}
-				else if (ColorInheritance == EPCGCopyPointsInheritanceMode::Source)
-				{ 
-					OutPoint.Color = SourcePoint.Color;
-				}
-				else // if (ColorInheritance == EPCGCopyPointsInheritanceMode::Target)
-				{ 
-					OutPoint.Color = TargetPoint.Color;
-				}
+			if (ScaleInheritance == EPCGCopyPointsInheritanceMode::Source)
+			{ 
+				OutPoint.Transform.SetScale3D(SourcePoint.Transform.GetScale3D());
+			}
+			else if (ScaleInheritance == EPCGCopyPointsInheritanceMode::Target)
+			{
+				OutPoint.Transform.SetScale3D(TargetPoint.Transform.GetScale3D());
+			}
 
-				// Set seed based on inheritance mode
-				if (SeedInheritance == EPCGCopyPointsInheritanceMode::Relative)
-				{
-					OutPoint.Seed = PCGHelpers::ComputeSeed(SourcePoint.Seed, TargetPoint.Seed);
-				}
-				else if (SeedInheritance == EPCGCopyPointsInheritanceMode::Target)
-				{
-					OutPoint.Seed = TargetPoint.Seed;
-				}
+			if (ColorInheritance == EPCGCopyPointsInheritanceMode::Relative)
+			{
+				OutPoint.Color = SourcePoint.Color * TargetPoint.Color;
+			}
+			else if (ColorInheritance == EPCGCopyPointsInheritanceMode::Source)
+			{ 
+				OutPoint.Color = SourcePoint.Color;
+			}
+			else // if (ColorInheritance == EPCGCopyPointsInheritanceMode::Target)
+			{ 
+				OutPoint.Color = TargetPoint.Color;
+			}
 
-				if (bProcessMetadata)
-				{
-					const FPCGPoint* RootPoint = (bInheritMetadataFromSource ? &SourcePoint : &TargetPoint);
+			// Set seed based on inheritance mode
+			if (SeedInheritance == EPCGCopyPointsInheritanceMode::Relative)
+			{
+				OutPoint.Seed = PCGHelpers::ComputeSeed(SourcePoint.Seed, TargetPoint.Seed);
+			}
+			else if (SeedInheritance == EPCGCopyPointsInheritanceMode::Target)
+			{
+				OutPoint.Seed = TargetPoint.Seed;
+			}
 
-					OutPoint.MetadataEntry = OutMetadata->AddEntryPlaceholder();
-					AllMetadataEntries[Index] = TTuple<int64, int64>(OutPoint.MetadataEntry, RootPoint->MetadataEntry);
+			if (bProcessMetadata)
+			{
+				const FPCGPoint* RootPoint = (bInheritMetadataFromSource ? &SourcePoint : &TargetPoint);
 
-					if (NonRootMetadata)
-					{
-						const FPCGPoint* NonRootPoint = (bInheritMetadataFromSource ? &TargetPoint : &SourcePoint);
+				OutPoint.MetadataEntry = OutMetadata->AddEntryPlaceholder();
+				AllMetadataEntries[Index] = TTuple<int64, int64>(OutPoint.MetadataEntry, RootPoint->MetadataEntry);
 
-						// Copy EntryToValue key mappings from NonRootAttributes - no need to do it if the non-root uses the default values
-						if (NonRootPoint->MetadataEntry != PCGInvalidEntryKey)
+				if (NonRootMetadata)
+				{
+					const FPCGPoint* NonRootPoint = (bInheritMetadataFromSource ? &TargetPoint : &SourcePoint);
+
+					// Copy EntryToValue key mappings from NonRootAttributes - no need to do it if the non-root uses the default values
+					if (NonRootPoint->MetadataEntry != PCGInvalidEntryKey)
+					{
+						for (int32 AttributeIndex = 0; AttributeIndex < NonRootAttributes.Num(); ++AttributeIndex)
 						{
-							for (int32 AttributeIndex = 0; AttributeIndex < NonRootAttributes.Num(); ++AttributeIndex)
-							{
-								AttributeValuesToSet[AttributeIndex][Index] = TTuple<PCGMetadataEntryKey, PCGMetadataValueKey>(OutPoint.MetadataEntry, NonRootAttributes[AttributeIndex]->GetValueKey(NonRootPoint->MetadataEntry));
-							}
+							AttributeValuesToSet[AttributeIndex][Index] = TTuple<PCGMetadataEntryKey, PCGMetadataValueKey>(OutPoint.MetadataEntry, NonRootAttributes[AttributeIndex]->GetValueKey(NonRootPoint->MetadataEntry));
 						}
-						else
+					}
+					else
+					{
+						for (int32 AttributeIndex = 0; AttributeIndex < NonRootAttributes.Num(); ++AttributeIndex)
 						{
-							for (int32 AttributeIndex = 0; AttributeIndex < NonRootAttributes.Num(); ++AttributeIndex)
-							{
-								AttributeValuesToSet[AttributeIndex][Index] = TTuple<PCGMetadataEntryKey, PCGMetadataValueKey>(OutPoint.MetadataEntry, PCGDefaultValueKey);
-							}
+							AttributeValuesToSet[AttributeIndex][Index] = TTuple<PCGMetadataEntryKey, PCGMetadataValueKey>(OutPoint.MetadataEntry, PCGDefaultValueKey);
 						}
 					}
 				}
-				else
-				{
-					// Reset the metadata entry if we have no metadata.
-					OutPoint.MetadataEntry = PCGInvalidEntryKey;
-				}
+			}
+			else
+			{
+				// Reset the metadata entry if we have no metadata.
+				OutPoint.MetadataEntry = PCGInvalidEntryKey;
+			}
 
-				return true;
-			});
+			return true;
+		});
 
-			if (bProcessMetadata)
+		if (bProcessMetadata)
+		{
+			TRACE_CPUPROFILER_EVENT_SCOPE(FPCGCopyPointsElement::Execute::SetMetadata);
+			check(AttributesToSet.Num() == AttributeValuesToSet.Num());
+			if (AttributesToSet.Num() > 0)
 			{
-				TRACE_CPUPROFILER_EVENT_SCOPE(FPCGCopyPointsElement::Execute::SetMetadata);
-				check(AttributesToSet.Num() == AttributeValuesToSet.Num());
-				if (AttributesToSet.Num() > 0)
-				{
-					int32 AttributeOffset = 0;
-					const int32 AttributePerDispatch = FMath::Max(1, Context->AsyncState.NumAvailableTasks);
+				int32 AttributeOffset = 0;
+				const int32 AttributePerDispatch = FMath::Max(1, Context->AsyncState.NumAvailableTasks);
 
-					while (AttributeOffset < AttributesToSet.Num())
+				while (AttributeOffset < AttributesToSet.Num())
+				{
+					const int32 AttributeCountInCurrentDispatch = FMath::Min(AttributePerDispatch, AttributesToSet.Num() - AttributeOffset);
+					ParallelFor(AttributeCountInCurrentDispatch, [AttributeOffset, &AttributesToSet, &AttributeValuesToSet](int32 WorkerIndex)
 					{
-						const int32 AttributeCountInCurrentDispatch = FMath::Min(AttributePerDispatch, AttributesToSet.Num() - AttributeOffset);
-						ParallelFor(AttributeCountInCurrentDispatch, [AttributeOffset, &AttributesToSet, &AttributeValuesToSet](int32 WorkerIndex)
-						{
-							FPCGMetadataAttributeBase* Attribute = AttributesToSet[AttributeOffset + WorkerIndex];
-							const TArray<TTuple<PCGMetadataEntryKey, PCGMetadataValueKey>>& Values = AttributeValuesToSet[AttributeOffset + WorkerIndex];
-							check(Attribute);
-							Attribute->SetValuesFromValueKeys(Values, /*bResetValueOnDefaultValueKey*/false); // no need for the reset here, our points will not have any prior value for these attributes
-						});
+						FPCGMetadataAttributeBase* Attribute = AttributesToSet[AttributeOffset + WorkerIndex];
+						const TArray<TTuple<PCGMetadataEntryKey, PCGMetadataValueKey>>& Values = AttributeValuesToSet[AttributeOffset + WorkerIndex];
+						check(Attribute);
+						Attribute->SetValuesFromValueKeys(Values, /*bResetValueOnDefaultValueKey*/false); // no need for the reset here, our points will not have any prior value for these attributes
+					});
 
-						AttributeOffset += AttributeCountInCurrentDispatch;
-					}
+					AttributeOffset += AttributeCountInCurrentDispatch;
 				}
-
-				OutMetadata->AddDelayedEntries(AllMetadataEntries);
 			}
+
+			OutMetadata->AddDelayedEntries(AllMetadataEntries);
 		}
 	}
 
diff --git a/Engine/Plugins/PCG/Source/PCG/Public/Elements/PCGCopyPoints.h b/Engine/Plugins/PCG/Source/PCG/Public/Elements/PCGCopyPoints.h
--- a/Engine/Plugins/PCG/Source/PCG/Public/Elements/PCGCopyPoints.h
+++ b/Engine/Plugins/PCG/Source/PCG/Public/Elements/PCGCopyPoints.h
@@ -83,6 +83,10 @@ public:
 	/** The method used to determine the output data tags */
 	UPROPERTY(BlueprintReadWrite, EditAnywhere, Category = Settings, meta = (PCG_Overridable))
 	EPCGCopyPointsTagInheritanceMode TagInheritance = EPCGCopyPointsTagInheritanceMode::Both;
+
+	/** If this option is set, each source point data will be copied to every target point data (cartesian product), producing N * M point data. Otherwise, will do a N:N (or N:1 or 1:N) operation, producing N point data. */
+	UPROPERTY(BlueprintReadWrite, EditAnywhere, Category = Settings, meta = (PCG_Overridable))
+	bool bCopyEachSourceOnEveryTarget = true;
 };
 
 class FPCGCopyPointsElement : public IPCGElement
