More efficient data management for PCG component generated data, prevents copying when we don't need to.

diff --git a/Engine/Plugins/PCG/Source/PCG/Private/PCGComponent.cpp b/Engine/Plugins/PCG/Source/PCG/Private/PCGComponent.cpp
--- a/Engine/Plugins/PCG/Source/PCG/Private/PCGComponent.cpp
+++ b/Engine/Plugins/PCG/Source/PCG/Private/PCGComponent.cpp
@@ -561,25 +561,58 @@ void UPCGComponent::PostProcessGraph(const FBox& InNewBounds, bool bInGenerated,
 		if (Context)
 		{
 			// TODO: should we filter based on supported serialized types here?
-			// TOOD: should reouter the contained data to this component
-			// .. and also remove it from the rootset information in the graph executor
-			//GeneratedGraphOutput = Context->InputData;
 			for (const FPCGTaggedData& TaggedData : Context->InputData.TaggedData)
 			{
 				if (ensure(TaggedData.Data))
 				{
-					// TODO: outering the first layer might not be sufficient here - might need to expose
-					// some methods in the data to traverse all the data to outer everything for serialization
-					if (UPCGData* DuplicatedData = TaggedData.Data->DuplicateData(Context))
+					// Visit the generated data; if the outer is the transient package or this component for all data, then
+					// we don't need to duplicate the data and can change the outer & flatten the data without any additional copies.
+					bool bDataCanBeStolen = true;
+
+					TaggedData.Data->VisitDataNetwork([this, &bDataCanBeStolen](const UPCGData* InData)
 					{
-						FPCGTaggedData& DuplicatedTaggedData = GeneratedGraphOutput.TaggedData.Add_GetRef(TaggedData);
-						DuplicatedTaggedData.Data = DuplicatedData;
+						if (InData && InData->GetOuter() != GetTransientPackage() && InData->GetOuter() != this)
+						{
+							bDataCanBeStolen = false;
+						}
+					});
 
-						// NOTE: Flatten before outering to this, as doing the flatten afterwards can dirty this package.
-						DuplicatedData->Flatten();
+					FPCGTaggedData OutputTaggedData = TaggedData;
 
-						DuplicatedData->Rename(nullptr, this, IsInPreviewMode() ? REN_DoNotDirty : REN_None);
+					if (!bDataCanBeStolen)
+					{
+						if (UPCGData* DuplicatedData = TaggedData.Data->DuplicateData(Context))
+						{
+							OutputTaggedData.Data = DuplicatedData;
+						}
+						else
+						{
+							// Duplication failed, don't keep that data
+							UE_LOG(LogPCG, Warning, TEXT("Failed data duplication in the PostProcessGraph - will be missing from the generated output data."));
+							continue;
+						}
 					}
+
+					// Flatten data
+					OutputTaggedData.Data->VisitDataNetwork([](const UPCGData* InData)
+					{
+						if (InData)
+						{
+							const_cast<UPCGData*>(InData)->Flatten();
+						}
+					});
+
+					// Reouter data
+					OutputTaggedData.Data->VisitDataNetwork([this](const UPCGData* InData)
+					{
+						if (InData)
+						{
+							const_cast<UPCGData*>(InData)->Rename(nullptr, this, IsInPreviewMode() ? REN_DoNotDirty : REN_None);
+						}
+					});
+
+					// Finally add to the generated output collection
+					GeneratedGraphOutput.TaggedData.Add(OutputTaggedData);
 				}
 			}
 
