[PCG] Fix case where CullInactiveDownstreamNodes can cause ReadyTasks to resize causing reallocation of the array and invalidating the Task reference

#rnx

diff --git a/Engine/Plugins/PCG/Source/PCG/Private/Graph/PCGGraphExecutor.cpp b/Engine/Plugins/PCG/Source/PCG/Private/Graph/PCGGraphExecutor.cpp
--- a/Engine/Plugins/PCG/Source/PCG/Private/Graph/PCGGraphExecutor.cpp
+++ b/Engine/Plugins/PCG/Source/PCG/Private/Graph/PCGGraphExecutor.cpp
@@ -881,35 +881,37 @@ void FPCGGraphExecutor::Execute()
 
 					if (!bNeedsToCreateActiveTask)
 					{
+						// Move Task out of array in case CullInactiveDownstreamNodes ends up queueing new tasks and resizes ReadyTasks (which can cause reallocation)
+						FPCGGraphTask SkippedTask = MoveTemp(Task);
+						ReadyTasks.RemoveAtSwap(ReadyTaskIndex);
 #if WITH_EDITOR
 						// Doing this now since we're about to modify ReadyTasks potentially reallocating while Task is a reference. 
-						if (UPCGComponent* SourceComponent = Task.SourceComponent.Get())
+						if (UPCGComponent* SourceComponent = SkippedTask.SourceComponent.Get())
 						{
-							if (Task.StackIndex != INDEX_NONE)
+							if (SkippedTask.StackIndex != INDEX_NONE)
 							{
-								const FPCGStack* Stack = Task.GetStack();
-								SourceComponent->StoreInspectionData(Stack, Task.Node, nullptr, TaskInput, CachedOutput, /*bUsedCache=*/true);
+								const FPCGStack* Stack = SkippedTask.GetStack();
+								SourceComponent->StoreInspectionData(Stack, SkippedTask.Node, nullptr, TaskInput, CachedOutput, /*bUsedCache=*/true);
 							}
 						}
 #endif
 
 						if (bDynamicTaskCulling && TaskSettings && TaskSettings->OutputPinsCanBeDeactivated() && CachedOutput.InactiveOutputPinBitmask != 0)
 						{
-							CullInactiveDownstreamNodes(Task.NodeId, CachedOutput.InactiveOutputPinBitmask);
+							CullInactiveDownstreamNodes(SkippedTask.NodeId, CachedOutput.InactiveOutputPinBitmask);
 
 #if WITH_EDITOR
-							SendInactivePinNotification(Task.Node, Task.GetStack(), CachedOutput.InactiveOutputPinBitmask);
+							SendInactivePinNotification(SkippedTask.Node, SkippedTask.GetStack(), CachedOutput.InactiveOutputPinBitmask);
 #endif
 						}
 
 						// If the task is a post execute, then we can safely clear the data after getting it from the results.
-						const bool bTaskIsPostExecute = (Task.Element == GraphCompiler.GetSharedTrivialPostGraphElement());
-
+						const bool bTaskIsPostExecute = (SkippedTask.Element == GraphCompiler.GetSharedTrivialPostGraphElement());
 						// Fast-forward cached result to stored results
-						FPCGTaskId SkippedTaskId = Task.NodeId;
+						FPCGTaskId SkippedTaskId = SkippedTask.NodeId;
 						StoreResults(SkippedTaskId, CachedOutput, bTaskIsPostExecute);
-						delete Task.Context;
-						ReadyTasks.RemoveAtSwap(ReadyTaskIndex);
+						delete SkippedTask.Context;
+
 						QueueNextTasks(SkippedTaskId);
 						bAnyTaskEnded = true;
 
