[PCG] Fixed main thread predicate in the Create Spline node. Will be updated once we process overrides better.

diff --git a/Engine/Plugins/PCG/Source/PCG/Private/Elements/PCGCreateSpline.cpp b/Engine/Plugins/PCG/Source/PCG/Private/Elements/PCGCreateSpline.cpp
--- a/Engine/Plugins/PCG/Source/PCG/Private/Elements/PCGCreateSpline.cpp
+++ b/Engine/Plugins/PCG/Source/PCG/Private/Elements/PCGCreateSpline.cpp
@@ -41,6 +41,13 @@ FPCGElementPtr UPCGCreateSplineSettings::CreateElement() const
 	return MakeShared<FPCGCreateSplineElement>();
 }
 
+bool FPCGCreateSplineElement::CanExecuteOnlyOnMainThread(FPCGContext* Context) const
+{
+	// Creating the spline component requires to run on the main thread, but if the settings/context aren't available we'll err on the side of caution.
+	const UPCGCreateSplineSettings* Settings = Context ? Context->GetInputSettings<UPCGCreateSplineSettings>() : nullptr;
+	return !Settings || Settings->Mode == EPCGCreateSplineMode::CreateComponent;
+}
+
 bool FPCGCreateSplineElement::IsCacheable(const UPCGSettings* InSettings) const
 {
 	const UPCGCreateSplineSettings* Settings = Cast<const UPCGCreateSplineSettings>(InSettings);
diff --git a/Engine/Plugins/PCG/Source/PCG/Public/Elements/PCGCreateSpline.h b/Engine/Plugins/PCG/Source/PCG/Public/Elements/PCGCreateSpline.h
--- a/Engine/Plugins/PCG/Source/PCG/Public/Elements/PCGCreateSpline.h
+++ b/Engine/Plugins/PCG/Source/PCG/Public/Elements/PCGCreateSpline.h
@@ -71,7 +71,10 @@ public:
 
 class FPCGCreateSplineElement : public IPCGElement
 {
-protected:
+public:
+	virtual bool CanExecuteOnlyOnMainThread(FPCGContext* Context) const override;
 	virtual bool IsCacheable(const UPCGSettings* InSettings) const override;
+
+protected:	
 	virtual bool ExecuteInternal(FPCGContext* Context) const override;
 };
\ No newline at end of file
