[PCG] V1 for Grammar support in PCG, with Shape Grammar (with Slicers)

* Add a Spline to Segment (name pending) to convert a linear spline into points, representing the segment
* Add a grammar parser to parse our grammar
* Add Volume Slicer (name pending) to duplicate a spline along an axis, following a grammar
* Add a Segment Slicer (name pending) to slice a point into multiple points, following a grammar
* Add redirectors for PCG sample plugin
* Removes duplicated nodes

diff --git a/Engine/Plugins/PCG/Source/PCG/Private/Elements/Grammar/PCGSegmentSlicer.cpp b/Engine/Plugins/PCG/Source/PCG/Private/Elements/Grammar/PCGSegmentSlicer.cpp
new file mode 100644
--- /dev/null
+++ b/Engine/Plugins/PCG/Source/PCG/Private/Elements/Grammar/PCGSegmentSlicer.cpp
@@ -0,0 +1,377 @@
+// Copyright Epic Games, Inc. All Rights Reserved.
+
+#include "Elements/Grammar/PCGSegmentSlicer.h"
+
+#include "PCGContext.h"
+#include "PCGParamData.h"
+#include "Data/PCGPointData.h"
+#include "Elements/Metadata/PCGMetadataElementCommon.h"
+#include "Grammar/PCGGrammarParser.h"
+#include "Metadata/Accessors/IPCGAttributeAccessor.h"
+#include "Metadata/Accessors/PCGAttributeAccessorHelpers.h"
+#include "Metadata/Accessors/PCGAttributeAccessorKeys.h"
+
+#define LOCTEXT_NAMESPACE "PCGSegmentSlicerElement"
+
+class PCGSegmentSlicerHelpers
+{
+public:
+	struct FParameters
+	{
+		FPCGMetadataAttribute<FName>* SymbolAttribute = nullptr;
+		FPCGMetadataAttribute<FVector4>* DebugColorAttribute = nullptr;
+		FPCGMetadataAttribute<int32>* ModuleIndexAttribute = nullptr;
+		FPCGMetadataAttribute<bool>* ExtremityAttribute = nullptr;
+		FPCGMetadataAttribute<int32>* ExtremityNeighborIndexAttribute = nullptr;
+
+		FPCGSlicingBaseElement::FPCGModulesInfoMap ModulesInfo;
+		TMap<FString, TArray<FPCGTokenizedGrammar>> CachedGrammars;
+		TArray<int32> CornerIndexes;
+
+		FVector SlicingDirection;
+		FVector PerpendicularSlicingDirection;
+
+		const UPCGSegmentSlicerSettings* Settings = nullptr;
+		FPCGContext* Context = nullptr;
+		const TArray<FPCGPoint>* InPoints = nullptr;
+		TArray<FPCGPoint>* OutPoints = nullptr;
+		UPCGMetadata* OutputMetadata = nullptr;
+	};
+
+	static void Process(const FPCGSegmentSlicerElement* InElement, FParameters& InOutParameters, const FString& InGrammar, const bool bFlipAxis, int32 Index)
+	{
+		if (!InOutParameters.CachedGrammars.Contains(InGrammar))
+		{
+			double MinSize;
+			InOutParameters.CachedGrammars.Emplace(InGrammar, InElement->GetTokenizeGrammar(InOutParameters.Context, InGrammar, InOutParameters.ModulesInfo, MinSize));
+		}
+
+		const TArray<FPCGTokenizedGrammar>& CurrentTokenizeGrammar = InOutParameters.CachedGrammars[InGrammar];
+
+		if (CurrentTokenizeGrammar.IsEmpty())
+		{
+			return;
+		}
+
+		const FPCGPoint& Point = (*InOutParameters.InPoints)[Index];
+		FTransform TransformNoTranslation = Point.Transform;
+		TransformNoTranslation.SetLocation(FVector::ZeroVector);
+
+		const int32 FirstModuleIndex = InOutParameters.OutPoints->Num();
+
+		FBox Segment = Point.GetLocalBounds();
+		FVector FullExtents = Point.GetScaledExtents() * 2.0;
+		if (bFlipAxis)
+		{
+			// Swap coordinates on the slicing direction
+			const FVector PreviousMin = Segment.Min;
+			Segment.Min = Segment.Min * InOutParameters.PerpendicularSlicingDirection + Segment.Max * InOutParameters.SlicingDirection;
+			Segment.Max = Segment.Max * InOutParameters.PerpendicularSlicingDirection + PreviousMin * InOutParameters.SlicingDirection;
+			FullExtents *= (InOutParameters.PerpendicularSlicingDirection - InOutParameters.SlicingDirection);
+		}
+
+		const FVector Direction = TransformNoTranslation.TransformVectorNoScale(InOutParameters.SlicingDirection).GetSafeNormal();
+		const FVector OtherDirection = TransformNoTranslation.TransformVectorNoScale(FullExtents * InOutParameters.PerpendicularSlicingDirection) * 0.5;
+		const FVector HalfExtents2D = FullExtents * InOutParameters.PerpendicularSlicingDirection * 0.5;
+		const double Size = FullExtents.Dot(InOutParameters.SlicingDirection);
+
+		TArray<PCGSlicingBase::PCGSubDivModuleInstance<FPCGTokenizedGrammar>> ModulesInstances;
+		double RemainingSubdivide;
+		const bool bSubdivideSuccess = PCGSlicingBase::Subdivide(CurrentTokenizeGrammar, Size, ModulesInstances, RemainingSubdivide, InOutParameters.Context);
+
+		if (!bSubdivideSuccess)
+		{
+			return;
+		}
+
+		if (!InOutParameters.Settings->bAcceptIncompleteSlicing && !FMath::IsNearlyZero(RemainingSubdivide))
+		{
+			PCGLog::LogWarningOnGraph(LOCTEXT("FailSliceFullLength", "One segment has an incomplete slicing (grammar doesn't fit the whole segment)."), InOutParameters.Context);
+			return;
+		}
+
+		// Now we have our segment subdivided, create the final points
+		const FVector InitialPos = Point.Transform.TransformPosition(Segment.Min);
+		FVector CurrentPos = InitialPos;
+		int32 ModuleIndex = 0;
+
+		for (int32 ModuleInstanceIndex = 0; ModuleInstanceIndex < ModulesInstances.Num(); ModuleInstanceIndex++)
+		{
+			const PCGSlicingBase::PCGSubDivModuleInstance<FPCGTokenizedGrammar>& ModuleInstance = ModulesInstances[ModuleInstanceIndex];
+			for (int32 j = 0; j < ModuleInstance.NumRepeat; ++j)
+			{
+				for (int32 SymbolIndex = 0; SymbolIndex < ModuleInstance.Module->Symbols.Num(); ++SymbolIndex)
+				{
+					const FName Symbol = ModuleInstance.Module->Symbols[SymbolIndex];
+					const FVector Scale = FVector::OneVector + (InOutParameters.SlicingDirection * ModuleInstance.ExtraScales[SymbolIndex]);
+					const FPCGSlicingModule& SlicingModule = InOutParameters.ModulesInfo[Symbol];
+
+					const bool bIsFirstModule = (ModuleInstanceIndex == 0) && (j == 0) && (SymbolIndex == 0);
+					const bool bIsLastModule = (ModuleInstanceIndex == ModulesInstances.Num() - 1) && (j == ModuleInstance.NumRepeat - 1) && (SymbolIndex == ModuleInstance.Module->Symbols.Num() - 1);
+					const double HalfDisplacement = SlicingModule.Size * 0.5;
+					const double HalfScaledDisplacement = Scale.Dot(InOutParameters.SlicingDirection) * HalfDisplacement;
+
+					const FVector LocalBoundsExtents = InOutParameters.SlicingDirection * HalfDisplacement + HalfExtents2D;
+					const FVector HalfStep = HalfScaledDisplacement * Direction;
+					const FVector Position = CurrentPos + HalfStep;
+					CurrentPos = Position + HalfStep;
+
+					FPCGPoint& OutPoint = InOutParameters.OutPoints->Emplace_GetRef(Point);
+					OutPoint.Transform = FTransform(Point.Transform.GetRotation(), Position + OtherDirection, Scale);
+					OutPoint.SetLocalBounds(FBox(-LocalBoundsExtents, LocalBoundsExtents));
+					InOutParameters.OutputMetadata->InitializeOnSet(OutPoint.MetadataEntry);
+					if (InOutParameters.SymbolAttribute)
+					{
+						InOutParameters.SymbolAttribute->SetValue(OutPoint.MetadataEntry, Symbol);
+					}
+
+					if (InOutParameters.DebugColorAttribute)
+					{
+						InOutParameters.DebugColorAttribute->SetValue(OutPoint.MetadataEntry, FVector4(SlicingModule.DebugColor, 1.0));
+					}
+
+					if (InOutParameters.ModuleIndexAttribute)
+					{
+						InOutParameters.ModuleIndexAttribute->SetValue(OutPoint.MetadataEntry, ModuleIndex++);
+					}
+
+					if (bIsFirstModule || bIsLastModule)
+					{
+						if (InOutParameters.ExtremityAttribute)
+						{
+							InOutParameters.ExtremityAttribute->SetValue(OutPoint.MetadataEntry, true);
+						}
+					}
+				}
+			}
+		}
+
+		if (bFlipAxis)
+		{
+			InOutParameters.CornerIndexes.Add(InOutParameters.OutPoints->Num() - 1);
+			InOutParameters.CornerIndexes.Add(FirstModuleIndex);
+		}
+		else
+		{
+			InOutParameters.CornerIndexes.Add(FirstModuleIndex);
+			InOutParameters.CornerIndexes.Add(InOutParameters.OutPoints->Num() - 1);
+		}
+	}
+};
+
+#if WITH_EDITOR
+FName UPCGSegmentSlicerSettings::GetDefaultNodeName() const
+{
+	return FName(TEXT("SegmentSlicer"));
+}
+
+FText UPCGSegmentSlicerSettings::GetDefaultNodeTitle() const
+{
+	return LOCTEXT("NodeTitle", "Segment Slicer");
+}
+#endif // WITH_EDITOR
+
+FPCGElementPtr UPCGSegmentSlicerSettings::CreateElement() const
+{
+	return MakeShared<FPCGSegmentSlicerElement>();
+}
+
+TArray<FPCGPinProperties> UPCGSegmentSlicerSettings::InputPinProperties() const
+{
+	TArray<FPCGPinProperties> Result;
+	FPCGPinProperties& InputPin = Result.Emplace_GetRef(PCGPinConstants::DefaultInputLabel, EPCGDataType::Point);
+	InputPin.SetRequiredPin();
+
+	if (bModuleInfoAsInput)
+	{
+		FPCGPinProperties& ModuleInfoPin = Result.Emplace_GetRef(PCGSlicingBaseConstants::ModulesInfoPinLabel, EPCGDataType::Param);
+		ModuleInfoPin.SetRequiredPin();
+	}
+
+	return Result;
+}
+
+TArray<FPCGPinProperties> UPCGSegmentSlicerSettings::OutputPinProperties() const
+{
+	return Super::DefaultPointOutputPinProperties();
+}
+
+bool FPCGSegmentSlicerElement::ExecuteInternal(FPCGContext* InContext) const
+{
+	TRACE_CPUPROFILER_EVENT_SCOPE(FPCGSegmentSlicerElement::Execute);
+
+	check(InContext);
+
+	const UPCGSegmentSlicerSettings* Settings = InContext->GetInputSettings<UPCGSegmentSlicerSettings>();
+	check(Settings);
+
+	const TArray<FPCGTaggedData> Inputs = InContext->InputData.GetInputsByPin(PCGPinConstants::DefaultInputLabel);
+	TArray<FPCGTaggedData>& Outputs = InContext->OutputData.TaggedData;
+
+	PCGSegmentSlicerHelpers::FParameters Parameters{};
+
+	Parameters.SlicingDirection = FVector::ZeroVector;
+
+	switch (Settings->SlicingAxis)
+	{
+	case EPCGSplitAxis::X:
+		Parameters.SlicingDirection = FVector::XAxisVector;
+		break;
+	case EPCGSplitAxis::Y:
+		Parameters.SlicingDirection = FVector::YAxisVector;
+		break;
+	case EPCGSplitAxis::Z:
+		Parameters.SlicingDirection = FVector::ZAxisVector;
+		break;
+	default:
+		PCGLog::LogErrorOnGraph(LOCTEXT("InvalidAxis", "Invalid Slicing Axis enum value."), InContext);
+		return true;
+	}
+
+	Parameters.PerpendicularSlicingDirection = FVector::OneVector - Parameters.SlicingDirection;
+
+	const UPCGParamData* ModuleInfoParamData = nullptr;
+	Parameters.ModulesInfo = GetModulesInfoMap(InContext, Settings, ModuleInfoParamData);
+
+	for (const FPCGTaggedData& Input : Inputs)
+	{
+		const UPCGPointData* InputPointData = Cast<const UPCGPointData>(Input.Data);
+
+		if (!InputPointData)
+		{
+			continue;
+		}
+		
+		TUniquePtr<const IPCGAttributeAccessor> GrammarAccessor;
+		TUniquePtr<const IPCGAttributeAccessor> FlipAxisAccessor;
+		TUniquePtr<const IPCGAttributeAccessorKeys> Keys;
+
+		if (Settings->bGrammarAsAttribute)
+		{
+			const FPCGAttributePropertyInputSelector Selector = Settings->GrammarAttribute.CopyAndFixLast(InputPointData);
+			GrammarAccessor = PCGAttributeAccessorHelpers::CreateConstAccessor(InputPointData, Selector);
+			Keys = PCGAttributeAccessorHelpers::CreateConstKeys(InputPointData, Selector);
+			if (!GrammarAccessor || !Keys)
+			{
+				PCGLog::LogErrorOnGraph(FText::Format(LOCTEXT("GrammarAccessor", "Attribute {0} was not found for the grammar."), Selector.GetDisplayText()), InContext);
+				continue;
+			}
+		}
+
+		if (Settings->bFlipAxisAsAttribute)
+		{
+			const FPCGAttributePropertyInputSelector Selector = Settings->FlipAxisAttribute.CopyAndFixLast(InputPointData);
+			FlipAxisAccessor = PCGAttributeAccessorHelpers::CreateConstAccessor(InputPointData, Selector);
+			if (!Keys)
+			{
+				Keys = PCGAttributeAccessorHelpers::CreateConstKeys(InputPointData, Selector);
+			}
+
+			if (!FlipAxisAccessor || !Keys)
+			{
+				PCGLog::LogErrorOnGraph(FText::Format(LOCTEXT("FlipAxisAccessor", "Attribute {0} was not found for the flip axis attribute."), Selector.GetDisplayText()), InContext);
+				continue;
+			}
+		}
+
+		UPCGPointData* OutputPointData = FPCGContext::NewObject_AnyThread<UPCGPointData>(InContext);
+		OutputPointData->InitializeFromData(InputPointData);
+
+		Parameters.Settings = Settings;
+		Parameters.Context = InContext;
+		Parameters.InPoints = &InputPointData->GetPoints();
+		Parameters.OutPoints = &OutputPointData->GetMutablePoints();
+		Parameters.OutputMetadata = OutputPointData->Metadata;
+
+		auto CreateAndValidateAttribute = [InContext, OutputMetadata = OutputPointData->Metadata]<typename T>(const FName AttributeName, const T DefaultValue, const bool bShouldCreate, FPCGMetadataAttribute<T>*& OutAttribute) -> bool
+		{
+			if (bShouldCreate)
+			{
+				OutAttribute = OutputMetadata->FindOrCreateAttribute<T>(AttributeName, DefaultValue, false, true);
+				if (!OutAttribute)
+				{
+					PCGLog::LogWarningOnGraph(LOCTEXT("FailCreateAttributes", "Fail to create output attributes."), InContext);
+					return false;
+				}
+			}
+
+			return true;
+		};
+
+		if (!CreateAndValidateAttribute(Settings->SymbolAttributeName, FName(NAME_None), true, Parameters.SymbolAttribute)
+			|| !CreateAndValidateAttribute(Settings->DebugColorAttributeName, FVector4::Zero(), Settings->bOutputDebugColorAttribute, Parameters.DebugColorAttribute)
+			|| !CreateAndValidateAttribute(Settings->ModuleIndexAttributeName, -1, Settings->bOutputModuleIndexAttribute, Parameters.ModuleIndexAttribute)
+			|| !CreateAndValidateAttribute(Settings->ExtremityAttributeName, false, Settings->bOutputExtremityAttribute, Parameters.ExtremityAttribute)
+			|| !CreateAndValidateAttribute(Settings->ExtremityNeighborIndexAttributeName, -1, Settings->bOutputExtremityNeighborIndexAttribute, Parameters.ExtremityNeighborIndexAttribute))
+		{
+			continue;
+		}
+
+
+		if (Settings->bGrammarAsAttribute && Settings->bFlipAxisAsAttribute)
+		{
+			auto Process = [this, &Parameters](const FString& InGrammar, const bool bFlipAxis, int32 Index) -> void
+			{
+				PCGSegmentSlicerHelpers::Process(this, Parameters, InGrammar, bFlipAxis, Index);
+			};
+
+			PCGMetadataElementCommon::ApplyOnMultiAccessors<FString, bool>(*Keys, { GrammarAccessor.Get(), FlipAxisAccessor.Get() }, Process);
+		}
+		else if (Settings->bGrammarAsAttribute)
+		{
+			auto Process = [this, &Parameters, bShouldFlipAxis = Settings->bShouldFlipAxis](const FString& InGrammar, int32 Index) -> void
+			{
+				PCGSegmentSlicerHelpers::Process(this, Parameters, InGrammar, bShouldFlipAxis, Index);
+			};
+
+			PCGMetadataElementCommon::ApplyOnAccessor<FString>(*Keys, *GrammarAccessor, Process);
+		}
+		else if (Settings->bFlipAxisAsAttribute)
+		{
+			auto Process = [this, &Parameters, Grammar = Settings->Grammar](const bool bFlipAxis, int32 Index) -> void
+			{
+				PCGSegmentSlicerHelpers::Process(this, Parameters, Grammar, bFlipAxis, Index);
+			};
+
+			PCGMetadataElementCommon::ApplyOnAccessor<bool>(*Keys, *FlipAxisAccessor, Process);
+		}
+		else
+		{
+			for (int32 SegmentIndex = 0; SegmentIndex < Parameters.InPoints->Num(); ++SegmentIndex)
+			{
+				PCGSegmentSlicerHelpers::Process(this, Parameters, Settings->Grammar, Settings->bShouldFlipAxis, SegmentIndex);
+			}
+		}
+
+		if (!Parameters.OutPoints->IsEmpty())
+		{
+			// Set the extremity neighbor indexes
+			if (Parameters.ExtremityNeighborIndexAttribute)
+			{
+				// Go 2 by 2
+				check(Parameters.CornerIndexes.Num() % 2 == 0);
+
+				for (int32 j = 0; j < Parameters.CornerIndexes.Num(); j += 2)
+				{
+					const int32 PreviousModuleIndex = Parameters.CornerIndexes[j == 0 ? Parameters.CornerIndexes.Num() - 1 : j - 1];
+					const int32 NextModuleIndex = Parameters.CornerIndexes[j == Parameters.CornerIndexes.Num() - 2 ? 0 : j + 2];
+
+					Parameters.ExtremityNeighborIndexAttribute->SetValue((*Parameters.OutPoints)[Parameters.CornerIndexes[j]].MetadataEntry, PreviousModuleIndex);
+					Parameters.ExtremityNeighborIndexAttribute->SetValue((*Parameters.OutPoints)[Parameters.CornerIndexes[j+1]].MetadataEntry, NextModuleIndex);
+				}
+			}
+
+			FPCGTaggedData& Output = Outputs.Emplace_GetRef(Input);
+			Output.Data = OutputPointData;
+		}
+	}
+
+	if (Settings->bForwardAttributesFromModulesInfo && ModuleInfoParamData)
+	{
+		MatchAndSetAttributes(Inputs, Outputs, ModuleInfoParamData, Settings);
+	}
+
+	return true;
+}
+
+#undef LOCTEXT_NAMESPACE
diff --git a/Engine/Plugins/PCG/Source/PCG/Private/Elements/Grammar/PCGSlicingBase.cpp b/Engine/Plugins/PCG/Source/PCG/Private/Elements/Grammar/PCGSlicingBase.cpp
new file mode 100644
--- /dev/null
+++ b/Engine/Plugins/PCG/Source/PCG/Private/Elements/Grammar/PCGSlicingBase.cpp
@@ -0,0 +1,412 @@
+// Copyright Epic Games, Inc. All Rights Reserved.
+
+#include "Elements/Grammar/PCGSlicingBase.h"
+
+#include "PCGContext.h"
+#include "PCGParamData.h"
+#include "Data/PCGPointData.h"
+#include "Elements/Metadata/PCGMetadataElementCommon.h"
+#include "Grammar/PCGGrammarParser.h"
+#include "Metadata/Accessors/IPCGAttributeAccessor.h"
+#include "Metadata/Accessors/PCGAttributeAccessorHelpers.h"
+#include "Metadata/Accessors/PCGAttributeAccessorKeys.h"
+
+#define LOCTEXT_NAMESPACE "PCGSlicingBaseElement"
+
+namespace PCGSlicingBase
+{
+	static const FText AccessorFailGetText = LOCTEXT("AccessorFailGet", "Couldn't retrieve attribute {0} value, could it be the wrong type? Expected type: {1}, Attribute Type: {2}.");
+	static const FText AccessorGrammarFailCreate = LOCTEXT("GrammarAccessor", "Attribute {0} was not found for the grammar.");
+	static const FText DuplicatedSymbolText = LOCTEXT("SymbolDuplicate", "Symbol {0} is duplicated, ignored.");
+}
+
+FPCGSlicingBaseElement::FPCGModulesInfoMap FPCGSlicingBaseElement::GetModulesInfoMap(FPCGContext* InContext, const UPCGSlicingBaseSettings* InSettings, const UPCGParamData*& OutModuleInfoParamData) const
+{
+	FPCGModulesInfoMap ModulesInfo;
+	OutModuleInfoParamData = nullptr;
+
+	if (InSettings->bModuleInfoAsInput)
+	{
+		const TArray<FPCGTaggedData> ModulesInfoInputs = InContext->InputData.GetInputsByPin(PCGSlicingBaseConstants::ModulesInfoPinLabel);
+
+		if (ModulesInfoInputs.IsEmpty())
+		{
+			PCGLog::LogWarningOnGraph(LOCTEXT("NoModuleInfo", "No data was found on the module info pin."), InContext);
+			return ModulesInfo;
+		}
+
+		const UPCGParamData* ParamData = Cast<UPCGParamData>(ModulesInfoInputs[0].Data);
+		if (!ParamData)
+		{
+			PCGLog::LogWarningOnGraph(LOCTEXT("ModuleInfoWrongType", "Module info input is not of type attribute set."), InContext);
+			return ModulesInfo;
+		}
+
+		static const FName SymbolAttributeName = PCGSlicingBaseConstants::SymbolAttributeName;
+		static const FName SizeAttributeName = GET_MEMBER_NAME_CHECKED(FPCGSlicingModule, Size);
+		static const FName ScalableAttributeName = GET_MEMBER_NAME_CHECKED(FPCGSlicingModule, bScalable);
+		static const FName DebugColorAttributeName = GET_MEMBER_NAME_CHECKED(FPCGSlicingModule, DebugColor);
+		static const FName ExtraScalableAttributeName = TEXT("Scalable"); // To also support the name without a "b".
+
+		const TUniquePtr<const IPCGAttributeAccessorKeys> Keys = PCGAttributeAccessorHelpers::CreateConstKeys(ParamData, FPCGAttributePropertySelector::CreateAttributeSelector(SymbolAttributeName));
+
+		const TUniquePtr<const IPCGAttributeAccessor> SymbolAccessor = PCGAttributeAccessorHelpers::CreateConstAccessor(ParamData, FPCGAttributePropertySelector::CreateAttributeSelector(SymbolAttributeName));
+		const TUniquePtr<const IPCGAttributeAccessor> SizeAccessor = PCGAttributeAccessorHelpers::CreateConstAccessor(ParamData, FPCGAttributePropertySelector::CreateAttributeSelector(SizeAttributeName));
+		TUniquePtr<const IPCGAttributeAccessor> ScalableAccessor = PCGAttributeAccessorHelpers::CreateConstAccessor(ParamData, FPCGAttributePropertySelector::CreateAttributeSelector(ScalableAttributeName));
+		const TUniquePtr<const IPCGAttributeAccessor> DebugColorAccessor = PCGAttributeAccessorHelpers::CreateConstAccessor(ParamData, FPCGAttributePropertySelector::CreateAttributeSelector(DebugColorAttributeName));
+
+		if (!ScalableAccessor.IsValid())
+		{
+			// Retry with Scalable without a b.
+			ScalableAccessor = PCGAttributeAccessorHelpers::CreateConstAccessor(ParamData, FPCGAttributePropertySelector::CreateAttributeSelector(ExtraScalableAttributeName));
+		}
+
+		auto GetRange = [InContext, &Keys, ParamData]<typename T>(TArray<T>& OutValues, const FName AttributeName)
+		{
+			const TUniquePtr<const IPCGAttributeAccessor> Accessor = PCGAttributeAccessorHelpers::CreateConstAccessor(ParamData, FPCGAttributePropertySelector::CreateAttributeSelector(AttributeName));
+			if (!Accessor || !Keys)
+			{
+				PCGLog::LogErrorOnGraph(FText::Format(LOCTEXT("ModuleInfoMissingAttribute", "Module info input is missing attribute {0}."), FText::FromName(AttributeName)), InContext);
+				return false;
+			}
+
+			OutValues.SetNum(Keys->GetNum());
+
+			if (!Accessor->GetRange<T>(OutValues, 0, *Keys, EPCGAttributeAccessorFlags::AllowBroadcastAndConstructible))
+			{
+				PCGLog::LogErrorOnGraph(FText::Format(PCGSlicingBase::AccessorFailGetText, FText::FromName(AttributeName), PCG::Private::GetTypeNameText<T>(), PCG::Private::GetTypeNameText(Accessor->GetUnderlyingType())), InContext);
+				return false;
+			}
+
+			return true;
+		};
+
+		TArray<FName> Symbols;
+		TArray<double> Sizes;
+		TArray<bool> Scalables;
+		TArray<FVector4> DebugColors;
+
+		if (!GetRange(Symbols, SymbolAttributeName)
+			|| !GetRange(Sizes, SizeAttributeName)
+			|| !GetRange(Scalables, ScalableAttributeName)
+			|| !GetRange(DebugColors, DebugColorAttributeName))
+		{
+			return ModulesInfo;
+		}
+
+		check(Keys);
+		ModulesInfo.Reserve(Keys->GetNum());
+
+		for (int32 i = 0; i < Keys->GetNum(); ++i)
+		{
+			if (ModulesInfo.Contains(Symbols[i]))
+			{
+				PCGLog::LogWarningOnGraph(FText::Format(PCGSlicingBase::DuplicatedSymbolText, FText::FromName(Symbols[i])), InContext);
+				continue;
+			}
+
+			FPCGSlicingModule& Module = ModulesInfo.Emplace(Symbols[i]);
+			Module.Symbol = Symbols[i];
+			Module.Size = Sizes[i];
+			Module.bScalable = Scalables[i];
+			Module.DebugColor = DebugColors[i];
+		}
+
+		OutModuleInfoParamData = ParamData;
+	}
+	else
+	{
+		ModulesInfo.Reserve(InSettings->ModulesInfo.Num());
+		for (const FPCGSlicingModule& SlicingModule : InSettings->ModulesInfo)
+		{
+			if (ModulesInfo.Contains(SlicingModule.Symbol))
+			{
+				PCGLog::LogWarningOnGraph(FText::Format(PCGSlicingBase::DuplicatedSymbolText, FText::FromName(SlicingModule.Symbol)), InContext);
+				continue;
+			}
+
+			ModulesInfo.Emplace(SlicingModule.Symbol, SlicingModule);
+		}
+	}
+
+	return ModulesInfo;
+}
+
+TArray<FPCGTokenizedGrammar> FPCGSlicingBaseElement::GetTokenizeGrammar(FPCGContext* InContext, const UPCGData* InputData, const UPCGSlicingBaseSettings* InSettings, const FPCGModulesInfoMap& InModulesInfo, double& OutMinSize) const
+{
+	FString Grammar = InSettings->Grammar;
+
+	if (InSettings->bGrammarAsAttribute)
+	{
+		const FPCGAttributePropertyInputSelector Selector = InSettings->GrammarAttribute.CopyAndFixLast(InputData);
+		const TUniquePtr<const IPCGAttributeAccessor> Accessor = PCGAttributeAccessorHelpers::CreateConstAccessor(InputData, Selector);
+		if (!Accessor)
+		{
+			PCGLog::LogErrorOnGraph(FText::Format(PCGSlicingBase::AccessorGrammarFailCreate, Selector.GetDisplayText()), InContext);
+			return {};
+		}
+
+		if (!Accessor->Get(Grammar, FPCGAttributeAccessorKeysEntries(PCGInvalidEntryKey), EPCGAttributeAccessorFlags::AllowBroadcastAndConstructible))
+		{
+			PCGLog::LogErrorOnGraph(FText::Format(PCGSlicingBase::AccessorFailGetText, Selector.GetDisplayText(), PCG::Private::GetTypeNameText<FString>(), PCG::Private::GetTypeNameText(Accessor->GetUnderlyingType())), InContext);
+			return {};
+		}
+	}
+
+	return GetTokenizeGrammar(InContext, Grammar, InModulesInfo, OutMinSize);
+}
+
+TArray<FPCGTokenizedGrammar> FPCGSlicingBaseElement::GetTokenizeGrammar(FPCGContext* InContext, const FString& InGrammar, const FPCGModulesInfoMap& InModulesInfo, double& OutMinSize) const
+{
+	const FPCGGrammarResult Result = PCGGrammar::Parse(InGrammar);
+
+	if (!Result.bSuccess)
+	{
+		PCGLog::LogErrorOnGraph(LOCTEXT("GrammarParseFail", "Problem while parsing grammar:"), InContext);
+
+		for (const FPCGGrammarResult::FLog& Log : Result.GetLogs())
+		{
+			switch (Log.Verbosity)
+			{
+			case FPCGGrammarResult::ELogType::Error:
+				PCGLog::LogErrorOnGraph(Log.Message, InContext);
+				break;
+			case FPCGGrammarResult::ELogType::Warning:
+				PCGLog::LogWarningOnGraph(Log.Message, InContext);
+				break;
+			default:
+				UE_LOG(LogPCG, Log, TEXT("%s"), *Log.Message.ToString());
+				break;
+			}
+		}
+
+		return {};
+	}
+
+	TArray<FPCGTokenizedGrammar> TokenizeGrammar;
+	OutMinSize = 0.0;
+
+	for (const PCGGrammar::FModuleDescriptor& ModuleDescriptor : Result.Modules)
+	{
+		FPCGTokenizedGrammar& CurrentModule = TokenizeGrammar.Emplace_GetRef();
+		CurrentModule.NumRepeat = ModuleDescriptor.Repetitions;
+
+		for (const PCGGrammar::FModuleDescriptor::FSubmodule& SubmoduleDescriptor : ModuleDescriptor.Submodules)
+		{
+			if (const FPCGSlicingModule* It = InModulesInfo.Find(SubmoduleDescriptor.ID))
+			{
+				CurrentModule.Symbols.Add(SubmoduleDescriptor.ID);
+				CurrentModule.Size += It->Size;
+				CurrentModule.bScalable |= It->bScalable;
+				CurrentModule.AreSymbolsScalable.Add(It->bScalable);
+				CurrentModule.SymbolSizes.Add(It->Size);
+			}
+		}
+		if (CurrentModule.NumRepeat > 0)
+		{
+			OutMinSize += CurrentModule.Size * CurrentModule.NumRepeat;
+		}
+
+		if (CurrentModule.Symbols.IsEmpty())
+		{
+			// If we have no symbol, we skip.
+			continue;
+		}
+	}
+
+	return TokenizeGrammar;
+}
+
+TMap<FString, TArray<FPCGTokenizedGrammar>> FPCGSlicingBaseElement::GetTokenizeGrammarForPoints(FPCGContext* InContext, const UPCGPointData* InputData, const UPCGSlicingBaseSettings* InSettings, const FPCGModulesInfoMap& InModulesInfo, double& OutMinSize) const
+{
+	TMap<FString, TArray<FPCGTokenizedGrammar>> Result;
+
+	if (InSettings->bGrammarAsAttribute)
+	{
+		const FPCGAttributePropertyInputSelector Selector = InSettings->GrammarAttribute.CopyAndFixLast(InputData);
+		const TUniquePtr<const IPCGAttributeAccessor> Accessor = PCGAttributeAccessorHelpers::CreateConstAccessor(InputData, Selector);
+		const TUniquePtr<const IPCGAttributeAccessorKeys> Keys = PCGAttributeAccessorHelpers::CreateConstKeys(InputData, Selector);
+		if (!Accessor || !Keys)
+		{
+			PCGLog::LogErrorOnGraph(FText::Format(PCGSlicingBase::AccessorGrammarFailCreate, Selector.GetDisplayText()), InContext);
+			return Result;
+		}
+
+		const bool bSuccess = PCGMetadataElementCommon::ApplyOnAccessor<FString>(*Keys, *Accessor, [&Result](const FString& InValue, int32)
+		{
+			if (!Result.Contains(InValue))
+			{
+				Result.Emplace(InValue);
+			}
+		});
+
+		if (!bSuccess)
+		{
+			PCGLog::LogErrorOnGraph(FText::Format(PCGSlicingBase::AccessorFailGetText, Selector.GetDisplayText(), PCG::Private::GetTypeNameText<FString>(), PCG::Private::GetTypeNameText(Accessor->GetUnderlyingType())), InContext);
+			return Result;
+		}
+	}
+	else
+	{
+		Result.Emplace(InSettings->Grammar);
+	}
+
+	for (auto& [Grammar, TokenizeGrammar] : Result)
+	{
+		TokenizeGrammar = GetTokenizeGrammar(InContext, Grammar, InModulesInfo, OutMinSize);
+	}
+
+	return Result;
+}
+
+bool FPCGSlicingBaseElement::MatchAndSetAttributes(const TArray<FPCGTaggedData>& InputData, TArray<FPCGTaggedData>& OutputData, const UPCGParamData* InModuleInfoParamData, const UPCGSlicingBaseSettings* InSettings) const
+{
+	check(InModuleInfoParamData && InModuleInfoParamData->Metadata);
+
+	const UPCGMetadata* InputMetadata = InModuleInfoParamData->Metadata;
+
+	// We prepare everything to process all the output data afterwrads.
+	// Build the Symbol -> EntryKey mapping
+	// Since we don't know if it is a Name or a String, we need to get both.
+	TMap<FName, PCGMetadataEntryKey> SymbolToEntryKeyMapping;
+	if (const FPCGMetadataAttribute<FName>* InSymbolNameAttribute = InputMetadata->GetConstTypedAttribute<FName>(PCGSlicingBaseConstants::SymbolAttributeName))
+	{
+		for (PCGMetadataEntryKey EntryKey = InputMetadata->GetItemKeyCountForParent(); EntryKey < InputMetadata->GetItemCountForChild(); ++EntryKey)
+		{
+			const FName Symbol = InSymbolNameAttribute->GetValueFromItemKey(EntryKey);
+			if (!SymbolToEntryKeyMapping.Contains(Symbol))
+			{
+				SymbolToEntryKeyMapping.Emplace(Symbol, EntryKey);
+			}
+		}
+	}
+	else if (const FPCGMetadataAttribute<FString>* InSymbolStrAttribute = InputMetadata->GetConstTypedAttribute<FString>(PCGSlicingBaseConstants::SymbolAttributeName))
+	{
+		for (PCGMetadataEntryKey EntryKey = InputMetadata->GetItemKeyCountForParent(); EntryKey < InputMetadata->GetItemCountForChild(); ++EntryKey)
+		{
+			const FName Symbol(InSymbolStrAttribute->GetValueFromItemKey(EntryKey));
+			if (!SymbolToEntryKeyMapping.Contains(Symbol))
+			{
+				SymbolToEntryKeyMapping.Emplace(Symbol, EntryKey);
+			}
+		}
+	}
+	else
+	{
+		return false;
+	}
+
+	// Also gather all the attribute in the input metadata
+	TArray<FName> AttributeNames;
+	TArray<EPCGMetadataTypes> AttributeTypes;
+	InModuleInfoParamData->Metadata->GetAttributes(AttributeNames, AttributeTypes);
+	check(AttributeNames.Num() == AttributeTypes.Num());
+
+	for (FPCGTaggedData& TaggedData : OutputData)
+	{
+		// Be careful with the const cast, only allow it if the data is not present in the input too (forwarded)
+		if (InputData.ContainsByPredicate([TaggedData](const FPCGTaggedData& InData) { return InData.Data == TaggedData.Data; }))
+		{
+			continue;
+		}
+
+		UPCGData* OutData = const_cast<UPCGData*>(TaggedData.Data.Get());
+		check(OutData && OutData->MutableMetadata());
+
+		UPCGMetadata* OutputMetadata = OutData->MutableMetadata();
+
+		// Look for the Symbol Attribute in the output metadata to query its value.
+		const FPCGMetadataAttribute<FName>* OutSymbolAttribute = OutputMetadata->GetConstTypedAttribute<FName>(InSettings->SymbolAttributeName);
+		if (!ensure(OutSymbolAttribute))
+		{
+			continue;
+		}
+
+		if (UPCGPointData* OutPointData = Cast<UPCGPointData>(OutData))
+		{
+			TArray<const FPCGMetadataAttributeBase*> InAttributes;
+			TArray<FPCGMetadataAttributeBase*> OutAttributes;
+
+			InAttributes.Reserve(AttributeNames.Num());
+			OutAttributes.Reserve(AttributeNames.Num());
+
+			// Copy all the attributes in the output metadata
+			for (const FName AttributeName : AttributeNames)
+			{
+				// Skip the symbol, it already exists, but perhaps with a different attribute name.
+				if (AttributeName == PCGSlicingBaseConstants::SymbolAttributeName)
+				{
+					continue;
+				}
+
+				// Skip any attribute that already exists in the output data
+				if (OutputMetadata->HasAttribute(AttributeName))
+				{
+					continue;
+				}
+
+				InAttributes.Add(InputMetadata->GetConstAttribute(AttributeName));
+				OutAttributes.Add(OutputMetadata->CopyAttribute(InAttributes.Last(), AttributeName, /*bKeepParent=*/false, /*bCopyEntries=*/false, /*bCopyValues=*/true));
+			}
+
+			for (const FPCGPoint& OutPoint : OutPointData->GetPoints())
+			{
+				const FName Symbol = OutSymbolAttribute->GetValueFromItemKey(OutPoint.MetadataEntry);
+				const PCGMetadataEntryKey InputEntryKey = SymbolToEntryKeyMapping[Symbol];
+
+				if (InputEntryKey == PCGInvalidEntryKey)
+				{
+					continue;
+				}
+
+				for (int32 i = 0; i < InAttributes.Num(); ++i)
+				{
+					PCGMetadataValueKey InputValueKey = InAttributes[i]->GetValueKey(InputEntryKey);
+					OutAttributes[i]->SetValueFromValueKey(OutPoint.MetadataEntry, InputValueKey);
+				}
+			}
+		}
+		else
+		{
+			// If we have a spatial data that is not a point data, we can only operate on Default values.
+			// So add all the attributes, and set the default value to the same value for the chosen entry key.
+			const FName Symbol = OutSymbolAttribute->GetValue(PCGDefaultValueKey);
+			const PCGMetadataEntryKey InputEntryKey = SymbolToEntryKeyMapping[Symbol];
+
+			auto AddAttribute = [InputMetadata, OutputMetadata, InputEntryKey]<typename T>(T Dummy, const FName AttributeName)
+			{
+				const FPCGMetadataAttribute<T>* InAttribute = InputMetadata->GetConstTypedAttribute<T>(AttributeName);
+				check(InAttribute);
+
+				// Overwrite any existing attribute
+				if (OutputMetadata->HasAttribute(AttributeName))
+				{
+					OutputMetadata->DeleteAttribute(AttributeName);
+				}
+
+				OutputMetadata->CreateAttribute<T>(AttributeName, InAttribute->GetValueFromItemKey(InputEntryKey), InAttribute->AllowsInterpolation(), /*bOverrideParent=*/true);
+			};
+
+			for (int32 i = 0; i < AttributeNames.Num(); ++i)
+			{
+				// Skip the symbol attribute, as it is already in the OutData (with perhaps a different name)
+				if (AttributeNames[i] == PCGSlicingBaseConstants::SymbolAttributeName)
+				{
+					continue;
+				}
+
+				// Skip any attribute that already exists in the output data
+				if (OutputMetadata->HasAttribute(AttributeNames[i]))
+				{
+					continue;
+				}
+
+				PCGMetadataAttribute::CallbackWithRightType(static_cast<int16>(AttributeTypes[i]), AddAttribute, AttributeNames[i]);
+			}
+		}
+	}
+
+	return true;
+}
+
+#undef LOCTEXT_NAMESPACE
\ No newline at end of file
diff --git a/Engine/Plugins/PCG/Source/PCG/Private/Elements/Grammar/PCGSplineToSegment.cpp b/Engine/Plugins/PCG/Source/PCG/Private/Elements/Grammar/PCGSplineToSegment.cpp
new file mode 100644
--- /dev/null
+++ b/Engine/Plugins/PCG/Source/PCG/Private/Elements/Grammar/PCGSplineToSegment.cpp
@@ -0,0 +1,256 @@
+// Copyright Epic Games, Inc. All Rights Reserved.
+
+#include "Elements/Grammar/PCGSplineToSegment.h"
+
+#include "PCGContext.h"
+#include "Data/PCGPointData.h"
+#include "Data/PCGSplineData.h"
+#include "Helpers/PCGHelpers.h"
+
+#define LOCTEXT_NAMESPACE "PCGSegmentSplineElement"
+
+namespace PCGSplineToSegment
+{
+	static const FName PreviousTangentAttributeName = TEXT("PreviousTangent");
+	static const FName NextTangentAttributeName = TEXT("NextTangent");
+	static const FName SegmentIndexAttributeName = TEXT("SegmentIndex");
+	static const FName PreviousIndexAttributeName = TEXT("SegmentPreviousIndex");
+	static const FName NextIndexAttributeName = TEXT("SegmentNextIndex");
+	static const FName PreviousAngleAttributeName = TEXT("PreviousAngle");
+	static const FName NextAngleAttributeName = TEXT("NextAngle");
+	static const FName ClockwiseAttributeName = TEXT("Clockwise");
+
+	template <typename T>
+	FPCGMetadataAttribute<T>* CreateAttributeWithLog(const FName AttributeName, const bool bAllowInterpolation, UPCGMetadata* Metadata, FPCGContext* Context, T DefaultValue = PCG::Private::MetadataTraits<T>::ZeroValue())
+	{
+		FPCGMetadataAttribute<T>* Attribute = Metadata->FindOrCreateAttribute<T>(AttributeName, DefaultValue, bAllowInterpolation, false);
+		if (!Attribute)
+		{
+			PCGLog::LogErrorOnGraph(FText::Format(LOCTEXT("FailCreateAttribute", "Failed to create attribute {0}."), FText::FromName(AttributeName)), Context);
+		}
+
+		return Attribute;
+	}
+}
+
+#if WITH_EDITOR
+FName UPCGSplineToSegmentSettings::GetDefaultNodeName() const
+{
+	return FName(TEXT("SplineToSegment"));
+}
+
+FText UPCGSplineToSegmentSettings::GetDefaultNodeTitle() const
+{
+	return LOCTEXT("NodeTitle", "Spline to Segment");
+}
+
+FText UPCGSplineToSegmentSettings::GetNodeTooltipText() const
+{
+	return LOCTEXT("NodeTooltip", "Take a spline as input and create a point data, with each point being a segment defined by 2 connected control points.\n"
+		"The point position will be in the middle of those 2 control points, and the extents of the point will be half the difference between those 2 control points.");
+}
+#endif // WITH_EDITOR
+
+FPCGElementPtr UPCGSplineToSegmentSettings::CreateElement() const
+{
+	return MakeShared<FPCGSplineToSegmentElement>();
+}
+
+TArray<FPCGPinProperties> UPCGSplineToSegmentSettings::InputPinProperties() const
+{
+	TArray<FPCGPinProperties> Result;
+	FPCGPinProperties& InputPin = Result.Emplace_GetRef(PCGPinConstants::DefaultInputLabel, EPCGDataType::Spline);
+	InputPin.SetRequiredPin();
+
+	return Result;
+}
+
+TArray<FPCGPinProperties> UPCGSplineToSegmentSettings::OutputPinProperties() const
+{ 
+	return Super::DefaultPointOutputPinProperties();
+}
+
+bool FPCGSplineToSegmentElement::ExecuteInternal(FPCGContext* InContext) const
+{
+	TRACE_CPUPROFILER_EVENT_SCOPE(FPCGSegmentSplineElement::Execute);
+
+	check(InContext);
+
+	const UPCGSplineToSegmentSettings* Settings = InContext->GetInputSettings<UPCGSplineToSegmentSettings>();
+	check(Settings);
+
+	const TArray<FPCGTaggedData> Inputs = InContext->InputData.GetInputsByPin(PCGPinConstants::DefaultInputLabel);
+	TArray<FPCGTaggedData>& Outputs = InContext->OutputData.TaggedData;
+
+	for (const FPCGTaggedData& Input : Inputs)
+	{
+		const UPCGSplineData* InputSplineData = Cast<const UPCGSplineData>(Input.Data);
+
+		if (!InputSplineData)
+		{
+			continue;
+		}
+
+		const bool bIsClosed = InputSplineData->SplineStruct.IsClosedLoop();
+
+		UPCGPointData* OutputPointData = FPCGContext::NewObject_AnyThread<UPCGPointData>(InContext);
+		OutputPointData->InitializeFromData(InputSplineData);
+
+		const TArray<FInterpCurvePointVector>& ControlPoints = InputSplineData->SplineStruct.SplineCurves.Position.Points;
+		const int32 ControlPointsNum = ControlPoints.Num();
+		if (ControlPointsNum < 2)
+		{
+			// If we have a malformed spline, just make an empty point data
+			Outputs.Emplace_GetRef(Input).Data = OutputPointData;
+			continue;
+		}
+
+		FPCGMetadataAttribute<FVector>* PreviousTangentAttribute = nullptr;
+		FPCGMetadataAttribute<FVector>* NextTangentAttribute = nullptr;
+		FPCGMetadataAttribute<int32>* SegmentIndexAttribute = nullptr;
+		FPCGMetadataAttribute<int32>* PreviousIndexAttribute = nullptr;
+		FPCGMetadataAttribute<int32>* NextIndexAttribute = nullptr;
+		FPCGMetadataAttribute<double>* PreviousAngleAttribute = nullptr;
+		FPCGMetadataAttribute<double>* NextAngleAttribute = nullptr;
+
+		bool bHasMetadata = false;
+
+		if (Settings->bExtractTangents)
+		{
+			PreviousTangentAttribute = PCGSplineToSegment::CreateAttributeWithLog<FVector>(PCGSplineToSegment::PreviousTangentAttributeName, false, OutputPointData->Metadata, InContext);
+			NextTangentAttribute = PCGSplineToSegment::CreateAttributeWithLog<FVector>(PCGSplineToSegment::NextTangentAttributeName, false, OutputPointData->Metadata, InContext);
+
+			if (!PreviousTangentAttribute || !NextTangentAttribute)
+			{
+				continue;
+			}
+
+			bHasMetadata = true;
+		}
+
+		if (Settings->bExtractAngles)
+		{
+			PreviousAngleAttribute = PCGSplineToSegment::CreateAttributeWithLog<double>(PCGSplineToSegment::PreviousAngleAttributeName, true, OutputPointData->Metadata, InContext);
+			NextAngleAttribute = PCGSplineToSegment::CreateAttributeWithLog<double>(PCGSplineToSegment::NextAngleAttributeName, true, OutputPointData->Metadata, InContext);
+
+			if (!PreviousAngleAttribute || !NextAngleAttribute)
+			{
+				continue;
+			}
+
+			bHasMetadata = true;
+		}
+
+		if (Settings->bExtractConnectivityInfo)
+		{
+			PreviousIndexAttribute = PCGSplineToSegment::CreateAttributeWithLog<int32>(PCGSplineToSegment::PreviousIndexAttributeName, false, OutputPointData->Metadata, InContext);
+			NextIndexAttribute = PCGSplineToSegment::CreateAttributeWithLog<int32>(PCGSplineToSegment::NextIndexAttributeName, false, OutputPointData->Metadata, InContext);
+			SegmentIndexAttribute = PCGSplineToSegment::CreateAttributeWithLog<int32>(PCGSplineToSegment::SegmentIndexAttributeName, false, OutputPointData->Metadata, InContext);
+
+			if (!PreviousIndexAttribute || !NextIndexAttribute || !SegmentIndexAttribute)
+			{
+				continue;
+			}
+
+			bHasMetadata = true;
+		}
+
+		TArray<FPCGPoint>& OutputPoints = OutputPointData->GetMutablePoints();
+		OutputPoints.Reserve(bIsClosed ? ControlPointsNum : ControlPointsNum - 1);
+
+		TArray<FBox> Segments;
+		Segments.Reserve(ControlPointsNum);
+		double CumulativeAngle = 0.0;
+		double FirstAngle = 0.0;
+
+		for (int32 Index = 0; Index < ControlPointsNum; ++Index)
+		{
+			const int32 NextIndex = (Index + 1) % ControlPointsNum;
+			const FVector ControlPoint = InputSplineData->SplineStruct.GetLocationAtSplineInputKey(Index, ESplineCoordinateSpace::World);
+			const FVector NextControlPoint = InputSplineData->SplineStruct.GetLocationAtSplineInputKey(NextIndex, ESplineCoordinateSpace::World);
+			const FVector ControlPointZUp = InputSplineData->SplineStruct.GetQuaternionAtSplineInputKey(Index, ESplineCoordinateSpace::World).GetAxisZ();
+
+			FVector ArriveTangent, LeaveTangent;
+
+			InputSplineData->GetTangentsAtSegmentStart(Index, ArriveTangent, LeaveTangent);
+			const double CrossProduct = (ArriveTangent.X * LeaveTangent.Y - ArriveTangent.Y * LeaveTangent.X);
+			const double AbsAngle = FMath::Acos(ArriveTangent.CosineAngle2D(LeaveTangent));
+			const double Angle = FMath::RadiansToDegrees(FMath::Sign(CrossProduct) * AbsAngle);
+
+			FPCGPoint& OutputPoint = OutputPoints.Emplace_GetRef();
+			if (bHasMetadata)
+			{
+				OutputPointData->Metadata->InitializeOnSet(OutputPoint.MetadataEntry);
+			}
+
+			// Points will be placed at the middle of the segment with bounds min and max matching control points on X (Y and Z will be set to -1 and 1 for min and max respectively)
+			// We'll also reset the point scale
+			const FVector Extents = (NextControlPoint - ControlPoint) * 0.5;
+			const double ExtentsLength = Extents.Length();
+			OutputPoint.Transform.SetLocation(ControlPoint + Extents);
+			OutputPoint.Transform.SetRotation(FRotationMatrix::MakeFromXZ(Extents, ControlPointZUp).ToQuat());
+			OutputPoint.BoundsMin = FVector(-ExtentsLength, 0.0, 0.0);
+			OutputPoint.BoundsMax = FVector(ExtentsLength, 1.0, 1.0);
+
+			OutputPoint.Seed = PCGHelpers::ComputeSeedFromPosition(OutputPoint.Transform.GetLocation());
+
+			auto SetMetadata = [&](const int32 Index, const FVector& Tangent, const double Angle)
+			{
+				const int32 PreviousIndex = (Index + ControlPointsNum - 1) % ControlPointsNum;
+				const PCGMetadataEntryKey PreviousPointKey = OutputPoints[PreviousIndex].MetadataEntry;
+				const PCGMetadataEntryKey CurrentPointKey = OutputPoints[Index].MetadataEntry;
+
+				if (PreviousTangentAttribute && NextTangentAttribute)
+				{
+					NextTangentAttribute->SetValue(PreviousPointKey, Tangent);
+					PreviousTangentAttribute->SetValue(CurrentPointKey, Tangent);
+				}
+
+				if (PreviousAngleAttribute && NextAngleAttribute)
+				{
+					NextAngleAttribute->SetValue(PreviousPointKey, Angle);
+					PreviousAngleAttribute->SetValue(CurrentPointKey, Angle);
+				}
+
+				if (PreviousIndexAttribute && NextIndexAttribute)
+				{
+					NextIndexAttribute->SetValue(PreviousPointKey, Index);
+					PreviousIndexAttribute->SetValue(CurrentPointKey, PreviousIndex);
+					SegmentIndexAttribute->SetValue(CurrentPointKey, Index);
+				}
+			};
+
+			if (Index > 0)
+			{
+				CumulativeAngle += Angle;
+
+				if (bHasMetadata)
+				{
+					SetMetadata(Index, ArriveTangent, Angle);
+				}
+			}
+			else
+			{
+				// Need to save the first angle to set it later.
+				FirstAngle = Angle;
+			}
+
+			if (NextIndex == 0 && bHasMetadata)
+			{
+				SetMetadata(NextIndex, LeaveTangent, FirstAngle);
+			}
+		}
+
+		if (Settings->bExtractClockwiseInfo && bIsClosed)
+		{
+			PCGSplineToSegment::CreateAttributeWithLog<bool>(PCGSplineToSegment::ClockwiseAttributeName, false, OutputPointData->Metadata, InContext, CumulativeAngle <= 0.0);
+		}
+
+		FPCGTaggedData& Output = Outputs.Emplace_GetRef(Input);
+		Output.Data = OutputPointData;
+	}
+
+	return true;
+}
+
+#undef LOCTEXT_NAMESPACE
diff --git a/Engine/Plugins/PCG/Source/PCG/Private/Elements/Grammar/PCGVolumeSlicer.cpp b/Engine/Plugins/PCG/Source/PCG/Private/Elements/Grammar/PCGVolumeSlicer.cpp
new file mode 100644
--- /dev/null
+++ b/Engine/Plugins/PCG/Source/PCG/Private/Elements/Grammar/PCGVolumeSlicer.cpp
@@ -0,0 +1,199 @@
+// Copyright Epic Games, Inc. All Rights Reserved.
+
+#include "Elements/Grammar/PCGVolumeSlicer.h"
+
+#include "PCGContext.h"
+#include "PCGParamData.h"
+#include "Data/PCGPointData.h"
+#include "Data/PCGSplineData.h"
+#include "Grammar/PCGGrammarParser.h"
+#include "Metadata/Accessors/IPCGAttributeAccessor.h"
+#include "Metadata/Accessors/PCGAttributeAccessorHelpers.h"
+#include "Metadata/Accessors/PCGAttributeAccessorKeys.h"
+
+#define LOCTEXT_NAMESPACE "PCGVolumeSlicerElement"
+
+#if WITH_EDITOR
+FName UPCGVolumeSlicerSettings::GetDefaultNodeName() const
+{
+	return FName(TEXT("VolumeSlicer"));
+}
+
+FText UPCGVolumeSlicerSettings::GetDefaultNodeTitle() const
+{
+	return LOCTEXT("NodeTitle", "Volume Slicer");
+}
+#endif // WITH_EDITOR
+
+FPCGElementPtr UPCGVolumeSlicerSettings::CreateElement() const
+{
+	return MakeShared<FPCGVolumeSlicerElement>();
+}
+
+TArray<FPCGPinProperties> UPCGVolumeSlicerSettings::InputPinProperties() const
+{
+	TArray<FPCGPinProperties> Result;
+	FPCGPinProperties& InputPin = Result.Emplace_GetRef(PCGPinConstants::DefaultInputLabel, EPCGDataType::Spline);
+	InputPin.SetRequiredPin();
+
+	if (bModuleInfoAsInput)
+	{
+		FPCGPinProperties& ModuleInfoPin = Result.Emplace_GetRef(PCGSlicingBaseConstants::ModulesInfoPinLabel, EPCGDataType::Param);
+		ModuleInfoPin.SetRequiredPin();
+	}
+
+	return Result;
+}
+
+TArray<FPCGPinProperties> UPCGVolumeSlicerSettings::OutputPinProperties() const
+{
+	TArray<FPCGPinProperties> Result;
+	Result.Emplace_GetRef(PCGPinConstants::DefaultOutputLabel, EPCGDataType::Spline);
+
+	return Result;
+}
+
+bool FPCGVolumeSlicerElement::ExecuteInternal(FPCGContext* InContext) const
+{
+	TRACE_CPUPROFILER_EVENT_SCOPE(FPCGVolumeSlicerElement::Execute);
+
+	check(InContext);
+
+	const UPCGVolumeSlicerSettings* Settings = InContext->GetInputSettings<UPCGVolumeSlicerSettings>();
+	check(Settings);
+
+	const TArray<FPCGTaggedData> Inputs = InContext->InputData.GetInputsByPin(PCGPinConstants::DefaultInputLabel);
+	TArray<FPCGTaggedData>& Outputs = InContext->OutputData.TaggedData;
+
+	const UPCGParamData* ModuleInfoParamData = nullptr;
+	FPCGModulesInfoMap ModulesInfo = GetModulesInfoMap(InContext, Settings, ModuleInfoParamData);
+
+	const bool bMatchAndSetAttributes = Settings->bForwardAttributesFromModulesInfo && ModuleInfoParamData;
+
+	for (const FPCGTaggedData& Input : Inputs)
+	{
+		const UPCGSplineData* InputSplineData = Cast<const UPCGSplineData>(Input.Data);
+
+		if (!InputSplineData)
+		{
+			continue;
+		}
+
+		const TArray<FInterpCurvePointVector>& ControlPoints = InputSplineData->SplineStruct.SplineCurves.Position.Points;
+		if (ControlPoints.Num() < 2)
+		{
+			continue;
+		}
+
+		FVector ExtrudeVector = Settings->ExtrudeVector;
+
+		auto GetValueFromAttribute = [InputSplineData, InContext]<typename T>(const FPCGAttributePropertyInputSelector& InSelector, T& OutValue) -> bool
+		{
+			const FPCGAttributePropertyInputSelector Selector = InSelector.CopyAndFixLast(InputSplineData);
+			const TUniquePtr<const IPCGAttributeAccessor> Accessor = PCGAttributeAccessorHelpers::CreateConstAccessor(InputSplineData, Selector);
+			if (!Accessor)
+			{
+				PCGLog::LogErrorOnGraph(FText::Format(LOCTEXT("AccessorFailCreate", "Attribute {0} was not found in the input data."), Selector.GetDisplayText()), InContext);
+				return false;
+			}
+
+			if (!Accessor->Get(OutValue, FPCGAttributeAccessorKeysEntries(PCGInvalidEntryKey), EPCGAttributeAccessorFlags::AllowBroadcastAndConstructible))
+			{
+				PCGLog::LogErrorOnGraph(FText::Format(LOCTEXT("AccessorFailGet", "Couldn't retrieve attribute {0} value, could it be the wrong type? Expected type: {1}, Attribute Type: {2}."), Selector.GetDisplayText(), PCG::Private::GetTypeNameText<T>(), PCG::Private::GetTypeNameText(Accessor->GetUnderlyingType())), InContext);
+				return false;
+			}
+
+			return true;
+		};
+
+		if (Settings->bExtrudeVectorAsAttribute && !GetValueFromAttribute(Settings->ExtrudeVectorAttribute, ExtrudeVector))
+		{
+			continue;
+		}
+
+		const double ExtrudeLength = ExtrudeVector.Length();
+		if (FMath::IsNearlyZero(ExtrudeLength))
+		{
+			continue;
+		}
+
+		const FVector ExtrudeDirection = ExtrudeVector / ExtrudeLength;
+
+		double MinSize = 0.0;
+		TArray<FPCGTokenizedGrammar> TokenizeGrammar = GetTokenizeGrammar(InContext, InputSplineData, Settings, ModulesInfo, MinSize);
+
+		if (TokenizeGrammar.IsEmpty())
+		{
+			continue;
+		}
+
+		TArray<PCGSlicingBase::PCGSubDivModuleInstance<FPCGTokenizedGrammar>> Instances;
+		double RemainingLength = 0.0;
+		const bool bHeightSubdivideSuccess = PCGSlicingBase::Subdivide(TokenizeGrammar, ExtrudeLength, Instances, RemainingLength);
+
+		if (!bHeightSubdivideSuccess)
+		{
+			PCGLog::LogErrorOnGraph(LOCTEXT("SubdivideFailed", "Grammar doesn't fit."), InContext);
+			continue;
+		}
+
+		FVector CurrentDisplacement = FVector::ZeroVector;
+		int32 SplineIndex = 0;
+
+		for (const PCGSlicingBase::PCGSubDivModuleInstance<FPCGTokenizedGrammar>& Instance : Instances)
+		{
+			for (int i = 0; i < Instance.NumRepeat; ++i)
+			{
+				for (int32 SymbolIndex = 0; SymbolIndex < Instance.Module->Symbols.Num(); ++SymbolIndex)
+				{
+					const FName Symbol = Instance.Module->Symbols[SymbolIndex];
+					const FVector Size = ExtrudeDirection * Instance.Module->SymbolSizes[SymbolIndex] * (FVector::OneVector + Instance.ExtraScales[SymbolIndex]);
+					const FPCGSlicingModule& CurrentBlock = ModulesInfo[Symbol];
+					FPCGSplineStruct NewSpline = InputSplineData->SplineStruct;
+					for (FInterpCurvePointVector& ControlPoint : NewSpline.SplineCurves.Position.Points)
+					{
+						ControlPoint.OutVal += CurrentDisplacement;
+					}
+
+					UPCGSplineData* NewSplineData = FPCGContext::NewObject_AnyThread<UPCGSplineData>(InContext);
+					NewSplineData->Initialize(NewSpline);
+					NewSplineData->InitializeFromData(InputSplineData);
+
+					NewSplineData->Metadata->FindOrCreateAttribute<FName>(Settings->SymbolAttributeName, Symbol, false, false);
+
+					if (!bMatchAndSetAttributes && Settings->bOutputDebugColorAttribute)
+					{
+						NewSplineData->Metadata->FindOrCreateAttribute<FVector4>(Settings->DebugColorAttributeName, FVector4(CurrentBlock.DebugColor, 1.0), false, false);
+					}
+
+					if (Settings->bOutputSizeAttribute)
+					{
+						NewSplineData->Metadata->FindOrCreateAttribute<FVector>(Settings->SizeAttributeName, Size, false, false);
+					}
+
+					if (!bMatchAndSetAttributes && Settings->bOutputScalableAttribute)
+					{
+						NewSplineData->Metadata->FindOrCreateAttribute<bool>(Settings->ScalableAttributeName, CurrentBlock.bScalable, false, false);
+					}
+
+					if (Settings->bOutputSplineIndexAttribute)
+					{
+						NewSplineData->Metadata->FindOrCreateAttribute<int32>(Settings->SplineIndexAttributeName, SplineIndex++, false, false);
+					}
+
+					Outputs.Add_GetRef(Input).Data = NewSplineData;
+					CurrentDisplacement += Size;
+				}
+			}
+		}
+	}
+
+	if (bMatchAndSetAttributes)
+	{
+		MatchAndSetAttributes(Inputs, Outputs, ModuleInfoParamData, Settings);
+	}
+
+	return true;
+}
+
+#undef LOCTEXT_NAMESPACE
diff --git a/Engine/Plugins/PCG/Source/PCG/Private/Grammar/PCGGrammarParser.cpp b/Engine/Plugins/PCG/Source/PCG/Private/Grammar/PCGGrammarParser.cpp
new file mode 100644
--- /dev/null
+++ b/Engine/Plugins/PCG/Source/PCG/Private/Grammar/PCGGrammarParser.cpp
@@ -0,0 +1,238 @@
+// Copyright Epic Games, Inc. All Rights Reserved.
+
+#include "Grammar/PCGGrammarParser.h"
+
+#include "Algo/Transform.h"
+#include "Internationalization/Internationalization.h"
+#include "Internationalization/Regex.h"
+
+#define LOCTEXT_NAMESPACE "PCGGrammar"
+
+void FPCGGrammarResult::AddLog(FText Message, ELogType Verbosity)
+{
+	Logs.AddUnique({ std::move(Message), Verbosity });
+	if (Verbosity == ELogType::Error)
+	{
+		bSuccess = false;
+	}
+}
+
+namespace PCGGrammar
+{
+	namespace Patterns
+	{
+		static constexpr const TCHAR* Base = TEXT("(\\[([^\\]]+)\\](?:\\s*(\\d+|\\*))?)");
+		static constexpr const TCHAR* Stochastic = TEXT("(\\{([^\\}]+)\\}(?:\\s*(\\d++|\\*))?)");
+		static constexpr const TCHAR* Priority = TEXT("(<([^>]+)>(?:\\s*(\\d++|\\*))?)");
+	}
+
+	namespace Tokens
+	{
+		static constexpr const TCHAR* SubmoduleDelimiter = TEXT(",");
+		static constexpr const TCHAR* WeightDelimiter = TEXT(":");
+		static constexpr const TCHAR* InfiniteRepetition = TEXT("*");
+	}
+
+	enum ECaptureGroups : int
+	{
+		// Capture Group 0 is the matched expression for this iteration
+		MatchedExpression = 0,
+
+		// Capture group 1 is the whole group of submodules, including [ ]
+		WholeGroupOfSubmodules = 1,
+
+		// Capture group 2 is the submodule expression. Ie. [A,B], {A,B}, or <A,B> -> A,B
+		SubmoduleExpression = 2,
+
+		// Capture group 3 is the number of repetitions - [A,B]3 -> 3
+		NumberOfRepetitionsOrAsterisk = 3
+	};
+
+	bool FindModules(
+		const FString& Grammar,
+		const EModuleType ModuleType,
+		TArray<TPair<int32, int32>, TInlineAllocator<16>>& OutMatchedIndices,
+		FPCGGrammarResult& OutGrammarResult)
+	{
+		const TCHAR* Pattern = nullptr;
+		switch (ModuleType)
+		{
+			case EModuleType::Base:
+				Pattern = Patterns::Base;
+				break;
+			case EModuleType::Stochastic:
+				Pattern = Patterns::Stochastic;
+				break;
+			case EModuleType::Priority:
+				Pattern = Patterns::Priority;
+				break;
+			default:
+				checkNoEntry();
+				return false;
+		}
+
+		check(Pattern);
+
+		bool bFoundMatch = false;
+
+		TArray<FName> SubmoduleNames;
+		TArray<int32> SubmoduleWeights;
+
+		FRegexMatcher ModuleMatcher(FRegexPattern(Pattern), Grammar);
+		while (ModuleMatcher.FindNext())
+		{
+			bFoundMatch = true;
+
+			const FString SubmoduleCapture = ModuleMatcher.GetCaptureGroup(ECaptureGroups::SubmoduleExpression);
+
+			if (SubmoduleCapture.IsEmpty())
+			{
+				FText Message = FText::Format(LOCTEXT("NoModuleNameFound", "Unable to find module name match within the grammar declaration: {0}"), FText::FromString(ModuleMatcher.GetCaptureGroup(ECaptureGroups::MatchedExpression)));
+				OutGrammarResult.AddLog(std::move(Message), FPCGGrammarResult::ELogType::Warning);
+				continue;
+			}
+
+			TArray<FString> SubmoduleStrings;
+			SubmoduleCapture.ParseIntoArray(SubmoduleStrings, Tokens::SubmoduleDelimiter);
+
+			SubmoduleNames.Reset(SubmoduleStrings.Num());
+			SubmoduleWeights.Reset(SubmoduleStrings.Num());
+
+			for (const FString& SubmoduleString : SubmoduleStrings)
+			{
+				TArray<FString> SubmoduleIDString;
+				SubmoduleString.ParseIntoArray(SubmoduleIDString, Tokens::WeightDelimiter);
+
+				if (SubmoduleIDString.IsEmpty())
+				{
+					FText Message = FText::Format(LOCTEXT("InvalidSubmoduleString", "Parsed submodule string is invalid."), FText::FromString(ModuleMatcher.GetCaptureGroup(ECaptureGroups::MatchedExpression)));
+					OutGrammarResult.AddLog(std::move(Message), FPCGGrammarResult::ELogType::Warning);
+					continue;
+				}
+
+				SubmoduleIDString[0].TrimStartAndEndInline();
+				if (SubmoduleIDString[0].IsEmpty())
+				{
+					FText Message = FText::Format(LOCTEXT("EmptyModuleID", "Module ID must not be empty."), FText::FromString(ModuleMatcher.GetCaptureGroup(ECaptureGroups::MatchedExpression)));
+					OutGrammarResult.AddLog(std::move(Message), FPCGGrammarResult::ELogType::Warning);
+					continue;
+				}
+
+				SubmoduleNames.Emplace(SubmoduleIDString[0]);
+				SubmoduleWeights.Emplace(1);
+
+				if (SubmoduleIDString.Num() > 1)
+				{
+					if (ModuleType != EModuleType::Stochastic)
+					{
+						FText Message = FText::Format(LOCTEXT("WeightOnNonStochasticType", "Weight added to non-stochastic module type."), FText::FromString(ModuleMatcher.GetCaptureGroup(ECaptureGroups::MatchedExpression)));
+						OutGrammarResult.AddLog(std::move(Message), FPCGGrammarResult::ELogType::Warning);
+						continue;
+					}
+
+					if (SubmoduleIDString.Num() > 2)
+					{
+						FText Message = LOCTEXT("MultiCharacterWeightDelimiter", "Multi-character delimiter for weight. Weight ignored.");
+						OutGrammarResult.AddLog(std::move(Message), FPCGGrammarResult::ELogType::Warning);
+						continue;
+					}
+
+					SubmoduleIDString[1].TrimStartAndEndInline();
+					if (!SubmoduleIDString[1].IsNumeric())
+					{
+						FText Message = FText::Format(LOCTEXT("InvalidWeightCharacter", "Invalid weight character '{0}' Weight ignored."), FText::FromString(SubmoduleIDString[1]));
+						OutGrammarResult.AddLog(std::move(Message), FPCGGrammarResult::ELogType::Warning);
+						continue;
+					}
+
+					SubmoduleWeights.Last() = FCString::Atoi(*SubmoduleIDString[1]);
+				}
+			}
+
+			// Capture repetition count
+			int32 Repetitions = 1;
+			const FString RepetitionString = ModuleMatcher.GetCaptureGroup(ECaptureGroups::NumberOfRepetitionsOrAsterisk);
+			if (RepetitionString.IsNumeric())
+			{
+				Repetitions = FCString::Atoi(*RepetitionString);
+			}
+			else if (!RepetitionString.IsEmpty())
+			{
+				// It should only be the infinite repetition character, to match the
+				ensure(RepetitionString == Tokens::InfiniteRepetition);
+				Repetitions = -1;
+			}
+
+			FModuleDescriptor& Module = OutGrammarResult.Modules.Emplace_GetRef();
+			Module.GrammarStartEndIndices = {ModuleMatcher.GetMatchBeginning(), ModuleMatcher.GetMatchEnding()};
+			Module.Type = ModuleType;
+			Module.Repetitions = Repetitions;
+
+			for (int Index = 0; Index < SubmoduleNames.Num(); ++Index)
+			{
+				Module.Submodules.Emplace(SubmoduleNames[Index], SubmoduleWeights[Index]);
+			}
+
+			OutMatchedIndices.Emplace(ModuleMatcher.GetMatchBeginning(), ModuleMatcher.GetMatchEnding());
+		}
+
+		return bFoundMatch;
+	}
+
+	PCG_API FPCGGrammarResult Parse(const FString& Grammar, bool bValidateGrammar)
+	{
+		FPCGGrammarResult Result;
+
+		if (Grammar.IsEmpty())
+		{
+			Result.AddLog(LOCTEXT("EmptyGrammar", "Grammar is empty."));
+			return Result;
+		}
+
+		bool bFoundMatch = true;
+		TArray<TPair<int32, int32>, TInlineAllocator<16>> MatchedIndices;
+
+		bFoundMatch |= FindModules(Grammar, EModuleType::Base, MatchedIndices, Result);
+		bFoundMatch |= FindModules(Grammar, EModuleType::Stochastic, MatchedIndices, Result);
+		bFoundMatch |= FindModules(Grammar, EModuleType::Priority, MatchedIndices, Result);
+
+		if (!bFoundMatch)
+		{
+			Result.AddLog(LOCTEXT("NoModuleMatch", "Unable to find module match in grammar."), FPCGGrammarResult::ELogType::Warning);
+			return Result;
+		}
+
+		// Sort by start index - where they appeared in the grammar.
+		Algo::Sort(Result.Modules, [](const FModuleDescriptor& LHS, const FModuleDescriptor& RHS)
+		{
+			return LHS.GrammarStartEndIndices.Key < RHS.GrammarStartEndIndices.Key;
+		});
+
+		if (bValidateGrammar)
+		{
+			FString RemainingCharacters = Grammar;
+
+			MatchedIndices.Sort([](const TPair<int32, int32>& Pair1, const TPair<int32, int32>& Pair2)
+			{
+				return Pair1.Value < Pair2.Value;
+			});
+
+			// Must remove in inverse order to avoid conflict
+			for (int I = MatchedIndices.Num() - 1; I >= 0; --I)
+			{
+				const TPair<int32, int32>& Pair = MatchedIndices[I];
+				RemainingCharacters.RemoveAt(Pair.Key, Pair.Value - Pair.Key, EAllowShrinking::No);
+			}
+
+			RemainingCharacters.TrimStartAndEndInline();
+			if (!RemainingCharacters.IsEmpty())
+			{
+				Result.AddLog(FText::Format(LOCTEXT("ExtraCharactersInGrammar", "Extraneous characters in grammar {0}"), FText::FromString(RemainingCharacters)), FPCGGrammarResult::ELogType::Warning);
+			}
+		}
+
+		return Result;
+	}
+}
+
+#undef LOCTEXT_NAMESPACE
diff --git a/Engine/Plugins/PCG/Source/PCG/Public/Elements/Grammar/PCGSegmentSlicer.h b/Engine/Plugins/PCG/Source/PCG/Public/Elements/Grammar/PCGSegmentSlicer.h
new file mode 100644
--- /dev/null
+++ b/Engine/Plugins/PCG/Source/PCG/Public/Elements/Grammar/PCGSegmentSlicer.h
@@ -0,0 +1,84 @@
+// Copyright Epic Games, Inc. All Rights Reserved.
+
+#pragma once
+
+#include "PCGSettings.h"
+
+#include "Elements/Grammar/PCGSlicingBase.h"
+
+#include "Elements/PCGSplitPoints.h"
+#include "Metadata/PCGAttributePropertySelector.h"
+
+#include "PCGSegmentSlicer.generated.h"
+
+UCLASS(MinimalAPI, BlueprintType, ClassGroup = (Procedural))
+class UPCGSegmentSlicerSettings : public UPCGSlicingBaseSettings
+{
+	GENERATED_BODY()
+
+public:
+	//~Begin UPCGSettings interface
+#if WITH_EDITOR
+	virtual FName GetDefaultNodeName() const override;
+	virtual FText GetDefaultNodeTitle() const override;
+#endif
+
+protected:
+	virtual TArray<FPCGPinProperties> InputPinProperties() const override;
+	virtual TArray<FPCGPinProperties> OutputPinProperties() const override;
+	virtual FPCGElementPtr CreateElement() const override;
+	//~End UPCGSettings interface
+
+public:
+	/** Slicing direction in point local space. */
+	UPROPERTY(BlueprintReadWrite, EditAnywhere, Category = Settings)
+	EPCGSplitAxis SlicingAxis = EPCGSplitAxis::X;
+
+	/** Use an attribute to determine whether we should flip axis. */
+	UPROPERTY(BlueprintReadWrite, EditAnywhere, Category = Settings)
+	bool bFlipAxisAsAttribute = false;
+
+	/** If we need to flip axis. */
+	UPROPERTY(BlueprintReadWrite, EditAnywhere, Category = Settings, meta = (PCG_Overridable, EditCondition = "!bFlipAxisAsAttribute", EditConditionHides, ShowAfter = "bFlipAxisAsAttribute"))
+	bool bShouldFlipAxis = false;
+
+	/** Name of the attribute to know if we need to flip axis. */
+	UPROPERTY(BlueprintReadWrite, EditAnywhere, Category = Settings, meta = (PCG_Overridable, EditCondition = "bFlipAxisAsAttribute", EditConditionHides, ShowAfter="bFlipAxisAsAttribute"))
+	FPCGAttributePropertyInputSelector FlipAxisAttribute;
+
+	/** If the slicing with a given grammar doesn't fill the entire segment, setting it to true makes it a valid case. */
+	UPROPERTY(BlueprintReadWrite, EditAnywhere, Category = Settings)
+	bool bAcceptIncompleteSlicing = false;
+
+	UPROPERTY(BlueprintReadWrite, EditAnywhere, Category = "Settings|ExtraAttributes", meta = (PCG_Overridable, InlineEditConditionToggle))
+	bool bOutputModuleIndexAttribute = false;
+
+	/** Name of the module index output attribute name. */
+	UPROPERTY(BlueprintReadWrite, EditAnywhere, Category = "Settings|ExtraAttributes", meta = (PCG_Overridable, EditCondition = "bOutputModuleIndexAttribute"))
+	FName ModuleIndexAttributeName = TEXT("ModuleIndex");
+
+	UPROPERTY(BlueprintReadWrite, EditAnywhere, Category = "Settings|ExtraAttributes", meta = (PCG_Overridable, InlineEditConditionToggle))
+	bool bOutputExtremityAttribute = false;
+
+	/** Name of the Extremity output attribute name. */
+	UPROPERTY(BlueprintReadWrite, EditAnywhere, Category = "Settings|ExtraAttributes", meta = (PCG_Overridable, EditCondition = "bOutputExtremityAttribute"))
+	FName ExtremityAttributeName = TEXT("Extremity");
+
+	UPROPERTY(BlueprintReadWrite, EditAnywhere, Category = "Settings|ExtraAttributes", meta = (PCG_Overridable, InlineEditConditionToggle))
+	bool bOutputExtremityNeighborIndexAttribute = false;
+
+	/** Name of the extremity neighbor index output attribute name. */
+	UPROPERTY(BlueprintReadWrite, EditAnywhere, Category = "Settings|ExtraAttributes", meta = (PCG_Overridable, EditCondition = "bOutputExtremityNeighborIndexAttribute"))
+	FName ExtremityNeighborIndexAttributeName = TEXT("ExtremityNeighborIndex");
+
+};
+
+class PCGSegmentSlicerHelpers;
+
+class FPCGSegmentSlicerElement : public FPCGSlicingBaseElement
+{
+public:
+	friend class PCGSegmentSlicerHelpers;
+protected:
+	virtual bool ExecuteInternal(FPCGContext* InContext) const override;
+};
diff --git a/Engine/Plugins/PCG/Source/PCG/Public/Elements/Grammar/PCGSlicingBase.h b/Engine/Plugins/PCG/Source/PCG/Public/Elements/Grammar/PCGSlicingBase.h
new file mode 100644
--- /dev/null
+++ b/Engine/Plugins/PCG/Source/PCG/Public/Elements/Grammar/PCGSlicingBase.h
@@ -0,0 +1,292 @@
+// Copyright Epic Games, Inc. All Rights Reserved.
+
+#pragma once
+
+#include "PCGModule.h"
+#include "PCGSettings.h"
+#include "Metadata/PCGAttributePropertySelector.h"
+
+#include "PCGSlicingBase.generated.h"
+
+class UPCGData;
+class UPCGPointData;
+
+USTRUCT(BlueprintType)
+struct FPCGSlicingModule
+{
+	GENERATED_BODY()
+
+	/** Symbol for the grammar. */
+	UPROPERTY(BlueprintReadWrite, EditAnywhere, Category = "")
+	FName Symbol = NAME_None;
+
+	/** Size of the block, aligned on the segment direction. */
+	UPROPERTY(BlueprintReadWrite, EditAnywhere, Category = "")
+	double Size = 100.0;
+
+	/** If the volume can be scaled to fit the remaining space or not. */
+	UPROPERTY(BlueprintReadWrite, EditAnywhere, Category = "")
+	bool bScalable = false;
+
+	/** For easier debugging, using Point color in conjunction with PCG Debug Color Material. */
+	UPROPERTY(BlueprintReadWrite, EditAnywhere, Category = "Debug")
+	FVector4 DebugColor = FVector4::One();
+};
+
+namespace PCGSlicingBaseConstants
+{
+	const FName ModulesInfoPinLabel = TEXT("ModulesInfo");
+	const FName SymbolAttributeName = TEXT("Symbol");
+	const FName SizeAttributeName = TEXT("Size");
+	const FName ScalableAttributeName = TEXT("bScalable");
+	const FName DebugColorAttributeName = TEXT("DebugColor");
+}
+
+UCLASS(MinimalAPI, Abstract, ClassGroup = (Procedural))
+class UPCGSlicingBaseSettings : public UPCGSettings
+{
+	GENERATED_BODY()
+
+public:
+	//~Begin UPCGSettings interface
+#if WITH_EDITOR
+	virtual EPCGSettingsType GetType() const override { return EPCGSettingsType::Spatial; }
+#endif
+	//~End UPCGSettings interface
+
+public:
+	/** Set it to true to pass the info as attribute set. Module info needs this structure (type + name): { FName Symbol, double Size, bool bScalable, FLinearColor DebugColor }. Types don't need to match exactly, must be able to be converted to the expected type. */
+	UPROPERTY(BlueprintReadWrite, EditAnywhere, Category = Settings)
+	bool bModuleInfoAsInput = false;
+
+	/** Read the grammar from the input spline as an attribute rather than directly from the settings.
+	* Grammar syntax:
+	* - Each symbol can have multiple characters
+	* - Modules are defined in '[]', multiple symbols in a module are seperated with ','
+	* - Modules can be repeated a fixed number of times, by adding a number after it (like [A,B]3 will produce ABABAB)
+	* - Modules can be marked repeated an indefinite number of times, with '*'. (like [A,B]* will produce ABABABAB... while it fits the allowed size).
+	*/
+	UPROPERTY(BlueprintReadWrite, EditAnywhere, Category = Settings)
+	bool bGrammarAsAttribute = false;
+
+	/** Fixed array of modules used for the slicing. */
+	UPROPERTY(BlueprintReadWrite, EditAnywhere, Category = Settings, meta = (EditCondition = "!bModuleInfoAsInput", EditConditionHides, DisplayAfter = bModuleInfoAsInput))
+	TArray<FPCGSlicingModule> ModulesInfo;
+
+	/** An encoded string that represents how to apply a set of rules to a series of defined modules. */
+	UPROPERTY(BlueprintReadWrite, EditAnywhere, Category = Settings, meta = (EditCondition = "!bGrammarAsAttribute", EditConditionHides, DisplayAfter = bGrammarAsAttribute, PCG_Overridable))
+	FString Grammar;
+
+	/** Attribute to be taken from the input spline containing the grammar to use for the slicing. */
+	UPROPERTY(BlueprintReadWrite, EditAnywhere, Category = Settings, meta = (EditCondition = "bGrammarAsAttribute", EditConditionHides, DisplayAfter = bGrammarAsAttribute, PCG_Overridable))
+	FPCGAttributePropertyInputSelector GrammarAttribute;
+
+	/** Do a match and set with the incoming module info, only if the module info is passed as input. */
+	UPROPERTY(BlueprintReadWrite, EditAnywhere, Category = "Settings|ExtraAttributes", meta = (EditCondition = "bModuleInfoAsInput", PCG_Overridable))
+	bool bForwardAttributesFromModulesInfo = false;
+
+	/** Name of the Symbol output attribute name. */
+	UPROPERTY(BlueprintReadWrite, EditAnywhere, Category = "Settings|ExtraAttributes", meta = (PCG_Overridable))
+	FName SymbolAttributeName = PCGSlicingBaseConstants::SymbolAttributeName;
+
+	UPROPERTY(BlueprintReadWrite, EditAnywhere, Category = "Settings|ExtraAttributes", meta = (PCG_Overridable, InlineEditConditionToggle))
+	bool bOutputSizeAttribute = true;
+
+	/** Name of the Size output attribute name, ignored if match and set from module info is true. */
+	UPROPERTY(BlueprintReadWrite, EditAnywhere, Category = "Settings|ExtraAttributes", meta = (PCG_Overridable, EditCondition = "bOutputSizeAttribute"))
+	FName SizeAttributeName = PCGSlicingBaseConstants::SizeAttributeName;
+
+	UPROPERTY(BlueprintReadWrite, EditAnywhere, Category = "Settings|ExtraAttributes", meta = (PCG_Overridable, InlineEditConditionToggle))
+	bool bOutputScalableAttribute = true;
+
+	/** Name of the Scalable output attribute name, ignored if match and set from module info is true. */
+	UPROPERTY(BlueprintReadWrite, EditAnywhere, Category = "Settings|ExtraAttributes", meta = (PCG_Overridable, EditCondition = "bOutputScalableAttribute"))
+	FName ScalableAttributeName = PCGSlicingBaseConstants::ScalableAttributeName;
+
+	UPROPERTY(BlueprintReadWrite, EditAnywhere, Category = "Settings|ExtraAttributes", meta = (PCG_Overridable, InlineEditConditionToggle))
+	bool bOutputDebugColorAttribute = false;
+
+	/** Name of the Debug Color output attribute name, ignored if match and set from module info is true. */
+	UPROPERTY(BlueprintReadWrite, EditAnywhere, Category = "Settings|ExtraAttributes", meta = (PCG_Overridable, EditCondition = "bOutputDebugColorAttribute"))
+	FName DebugColorAttributeName = PCGSlicingBaseConstants::DebugColorAttributeName;
+};
+
+struct FPCGTokenizedGrammar
+{
+	TArray<FName> Symbols;
+	TArray<bool> AreSymbolsScalable;
+	TArray<double> SymbolSizes;
+	int NumRepeat = 1;
+	bool bScalable = false;
+	double Size = 0.0;
+
+	bool IsValid() const { return Size > 0 && !Symbols.IsEmpty() && Symbols.Num() == AreSymbolsScalable.Num() && Symbols.Num() == SymbolSizes.Num(); }
+	double GetSize() const { return Size; }
+	int GetNumRepeat() const { return NumRepeat; }
+	bool IsScalable() const { return bScalable; }
+	int32 GetSubmodulesCount() const { return Symbols.Num(); }
+	TArrayView<const bool> AreSubmodulesScalable() const { return AreSymbolsScalable; }
+	TArrayView<const double> SubmoduleSizes() const { return SymbolSizes; }
+};
+
+class FPCGSlicingBaseElement : public IPCGElement
+{
+public:
+	using FPCGModulesInfoMap = TMap<FName, FPCGSlicingModule>;
+
+protected:
+	FPCGModulesInfoMap GetModulesInfoMap(FPCGContext* InContext, const UPCGSlicingBaseSettings* InSettings, const UPCGParamData*& OutModuleInfoParamData) const;
+	TArray<FPCGTokenizedGrammar> GetTokenizeGrammar(FPCGContext* InContext, const UPCGData* InputData, const UPCGSlicingBaseSettings* InSettings, const FPCGModulesInfoMap& InModulesInfo, double& OutMinSize) const;
+	TArray<FPCGTokenizedGrammar> GetTokenizeGrammar(FPCGContext* InContext, const FString& InGrammar, const FPCGModulesInfoMap& InModulesInfo, double& OutMinSize) const;
+	TMap<FString, TArray<FPCGTokenizedGrammar>> GetTokenizeGrammarForPoints(FPCGContext* InContext, const UPCGPointData* InputData, const UPCGSlicingBaseSettings* InSettings, const FPCGModulesInfoMap& InModulesInfo, double& OutMinSize) const;
+	bool MatchAndSetAttributes(const TArray<FPCGTaggedData>& InputData, TArray<FPCGTaggedData>& OutputData, const UPCGParamData* InModuleInfoParamData, const UPCGSlicingBaseSettings* InSettings) const;
+};
+
+namespace PCGSlicingBase
+{
+#define PCG_SLICING_BASE_USES_CONCEPTS 0
+
+#if PCG_SLICING_BASE_USES_CONCEPTS
+	// To validate that the type has the right interface
+	// To be enable when this is supported on all our platforms.
+	template <typename T>
+	concept IsValidPCGSubDivModuleInstance =
+		requires(T t) {
+			{ t.IsValid() } -> std::same_as<bool>;
+			{ t.GetSize() } -> std::same_as<double>;
+			{ t.GetNumRepeat() } -> std::same_as<int>;
+			{ t.IsScalable() } -> std::same_as<bool>;
+			{ t.GetSubmodulesCount() } -> std::same_as<int32>;
+			{ t.AreSubmodulesScalable() } -> std::same_as<TArrayView<const bool>>;
+			{ t.SubmoduleSizes() } -> std::same_as<TArrayView<const double>>;
+	};
+#endif // PCG_SLICING_BASE_USES_CONCEPTS
+
+	template <typename T>
+	struct PCGSubDivModuleInstance
+	{
+		const T* Module = nullptr;
+		int32 NumRepeat = 0;
+		// If the module is made of multiple submodules, we'll have 1 extra scale per submodule.
+		// Extra scale to add to the module initial scale on the slicing direction.
+		TArray<double> ExtraScales;
+	};
+
+	template <typename T>
+	bool Subdivide(const TArray<T>& Modules, double Length, TArray<PCGSubDivModuleInstance<T>>& OutModuleInstances, double& RemainingLength, FPCGContext* InOptionalContext = nullptr)
+	{
+		OutModuleInstances.Empty(Modules.Num());
+		RemainingLength = Length;
+
+		if (Modules.IsEmpty() || FMath::IsNearlyZero(Length))
+		{
+			return true;
+		}
+
+		for (const T& Module : Modules)
+		{
+			if (!Module.IsValid())
+			{
+				continue;
+			}
+
+			PCGSubDivModuleInstance<T>& ModuleInstance = OutModuleInstances.Emplace_GetRef();
+			ModuleInstance.Module = &Module;
+			ModuleInstance.ExtraScales.SetNumZeroed(Module.GetSubmodulesCount());
+			if (Module.GetNumRepeat() > 0)
+			{
+				ModuleInstance.NumRepeat = Module.GetNumRepeat();
+				RemainingLength -= Module.GetSize() * ModuleInstance.NumRepeat;
+				if (RemainingLength < 0)
+				{
+					PCGLog::LogErrorOnGraph(NSLOCTEXT("PCGSlicingBase", "SegmentCutFail", "Grammar doesn't fit for this segment."), InOptionalContext);
+					return false;
+				}
+			}
+		}
+
+		if (OutModuleInstances.IsEmpty())
+		{
+			return false;
+		}
+
+		// When we are done and we still have some segemnt left, place the repeatable
+		int32 CurrentModuleIndex = 0;
+		bool bHasModifiedSomething = false;
+		while (RemainingLength >= 0)
+		{
+			PCGSubDivModuleInstance<T>& ModuleInstance = OutModuleInstances[CurrentModuleIndex];
+			if (ModuleInstance.Module->GetNumRepeat() <= 0)
+			{
+				if (RemainingLength >= ModuleInstance.Module->GetSize())
+				{
+					++ModuleInstance.NumRepeat;
+					RemainingLength -= ModuleInstance.Module->GetSize();
+					bHasModifiedSomething = true;
+				}
+			}
+
+			if (++CurrentModuleIndex == OutModuleInstances.Num())
+			{
+				if (!bHasModifiedSomething)
+				{
+					// Nothing left
+					break;
+				}
+
+				bHasModifiedSomething = false;
+				CurrentModuleIndex = 0;
+			}
+		}
+
+		// Finally, try to stretch the scalable to get a complete match.
+		if (!FMath::IsNearlyZero(RemainingLength))
+		{
+			int32 NumScalableSubmodules = 0;
+
+			TArray<PCGSubDivModuleInstance<T>*> ScalableInstances;
+			for (PCGSubDivModuleInstance<T>& ModuleInstance : OutModuleInstances)
+			{
+				if (ModuleInstance.Module->IsScalable() && ModuleInstance.NumRepeat > 0)
+				{
+					ScalableInstances.Add(&ModuleInstance);
+					for (const bool bIsScalable : ModuleInstance.Module->AreSubmodulesScalable())
+					{
+						if (bIsScalable)
+						{
+							NumScalableSubmodules += ModuleInstance.NumRepeat;
+						}
+					}
+				}
+			}
+
+
+			if (NumScalableSubmodules > 0 && !ScalableInstances.IsEmpty())
+			{
+				const double ExtraLengthPerSubmodule = RemainingLength / NumScalableSubmodules;
+
+				for (PCGSubDivModuleInstance<T>* ModuleInstance : ScalableInstances)
+				{
+					const TArrayView<const bool> AreSubmodulesScalable = ModuleInstance->Module->AreSubmodulesScalable();
+					const TArrayView<const double> SubmoduleSizes = ModuleInstance->Module->SubmoduleSizes();
+					check(AreSubmodulesScalable.Num() == SubmoduleSizes.Num());
+					for (int32 i = 0; i < SubmoduleSizes.Num(); ++i)
+					{
+						if (AreSubmodulesScalable[i])
+						{
+							const double SubmoduleSize = SubmoduleSizes[i];
+							if (!FMath::IsNearlyZero(SubmoduleSize))
+							{
+								ModuleInstance->ExtraScales[i] = ExtraLengthPerSubmodule / SubmoduleSize;
+							}
+						}
+					}
+				}
+
+				RemainingLength = 0.0;
+			}
+		}
+
+		return true;
+	}
+}
\ No newline at end of file
diff --git a/Engine/Plugins/PCG/Source/PCG/Public/Elements/Grammar/PCGSplineToSegment.h b/Engine/Plugins/PCG/Source/PCG/Public/Elements/Grammar/PCGSplineToSegment.h
new file mode 100644
--- /dev/null
+++ b/Engine/Plugins/PCG/Source/PCG/Public/Elements/Grammar/PCGSplineToSegment.h
@@ -0,0 +1,56 @@
+// Copyright Epic Games, Inc. All Rights Reserved.
+
+#pragma once
+
+#include "PCGSettings.h"
+
+#include "PCGSplineToSegment.generated.h"
+
+/**
+* Take a spline as input and create a point data, with each point being a segment defined by 2 connected control points.
+* The point position will be in the middle of those 2 control points, and the extents of the point will be half the difference between those 2 control points.
+*/
+UCLASS(MinimalAPI, BlueprintType, ClassGroup = (Procedural))
+class UPCGSplineToSegmentSettings : public UPCGSettings
+{
+	GENERATED_BODY()
+
+public:
+	//~Begin UPCGSettings interface
+#if WITH_EDITOR
+	virtual FName GetDefaultNodeName() const override;
+	virtual FText GetDefaultNodeTitle() const override;
+	virtual FText GetNodeTooltipText() const override;
+	virtual EPCGSettingsType GetType() const override { return EPCGSettingsType::Spatial; }
+#endif
+
+protected:
+	virtual TArray<FPCGPinProperties> InputPinProperties() const override;
+	virtual TArray<FPCGPinProperties> OutputPinProperties() const override;
+	virtual FPCGElementPtr CreateElement() const override;
+	//~End UPCGSettings interface
+
+public:
+	/** Can extract the tangents with previous segment and next segment. */
+	UPROPERTY(BlueprintReadWrite, EditAnywhere, Category = Settings)
+	bool bExtractTangents = false;
+
+	/** Can extract the angle with previous tangent and next tangent. Will be 0 at the extremity for non-closed splines. */
+	UPROPERTY(BlueprintReadWrite, EditAnywhere, Category = Settings)
+	bool bExtractAngles = true;
+
+	/** Can set the index of the previous and next segment (to keep information on connectivity). */
+	UPROPERTY(BlueprintReadWrite, EditAnywhere, Category = Settings)
+	bool bExtractConnectivityInfo = true;
+
+	/** Can output a global attribute to know if the spline points order is clockwise or counterclockwise. (Only for closed loops). */
+	UPROPERTY(BlueprintReadWrite, EditAnywhere, Category = Settings)
+	bool bExtractClockwiseInfo = true;
+};
+
+class FPCGSplineToSegmentElement : public IPCGElement
+{
+protected:
+	virtual bool ExecuteInternal(FPCGContext* InContext) const override;
+};
+
diff --git a/Engine/Plugins/PCG/Source/PCG/Public/Elements/Grammar/PCGVolumeSlicer.h b/Engine/Plugins/PCG/Source/PCG/Public/Elements/Grammar/PCGVolumeSlicer.h
new file mode 100644
--- /dev/null
+++ b/Engine/Plugins/PCG/Source/PCG/Public/Elements/Grammar/PCGVolumeSlicer.h
@@ -0,0 +1,54 @@
+// Copyright Epic Games, Inc. All Rights Reserved.
+
+#pragma once
+
+#include "PCGSettings.h"
+#include "Metadata/PCGAttributePropertySelector.h"
+#include "Elements/Grammar/PCGSlicingBase.h"
+
+#include "PCGVolumeSlicer.generated.h"
+
+UCLASS(MinimalAPI, BlueprintType, ClassGroup = (Procedural))
+class UPCGVolumeSlicerSettings : public UPCGSlicingBaseSettings
+{
+	GENERATED_BODY()
+
+public:
+	//~Begin UPCGSettings interface
+#if WITH_EDITOR
+	virtual FName GetDefaultNodeName() const override;
+	virtual FText GetDefaultNodeTitle() const override;
+#endif
+
+protected:
+	virtual TArray<FPCGPinProperties> InputPinProperties() const override;
+	virtual TArray<FPCGPinProperties> OutputPinProperties() const override;
+	virtual FPCGElementPtr CreateElement() const override;
+	//~End UPCGSettings interface
+
+public:
+	/** Set it to true if you want the extrude vector to be taken from the input spline as attribute, or fixed in the settings. */
+	UPROPERTY(BlueprintReadWrite, EditAnywhere, Category = Settings)
+	bool bExtrudeVectorAsAttribute = false;
+
+	UPROPERTY(BlueprintReadWrite, EditAnywhere, Category = Settings, meta = (EditCondition = "!bExtrudeVectorAsAttribute", EditConditionHides, DisplayAfter = bExtrudeVectorAsAttribute, PCG_Overridable))
+	FVector ExtrudeVector = FVector(0.0, 0.0, 1000.0);
+
+	/** Attribute to be taken from the input spline containing the extrude vector for the slicing. */
+	UPROPERTY(BlueprintReadWrite, EditAnywhere, Category = Settings, meta = (EditCondition = "bExtrudeVectorAsAttribute", EditConditionHides, DisplayAfter = bExtrudeVectorAsAttribute, PCG_Overridable))
+	FPCGAttributePropertyInputSelector ExtrudeVectorAttribute;
+
+	UPROPERTY(BlueprintReadWrite, EditAnywhere, Category = "Settings|ExtraAttributes", meta = (PCG_Overridable, InlineEditConditionToggle))
+	bool bOutputSplineIndexAttribute = true;
+
+	/** Name of the spline index output attribute name. */
+	UPROPERTY(BlueprintReadWrite, EditAnywhere, Category = "Settings|ExtraAttributes", meta = (PCG_Overridable, EditCondition = "bOutputSplineIndexAttribute"))
+	FName SplineIndexAttributeName = TEXT("SplineIndex");
+};
+
+class FPCGVolumeSlicerElement : public FPCGSlicingBaseElement
+{
+protected:
+	virtual bool ExecuteInternal(FPCGContext* InContext) const override;
+};
+
diff --git a/Engine/Plugins/PCG/Source/PCG/Public/Grammar/PCGGrammarParser.h b/Engine/Plugins/PCG/Source/PCG/Public/Grammar/PCGGrammarParser.h
new file mode 100644
--- /dev/null
+++ b/Engine/Plugins/PCG/Source/PCG/Public/Grammar/PCGGrammarParser.h
@@ -0,0 +1,70 @@
+// Copyright Epic Games, Inc. All Rights Reserved.
+
+#pragma once
+
+#include "Containers/Array.h"
+#include "Internationalization/Text.h"
+#include "UObject/NameTypes.h"
+
+namespace PCGGrammar
+{
+	enum class EModuleType
+	{
+		Base,
+		Stochastic,
+		Priority
+	};
+
+	struct FModuleDescriptor
+	{
+		struct FSubmodule
+		{
+			FName ID;
+			int8 Weight;
+		};
+
+		EModuleType Type = EModuleType::Base;
+		int32 Repetitions = 1;
+		// The matched beginning and ending indices of the grammar string for this module
+		TPair<int32, int32> GrammarStartEndIndices;
+		TArray<FSubmodule> Submodules;
+	};
+}
+
+struct PCG_API FPCGGrammarResult
+{
+	enum class ELogType
+	{
+		Log = ELogVerbosity::Log,
+		Warning = ELogVerbosity::Warning,
+		Error = ELogVerbosity::Error
+	};
+
+	struct FLog
+	{
+		FLog(FText InMessage, ELogType InVerbosity = ELogType::Log)
+			: Message(std::move(InMessage))
+			, Verbosity(InVerbosity)
+		{}
+
+		FText Message;
+		ELogType Verbosity = ELogType::Log;
+
+		bool operator==(const FLog& Other) const { return Message.EqualTo(Other.Message) && Verbosity == Other.Verbosity; }
+	};
+
+	void AddLog(FText Message, ELogType Verbosity = ELogType::Log);
+
+	const TArray<FLog>& GetLogs() const{ return Logs; }
+
+	bool bSuccess = true;
+	TArray<PCGGrammar::FModuleDescriptor> Modules;
+
+private:
+	TArray<FLog> Logs;
+};
+
+namespace PCGGrammar
+{
+	PCG_API FPCGGrammarResult Parse(const FString& Grammar, bool bValidateGrammar = true);
+}
