[PCG] Fix Mac compile issue

#rnx

diff --git a/Engine/Plugins/PCG/Source/PCG/Public/Grammar/PCGGrammarParser.h b/Engine/Plugins/PCG/Source/PCG/Public/Grammar/PCGGrammarParser.h
--- a/Engine/Plugins/PCG/Source/PCG/Public/Grammar/PCGGrammarParser.h
+++ b/Engine/Plugins/PCG/Source/PCG/Public/Grammar/PCGGrammarParser.h
@@ -19,6 +19,11 @@ namespace PCGGrammar
 	{
 		struct FSubmodule
 		{
+			FSubmodule(const FName InID, const int8 InWeight)
+				: ID(InID)
+				, Weight(InWeight)
+			{}
+
 			FName ID;
 			int8 Weight;
 		};
