[PCG] Prevent crash when the render data is null, could happen in certain server situations.

diff --git a/Engine/Plugins/PCG/Source/PCG/Private/Helpers/PCGActorHelpers.cpp b/Engine/Plugins/PCG/Source/PCG/Private/Helpers/PCGActorHelpers.cpp
--- a/Engine/Plugins/PCG/Source/PCG/Private/Helpers/PCGActorHelpers.cpp
+++ b/Engine/Plugins/PCG/Source/PCG/Private/Helpers/PCGActorHelpers.cpp
@@ -97,7 +97,7 @@ UPCGManagedISMComponent* UPCGActorHelpers::GetOrCreateManagedISMC(AActor* InTarg
 #if WITH_EDITOR
 	const bool bMeshHasNaniteData = StaticMesh->IsNaniteEnabled();
 #else
-	const bool bMeshHasNaniteData = StaticMesh->GetRenderData()->HasValidNaniteData();
+	const bool bMeshHasNaniteData = StaticMesh->GetRenderData() && StaticMesh->GetRenderData()->HasValidNaniteData();
 #endif
 
 	FString ComponentName;
