Update CreateSpline node following the thread usage changes
#rnx

diff --git a/Engine/Plugins/PCG/Source/PCG/Private/Elements/PCGCreateSpline.cpp b/Engine/Plugins/PCG/Source/PCG/Private/Elements/PCGCreateSpline.cpp
--- a/Engine/Plugins/PCG/Source/PCG/Private/Elements/PCGCreateSpline.cpp
+++ b/Engine/Plugins/PCG/Source/PCG/Private/Elements/PCGCreateSpline.cpp
@@ -122,8 +122,7 @@ bool FPCGCreateSplineElement::ExecuteInternal(FPCGContext* Context) const
 			LeaveTangentAttribute = static_cast<const FPCGMetadataAttribute<FVector>*>(LeaveTangentBaseAttribute);
 		}
 
-		check(IsInGameThread());
-		UPCGSplineData* SplineData = NewObject<UPCGSplineData>();
+		UPCGSplineData* SplineData = FPCGContext::NewObject_AnyThread<UPCGSplineData>(Context);
 		AActor* SplineActor = TargetActor;
 
 		const TArray<FPCGPoint>& Points = PointData->GetPoints();
@@ -163,6 +162,7 @@ bool FPCGCreateSplineElement::ExecuteInternal(FPCGContext* Context) const
 
 		if (Settings->Mode != EPCGCreateSplineMode::CreateDataOnly)
 		{
+			check(IsInGameThread());
 			SplineComponent = NewObject<USplineComponent>(SplineActor);
 			SplineComponent->ComponentTags.Add(Context->SourceComponent.Get()->GetFName());
 			SplineComponent->ComponentTags.Add(PCGHelpers::DefaultPCGTag);
