Prevent crash in standalone mode when GEditor is null

#rnx

diff --git a/Engine/Plugins/PCG/Source/PCG/Private/PCGComponent.cpp b/Engine/Plugins/PCG/Source/PCG/Private/PCGComponent.cpp
--- a/Engine/Plugins/PCG/Source/PCG/Private/PCGComponent.cpp
+++ b/Engine/Plugins/PCG/Source/PCG/Private/PCGComponent.cpp
@@ -1368,7 +1368,7 @@ FPCGTaskId UPCGComponent::CreateCleanupTask(bool bRemoveComponents, const TArray
 		if (UPCGComponent* ThisComponent = ThisComponentWeakPtr.Get())
 		{
 #if WITH_EDITOR
-			if (UWorld* ThisWorld = ThisComponent->GetWorld())
+			if (UWorld* ThisWorld = ThisComponent->GetWorld(); ThisWorld && GEditor) // FActorFolders require the editor
 			{
 				// Look for a nested generated results subfolder and remove it if it exists
 				FString FolderPath;
