[PCG] Add FPCGContext param to GetSpatialUnionOfInputsByPin to allow call to FPCGContext::NewObject_AnyThread

#rnx

diff --git a/Engine/Plugins/PCG/Source/PCG/Private/Elements/PCGElevationIsolines.cpp b/Engine/Plugins/PCG/Source/PCG/Private/Elements/PCGElevationIsolines.cpp
--- a/Engine/Plugins/PCG/Source/PCG/Private/Elements/PCGElevationIsolines.cpp
+++ b/Engine/Plugins/PCG/Source/PCG/Private/Elements/PCGElevationIsolines.cpp
@@ -64,7 +64,7 @@ bool FPCGElevationIsolinesElement::ExecuteInternal(FPCGContext* InContext) const
 	}
 
 	bool bOutUnionDataCreated = false;
-	const UPCGSpatialData* BoundingShape = InContext->InputData.GetSpatialUnionOfInputsByPin(TEXT("BoundingShape"), bOutUnionDataCreated);
+	const UPCGSpatialData* BoundingShape = InContext->InputData.GetSpatialUnionOfInputsByPin(InContext, TEXT("BoundingShape"), bOutUnionDataCreated);
 
 	// Fallback to getting bounds from actor
 	if (!BoundingShape && InContext->SourceComponent.IsValid())
diff --git a/Engine/Plugins/PCG/Source/PCG/Private/Elements/PCGOuterIntersectionElement.cpp b/Engine/Plugins/PCG/Source/PCG/Private/Elements/PCGOuterIntersectionElement.cpp
--- a/Engine/Plugins/PCG/Source/PCG/Private/Elements/PCGOuterIntersectionElement.cpp
+++ b/Engine/Plugins/PCG/Source/PCG/Private/Elements/PCGOuterIntersectionElement.cpp
@@ -115,7 +115,7 @@ bool FPCGOuterIntersectionElement::ExecuteInternal(FPCGContext* Context) const
 		}
 
 		bool bDummy;
-		if (const UPCGSpatialData* UnionData = Context->InputData.GetSpatialUnionOfInputsByPin(InputPin->Properties.Label, bDummy))
+		if (const UPCGSpatialData* UnionData = Context->InputData.GetSpatialUnionOfInputsByPin(Context, InputPin->Properties.Label, bDummy))
 		{
 			SecondarySourceUnionArray.Push(UnionData);
 		}
diff --git a/Engine/Plugins/PCG/Source/PCG/Private/Elements/PCGSplineSampler.cpp b/Engine/Plugins/PCG/Source/PCG/Private/Elements/PCGSplineSampler.cpp
--- a/Engine/Plugins/PCG/Source/PCG/Private/Elements/PCGSplineSampler.cpp
+++ b/Engine/Plugins/PCG/Source/PCG/Private/Elements/PCGSplineSampler.cpp
@@ -1513,7 +1513,7 @@ bool FPCGSplineSamplerElement::ExecuteInternal(FPCGContext* Context) const
 		// TODO: Once we support time-slicing, put this in the context and root (see FPCGSurfaceSamplerContext)
 		bool bUnionCreated = false;
 		// Grab the Bounding Shape input if there is one.
-		BoundingShape = Context->InputData.GetSpatialUnionOfInputsByPin(PCGSplineSamplerConstants::BoundingShapeLabel, bUnionCreated);
+		BoundingShape = Context->InputData.GetSpatialUnionOfInputsByPin(Context, PCGSplineSamplerConstants::BoundingShapeLabel, bUnionCreated);
 
 		// Fallback to getting bounds from actor
 		if (!BoundingShape && Context->SourceComponent.IsValid())
diff --git a/Engine/Plugins/PCG/Source/PCG/Private/Elements/PCGVolumeSampler.cpp b/Engine/Plugins/PCG/Source/PCG/Private/Elements/PCGVolumeSampler.cpp
--- a/Engine/Plugins/PCG/Source/PCG/Private/Elements/PCGVolumeSampler.cpp
+++ b/Engine/Plugins/PCG/Source/PCG/Private/Elements/PCGVolumeSampler.cpp
@@ -194,7 +194,7 @@ namespace PCGVolumeSamplerHelpers
 		{
 			bool bUnionWasCreated;
 			// Get a union of inputs and if successful, add it to the root. Will be removed and marked for GC in the state destructor
-			OutState.BoundingShape = Context->InputData.GetSpatialUnionOfInputsByPin(PCGVolumeSamplerConstants::BoundingShapeLabel, bUnionWasCreated);
+			OutState.BoundingShape = Context->InputData.GetSpatialUnionOfInputsByPin(Context, PCGVolumeSamplerConstants::BoundingShapeLabel, bUnionWasCreated);
 			if (OutState.BoundingShape && bUnionWasCreated)
 			{
 				Context->TrackObject(OutState.BoundingShape);
diff --git a/Engine/Plugins/PCG/Source/PCG/Private/PCGData.cpp b/Engine/Plugins/PCG/Source/PCG/Private/PCGData.cpp
--- a/Engine/Plugins/PCG/Source/PCG/Private/PCGData.cpp
+++ b/Engine/Plugins/PCG/Source/PCG/Private/PCGData.cpp
@@ -2,6 +2,7 @@
 
 #include "PCGData.h"
 
+#include "PCGContext.h"
 #include "PCGNode.h"
 #include "PCGParamData.h"
 #include "PCGPin.h"
@@ -191,6 +192,12 @@ int32 FPCGDataCollection::GetSpatialInputCountByPin(const FName& InPinLabel) con
 }
 
 const UPCGSpatialData* FPCGDataCollection::GetSpatialUnionOfInputsByPin(const FName& InPinLabel, bool& bOutUnionDataCreated) const
+{
+	check(IsInGameThread());
+	return GetSpatialUnionOfInputsByPin(nullptr, InPinLabel, bOutUnionDataCreated);
+}
+
+const UPCGSpatialData* FPCGDataCollection::GetSpatialUnionOfInputsByPin(FPCGContext* InContext, const FName& InPinLabel, bool& bOutUnionDataCreated) const
 {
 	TArray<FPCGTaggedData> SpatialDataOnPin = TaggedData.FilterByPredicate([&InPinLabel](const FPCGTaggedData& Data) {
 		return Data.Pin == InPinLabel && Data.Data && Data.Data->IsA<UPCGSpatialData>();
@@ -216,7 +223,7 @@ const UPCGSpatialData* FPCGDataCollection::GetSpatialUnionOfInputsByPin(const FN
 			if (!Union)
 			{
 				// Second valid data - set up union
-				Union = NewObject<UPCGUnionData>();
+				Union = FPCGContext::NewObject_AnyThread<UPCGUnionData>(InContext);
 				Union->Initialize(Result, SpatialInput);
 
 				// Make result union
diff --git a/Engine/Plugins/PCG/Source/PCG/Public/Helpers/PCGSettingsHelpers.h b/Engine/Plugins/PCG/Source/PCG/Public/Helpers/PCGSettingsHelpers.h
--- a/Engine/Plugins/PCG/Source/PCG/Public/Helpers/PCGSettingsHelpers.h
+++ b/Engine/Plugins/PCG/Source/PCG/Public/Helpers/PCGSettingsHelpers.h
@@ -207,9 +207,9 @@ PRAGMA_ENABLE_DEPRECATION_WARNINGS
 		});
 	}
 
-	inline const UPCGSpatialData* ComputeBoundingShape(const FPCGContext* Context, FName BoundingShapeLabel, bool& bOutUnionWasCreated)
+	inline const UPCGSpatialData* ComputeBoundingShape(FPCGContext* Context, FName BoundingShapeLabel, bool& bOutUnionWasCreated)
 	{
-		const UPCGSpatialData* BoundingShape = Context->InputData.GetSpatialUnionOfInputsByPin(BoundingShapeLabel, bOutUnionWasCreated);
+		const UPCGSpatialData* BoundingShape = Context->InputData.GetSpatialUnionOfInputsByPin(Context, BoundingShapeLabel, bOutUnionWasCreated);
 
 		// Fallback to getting bounds from actor
 		if (!BoundingShape && Context->SourceComponent.IsValid())
diff --git a/Engine/Plugins/PCG/Source/PCG/Public/PCGData.h b/Engine/Plugins/PCG/Source/PCG/Public/PCGData.h
--- a/Engine/Plugins/PCG/Source/PCG/Public/PCGData.h
+++ b/Engine/Plugins/PCG/Source/PCG/Public/PCGData.h
@@ -147,7 +147,11 @@ struct PCG_API FPCGDataCollection
 	/** Gets number of spatial data items on a given pin */
 	int32 GetSpatialInputCountByPin(const FName& InPinLabel) const;
 	/** Returns spatial union of all data on a given pin, returns null if no such data exists. bOutUnionDataCreated indicates if new data created that may need rooting. */
+	const UPCGSpatialData* GetSpatialUnionOfInputsByPin(FPCGContext* InContext, const FName& InPinLabel, bool& bOutUnionDataCreated) const;
+
+	UE_DEPRECATED(5.5, "Use version with FPCGContext param")
 	const UPCGSpatialData* GetSpatialUnionOfInputsByPin(const FName& InPinLabel, bool& bOutUnionDataCreated) const;
+
 	/** Returns all spatial data in the collection with the given tag */
 	TArray<FPCGTaggedData> GetTaggedInputs(const FString& InTag) const;
 	/** Returns all settings in the collection */
