Fixed issues that arose when the cache would contain elements that are outered to things in sublevels, which can lead to reference leaks.

#rnx

diff --git a/Engine/Plugins/PCG/Source/PCG/Private/Elements/PCGDataFromActor.cpp b/Engine/Plugins/PCG/Source/PCG/Private/Elements/PCGDataFromActor.cpp
--- a/Engine/Plugins/PCG/Source/PCG/Private/Elements/PCGDataFromActor.cpp
+++ b/Engine/Plugins/PCG/Source/PCG/Private/Elements/PCGDataFromActor.cpp
@@ -18,6 +18,7 @@
 #include "Utils/PCGGraphExecutionLogging.h"
 
 #include "Algo/AnyOf.h"
+#include "Engine/World.h"
 #include "GameFramework/Actor.h"
 #include "UObject/Package.h"
 
@@ -787,6 +788,15 @@ void FPCGDataFromActorElement::ProcessActor(FPCGContext* Context, const UPCGData
 				if (ensure(TaggedData.Data))
 				{
 					FPCGTaggedData& DuplicatedTaggedData = Outputs.Add_GetRef(TaggedData);
+
+					// In cases where the owner of the data (component -> owner) isn't in the main world, we MUST duplicate the data, otherwise we'll cause reference leaks in the cache.
+#if WITH_EDITOR
+					if (!SourceOwner || (Component->GetOwner() && Component->GetOwner()->GetLevel() != SourceOwner->GetWorld()->PersistentLevel))
+					{
+						DuplicatedTaggedData.Data = Cast<UPCGData>(StaticDuplicateObject(TaggedData.Data, GetTransientPackage()));
+					}
+#endif
+
 					DuplicatedTaggedData.Tags.Add(PCGDataFromActorConstants::PCGComponentDataGridSizeTagPrefix + FString::FromInt(PCGHiGenGrid::GridToGridSize(Component->GetGenerationGrid())));
 				}
 			}
diff --git a/Engine/Plugins/PCG/Source/PCG/Private/PCGComponent.cpp b/Engine/Plugins/PCG/Source/PCG/Private/PCGComponent.cpp
--- a/Engine/Plugins/PCG/Source/PCG/Private/PCGComponent.cpp
+++ b/Engine/Plugins/PCG/Source/PCG/Private/PCGComponent.cpp
@@ -567,15 +567,24 @@ void UPCGComponent::PostProcessGraph(const FBox& InNewBounds, bool bInGenerated,
 				{
 					// Visit the generated data; if the outer is the transient package or this component for all data, then
 					// we don't need to duplicate the data and can change the outer & flatten the data without any additional copies.
+					// Caveat: since that data could be in the cache currently, we shouldn't reouter it to a component that's not in the persistent level, otherwise when the
+					// sublevel is unloaded, it could leak references.
+#if WITH_EDITOR
+					bool bDataCanBeStolen = GetOwner() && GetOwner()->GetWorld() && (GetOwner()->GetLevel() == GetOwner()->GetWorld()->PersistentLevel);
+#else
 					bool bDataCanBeStolen = true;
+#endif
 
-					TaggedData.Data->VisitDataNetwork([this, &bDataCanBeStolen](const UPCGData* InData)
+					if (bDataCanBeStolen)
 					{
-						if (InData && InData->GetOuter() != GetTransientPackage() && InData->GetOuter() != this)
+						TaggedData.Data->VisitDataNetwork([this, &bDataCanBeStolen](const UPCGData* InData)
 						{
-							bDataCanBeStolen = false;
-						}
-					});
+							if (InData && InData->GetOuter() != GetTransientPackage() && InData->GetOuter() != this)
+							{
+								bDataCanBeStolen = false;
+							}
+						});
+					}
 
 					FPCGTaggedData OutputTaggedData = TaggedData;
 
