Collision wrapper/self-pruning improvements
- Normalized behavior when there are no collision shapes
- Allowed decay to simple/complex with an option
- Removed the recompute octree option, it was a bit too complex and ambiguous

#rnx

diff --git a/Engine/Plugins/PCG/Source/PCG/Private/Data/PCGCollisionWrapperData.cpp b/Engine/Plugins/PCG/Source/PCG/Private/Data/PCGCollisionWrapperData.cpp
--- a/Engine/Plugins/PCG/Source/PCG/Private/Data/PCGCollisionWrapperData.cpp
+++ b/Engine/Plugins/PCG/Source/PCG/Private/Data/PCGCollisionWrapperData.cpp
@@ -7,7 +7,9 @@
 #include "Elements/Metadata/PCGMetadataElementCommon.h"
 
 #include "Components/StaticMeshComponent.h"
+#include "Chaos/GeometryQueries.h"
 #include "Engine/StaticMesh.h"
+#include "Physics/PhysicsFiltering.h"
 #include "PhysicsEngine/BodyInstance.h"
 #include "PhysicsEngine/BodySetup.h"
 #include "Serialization/ArchiveCrc32.h"
@@ -135,76 +137,57 @@ FBodyInstance* FPCGCollisionWrapper::GetBodyInstance(int32 EntryIndex) const
 	return ((bInitialized && IndexToBodyInstance[EntryIndex] != INDEX_NONE) ? BodyInstances[IndexToBodyInstance[EntryIndex]] : nullptr);
 }
 
-bool FPCGCollisionWrapper::InitializeOctree(const UPCGPointData* InPointData, const TArray<FSoftObjectPath>& InMeshPaths, UPCGPointData::PointOctree& OutOctree, TArray<FPCGPointRef>* OutOctreePointRefs) const
+void FPCGCollisionWrapper::GetShapeArray(int32 EntryIndex, EPCGCollisionQueryFlag QueryFlag, PhysicsInterfaceTypes::FInlineShapeArray& OutShapeArray) const
 {
-	TArray<FBox> MeshBoundsList;
-	MeshBoundsList.Reserve(InMeshPaths.Num());
-	bool bHasValidBounds = false;
+	GetShapeArray(GetBodyInstance(EntryIndex), QueryFlag, OutShapeArray);
+}
 
-	for (int32 MeshIndex = 0; MeshIndex < InMeshPaths.Num(); ++MeshIndex)
-	{
-		TSoftObjectPtr<UStaticMesh> Mesh(InMeshPaths[MeshIndex]);
-		if (Mesh.Get()) // should already be loaded by now.
-		{
-			MeshBoundsList.Add(Mesh->GetBoundingBox());
-			bHasValidBounds = true;
-		}
-		else
-		{
-			MeshBoundsList.Emplace(EForceInit::ForceInit);
-		}
-	}
+bool FPCGCollisionWrapper::GetShapeArray(FBodyInstance* Body, EPCGCollisionQueryFlag QueryFlag, PhysicsInterfaceTypes::FInlineShapeArray& OutShapeArray)
+{
+	OutShapeArray.Reset();
 
-	if (bHasValidBounds)
+	if (!Body)
 	{
-		const TArray<FPCGPoint>& Points = InPointData->GetPoints();
-		TArray<FPCGPointRef> PointRefs;
-		PointRefs.Reserve(Points.Num());
-
-		for (int32 PointIndex = 0; PointIndex < Points.Num(); ++PointIndex)
-		{
-			const FPCGPoint& Point = Points[PointIndex];
-			const FBox& MeshBounds = MeshBoundsList[IndexToBodyInstance[PointIndex]];
-
-			if (MeshBounds.IsValid)
-			{
-				PointRefs.Emplace(Point, MeshBounds);
-			}
-			else
-			{
-				PointRefs.Emplace(Point);
-			}
-		}
+		return true;
+	}
 
-		FBox PointsBounds = InPointData->GetBounds();
-		TOctree2<FPCGPointRef, FPCGPointRefSemantics> NewOctree(PointsBounds.GetCenter(), PointsBounds.GetExtent().Length());
+	PhysicsInterfaceTypes::FInlineShapeArray AllShapes;
+	FillInlineShapeArray_AssumesLocked(AllShapes, Body->GetPhysicsActorHandle());
 
-		for (const FPCGPointRef& PointRef : PointRefs)
+	for (auto& Shape : AllShapes)
+	{
+		if (!Shape.Shape)
 		{
-			NewOctree.AddElement(PointRef);
+			continue;
 		}
 
-		OutOctree = MoveTemp(NewOctree);
+		FCollisionFilterData ShapeFilter = Shape.Shape->GetQueryData();
+		const bool bShapeIsComplex = (ShapeFilter.Word3 & EPDF_ComplexCollision) != 0;
+		const bool bShapeIsSimple = (ShapeFilter.Word3 & EPDF_SimpleCollision) != 0;
 
-		if (OutOctreePointRefs)
+		if ((bShapeIsSimple && (QueryFlag == EPCGCollisionQueryFlag::Simple || QueryFlag == EPCGCollisionQueryFlag::SimpleFirst)) ||
+			(bShapeIsComplex && (QueryFlag == EPCGCollisionQueryFlag::Complex || QueryFlag == EPCGCollisionQueryFlag::ComplexFirst)))
 		{
-			*OutOctreePointRefs = MoveTemp(PointRefs);
+			OutShapeArray.Add(Shape);
 		}
-
-		return true;
 	}
-	else
+
+	if (!OutShapeArray.IsEmpty() || QueryFlag == EPCGCollisionQueryFlag::Simple || QueryFlag == EPCGCollisionQueryFlag::Complex)
 	{
-		return false;
+		return true;
 	}
+
+	// Otherwise, we can take the remainder of the shapes, which means we can take everything
+	OutShapeArray = AllShapes;
+	return false;
 }
 
-bool UPCGCollisionWrapperData::Initialize(const UPCGPointData* InPointData, const FPCGAttributePropertyInputSelector& InCollisionSelector, bool bInUseComplexCollision, bool bInUseAccurateOctree)
+bool UPCGCollisionWrapperData::Initialize(const UPCGPointData* InPointData, const FPCGAttributePropertyInputSelector& InCollisionSelector, EPCGCollisionQueryFlag InCollisionQueryFlag)
 {
 	TArray<FSoftObjectPath> MeshesToLoad;
-	if (PreInitializeAndGatherMeshesEx(InPointData, InCollisionSelector, bInUseComplexCollision, MeshesToLoad))
+	if (PreInitializeAndGatherMeshesEx(InPointData, InCollisionSelector, InCollisionQueryFlag, MeshesToLoad))
 	{
-		FinalizeInitializationEx(InPointData, MeshesToLoad, bInUseAccurateOctree);
+		FinalizeInitializationEx(MeshesToLoad);
 		return true;
 	}
 	else
@@ -213,7 +196,7 @@ bool UPCGCollisionWrapperData::Initialize(const UPCGPointData* InPointData, cons
 	}
 }
 
-bool UPCGCollisionWrapperData::PreInitializeAndGatherMeshesEx(const UPCGPointData* InPointData, const FPCGAttributePropertyInputSelector& InCollisionSelector, bool bInUseComplexCollision, TArray<FSoftObjectPath>& OutMeshesToLoad)
+bool UPCGCollisionWrapperData::PreInitializeAndGatherMeshesEx(const UPCGPointData* InPointData, const FPCGAttributePropertyInputSelector& InCollisionSelector, EPCGCollisionQueryFlag InCollisionQueryFlag, TArray<FSoftObjectPath>& OutMeshesToLoad)
 {
 	CollisionWrapper.Uninitialize();
 
@@ -224,7 +207,7 @@ bool UPCGCollisionWrapperData::PreInitializeAndGatherMeshesEx(const UPCGPointDat
 #endif
 
 	CollisionSelector = InCollisionSelector;
-	bUseComplexCollision = bInUseComplexCollision;
+	CollisionQueryFlag = InCollisionQueryFlag;
 
 	// Go through the point data on the attribute we'll use for the shapes
 	// StaticMesh -> see what we're doing in UInstanceStaticMeshComponent::InitInstanceBody, or UStaticMeshComponent::UpdateCollisionFromStaticMesh
@@ -242,13 +225,37 @@ bool UPCGCollisionWrapperData::PreInitializeAndGatherMeshesEx(const UPCGPointDat
 	return CollisionWrapper.Prepare(InputAccessor.Get(), InputKeys.Get(), OutMeshesToLoad);
 }
 
-void UPCGCollisionWrapperData::FinalizeInitializationEx(const UPCGPointData* InPointData, const TArray<FSoftObjectPath>& InMeshPaths, bool bInUseAccurateOctree)
+void UPCGCollisionWrapperData::FinalizeInitializationEx(const TArray<FSoftObjectPath>& InMeshPaths)
 {
+	const bool bWasInitialized = CollisionWrapper.bInitialized;
+
 	CollisionWrapper.CreateBodyInstances(InMeshPaths);
 
-	if (InPointData && bInUseAccurateOctree)
+	// Finally, cache the shapes we'll be interested in according to the collision query flag
+	if (!bWasInitialized && CollisionWrapper.bInitialized)
 	{
-		bUseCollisionAccurateOctree = CollisionWrapper.InitializeOctree(InPointData, InMeshPaths, CollisionAccurateOctree);
+		CachedShapes.Reset();
+		CachedShapes.Reserve(CollisionWrapper.BodyInstances.Num());
+
+		for (FBodyInstance* BodyInstance : CollisionWrapper.BodyInstances)
+		{
+			check(BodyInstance);
+			PhysicsInterfaceTypes::FInlineShapeArray& Shapes = CachedShapes.Emplace_GetRef();
+			CollisionWrapper.GetShapeArray(BodyInstance, CollisionQueryFlag, Shapes);
+		}
+	}
+}
+
+const PhysicsInterfaceTypes::FInlineShapeArray& UPCGCollisionWrapperData::GetCachedShapes(int32 EntryIndex) const
+{
+	static PhysicsInterfaceTypes::FInlineShapeArray EmptyArray;
+	if (!CollisionWrapper.bInitialized || CollisionWrapper.IndexToBodyInstance[EntryIndex] == INDEX_NONE)
+	{
+		return EmptyArray;
+	}
+	else
+	{
+		return CachedShapes[CollisionWrapper.IndexToBodyInstance[EntryIndex]];
 	}
 }
 
@@ -258,11 +265,9 @@ void UPCGCollisionWrapperData::AddToCrc(FArchiveCrc32& Ar, bool bFullDataCrc) co
 	GetPointData()->AddToCrc(Ar, bFullDataCrc);
 
 	CollisionSelector.AddToCrc(Ar);
-	uint32 UseCollisionAccurateOctree = bUseCollisionAccurateOctree ? 1 : 0;
-	Ar << UseCollisionAccurateOctree;
 
-	uint32 UseComplexCollision = bUseComplexCollision ? 1 : 0;
-	Ar << UseComplexCollision;
+	uint32 CollisionFlag = static_cast<int32>(CollisionQueryFlag);
+	Ar << CollisionFlag;
 }
 
 void UPCGCollisionWrapperData::GetResourceSizeEx(FResourceSizeEx& CumulativeResourceSize)
@@ -296,7 +301,7 @@ bool UPCGCollisionWrapperData::SamplePoint(const FTransform& InTransform, const
 	// For all touching points, check against the actual physics shape we have for that point
 	// E.g. get the point collision attribute value, match to our internal list, get the shape (might require loading), then test against the matching shapes
 	// note that we need to place the stuff in the right referential too, since our shapes will all be at the origin
-	const UPCGPointData::PointOctree& Octree = (bUseCollisionAccurateOctree ? CollisionAccurateOctree : GetPointData()->GetOctree());
+	const UPCGPointData::PointOctree& Octree = GetPointData()->GetOctree();
 
 	float Density = 0;
 	FBox TransformedBounds = InBounds.TransformBy(InTransform);
@@ -308,21 +313,36 @@ bool UPCGCollisionWrapperData::SamplePoint(const FTransform& InTransform, const
 			return;
 		}
 
-		bool bHasOverlap = true;
+		bool bHasOverlap = false;
 
-		// Compute point index & find body instance
-		// Test against body instance when present
-		if (FBodyInstance* BodyInstance = CollisionWrapper.GetBodyInstance(InPointRef.Point - GetPointData()->GetPoints().GetData()))
-		{
-			// Compute collision shape rotation and transform, knowing that the body instance is at the origin.
-			// e.g. we need to take the InTransform and apply the inverse point transform to it.
-			FTransform RelativeTransform = InTransform.GetRelativeTransform(InPointRef.Point->Transform);
-			FTransform GeomTransform(RelativeTransform.GetRotation(), RelativeTransform.GetLocation());
+		FTransform RelativeTransform = InTransform.GetRelativeTransform(InPointRef.Point->Transform);
+		FCollisionShape CollisionShape;
+		CollisionShape.SetBox(FVector3f(InBounds.GetExtent() * RelativeTransform.GetScale3D()));
+		FPhysicsShapeAdapter CollAdapter(RelativeTransform.GetRotation(), CollisionShape);
+		const FPhysicsGeometry& Geom = CollAdapter.GetGeometry();
+		FTransform GeomTransform = CollAdapter.GetGeomPose(RelativeTransform.GetLocation());
+
+		const int32 EntryIndex = InPointRef.Point - GetPointData()->GetPoints().GetData();
+		const PhysicsInterfaceTypes::FInlineShapeArray& Shapes = GetCachedShapes(EntryIndex);
 
-			FCollisionShape CollisionShape;
-			CollisionShape.SetBox(FVector3f(InBounds.GetExtent() * RelativeTransform.GetScale3D()));
+		if (!Shapes.IsEmpty())
+		{
+			for (const FPhysicsShapeHandle& Shape : Shapes)
+			{
+				if (Chaos::Utilities::CastHelper(Geom, GeomTransform, [&Shape](const auto& Downcast, const auto& FullGeomTransform) { return Chaos::OverlapQuery(*Shape.Shape->GetGeometry(), FTransform::Identity, Downcast, FullGeomTransform); }))
+				{
+					bHasOverlap = true;
+					break;
+				}
+			}
+		}
+		else
+		{
+			FCollisionShape ThisSimpleShape;
+			ThisSimpleShape.SetBox(FVector3f(InPointRef.Bounds.BoxExtent));
+			FPhysicsShapeAdapter ThisAdapter(FQuat::Identity, ThisSimpleShape);
 
-			bHasOverlap = FPhysicsInterface::Overlap_Geom(BodyInstance, CollisionShape, GeomTransform.GetRotation(), GeomTransform, nullptr, bUseComplexCollision);
+			bHasOverlap = Chaos::Utilities::CastHelper(Geom, GeomTransform, [&ThisAdapter](const auto& Downcast, const auto& FullGeomTransform) { return Chaos::OverlapQuery(ThisAdapter.GetGeometry(), FTransform::Identity, Downcast, FullGeomTransform, /*Thickness=*/0); });
 		}
 
 		if (bHasOverlap)
@@ -349,7 +369,7 @@ bool UPCGCollisionWrapperData::SamplePoint(const FTransform& InTransform, const
 UPCGSpatialData* UPCGCollisionWrapperData::CopyInternal(FPCGContext* Context) const
 {
 	UPCGCollisionWrapperData* NewCollisionWrapperData = FPCGContext::NewObject_AnyThread<UPCGCollisionWrapperData>(Context);
-	NewCollisionWrapperData->Initialize(PointData, CollisionSelector, bUseComplexCollision, bUseCollisionAccurateOctree);
+	NewCollisionWrapperData->Initialize(PointData, CollisionSelector, CollisionQueryFlag);
 
 	return NewCollisionWrapperData;
 }
diff --git a/Engine/Plugins/PCG/Source/PCG/Private/Elements/PCGCreateCollisionData.cpp b/Engine/Plugins/PCG/Source/PCG/Private/Elements/PCGCreateCollisionData.cpp
--- a/Engine/Plugins/PCG/Source/PCG/Private/Elements/PCGCreateCollisionData.cpp
+++ b/Engine/Plugins/PCG/Source/PCG/Private/Elements/PCGCreateCollisionData.cpp
@@ -10,6 +10,19 @@
 
 #define LOCTEXT_NAMESPACE "PCGCreateCollisionDataElement"
 
+void UPCGCreateCollisionDataSettings::PostLoad()
+{
+	Super::PostLoad();
+
+#if WITH_EDITOR
+	if (bUseComplexCollision_DEPRECATED)
+	{
+		CollisionQueryFlag = EPCGCollisionQueryFlag::Complex;
+		bUseComplexCollision_DEPRECATED = false;
+	}
+#endif
+}
+
 #if WITH_EDITOR
 FText UPCGCreateCollisionDataSettings::GetNodeTooltipText() const
 {
@@ -72,7 +85,7 @@ bool FPCGCreateCollisionDataElement::PrepareDataInternal(FPCGContext* InContext)
 				InputMeshData.InputIndex = InputIndex;
 				InputMeshData.Data = FPCGContext::NewObject_AnyThread<UPCGCollisionWrapperData>(Context);
 
-				if (!InputMeshData.Data->PreInitializeAndGatherMeshesEx(PointData, Settings->CollisionAttribute, Settings->bUseComplexCollision, InputMeshData.MeshPaths))
+				if (!InputMeshData.Data->PreInitializeAndGatherMeshesEx(PointData, Settings->CollisionAttribute, Settings->CollisionQueryFlag, InputMeshData.MeshPaths))
 				{
 					if (Settings->bWarnIfAttributeCouldNotBeUsed)
 					{
@@ -125,7 +138,7 @@ bool FPCGCreateCollisionDataElement::ExecuteInternal(FPCGContext* InContext) con
 		if (FPCGCreateCollisionContext::InputMeshData* MatchingData = Context->PerInputData.FindByPredicate([InputIndex](const FPCGCreateCollisionContext::InputMeshData& Data) { return Data.InputIndex == InputIndex; }))
 		{
 			check(MatchingData->Data);
-			MatchingData->Data->FinalizeInitializationEx(CastChecked<UPCGPointData>(Input.Data), MatchingData->MeshPaths, Settings->bRecomputeOctreeAccordingToMeshes);
+			MatchingData->Data->FinalizeInitializationEx(MatchingData->MeshPaths);
 
 			// Unroot collision wrapper data object since it'll be tracked through the collection
 			MatchingData->Data->RemoveFromRoot();
diff --git a/Engine/Plugins/PCG/Source/PCG/Private/Elements/PCGSelfPruning.cpp b/Engine/Plugins/PCG/Source/PCG/Private/Elements/PCGSelfPruning.cpp
--- a/Engine/Plugins/PCG/Source/PCG/Private/Elements/PCGSelfPruning.cpp
+++ b/Engine/Plugins/PCG/Source/PCG/Private/Elements/PCGSelfPruning.cpp
@@ -12,6 +12,7 @@
 #include "Metadata/Accessors/PCGAttributeExtractor.h"
 
 #include "CollisionShape.h"
+#include "Chaos/GeometryQueries.h"
 #include "Chaos/ImplicitObject.h"
 #include "PhysicsEngine/BodyInstance.h"
 
@@ -158,17 +159,67 @@ namespace PCGSelfPruningElement
 	}
 
 	/** Self-pruning driven by use of collision shapes. Implementation is in practice just a secondary step after the octree query to filter out points if their collisions don`t intersect. */
-	bool CollisionExclusion(FIterationState& IterationState, FPCGContext* InOptionalContext, bool bUseComplexCollision)
+	bool CollisionExclusion(FIterationState& IterationState, FPCGContext* InOptionalContext, EPCGCollisionQueryFlag InCollisionQueryFlag)
 	{
 		TRACE_CPUPROFILER_EVENT_SCOPE(FPCGSelfPruningElement::Execute::CollisionExclusion);
 
 		check(IterationState.InputPointData);
-		const UPCGPointData::PointOctree& Octree = IterationState.bUseCollisionAccurateOctree ? IterationState.CollisionAccurateOctree : IterationState.InputPointData->GetOctree();
+		const UPCGPointData::PointOctree& Octree = IterationState.InputPointData->GetOctree();
 
 		int32 CheckTimeSlicingCount = 0;
 
 		TArray<const FPCGPointRef*, TInlineAllocator<256>> ElementsToTest;
 
+		auto SetupQueryInfo = [&IterationState](const FPCGPointRef& PointRef, FBodyInstance* OtherBodyInstance, EPCGCollisionQueryFlag CollisionQueryFlag, FBodyInstance*& OutInstance, PhysicsInterfaceTypes::FInlineShapeArray& OutShapes, FCollisionShape& OutSimpleShape)
+		{
+			const int32 EntryIndex = PointRef.Point - IterationState.InputPointData->GetPoints().GetData();
+			OutInstance = IterationState.CollisionWrapper.GetBodyInstance(EntryIndex);
+			bool bHasComplexShapes = false;
+
+			if (OutInstance)
+			{
+				if (OutInstance == OtherBodyInstance)
+				{
+					if (FBodyInstance** TemporaryInstance = IterationState.TemporaryBodyInstances.Find(OtherBodyInstance))
+					{
+						OutInstance = *TemporaryInstance;
+					}
+					else
+					{
+						check(OtherBodyInstance->GetBodySetup());
+						OutInstance = new FBodyInstance();
+						OutInstance->bAutoWeld = false;
+						OutInstance->bSimulatePhysics = false;
+						OutInstance->InitBody(OtherBodyInstance->GetBodySetup(), FTransform::Identity, nullptr, nullptr);
+
+						IterationState.TemporaryBodyInstances.Add(OtherBodyInstance, OutInstance);
+					}
+				}
+
+				OutInstance->UpdateBodyScale(PointRef.Point->Transform.GetScale3D());
+				const bool bFirstChoice = FPCGCollisionWrapper::GetShapeArray(OutInstance, CollisionQueryFlag, OutShapes);
+
+				if (OutShapes.IsEmpty())
+				{
+					OutInstance = nullptr;
+				}
+				else
+				{
+					bHasComplexShapes = (CollisionQueryFlag == EPCGCollisionQueryFlag::Complex) ||
+						(CollisionQueryFlag == EPCGCollisionQueryFlag::ComplexFirst && bFirstChoice) ||
+						(CollisionQueryFlag == EPCGCollisionQueryFlag::SimpleFirst && !bFirstChoice);
+				}
+			}
+
+			if (!OutInstance)
+			{
+				OutShapes.Reset();
+				OutSimpleShape.SetBox(FVector3f(PointRef.Bounds.BoxExtent));
+			}
+
+			return bHasComplexShapes;
+		};
+
 		while (IterationState.CurrentPointIndex < IterationState.SortedPoints.Num())
 		{
 			// Don't check too many times. Pre-increment will make sure we always at least process "TimeSliceFrequencyCheck" elements at each iteration.
@@ -208,98 +259,102 @@ namespace PCGSelfPruningElement
 			}
 
 			// Implementation note: this is a deconstruction of FBodyInstance::OverlapTestForBodiesImpl
-			// Get this point's body instance, if any.
-			FBodyInstance* ThisInstance = IterationState.CollisionWrapper.GetBodyInstance(PointRef.Point - IterationState.InputPointData->GetPoints().GetData());
-			FCollisionShape ThisCollisionShape;
-			FTransform TransformNoScale = FTransform(PointRef.Point->Transform.GetRotation(), PointRef.Point->Transform.GetLocation());
+			FBodyInstance* ThisInstance = nullptr;
+			PhysicsInterfaceTypes::FInlineShapeArray ThisShapes;
+			FCollisionShape ThisSimpleShape;
 
-			PhysicsInterfaceTypes::FInlineShapeArray TargetShapes;
+			FBodyInstance* OtherInstance = nullptr;
+			PhysicsInterfaceTypes::FInlineShapeArray OtherShapes;
+			FCollisionShape OtherSimpleShape;
 
-			if (ThisInstance)
-			{
-				// Since scale can change on a per-point basis and we can't fold it in the relative transform, we must apply scale here.
-				ThisInstance->UpdateBodyScale(PointRef.Point->Transform.GetScale3D());
-				// Implementation note: all of these bodies don't exist in the physics scene, so there's no need to lock
-				FillInlineShapeArray_AssumesLocked(TargetShapes, ThisInstance->GetPhysicsActorHandle());
-			}
-			else
-			{
-				ThisCollisionShape.SetBox(FVector3f(PointRef.Bounds.BoxExtent * PointRef.Point->Transform.GetScale3D()));
-			}
+			const bool bThisHasComplexShapes = SetupQueryInfo(PointRef, nullptr, InCollisionQueryFlag, ThisInstance, ThisShapes, ThisSimpleShape);
+			FTransform TransformNoScale = FTransform(PointRef.Point->Transform.GetRotation(), PointRef.Point->Transform.GetLocation());
+
+			// We must force the other collision flag to simple if we have complex shapes in the leading shape here, because complex-complex overlaps aren't supported.
+			EPCGCollisionQueryFlag OtherCollisionQueryFlag = (bThisHasComplexShapes ? EPCGCollisionQueryFlag::Simple : InCollisionQueryFlag);
 
 			for (const FPCGPointRef* OtherPointRef : ElementsToTest)
 			{
-				FBodyInstance* OtherInstance = IterationState.CollisionWrapper.GetBodyInstance(OtherPointRef->Point - IterationState.InputPointData->GetPoints().GetData());
+				bool bOverlaps = false;
+
+				bool bOtherHasComplexShapes = SetupQueryInfo(*OtherPointRef, ThisInstance, OtherCollisionQueryFlag, OtherInstance, OtherShapes, OtherSimpleShape);
 
-				bool bOverlaps = true;
-				if (ThisInstance && OtherInstance)
+				FTransform OtherTransformNoScale = FTransform(OtherPointRef->Point->Transform.GetRotation(), OtherPointRef->Point->Transform.GetLocation());
+
+				// Four cases here:
+				// 1 - shapes vs shapes
+				if (!ThisShapes.IsEmpty() && !OtherShapes.IsEmpty())
 				{
-					// If both bodies are the same, then we can't really have them having different scales; we need to do a (temporary) copy when that happens
-					if (ThisInstance == OtherInstance && !OtherPointRef->Point->Transform.GetScale3D().Equals(PointRef.Point->Transform.GetScale3D()))
+					auto CheckForOverlap = [](const PhysicsInterfaceTypes::FInlineShapeArray& ComplexShapes, const FTransform& ComplexShapesTransform, const PhysicsInterfaceTypes::FInlineShapeArray& SimpleShapes, const FTransform& SimpleShapesTransform)
 					{
-						if (FBodyInstance** TemporaryInstance = IterationState.TemporaryBodyInstances.Find(OtherInstance))
-						{
-							OtherInstance = *TemporaryInstance;
-						}
-						else
+						FTransform RelativeTransform = SimpleShapesTransform.GetRelativeTransform(ComplexShapesTransform);
+
+						for (const FPhysicsShapeHandle& ComplexShape : ComplexShapes)
 						{
-							check(ThisInstance->GetBodySetup());
-							OtherInstance = new FBodyInstance();
-							OtherInstance->bAutoWeld = false;
-							OtherInstance->bSimulatePhysics = false;
-							OtherInstance->InitBody(ThisInstance->GetBodySetup(), FTransform::Identity, nullptr, nullptr);
+							for (const FPhysicsShapeHandle& SimpleShape : SimpleShapes)
+							{
+								FPhysicsGeometryCollection SimpleShapeCollection = FPhysicsInterface::GetGeometryCollection(SimpleShape);
+								const Chaos::FImplicitObject& SimpleShapeGeom = SimpleShapeCollection.GetGeometry();
 
-							IterationState.TemporaryBodyInstances.Add(ThisInstance, OtherInstance);
+								if (Chaos::Utilities::CastHelper(SimpleShapeGeom, RelativeTransform, [&ComplexShape](const auto& Downcast, const auto& FullGeomTransform) { return Chaos::OverlapQuery(*ComplexShape.Shape->GetGeometry(), FTransform::Identity, Downcast, FullGeomTransform); }))
+								{
+									return true;
+								}
+							}
 						}
-					}
 
-					if (ThisInstance != OtherInstance)
+						return false;
+					};
+
+					if (!bOtherHasComplexShapes)
 					{
-						OtherInstance->UpdateBodyScale(OtherPointRef->Point->Transform.GetScale3D());
+						bOverlaps = CheckForOverlap(ThisShapes, TransformNoScale, OtherShapes, OtherTransformNoScale);
 					}
-
-					FTransform OtherTransformNoScale = FTransform(OtherPointRef->Point->Transform.GetRotation(), OtherPointRef->Point->Transform.GetLocation());
-					FTransform RelativeTransform = TransformNoScale.GetRelativeTransform(OtherTransformNoScale);
-
+					else
 					{
-						// Copied more or less verbatim from FBodyInstance::OverlapTestForBodiesImpl
-						//TRACE_CPUPROFILER_EVENT_SCOPE(FPCGSelfPruningElement::Execute::CollisionExclusion::OverlapTest);
-						bOverlaps = false;
-						for (const FPhysicsShapeHandle& Shape : TargetShapes)
-						{
-							if (!Shape.GetGeometry().IsConvex())
-							{
-								continue; //we skip complex shapes - should this respect ComplexAsSimple?
-							}
+						bOverlaps = CheckForOverlap(OtherShapes, OtherTransformNoScale, ThisShapes, TransformNoScale);
+					}
+				}
+				// 2 - shapes vs simple shape
+				// 3 - simple shape vs shapes
+				else if (!ThisShapes.IsEmpty() || !OtherShapes.IsEmpty())
+				{
+					auto CheckForOverlap = [](const PhysicsInterfaceTypes::FInlineShapeArray& Shapes, const FTransform& ShapesTransform, const FCollisionShape& CollShape, const FTransform& CollTransform)
+					{
+						FTransform RelativeTransform = CollTransform.GetRelativeTransform(ShapesTransform);
+						FPhysicsShapeAdapter CollAdapter(RelativeTransform.GetRotation(), CollShape);
+						const FPhysicsGeometry& Geom = CollAdapter.GetGeometry();
+						FTransform GeomTransform = CollAdapter.GetGeomPose(RelativeTransform.GetLocation());
 
-							bOverlaps = FPhysicsInterface::Overlap_Geom(OtherInstance, FPhysicsInterface::GetGeometryCollection(Shape), RelativeTransform, /*bOutMTD=*/nullptr, bUseComplexCollision);
-							if (bOverlaps)
+						for (const FPhysicsShapeHandle& Shape : Shapes)
+						{
+							if (Chaos::Utilities::CastHelper(Geom, GeomTransform, [&Shape](const auto& Downcast, const auto& FullGeomTransform) { return Chaos::OverlapQuery(*Shape.Shape->GetGeometry(), FTransform::Identity, Downcast, FullGeomTransform); }))
 							{
-								break;
+								return true;
 							}
 						}
-					}
-				}
-				else if (ThisInstance)
-				{
-					FTransform RelativeTransform = TransformNoScale.GetRelativeTransform(OtherPointRef->Point->Transform);
-					FTransform GeomTransform(RelativeTransform.GetRotation(), RelativeTransform.GetLocation());
 
-					// Create collision shape for other instance
-					FCollisionShape OtherCollisionShape;
-					OtherCollisionShape.SetBox(FVector3f(OtherPointRef->Bounds.BoxExtent * RelativeTransform.GetScale3D()));
+						return false;
+					};
 
-					bOverlaps = FPhysicsInterface::Overlap_Geom(ThisInstance, OtherCollisionShape, GeomTransform.GetRotation(), GeomTransform, nullptr, bUseComplexCollision);
+					if(!ThisShapes.IsEmpty())
+					{
+						bOverlaps = CheckForOverlap(ThisShapes, TransformNoScale, OtherSimpleShape, OtherTransformNoScale);
+					}
+					else
+					{
+						check(!OtherShapes.IsEmpty());
+						bOverlaps = CheckForOverlap(OtherShapes, OtherTransformNoScale, ThisSimpleShape, TransformNoScale);
+					}
 				}
-				else if (OtherInstance)
+				// 4 - simple shape vs simple shape <- default case.
+				else
 				{
-					OtherInstance->UpdateBodyScale(OtherPointRef->Point->Transform.GetScale3D());
-
-					FTransform OtherTransformNoScale = FTransform(OtherPointRef->Point->Transform.GetRotation(), OtherPointRef->Point->Transform.GetLocation());
 					FTransform RelativeTransform = OtherTransformNoScale.GetRelativeTransform(TransformNoScale);
-					FTransform GeomTransform(RelativeTransform.GetRotation(), RelativeTransform.GetLocation());
+					FPhysicsShapeAdapter ThisAdapter(FQuat::Identity, ThisSimpleShape);
+					FPhysicsShapeAdapter OtherAdapter(RelativeTransform.GetRotation(), OtherSimpleShape);
 
-					bOverlaps = FPhysicsInterface::Overlap_Geom(OtherInstance, ThisCollisionShape, GeomTransform.GetRotation(), GeomTransform, nullptr, bUseComplexCollision);
+					bOverlaps = Chaos::Utilities::CastHelper(OtherAdapter.GetGeometry(), OtherAdapter.GetGeomPose(RelativeTransform.GetTranslation()), [&ThisAdapter](const auto& Downcast, const auto& FullGeomTransform) { return Chaos::OverlapQuery(ThisAdapter.GetGeometry(), ThisAdapter.GetGeomPose(FVector::ZeroVector), Downcast, FullGeomTransform, /*Thickness=*/0); });
 				}
 
 				if (bOverlaps)
@@ -446,16 +501,11 @@ namespace PCGSelfPruningElement
 		if (!InState.bSortDone)
 		{
 			// In the case of the collision-driven self-pruning, we have to populate the sorted points array earlier since we're playing with the bounds.
-			if (!InState.bSortedPointsArrayPopulateDone)
+			InState.SortedPoints.Empty();
+			InState.SortedPoints.Reserve(Points.Num());
+			for (const FPCGPoint& Point : Points)
 			{
-				InState.SortedPoints.Empty();
-				InState.SortedPoints.Reserve(Points.Num());
-				for (const FPCGPoint& Point : Points)
-				{
-					InState.SortedPoints.Emplace(Point);
-				}
-
-				InState.bSortedPointsArrayPopulateDone = true;
+				InState.SortedPoints.Emplace(Point);
 			}
 
 			FPCGAttributePropertySelector ComparisonSource = InParameters.ComparisonSource.CopyAndFixLast(InState.InputPointData);
@@ -588,7 +638,7 @@ namespace PCGSelfPruningElement
 		}
 		else if (InParameters.bUseCollisionAttribute)
 		{
-			bIsDone = PCGSelfPruningElement::CollisionExclusion(InState, InOptionalContext, InParameters.bUseComplexCollision);
+			bIsDone = PCGSelfPruningElement::CollisionExclusion(InState, InOptionalContext, InParameters.CollisionQueryFlag);
 		}
 		else
 		{
@@ -629,6 +679,17 @@ namespace PCGSelfPruningElement
 	}
 }
 
+void FPCGSelfPruningParameters::PostLoad()
+{
+#if WITH_EDITOR
+	if (bUseComplexCollision_DEPRECATED)
+	{
+		CollisionQueryFlag = EPCGCollisionQueryFlag::Complex;
+		bUseComplexCollision_DEPRECATED = false;
+	}
+#endif
+}
+
 UPCGSelfPruningSettings::UPCGSelfPruningSettings()
 {
 	// Previous Default behavior was Extents
@@ -639,6 +700,8 @@ void UPCGSelfPruningSettings::PostLoad()
 {
 	Super::PostLoad();
 
+	Parameters.PostLoad();
+
 #if WITH_EDITOR
 	if (PruningType_DEPRECATED != EPCGSelfPruningType::LargeToSmall)
 	{
@@ -709,12 +772,6 @@ bool FPCGSelfPruningElement::PrepareDataInternal(FPCGContext* InContext) const
 			if (OutState.CollisionWrapper.Prepare(InputAccessor.Get(), InputKeys.Get(), Meshes))
 			{
 				OutState.CollisionWrapper.CreateBodyInstances(Meshes);
-
-				if (Settings->Parameters.bRecomputeOctreeAccordingToMeshes)
-				{
-					OutState.bUseCollisionAccurateOctree = OutState.CollisionWrapper.InitializeOctree(OutState.InputPointData, Meshes, OutState.CollisionAccurateOctree, &OutState.SortedPoints);
-					OutState.bSortedPointsArrayPopulateDone = OutState.bUseCollisionAccurateOctree;
-				}
 			}
 		}
 
diff --git a/Engine/Plugins/PCG/Source/PCG/Public/Data/PCGCollisionWrapperData.h b/Engine/Plugins/PCG/Source/PCG/Public/Data/PCGCollisionWrapperData.h
--- a/Engine/Plugins/PCG/Source/PCG/Public/Data/PCGCollisionWrapperData.h
+++ b/Engine/Plugins/PCG/Source/PCG/Public/Data/PCGCollisionWrapperData.h
@@ -7,10 +7,21 @@
 #include "Metadata/PCGAttributePropertySelector.h"
 #include "Metadata/Accessors/PCGAttributeAccessorHelpers.h"
 
+#include "Chaos/ChaosEngineInterface.h"
+
 #include "PCGCollisionWrapperData.generated.h"
 
 struct FBodyInstance;
 
+UENUM(BlueprintType)
+enum class EPCGCollisionQueryFlag : uint8
+{
+	Simple,
+	Complex,
+	SimpleFirst,
+	ComplexFirst
+};
+
 struct PCG_API FPCGCollisionWrapper
 {
 	FPCGCollisionWrapper() = default;
@@ -27,11 +38,16 @@ struct PCG_API FPCGCollisionWrapper
 	// Advanced API - allows to do async loading as we separate the mesh finding part from the body creation part
 	bool Prepare(const IPCGAttributeAccessor* Accessor, const IPCGAttributeAccessorKeys* Keys, TArray<FSoftObjectPath>& MeshPathsToLoad);
 	void CreateBodyInstances(const TArray<FSoftObjectPath>& MeshPaths);
-	bool InitializeOctree(const UPCGPointData* InPointData, const TArray<FSoftObjectPath>& InMeshPaths, UPCGPointData::PointOctree& OutOctree, TArray<FPCGPointRef>* OutOctreePointRefs = nullptr) const;
 	
 	// Retrieves the body instance associated to the entry given by its index
 	FBodyInstance* GetBodyInstance(int32 EntryIndex) const;
 
+	// Retrieves the shape list for a given entry matching the query flag
+	void GetShapeArray(int32 EntryIndex, EPCGCollisionQueryFlag QueryFlag, PhysicsInterfaceTypes::FInlineShapeArray& OutShapeArray) const;
+
+	// Retrieves the shape list for a given body, matching the query flag. Returns false if we selected the other type for the 'SimpleFirst' or 'ComplexFirst' cases.
+	static bool GetShapeArray(FBodyInstance* BodyInstance, EPCGCollisionQueryFlag QueryFlag, PhysicsInterfaceTypes::FInlineShapeArray& OutShapeArray);
+
 	TArray<FBodyInstance*> BodyInstances;
 	TArray<int32> IndexToBodyInstance;
 	bool bInitialized = false;
@@ -44,11 +60,11 @@ class UPCGCollisionWrapperData : public UPCGSpatialData
 
 public:
 	/** Inititializes the collision wrapper on a point data based on the provided attribute selector */
-	PCG_API bool Initialize(const UPCGPointData* InPointData, const FPCGAttributePropertyInputSelector& InCollisionSelector, bool bInUseComplexCollision, bool bUseAccurateOctree);
+	PCG_API bool Initialize(const UPCGPointData* InPointData, const FPCGAttributePropertyInputSelector& InCollisionSelector, EPCGCollisionQueryFlag InCollisionQueryFlag);
 
 	/** Advanced API for async loading */
-	bool PreInitializeAndGatherMeshesEx(const UPCGPointData* InPointData, const FPCGAttributePropertyInputSelector& InCollisionSelector, bool bInUseComplexCollision, TArray<FSoftObjectPath>& OutMeshesToLoad);
-	void FinalizeInitializationEx(const UPCGPointData* InPointData, const TArray<FSoftObjectPath>& InMeshPaths, bool bUseAccurateOctree);
+	bool PreInitializeAndGatherMeshesEx(const UPCGPointData* InPointData, const FPCGAttributePropertyInputSelector& InCollisionSelector, EPCGCollisionQueryFlag InCollisionQueryFlag, TArray<FSoftObjectPath>& OutMeshesToLoad);
+	void FinalizeInitializationEx(const TArray<FSoftObjectPath>& InMeshPaths);
 
 	// ~Begin UPCGData Interface
 	virtual EPCGDataType GetDataType() const override { return EPCGDataType::Primitive; }
@@ -74,6 +90,8 @@ protected:
 #endif
 
 private:
+	const PhysicsInterfaceTypes::FInlineShapeArray& GetCachedShapes(int32 EntryIndex) const;
+
 	UPROPERTY()
 	TObjectPtr<const UPCGPointData> PointData;
 
@@ -81,15 +99,12 @@ private:
 	UPROPERTY()
 	FPCGAttributePropertyInputSelector CollisionSelector;
 
-	/** Uses a new octree based on the mesh bounds, performance warning (does similar work to the BoundsFromMesh node, but will not change the point data). */
 	UPROPERTY()
-	bool bUseCollisionAccurateOctree = false;
-
-	UPROPERTY()
-	bool bUseComplexCollision = false;
+	EPCGCollisionQueryFlag CollisionQueryFlag = EPCGCollisionQueryFlag::Simple;
 
 	FPCGCollisionWrapper CollisionWrapper;
-	UPCGPointData::PointOctree CollisionAccurateOctree;
+
+	TArray<PhysicsInterfaceTypes::FInlineShapeArray> CachedShapes;
 
 #if WITH_EDITOR
 	const UPCGPointData* RawPointData = nullptr;
diff --git a/Engine/Plugins/PCG/Source/PCG/Public/Elements/PCGCreateCollisionData.h b/Engine/Plugins/PCG/Source/PCG/Public/Elements/PCGCreateCollisionData.h
--- a/Engine/Plugins/PCG/Source/PCG/Public/Elements/PCGCreateCollisionData.h
+++ b/Engine/Plugins/PCG/Source/PCG/Public/Elements/PCGCreateCollisionData.h
@@ -15,6 +15,10 @@ class UPCGCreateCollisionDataSettings : public UPCGSettings
 	GENERATED_BODY()
 
 public:
+	//~Begin UObject interface
+	virtual void PostLoad() override;
+	//~End UObject interface
+
 	//~Begin UPCGSettings interface
 #if WITH_EDITOR
 	virtual FName GetDefaultNodeName() const override { return FName(TEXT("CreateCollisionData")); }
@@ -33,19 +37,22 @@ public:
 	UPROPERTY(BlueprintReadWrite, EditAnywhere, Category = Settings, meta = (PCG_Overridable, PCG_DiscardPropertySelection, PCG_DiscardExtraSelection))
 	FPCGAttributePropertyInputSelector CollisionAttribute;
 
-	/** Uses a new octree based on the mesh bounds, performance warning (does similar work to the BoundsFromMesh node, but will not change the point data). */
-	UPROPERTY(BlueprintReadWrite, EditAnywhere, Category = Settings, meta = (PCG_Overridable))
-	bool bRecomputeOctreeAccordingToMeshes = false;
-
-	/** Queries against complex collision if enabled, performance warning */
-	UPROPERTY(BlueprintReadWrite, EditAnywhere, Category = Settings, meta = (PCG_Overridable))
-	bool bUseComplexCollision = false;
+	/** Controls how shapes are selected from collision. Performance warning on using complex shapes. */
+	UPROPERTY(BlueprintReadWrite, EditAnywhere, Category = Settings, meta = (PCG_Overridable, PCG_OverrideAliases="bUseComplexCollision"))
+	EPCGCollisionQueryFlag CollisionQueryFlag = EPCGCollisionQueryFlag::Simple;
 
 	UPROPERTY(BlueprintReadWrite, EditAnywhere, Category = "Settings|Advanced")
 	bool bWarnIfAttributeCouldNotBeUsed = true;
 
 	UPROPERTY(BlueprintReadWrite, EditAnywhere, Category = "Settings|Debug")
 	bool bSynchronousLoad = false;
+
+private:
+#if WITH_EDITORONLY_DATA
+	// Implementation note: was introduced during UE 5.5 development and replaced, does not require a full public API deprecation mechanism
+	UPROPERTY()
+	bool bUseComplexCollision_DEPRECATED = false;
+#endif
 };
 
 struct FPCGCreateCollisionContext : public FPCGContext, public IPCGAsyncLoadingContext
diff --git a/Engine/Plugins/PCG/Source/PCG/Public/Elements/PCGSelfPruning.h b/Engine/Plugins/PCG/Source/PCG/Public/Elements/PCGSelfPruning.h
--- a/Engine/Plugins/PCG/Source/PCG/Public/Elements/PCGSelfPruning.h
+++ b/Engine/Plugins/PCG/Source/PCG/Public/Elements/PCGSelfPruning.h
@@ -28,6 +28,8 @@ struct FPCGSelfPruningParameters
 	GENERATED_BODY()
 
 public:
+	void PostLoad();
+
 	UPROPERTY(BlueprintReadWrite, EditAnywhere, Category = Settings, meta = (PCG_Overridable))
 	EPCGSelfPruningType PruningType = EPCGSelfPruningType::LargeToSmall;
 
@@ -49,13 +51,15 @@ public:
 	UPROPERTY(BlueprintReadWrite, EditAnywhere, Category = Settings, meta = (PCG_Overridable, PCG_DiscardPropertySelection, PCG_DiscardExtraSelection, EditCondition = "bUseCollisionAttribute && PruningType != EPCGSelfPruningType::RemoveDuplicates", EditConditionHides))
 	FPCGAttributePropertyInputSelector CollisionAttribute;
 
-	/** Uses a new octree based on the mesh bounds, performance warning (does similar work to the BoundsFromMesh node, but will not change the point data). */
-	UPROPERTY(BlueprintReadWrite, EditAnywhere, Category = Settings, meta = (PCG_Overridable, EditCondition = "bUseCollisionAttribute && PruningType != EPCGSelfPruningType::RemoveDuplicates", EditConditionHides))
-	bool bRecomputeOctreeAccordingToMeshes = false;
+	/** Controls whether queries will be done against complex collisions or not. If enabled, performance warning. */
+	UPROPERTY(BlueprintReadWrite, EditAnywhere, Category = Settings, meta = (PCG_Overridable, PCG_OverrideAliases="bUseComplexCollision", EditCondition = "bUseCollisionAttribute && PruningType != EPCGSelfPruningType::RemoveDuplicates", EditConditionHides))
+	EPCGCollisionQueryFlag CollisionQueryFlag = EPCGCollisionQueryFlag::Simple;
 
-	/** Queries against complex collision (vs. collection of convex elements in the collision) if enabled, performance warning */
-	UPROPERTY(BlueprintReadWrite, EditAnywhere, Category = Settings, meta = (PCG_Overridable, EditCondition = "bUseCollisionAttribute && PruningType != EPCGSelfPruningType::RemoveDuplicates", EditConditionHides))
-	bool bUseComplexCollision = false;
+#if WITH_EDITORONLY_DATA
+	// Implementation note: was introduced during UE 5.5 development and replaced, does not require a full public API deprecation mechanism
+	UPROPERTY()
+	bool bUseComplexCollision_DEPRECATED = false;
+#endif
 };
 
 namespace PCGSelfPruningElement
@@ -81,13 +85,10 @@ namespace PCGSelfPruningElement
 		FPointBitSet ExclusionPoints;
 		int32 CurrentPointIndex = 0;
 		bool bSortDone = false;
-		bool bSortedPointsArrayPopulateDone = false;
-		bool bUseCollisionAccurateOctree = false;
 
 		// In the case of the self-pruning using the collision, we'll need instances of the unique bodies that are used in the point data.
 		// Additionally, in the cases where two points to compare use the same instance, we'll have to do a temporary copy.
 		FPCGCollisionWrapper CollisionWrapper;
-		UPCGPointData::PointOctree CollisionAccurateOctree;
 		TMap<FBodyInstance*, FBodyInstance*> TemporaryBodyInstances;
 	};
 
