[PCG] Fix extra frame taken by async loading when all objects are loaded

* All already loaded objects are now gathered as strong ptr to keep a reference on them
* Remove the condition of having all the paths not null, as they are now always checked (and removed if needed)

#rnx

diff --git a/Engine/Plugins/PCG/Source/PCG/Private/Async/PCGAsyncLoadingContext.cpp b/Engine/Plugins/PCG/Source/PCG/Private/Async/PCGAsyncLoadingContext.cpp
--- a/Engine/Plugins/PCG/Source/PCG/Private/Async/PCGAsyncLoadingContext.cpp
+++ b/Engine/Plugins/PCG/Source/PCG/Private/Async/PCGAsyncLoadingContext.cpp
@@ -22,13 +22,38 @@ void IPCGAsyncLoadingContext::CancelLoading()
 	LoadHandle.Reset();
 }
 
+void IPCGAsyncLoadingContext::ResolveAlreadyLoadedObjects(TArray<FSoftObjectPath>& ObjectsToLoad)
+{
+	for (int32 i = ObjectsToLoad.Num() - 1; i >= 0; --i)
+	{
+		UObject* ResolvedObject = ObjectsToLoad[i].ResolveObject();
+		if (ResolvedObject || ObjectsToLoad[i].IsNull())
+		{
+			ObjectsToLoad.RemoveAtSwap(i);
+			if (ResolvedObject)
+			{
+				AlreadyLoadedObjects.Emplace(ResolvedObject);
+			}
+		}
+	}
+}
+
 bool IPCGAsyncLoadingContext::RequestResourceLoad(FPCGContext* ThisContext, TArray<FSoftObjectPath>&& ObjectsToLoad, bool bAsynchronous)
 {
 	if (!ObjectsToLoad.IsEmpty() && !bLoadRequested)
 	{
+		ResolveAlreadyLoadedObjects(ObjectsToLoad);
+
+		// If there is nothing to load, early out.
+		if (ObjectsToLoad.IsEmpty())
+		{
+			bLoadRequested = true;
+			return true;
+		}
+
 		if (!bAsynchronous)
 		{
-			LoadHandle = UAssetManager::GetStreamableManager().RequestSyncLoad(std::forward<TArray<FSoftObjectPath>>(ObjectsToLoad));
+			LoadHandle = UAssetManager::GetStreamableManager().RequestSyncLoad(std::move(ObjectsToLoad));
 			bLoadRequested = true;
 
 			return true;
@@ -38,7 +63,7 @@ bool IPCGAsyncLoadingContext::RequestResourceLoad(FPCGContext* ThisContext, TArr
 			ThisContext->bIsPaused = true;
 
 			// It is a bit unsafe to pass this to a delegate lambda. But if the context dies before the completion of the loading, the context will cancel the loading in its dtor.
-			LoadHandle = UAssetManager::GetStreamableManager().RequestAsyncLoad(std::forward<TArray<FSoftObjectPath>>(ObjectsToLoad), [ThisContext]() { ThisContext->bIsPaused = false; });
+			LoadHandle = UAssetManager::GetStreamableManager().RequestAsyncLoad(std::move(ObjectsToLoad), [ThisContext]() { ThisContext->bIsPaused = false; });
 			bLoadRequested = true;
 
 			// If the load handle is not active it means objects were invalid
diff --git a/Engine/Plugins/PCG/Source/PCG/Public/Async/PCGAsyncLoadingContext.h b/Engine/Plugins/PCG/Source/PCG/Public/Async/PCGAsyncLoadingContext.h
--- a/Engine/Plugins/PCG/Source/PCG/Public/Async/PCGAsyncLoadingContext.h
+++ b/Engine/Plugins/PCG/Source/PCG/Public/Async/PCGAsyncLoadingContext.h
@@ -19,7 +19,6 @@ public:
 
 	/** Request a load. If load was already requested, do nothing. LoadHandle will be set in the context, meaning that assets will stay alive while context is loaded.
 	* Request can be synchronous or asynchronous. If loading is asynchronous, the current task is paused and will be woken up when the loading is done.
-	* WARNING: Make sure to call this function with soft paths that are NOT null.
 	* Returns true if the execution can continue (objects are loaded or invalid), or false if we need to wait for loading
 	*/
 	bool RequestResourceLoad(FPCGContext* ThisContext, TArray<FSoftObjectPath>&& ObjectsToLoad, bool bAsynchronous = true);
@@ -29,9 +28,18 @@ public:
 	bool WasLoadRequested() const { return bLoadRequested; }
 
 private:
+	/** Resolve already loaded objects and store a reference on them and also remove null paths, will modify in place the array of objects that actually require loading. */
+	void ResolveAlreadyLoadedObjects(TArray<FSoftObjectPath>& ObjectsToLoad);
+
 	/** If the load was already requested */
 	bool bLoadRequested = false;
 
 	/** Handle holder for any loaded resources */
 	TSharedPtr<FStreamableHandle> LoadHandle;
+
+	/** 
+	* Keep a reference on objects already loaded, to avoid starting a loading task if everything is already loaded.
+	* It's a strong object ptr to make sure it doesn't get GC before the end of the execution of the PCG Element.
+	*/
+	TArray<TStrongObjectPtr<UObject>> AlreadyLoadedObjects;
 };
\ No newline at end of file
