Moved asset exporter utils in library to make it easier to use in PCG elements.

#rnx

diff --git a/Engine/Plugins/PCG/Source/PCG/Private/PCGAssetExporter.cpp b/Engine/Plugins/PCG/Source/PCG/Private/PCGAssetExporter.cpp
--- a/Engine/Plugins/PCG/Source/PCG/Private/PCGAssetExporter.cpp
+++ b/Engine/Plugins/PCG/Source/PCG/Private/PCGAssetExporter.cpp
@@ -71,6 +71,7 @@ bool UPCGAssetExporter::Export(const FString& PackageName, UPCGDataAsset* Asset)
 		Asset->ExporterClass = GetClass();
 		if (ExportAsset(PackageName, Asset))
 		{
+			DuplicateAndReOuterData(Asset);
 			SerializeMetadataToAsset(Asset);
 			Asset->MarkPackageDirty();
 			return true;
@@ -142,4 +143,25 @@ void UPCGAssetExporter::SerializeMetadataToAsset(const FAssetData& PCGAsset)
 
 	Asset->ExporterMetadata = FBase64::Encode(SerializedMetadata);
 #endif // WITH_EDITOR
+}
+
+void UPCGAssetExporter::DuplicateAndReOuterData(UPCGDataAsset* Asset)
+{
+	if (!Asset)
+	{
+		return;
+	}
+
+	// Implementation note: we can't easily update the contents of data even through visit data
+	// because it doesn't allow changing its parts.
+	for (FPCGTaggedData& TaggedData : Asset->Data.TaggedData)
+	{
+		if (TaggedData.Data && TaggedData.Data->GetOuter() == GetTransientPackage())
+		{
+			UPCGData* DuplicatedData = TaggedData.Data->DuplicateData(nullptr);
+			DuplicatedData->Rename(nullptr, Asset);
+
+			TaggedData.Data = DuplicatedData;
+		}
+	}
 }
\ No newline at end of file
diff --git a/Engine/Plugins/PCG/Source/PCG/Private/PCGAssetExporterUtils.cpp b/Engine/Plugins/PCG/Source/PCG/Private/PCGAssetExporterUtils.cpp
new file mode 100644
--- /dev/null
+++ b/Engine/Plugins/PCG/Source/PCG/Private/PCGAssetExporterUtils.cpp
@@ -0,0 +1,182 @@
+// Copyright Epic Games, Inc. All Rights Reserved.
+
+#include "PCGAssetExporterUtils.h"
+
+#include "PCGModule.h"
+
+#if WITH_EDITOR
+#include "ContentBrowserModule.h"
+#include "FileHelpers.h"
+#include "IContentBrowserSingleton.h"
+#include "AssetRegistry/AssetRegistryModule.h"
+#endif // WITH_EDITOR
+
+UPackage* UPCGAssetExporterUtils::CreateAsset(UPCGAssetExporter* Exporter, FPCGAssetExporterParameters Parameters)
+{
+#if WITH_EDITOR
+	if (!Exporter || !Exporter->GetAssetType())
+	{
+		UE_LOG(LogPCG, Error, TEXT("Unable to create asset without an exporter, or exporter is not setup properly."));
+		return nullptr;
+	}
+
+	const bool bExporterNeedsRooting = !Exporter->IsRooted();
+	if (bExporterNeedsRooting)
+	{
+		Exporter->AddToRoot();
+	}
+
+	ON_SCOPE_EXIT
+	{
+		if (bExporterNeedsRooting)
+		{
+			Exporter->RemoveFromRoot();
+		}
+	};
+
+	FString AssetName = Parameters.AssetName;
+	FString AssetPath = Parameters.AssetPath;
+	FString PackageName = FPaths::Combine(AssetPath, AssetName);
+
+	if (Parameters.bOpenSaveDialog)
+	{
+		FSaveAssetDialogConfig SaveAssetDialogConfig;
+		SaveAssetDialogConfig.DefaultPath = AssetPath;
+		SaveAssetDialogConfig.DefaultAssetName = AssetName;
+		SaveAssetDialogConfig.AssetClassNames.Add(Exporter->GetAssetType()->GetClassPathName());
+		SaveAssetDialogConfig.ExistingAssetPolicy = ESaveAssetDialogExistingAssetPolicy::AllowButWarn;
+		SaveAssetDialogConfig.DialogTitleOverride = NSLOCTEXT("PCGAssetExporter", "SaveAssetToFileDialogTitle", "Save PCG Asset");
+
+		FContentBrowserModule& ContentBrowserModule = FModuleManager::LoadModuleChecked<FContentBrowserModule>("ContentBrowser");
+		FString SaveObjectPath = ContentBrowserModule.Get().CreateModalSaveAssetDialog(SaveAssetDialogConfig);
+		if (!SaveObjectPath.IsEmpty())
+		{
+			AssetName = FPackageName::ObjectPathToObjectName(SaveObjectPath);
+			AssetPath = FString(); // not going to be reused
+			PackageName = FPackageName::ObjectPathToPackageName(SaveObjectPath);
+		}
+		else
+		{
+			return nullptr;
+		}
+	}
+
+	UPackage* Package = FPackageName::DoesPackageExist(PackageName) ? LoadPackage(nullptr, *PackageName, LOAD_None) : nullptr;
+
+	UPCGDataAsset* Asset = nullptr;
+	bool NewAssetCreated = false;
+
+	if (Package)
+	{
+		UObject* Object = FindObjectFast<UObject>(Package, *AssetName);
+		if (Object && Object->GetClass() != Exporter->GetAssetType())
+		{
+			Object->SetFlags(RF_Transient);
+			Object->Rename(nullptr, GetTransientPackage(), REN_DontCreateRedirectors | REN_NonTransactional);
+			NewAssetCreated = true;
+		}
+		else
+		{
+			Asset = Cast<UPCGDataAsset>(Object);
+		}
+	}
+	else
+	{
+		Package = CreatePackage(*PackageName);
+		NewAssetCreated = true;
+	}
+
+	if (!Asset)
+	{
+		const EObjectFlags Flags = RF_Public | RF_Standalone | RF_Transactional;
+		Asset = NewObject<UPCGDataAsset>(Package, Exporter->GetAssetType(), FName(*AssetName), Flags);
+	}
+
+	if (Asset)
+	{
+		if (NewAssetCreated)
+		{
+			// Notify the asset registry
+			FAssetRegistryModule::AssetCreated(Asset);
+		}
+
+		if (Exporter->Export(PackageName, Asset))
+		{
+			// Make sure everybody knows we changed the asset
+			if (!NewAssetCreated)
+			{
+				FCoreUObjectDelegates::BroadcastOnObjectModified(Asset);
+			}
+		}
+		else
+		{
+			return nullptr;
+		}
+	}
+
+	// Save the file
+	if (Package && Parameters.bSaveOnExportEnded)
+	{
+		FEditorFileUtils::PromptForCheckoutAndSave({ Package }, /*bCheckDirty=*/false, /*bPromptToSave=*/false);
+	}
+
+	return Package;
+#else
+	UE_LOG(LogPCG, Error, TEXT("PCG Asset Exporter Utils cannot be used in non-editor builds."));
+	return nullptr;
+#endif
+}
+
+void UPCGAssetExporterUtils::UpdateAssets(const TArray<FAssetData>& PCGAssets, FPCGAssetExporterParameters InParameters)
+{
+#if WITH_EDITOR
+	TArray<UPackage*> PackagesToSave;
+	FPCGAssetExporterParameters Parameters = InParameters;
+
+	// Never open a dialog when updating pre-existing assets
+	Parameters.bOpenSaveDialog = false;
+
+	for (const FAssetData& PCGAsset : PCGAssets)
+	{
+		TSubclassOf<UPCGAssetExporter> ExporterSubclass;
+
+		const FString ExporterClassString = PCGAsset.GetTagValueRef<FString>(GET_MEMBER_NAME_CHECKED(UPCGDataAsset, ExporterClass));
+		TSoftClassPtr<> ExporterClass(ExporterClassString);
+		if (ExporterClass.IsValid() && ExporterClass->IsChildOf<UPCGAssetExporter>())
+		{
+			ExporterSubclass = ExporterClass.Get();
+		}
+		else
+		{
+			UE_LOG(LogPCG, Error, TEXT("Unable to update asset '%s' because exporter isn't valid."), *PCGAsset.AssetName.ToString());
+			continue;
+		}
+
+		UPCGAssetExporter* Exporter = NewObject<UPCGAssetExporter>(GetTransientPackage(), ExporterSubclass);
+
+		if (!Exporter)
+		{
+			UE_LOG(LogPCG, Error, TEXT("Unable to create exporter for asset '%s' during update process."), *PCGAsset.AssetName.ToString());
+			continue;
+		}
+
+		Exporter->AddToRoot();
+
+		if (UPackage* Package = Exporter->Update(PCGAsset))
+		{
+			FCoreUObjectDelegates::BroadcastOnObjectModified(PCGAsset.GetAsset());
+			PackagesToSave.Add(Package);
+		}
+
+		Exporter->RemoveFromRoot();
+	}
+
+	// Save the file(s)
+	if (!PackagesToSave.IsEmpty() && Parameters.bSaveOnExportEnded)
+	{
+		FEditorFileUtils::PromptForCheckoutAndSave(PackagesToSave, /*bCheckDirty=*/false, /*bPromptToSave=*/false);
+	}
+#else
+	UE_LOG(LogPCG, Error, TEXT("PCG Asset Exporter Utils cannot be used in non-editor builds."));
+#endif // WITH_EDITOR
+}
\ No newline at end of file
diff --git a/Engine/Plugins/PCG/Source/PCG/Public/PCGAssetExporter.h b/Engine/Plugins/PCG/Source/PCG/Public/PCGAssetExporter.h
--- a/Engine/Plugins/PCG/Source/PCG/Public/PCGAssetExporter.h
+++ b/Engine/Plugins/PCG/Source/PCG/Public/PCGAssetExporter.h
@@ -68,4 +68,7 @@ protected:
 
 	/** Sets up the exporter prior to performing update */
 	virtual void SerializeMetadata(FArchive& Ar) {}
+
+	/** Duplicates data that is outered to the transient package so it is properly saved */
+	void DuplicateAndReOuterData(UPCGDataAsset* Asset);
 };
\ No newline at end of file
diff --git a/Engine/Plugins/PCG/Source/PCG/Public/PCGAssetExporterUtils.h b/Engine/Plugins/PCG/Source/PCG/Public/PCGAssetExporterUtils.h
new file mode 100644
--- /dev/null
+++ b/Engine/Plugins/PCG/Source/PCG/Public/PCGAssetExporterUtils.h
@@ -0,0 +1,25 @@
+// Copyright Epic Games, Inc. All Rights Reserved.
+
+#pragma once
+
+#include "PCGAssetExporter.h"
+
+#include "AssetRegistry/AssetData.h"
+
+#include "PCGAssetExporterUtils.generated.h"
+
+/**
+* Asset export utils - will work only in editor builds. 
+*/
+UCLASS()
+class PCG_API UPCGAssetExporterUtils : public UBlueprintFunctionLibrary
+{
+	GENERATED_BODY()
+
+public:
+	UFUNCTION(BlueprintCallable, Category = "PCG|IO")
+	static UPackage* CreateAsset(UPCGAssetExporter* Exporter, FPCGAssetExporterParameters Parameters = FPCGAssetExporterParameters());
+
+	UFUNCTION(BlueprintCallable, Category = "PCG|IO")
+	static void UpdateAssets(const TArray<FAssetData>& PCGAssets, FPCGAssetExporterParameters Parameters = FPCGAssetExporterParameters());
+};
\ No newline at end of file
diff --git a/Engine/Plugins/PCG/Source/PCGEditor/Private/PCGAssetExporterUtils.cpp b/Engine/Plugins/PCG/Source/PCGEditor/Private/PCGAssetExporterUtils.cpp
deleted file mode 100644
diff --git a/Engine/Plugins/PCG/Source/PCGEditor/Public/PCGAssetExporterUtils.h b/Engine/Plugins/PCG/Source/PCGEditor/Public/PCGAssetExporterUtils.h
deleted file mode 100644
