[PCG] Graph Executor V2: Supports Out of tick task execution and scheduling (1st pass)

BiomeSample
50ms tick budget: 3.8 seconds -> 2.7 seconds
5ms tick budget: 7.3seconds -> 4.8 seconds

Slower in PCGWorldPartitionBuilder when ActiveTasks grows too much, will be fixed before enabling this new path by default

- Enable with CVar: pcg.Graph.ExecuteV2 = true (disabled by default, will be tested by some users to find remaining issues)
- New ExecuteV2 supports non multithreaded execution (CVar: pcg.GraphMultithreading = false)
- PCGAsync now supports NumAvailableTasks = -1 (no limits)
- Fix CVarAsyncOverrideChunkSize with SampleBasedRangeProcessing which didn't properly support the override (add a bAllowChunkSizeOverride param)
- Fix some CanExecuteOnlyOnMainThread calls to support a null FPCGContext which can happen when setting up tasks (to decide if PrepareForExecute can run in worker thread or not)

#rnx

diff --git a/Engine/Plugins/PCG/Source/PCG/Private/Elements/PCGCopyPoints.cpp b/Engine/Plugins/PCG/Source/PCG/Private/Elements/PCGCopyPoints.cpp
--- a/Engine/Plugins/PCG/Source/PCG/Private/Elements/PCGCopyPoints.cpp
+++ b/Engine/Plugins/PCG/Source/PCG/Private/Elements/PCGCopyPoints.cpp
@@ -333,7 +333,12 @@ bool FPCGCopyPointsElement::ExecuteInternal(FPCGContext* Context) const
 			if (AttributesToSet.Num() > 0)
 			{
 				int32 AttributeOffset = 0;
-				const int32 AttributePerDispatch = FMath::Max(1, Context->AsyncState.NumAvailableTasks);
+				const int32 DefaultAttributePerDispatch = 128;
+				int32 AttributePerDispatch = DefaultAttributePerDispatch;
+				if (Context->AsyncState.NumAvailableTasks > 0)
+				{
+					AttributePerDispatch = FMath::Min(Context->AsyncState.NumAvailableTasks, AttributePerDispatch);
+				}
 
 				while (AttributeOffset < AttributesToSet.Num())
 				{
diff --git a/Engine/Plugins/PCG/Source/PCG/Private/Elements/PCGExecuteBlueprint.cpp b/Engine/Plugins/PCG/Source/PCG/Private/Elements/PCGExecuteBlueprint.cpp
--- a/Engine/Plugins/PCG/Source/PCG/Private/Elements/PCGExecuteBlueprint.cpp
+++ b/Engine/Plugins/PCG/Source/PCG/Private/Elements/PCGExecuteBlueprint.cpp
@@ -1,6 +1,8 @@
 // Copyright Epic Games, Inc. All Rights Reserved.
 
 #include "Elements/PCGExecuteBlueprint.h"
+
+#include "Algo/Transform.h"
 #include "Engine/Blueprint.h"
 #include "Math/RandomStream.h"
 #include "PCGComponent.h"
@@ -830,9 +832,22 @@ void FPCGExecuteBlueprintElement::PostExecuteInternal(FPCGContext* InContext) co
 	check(InContext);
 	FPCGBlueprintExecutionContext* Context = static_cast<FPCGBlueprintExecutionContext*>(InContext);
 
-	check(IsInGameThread());
 	if (Context->BlueprintElementInstance)
 	{
+		check(IsInGameThread());
+
+		// Build a list of InputObjects so that we don't remove Async flags from Outputs that come from the Input as it isn't our responsability
+		TSet<TObjectPtr<const UObject>> InputObjects;
+		Algo::TransformIf(Context->InputData.TaggedData, InputObjects, 
+		[](const FPCGTaggedData& TaggedData)
+		{
+			return TaggedData.Data != nullptr;
+		},
+		[](const FPCGTaggedData& TaggedData)
+		{
+			return TaggedData.Data;
+		});
+
 		// Log info on outputs
 		for (int32 OutputIndex = 0; OutputIndex < Context->OutputData.TaggedData.Num(); ++OutputIndex)
 		{
@@ -849,7 +864,7 @@ void FPCGExecuteBlueprintElement::PostExecuteInternal(FPCGContext* InContext) co
 			if (Output.Data)
 			{
 				// Clear Async flags on objects created outside of the main thread and not part of the Context known async objects
-				if (Output.Data->HasAnyInternalFlags(EInternalObjectFlags::Async) && !Context->ContainsAsyncObject(Output.Data))
+				if (Output.Data->HasAnyInternalFlags(EInternalObjectFlags::Async) && !Context->ContainsAsyncObject(Output.Data) && !InputObjects.Contains(Output.Data))
 				{
 					Output.Data->ClearInternalFlags(EInternalObjectFlags::Async);
 					ForEachObjectWithOuter(Output.Data, [](UObject* SubObject) { SubObject->ClearInternalFlags(EInternalObjectFlags::Async); }, true);
@@ -1080,7 +1095,11 @@ bool FPCGExecuteBlueprintElement::ShouldComputeFullOutputDataCrc(FPCGContext* Co
 
 bool FPCGExecuteBlueprintElement::CanExecuteOnlyOnMainThread(FPCGContext* Context) const
 {
-	check(Context);
+	if (!Context)
+	{
+		return true;
+	}
+
 	FPCGBlueprintExecutionContext* BPContext = static_cast<FPCGBlueprintExecutionContext*>(Context);
 
 	// Always execute PostExecute on main thread
diff --git a/Engine/Plugins/PCG/Source/PCG/Private/Elements/PCGLoopElement.cpp b/Engine/Plugins/PCG/Source/PCG/Private/Elements/PCGLoopElement.cpp
--- a/Engine/Plugins/PCG/Source/PCG/Private/Elements/PCGLoopElement.cpp
+++ b/Engine/Plugins/PCG/Source/PCG/Private/Elements/PCGLoopElement.cpp
@@ -359,7 +359,7 @@ bool FPCGLoopElement::ExecuteInternal(FPCGContext* InContext) const
 
 				FPCGDataCollection SubgraphOutput;
 				// While this should be always return true, if a scheduled task was cancelled, this can still happen because the dependency has been removed but this code isn't aware of this
-				if (Subsystem->GetOutputData(SubgraphTaskId, SubgraphOutput, /*bClearDataOnGet=*/true))
+				if (Subsystem->GetOutputData(SubgraphTaskId, SubgraphOutput))
 				{
 					if (FeedbackPinNames.IsEmpty() || bIsLastTask)
 					{
@@ -369,6 +369,7 @@ bool FPCGLoopElement::ExecuteInternal(FPCGContext* InContext) const
 					{
 						Algo::CopyIf(SubgraphOutput.TaggedData, Context->OutputData.TaggedData, [&FeedbackPinNames](const FPCGTaggedData& InTaggedData) { return !FeedbackPinNames.Contains(InTaggedData.Pin); });
 					}
+					Subsystem->ClearOutputData(SubgraphTaskId);
 				}
 			}
 		}
@@ -432,7 +433,7 @@ bool FPCGLoopInputForwardingElement::ExecuteInternal(FPCGContext* Context) const
 		if (UPCGSubsystem* Subsystem = Context->SourceComponent->GetSubsystem())
 		{
 			FPCGDataCollection PreviousTaskOutput;
-			if (Subsystem->GetOutputData(PreviousIterationTaskId, PreviousTaskOutput, /*bClearDataOnGet=*/false))
+			if (Subsystem->GetOutputData(PreviousIterationTaskId, PreviousTaskOutput))
 			{
 				Algo::CopyIf(PreviousTaskOutput.TaggedData, Context->OutputData.TaggedData, [this](const FPCGTaggedData& InTaggedData) { return FeedbackPinNames.Contains(InTaggedData.Pin); });
 			}
diff --git a/Engine/Plugins/PCG/Source/PCG/Private/Elements/PCGSplineSampler.cpp b/Engine/Plugins/PCG/Source/PCG/Private/Elements/PCGSplineSampler.cpp
--- a/Engine/Plugins/PCG/Source/PCG/Private/Elements/PCGSplineSampler.cpp
+++ b/Engine/Plugins/PCG/Source/PCG/Private/Elements/PCGSplineSampler.cpp
@@ -1206,9 +1206,12 @@ namespace PCGSplineSamplerHelpers
 		const FVector::FReal MaxY = FMath::FloorToDouble(MaxPoint.Y / Params.InteriorSampleSpacing) * Params.InteriorSampleSpacing;
 
 		constexpr int32 MinIterationPerDispatch = 4;
-		const int32 NumAvailableTasks = Context ? Context->AsyncState.NumAvailableTasks : 1;
 		const int32 NumIterations = (MaxY + UE_KINDA_SMALL_NUMBER - MinY) / Params.InteriorSampleSpacing;
-		const int32 NumDispatch = FMath::Max(1, FMath::Min(NumAvailableTasks, NumIterations / MinIterationPerDispatch));
+		int32 NumDispatch = Context ? (NumIterations / MinIterationPerDispatch) : 1;
+		if (Context && Context->AsyncState.NumAvailableTasks > 0)
+		{
+			NumDispatch = FMath::Min(Context->AsyncState.NumAvailableTasks, NumDispatch);
+		}
 		const int32 NumIterationsPerDispatch = NumIterations / NumDispatch;
 
 		TArray<TArray<TTuple<FTransform, FVector, float>>> InteriorSplinePointData;
diff --git a/Engine/Plugins/PCG/Source/PCG/Private/Elements/PCGStaticMeshSpawner.cpp b/Engine/Plugins/PCG/Source/PCG/Private/Elements/PCGStaticMeshSpawner.cpp
--- a/Engine/Plugins/PCG/Source/PCG/Private/Elements/PCGStaticMeshSpawner.cpp
+++ b/Engine/Plugins/PCG/Source/PCG/Private/Elements/PCGStaticMeshSpawner.cpp
@@ -370,7 +370,7 @@ bool FPCGStaticMeshSpawnerElement::CanExecuteOnlyOnMainThread(FPCGContext* Conte
 	// PrepareData can call UPCGManagedComponent::MarkAsReused which registers the ISMC, which can go into Chaos code that asserts if not on main thread.
 	// TODO: We can likely re-enable multi-threading for PrepareData if we move the call to MarkAsReused to Execute. There should hopefully not be
 	// wider contention on resources resources are not shared across nodes and are also per-component.
-	return Context->CurrentPhase == EPCGExecutionPhase::Execute || Context->CurrentPhase == EPCGExecutionPhase::PrepareData;
+	return !Context || Context->CurrentPhase == EPCGExecutionPhase::Execute || Context->CurrentPhase == EPCGExecutionPhase::PrepareData;
 }
 
 void FPCGStaticMeshSpawnerElement::SpawnStaticMeshInstances(FPCGStaticMeshSpawnerContext* Context, const FPCGMeshInstanceList& InstanceList, AActor* TargetActor, const FPCGPackedCustomData* InPackedCustomData) const
@@ -479,7 +479,13 @@ void FPCGStaticMeshSpawnerElement::SpawnStaticMeshInstances(FPCGStaticMeshSpawne
 
 void FPCGStaticMeshSpawnerElement::AbortInternal(FPCGContext* InContext) const
 {
-	TRACE_CPUPROFILER_EVENT_SCOPE(FPCGStaticMeshSpawnerElement::Execute);
+	TRACE_CPUPROFILER_EVENT_SCOPE(FPCGStaticMeshSpawnerElement::AbortInternal);
+	// It is possible to Abort a ready task with no context yet
+	if (!InContext)
+	{
+		return;
+	}
+
 	FPCGStaticMeshSpawnerContext* Context = static_cast<FPCGStaticMeshSpawnerContext*>(InContext);
 
 	// Any resources we've touched during the execution of this node can potentially be in a "not-quite complete state" especially if we have multiple sources of data writing to the same ISMC.
diff --git a/Engine/Plugins/PCG/Source/PCG/Private/Graph/PCGGraphExecutor.cpp b/Engine/Plugins/PCG/Source/PCG/Private/Graph/PCGGraphExecutor.cpp
--- a/Engine/Plugins/PCG/Source/PCG/Private/Graph/PCGGraphExecutor.cpp
+++ b/Engine/Plugins/PCG/Source/PCG/Private/Graph/PCGGraphExecutor.cpp
@@ -19,8 +19,10 @@
 #include "Utils/PCGGraphExecutionLogging.h"
 
 #include "Algo/AnyOf.h"
-#include "Async/Async.h"
+#include "Algo/ForEach.h"
+#include "Containers/Ticker.h"
 #include "GameFramework/Actor.h"
+#include "Tasks/Task.h"
 
 #if WITH_EDITOR
 #include "Editor.h"
@@ -73,8 +75,92 @@ namespace PCGGraphExecutor
 		TEXT("pcg.Graph.DynamicTaskCulling"),
 		true,
 		TEXT("Controls whether tasks are culled at execution time, for example in response to an deactivated dynamic branch pin"));
+
+	TAutoConsoleVariable<bool> CVarExecuteV2(
+		TEXT("pcg.Graph.ExecuteV2"),
+		false,
+		TEXT("Enables the new execute allowing tasks to run out of tick"));
+
+	TAutoConsoleVariable<int32> CVarMaxWorkerTasks(
+		TEXT("pcg.Graph.MaxWorkerTasks"),
+		32,
+		TEXT("Max in flight PCG Element tasks allowed at the same time. Note that PCG Element execution can spawn async tasks themselves"));
+	
+
+	template<typename FunctorType>
+	void ExecuteOnGameThread(const TCHAR* DebugName, FunctorType&& Functor)
+	{
+		if (IsInGameThread())
+		{
+			Functor();
+		}
+		else
+		{
+			::ExecuteOnGameThread(DebugName, std::forward<FunctorType>(Functor));
+		}
+	}
+
+	void ClearAsyncFlags(TSet<TObjectPtr<UObject>>& AsyncObjects)
+	{
+		// Remove Async Flags from Objects created in Async PCG Tasks so that they can get tracked by the GarbageCollector properly
+		// The reason to do this here and not in the FPCGContext destructor is that FPCGContext might get copied in Blueprint node execution causing multiple instances of the same AsyncObjects list
+		// The FPCGGraphActiveTask is the true owner of the original context (throughg TUniquePtr) so we let it do the cleanup
+		for (TObjectPtr<UObject>& AsyncObject : AsyncObjects)
+		{
+			if (ensure(AsyncObject->HasAnyInternalFlags(EInternalObjectFlags::Async)))
+			{
+				AsyncObject->ClearInternalFlags(EInternalObjectFlags::Async);
+				ForEachObjectWithOuter(AsyncObject, [](UObject* SubObject) { SubObject->ClearInternalFlags(EInternalObjectFlags::Async); }, true);
+			}
+		}
+	}
+
+	// Similar to existing UE::TScopeLock but allows optional TryLock in which case caller is responsible to 
+	// check if Lock succeeded or not with IsLocked(). This allows execution code to be skipped if non main thread
+	// tasks fail to get the lock instead of spinning for nothing
+	template<typename MutexType>
+	class TScopeLock
+	{
+	public:
+		UE_NONCOPYABLE(TScopeLock);
+
+		UE_NODISCARD_CTOR TScopeLock(MutexType& InMutex, bool bTryLock=false)
+			: Mutex(&InMutex), bLocked(false)
+		{
+			check(Mutex);
+			if (bTryLock)
+			{
+				bLocked = Mutex->TryLock();
+			}
+			else
+			{
+				Mutex->Lock();
+				bLocked = true;
+			}
+		}
+
+		~TScopeLock()
+		{
+			if (bLocked && Mutex)
+			{
+				Mutex->Unlock();
+				Mutex = nullptr;
+			}
+		}
+			
+		bool IsLocked() const
+		{
+			return bLocked;
+		}
+
+	private:
+		MutexType* Mutex;
+		bool bLocked;
+	};
 }
 
+int32 FPCGGraphActiveTask::NumExecuting = 0;
+
 const FPCGStack* FPCGGraphTask::GetStack() const
 {
 	return (StackContext && StackIndex != INDEX_NONE) ? StackContext->GetStack(StackIndex) : nullptr;
@@ -142,17 +228,56 @@ FPCGGraphActiveTask::~FPCGGraphActiveTask()
 {
 	if (Context)
 	{
-		// Remove Async Flags from Objects created in Async PCG Tasks so that they can get tracked by the GarbageCollector properly
-		// The reason to do this here and not in the FPCGContext destructor is that FPCGContext might get copied in Blueprint node execution causing multiple instances of the same AsyncObjects list
-		// The FPCGGraphActiveTask is the true owner of the original context (throughg TUniquePtr) so we let it do the cleanup
-		for (TObjectPtr<UObject>& AsyncObject : Context->AsyncObjects)
+		PCGGraphExecutor::ExecuteOnGameThread(UE_SOURCE_LOCATION, [ContextPtr = Context.Release()]()
 		{
-			if (ensure(AsyncObject->HasAnyInternalFlags(EInternalObjectFlags::Async)))
+			PCGGraphExecutor::ClearAsyncFlags(ContextPtr->AsyncObjects);
+			delete ContextPtr;
+		});	
+	}
+}
+
+void FPCGGraphActiveTask::StartExecuting()
+{
+	check(!bIsExecutingTask);
+	bIsExecutingTask = true;
+	++NumExecuting;
+
+	// @todo_pcg this should be merged through some kind of visitor pattern with AddStructReferencedObjects
+	// We need to keep a immutable list of references prior to execution mainly so that TaggedData contents don't change while we are running GC
+	{
+		ExecutingReferences.Reserve(Context->InputData.TaggedData.Num() + Context->OutputData.TaggedData.Num() + 1);
+		for (const FPCGTaggedData& TaggedInputData : Context->InputData.TaggedData)
+		{
+			if (TaggedInputData.Data)
 			{
-				AsyncObject->ClearInternalFlags(EInternalObjectFlags::Async);
-				ForEachObjectWithOuter(AsyncObject, [](UObject* SubObject) { SubObject->ClearInternalFlags(EInternalObjectFlags::Async); }, true);
+				ExecutingReferences.Add(TaggedInputData.Data);
+			}
+		}
+
+		for (const FPCGTaggedData& TaggedOutputData : Context->OutputData.TaggedData)
+		{
+			if (TaggedOutputData.Data)
+			{
+				ExecutingReferences.Add(TaggedOutputData.Data);
 			}
 		}
+
+		if (Context->SettingsWithOverride)
+		{
+			ExecutingReferences.Add(Context->SettingsWithOverride);
+		}
+	}
+}
+
+void FPCGGraphActiveTask::StopExecuting()
+{
+	if (bIsExecutingTask)
+	{
+		bIsExecutingTask = false;
+		ExecutingTask = {};
+		--NumExecuting;
+
+		ExecutingReferences.Empty();
 	}
 }
 
@@ -249,7 +374,7 @@ FPCGTaskId FPCGGraphExecutor::Schedule(
 		// Setup fetch task on input node
 		CompiledTasks[0].Element = InputElement;
 
-		ScheduleLock.Lock();
+		PCGGraphExecutor::TScopeLock ScopeLock(ScheduleLock);
 
 		FPCGGraphScheduleTask& ScheduledTask = ScheduledTasks.Emplace_GetRef();
 		ScheduledTask.Tasks = MoveTemp(CompiledTasks);
@@ -278,8 +403,6 @@ FPCGTaskId FPCGGraphExecutor::Schedule(
 
 		ScheduledTask.FirstTaskIndex = ScheduledTask.Tasks.Num() - 2;
 		ScheduledTask.LastTaskIndex = ScheduledTask.Tasks.Num() - 1;
-
-		ScheduleLock.Unlock();
 	}
 
 	return ScheduledId;
@@ -317,43 +440,49 @@ TArray<UPCGComponent*> FPCGGraphExecutor::CancelAll()
 
 TSet<UPCGComponent*> FPCGGraphExecutor::Cancel(TFunctionRef<bool(TWeakObjectPtr<UPCGComponent>)> CancelFilter)
 {
+	check(IsInGameThread());
 	TSet<UPCGComponent*> CancelledComponents;
+	const FPCGGraphExecutor::EExecuteVersion CurrentExecuteVersion = GetExecuteVersion();
 
 	// Visit scheduled tasks
-	ScheduleLock.Lock();
-	for (const FPCGGraphScheduleTask& ScheduledTask : ScheduledTasks)
 	{
-		if (CancelFilter(ScheduledTask.SourceComponent))
+		PCGGraphExecutor::TScopeLock ScopeLock(ScheduleLock);
+		for (const FPCGGraphScheduleTask& ScheduledTask : ScheduledTasks)
 		{
-			CancelledComponents.Add(ScheduledTask.SourceComponent.Get());
+			if (CancelFilter(ScheduledTask.SourceComponent))
+			{
+				CancelledComponents.Add(ScheduledTask.SourceComponent.Get());
+			}
 		}
 	}
-	ScheduleLock.Unlock();
-
-	// Visit ready tasks
-	for (const FPCGGraphTask& Task : ReadyTasks)
+		
 	{
-		if (CancelFilter(Task.SourceComponent))
+		// Visit ready tasks
+		PCGGraphExecutor::TScopeLock ScopeLock(LiveTasksLock);
+		for (const FPCGGraphTask& Task : ReadyTasks)
 		{
-			CancelledComponents.Add(Task.SourceComponent.Get());
+			if (CancelFilter(Task.SourceComponent))
+			{
+				CancelledComponents.Add(Task.SourceComponent.Get());
+			}
 		}
-	}
 
-	// Visit active tasks
-	for (const FPCGGraphActiveTask& Task : ActiveTasks)
-	{
-		if (Task.Context && CancelFilter(Task.Context->SourceComponent))
+		// Visit active tasks
+		for (const TSharedPtr<FPCGGraphActiveTask>& Task : ActiveTasks)
 		{
-			CancelledComponents.Add(Task.Context->SourceComponent.Get());
+			if (Task->Context && CancelFilter(Task->Context->SourceComponent))
+			{
+				CancelledComponents.Add(Task->Context->SourceComponent.Get());
+			}
 		}
-	}
-
-	// Visit sleeping tasks
-	for (const FPCGGraphActiveTask& Task : SleepingTasks)
-	{
-		if (Task.Context && CancelFilter(Task.Context->SourceComponent))
+	
+		// Visit sleeping tasks
+		for (const TSharedPtr<FPCGGraphActiveTask>& Task : SleepingTasks)
 		{
-			CancelledComponents.Add(Task.Context->SourceComponent.Get());
+			if (Task->Context && CancelFilter(Task->Context->SourceComponent))
+			{
+				CancelledComponents.Add(Task->Context->SourceComponent.Get());
+			}
 		}
 	}
 
@@ -393,99 +522,146 @@ TSet<UPCGComponent*> FPCGGraphExecutor::Cancel(TFunctionRef<bool(TWeakObjectPtr<
 		bStableCancellationSet = true;
 
 		// Remove from scheduled tasks
-		ScheduleLock.Lock();
-		for (int32 ScheduledTaskIndex = ScheduledTasks.Num() - 1; ScheduledTaskIndex >= 0; --ScheduledTaskIndex)
 		{
-			FPCGGraphScheduleTask& ScheduledTask = ScheduledTasks[ScheduledTaskIndex];
-			if (CancelledComponents.Contains(ScheduledTask.SourceComponent.Get()))
-			{
-				CancelledScheduledTasks.Add(ScheduledTask.Tasks[ScheduledTask.LastTaskIndex].NodeId);
-
-				TryAbortScheduledTasks(ScheduledTask);
-				ScheduledTasks.RemoveAtSwap(ScheduledTaskIndex);
-			}
-		}
-
-		// WARNING: variable upper bound
-		for (int32 CancelledTaskIdIndex = 0; CancelledTaskIdIndex < CancelledScheduledTasks.Num(); ++CancelledTaskIdIndex)
-		{
-			const FPCGTaskId EndTaskId = CancelledScheduledTasks[CancelledTaskIdIndex];
+			PCGGraphExecutor::TScopeLock ScopeLock(ScheduleLock);
 			for (int32 ScheduledTaskIndex = ScheduledTasks.Num() - 1; ScheduledTaskIndex >= 0; --ScheduledTaskIndex)
 			{
 				FPCGGraphScheduleTask& ScheduledTask = ScheduledTasks[ScheduledTaskIndex];
-				const bool bContainsDependency = Algo::AnyOf(ScheduledTask.Tasks[ScheduledTask.FirstTaskIndex].Inputs, [EndTaskId](const FPCGGraphTaskInput& Input) { return Input.TaskId == EndTaskId; });
-
-				if (bContainsDependency)
+				if (CancelledComponents.Contains(ScheduledTask.SourceComponent.Get()))
 				{
-					UPCGComponent* TaskComponent = ScheduledTask.SourceComponent.Get();
-					if (TaskComponent && !CancelledComponents.Contains(TaskComponent))
-					{
-						CancelledComponents.Add(TaskComponent);
-						bStableCancellationSet = false;
-					}
-
 					CancelledScheduledTasks.Add(ScheduledTask.Tasks[ScheduledTask.LastTaskIndex].NodeId);
 
 					TryAbortScheduledTasks(ScheduledTask);
 					ScheduledTasks.RemoveAtSwap(ScheduledTaskIndex);
 				}
 			}
-		}
 
-		CancelledScheduledTasks.Reset();
-		ScheduleLock.Unlock();
+			// WARNING: variable upper bound
+			for (int32 CancelledTaskIdIndex = 0; CancelledTaskIdIndex < CancelledScheduledTasks.Num(); ++CancelledTaskIdIndex)
+			{
+				const FPCGTaskId EndTaskId = CancelledScheduledTasks[CancelledTaskIdIndex];
+				for (int32 ScheduledTaskIndex = ScheduledTasks.Num() - 1; ScheduledTaskIndex >= 0; --ScheduledTaskIndex)
+				{
+					FPCGGraphScheduleTask& ScheduledTask = ScheduledTasks[ScheduledTaskIndex];
+					const bool bContainsDependency = Algo::AnyOf(ScheduledTask.Tasks[ScheduledTask.FirstTaskIndex].Inputs, [EndTaskId](const FPCGGraphTaskInput& Input) { return Input.TaskId == EndTaskId; });
 
-		// Remove from ready tasks
-		for (int32 ReadyTaskIndex = ReadyTasks.Num() - 1; ReadyTaskIndex >= 0; --ReadyTaskIndex)
-		{
-			FPCGGraphTask& Task = ReadyTasks[ReadyTaskIndex];
+					if (bContainsDependency)
+					{
+						UPCGComponent* TaskComponent = ScheduledTask.SourceComponent.Get();
+						if (TaskComponent && !CancelledComponents.Contains(TaskComponent))
+						{
+							CancelledComponents.Add(TaskComponent);
+							bStableCancellationSet = false;
+						}
+
+						CancelledScheduledTasks.Add(ScheduledTask.Tasks[ScheduledTask.LastTaskIndex].NodeId);
+
+						TryAbortScheduledTasks(ScheduledTask);
+						ScheduledTasks.RemoveAtSwap(ScheduledTaskIndex);
+					}
+				}
+			}
 
-			if (CancelledComponents.Contains(Task.SourceComponent.Get()))
+			CancelledScheduledTasks.Reset();
+		}
+				
+		// Mark as cancelled in the active tasks - needed to make sure we're not breaking the current execution (if any)
+		{
+			TArray<TTuple<TSharedPtr<FPCGGraphActiveTask>, UE::Tasks::TTask<bool>>> CancelledActiveTasks;
 			{
-				if (Task.Element)
+				PCGGraphExecutor::TScopeLock ScopeLock(LiveTasksLock);
+				for (int32 ActiveTaskIndex = ActiveTasks.Num() - 1; ActiveTaskIndex >= 0; --ActiveTaskIndex)
 				{
-					Task.Element->Abort(Task.Context);
-				}
+					TSharedPtr<FPCGGraphActiveTask>& ActiveTask = ActiveTasks[ActiveTaskIndex];
+					if (ActiveTask->Context && CancelledComponents.Contains(ActiveTask->Context->SourceComponent.Get()))
+					{
+						// While we have lock Task can't complete, but we can't wait on this task with the lock neither so we capture it here 
+						// and are going to wait on it outside of the LiveTaskLock
+						UE::Tasks::TTask<bool> TaskHandle = ActiveTask->ExecutingTask;
+						ActiveTask->bWasCancelled = true;
 
-				FPCGTaskId CancelledTaskId = Task.NodeId;
-				RemoveTaskFromInputSuccessors(CancelledTaskId, Task.Inputs);
+						CancelledActiveTasks.Add({ ActiveTask, MoveTemp(TaskHandle) });
+						ActiveTasks.RemoveAtSwap(ActiveTaskIndex);
+					}
+				}
+			}
 
-				delete Task.Context;
-				ReadyTasks.RemoveAtSwap(ReadyTaskIndex);
-				bStableCancellationSet &= !CancelNextTasks(CancelledTaskId, CancelledComponents);
+			for (TTuple<TSharedPtr<FPCGGraphActiveTask>, UE::Tasks::TTask<bool>>& CancelledActiveTask : CancelledActiveTasks)
+			{
+				// Wait for Task (if any) to complete
+				UE::Tasks::TTask<bool> TaskHandle = CancelledActiveTask.Get<UE::Tasks::TTask<bool>>();
+				TaskHandle.Wait();
+				
+				TSharedPtr<FPCGGraphActiveTask> ActiveTask = CancelledActiveTask.Get<TSharedPtr<FPCGGraphActiveTask>>();			
+				check(ActiveTask->Element);
+				ActiveTask->Element->Abort(ActiveTask->Context.Get());
+												
+				bStableCancellationSet &= !CancelNextTasks(ActiveTask->NodeId, CancelledComponents);
+
+				// In V2 this is done inside the task
+				if (CurrentExecuteVersion == EExecuteVersion::V1)
+				{
+					PostTaskExecute(ActiveTask);
+				}
 			}
 		}
 
-		// Mark as cancelled in the active tasks - needed to make sure we're not breaking the current execution (if any)
-		for (int32 ActiveTaskIndex = ActiveTasks.Num() - 1; ActiveTaskIndex >= 0; --ActiveTaskIndex)
+		// Remove from ready tasks
 		{
-			FPCGGraphActiveTask& Task = ActiveTasks[ActiveTaskIndex];
-			if (Task.Context && CancelledComponents.Contains(Task.Context->SourceComponent.Get()))
+			TArray<FPCGGraphTask> CancelledReadyTasks;
 			{
-				check(Task.Element);
-				Task.Element->Abort(Task.Context.Get());
+				PCGGraphExecutor::TScopeLock ScopeLock(LiveTasksLock);
+				for (int32 ReadyTaskIndex = ReadyTasks.Num() - 1; ReadyTaskIndex >= 0; --ReadyTaskIndex)
+				{
+					FPCGGraphTask& Task = ReadyTasks[ReadyTaskIndex];
 
-				FPCGTaskId CancelledTaskId = Task.NodeId;
-				Task.bWasCancelled = true;
-				RemoveTaskFromInputSuccessors(CancelledTaskId, Task.Inputs);
+					if (CancelledComponents.Contains(Task.SourceComponent.Get()))
+					{
+						CancelledReadyTasks.Add(MoveTemp(Task));
+						ReadyTasks.RemoveAtSwap(ReadyTaskIndex);
+					}
+				}
+			}
 
+			for (FPCGGraphTask& CancelledReadyTask : CancelledReadyTasks)
+			{
+				if (CancelledReadyTask.Element)
+				{
+					CancelledReadyTask.Element->Abort(CancelledReadyTask.Context);
+				}
+
+				FPCGTaskId CancelledTaskId = CancelledReadyTask.NodeId;
+				RemoveTaskFromInputSuccessors(CancelledTaskId, CancelledReadyTask.Inputs);
+
+				delete CancelledReadyTask.Context;
 				bStableCancellationSet &= !CancelNextTasks(CancelledTaskId, CancelledComponents);
 			}
 		}
 
 		// Remove from sleeping tasks
-		for (int32 SleepingTaskIndex = SleepingTasks.Num() - 1; SleepingTaskIndex >= 0; --SleepingTaskIndex)
 		{
-			FPCGGraphActiveTask& Task = SleepingTasks[SleepingTaskIndex];
-			if (Task.Context && CancelledComponents.Contains(Task.Context->SourceComponent.Get()))
+			TArray<TSharedPtr<FPCGGraphActiveTask>> CancelledSleepingTasks;
 			{
-				check(Task.Element);
-				Task.Element->Abort(Task.Context.Get());
+				PCGGraphExecutor::TScopeLock ScopeLock(LiveTasksLock);
+				for (int32 SleepingTaskIndex = SleepingTasks.Num() - 1; SleepingTaskIndex >= 0; --SleepingTaskIndex)
+				{
+					TSharedPtr<FPCGGraphActiveTask>& Task = SleepingTasks[SleepingTaskIndex];
+					if (Task->Context && CancelledComponents.Contains(Task->Context->SourceComponent.Get()))
+					{
+						CancelledSleepingTasks.Add(MoveTemp(Task));
+						SleepingTasks.RemoveAtSwap(SleepingTaskIndex);
+					}
+				}
+			}
 
-				FPCGTaskId CancelledTaskId = Task.NodeId;
-				RemoveTaskFromInputSuccessors(CancelledTaskId, Task.Inputs);
+			for (TSharedPtr<FPCGGraphActiveTask>& CancelledSleepingTask : CancelledSleepingTasks)
+			{
+				check(CancelledSleepingTask->Element);
+				CancelledSleepingTask->Element->Abort(CancelledSleepingTask->Context.Get());
 
-				SleepingTasks.RemoveAtSwap(SleepingTaskIndex);
+				FPCGTaskId CancelledTaskId = CancelledSleepingTask->NodeId;
+				RemoveTaskFromInputSuccessors(CancelledTaskId, CancelledSleepingTask->Inputs);
+								
 				bStableCancellationSet &= !CancelNextTasks(CancelledTaskId, CancelledComponents);
 			}
 		}
@@ -522,7 +698,6 @@ bool FPCGGraphExecutor::IsAnyGraphCurrentlyExecuting() const
 		return true;
 	}
 
-	// No need for locking here as we don't need the value to be precise here and it can change right after a lock anyways.
 	return ScheduledTasks.Num() > 0;
 }
 
@@ -590,26 +765,26 @@ FPCGTaskId FPCGGraphExecutor::ScheduleGenericWithContext(TFunction<bool(FPCGCont
 	Task.SourceComponent = InSourceComponent;
 	Task.Element = MakeShared<FPCGGenericElement>(InOperation, InAbortOperation);
 
-	ScheduleLock.Lock();
-
-	// Assign task id
-	Task.NodeId = NextTaskId++;
-
-	FPCGGraphScheduleTask& ScheduledTask = ScheduledTasks.Emplace_GetRef();
-	ScheduledTask.Tasks.Add(Task);
-	ScheduledTask.SourceComponent = InSourceComponent;
-	ScheduledTask.bHasAbortCallbacks = !!InAbortOperation;
+	{
+		PCGGraphExecutor::TScopeLock ScopeLock(ScheduleLock);
 
-	ScheduleLock.Unlock();
+		// Assign task id
+		Task.NodeId = NextTaskId++;
 
+		FPCGGraphScheduleTask& ScheduledTask = ScheduledTasks.Emplace_GetRef();
+		ScheduledTask.Tasks.Add(Task);
+		ScheduledTask.SourceComponent = InSourceComponent;
+		ScheduledTask.bHasAbortCallbacks = !!InAbortOperation;
+	}
 	return Task.NodeId;
 }
 
 void FPCGGraphExecutor::MarkInputResults(TArrayView<const FPCGTaskId> InInputResults)
 {
+	TRACE_CPUPROFILER_EVENT_SCOPE(FPCGGraphExecutor::MarkInputResults);
 	if (InInputResults.Num() > 0)
 	{
-		FWriteScopeLock WriteLock(TaskOutputsRWLock);
+		PCGGraphExecutor::TScopeLock ScopeLock(TaskOutputsLock);
 		for (FPCGTaskId TaskId : InInputResults)
 		{
 			FOutputDataInfo& OutputDataInfo = TaskOutputs[TaskId];
@@ -621,34 +796,455 @@ void FPCGGraphExecutor::MarkInputResults(TArrayView<const FPCGTaskId> InInputRes
 	}
 }
 
-bool FPCGGraphExecutor::GetOutputData(FPCGTaskId TaskId, FPCGDataCollection& OutData, bool bClearDataOnGet)
+bool FPCGGraphExecutor::GetOutputData(FPCGTaskId TaskId, FPCGDataCollection& OutData)
+{
+	TRACE_CPUPROFILER_EVENT_SCOPE(FPCGGraphExecutor::GetOutputData);
+	PCGGraphExecutor::TScopeLock ScopeLock(TaskOutputsLock);
+	if (FOutputDataInfo* OutputDataInfo = TaskOutputs.Find(TaskId))
+	{
+		ensure(!OutputDataInfo->bCulled);
+		OutData = OutputDataInfo->DataCollection;
+		return true;
+	}
+	else
+	{
+		return false;
+	}
+}
+
+void FPCGGraphExecutor::ClearOutputData(FPCGTaskId TaskId)
+{
+	TRACE_CPUPROFILER_EVENT_SCOPE(FPCGGraphExecutor::ClearOutputData);
+	PCGGraphExecutor::TScopeLock ScopeLock(TaskOutputsLock);
+	if (FOutputDataInfo* OutputDataInfo = TaskOutputs.Find(TaskId); OutputDataInfo && ensure(OutputDataInfo->bNeedsManualClear))
+	{
+		TaskOutputs.Remove(TaskId);
+	}
+}
+
+bool FPCGGraphExecutor::ProcessScheduledTasks()
+{
+	TRACE_CPUPROFILER_EVENT_SCOPE(FPCGGraphExecutor::ProcessScheduledTasks);
+	bool bProcessedNewTasks = false;
+	const bool bIsInGameThread = IsInGameThread();
+	TArray<FCachedResult*> CachedResults;
+
+	{
+		// Process any newly scheduled graphs to execute
+		PCGGraphExecutor::TScopeLock ScopeLock(ScheduleLock, !bIsInGameThread);
+		if (ScopeLock.IsLocked())
+		{
+			for (int32 ScheduledTaskIndex = ScheduledTasks.Num() - 1; ScheduledTaskIndex >= 0; --ScheduledTaskIndex)
+			{
+				FPCGGraphScheduleTask& ScheduledTask = ScheduledTasks[ScheduledTaskIndex];
+
+				check(ScheduledTask.Tasks.Num() > 0);
+				// Push tasks to the primary task list & build successors map
+				for (FPCGGraphTask& Task : ScheduledTask.Tasks)
+				{
+					const FPCGTaskId TaskId = Task.NodeId;
+
+					// TODO: review if it's actually possible for a task with inputs to be ready at this point - it seems very unlikely
+					bool bPushToReady = true;
+					{
+						PCGGraphExecutor::TScopeLock ChildScopeLock(TasksLock);
+						PCGGraphExecutor::TScopeLock OutputLock(TaskOutputsLock);
+						for (const FPCGGraphTaskInput& Input : Task.Inputs)
+						{
+							if (!TaskOutputs.Contains(Input.TaskId))
+							{
+								TaskSuccessors.FindOrAdd(Input.TaskId).Add(TaskId);
+								bPushToReady = false;
+							}
+						}
+					}
+
+					// Automatically push inputless/already satisfied tasks to the ready queue
+					if (bPushToReady)
+					{
+						OnTaskInputsReady(Task, CachedResults, bIsInGameThread);
+					}
+					else
+					{
+						PCGGraphExecutor::TScopeLock ChildScopeLock(TasksLock);
+						Tasks.Add(TaskId, MoveTemp(Task));
+					}
+					bProcessedNewTasks = true;
+				}
+			}
+
+			if (!ScheduledTasks.IsEmpty())
+			{
+				PCGGraphExecutor::TScopeLock ChildScopeLock(TasksLock);
+				PCGGraphExecutionLogging::LogGraphPostSchedule(Tasks, TaskSuccessors);
+			}
+
+			ScheduledTasks.Reset();
+		}
+	}
+
+	ProcessCachedResults(CachedResults);
+		
+	if (bIsInGameThread)
+	{
+		PCGGraphExecutor::TScopeLock ScopeLock(TasksLock);
+		PCGGraphExecutor::TScopeLock ChildScopeLock(LiveTasksLock);
+		// This is a safeguard to check if we're in a stuck state
+		if (CachingResults.Num() == 0 && ReadyTasks.Num() == 0 && ActiveTasks.Num() == 0 && SleepingTasks.Num() == 0 && Tasks.Num() > 0)
+		{
+			UE_LOG(LogPCG, Error, TEXT("PCG Graph executor error: tasks are in a deadlocked state. Will drop all tasks."));
+			ClearAllTasks();
+		}
+	}
+
+	return bProcessedNewTasks;
+}
+
+void FPCGGraphExecutor::ExecuteTasksEnded()
+{
+	check(IsInGameThread());
+	if (!bNeedToExecuteTasksEnded)
+	{
+		return;
+	}
+	bNeedToExecuteTasksEnded = false;
+
+	check(IsInGameThread());
+	// Nothing left to do; we'll release everything here.
+	// TODO: this is fine and will make sure any intermediate data is properly
+	// garbage collected, however, this goes a bit against our goals if we want to
+	// keep a cache of intermediate results.
+	if (GetNonScheduledRemainingTaskCount() == 0)
+	{
+		{
+			PCGGraphExecutor::TScopeLock ScopeLock(TasksLock);
+			if (!ensure(TaskSuccessors.IsEmpty()))
+			{
+				TaskSuccessors.Reset();
+			}
+		}
+
+		ClearResults();
+
+#if WITH_EDITOR
+		// Call the notification update here to prevent it from sticking around - needed because we early out before this
+		UpdateGenerationNotification();
+#endif
+
+		PCGGraphExecutionLogging::LogGraphExecuteFrameFinished();
+	}
+
+	// Purge things from cache if memory usage is too high
+	if (GraphCache.EnforceMemoryBudget())
+	{
+		CollectGarbage(GARBAGE_COLLECTION_KEEPFLAGS, true);
+	}
+}
+
+void FPCGGraphExecutor::PostTaskExecute(TSharedPtr<FPCGGraphActiveTask> ActiveTaskPtr)
+{
+	TRACE_CPUPROFILER_EVENT_SCOPE(FPCGGraphExecutor::PostTaskExecute);
+	check(ActiveTaskPtr);
+	FPCGGraphActiveTask& ActiveTask = *ActiveTaskPtr;
+	check(ActiveTask.Context);
+	check(ActiveTask.bIsExecutingTask || ActiveTask.bWasCancelled);
+	
+	const bool bTaskWasCancelled = ActiveTask.bWasCancelled;
+#if WITH_EDITOR
+	const bool bTaskWasBypassed = ActiveTask.bIsBypassed;
+#else
+	const bool bTaskWasBypassed = false;
+#endif
+
+	// Needs to be done before QueueNextTasks
+	if (bTaskWasCancelled)
+	{
+		RemoveTaskFromInputSuccessors(ActiveTask.NodeId, ActiveTask.Inputs);
+	}
+
+	bNeedToExecuteTasksEnded = true;
+
+	const UPCGSettingsInterface* ActiveTaskSettingsInterface = ActiveTask.Context->GetInputSettingsInterface();
+	const uint64 InactivePinMask = ActiveTask.Context->OutputData.InactiveOutputPinBitmask;
+
+	if (InactivePinMask != 0 && ActiveTaskSettingsInterface)
+	{
+		const UPCGSettings* ActiveTaskSettings = ActiveTaskSettingsInterface ? ActiveTaskSettingsInterface->GetSettings() : nullptr;
+
+		// If output pins may have been deactivated then perform culling and update information for editor visualization.
+		if (ActiveTaskSettings && ActiveTaskSettings->OutputPinsCanBeDeactivated())
+		{
+			CullInactiveDownstreamNodes(ActiveTask.NodeId, InactivePinMask);
+
+#if WITH_EDITOR
+			SendInactivePinNotification(ActiveTask.Context->Node, ActiveTask.StackContext->GetStack(ActiveTask.StackIndex), InactivePinMask);
+#endif
+		}
+	}
+
+	if (!bTaskWasCancelled && !bTaskWasBypassed)
+	{
+		// Store result in cache as needed - done here because it needs to be done on the main thread
+
+		// Don't store if errors or warnings present
+#if WITH_EDITOR
+		const bool bHasErrorOrWarning = ActiveTask.Context->Node && (ActiveTask.Context->HasVisualLogs());
+#else
+		const bool bHasErrorOrWarning = false;
+#endif
+
+		if (ActiveTaskSettingsInterface && !bHasErrorOrWarning && ActiveTask.Element->IsCacheableInstance(ActiveTaskSettingsInterface))
+		{
+			GraphCache.StoreInCache(ActiveTask.Element.Get(), ActiveTask.Context->DependenciesCrc, ActiveTask.Context->OutputData);
+		}
+	}
+
+	if (ActiveTask.Context->AsyncState.NumAvailableTasks > 0)
+	{
+		CurrentlyUsedThreads -= ActiveTask.Context->AsyncState.NumAvailableTasks;
+	}
+
+#if WITH_EDITOR
+	if (!bTaskWasCancelled)
+	{
+		// No need to add if we are post executing off of main thread
+		if(!IsInGameThread())
+		{
+			PCGGraphExecutor::TScopeLock ScopeLock(CollectGCReferenceTasksLock);
+			CollectGCReferenceTasks.Add(ActiveTaskPtr);
+		}
+
+		PCGGraphExecutor::ExecuteOnGameThread(UE_SOURCE_LOCATION, [this, ActiveTaskPtr]()
+		{
+			ActiveTaskPtr->Element->DebugDisplay(ActiveTaskPtr->Context.Get());
+
+			PCGGraphExecutor::TScopeLock ScopeLock(CollectGCReferenceTasksLock);
+			CollectGCReferenceTasks.Remove(ActiveTaskPtr);
+		});
+		
+		if (UPCGComponent* SourceComponent = ActiveTask.Context->SourceComponent.Get())
+		{
+			if (ActiveTask.StackIndex != INDEX_NONE)
+			{
+				const FPCGStack* Stack = ActiveTask.StackContext->GetStack(ActiveTask.StackIndex);
+				SourceComponent->StoreInspectionData(Stack, ActiveTask.Context->Node, &ActiveTask.Context->Timer, ActiveTask.Context->InputData, ActiveTask.Context->OutputData, /*bUsedCache=*/false);
+			}
+		}
+	}
+#endif
+
+	// If the task is a post execute, then we can safely clear the data after getting it from the results.
+	const bool bTaskIsPostExecute = (ActiveTask.Element == GraphCompiler.GetSharedTrivialPostGraphElement());
+		
+	// Store output in data map.
+	// TODO - investigate if we should avoid doing this if the task was cancelled.
+	StoreResults(ActiveTask.NodeId, ActiveTask.Context->OutputData, bTaskIsPostExecute);
+
+	// Book-keeping
+	QueueNextTasks(ActiveTask.NodeId);
+
+	// Erase from ActiveTasks
+	{
+		PCGGraphExecutor::TScopeLock ScopeLock(LiveTasksLock);
+		if (int32 TaskIndex = ActiveTasks.Find(ActiveTaskPtr); TaskIndex != INDEX_NONE)
+		{
+			// Remove current active task from list
+			ActiveTasks.RemoveAtSwap(TaskIndex);
+		}
+
+		// Make sure we set this to false inside lock
+		ActiveTask.StopExecuting();
+	}
+}
+
+void FPCGGraphExecutor::OnTaskInputsReady(FPCGGraphTask& Task, TArray<FCachedResult*>& OutCachedResults, bool bIsInGameThread)
 {
-	FRWScopeLock ReadWriteLock(TaskOutputsRWLock, SLT_ReadOnly);
-	if (TaskOutputs.Contains(TaskId))
+	TRACE_CPUPROFILER_EVENT_SCOPE(FPCGGraphExecutor::OnTaskInputsReady);
+	TArray<FPCGTaskId> ResultsToMarkAsRead;
+	if (SetupTask(Task, ResultsToMarkAsRead))
 	{
-		FOutputDataInfo& OutputDataInfo = TaskOutputs[TaskId];
-		OutData = OutputDataInfo.DataCollection;
+		check(Task.Element);
+
+		FCachedResult* CachedResult = nullptr;
 
-		if (bClearDataOnGet && ensure(OutputDataInfo.bNeedsManualClear))
+		// Query Cache if we can here, if not the Task will be pushed to ReadyTasks and PrepareForExecute will be called on MainThread if needed
+		if (bIsInGameThread || !Task.Element->CanExecuteOnlyOnMainThread(nullptr))
 		{
-			// Safe here as we are just removing from key, if release of read lock causes TaskOutputs to change
-			// before we acquire the write lock this will result in a no op.
-			ReadWriteLock.ReleaseReadOnlyLockAndAcquireWriteLock_USE_WITH_CAUTION();
-			TaskOutputs.Remove(TaskId);
+			PrepareForExecute(Task, CachedResult, /*bLiveTasksLockAlreadyLocked=*/false);
 		}
 
+		PCGGraphExecutor::TScopeLock ChildScopeLock(LiveTasksLock);
+		if (CachedResult && !CachedResult->bIsBypassed)
+		{
+			check(CachedResult->TaskId != InvalidPCGTaskId);
+			OutCachedResults.Add(CachedResult);
+		}
+		else
+		{
+			ReadyTasks.Emplace(MoveTemp(Task));
+		}
+	}
+
+	MarkInputResults(ResultsToMarkAsRead);
+}
+
+bool FPCGGraphExecutor::SetupTask(FPCGGraphTask& Task, TArray<FPCGTaskId>& ResultsToMarkAsRead)
+{
+	TRACE_CPUPROFILER_EVENT_SCOPE(FPCGGraphExecutor::SetupTask);
+	if (Task.bHasDoneSetup)
+	{
 		return true;
 	}
-	else
+
+	Task.bHasDoneSetup = true;
+
+	// Build input
+	BuildTaskInput(Task, ResultsToMarkAsRead);
+
+	// Initialize the element if needed (required to know whether it will run on the main thread or not)
+	if (!Task.Element)
+	{
+		// Get appropriate settings
+		check(Task.Node);
+		const UPCGSettings* Settings = Task.TaskInput.GetSettings(Task.Node->GetSettings());
+
+		if (Settings)
+		{
+			Task.Element = Settings->GetElement();
+		}
+	}
+
+	// At this point, if the task doesn't have an element, we will never be able to execute it, so we can drop it.
+	if (!Task.Element)
 	{
+		check(!Task.Context);
 		return false;
 	}
+
+	return true;
 }
 
-void FPCGGraphExecutor::Execute()
+void FPCGGraphExecutor::PrepareForExecute(FPCGGraphTask& Task, FCachedResult*& OutCachedResult, bool bLiveTasksLockAlreadyLocked)
 {
-	TRACE_CPUPROFILER_EVENT_SCOPE(FPCGGraphExecutor::Execute);
+	TRACE_CPUPROFILER_EVENT_SCOPE(FPCGGraphExecutor::PrepareForExecute);
+	check(Task.bHasDoneSetup && !Task.bHasDonePrepareForExecute);
+	Task.bHasDonePrepareForExecute = true;
+
+	PCGGraphExecutionLogging::LogTaskExecute(Task);
+
+	// If a task is cacheable and has been cached, then we don't need to create an active task for it unless
+	// there is an execution mode that would prevent us from doing so.
+	const UPCGSettingsInterface* TaskSettingsInterface = Task.TaskInput.GetSettingsInterface(Task.Node ? Task.Node->GetSettingsInterface() : nullptr);
+	const UPCGSettings* TaskSettings = TaskSettingsInterface ? TaskSettingsInterface->GetSettings() : nullptr;
+	const bool bCacheable = Task.Element->IsCacheableInstance(TaskSettingsInterface);
+
+	// Calculate Crc of dependencies (input data Crcs, settings) and use this as the key in the cache lookup
+	FPCGCrc DependenciesCrc;
+	if (TaskSettings && bCacheable)
+	{
+		Task.Element->GetDependenciesCrc(Task.TaskInput, TaskSettings, Task.SourceComponent.Get(), DependenciesCrc);
+	}
+
+	if (!bCacheable)
+	{
+		PCGGraphExecutionLogging::LogTaskExecuteCachingDisabled(Task);
+	}
+		
+	auto GetFromCache = [this, bLiveTasksLockAlreadyLocked, &Task, &DependenciesCrc, &OutCachedResult]() -> bool
+	{
+		// Lock LiveTasksLock if not already done so that we don't have a window where Cache can get cleaned on mainthread and GC runs while we have a unreferenced FPCGDataCollection
+		if(!bLiveTasksLockAlreadyLocked)
+		{
+			LiveTasksLock.Lock();
+		}
+		
+		FCachedResult LocalCachedResult;
+		const bool bFoundInCache = GraphCache.GetFromCache(Task.Node, Task.Element.Get(), DependenciesCrc, Task.SourceComponent.Get(), LocalCachedResult.Output);
+		if (bFoundInCache)
+		{
+			check(!CachingResults.Contains(Task.NodeId));
+			TUniquePtr<FCachedResult>& CachingResultPtr = CachingResults.Add(Task.NodeId, MakeUnique<FCachedResult>());
+			OutCachedResult = CachingResultPtr.Get();
+			OutCachedResult->Output = MoveTemp(LocalCachedResult.Output);
+		}
+
+		if (!bLiveTasksLockAlreadyLocked)
+		{
+			LiveTasksLock.Unlock();
+		}
+
+		return bFoundInCache;
+	};
 
+	const bool bResultAlreadyInCache = bCacheable && DependenciesCrc.IsValid() && GetFromCache();
+#if WITH_EDITOR
+	const bool bNeedsToCreateActiveTask = !bResultAlreadyInCache || TaskSettingsInterface->bDebug;
+#else
+	const bool bNeedsToCreateActiveTask = !bResultAlreadyInCache;
+#endif
+
+	if (!bNeedsToCreateActiveTask)
+	{
+		check(OutCachedResult);
+#if WITH_EDITOR
+		// Doing this now since we're about to modify ReadyTasks potentially reallocating while Task is a reference. 
+		if (UPCGComponent* SourceComponent = Task.SourceComponent.Get())
+		{
+			if (Task.StackIndex != INDEX_NONE)
+			{
+				const FPCGStack* Stack = Task.GetStack();
+				SourceComponent->StoreInspectionData(Stack, Task.Node, nullptr, Task.TaskInput, OutCachedResult->Output, /*bUsedCache=*/true);
+			}
+		}
+#endif
+
+		if (PCGGraphExecutor::CVarDynamicTaskCulling.GetValueOnAnyThread() && TaskSettings && TaskSettings->OutputPinsCanBeDeactivated() && OutCachedResult->Output.InactiveOutputPinBitmask != 0)
+		{
+			OutCachedResult->bDoDynamicTaskCulling = true;
+			OutCachedResult->Stack = Task.GetStack();
+			OutCachedResult->Node = Task.Node;
+		}
+
+		// If the task is a post execute, then we can safely clear the data after getting it from the results.
+		OutCachedResult->bIsPostGraphTask = (Task.Element == GraphCompiler.GetSharedTrivialPostGraphElement());
+
+		// Fast-forward cached result to stored results
+		OutCachedResult->TaskId = Task.NodeId;
+								
+		bNeedToExecuteTasksEnded = true;
+
+		return;
+	}
+		
+	// Allocate context if not previously done
+	check(!Task.Context);
+	
+	{
+		Task.Context = Task.Element->Initialize(Task.TaskInput, Task.SourceComponent, Task.Node);
+		Task.Context->InitializeSettings();
+		Task.Context->TaskId = Task.NodeId;
+		Task.Context->CompiledTaskId = Task.CompiledTaskId;
+		Task.Context->DependenciesCrc = DependenciesCrc;
+		Task.Context->Stack = Task.GetStack();
+
+		// Safe to remove async flags on CombineParams here
+		PCGGraphExecutor::ClearAsyncFlags(Task.CombineParamsAsyncObjects);
+	}
+
+#if WITH_EDITOR
+	if (bResultAlreadyInCache)
+	{
+		Task.bIsBypassed = true;
+		OutCachedResult->bIsBypassed = true;
+		Task.Context->OutputData = OutCachedResult->Output;
+	}
+#endif
+
+	return;
+}
+
+void FPCGGraphExecutor::Execute()
+{
 #if WITH_EDITOR
 	if (PCGSystemSwitches::CVarPausePCGExecution.GetValueOnAnyThread())
 	{
@@ -656,66 +1252,319 @@ void FPCGGraphExecutor::Execute()
 	}
 #endif // WITH_EDITOR
 
-	// Process any newly scheduled graphs to execute
-	ScheduleLock.Lock();
+	if (GetExecuteVersion() == EExecuteVersion::V2)
+	{
+		ExecuteVersion = EExecuteVersion::V2;
+		ExecuteV2();
+	}
+	else
+	{
+		ExecuteVersion = EExecuteVersion::V1;
+		ExecuteV1();
+	}
 
-	for (int32 ScheduledTaskIndex = ScheduledTasks.Num() - 1; ScheduledTaskIndex >= 0; --ScheduledTaskIndex)
+	// Allow switching execute method if we are done
+	if (GetNonScheduledRemainingTaskCount() == 0)
 	{
-		FPCGGraphScheduleTask& ScheduledTask = ScheduledTasks[ScheduledTaskIndex];
+		ExecuteVersion = EExecuteVersion::None;
+	}
+}
 
-		check(ScheduledTask.Tasks.Num() > 0);
-		// Push tasks to the primary task list & build successors map
-		for (FPCGGraphTask& Task : ScheduledTask.Tasks)
-		{
-			FPCGTaskId TaskId = Task.NodeId;
+double FPCGGraphExecutor::GetTickBudgetInSeconds() const
+{
+#if WITH_EDITOR
+	if (GEditor && !GEditor->IsPlaySessionInProgress())
+	{
+		return PCGGraphExecutor::CVarEditorTimePerFrame.GetValueOnAnyThread() / 1000.0;
+	}
+#endif
+
+	return PCGGraphExecutor::CVarTimePerFrame.GetValueOnAnyThread() / 1000.0;
+}
+
+FPCGGraphExecutor::EExecuteVersion FPCGGraphExecutor::GetExecuteVersion() const
+{
+	if (ExecuteVersion == EExecuteVersion::None)
+	{
+		check(GetNonScheduledRemainingTaskCount() == 0);
+		return PCGGraphExecutor::CVarExecuteV2.GetValueOnAnyThread() ? EExecuteVersion::V2 : EExecuteVersion::V1;
+	}
 
-			// TODO: review if it's actually possible for a task with inputs to be ready at this point - it seems very unlikely
-			bool bPushToReady = true;
-			if (Task.Inputs.Num() > 0)
+	return ExecuteVersion;
+}
+
+bool FPCGGraphExecutor::ExecuteScheduling(double EndTime, TSharedPtr<FPCGGraphActiveTask>* OutMainThreadTask)
+{
+	bool bStateChanged = false;
+	const bool bIsInGameThread = IsInGameThread();
+	TRACE_CPUPROFILER_EVENT_SCOPE(FPCGGraphExecutor::ExecuteScheduling);
+		
+	TArray<FCachedResult*> CachedResults;
+	{
+		TRACE_CPUPROFILER_EVENT_SCOPE(FPCGGraphExecutor::ExecuteScheduling::ReadyTasks);
+		PCGGraphExecutor::TScopeLock ScopeLock(LiveTasksLock, !bIsInGameThread);
+		if (ScopeLock.IsLocked())
+		{
+			for (int32 ReadyTaskIndex = ReadyTasks.Num() - 1; ReadyTaskIndex >= 0; --ReadyTaskIndex)
 			{
-				FReadScopeLock ReadLock(TaskOutputsRWLock);
-				for (const FPCGGraphTaskInput& Input : Task.Inputs)
+				FPCGGraphTask& Task = ReadyTasks[ReadyTaskIndex];
+				check(Task.bHasDoneSetup);
+
+				if (!Task.bHasDonePrepareForExecute)
 				{
-					if (!TaskOutputs.Contains(Input.TaskId))
+					if (!bIsInGameThread && Task.Element->CanExecuteOnlyOnMainThread(nullptr))
 					{
-						TaskSuccessors.FindOrAdd(Input.TaskId).Add(TaskId);
-						bPushToReady = false;
+						continue;
 					}
+
+					FCachedResult* CachedResult = nullptr;
+					PrepareForExecute(Task, CachedResult, /*bLiveTasksLockAlreadyLocked=*/true);
+					check(Task.bHasDonePrepareForExecute);
+
+					if (CachedResult && !CachedResult->bIsBypassed)
+					{
+						check(CachedResult->TaskId != InvalidPCGTaskId);
+						CachedResults.Add(MoveTemp(CachedResult));
+						ReadyTasks.RemoveAtSwap(ReadyTaskIndex);
+						bStateChanged = true;
+						continue;
+					}
+
+					check(Task.Context);
 				}
+
+				// Validate that we can start this task now
+				FPCGGraphActiveTask& ActiveTask = *ActiveTasks.Emplace_GetRef(MakeShared<FPCGGraphActiveTask>());
+				ActiveTask.Inputs = MoveTemp(Task.Inputs);
+				ActiveTask.Element = Task.Element;
+				ActiveTask.NodeId = Task.NodeId;
+				ActiveTask.Context = TUniquePtr<FPCGContext>(Task.Context);
+				ActiveTask.StackIndex = Task.StackIndex;
+				ActiveTask.StackContext = Task.StackContext;
+#if WITH_EDITOR
+				ActiveTask.bIsBypassed = Task.bIsBypassed;
+#endif
+				// Move the task up front if it needs to run on the main thread
+				ReadyTasks.RemoveAtSwap(ReadyTaskIndex);
+				bStateChanged = true;
 			}
+		}
+	}
+		
+	ProcessCachedResults(CachedResults);
 
-			// Automatically push inputless/already satisfied tasks to the ready queue
-			if (bPushToReady)
+	// Dispatch Tasks
+	{
+		PCGGraphExecutor::TScopeLock ScopeLock(LiveTasksLock, !bIsInGameThread);
+		if (ScopeLock.IsLocked())
+		{
+			if (bIsInGameThread && OutMainThreadTask)
 			{
-				ReadyTasks.Emplace(MoveTemp(Task));
+				TRACE_CPUPROFILER_EVENT_SCOPE(FPCGGraphExecutor::Execute::SelectMainThreadTask);
+				// Execute main thread task
+				if (!ActiveTasks.IsEmpty())
+				{
+					// Sort tasks so that we might find an main thread executable task at index 0
+					ActiveTasks.Sort([](const TSharedPtr<FPCGGraphActiveTask>& Left, const TSharedPtr<FPCGGraphActiveTask>& Right)
+					{
+#if WITH_EDITOR
+						if (Left->bIsBypassed != Right->bIsBypassed) { return Left->bIsBypassed; }
+#endif
+						if (Left->Context->bIsPaused != Right->Context->bIsPaused) { return Right->Context->bIsPaused; }
+
+						if (Left->bIsExecutingTask != Right->bIsExecutingTask) { return Right->bIsExecutingTask; }
+
+						const bool LeftMainThreadOnly = Left->Element->CanExecuteOnlyOnMainThread(Left->Context.Get());
+						const bool RightMainThreadOnly = Right->Element->CanExecuteOnlyOnMainThread(Right->Context.Get());
+
+						if (LeftMainThreadOnly != RightMainThreadOnly) { return LeftMainThreadOnly; }
+
+						return Left->NodeId < Right->NodeId;
+					});
+
+					// Sort should give us best target to execute
+					TSharedPtr<FPCGGraphActiveTask>& ActiveTask = ActiveTasks[0];
+
+					// Skip already running task
+					if (!ActiveTask->bIsExecutingTask && !ActiveTask->Context->bIsPaused)
+					{
+						ActiveTask->StartExecuting();
+						ActiveTask->Context->AsyncState.NumAvailableTasks = -1;
+						ActiveTask->Context->AsyncState.bIsRunningOnMainThread = true;
+						ActiveTask->Context->AsyncState.bIsRunningOutOfTick = false;
+						ActiveTask->Context->AsyncState.EndTime = EndTime;
+						*OutMainThreadTask = ActiveTask;
+						bStateChanged = true;
+					}
+				}
 			}
-			else
+
+			// PCGGraphExecutor::CVarMaxWorkerTasks.GetValueOnAnyThread()+1 because the mainthread doesn't count
+			if (PCGGraphExecutor::CVarGraphMultithreading.GetValueOnAnyThread() && PCGGraphExecutor::CVarMaxWorkerTasks.GetValueOnAnyThread()+1 > FPCGGraphActiveTask::NumExecuting)
 			{
-				Tasks.Add(TaskId, MoveTemp(Task));
+				TRACE_CPUPROFILER_EVENT_SCOPE(FPCGGraphExecutor::ExecuteScheduling::LaunchTasks);
+				for (int32 ExecutionIndex = 0; ExecutionIndex < ActiveTasks.Num(); ++ExecutionIndex)
+				{
+					TSharedPtr<FPCGGraphActiveTask>& ActiveTask = ActiveTasks[ExecutionIndex];
+
+					if (PCGGraphExecutor::CVarMaxWorkerTasks.GetValueOnAnyThread()+1 <= FPCGGraphActiveTask::NumExecuting)
+					{
+						break;
+					}
+
+					// Running task
+					if (ActiveTask->bIsExecutingTask || ActiveTask->Context->bIsPaused)
+					{
+						continue;
+					}
+
+					// Main thread only
+					if (ActiveTask->Element->CanExecuteOnlyOnMainThread(ActiveTask->Context.Get()))
+					{
+						continue;
+					}
+#if WITH_EDITOR
+					if (!ActiveTask->bIsBypassed)
+#endif
+					{
+						check(ActiveTask->Context->CurrentPhase != EPCGExecutionPhase::Done);
+						ActiveTask->Context->AsyncState.NumAvailableTasks = -1;
+						ActiveTask->Context->AsyncState.EndTime = EndTime;
+						ActiveTask->Context->AsyncState.bIsRunningOnMainThread = false;
+						ActiveTask->Context->AsyncState.bIsRunningOutOfTick = true;
+
+						// Capture copy of SharedPtr so task stays valid while Async task is alive
+						ActiveTask->StartExecuting();
+						ActiveTask->ExecutingTask = UE::Tasks::Launch(UE_SOURCE_LOCATION, [this, ActiveTask]() -> bool
+						{
+							TRACE_CPUPROFILER_EVENT_SCOPE(FPCGGraphExecutor::ExecuteAsyncTask);
+
+							const bool bIsDone = ActiveTask->bWasCancelled || ActiveTask->Element->Execute(ActiveTask->Context.Get());
+							const bool bIsPaused = ActiveTask->Context->bIsPaused;
+
+							if (bIsDone)
+							{
+								PostTaskExecute(ActiveTask);
+							}
+							else
+							{
+								PCGGraphExecutor::TScopeLock ScopeLock(LiveTasksLock);
+								ActiveTask->StopExecuting();
+							}
+
+							if (bIsDone || bIsPaused)
+							{
+								// Budget is provided by CVarAsyncOutOfTickBudgetInMilliseconds
+								ExecuteScheduling(-1.0);
+							}
+
+							return bIsDone;
+						});
+						bStateChanged = true;
+					}
+				}
 			}
 		}
 	}
 
-	if (!ScheduledTasks.IsEmpty())
-	{
-		PCGGraphExecutionLogging::LogGraphPostSchedule(Tasks, TaskSuccessors);
-	}
+	bStateChanged |= ProcessScheduledTasks();
 
-	ScheduledTasks.Reset();
+	return bStateChanged;
+}
 
-	ScheduleLock.Unlock();
+/**
+ * Graph Execution Pipeline:
+ * 
+ * - New Tasks first enter the execution pipeline through [ScheduledTasks]
+ * - Each [ScheduledTasks] is a list of Task with Input dependencies on other Tasks
+ * - [ScheduleTasks] are processed through [FPCGraphExecutor::ProcessScheduledTasks] which can happen on GameThread/WorkerThread 
+ * - On WorkerThreads [FPCGraphExecutor::ProcessScheduledTasks] will try to acquire the Lock and if it can't will just skip through, on GameThread it will wait until it can grab the lock (ScheduleLock)
+*	 The reason why WorkerThreads will only Try to lock is because [FPCGGraphExecutor::ExecuteScheduling] runs in the same UE::Task as the Node execution and we don't want to make those spin for too long. (Let's call it opportunistic scheduling)
+* 
+ * - [FPCGGraphExecutor::ProcessScheduledTasks] will move [ScheduledTasks] into one of [Tasks] if task still waiting on other task inputs or [ReadyTasks] if tasks is ready to execute
+ * 
+ * - When moving a [ScheduledTasks] into [ReadyTasks] we call [FPCGGraphExecutor::SetupTask] which will build the task inputs and its FPCGGraphElement used for execution.
+ *	 It will also try and call [FPCGGraphExecutor::PrepareForExecute] on the task if it is on the GameThread or if it is on a WorkerThread and the FPCGGraphElement allows it [!FPCGGraphElement::CanExecuteOnlyOnMainThread(nullptr)]
+ * - [FPCGGraphExecutor::PrepareForExecute] will query the Cache and might return a valid Output or it will create the FPCGContext for the task so that it can be executed
+ *
+ * - [ReadyTasks] are processed through [FPCGGraphExecutor::ExecuteScheduling] which can happen on GameThread/WorkerThread
+ * - On WorkerThreads [FPCGGraphExecutor::ExecuteScheduling] will try to acquire the Lock and if it can't will just skip through, on GameThread it will wait until it can grab the lock. (LiveTasksLock) (same reason as above)
+ *	 
+ * - [ReadyTasks] will be moved to [ActiveTasks] through [FPCGGraphExecutor::ExecuteScheduling]
+ * - This will only happen if Task.bHasDonePrepareForExecute is true or if we can run [FPCGGraphExecutor::PrepareForExecute] on it. (this can always happen on GameThread but not always on WorkerThreads as explained before)
+ * 
+ * - [ActiveTasks] are also processed in [FPCGGraphExecutor::ExecuteScheduling]
+ * - On the GameThread the [ActiveTasks] will get sorted so that tasks that can execute on the GameThread only will be first in the array
+ * - After the sort, GameThread will always try to execute [ActiveTasks[0]], it will do this [FPCGGraphExecutor::ExecuteScheduling] while it exhausts its GameThread budget
+ * - Both GameThread/WorkerThread will also spawn new UE::Tasks for the remaining [ActiveTasks] that are allowed to run outside of the GameThread still happens in [FPCGGraphExecutor::ExecuteScheduling]
+ * 
+ * - When an ActiveTasks Execute call returns and it returns true, it means that the task is done. In this case we call [FPCGGraphExecutor::StoreResults] followed by [FPCGGraphExecutor::QueueNextTasks] this will potentially move some [Tasks] to [ReadyTasks]
+ *   so that new tasks can get executed now that their inputs have been filled
+ * 
+ * - Cache information: When we call [FPCGGraphExecutor::PrepareForExecute] is called on a task, if a Cache result is found for this task then we can directly call [FPCGGraphExecutor::StoreResults] and [FPCGGraphExecutor::QueueNextTasks] on it skipping execution [FPCGGraphExecutor::ProcessCachedResults]
+ * - An exception to this is when the FPCGElement is marked for debugging then we move the task through the normal execution pipeline but set a flag on it so that the actual Execute call can be skipped
+ * - [CachingResults] exists only to keep GC references to PCG Data that has been read from the Cache but hasn't yet been moved to the [TaskOutputs]
+ * 
+ * - [CVars]
+ *   pcg.Graph.ExecuteV2 : To enable this pipeline (this will be removed once stability is considered on par with previous impl.)
+ *   pcg.GraphMultithreading : Enables multithreading on this execution pipeline, without that flag most multithreading happens only in PCGAsync.h/.cpp
+ *	 pcg.Graph.MaxWorkerTasks : Limits the number of UE::Tasks that [FPCGGraphExecutor::ExecuteScheduling] can launch
+ *   pcg.DisableAsyncTimeSlicing : Disables timeslicing in UE::Tasks launched by [FPCGGraphExecutor::ExecuteScheduling]
+ *   pcg.Async.OutOfTickBudgetInMilliseconds : Timeslicing budget for UE::Tasks launched by [FPCGGraphExecutor::ExecuteScheduling]
+ *   pcg.DisableAsyncTimeSlicingOnGameThread : Disables timeslicing on GameThread execution
+ */
+void FPCGGraphExecutor::ExecuteV2()
+{
+	TRACE_CPUPROFILER_EVENT_SCOPE(FPCGGraphExecutor::Execute);
 
-	// TODO: add optimization phase if we've added new graph(s)/tasks to be executed
+	ProcessScheduledTasks();
+#if WITH_EDITOR
+	// Update Notifications before capturing StartTime so that this call doesn't eat up task budget
+	UpdateGenerationNotification();
+#endif
+				
+	bool bContinueExecute = GetNonScheduledRemainingTaskCount() > 0;
+	bool bFirstLoop = true;
 
-	// This is a safeguard to check if we're in a stuck state
-	if (ReadyTasks.Num() == 0 && ActiveTasks.Num() == 0 && SleepingTasks.Num() == 0 && Tasks.Num() > 0)
+	const double StartTime = FPlatformTime::Seconds();
+	const double EndTime = StartTime + GetTickBudgetInSeconds();
+		
+	while ((bFirstLoop || (FPlatformTime::Seconds() < EndTime)) && bContinueExecute)
 	{
-		UE_LOG(LogPCG, Error, TEXT("PCG Graph executor error: tasks are in a deadlocked state. Will drop all tasks."));
-		ClearAllTasks();
+		bFirstLoop = false;
+		
+		TSharedPtr<FPCGGraphActiveTask> MainThreadTask;
+		bContinueExecute = ExecuteScheduling(EndTime, &MainThreadTask);
+								
+		if (MainThreadTask)
+		{
+			TRACE_CPUPROFILER_EVENT_SCOPE(FPCGGraphExecutor::Execute::ExecuteMainThreadTask);
+#if WITH_EDITOR
+			if (MainThreadTask->bIsBypassed || MainThreadTask->bWasCancelled || MainThreadTask->Element->Execute(MainThreadTask->Context.Get()))
+#else
+			if (MainThreadTask->bWasCancelled || MainThreadTask->Element->Execute(MainThreadTask->Context.Get()))
+#endif
+			{
+				PostTaskExecute(MainThreadTask);
+			}
+			else
+			{
+				PCGGraphExecutor::TScopeLock ScopeLock(LiveTasksLock);
+				MainThreadTask->StopExecuting();
+			}
+		}
 	}
 
+	ExecuteTasksEnded();
+}
+
+void FPCGGraphExecutor::ExecuteV1()
+{
+	TRACE_CPUPROFILER_EVENT_SCOPE(FPCGGraphExecutor::Execute);
+
+	ProcessScheduledTasks();
+		
 	// TODO: change this if we support tasks that are not framebound
-	bool bAnyTaskEnded = false;
 	bool bHasAlreadyCheckedSleepingTasks = false;
 
 #if WITH_EDITOR
@@ -723,22 +1572,12 @@ void FPCGGraphExecutor::Execute()
 	UpdateGenerationNotification();
 #endif
 
-	const double StartTime = FPlatformTime::Seconds();
-
-	double VarTimePerFrame = PCGGraphExecutor::CVarTimePerFrame.GetValueOnAnyThread() / 1000.0;
-
-#if WITH_EDITOR
-	if (GEditor && !GEditor->IsPlaySessionInProgress())
-	{
-		VarTimePerFrame = PCGGraphExecutor::CVarEditorTimePerFrame.GetValueOnAnyThread() / 1000.0;
-	}
-#endif
-
-	const double EndTime = StartTime + VarTimePerFrame;
 	const float MaxPercentageOfThreadsToUse = FMath::Clamp(CVarMaxPercentageOfThreadsToUse.GetValueOnAnyThread(), 0.0f, 1.0f);
 	const int32 MaxNumThreads = FMath::Max(0, FMath::Min((int32)(FPlatformMisc::NumberOfCoresIncludingHyperthreads() * MaxPercentageOfThreadsToUse), CVarMaxNumTasks.GetValueOnAnyThread() - 1));
 	const bool bAllowMultiDispatch = PCGGraphExecutor::CVarGraphMultithreading.GetValueOnAnyThread();
-	const bool bDynamicTaskCulling = PCGGraphExecutor::CVarDynamicTaskCulling.GetValueOnAnyThread();
+	
+	const double StartTime = FPlatformTime::Seconds();
+	const double EndTime = StartTime + GetTickBudgetInSeconds();
 
 	while (ReadyTasks.Num() > 0 || ActiveTasks.Num() > 0 || (!bHasAlreadyCheckedSleepingTasks && SleepingTasks.Num() > 0))
 	{
@@ -747,7 +1586,7 @@ void FPCGGraphExecutor::Execute()
 		bHasAlreadyCheckedSleepingTasks = ReadyTasks.Num() == 0 && ActiveTasks.Num() == 0 && SleepingTasks.Num() > 0;
 
 		// First: if we have free resources, move ready tasks to the active tasks
-		bool bMainThreadAvailable = (ActiveTasks.Num() == 0 || !ActiveTasks[0].Element->CanExecuteOnlyOnMainThread(ActiveTasks[0].Context.Get()));
+		bool bMainThreadAvailable = (ActiveTasks.Num() == 0 || !ActiveTasks[0]->Element->CanExecuteOnlyOnMainThread(ActiveTasks[0]->Context.Get()));
 		int32 NumAvailableThreads = FMath::Max(0, MaxNumThreads - CurrentlyUsedThreads);
 
 		const bool bMainThreadWasAvailable = bMainThreadAvailable;
@@ -787,15 +1626,15 @@ void FPCGGraphExecutor::Execute()
 					break;
 				}
 
-				FPCGGraphActiveTask& SleepingTask = SleepingTasks[SleepingTaskIndex];
+				TSharedPtr<FPCGGraphActiveTask>& SleepingTask = SleepingTasks[SleepingTaskIndex];
 
-				if (SleepingTask.Context->bIsPaused)
+				if (SleepingTask->Context->bIsPaused)
 				{
 					continue; // still sleeping
 				}
 
 				// Validate that we can start this task now
-				const bool bIsMainThreadTask = SleepingTask.Element->CanExecuteOnlyOnMainThread(SleepingTask.Context.Get());
+				const bool bIsMainThreadTask = SleepingTask->Element->CanExecuteOnlyOnMainThread(SleepingTask->Context.Get());
 
 				if (!bIsMainThreadTask || bMainThreadAvailable)
 				{
@@ -821,129 +1660,14 @@ void FPCGGraphExecutor::Execute()
 				}
 
 				FPCGGraphTask& Task = ReadyTasks[ReadyTaskIndex];
-
-				if (!Task.bHasDoneSetup)
-				{
-					Task.bHasDoneSetup = true;
-
-					// Build input
-					FPCGDataCollection TaskInput;
-					BuildTaskInput(Task, TaskInput);
-
-					// Initialize the element if needed (required to know whether it will run on the main thread or not)
-					if (!Task.Element)
-					{
-						// Get appropriate settings
-						check(Task.Node);
-						const UPCGSettings* Settings = TaskInput.GetSettings(Task.Node->GetSettings());
-
-						if (Settings)
-						{
-							Task.Element = Settings->GetElement();
-						}
-					}
-
-					// At this point, if the task doesn't have an element, we will never be able to execute it, so we can drop it.
-					if (!Task.Element)
-					{
-						check(!Task.Context);
-						ReadyTasks.RemoveAtSwap(ReadyTaskIndex);
-						continue;
-					}
-
-					PCGGraphExecutionLogging::LogTaskExecute(Task);
-
-					// If a task is cacheable and has been cached, then we don't need to create an active task for it unless
-					// there is an execution mode that would prevent us from doing so.
-					const UPCGSettingsInterface* TaskSettingsInterface = TaskInput.GetSettingsInterface(Task.Node ? Task.Node->GetSettingsInterface() : nullptr);
-					const UPCGSettings* TaskSettings = TaskSettingsInterface ? TaskSettingsInterface->GetSettings() : nullptr;
-					const bool bCacheable = Task.Element->IsCacheableInstance(TaskSettingsInterface);
-
-					// Calculate Crc of dependencies (input data Crcs, settings) and use this as the key in the cache lookup
-					FPCGCrc DependenciesCrc;
-					if (TaskSettings && bCacheable)
-					{
-						Task.Element->GetDependenciesCrc(TaskInput, TaskSettings, Task.SourceComponent.Get(), DependenciesCrc);
-					}
-
-					if (!bCacheable)
-					{
-						PCGGraphExecutionLogging::LogTaskExecuteCachingDisabled(Task);
-					}
-
-					FPCGDataCollection CachedOutput;
-					const bool bResultAlreadyInCache = bCacheable && DependenciesCrc.IsValid() && GraphCache.GetFromCache(Task.Node, Task.Element.Get(), DependenciesCrc, Task.SourceComponent.Get(), CachedOutput);
-#if WITH_EDITOR
-					const bool bNeedsToCreateActiveTask = !bResultAlreadyInCache || TaskSettingsInterface->bDebug;
-#else
-					const bool bNeedsToCreateActiveTask = !bResultAlreadyInCache;
-#endif
-
-					if (!bNeedsToCreateActiveTask)
-					{
-						// Move Task out of array in case CullInactiveDownstreamNodes ends up queueing new tasks and resizes ReadyTasks (which can cause reallocation)
-						FPCGGraphTask SkippedTask = MoveTemp(Task);
-						ReadyTasks.RemoveAtSwap(ReadyTaskIndex);
-#if WITH_EDITOR
-						// Doing this now since we're about to modify ReadyTasks potentially reallocating while Task is a reference. 
-						if (UPCGComponent* SourceComponent = SkippedTask.SourceComponent.Get())
-						{
-							if (SkippedTask.StackIndex != INDEX_NONE)
-							{
-								const FPCGStack* Stack = SkippedTask.GetStack();
-								SourceComponent->StoreInspectionData(Stack, SkippedTask.Node, nullptr, TaskInput, CachedOutput, /*bUsedCache=*/true);
-							}
-						}
-#endif
-
-						if (bDynamicTaskCulling && TaskSettings && TaskSettings->OutputPinsCanBeDeactivated() && CachedOutput.InactiveOutputPinBitmask != 0)
-						{
-							CullInactiveDownstreamNodes(SkippedTask.NodeId, CachedOutput.InactiveOutputPinBitmask);
-
-#if WITH_EDITOR
-							SendInactivePinNotification(SkippedTask.Node, SkippedTask.GetStack(), CachedOutput.InactiveOutputPinBitmask);
-#endif
-						}
-
-						// If the task is a post execute, then we can safely clear the data after getting it from the results.
-						const bool bTaskIsPostExecute = (SkippedTask.Element == GraphCompiler.GetSharedTrivialPostGraphElement());
-						// Fast-forward cached result to stored results
-						FPCGTaskId SkippedTaskId = SkippedTask.NodeId;
-						StoreResults(SkippedTaskId, CachedOutput, bTaskIsPostExecute);
-						delete SkippedTask.Context;
-
-						QueueNextTasks(SkippedTaskId);
-						bAnyTaskEnded = true;
-
-						continue;
-					}
-
-					// Allocate context if not previously done
-					if (!Task.Context)
-					{
-						Task.Context = Task.Element->Initialize(TaskInput, Task.SourceComponent, Task.Node);
-						Task.Context->InitializeSettings();
-						Task.Context->TaskId = Task.NodeId;
-						Task.Context->CompiledTaskId = Task.CompiledTaskId;
-						Task.Context->DependenciesCrc = DependenciesCrc;
-						Task.Context->Stack = Task.GetStack();
-					}
-
-#if WITH_EDITOR
-					if (bResultAlreadyInCache)
-					{
-						Task.bIsBypassed = true;
-						Task.Context->OutputData = CachedOutput;
-					}
-#endif
-				}
+				check(Task.bHasDoneSetup && Task.bHasDonePrepareForExecute);
 
 				// Validate that we can start this task now
 				const bool bIsMainThreadTask = Task.Element->CanExecuteOnlyOnMainThread(Task.Context);
 
 				if (!bIsMainThreadTask || bMainThreadAvailable)
 				{
-					FPCGGraphActiveTask& ActiveTask = ActiveTasks.Emplace_GetRef();
+					FPCGGraphActiveTask& ActiveTask = *ActiveTasks.Emplace_GetRef(MakeShared<FPCGGraphActiveTask>());
 					ActiveTask.Inputs = MoveTemp(Task.Inputs);
 					ActiveTask.Element = Task.Element;
 					ActiveTask.NodeId = Task.NodeId;
@@ -980,7 +1704,7 @@ void FPCGGraphExecutor::Execute()
 
 			for (int32 ExecutionIndex = 0; ExecutionIndex < ActiveTasks.Num(); ++ExecutionIndex)
 			{
-				FPCGGraphActiveTask& ActiveTask = ActiveTasks[ExecutionIndex];
+				FPCGGraphActiveTask& ActiveTask = *ActiveTasks[ExecutionIndex];
 
 				// Tasks that were already launched already have assigned tasks, so don't touch them
 				if (ActiveTask.Context->AsyncState.NumAvailableTasks == 0)
@@ -991,14 +1715,11 @@ void FPCGGraphExecutor::Execute()
 			}
 		}
 
-		// Dispatch async tasks
-		TMap<int32, TFuture<bool>> Futures;
-
 		{
-			TRACE_CPUPROFILER_EVENT_SCOPE(FPCGGraphExecutor::Execute::StartFutures);
+			TRACE_CPUPROFILER_EVENT_SCOPE(FPCGGraphExecutor::Execute::StartTasks);
 			for (int32 ExecutionIndex = 1; ExecutionIndex < ActiveTasks.Num(); ++ExecutionIndex)
 			{
-				FPCGGraphActiveTask& ActiveTask = ActiveTasks[ExecutionIndex];
+				FPCGGraphActiveTask& ActiveTask = *ActiveTasks[ExecutionIndex];
 				check(!ActiveTask.Context->bIsPaused);
 
 	#if WITH_EDITOR
@@ -1008,95 +1729,15 @@ void FPCGGraphExecutor::Execute()
 					check(!ActiveTask.Element->CanExecuteOnlyOnMainThread(ActiveTask.Context.Get()));
 					ActiveTask.Context->AsyncState.EndTime = EndTime;
 					ActiveTask.Context->AsyncState.bIsRunningOnMainThread = false;
-
-					Futures.Emplace(ExecutionIndex, Async(EAsyncExecution::ThreadPool, [&ActiveTask]()
+					ActiveTask.StartExecuting();
+					ActiveTask.ExecutingTask = UE::Tasks::Launch(UE_SOURCE_LOCATION, [&ActiveTask]()
 					{
 						return ActiveTask.Element->Execute(ActiveTask.Context.Get());
-					}));
+					});
 				}
 			}
 		}
-
-		auto PostTaskExecute = [this, &bAnyTaskEnded](int32 TaskIndex)
-		{
-			FPCGGraphActiveTask& ActiveTask = ActiveTasks[TaskIndex];
-			check(ActiveTask.Context);
-
-			const UPCGSettingsInterface* ActiveTaskSettingsInterface = ActiveTask.Context->GetInputSettingsInterface();
-			const uint64 InactivePinMask = ActiveTask.Context->OutputData.InactiveOutputPinBitmask;
-
-			if (InactivePinMask != 0 && ActiveTaskSettingsInterface)
-			{
-				const UPCGSettings* ActiveTaskSettings = ActiveTaskSettingsInterface ? ActiveTaskSettingsInterface->GetSettings() : nullptr;
-
-				// If output pins may have been deactivated then perform culling and update information for editor visualization.
-				if (ActiveTaskSettings && ActiveTaskSettings->OutputPinsCanBeDeactivated())
-				{
-					CullInactiveDownstreamNodes(ActiveTask.NodeId, InactivePinMask);
-
-#if WITH_EDITOR
-					SendInactivePinNotification(ActiveTask.Context->Node, ActiveTask.StackContext->GetStack(ActiveTask.StackIndex), InactivePinMask);
-#endif
-				}
-			}
-
-#if WITH_EDITOR
-			if (!ActiveTask.bWasCancelled && !ActiveTask.bIsBypassed)
-#else
-			if (!ActiveTask.bWasCancelled)
-#endif
-			{
-				// Store result in cache as needed - done here because it needs to be done on the main thread
-
-				// Don't store if errors or warnings present
-#if WITH_EDITOR
-				const bool bHasErrorOrWarning = ActiveTask.Context->Node && (ActiveTask.Context->HasVisualLogs());
-#else
-				const bool bHasErrorOrWarning = false;
-#endif
-
-				if (ActiveTaskSettingsInterface && !bHasErrorOrWarning && ActiveTask.Element->IsCacheableInstance(ActiveTaskSettingsInterface))
-				{
-					GraphCache.StoreInCache(ActiveTask.Element.Get(), ActiveTask.Context->DependenciesCrc, ActiveTask.Context->OutputData);
-				}
-			}
-
-			check(ActiveTask.Context->AsyncState.NumAvailableTasks >= 0);
-			CurrentlyUsedThreads -= ActiveTask.Context->AsyncState.NumAvailableTasks;
-
-#if WITH_EDITOR
-			if (!ActiveTask.bWasCancelled)
-			{
-				// Execute debug display code as needed - done here because it needs to be done on the main thread
-				// Additional note: this needs to be executed before the StoreResults since debugging might cancel further tasks
-				ActiveTask.Element->DebugDisplay(ActiveTask.Context.Get());
-
-				if (UPCGComponent* SourceComponent = ActiveTask.Context->SourceComponent.Get())
-				{
-					if (ActiveTask.StackIndex != INDEX_NONE)
-					{
-						const FPCGStack* Stack = ActiveTask.StackContext->GetStack(ActiveTask.StackIndex);
-						SourceComponent->StoreInspectionData(Stack, ActiveTask.Context->Node, &ActiveTask.Context->Timer, ActiveTask.Context->InputData, ActiveTask.Context->OutputData, /*bUsedCache=*/false);
-					}
-				}
-			}
-#endif
-
-			// If the task is a post execute, then we can safely clear the data after getting it from the results.
-			const bool bTaskIsPostExecute = (ActiveTask.Element == GraphCompiler.GetSharedTrivialPostGraphElement());
-
-			// Store output in data map.
-			// TODO - investigate if we should avoid doing this if the task was cancelled.
-			StoreResults(ActiveTask.NodeId, ActiveTask.Context->OutputData, bTaskIsPostExecute);
-
-			// Book-keeping
-			QueueNextTasks(ActiveTask.NodeId);
-			bAnyTaskEnded = true;
-
-			// Remove current active task from list
-			ActiveTasks.RemoveAtSwap(TaskIndex);
-		};
-
+				
 		{
 			TRACE_CPUPROFILER_EVENT_SCOPE(FPCGGraphExecutor::Execute::ExecuteTasks);
 			bool bMainTaskDone = false;
@@ -1104,10 +1745,11 @@ void FPCGGraphExecutor::Execute()
 			if (!ActiveTasks.IsEmpty())
 			{
 				TRACE_CPUPROFILER_EVENT_SCOPE(FPCGGraphExecutor::Execute::ExecuteTasks::MainThreadTask);
-				FPCGGraphActiveTask& MainThreadTask = ActiveTasks[0];
+				FPCGGraphActiveTask& MainThreadTask = *ActiveTasks[0];
 				check(!MainThreadTask.Context->bIsPaused);
 				MainThreadTask.Context->AsyncState.EndTime = EndTime;
 				MainThreadTask.Context->AsyncState.bIsRunningOnMainThread = true;
+				MainThreadTask.StartExecuting();
 
 #if WITH_EDITOR
 				if(MainThreadTask.bIsBypassed || MainThreadTask.bWasCancelled || MainThreadTask.Element->Execute(MainThreadTask.Context.Get()))
@@ -1120,29 +1762,30 @@ void FPCGGraphExecutor::Execute()
 			}
 
 			{
-				TRACE_CPUPROFILER_EVENT_SCOPE(FPCGGraphExecutor::Execute::ExecuteTasks::WaitForFutures);
+				TRACE_CPUPROFILER_EVENT_SCOPE(FPCGGraphExecutor::Execute::ExecuteTasks::WaitForTasks);
 
 				// Then wait after all futures - start from the back so we can more easily manage the ActiveTasks array
 				for (int32 ExecutionIndex = ActiveTasks.Num() - 1; ExecutionIndex > 0; --ExecutionIndex)
 				{
+					FPCGGraphActiveTask& ActiveTask = *ActiveTasks[ExecutionIndex];
 					bool bTaskDone = false;
 					// Wait on the future if any
-					if (TFuture<bool>* Future = Futures.Find(ExecutionIndex))
+					if (ActiveTask.ExecutingTask.IsValid())
 					{
-						Future->Wait();
-						bTaskDone = Future->Get();
+						ActiveTask.ExecutingTask.Wait();
+						bTaskDone = ActiveTask.ExecutingTask.GetResult();
 					}
 
-					if (bTaskDone || ActiveTasks[ExecutionIndex].bWasCancelled)
+					if (bTaskDone || ActiveTasks[ExecutionIndex]->bWasCancelled)
 					{
-						PostTaskExecute(ExecutionIndex);
+						PostTaskExecute(ActiveTasks[ExecutionIndex]);
 					}
 				}
 			}
 
 			if (bMainTaskDone)
 			{
-				PostTaskExecute(0);
+				PostTaskExecute(ActiveTasks[0]);
 			}
 		}
 
@@ -1152,20 +1795,25 @@ void FPCGGraphExecutor::Execute()
 			TRACE_CPUPROFILER_EVENT_SCOPE(FPCGGraphExecutor::Execute::ExecuteTasks::CheckSleepingTasks);
 			for (int32 ActiveTaskIndex = ActiveTasks.Num() - 1; ActiveTaskIndex >= 0; --ActiveTaskIndex)
 			{
-				FPCGGraphActiveTask& ActiveTask = ActiveTasks[ActiveTaskIndex];
-				
-				// Any task that asks to be paused or now needs to run on the main thread but doesn't have that slot currently will be moved to the sleeping queue
-				const bool bTaskShouldBePutAside = (ActiveTask.Context->bIsPaused || (ActiveTaskIndex > 0 && ActiveTask.Element->CanExecuteOnlyOnMainThread(ActiveTask.Context.Get())));
+				// Any paused tasks at that point should relinquish their resources
+				TSharedPtr<FPCGGraphActiveTask> ActiveTask = ActiveTasks[ActiveTaskIndex];
+				check(ActiveTask->bIsExecutingTask);
 
-				if(bTaskShouldBePutAside)
+				// Any task that asks to be paused or now needs to run on the main thread but doesn't have that slot currently will be moved to the sleeping queue
+				const bool bTaskShouldBePutAside = (ActiveTask->Context->bIsPaused || (ActiveTaskIndex > 0 && ActiveTask->Element->CanExecuteOnlyOnMainThread(ActiveTask->Context.Get())));
+				if (bTaskShouldBePutAside)
 				{
-					check(ActiveTask.Context->AsyncState.NumAvailableTasks > 0);
-					CurrentlyUsedThreads -= ActiveTask.Context->AsyncState.NumAvailableTasks;
-					ActiveTask.Context->AsyncState.NumAvailableTasks = 0;
+					if (ActiveTask->Context->AsyncState.NumAvailableTasks > 0)
+					{
+						CurrentlyUsedThreads -= ActiveTask->Context->AsyncState.NumAvailableTasks;
+					}
+					ActiveTask->Context->AsyncState.NumAvailableTasks = 0;
 
-					SleepingTasks.Emplace(MoveTemp(ActiveTask));
+					SleepingTasks.Emplace(ActiveTask);
 					ActiveTasks.RemoveAtSwap(ActiveTaskIndex);
 				}
+
+				ActiveTask->StopExecuting();
 			}
 		}
 
@@ -1177,41 +1825,16 @@ void FPCGGraphExecutor::Execute()
 		}
 	}
 
-	if (bAnyTaskEnded)
-	{
-		// Nothing left to do; we'll release everything here.
-		// TODO: this is fine and will make sure any intermediate data is properly
-		// garbage collected, however, this goes a bit against our goals if we want to
-		// keep a cache of intermediate results.
-		if (GetNonScheduledRemainingTaskCount() == 0)
-		{
-			if (!ensure(TaskSuccessors.IsEmpty()))
-			{
-				TaskSuccessors.Reset();
-			}
-			
-			ClearResults();
-
-#if WITH_EDITOR
-			// Call the notification update here to prevent it from sticking around - needed because we early out before this
-			UpdateGenerationNotification();
-#endif
-
-			PCGGraphExecutionLogging::LogGraphExecuteFrameFinished();
-		}
-
-		// Purge things from cache if memory usage is too high
-		if (GraphCache.EnforceMemoryBudget())
-		{
-			CollectGarbage(GARBAGE_COLLECTION_KEEPFLAGS, true);
-		}
-	}
+	ExecuteTasksEnded();
 }
 
 void FPCGGraphExecutor::ClearAllTasks()
 {
-	Tasks.Reset();
+	check(IsInGameThread());
 
+	Tasks.Reset();
+	TaskSuccessors.Reset();
+	
 	// Make sure we don't leak preallocated contexts
 	for (FPCGGraphTask& ReadyTask : ReadyTasks)
 	{
@@ -1221,99 +1844,155 @@ void FPCGGraphExecutor::ClearAllTasks()
 	ReadyTasks.Reset();
 	ActiveTasks.Reset();
 	SleepingTasks.Reset();
-	TaskSuccessors.Reset();
 }
 
-void FPCGGraphExecutor::QueueNextTasks(FPCGTaskId FinishedTask, bool bIgnoreMissingTasks)
+void FPCGGraphExecutor::QueueNextTasks(FPCGTaskId FinishedTask)
 {
 	TRACE_CPUPROFILER_EVENT_SCOPE(FPCGGraphExecutor::QueueNextTasks);
+	TArray<FCachedResult*> CachedResults;
+	const bool bIsInGameThread = IsInGameThread();
 
-	if (TSet<FPCGTaskId>* Successors = TaskSuccessors.Find(FinishedTask))
 	{
-		for (FPCGTaskId Successor : *Successors)
+		TSet<FPCGTaskId> Successors;
 		{
-			bool bAllPrerequisitesMet = true;
-			FPCGGraphTask* SuccessorTaskPtr = Tasks.Find(Successor);
-
-			// This should never be null, but later recovery should be able to cleanup this properly
-			if (SuccessorTaskPtr)
+			PCGGraphExecutor::TScopeLock ScopeLock(TasksLock);
+		
+			if (TaskSuccessors.RemoveAndCopyValue(FinishedTask, Successors))
 			{
-				FPCGGraphTask& SuccessorTask = *SuccessorTaskPtr;
-
-				if (SuccessorTask.Inputs.Num() > 0)
+				for (FPCGTaskId Successor : Successors)
 				{
-					FReadScopeLock ReadLock(TaskOutputsRWLock);
-					for (const FPCGGraphTaskInput& Input : SuccessorTask.Inputs)
+					bool bAllPrerequisitesMet = true;
+
+					// When running multithreaded it is possible for the Successor to already have been moved to ReadyTask after the current task stored its results
+					if (FPCGGraphTask* SuccessorTaskPtr = Tasks.Find(Successor))
 					{
-						bAllPrerequisitesMet &= TaskOutputs.Contains(Input.TaskId);
-					}
-				}
+						FPCGGraphTask& SuccessorTask = *SuccessorTaskPtr;
 
-				if (bAllPrerequisitesMet)
-				{
-					ReadyTasks.Emplace(MoveTemp(SuccessorTask));
-					Tasks.Remove(Successor);
+						if (SuccessorTask.Inputs.Num() > 0)
+						{
+							PCGGraphExecutor::TScopeLock OutputLock(TaskOutputsLock);
+							for (const FPCGGraphTaskInput& Input : SuccessorTask.Inputs)
+							{
+								bAllPrerequisitesMet &= TaskOutputs.Contains(Input.TaskId);
+							}
+						}
+
+						if (bAllPrerequisitesMet)
+						{
+							OnTaskInputsReady(SuccessorTask, CachedResults, bIsInGameThread);
+							Tasks.Remove(Successor);
+						}
+					}
 				}
 			}
-			else
-			{
-				ensure(bIgnoreMissingTasks);
-			}
 		}
+	}
+
+	ProcessCachedResults(CachedResults);
+}
+
+void FPCGGraphExecutor::ProcessCachedResults(const TArray<FCachedResult*>& CachedResults)
+{
+	TRACE_CPUPROFILER_EVENT_SCOPE(FPCGGraphExecutor::ProcessCachedResults);
+
+	for (FCachedResult* CachedResult : CachedResults)
+	{
+#if WITH_EDITOR
+		// Cache was used to bypass actual execution but we let the task go through the normal execution path
+		if (CachedResult->bIsBypassed)
+		{
+			continue;
+		}
+#endif
+
+		check(CachedResult);
+		if (CachedResult->bDoDynamicTaskCulling)
+		{
+			CullInactiveDownstreamNodes(CachedResult->TaskId, CachedResult->Output.InactiveOutputPinBitmask);
+#if WITH_EDITOR
+			SendInactivePinNotification(CachedResult->Node, CachedResult->Stack, CachedResult->Output.InactiveOutputPinBitmask);
+#endif
+		}
+
+		StoreResults(CachedResult->TaskId, CachedResult->Output, CachedResult->bIsPostGraphTask);
+		QueueNextTasks(CachedResult->TaskId);
+	}
 
-		TaskSuccessors.Remove(FinishedTask);
+	if (CachedResults.Num() > 0)
+	{
+		PCGGraphExecutor::TScopeLock ScopeLock(LiveTasksLock);
+		for (const FCachedResult* CachedResult : CachedResults)
+		{
+			check(CachingResults.Contains(CachedResult->TaskId));
+			CachingResults.Remove(CachedResult->TaskId);
+		}
 	}
 }
 
 bool FPCGGraphExecutor::CancelNextTasks(FPCGTaskId CancelledTask, TSet<UPCGComponent*>& OutCancelledComponents)
 {
+	check(IsInGameThread());
+
 	bool bAddedComponents = false;
 
-	if (TSet<FPCGTaskId>* Successors = TaskSuccessors.Find(CancelledTask))
+	TSet<FPCGTaskId> LocalSuccessors;
 	{
-		TSet<FPCGTaskId> LocalSuccessors = MoveTemp(*Successors);
-		TaskSuccessors.Remove(CancelledTask);
-
-		for (FPCGTaskId Successor : LocalSuccessors)
+		PCGGraphExecutor::TScopeLock ScopeLock(TasksLock);
+		if (TSet<FPCGTaskId>* Successors = TaskSuccessors.Find(CancelledTask))
 		{
-			if (FPCGGraphTask* Task = Tasks.Find(Successor))
+			LocalSuccessors = MoveTemp(*Successors);
+			TaskSuccessors.Remove(CancelledTask);
+
+			for (FPCGTaskId Successor : LocalSuccessors)
 			{
-				if(!OutCancelledComponents.Contains(Task->SourceComponent.Get()))
+				if (FPCGGraphTask* Task = Tasks.Find(Successor))
 				{
-					OutCancelledComponents.Add(Task->SourceComponent.Get());
-					bAddedComponents = true;
-				}
+					if (!OutCancelledComponents.Contains(Task->SourceComponent.Get()))
+					{
+						OutCancelledComponents.Add(Task->SourceComponent.Get());
+						bAddedComponents = true;
+					}
 
-				if (Task->Element)
-				{
-					Task->Element->Abort(Task->Context);
-				}
+					if (Task->Element)
+					{
+						Task->Element->Abort(Task->Context);
+					}
 
-				RemoveTaskFromInputSuccessors(Task->NodeId, Task->Inputs);
-				Tasks.Remove(Successor);
+					RemoveTaskFromInputSuccessorsNoLock(Task->NodeId, Task->Inputs);
+					Tasks.Remove(Successor);
+				}
 			}
-
-			bAddedComponents |= CancelNextTasks(Successor, OutCancelledComponents);
 		}
 	}
 
+	for (FPCGTaskId Successor : LocalSuccessors)
+	{
+		bAddedComponents |= CancelNextTasks(Successor, OutCancelledComponents);
+	}
+		
 	// Tasks cancelled might have an impact on scheduled-but-not-processed tasks
-	ScheduleLock.Lock();
-	for (FPCGGraphScheduleTask& ScheduledTask : ScheduledTasks)
 	{
-		if (!OutCancelledComponents.Contains(ScheduledTask.SourceComponent.Get()) &&
-			Algo::AnyOf(ScheduledTask.Tasks[ScheduledTask.FirstTaskIndex].Inputs, [CancelledTask](const FPCGGraphTaskInput& Input) { return Input.TaskId == CancelledTask; }))
+		PCGGraphExecutor::TScopeLock ScopeLock(ScheduleLock);
+		for (FPCGGraphScheduleTask& ScheduledTask : ScheduledTasks)
 		{
-			OutCancelledComponents.Add(ScheduledTask.SourceComponent.Get());
-			bAddedComponents = true;
+			if (!OutCancelledComponents.Contains(ScheduledTask.SourceComponent.Get()) &&
+				Algo::AnyOf(ScheduledTask.Tasks[ScheduledTask.FirstTaskIndex].Inputs, [CancelledTask](const FPCGGraphTaskInput& Input) { return Input.TaskId == CancelledTask; }))
+			{
+				OutCancelledComponents.Add(ScheduledTask.SourceComponent.Get());
+				bAddedComponents = true;
+			}
 		}
 	}
-	ScheduleLock.Unlock();
-
 	return bAddedComponents;
 }
 
 void FPCGGraphExecutor::RemoveTaskFromInputSuccessors(FPCGTaskId CancelledTask, const TArray<FPCGGraphTaskInput>& CancelledTaskInputs)
+{
+	PCGGraphExecutor::TScopeLock ScopeLock(TasksLock);
+	RemoveTaskFromInputSuccessorsNoLock(CancelledTask, CancelledTaskInputs);
+}
+
+void FPCGGraphExecutor::RemoveTaskFromInputSuccessorsNoLock(FPCGTaskId CancelledTask, const TArray<FPCGGraphTaskInput>& CancelledTaskInputs)
 {
 	for (const FPCGGraphTaskInput& Input : CancelledTaskInputs)
 	{
@@ -1324,7 +2003,7 @@ void FPCGGraphExecutor::RemoveTaskFromInputSuccessors(FPCGTaskId CancelledTask,
 	}
 }
 
-void FPCGGraphExecutor::BuildTaskInput(const FPCGGraphTask& Task, FPCGDataCollection& TaskInput)
+void FPCGGraphExecutor::BuildTaskInput(FPCGGraphTask& Task, TArray<FPCGTaskId>& ResultsToMarkAsRead)
 {
 	TRACE_CPUPROFILER_EVENT_SCOPE(FPCGGraphExecutor::BuildTaskInput);
 
@@ -1363,15 +2042,21 @@ void FPCGGraphExecutor::BuildTaskInput(const FPCGGraphTask& Task, FPCGDataCollec
 	// Hoisted out of loop for performance reasons.
 	TArray<FPCGTaggedData, TInlineAllocator<16>> InputDataOnPin;
 	TArray<FPCGCrc, TInlineAllocator<16>> InputDataCrcsOnPin;
-	TArray<FPCGTaskId, TInlineAllocator<16>> ResultsToMarkAsRead;
-
+	
 	for (const FPCGGraphTaskInput& Input : Task.Inputs)
 	{
-		FReadScopeLock ReadLock(TaskOutputsRWLock);
-		check(TaskOutputs.Contains(Input.TaskId));
+		PCGGraphExecutor::TScopeLock OutputLock(TaskOutputsLock);
+		
+		const FOutputDataInfo& TaskOutput = TaskOutputs.FindChecked(Input.TaskId);
 		
 		ResultsToMarkAsRead.AddUnique(Input.TaskId);
 
+		// Input was culled before this task managed to get scheduled (moved from ScheduleTasks to Tasks)
+		if (TaskOutput.bCulled)
+		{
+			continue;
+		}
+
 		// If the input does not provide any data, don't add it to the task input.
 		if (!Input.bProvideData)
 		{
@@ -1382,17 +2067,17 @@ void FPCGGraphExecutor::BuildTaskInput(const FPCGGraphTask& Task, FPCGDataCollec
 		const uint32 InputPinLabelCrc = Input.DownstreamPin.IsSet() ? GetTypeHash(Input.DownstreamPin.GetValue().Label) : DefaultHashForNoOutputPin;
 
 		// Enforce single data - if already have input for this pin, don't add more. Early check before other side effects below.
-		if (Input.DownstreamPin.IsSet() && !bAllowMultipleData && TaskInput.GetInputCountByPin(Input.DownstreamPin.GetValue().Label) > 0)
+		if (Input.DownstreamPin.IsSet() && !bAllowMultipleData && Task.TaskInput.GetInputCountByPin(Input.DownstreamPin.GetValue().Label) > 0)
 		{
 			LogDiscardedData(Input.DownstreamPin);
 			continue;
 		}
 
-		const FPCGDataCollection& InputCollection = TaskOutputs[Input.TaskId].DataCollection;
-
-		TaskInput.bCancelExecution |= InputCollection.bCancelExecution;
+		const FPCGDataCollection& InputCollection = TaskOutput.DataCollection;
+		
+		Task.TaskInput.bCancelExecution |= InputCollection.bCancelExecution;
 
-		const int32 TaggedDataOffset = TaskInput.TaggedData.Num();
+		const int32 TaggedDataOffset = Task.TaskInput.TaggedData.Num();
 
 		// Get input data at the given pin (or everything). This will add the data and include the input pin Crc to uniquely identify
 		// inputs per-pin, or use a placeholder for symmetry.
@@ -1413,7 +2098,7 @@ void FPCGGraphExecutor::BuildTaskInput(const FPCGGraphTask& Task, FPCGDataCollec
 
 				const int NumberDataItemsToTake = bAllowMultipleData ? InputDataOnPin.Num() : 1;
 
-				TaskInput.AddDataForPin(
+				Task.TaskInput.AddDataForPin(
 					MakeArrayView(InputDataOnPin.GetData(), NumberDataItemsToTake),
 					MakeArrayView(InputDataCrcsOnPin.GetData(), NumberDataItemsToTake),
 					InputPinLabelCrc);
@@ -1426,33 +2111,29 @@ void FPCGGraphExecutor::BuildTaskInput(const FPCGGraphTask& Task, FPCGDataCollec
 		}
 		else
 		{
-			TaskInput.AddData(InputCollection.TaggedData, InputCollection.DataCrcs);
+			Task.TaskInput.AddData(InputCollection.TaggedData, InputCollection.DataCrcs);
 		}
 
 		// Apply labelling on data; technically, we should ensure that we do this only for pass-through nodes,
 		// Otherwise we could also null out the label on the input...
 		if (Input.DownstreamPin.IsSet())
 		{
-			for (int32 TaggedDataIndex = TaggedDataOffset; TaggedDataIndex < TaskInput.TaggedData.Num(); ++TaggedDataIndex)
+			for (int32 TaggedDataIndex = TaggedDataOffset; TaggedDataIndex < Task.TaskInput.TaggedData.Num(); ++TaggedDataIndex)
 			{
-				TaskInput.TaggedData[TaggedDataIndex].Pin = Input.DownstreamPin.GetValue().Label;
+				Task.TaskInput.TaggedData[TaggedDataIndex].Pin = Input.DownstreamPin.GetValue().Label;
 			}
 		}
 	}
 
 	// Then combine params if needed
-	CombineParams(Task.NodeId, TaskInput);
-
-	// Finally, mark inputs we read from so we can release them if we were the last remaining consumer.
-	MarkInputResults(ResultsToMarkAsRead);
+	CombineParams(Task);
 }
 
-void FPCGGraphExecutor::CombineParams(FPCGTaskId InTaskId, FPCGDataCollection& InTaskInput)
+void FPCGGraphExecutor::CombineParams(FPCGGraphTask& Task)
 {
 	TRACE_CPUPROFILER_EVENT_SCOPE(FPCGGraphExecutor::CombineParams);
-	ensure(IsInGameThread());
-
-	TArray<FPCGTaggedData> AllParamsData = InTaskInput.GetParamsByPin(PCGPinConstants::DefaultParamsLabel);
+	
+	TArray<FPCGTaggedData> AllParamsData = Task.TaskInput.GetParamsByPin(PCGPinConstants::DefaultParamsLabel);
 	if (AllParamsData.Num() > 1)
 	{
 		UPCGParamData* CombinedParamData = nullptr;
@@ -1463,7 +2144,16 @@ void FPCGGraphExecutor::CombineParams(FPCGTaskId InTaskId, FPCGDataCollection& I
 			const UPCGParamData* ParamData = CastChecked<UPCGParamData>(TaggedDatum.Data);
 			if (!CombinedParamData)
 			{
-				CombinedParamData = ParamData->DuplicateData(nullptr);
+				// We always need a Context to allow capture of Async objects if this runs outside of the game thread
+				FPCGContext TempContext;
+				TempContext.AsyncState.bIsRunningOnMainThread = IsInGameThread();
+				
+				CombinedParamData = ParamData->DuplicateData(&TempContext);
+
+				if (!TempContext.AsyncState.bIsRunningOnMainThread)
+				{
+					Task.CombineParamsAsyncObjects = std::move(TempContext.AsyncObjects);
+				}
 			}
 			else
 			{
@@ -1476,12 +2166,12 @@ void FPCGGraphExecutor::CombineParams(FPCGTaskId InTaskId, FPCGDataCollection& I
 			return;
 		}
 
-		const int32 NewNumberOfInputs = InTaskInput.TaggedData.Num() - AllParamsData.Num() + 1;
+		const int32 NewNumberOfInputs = Task.TaskInput.TaggedData.Num() - AllParamsData.Num() + 1;
 		check(NewNumberOfInputs >= 1);
 
 		TArray<FPCGTaggedData> TempTaggedData{};
 		TempTaggedData.Reserve(NewNumberOfInputs);
-		for (FPCGTaggedData& TaggedData : InTaskInput.TaggedData)
+		for (FPCGTaggedData& TaggedData : Task.TaskInput.TaggedData)
 		{
 			if (TaggedData.Pin != PCGPinConstants::DefaultParamsLabel)
 			{
@@ -1495,7 +2185,7 @@ void FPCGGraphExecutor::CombineParams(FPCGTaskId InTaskId, FPCGDataCollection& I
 		CombineParams.Pin = PCGPinConstants::DefaultParamsLabel;
 		TempTaggedData.Add(CombineParams);
 
-		InTaskInput.TaggedData = std::move(TempTaggedData);
+		Task.TaskInput.TaggedData = std::move(TempTaggedData);
 	}
 }
 
@@ -1506,30 +2196,34 @@ void FPCGGraphExecutor::StoreResults(FPCGTaskId InTaskId, const FPCGDataCollecti
 	FOutputDataInfo OutputDataInfo;
 	OutputDataInfo.DataCollection = InTaskOutput;
 	OutputDataInfo.bNeedsManualClear = bNeedsManualClear;
-	if (TSet<FPCGTaskId>* Successors = TaskSuccessors.Find(InTaskId))
+
 	{
-		OutputDataInfo.RemainingSuccessorCount = Successors->Num();
+		TRACE_CPUPROFILER_EVENT_SCOPE(FPCGGraphExecutor::StoreResults::Successors);
+		PCGGraphExecutor::TScopeLock ScopeLock(TasksLock);
+		if (TSet<FPCGTaskId>* Successors = TaskSuccessors.Find(InTaskId))
+		{
+			OutputDataInfo.RemainingSuccessorCount = Successors->Num();
+		}
 	}
 
 	// Store output in map
-	FWriteScopeLock WriteLock(TaskOutputsRWLock);
+	PCGGraphExecutor::TScopeLock OutputLock(TaskOutputsLock);
 	TaskOutputs.Add(InTaskId, MoveTemp(OutputDataInfo));
 }
 
 void FPCGGraphExecutor::ClearResults()
 {
-	ScheduleLock.Lock();
-
-	// Only reset if we have no more scheduled tasks, to avoid breaking dependencies
-	if (ScheduledTasks.IsEmpty())
 	{
-		NextTaskId = 0;
+		PCGGraphExecutor::TScopeLock ScopeLock(ScheduleLock);
+		// Only reset if we have no more scheduled tasks, to avoid breaking dependencies
+		if (ScheduledTasks.IsEmpty())
+		{
+			NextTaskId = 0;
+		}
 	}
 
-	FWriteScopeLock WriteLock(TaskOutputsRWLock);
+	PCGGraphExecutor::TScopeLock OutputLock(TaskOutputsLock);
 	TaskOutputs.Reset();
-
-	ScheduleLock.Unlock();
 }
 
 void FPCGGraphExecutor::GetPinIdsToDeactivate(FPCGTaskId TaskId, uint64 InactiveOutputPinBitmask, TArray<FPCGPinId>& InOutPinIds)
@@ -1573,6 +2267,7 @@ void FPCGGraphExecutor::CullInactiveDownstreamNodes(FPCGTaskId InCompletedTaskId
 		PCGGraphExecutionLogging::LogTaskCullingBeginLoop(PinTaskId, PCGPinIdHelpers::GetPinIndexFromPinId(PinId), PinIdsToDeactivate);
 		LogTaskState();
 
+		PCGGraphExecutor::TScopeLock ScopeLock(TasksLock);
 		const TSet<FPCGTaskId>* Successors = TaskSuccessors.Find(PinTaskId);
 		if (!Successors)
 		{
@@ -1627,7 +2322,7 @@ void FPCGGraphExecutor::CullInactiveDownstreamNodes(FPCGTaskId InCompletedTaskId
 				PinIdsToDeactivate.AddUnique(PCGPinIdHelpers::NodeIdToPinId(RemovedTaskId));
 
 				// Remove task as successor of upstream node.
-				RemoveTaskFromInputSuccessors(RemovedTaskId, RemovedTask.Inputs);
+				RemoveTaskFromInputSuccessorsNoLock(RemovedTaskId, RemovedTask.Inputs);
 
 				// Remove the deleted tasks from the inputs of downstream tasks.
 				if (const TSet<FPCGTaskId>* SuccessorsOfRemovedTask = TaskSuccessors.Find(RemovedTaskId))
@@ -1659,10 +2354,19 @@ void FPCGGraphExecutor::CullInactiveDownstreamNodes(FPCGTaskId InCompletedTaskId
 	// Ensure any downstream tasks are enqueued.
 	for (const FPCGTaskId TaskId : AllRemovedTasks)
 	{
+		{
+			// Store culled output
+			PCGGraphExecutor::TScopeLock OutputLock(TaskOutputsLock);
+			FOutputDataInfo CullDataInfo;
+			CullDataInfo.RemainingSuccessorCount = 0;
+			CullDataInfo.bCulled = true;
+			TaskOutputs.Add(TaskId, MoveTemp(CullDataInfo));
+		}
+
 		// Queue downstream tasks in a similar manner to when a task draws from the cache and is skipped.
 		// Some downstream tasks will have been culled which we don't care about (hence the ignore flag),
 		// but some may not be queued and may be ready for queuing.
-		QueueNextTasks(TaskId, /*bIgnoreMissingTasks=*/true);
+		QueueNextTasks(TaskId);
 	}
 }
 
@@ -1683,28 +2387,71 @@ void FPCGGraphExecutor::AddReferencedObjects(FReferenceCollector& Collector)
 
 	// Go through all data in the cached output map
 	{
-		FReadScopeLock ReadLock(TaskOutputsRWLock);
+		PCGGraphExecutor::TScopeLock ScopeLock(TaskOutputsLock);
 		for (auto& OutputDataEntry : TaskOutputs)
 		{
 			OutputDataEntry.Value.DataCollection.AddReferences(Collector);
 		}
 	}
 
-	// Go through ready tasks, active tasks and sleeping tasks contexts
-	auto AddReferences = [&Collector](auto& TaskContainer)
 	{
-		for (auto& Task : TaskContainer)
+		PCGGraphExecutor::TScopeLock ScopeLock(LiveTasksLock);
+		Algo::ForEach(ReadyTasks, [&Collector](auto& Task)
 		{
-			if (Task.Context)
+			if (!Task.bHasDonePrepareForExecute)
+			{
+				check(!Task.Context);
+				Task.TaskInput.AddReferences(Collector);
+			}
+			else if (Task.Context)
 			{
 				Task.Context->AddStructReferencedObjects(Collector);
 			}
-		}
-	};
+		});
+
+		Algo::ForEach(ActiveTasks, [&Collector](auto& Task)
+		{
+			if (Task->bIsExecutingTask)
+			{
+				Collector.AddReferencedObjects(Task->ExecutingReferences);
+				// @todo_pcg this is to allow referencing extra objects from the Context sub classes, should probably be part of the visiting pattern also
+				if (Task->Context)
+				{
+					Task->Context->AddExtraStructReferencedObjects(Collector);
+				}
+			}
+			else if (Task->Context)
+			{
+				Task->Context->AddStructReferencedObjects(Collector);
+			}
+		});
+
+		Algo::ForEach(SleepingTasks, [&Collector](auto& Task)
+		{
+			if (Task->Context)
+			{
+				Task->Context->AddStructReferencedObjects(Collector);
+			}
+		});
 
-	AddReferences(ReadyTasks);
-	AddReferences(ActiveTasks);
-	AddReferences(SleepingTasks);
+		Algo::ForEach(CachingResults, [&Collector](auto& CachingResult)
+		{
+			CachingResult.Value->Output.AddReferences(Collector);
+		});
+	}
+
+	// Some task might have been delayed to the game thread and we are no longer referencing them through (Ready/Sleep/Active)
+	// but we still want their UObject references to stay valid
+	{
+		PCGGraphExecutor::TScopeLock ScopeLock(CollectGCReferenceTasksLock);
+		Algo::ForEach(CollectGCReferenceTasks, [&Collector](auto& Task)
+		{
+			if (Task->Context)
+			{
+				Task->Context->AddStructReferencedObjects(Collector);
+			}
+		});
+	}
 }
 
 FPCGElementPtr FPCGGraphExecutor::GetFetchInputElement()
@@ -1723,6 +2470,7 @@ void FPCGGraphExecutor::LogTaskState() const
 	if (PCGGraphExecutionLogging::CullingLogEnabled())
 	{
 		UE_LOG(LogPCG, Log, TEXT("\tDORMANT (FPCGGraphExecutor::Tasks):"));
+		PCGGraphExecutor::TScopeLock ScopeLock(TasksLock);
 		PCGGraphExecutionLogging::LogGraphTasks(Tasks, &TaskSuccessors);
 	}
 #endif
@@ -1753,7 +2501,7 @@ FPCGTaskId FPCGGraphExecutor::ScheduleDebugWithTaskCallback(UPCGComponent* InCom
 		FPCGTaskId CaptureTaskId = ScheduleGeneric([this, TaskCompleteCallback, CompiledTask]
 		{
 			FPCGDataCollection TaskOutputData;
-			if (CompiledTask.Node && GetOutputData(CompiledTask.NodeId, TaskOutputData, /*bClearDataOnGet=*/false))
+			if (CompiledTask.Node && GetOutputData(CompiledTask.NodeId, TaskOutputData))
 			{
 				TaskCompleteCallback(CompiledTask.NodeId, CompiledTask.Node, TaskOutputData);
 			}
diff --git a/Engine/Plugins/PCG/Source/PCG/Private/Graph/PCGGraphExecutor.h b/Engine/Plugins/PCG/Source/PCG/Private/Graph/PCGGraphExecutor.h
--- a/Engine/Plugins/PCG/Source/PCG/Private/Graph/PCGGraphExecutor.h
+++ b/Engine/Plugins/PCG/Source/PCG/Private/Graph/PCGGraphExecutor.h
@@ -11,7 +11,10 @@
 #include "Graph/PCGPinDependencyExpression.h"
 #include "Graph/PCGStackContext.h"
 
+#include "Misc/SpinLock.h"
 #include "UObject/GCObject.h"
+#include "Tasks/Task.h"
+#include "Templates/UniquePtr.h"
 
 #if WITH_EDITOR
 #include "Editor/IPCGEditorProgressNotification.h"
@@ -93,8 +96,15 @@ struct FPCGGraphTask
 	int32 StackIndex = INDEX_NONE;
 	TSharedPtr<const FPCGStackContext> StackContext;
 
-	// Contains info around whether the element, context, cache check and more was already done.
+	// Whether SetupTask has been called on this task
 	bool bHasDoneSetup = false;
+	// BuildTaskInput will initialize this Collection which will later be used by PrepareForExecute
+	FPCGDataCollection TaskInput;
+	// CombineParams call might have created AsyncObjects
+	TSet<TObjectPtr<UObject>> CombineParamsAsyncObjects;
+
+	// Whether PrepareForExecute as been called on this task
+	bool bHasDonePrepareForExecute = false;
 
 #if WITH_EDITOR
 	// Can be true when we want to have debug display on a task but have taken the results from the cache
@@ -111,27 +121,38 @@ struct FPCGGraphScheduleTask
 	bool bHasAbortCallbacks = false;
 };
 
-struct FPCGGraphActiveTask
+struct FPCGGraphActiveTask : TSharedFromThis<FPCGGraphActiveTask>
 {
 	FPCGGraphActiveTask() = default;
-	~FPCGGraphActiveTask();
-	
+	virtual ~FPCGGraphActiveTask();
+
 	FPCGGraphActiveTask(const FPCGGraphActiveTask&) = delete;
 	FPCGGraphActiveTask& operator=(const FPCGGraphActiveTask&) = delete;
 
-	FPCGGraphActiveTask(FPCGGraphActiveTask&&) = default;
-	FPCGGraphActiveTask& operator=(FPCGGraphActiveTask&&) = default;
+	FPCGGraphActiveTask(FPCGGraphActiveTask&&) = delete;
+	FPCGGraphActiveTask& operator=(FPCGGraphActiveTask&&) = delete;
+
+	void StartExecuting();
+	void StopExecuting();
 
 	TArray<FPCGGraphTaskInput> Inputs;
 	FPCGElementPtr Element;
 	TUniquePtr<FPCGContext> Context;
 	FPCGTaskId NodeId = InvalidPCGTaskId;
-	bool bWasCancelled = false;
+	std::atomic<bool> bWasCancelled = false;
 #if WITH_EDITOR
 	bool bIsBypassed = false;
 #endif
 	int32 StackIndex = INDEX_NONE;
 	TSharedPtr<const FPCGStackContext> StackContext;
+		
+	// Those members need to be modified under the FPCGGraphExecutor::LiveTasksLock (unless we are running the old executor path)
+	UE::Tasks::TTask<bool> ExecutingTask;
+	bool bIsExecutingTask = false;
+	static int32 NumExecuting;
+	TArray<TObjectPtr<const UObject>> ExecutingReferences;
+
+	TSharedPtr<UE::Tasks::FCancellationToken> CancelToken;
 };
 
 class FPCGGraphExecutor : public FGCObject
@@ -195,7 +216,10 @@ public:
 	FPCGTaskId ScheduleGenericWithContext(TFunction<bool(FPCGContext*)> InOperation, TFunction<void(FPCGContext*)> InAbortOperation, UPCGComponent* InSourceComponent, const TArray<FPCGTaskId>& TaskExecutionDependencies, const TArray<FPCGTaskId>& TaskDataDependencies);
 
 	/** Gets data in the output results. Returns false if data is not ready. */
-	bool GetOutputData(FPCGTaskId InTaskId, FPCGDataCollection& OutData, bool bClearDataOnGet);
+	bool GetOutputData(FPCGTaskId InTaskId, FPCGDataCollection& OutData);
+
+	/** Clear output data */
+	void ClearOutputData(FPCGTaskId InTaskId);
 
 	/** Accessor so PCG tools (e.g. profiler) can easily decode graph task ids **/
 	FPCGGraphCompiler& GetCompiler() { return GraphCompiler; }
@@ -228,14 +252,47 @@ public:
 	//~End FGCObject interface
 
 private:
+	void ExecuteV1();
+	void ExecuteV2();
+	double GetTickBudgetInSeconds() const;
+		
+	struct FCachedResult
+	{
+		FPCGTaskId TaskId = InvalidPCGTaskId;
+		FPCGDataCollection Output;
+		const FPCGStack* Stack = nullptr;
+		const UPCGNode* Node = nullptr;
+		bool bDoDynamicTaskCulling = false;
+		bool bIsPostGraphTask = false;		
+		bool bIsBypassed = false;
+	};
+		
+	void PostTaskExecute(TSharedPtr<FPCGGraphActiveTask> ActiveTask);
+	bool ProcessScheduledTasks();
+	void ExecuteTasksEnded();
+	bool ExecuteScheduling(double EndTime, TSharedPtr<FPCGGraphActiveTask>* OutMainThreadTask = nullptr);
+
 	TSet<UPCGComponent*> Cancel(TFunctionRef<bool(TWeakObjectPtr<UPCGComponent>)> CancelFilter);
 	void ClearAllTasks();
-	void QueueNextTasks(FPCGTaskId FinishedTask, bool bIgnoreMissingTasks = false);
+	void QueueNextTasks(FPCGTaskId FinishedTask);
 	bool CancelNextTasks(FPCGTaskId CancelledTask, TSet<UPCGComponent*>& OutCancelledComponents);
 	void RemoveTaskFromInputSuccessors(FPCGTaskId CancelledTask, const TArray<FPCGGraphTaskInput>& CancelledTaskInputs);
-	void BuildTaskInput(const FPCGGraphTask& Task, FPCGDataCollection& TaskInput);
+	void RemoveTaskFromInputSuccessorsNoLock(FPCGTaskId CancelledTask, const TArray<FPCGGraphTaskInput>& CancelledTaskInputs);
+	
+	/** Called from QueueNextTasks/ProcessScheduledTasks will try and setup/prepare task for execution */
+	void OnTaskInputsReady(FPCGGraphTask& Task, TArray<FCachedResult*>& OutCachedResults, bool bIsInGameThread);
+
+	/** SetupTask will call BuildTaskInput and assign a IPCGElement to the FPCGGraphTask */
+	bool SetupTask(FPCGGraphTask& Task, TArray<FPCGTaskId>& ResultsToMarkAsRead);
+	void BuildTaskInput(FPCGGraphTask& Task, TArray<FPCGTaskId>& ResultsToMarkAsRead);
+	/** Will check the cache for existing result or create and initialize the FPCGContext to the task */
+	void PrepareForExecute(FPCGGraphTask& Task, FCachedResult*& OutCachedResult, bool bLiveTasksLockAlreadyLocked);
+
+	/** Store cache results and Queue next tasks */
+	void ProcessCachedResults(const TArray<FCachedResult*>& CachedResults);
+
 	/** Combine all param data into one on the Params pin, if any.*/
-	void CombineParams(FPCGTaskId InTaskId, FPCGDataCollection& InTaskInput);
+	void CombineParams(FPCGGraphTask& Task);
 	void StoreResults(FPCGTaskId InTaskId, const FPCGDataCollection& InTaskOutput, bool bNeedsManualClear);
 	void ClearResults();
 	void MarkInputResults(TArrayView<const FPCGTaskId> InInputResults);
@@ -244,7 +301,7 @@ private:
 	void CullInactiveDownstreamNodes(FPCGTaskId CompletedTaskId, uint64 InInactiveOutputPinBitmask);
 
 	/** Builds an array of all deactivated unique pin IDs. */
-	void GetPinIdsToDeactivate(FPCGTaskId TaskId, uint64 InactiveOutputPinBitmask, TArray<FPCGPinId>& InOutPinIds);
+	static void GetPinIdsToDeactivate(FPCGTaskId TaskId, uint64 InactiveOutputPinBitmask, TArray<FPCGPinId>& InOutPinIds);
 
 	FPCGElementPtr GetFetchInputElement();
 
@@ -271,14 +328,35 @@ private:
 	/** Input fetch element, stored here so we have only one */
 	FPCGElementPtr FetchInputElement;
 
-	FCriticalSection ScheduleLock;
+	/** 
+	 * Define a Lock level for future reference. Rule is when we have a lock, we can't lock a lower or equal level lock to prevent deadlocks.
+	 * Example: When locking LiveTasksLock we can't lock ScheduleLock or CollectGCReferenceTasksLock 
+	 */
+
+	/** Lock level - 1 (top most lock) */
+	UE::FSpinLock ScheduleLock;
 	TArray<FPCGGraphScheduleTask> ScheduledTasks;
 
+	/** Lock level 2 */
+	mutable UE::FSpinLock TasksLock;
 	TMap<FPCGTaskId, FPCGGraphTask> Tasks;
-	TArray<FPCGGraphTask> ReadyTasks;
-	TArray<FPCGGraphActiveTask> ActiveTasks;
-	TArray<FPCGGraphActiveTask> SleepingTasks;
 	TMap<FPCGTaskId, TSet<FPCGTaskId>> TaskSuccessors;
+
+	/** Lock level 3 */
+	UE::FSpinLock LiveTasksLock;
+	TArray<FPCGGraphTask> ReadyTasks;
+	TArray<TSharedPtr<FPCGGraphActiveTask>> ActiveTasks;
+	// Used to keep GC references to in flight caching results (not yet stored to output and might not be in cache anymore)
+	TMap<FPCGTaskId, TUniquePtr<FCachedResult>> CachingResults;
+	int32 NumWorkerTasks = 0;
+
+	// @todo_pcg: to remove when we remove ExecuteV1
+	TArray<TSharedPtr<FPCGGraphActiveTask>> SleepingTasks;
+	
+	/** Lock level 3 */
+	UE::FSpinLock CollectGCReferenceTasksLock;
+	TSet<TSharedPtr<FPCGGraphActiveTask>> CollectGCReferenceTasks;
+			
 	/** Map of node instances to their output, could be cleared once execution is done */
 	/** Note: this should at some point unload based on loaded/unloaded proxies, otherwise memory cost will be unbounded */
 	struct FOutputDataInfo
@@ -288,13 +366,18 @@ private:
 		bool bNeedsManualClear = false;
 		// Successor count, updated after a successor is done executing (MarkInputResults).
 		int32 RemainingSuccessorCount = 0;
+		// Culled
+		bool bCulled = false;
 	};
 
-	FRWLock TaskOutputsRWLock;
+	/** Lock level 4 */
+	UE::FSpinLock TaskOutputsLock;
 	TMap<FPCGTaskId, FOutputDataInfo> TaskOutputs;
-
+	
 	/** Monotonically increasing id. Should be reset once all tasks are executed, should be protected by the ScheduleLock */
 	FPCGTaskId NextTaskId = 0;
+	
+	std::atomic<bool> bNeedToExecuteTasksEnded = false;
 
 	/** Runtime information */
 	int32 CurrentlyUsedThreads = 0;
@@ -306,6 +389,17 @@ private:
 #endif
 
 	TObjectPtr<UWorld> World = nullptr;
+
+	// Temporary while the 2 schedulers exist
+	enum class EExecuteVersion : uint8
+	{
+		None,
+		V1,
+		V2
+	};
+	EExecuteVersion ExecuteVersion = EExecuteVersion::None;
+
+	EExecuteVersion GetExecuteVersion() const;
 };
 
 class FPCGFetchInputElement : public IPCGElement
diff --git a/Engine/Plugins/PCG/Source/PCG/Private/Helpers/PCGAsync.cpp b/Engine/Plugins/PCG/Source/PCG/Private/Helpers/PCGAsync.cpp
--- a/Engine/Plugins/PCG/Source/PCG/Private/Helpers/PCGAsync.cpp
+++ b/Engine/Plugins/PCG/Source/PCG/Private/Helpers/PCGAsync.cpp
@@ -5,14 +5,26 @@
 #include "PCGContext.h"
 #include "PCGPoint.h"
 
-#include "Async/Async.h"
+#include "Tasks/Task.h"
 
 namespace FPCGAsync
 {
 	TAutoConsoleVariable<bool> ConsoleVar::CVarDisableAsyncTimeSlicing {
 		TEXT("pcg.DisableAsyncTimeSlicing"),
 		false,
-		TEXT("To help debug, we can disable time slicing for async tasks.")
+		TEXT("To help debug, we can disable time slicing for tasks (when executing Element from worker thread)")
+	};
+
+	TAutoConsoleVariable<bool> ConsoleVar::CVarDisableAsyncTimeSlicingOnGameThread {
+		TEXT("pcg.DisableAsyncTimeSlicingOnGameThread"),
+		false,
+		TEXT("To help debug, we can disable time slicing for tasks (when executing Element from main thread).")
+	};
+
+	TAutoConsoleVariable<float> ConsoleVar::CVarAsyncOutOfTickBudgetInMilliseconds{
+		TEXT("pcg.Async.OutOfTickBudgetInMilliseconds"),
+		5.0f,
+		TEXT("Allocated time in milliseconds per task when running async tasks out of tick (depends on pcg.DisableAsyncTimeSlicing being false)")
 	};
 
 	TAutoConsoleVariable<int32> ConsoleVar::CVarAsyncOverrideChunkSize{
@@ -21,18 +33,31 @@ namespace FPCGAsync
 		TEXT("For quick benchmarking, we can override the value of chunk size for async processing. Any negative value is discarded.")
 	};
 
-	void AsyncPointProcessing(int32 NumAvailableTasks, int32 MinIterationsPerTask, int32 NumIterations, TArray<FPCGPoint>& OutPoints, const TFunction<int32(int32, int32)>& IterationInnerLoop)
+	int32 GetNumTasks(FPCGAsyncState* AsyncState, int32 InDefaultNumTasks)
+	{
+		// Get number of available threads from the async state
+		// Calculate Number of Tasks to spawn, if we have no AsyncState run only 1 task (reentrance)
+		int32 NumTasks = AsyncState ? FMath::Max(1, InDefaultNumTasks) : 1;
+		// If AsyncState limits the Max Tasks apply here
+		if (AsyncState && AsyncState->NumAvailableTasks > 0)
+		{
+			NumTasks = FMath::Min(AsyncState->NumAvailableTasks, NumTasks);
+		}
+
+		return NumTasks;
+	}
+
+	void PrivateAsyncPointProcessing(FPCGAsyncState* AsyncState, int32 MinIterationsPerTask, int32 NumIterations, TArray<FPCGPoint>& OutPoints, const TFunction<int32(int32, int32)>& IterationInnerLoop)
 	{
 		TRACE_CPUPROFILER_EVENT_SCOPE(IPCGElement::AsyncPointProcessing);
 
-		check(NumAvailableTasks > 0 && MinIterationsPerTask > 0 && NumIterations >= 0);
-		if (NumAvailableTasks <= 0 || MinIterationsPerTask <= 0 || NumIterations <= 0)
+		check(MinIterationsPerTask > 0 && NumIterations >= 0);
+		if (MinIterationsPerTask <= 0 || NumIterations <= 0)
 		{
 			return;
 		}
 
-		// Get number of available threads from the context
-		const int32 NumTasks = FMath::Min(NumAvailableTasks, FMath::Max(1, NumIterations / MinIterationsPerTask));
+		const int32 NumTasks = GetNumTasks(AsyncState, NumIterations / MinIterationsPerTask);
 		const int32 IterationsPerTask = NumIterations / NumTasks;
 		const int32 NumFutures = NumTasks - 1;
 		 
@@ -43,7 +68,7 @@ namespace FPCGAsync
 		}
 
 		// Setup [current, last, nb points] data per dispatch
-		TArray<TFuture<int32>> AsyncTasks;
+		TArray<UE::Tasks::TTask<int32>> AsyncTasks;
 		AsyncTasks.Reserve(NumFutures);
 
 		// Launch the async tasks
@@ -52,7 +77,7 @@ namespace FPCGAsync
 			const int32 StartIndex = TaskIndex * IterationsPerTask;
 			const int32 EndIndex = StartIndex + IterationsPerTask;
 
-			AsyncTasks.Emplace(Async(EAsyncExecution::ThreadPool, [&IterationInnerLoop, StartIndex, EndIndex]()
+			AsyncTasks.Emplace(UE::Tasks::Launch(UE_SOURCE_LOCATION, [&IterationInnerLoop, StartIndex, EndIndex]()
 			{
 				return IterationInnerLoop(StartIndex, EndIndex);
 			}));
@@ -84,9 +109,9 @@ namespace FPCGAsync
 			{
 				const int32 StartPointsIndex = AsyncIndex * IterationsPerTask;
 
-				TFuture<int32>& AsyncTask = AsyncTasks[AsyncIndex];
+				UE::Tasks::TTask<int32>& AsyncTask = AsyncTasks[AsyncIndex];
 				AsyncTask.Wait();
-				const int32 NumberOfPointsAdded = AsyncTask.Get();
+				const int32 NumberOfPointsAdded = AsyncTask.GetResult();
 				RangeIndex = CollapsePoints(RangeIndex, StartPointsIndex, NumberOfPointsAdded);
 			}
 
@@ -97,36 +122,16 @@ namespace FPCGAsync
 		}
 	}
 
-	void AsyncPointProcessing(FPCGContext* Context, int32 NumIterations, TArray<FPCGPoint>& OutPoints, const TFunction<bool(int32, FPCGPoint&)>& PointFunc)
+	void PrivateAsyncPointProcessing(FPCGAsyncState* AsyncState, int32 MinIterationsPerTask, int32 NumIterations, TArray<FPCGPoint>& OutPoints, const TFunction<bool(int32, FPCGPoint&)>& PointFunc)
 	{
-		const int32 MinIterationsPerTask = 256;
-
-		FPCGAsyncState* AsyncState = Context ? &Context->AsyncState : nullptr;
-		
-		if (AsyncState && !AsyncState->bIsRunningAsyncCall)
-		{
-			AsyncState->bIsRunningAsyncCall = true;
-			AsyncPointProcessing(AsyncState->NumAvailableTasks, MinIterationsPerTask, NumIterations, OutPoints, PointFunc);
-			AsyncState->bIsRunningAsyncCall = false;
-		}
-		else
-		{
-			AsyncPointProcessing(1, MinIterationsPerTask, NumIterations, OutPoints, PointFunc);
-		}
-	}
-
-	void AsyncPointProcessing(FPCGContext* Context, const TArray<FPCGPoint>& InPoints, TArray<FPCGPoint>& OutPoints, const TFunction<bool(const FPCGPoint&, FPCGPoint&)>& PointFunc)
-	{
-		const int32 MinIterationsPerTask = 256;
-
-		auto IterationInnerLoop = [&PointFunc, &InPoints, &OutPoints](int32 StartIndex, int32 EndIndex) -> int32
+		auto IterationInnerLoop = [&PointFunc, &OutPoints](int32 StartIndex, int32 EndIndex) -> int32
 		{
 			TRACE_CPUPROFILER_EVENT_SCOPE(IPCGElement::AsyncPointProcessing::InnerLoop);
 			int32 NumPointsWritten = 0;
 
 			for (int32 Index = StartIndex; Index < EndIndex; ++Index)
 			{
-				if (PointFunc(InPoints[Index], OutPoints[StartIndex + NumPointsWritten]))
+				if (PointFunc(Index, OutPoints[StartIndex + NumPointsWritten]))
 				{
 					++NumPointsWritten;
 				}
@@ -135,70 +140,75 @@ namespace FPCGAsync
 			return NumPointsWritten;
 		};
 
-		FPCGAsyncState* AsyncState = Context ? &Context->AsyncState : nullptr;
+		PrivateAsyncPointProcessing(AsyncState, MinIterationsPerTask, NumIterations, OutPoints, IterationInnerLoop);
+	}
 
-		if (AsyncState && !AsyncState->bIsRunningAsyncCall)
-		{
-			AsyncState->bIsRunningAsyncCall = true;
-			AsyncPointProcessing(AsyncState->NumAvailableTasks, MinIterationsPerTask, InPoints.Num(), OutPoints, IterationInnerLoop);
-			AsyncState->bIsRunningAsyncCall = false;
-		}
-		else
+	void PrivateAsyncMultiPointProcessing(FPCGAsyncState* AsyncState, int32 MinIterationsPerTask, int32 NumIterations, TArray<FPCGPoint>& OutPoints, const TFunction<TArray<FPCGPoint>(int32)>& PointFunc)
+	{
+		TRACE_CPUPROFILER_EVENT_SCOPE(IPCGElement::AsyncMultiPointProcessing);
+		check(MinIterationsPerTask > 0 && NumIterations >= 0);
+		if (MinIterationsPerTask <= 0 || NumIterations <= 0)
 		{
-			AsyncPointProcessing(1, MinIterationsPerTask, InPoints.Num(), OutPoints, IterationInnerLoop);
+			return;
 		}
-	}
 
-	void AsyncPointProcessing(int32 NumAvailableTasks, int32 MinIterationsPerTask, int32 NumIterations, TArray<FPCGPoint>& OutPoints, const TFunction<bool(int32, FPCGPoint&)>& PointFunc)
-	{
-		auto IterationInnerLoop = [&PointFunc, &OutPoints](int32 StartIndex, int32 EndIndex) -> int32
+		const int32 NumTasks = GetNumTasks(AsyncState, NumIterations / MinIterationsPerTask);
+		const int32 IterationsPerTask = NumIterations / NumTasks;
+		const int32 NumFutures = NumTasks - 1;
+
+		auto IterationInnerLoop = [&PointFunc](int32 StartIndex, int32 EndIndex) -> TArray<FPCGPoint>
 		{
-			TRACE_CPUPROFILER_EVENT_SCOPE(IPCGElement::AsyncPointProcessing::InnerLoop);
-			int32 NumPointsWritten = 0;
+			TRACE_CPUPROFILER_EVENT_SCOPE(IPCGElement::AsyncMultiPointProcessing::InnerLoop);
+			TArray<FPCGPoint> OutPoints;
 
 			for (int32 Index = StartIndex; Index < EndIndex; ++Index)
 			{
-				if (PointFunc(Index, OutPoints[StartIndex + NumPointsWritten]))
-				{
-					++NumPointsWritten;
-				}
+				OutPoints.Append(PointFunc(Index));
 			}
 
-			return NumPointsWritten;
+			return OutPoints;
 		};
 
-		AsyncPointProcessing(NumAvailableTasks, MinIterationsPerTask, NumIterations, OutPoints, IterationInnerLoop);
-	}
-
-	void AsyncPointFilterProcessing(FPCGContext* Context, int32 NumIterations, TArray<FPCGPoint>& InFilterPoints, TArray<FPCGPoint>& OutFilterPoints, const TFunction<bool(int32, FPCGPoint&, FPCGPoint&)>& PointFunc)
-	{
-		const int32 MinIterationsPerTask = 256;
-
-		FPCGAsyncState* AsyncState = Context ? &Context->AsyncState : nullptr;
+		TArray<UE::Tasks::TTask<TArray<FPCGPoint>>> AsyncTasks;
+		AsyncTasks.Reserve(NumFutures);
 
-		if (AsyncState && !AsyncState->bIsRunningAsyncCall)
+		// Launch the async tasks
+		for (int32 TaskIndex = 0; TaskIndex < NumFutures; ++TaskIndex)
 		{
-			AsyncState->bIsRunningAsyncCall = true;
-			AsyncPointFilterProcessing(AsyncState->NumAvailableTasks, MinIterationsPerTask, NumIterations, InFilterPoints, OutFilterPoints, PointFunc);
-			AsyncState->bIsRunningAsyncCall = false;
+			const int32 StartIndex = TaskIndex * IterationsPerTask;
+			const int32 EndIndex = StartIndex + IterationsPerTask;
+
+			AsyncTasks.Emplace(UE::Tasks::Launch(UE_SOURCE_LOCATION, [&IterationInnerLoop, StartIndex, EndIndex]() -> TArray<FPCGPoint>
+			{
+				return IterationInnerLoop(StartIndex, EndIndex);
+			}));
 		}
-		else
+
+		TArray<FPCGPoint> PointsFromThisThread = IterationInnerLoop(NumFutures * IterationsPerTask, NumIterations);
+
+		// Wait/Gather results & collapse points
 		{
-			AsyncPointFilterProcessing(1, MinIterationsPerTask, NumIterations, InFilterPoints, OutFilterPoints, PointFunc);
+			TRACE_CPUPROFILER_EVENT_SCOPE(IPCGElement::AsyncMultiPointProcessing::WaitAndCollapseArray);
+			for (UE::Tasks::TTask<TArray<FPCGPoint>>& AsyncTask : AsyncTasks)
+			{
+				AsyncTask.Wait();
+				OutPoints.Append(MoveTemp(AsyncTask.GetResult()));
+			}
+
+			OutPoints.Append(MoveTemp(PointsFromThisThread));
 		}
 	}
 
-	void AsyncPointFilterProcessing(int32 NumAvailableTasks, int32 MinIterationsPerTask, int32 NumIterations, TArray<FPCGPoint>& InFilterPoints, TArray<FPCGPoint>& OutFilterPoints, const TFunction<bool(int32, FPCGPoint&, FPCGPoint&)>& PointFunc)
+	void PrivateAsyncPointFilterProcessing(FPCGAsyncState* AsyncState, int32 MinIterationsPerTask, int32 NumIterations, TArray<FPCGPoint>& InFilterPoints, TArray<FPCGPoint>& OutFilterPoints, const TFunction<bool(int32, FPCGPoint&, FPCGPoint&)>& PointFunc)
 	{
 		TRACE_CPUPROFILER_EVENT_SCOPE(IPCGElement::AsyncPointFilterProcessing);
-		check(NumAvailableTasks > 0 && MinIterationsPerTask > 0 && NumIterations >= 0);
-		if (NumAvailableTasks <= 0 || MinIterationsPerTask <= 0 || NumIterations <= 0)
+		check(MinIterationsPerTask > 0 && NumIterations >= 0);
+		if (MinIterationsPerTask <= 0 || NumIterations <= 0)
 		{
 			return;
 		}
 
-		// Get number of available threads from the context
-		const int32 NumTasks = FMath::Min(NumAvailableTasks, FMath::Max(1, NumIterations / MinIterationsPerTask));
+		const int32 NumTasks = GetNumTasks(AsyncState, NumIterations / MinIterationsPerTask);
 		const int32 IterationsPerTask = NumIterations / NumTasks;
 		const int32 NumFutures = NumTasks - 1;
 
@@ -231,7 +241,7 @@ namespace FPCGAsync
 		};
 
 		// Setup [current, last, nb points] data per dispatch
-		TArray<TFuture<TPair<int32, int32>>> AsyncTasks;
+		TArray<UE::Tasks::TTask<TPair<int32, int32>>> AsyncTasks;
 		AsyncTasks.Reserve(NumFutures);
 
 		// Launch the async tasks
@@ -240,7 +250,7 @@ namespace FPCGAsync
 			const int32 StartIndex = TaskIndex * IterationsPerTask;
 			const int32 EndIndex = StartIndex + IterationsPerTask;
 
-			AsyncTasks.Emplace(Async(EAsyncExecution::ThreadPool, [&IterationInnerLoop, StartIndex, EndIndex]() -> TPair<int32, int32>
+			AsyncTasks.Emplace(UE::Tasks::Launch(UE_SOURCE_LOCATION, [&IterationInnerLoop, StartIndex, EndIndex]() -> TPair<int32, int32>
 			{
 				return IterationInnerLoop(StartIndex, EndIndex);
 			}));
@@ -296,9 +306,9 @@ namespace FPCGAsync
 			{
 				const int32 StartPointsIndex = AsyncIndex * IterationsPerTask;
 
-				TFuture<TPair<int32, int32>>& AsyncTask = AsyncTasks[AsyncIndex];
+				UE::Tasks::TTask<TPair<int32, int32>>& AsyncTask = AsyncTasks[AsyncIndex];
 				AsyncTask.Wait();
-				TPair<int32, int32> NumberOfPointsAdded = AsyncTask.Get();
+				TPair<int32, int32> NumberOfPointsAdded = AsyncTask.GetResult();
 
 				RangeIndices = CollapsePoints(RangeIndices, StartPointsIndex, NumberOfPointsAdded);
 			}
@@ -311,77 +321,85 @@ namespace FPCGAsync
 		}
 	}
 
-	void AsyncMultiPointProcessing(FPCGContext* Context, int32 NumIterations, TArray<FPCGPoint>& OutPoints, const TFunction<TArray<FPCGPoint>(int32)>& PointFunc)
+	void AsyncPointProcessing(FPCGContext* Context, int32 NumIterations, TArray<FPCGPoint>& OutPoints, const TFunction<bool(int32, FPCGPoint&)>& PointFunc)
 	{
-		FPCGAsyncState* AsyncState = Context ? &Context->AsyncState : nullptr;
-
 		const int32 MinIterationsPerTask = 256;
-		if (AsyncState && !AsyncState->bIsRunningAsyncCall)
+
+		if (Context && !Context->AsyncState.bIsRunningAsyncCall)
 		{
-			AsyncState->bIsRunningAsyncCall = true;
-			AsyncMultiPointProcessing(AsyncState->NumAvailableTasks, MinIterationsPerTask, NumIterations, OutPoints, PointFunc);
-			AsyncState->bIsRunningAsyncCall = false;
+			Context->AsyncState.bIsRunningAsyncCall = true;
+			PrivateAsyncPointProcessing(&Context->AsyncState, MinIterationsPerTask, NumIterations, OutPoints, PointFunc);
+			Context->AsyncState.bIsRunningAsyncCall = false;
 		}
 		else
 		{
-			AsyncMultiPointProcessing(1, MinIterationsPerTask, NumIterations, OutPoints, PointFunc);
+			// Reentrant case (no async support)
+			PrivateAsyncPointProcessing(nullptr, MinIterationsPerTask, NumIterations, OutPoints, PointFunc);
 		}
 	}
 
-	void AsyncMultiPointProcessing(int32 NumAvailableTasks, int32 MinIterationsPerTask, int32 NumIterations, TArray<FPCGPoint>& OutPoints, const TFunction<TArray<FPCGPoint>(int32)>& PointFunc)
+	void AsyncPointProcessing(FPCGContext* Context, const TArray<FPCGPoint>& InPoints, TArray<FPCGPoint>& OutPoints, const TFunction<bool(const FPCGPoint&, FPCGPoint&)>& PointFunc)
 	{
-		TRACE_CPUPROFILER_EVENT_SCOPE(IPCGElement::AsyncMultiPointProcessing);
-		check(NumAvailableTasks > 0 && MinIterationsPerTask > 0 && NumIterations >= 0);
-		if (NumAvailableTasks <= 0 || MinIterationsPerTask <= 0 || NumIterations <= 0)
-		{
-			return;
-		}
-
-		// Get number of available threads from the context
-		const int32 NumTasks = FMath::Min(NumAvailableTasks, FMath::Max(1, NumIterations / MinIterationsPerTask));
-		const int32 IterationsPerTask = NumIterations / NumTasks;
-		const int32 NumFutures = NumTasks - 1;
+		const int32 MinIterationsPerTask = 256;
 
-		auto IterationInnerLoop = [&PointFunc](int32 StartIndex, int32 EndIndex) -> TArray<FPCGPoint>
+		auto IterationInnerLoop = [&PointFunc, &InPoints, &OutPoints](int32 StartIndex, int32 EndIndex) -> int32
 		{
-			TRACE_CPUPROFILER_EVENT_SCOPE(IPCGElement::AsyncMultiPointProcessing::InnerLoop);
-			TArray<FPCGPoint> OutPoints;
+			TRACE_CPUPROFILER_EVENT_SCOPE(IPCGElement::AsyncPointProcessing::InnerLoop);
+			int32 NumPointsWritten = 0;
 
 			for (int32 Index = StartIndex; Index < EndIndex; ++Index)
 			{
-				OutPoints.Append(PointFunc(Index));
+				if (PointFunc(InPoints[Index], OutPoints[StartIndex + NumPointsWritten]))
+				{
+					++NumPointsWritten;
+				}
 			}
 
-			return OutPoints;
+			return NumPointsWritten;
 		};
-
-		TArray<TFuture<TArray<FPCGPoint>>> AsyncTasks;
-		AsyncTasks.Reserve(NumFutures);
-
-		// Launch the async tasks
-		for (int32 TaskIndex = 0; TaskIndex < NumFutures; ++TaskIndex)
+				
+		if (Context && !Context->AsyncState.bIsRunningAsyncCall)
 		{
-			const int32 StartIndex = TaskIndex * IterationsPerTask;
-			const int32 EndIndex = StartIndex + IterationsPerTask;
-
-			AsyncTasks.Emplace(Async(EAsyncExecution::ThreadPool, [&IterationInnerLoop, StartIndex, EndIndex]() -> TArray<FPCGPoint>
-			{
-				return IterationInnerLoop(StartIndex, EndIndex);
-			}));
+			Context->AsyncState.bIsRunningAsyncCall = true;
+			PrivateAsyncPointProcessing(&Context->AsyncState, MinIterationsPerTask, InPoints.Num(), OutPoints, IterationInnerLoop);
+			Context->AsyncState.bIsRunningAsyncCall = false;
 		}
-
-		TArray<FPCGPoint> PointsFromThisThread = IterationInnerLoop(NumFutures * IterationsPerTask, NumIterations);
-
-		// Wait/Gather results & collapse points
+		else
 		{
-			TRACE_CPUPROFILER_EVENT_SCOPE(IPCGElement::AsyncMultiPointProcessing::WaitAndCollapseArray);
-			for (TFuture<TArray<FPCGPoint>>& AsyncTask : AsyncTasks)
-			{
-				AsyncTask.Wait();
-				OutPoints.Append(AsyncTask.Get());
-			}
-
-			OutPoints.Append(PointsFromThisThread);
+			// Reentrant case (no async support)
+			PrivateAsyncPointProcessing(nullptr, MinIterationsPerTask, InPoints.Num(), OutPoints, IterationInnerLoop);
+		}
+	}
+		
+	void AsyncPointFilterProcessing(FPCGContext* Context, int32 NumIterations, TArray<FPCGPoint>& InFilterPoints, TArray<FPCGPoint>& OutFilterPoints, const TFunction<bool(int32, FPCGPoint&, FPCGPoint&)>& PointFunc)
+	{
+		const int32 MinIterationsPerTask = 256;
+		if (Context && !Context->AsyncState.bIsRunningAsyncCall)
+		{
+			Context->AsyncState.bIsRunningAsyncCall = true;
+			PrivateAsyncPointFilterProcessing(&Context->AsyncState, MinIterationsPerTask, NumIterations, InFilterPoints, OutFilterPoints, PointFunc);
+			Context->AsyncState.bIsRunningAsyncCall = false;
+		}
+		else
+		{
+			// Reentrant case (no async support)
+			PrivateAsyncPointFilterProcessing(nullptr, MinIterationsPerTask, NumIterations, InFilterPoints, OutFilterPoints, PointFunc);
+		}
+	}
+		
+	void AsyncMultiPointProcessing(FPCGContext* Context, int32 NumIterations, TArray<FPCGPoint>& OutPoints, const TFunction<TArray<FPCGPoint>(int32)>& PointFunc)
+	{
+		const int32 MinIterationsPerTask = 256;
+		if (Context && !Context->AsyncState.bIsRunningAsyncCall)
+		{
+			Context->AsyncState.bIsRunningAsyncCall = true;
+			PrivateAsyncMultiPointProcessing(&Context->AsyncState, MinIterationsPerTask, NumIterations, OutPoints, PointFunc);
+			Context->AsyncState.bIsRunningAsyncCall = false;
+		}
+		else
+		{
+			// Reentrant case (no async support)
+			PrivateAsyncMultiPointProcessing(nullptr, MinIterationsPerTask, NumIterations, OutPoints, PointFunc);
 		}
 	}
 }
\ No newline at end of file
diff --git a/Engine/Plugins/PCG/Source/PCG/Private/Helpers/PCGAsyncState.cpp b/Engine/Plugins/PCG/Source/PCG/Private/Helpers/PCGAsyncState.cpp
--- a/Engine/Plugins/PCG/Source/PCG/Private/Helpers/PCGAsyncState.cpp
+++ b/Engine/Plugins/PCG/Source/PCG/Private/Helpers/PCGAsyncState.cpp
@@ -17,4 +17,10 @@ void FPCGAsyncState::Reset()
 	ChunkToNumElementsWrittenMap.Reset();
 
 	bStarted = false;
+}
+
+void FPCGAsyncState::SetStarted(double InEndTime)
+{
+	bStarted = true;
+	EndTime = InEndTime;
 }
\ No newline at end of file
diff --git a/Engine/Plugins/PCG/Source/PCG/Private/Helpers/PCGHelpers.cpp b/Engine/Plugins/PCG/Source/PCG/Private/Helpers/PCGHelpers.cpp
--- a/Engine/Plugins/PCG/Source/PCG/Private/Helpers/PCGHelpers.cpp
+++ b/Engine/Plugins/PCG/Source/PCG/Private/Helpers/PCGHelpers.cpp
@@ -10,6 +10,7 @@
 #include "PCGWorldActor.h"
 #include "Grid/PCGPartitionActor.h"
 
+#include "Components/BillboardComponent.h"
 #include "Landscape.h"
 #include "Algo/AnyOf.h"
 #include "UObject/UObjectIterator.h"
@@ -164,11 +165,15 @@ namespace PCGHelpers
 
 				InActor->ForEachComponent<UPrimitiveComponent>(/*bIncludeFromChildActors=*/true, [bNonColliding, bIgnorePCGCreatedComponents, &WorldToActor, &Box](const UPrimitiveComponent* InPrimComp)
 				{
-					if ((bNonColliding || InPrimComp->IsCollisionEnabled()) &&
-						(!bIgnorePCGCreatedComponents || !InPrimComp->ComponentTags.Contains(DefaultPCGTag)))
+					// Billboard requires access to its texture prevent this from running outside of game thread
+					if(!InPrimComp->IsA<UBillboardComponent>())
 					{
-						const FTransform ComponentToActor = InPrimComp->GetComponentTransform() * WorldToActor;
-						Box += InPrimComp->CalcBounds(ComponentToActor).GetBox();
+						if ((bNonColliding || InPrimComp->IsCollisionEnabled()) &&
+							(!bIgnorePCGCreatedComponents || !InPrimComp->ComponentTags.Contains(DefaultPCGTag)))
+						{
+							const FTransform ComponentToActor = InPrimComp->GetComponentTransform() * WorldToActor;
+							Box += InPrimComp->CalcBounds(ComponentToActor).GetBox();
+						}
 					}
 				});
 			}
diff --git a/Engine/Plugins/PCG/Source/PCG/Private/Metadata/PCGMetadata.cpp b/Engine/Plugins/PCG/Source/PCG/Private/Metadata/PCGMetadata.cpp
--- a/Engine/Plugins/PCG/Source/PCG/Private/Metadata/PCGMetadata.cpp
+++ b/Engine/Plugins/PCG/Source/PCG/Private/Metadata/PCGMetadata.cpp
@@ -1360,7 +1360,12 @@ void UPCGMetadata::SetAttributes(const TArrayView<const PCGMetadataEntryKey>& In
 
 	AttributeLock.ReadLock();
 	int32 AttributeOffset = 0;
-	const int32 AttributesPerDispatch = OptionalContext ? FMath::Max(1, OptionalContext->AsyncState.NumAvailableTasks) : 1;
+	const int32 DefaultAttributesPerDispatch = 64;
+	int32 AttributesPerDispatch = OptionalContext ? DefaultAttributesPerDispatch : 1;
+	if (OptionalContext && OptionalContext->AsyncState.NumAvailableTasks > 0)
+	{
+		AttributesPerDispatch = FMath::Min(OptionalContext->AsyncState.NumAvailableTasks, AttributesPerDispatch);
+	}
 
 	while (AttributeOffset < Attributes.Num())
 	{
diff --git a/Engine/Plugins/PCG/Source/PCG/Private/PCGContext.cpp b/Engine/Plugins/PCG/Source/PCG/Private/PCGContext.cpp
--- a/Engine/Plugins/PCG/Source/PCG/Private/PCGContext.cpp
+++ b/Engine/Plugins/PCG/Source/PCG/Private/PCGContext.cpp
@@ -2,6 +2,7 @@
 
 #include "PCGContext.h"
 #include "PCGComponent.h"
+#include "PCGGraph.h"
 #include "PCGParamData.h"
 #include "PCGPin.h"
 #include "PCGSubsystem.h"
@@ -125,10 +126,42 @@ void FPCGContext::InitializeSettings()
 			}
 
 			if (bHasParamConnected)
-			{
-				SettingsWithOverride = Cast<UPCGSettings>(StaticDuplicateObject(NodeSettings, GetTransientPackage()));
-				SettingsWithOverride->SetFlags(RF_Transient);
+			{				
+				FObjectDuplicationParameters DuplicateParams(const_cast<UPCGSettings*>(NodeSettings), GetTransientPackage());
+				const bool bIsInGameThread = IsInGameThread();
+				DuplicateParams.bSkipPostLoad = !bIsInGameThread;
+				DuplicateParams.ApplyFlags = RF_Transient;
+			
+				TMap<UObject*, UObject*> CreatedObjects;
+				DuplicateParams.CreatedObjects = &CreatedObjects;
+				
+				{
+					FGCScopeGuard Scope;
+					SettingsWithOverride = Cast<UPCGSettings>(StaticDuplicateObjectEx(DuplicateParams));
+				}
+			
+				// StaticDuplicateObjectEx will not allow PostLoad to be called outside of GameThread (except AsyncLoadingThread) so we take care of the PostLoad and clear Async flags
+				for (auto& KeyValuePair : CreatedObjects)
+				{
+					if (KeyValuePair.Value)
+					{
+						if (!bIsInGameThread)
+						{
+							ensure(KeyValuePair.Value->IsPostLoadThreadSafe());
+							KeyValuePair.Value->ConditionalPostLoad();
+							KeyValuePair.Value->ClearInternalFlags(EInternalObjectFlags::Async);
+						}
 
+#if WITH_EDITOR
+						// @todo_pcg: find a way to avoid the call to SetupCallbacks() all together but for now unregister the callbacks after duplication
+						if (UPCGGraphInstance* GraphInstance = Cast<UPCGGraphInstance>(KeyValuePair.Value))
+						{
+							GraphInstance->TeardownCallbacks();
+						}
+#endif
+					}
+				}
+				
 				// Force seed copy to prevent issue due to delta serialization vs. Seed being initialized in the constructor only for new nodes
 				SettingsWithOverride->Seed = NodeSettings->Seed;
 				SettingsWithOverride->OriginalSettings = NodeSettings;
@@ -249,6 +282,8 @@ void FPCGContext::OverrideSettings()
 				return false;
 			}
 
+			// Setting properties (ex: FSoftObjectPath) can end up doing StaticFindObject which needs to be protected from running at same time as GC
+			FGCScopeGuard GCScope;
 			FPCGAttributeAccessorKeysSingleObjectPtr PropertyObjectKey(Container);
 			PropertyAccessor->Set<PropertyType>(Value, PropertyObjectKey);
 
diff --git a/Engine/Plugins/PCG/Source/PCG/Private/PCGElement.cpp b/Engine/Plugins/PCG/Source/PCG/Private/PCGElement.cpp
--- a/Engine/Plugins/PCG/Source/PCG/Private/PCGElement.cpp
+++ b/Engine/Plugins/PCG/Source/PCG/Private/PCGElement.cpp
@@ -34,7 +34,7 @@ static TAutoConsoleVariable<bool> CVarPCGValidatePointMetadata(
 
 bool IPCGElement::Execute(FPCGContext* Context) const
 {
-	check(Context && Context->AsyncState.NumAvailableTasks > 0 && Context->CurrentPhase < EPCGExecutionPhase::Done);
+	check(Context && Context->AsyncState.NumAvailableTasks != 0 && Context->CurrentPhase < EPCGExecutionPhase::Done);
 	check(Context->AsyncState.bIsRunningOnMainThread || !CanExecuteOnlyOnMainThread(Context));
 
 	while (Context->CurrentPhase != EPCGExecutionPhase::Done)
diff --git a/Engine/Plugins/PCG/Source/PCG/Private/PCGSubgraph.cpp b/Engine/Plugins/PCG/Source/PCG/Private/PCGSubgraph.cpp
--- a/Engine/Plugins/PCG/Source/PCG/Private/PCGSubgraph.cpp
+++ b/Engine/Plugins/PCG/Source/PCG/Private/PCGSubgraph.cpp
@@ -747,7 +747,10 @@ bool FPCGSubgraphElement::ExecuteInternal(FPCGContext* InContext) const
 					// as we just can call the GetOutputData on a fresh data collection and merge it afterwards in the output data,
 					// but this is not needed here so we will keep the full assignment & benefit from the crc as well.
 					ensure(Context->SubgraphTaskIds.Num() == 1);
-					Subsystem->GetOutputData(Context->SubgraphTaskIds[0], Context->OutputData, /*bClearDataOnGet=*/true);
+					if (Subsystem->GetOutputData(Context->SubgraphTaskIds[0], Context->OutputData))
+					{
+						Subsystem->ClearOutputData(Context->SubgraphTaskIds[0]);
+					}
 				}
 			}
 			else
diff --git a/Engine/Plugins/PCG/Source/PCG/Private/PCGSubsystem.cpp b/Engine/Plugins/PCG/Source/PCG/Private/PCGSubsystem.cpp
--- a/Engine/Plugins/PCG/Source/PCG/Private/PCGSubsystem.cpp
+++ b/Engine/Plugins/PCG/Source/PCG/Private/PCGSubsystem.cpp
@@ -1229,10 +1229,16 @@ FPCGGenSourceManager* UPCGSubsystem::GetGenSourceManager() const
 	return RuntimeGenScheduler ? RuntimeGenScheduler->GenSourceManager : nullptr;
 }
 
-bool UPCGSubsystem::GetOutputData(FPCGTaskId TaskId, FPCGDataCollection& OutData, bool bClearDataOnGet)
+bool UPCGSubsystem::GetOutputData(FPCGTaskId TaskId, FPCGDataCollection& OutData)
 {
 	check(GraphExecutor);
-	return GraphExecutor->GetOutputData(TaskId, OutData, bClearDataOnGet);
+	return GraphExecutor->GetOutputData(TaskId, OutData);
+}
+
+void UPCGSubsystem::ClearOutputData(FPCGTaskId TaskId)
+{
+	check(GraphExecutor);
+	GraphExecutor->ClearOutputData(TaskId);
 }
 
 #if WITH_EDITOR
diff --git a/Engine/Plugins/PCG/Source/PCG/Private/Tests/Subgraphs/PCGCollapseSubgraphTest.cpp b/Engine/Plugins/PCG/Source/PCG/Private/Tests/Subgraphs/PCGCollapseSubgraphTest.cpp
--- a/Engine/Plugins/PCG/Source/PCG/Private/Tests/Subgraphs/PCGCollapseSubgraphTest.cpp
+++ b/Engine/Plugins/PCG/Source/PCG/Private/Tests/Subgraphs/PCGCollapseSubgraphTest.cpp
@@ -283,7 +283,7 @@ bool FPCGCollapseSubgraphCombinationOfCollapses::RunTest(const FString& Paramete
 		const FPCGTaskId FinalTaskId = GraphExecutor.ScheduleGeneric([&OutputPoints, &GraphExecutor, TaskId, &bDone]() -> bool
 		{
 			FPCGDataCollection OutputData{};
-			GraphExecutor.GetOutputData(TaskId, OutputData, /*bClearDataOnGet=*/false);
+			GraphExecutor.GetOutputData(TaskId, OutputData);
 
 			if (OutputData.TaggedData.Num() == 1)
 			{
diff --git a/Engine/Plugins/PCG/Source/PCG/Public/Data/PCGSpatialDataTpl.h b/Engine/Plugins/PCG/Source/PCG/Public/Data/PCGSpatialDataTpl.h
--- a/Engine/Plugins/PCG/Source/PCG/Public/Data/PCGSpatialDataTpl.h
+++ b/Engine/Plugins/PCG/Source/PCG/Public/Data/PCGSpatialDataTpl.h
@@ -61,6 +61,6 @@ namespace FPCGSpatialDataProcessing
 			OutPoints.SetNum(Count);
 		};
 
-		ensure(FPCGAsync::AsyncProcessingRangeEx(AsyncState, NumIterations, Initialize, ProcessRange, MoveDataRange, Finished, /*bEnableTimeSlicing=*/false, ChunkSize));
+		ensure(FPCGAsync::AsyncProcessingRangeEx(AsyncState, NumIterations, Initialize, ProcessRange, MoveDataRange, Finished, /*bEnableTimeSlicing=*/false, ChunkSize, /*bAllowChunkSizeOverride=*/false));
 	}
 }
\ No newline at end of file
diff --git a/Engine/Plugins/PCG/Source/PCG/Public/Helpers/PCGAsync.h b/Engine/Plugins/PCG/Source/PCG/Public/Helpers/PCGAsync.h
--- a/Engine/Plugins/PCG/Source/PCG/Public/Helpers/PCGAsync.h
+++ b/Engine/Plugins/PCG/Source/PCG/Public/Helpers/PCGAsync.h
@@ -20,6 +20,8 @@ namespace FPCGAsync
 	namespace ConsoleVar
 	{
 		extern PCG_API TAutoConsoleVariable<bool> CVarDisableAsyncTimeSlicing;
+		extern PCG_API TAutoConsoleVariable<bool> CVarDisableAsyncTimeSlicingOnGameThread;
+		extern PCG_API TAutoConsoleVariable<float> CVarAsyncOutOfTickBudgetInMilliseconds;
 		extern PCG_API TAutoConsoleVariable<int32> CVarAsyncOverrideChunkSize;
 	};
 
@@ -41,16 +43,6 @@ namespace FPCGAsync
 	*/
 	PCG_API void AsyncPointProcessing(FPCGContext* Context, const TArray<FPCGPoint>& InPoints, TArray<FPCGPoint>& OutPoints, const TFunction<bool(const FPCGPoint&, FPCGPoint&)>& PointFunc);
 
-	/** 
-	* Helper to do more general 1:1 point processing loops
-	* @param NumAvailableTasks - The upper bound on the number of async tasks we'll start
-	* @param MinIterationsPerTask - The lower bound on the number of iterations per task we'll dispatch
-	* @param NumIterations - The number of calls that will be done to the provided function, also an upper bound on the number of points generated
-	* @param OutPoints - The array in which the results will be written to. Note that the array will be cleared before execution
-	* @param PointFunc - A function that has the index [0; NumIterations] and has to write to the point & return true when the current call generates a point
-	*/
-	void AsyncPointProcessing(int32 NumAvailableTasks, int32 MinIterationsPerTask, int32 NumIterations, TArray<FPCGPoint>& OutPoints, const TFunction<bool(int32, FPCGPoint&)>& PointFunc);
-
 	/**
 	* Helper to do simple point filtering loops
 	* @param Context - The context containing the information about how many tasks to launch
@@ -61,17 +53,6 @@ namespace FPCGAsync
 	*/
 	PCG_API void AsyncPointFilterProcessing(FPCGContext* Context, int32 NumIterations, TArray<FPCGPoint>& InFilterPoints, TArray<FPCGPoint>& OutFilterPoints, const TFunction<bool(int32, FPCGPoint&, FPCGPoint&)>& PointFunc);
 
-	/**
-	* Helper to do more general 1:1 point filtering loops
-	* @param NumAvailableTasks - The upper bound on the number of async tasks we'll start
-	* @param MinIterationsPerTask - The lower bound on the number of iterations per task we'll dispatch
-	* @param NumIterations - The number of calls that will be done to the provided function, also an upper bound on the number of points generated
-	* @param InFilterPoints - The array in which the in-filter results will be written to. Note that the array will be cleared before execution
-	* @param OutFilterPoints - The array in which the out-filter results will be written to. Note that the array will be cleared before execution
-	* @param PointFunc - A function that has the index [0; NumIterations] and has to write to the point & return true when the current call generates a point
-	*/
-	void AsyncPointFilterProcessing(int32 NumAvailableTasks, int32 MinIterationsPerTask, int32 NumIterations, TArray<FPCGPoint>& InFilterPoints, TArray<FPCGPoint>& OutFilterPoints, const TFunction<bool(int32, FPCGPoint&, FPCGPoint&)>& PointFunc);
-
 	/**
 	* Helper to do simple 1:N point processing loops
 	* @param Context - The context containing the information about how many tasks to launch
@@ -81,29 +62,22 @@ namespace FPCGAsync
 	*/
 	PCG_API void AsyncMultiPointProcessing(FPCGContext* Context, int32 NumIterations, TArray<FPCGPoint>& OutPoints, const TFunction<TArray<FPCGPoint>(int32)>& PointFunc);
 
-	/** 
-	* Helper to do more general 1:N point processing loops
-	* @param NumAvailableTasks - The upper bound on the number of async tasks we'll start
-	* @param MinIterationsPerTask - The lower bound on the number of iterations per task we'll dispatch
-	* @param NumIterations - The number of calls that will be done to the provided function, also an upper bound on the number of points generated
- 	* @param OutPoints - The array in which the results will be written to. Note that the array will be cleared before execution
-	* @param PointFunc - A function that has the index [0; NumIterations] and has to write to the point & return true when the current call generates a point
-	*/
-	void AsyncMultiPointProcessing(int32 NumAvailableTasks, int32 MinIterationsPerTask, int32 NumIterations, TArray<FPCGPoint>& OutPoints, const TFunction<TArray<FPCGPoint>(int32)>& PointFunc);
-
 	namespace Private
 	{
 	template <typename InitializeFunc, typename InnerLoopFunc, typename MoveDataRangeFunc, typename FinishedFunc>
-	bool AsyncProcessing(FPCGAsyncState& AsyncState, int32 NumIterations, InitializeFunc&& Initialize, InnerLoopFunc&& IterationInnerLoop, MoveDataRangeFunc&& MoveDataRange, FinishedFunc&& Finished, const bool bInEnableTimeSlicing, const int32 InChunkSize)
+	bool AsyncProcessing(FPCGAsyncState& AsyncState, int32 NumIterations, InitializeFunc&& Initialize, InnerLoopFunc&& IterationInnerLoop, MoveDataRangeFunc&& MoveDataRange, FinishedFunc&& Finished, const bool bInEnableTimeSlicing, const int32 InChunkSize, const bool bAllowChunkSizeOverride)
 	{
 		TRACE_CPUPROFILER_EVENT_SCOPE(FPCGAsync::AsyncProcessing);
 
 		const int32 OverrideChunkSize = ConsoleVar::CVarAsyncOverrideChunkSize.GetValueOnAnyThread();
-		const int32 ChunkSize = OverrideChunkSize > 0 ? OverrideChunkSize : InChunkSize;
+		const int32 ChunkSize = bAllowChunkSizeOverride && OverrideChunkSize > 0 ? OverrideChunkSize : InChunkSize;
 
-		const bool bEnableTimeSlicing = bInEnableTimeSlicing && !ConsoleVar::CVarDisableAsyncTimeSlicing.GetValueOnAnyThread();
+		const bool bIsInGameThread = IsInGameThread();
+		const bool bEnableTimeSlicing = bInEnableTimeSlicing && ((bIsInGameThread && !ConsoleVar::CVarDisableAsyncTimeSlicingOnGameThread.GetValueOnAnyThread()) || (!bIsInGameThread && !ConsoleVar::CVarDisableAsyncTimeSlicing.GetValueOnAnyThread()));
+				
+		const float OutOfTickBudgetInSeconds = ConsoleVar::CVarAsyncOutOfTickBudgetInMilliseconds.GetValueOnAnyThread() / 1000.f;
 
-		if (AsyncState.NumAvailableTasks <= 0 || ChunkSize <= 0 || NumIterations <= 0)
+		if (AsyncState.NumAvailableTasks == 0 || ChunkSize <= 0 || NumIterations <= 0)
 		{
 			// Invalid request
 			return true;
@@ -123,14 +97,26 @@ namespace FPCGAsync
 
 			Initialize();
 
-			AsyncState.bStarted = true;
+			double EndTime = AsyncState.EndTime;
+			
+			if (AsyncState.bIsRunningOutOfTick && bEnableTimeSlicing && OutOfTickBudgetInSeconds > 0.f)
+			{
+				EndTime = FPlatformTime::Seconds() + OutOfTickBudgetInSeconds;
+			}
+
+			AsyncState.SetStarted(EndTime);
 
 			// At the beginning, dispatch the chunks that each task needs to process.
 			// It will be stored on the AsyncState so that it can be reused.
 			// Each task will update its current counter directly in the AsyncState.
 			// It's OK to not protect it since they all will read/write from different place in memory.
 
-			int32 NumTasks = FMath::Min(AsyncState.NumAvailableTasks, ChunksNumber);
+			int32 NumTasks = ChunksNumber;
+			if (AsyncState.NumAvailableTasks > 0)
+			{
+				NumTasks = FMath::Min(AsyncState.NumAvailableTasks, NumTasks);
+			}
+
 			const int32 IterationsPerTask = ChunksNumber / NumTasks;
 			
 			// Some tasks will have an extra chunk of work if NumTasks doesn't divide ChunksNumber.
@@ -270,8 +256,11 @@ namespace FPCGAsync
 				{
 					SynchroStruct.bQuit = true;
 
-					// Wait for all futures to finish their job
-					UE::Tasks::Wait(AsyncTasks);
+					{
+						TRACE_CPUPROFILER_EVENT_SCOPE(FPCGAsync::AsyncProcessing::Wait);
+						// Wait for all futures to finish their job
+						UE::Tasks::Wait(AsyncTasks);
+					}
 
 					// Make sure to dequeue everything
 					FlushQueue();
@@ -377,15 +366,16 @@ namespace FPCGAsync
 	* @param Finished - Signature: void(int32 Count). Called once on finished, and tells you the total count of points written.
 	* @param bEnableTimeSlicing - If false, we will not stop until all the processing is done.
 	* @param ChunkSize - Size of the chunks to cut the input data with
+	* @param bAllowChunkSizeOverride - If true, ChunkSize can be overridden by 'pcg.AsyncOverrideChunkSize' CVar
 	* @returns true if the processing is done, false otherwise. Use this to know if you need to reschedule the task.
 	*/
 	template <typename InitializeFunc, typename ProcessRangeFunc, typename MoveDataRangeFunc, typename FinishedFunc>
-	bool AsyncProcessingRangeEx(FPCGAsyncState* AsyncState, int32 NumIterations, InitializeFunc&& Initialize, ProcessRangeFunc&& ProcessRange, MoveDataRangeFunc&& MoveDataRange, FinishedFunc&& Finished, const bool bEnableTimeSlicing, const int32 ChunkSize = 64)
+	bool AsyncProcessingRangeEx(FPCGAsyncState* AsyncState, int32 NumIterations, InitializeFunc&& Initialize, ProcessRangeFunc&& ProcessRange, MoveDataRangeFunc&& MoveDataRange, FinishedFunc&& Finished, const bool bEnableTimeSlicing, const int32 ChunkSize = 64, const bool bAllowChunkSizeOverride = true)
 	{
 		if (AsyncState && !AsyncState->bIsRunningAsyncCall)
 		{
 			AsyncState->bIsRunningAsyncCall = true;
-			const bool bIsDone = Private::AsyncProcessing(*AsyncState, NumIterations, Initialize, ProcessRange, MoveDataRange, Finished, bEnableTimeSlicing, ChunkSize);
+			const bool bIsDone = Private::AsyncProcessing(*AsyncState, NumIterations, Initialize, ProcessRange, MoveDataRange, Finished, bEnableTimeSlicing, ChunkSize, bAllowChunkSizeOverride);
 			AsyncState->bIsRunningAsyncCall = false;
 			return bIsDone;
 		}
@@ -396,7 +386,7 @@ namespace FPCGAsync
 			FPCGAsyncState DummyState;
 			DummyState.NumAvailableTasks = 1;
 			DummyState.bIsRunningAsyncCall = true;
-			return Private::AsyncProcessing(DummyState, NumIterations, Initialize, ProcessRange, MoveDataRange, Finished, /*bEnableTimeSlicing=*/false, ChunkSize);
+			return Private::AsyncProcessing(DummyState, NumIterations, Initialize, ProcessRange, MoveDataRange, Finished, /*bEnableTimeSlicing=*/false, ChunkSize, bAllowChunkSizeOverride);
 		}
 	}
 
@@ -420,10 +410,11 @@ namespace FPCGAsync
 	* @param Finished - Signature: void(int32 Count). Called once on finished, and tells you the total count of points written.
 	* @param bEnableTimeSlicing - If false, we will not stop until all the processing is done.
 	* @param ChunkSize - Size of the chunks to cut the input data with
+	* @param bAllowChunkSizeOverride - If true, ChunkSize can be overridden by 'pcg.AsyncOverrideChunkSize' CVar
 	* @returns true if the processing is done, false otherwise. Use this to know if you need to reschedule the task.
 	*/
 	template <typename InitializeFunc, typename ProcessRangeFunc>
-	bool AsyncProcessingOneToOneRangeEx(FPCGAsyncState* AsyncState, int32 NumIterations, InitializeFunc&& Initialize, ProcessRangeFunc&& ProcessRange, const bool bEnableTimeSlicing, const int32 ChunkSize = 64)
+	bool AsyncProcessingOneToOneRangeEx(FPCGAsyncState* AsyncState, int32 NumIterations, InitializeFunc&& Initialize, ProcessRangeFunc&& ProcessRange, const bool bEnableTimeSlicing, const int32 ChunkSize = 64, const bool bAllowChunkSizeOverride = true)
 	{
 		auto MoveDataRange = [](int32, int32, int32) { ensure(false); };
 		auto Finished = [NumIterations](int32 Count) { ensure(NumIterations == Count); };
@@ -431,7 +422,7 @@ namespace FPCGAsync
 		if (AsyncState && !AsyncState->bIsRunningAsyncCall)
 		{
 			AsyncState->bIsRunningAsyncCall = true;
-			const bool bIsDone = Private::AsyncProcessing(*AsyncState, NumIterations, Initialize, ProcessRange, MoveDataRange, Finished, bEnableTimeSlicing, ChunkSize);
+			const bool bIsDone = Private::AsyncProcessing(*AsyncState, NumIterations, Initialize, ProcessRange, MoveDataRange, Finished, bEnableTimeSlicing, ChunkSize, bAllowChunkSizeOverride);
 			AsyncState->bIsRunningAsyncCall = false;
 			return bIsDone;
 		}
@@ -442,7 +433,7 @@ namespace FPCGAsync
 			FPCGAsyncState DummyState;
 			DummyState.NumAvailableTasks = 1;
 			DummyState.bIsRunningAsyncCall = true;
-			return Private::AsyncProcessing(DummyState, NumIterations, Initialize, ProcessRange, MoveDataRange, Finished, /*bEnableTimeSlicing=*/false, ChunkSize);
+			return Private::AsyncProcessing(DummyState, NumIterations, Initialize, ProcessRange, MoveDataRange, Finished, /*bEnableTimeSlicing=*/false, ChunkSize, bAllowChunkSizeOverride);
 		}
 	}
 
@@ -464,11 +455,12 @@ namespace FPCGAsync
 	* @param MoveData - Signature: void(int32 ReadIndex, int32 WriteIndex). If the processing filters points, this will be used to move elements from one index to another 
 	* @param Finished - Signature: void(int32 Count). Called once on finished, and tells you the total count of points written.
 	* @param bEnableTimeSlicing - If false, we will not stop until all the processing is done.
-	* @param ChunkSize - Size of the chunks to cut the input data with
+	* @param ChunkSize - Size of the chunks to cut the input data with 
+	* @param bAllowChunkSizeOverride - If true, ChunkSize can be overridden by 'pcg.AsyncOverrideChunkSize' CVar
 	* @returns true if the processing is done, false otherwise. Use this to know if you need to reschedule the task.
 	*/
 	template <typename InitializeFunc, typename ProcessElementFunc, typename MoveFunc, typename FinishedFunc>
-	bool AsyncProcessingEx(FPCGAsyncState* AsyncState, int32 NumIterations, InitializeFunc&& Initialize, ProcessElementFunc&& ProcessElement, MoveFunc&& MoveData, FinishedFunc&& Finished, const bool bEnableTimeSlicing, const int32 ChunkSize = 64)
+	bool AsyncProcessingEx(FPCGAsyncState* AsyncState, int32 NumIterations, InitializeFunc&& Initialize, ProcessElementFunc&& ProcessElement, MoveFunc&& MoveData, FinishedFunc&& Finished, const bool bEnableTimeSlicing, const int32 ChunkSize = 64, const bool bAllowChunkSizeOverride = true)
 	{
 		auto IterationInnerLoop = [Func = MoveTemp(ProcessElement)](int32 StartReadIndex, int32 StartWriteIndex, int32 Count) -> int32
 		{
@@ -493,7 +485,7 @@ namespace FPCGAsync
 			}
 		};
 
-		return AsyncProcessingRangeEx(AsyncState, NumIterations, Initialize, IterationInnerLoop, MoveDataRange, Finished, bEnableTimeSlicing, ChunkSize);
+		return AsyncProcessingRangeEx(AsyncState, NumIterations, Initialize, IterationInnerLoop, MoveDataRange, Finished, bEnableTimeSlicing, ChunkSize, bAllowChunkSizeOverride);
 	}
  
  	/**
@@ -516,10 +508,11 @@ namespace FPCGAsync
 	* @param ProcessElement - Signature: bool(int32 ReadIndex, int32 WriteIndex). A function that processes  
 	* @param bEnableTimeSlicing - If false, we will not stop until all the processing is done.
 	* @param ChunkSize - Size of the chunks to cut the input data with
+	* @param bAllowChunkSizeOverride - If true, ChunkSize can be overridden by 'pcg.AsyncOverrideChunkSize' CVar
 	* @returns true if the processing is done, false otherwise. Use this to know if you need to reschedule the task.
 	*/
  	template <typename InitializeFunc, typename ProcessElementFunc>
-	bool AsyncProcessingOneToOneEx(FPCGAsyncState* AsyncState, int32 NumIterations, InitializeFunc&& Initialize, ProcessElementFunc&& ProcessElement, const bool bEnableTimeSlicing, const int32 ChunkSize = 64)
+	bool AsyncProcessingOneToOneEx(FPCGAsyncState* AsyncState, int32 NumIterations, InitializeFunc&& Initialize, ProcessElementFunc&& ProcessElement, const bool bEnableTimeSlicing, const int32 ChunkSize = 64, const bool bAllowChunkSizeOverride = true)
 	{
 		auto IterationInnerLoop = [Func = MoveTemp(ProcessElement)](int32 StartReadIndex, int32 StartWriteIndex, int32 Count) -> int32
 		{
@@ -537,7 +530,7 @@ namespace FPCGAsync
 		if (AsyncState && !AsyncState->bIsRunningAsyncCall)
 		{
 			AsyncState->bIsRunningAsyncCall = true;
-			const bool bIsDone = Private::AsyncProcessing(*AsyncState, NumIterations, Initialize, IterationInnerLoop, MoveDataRange, Finished, bEnableTimeSlicing, ChunkSize);
+			const bool bIsDone = Private::AsyncProcessing(*AsyncState, NumIterations, Initialize, IterationInnerLoop, MoveDataRange, Finished, bEnableTimeSlicing, ChunkSize, bAllowChunkSizeOverride);
 			AsyncState->bIsRunningAsyncCall = false;
 			return bIsDone;
 		}
@@ -548,7 +541,7 @@ namespace FPCGAsync
 			FPCGAsyncState DummyState;
 			DummyState.NumAvailableTasks = 1;
 			DummyState.bIsRunningAsyncCall = true;
-			return Private::AsyncProcessing(DummyState, NumIterations, Initialize, IterationInnerLoop, MoveDataRange, Finished, /*bEnableTimeSlicing=*/false, ChunkSize);
+			return Private::AsyncProcessing(DummyState, NumIterations, Initialize, IterationInnerLoop, MoveDataRange, Finished, /*bEnableTimeSlicing=*/false, ChunkSize, bAllowChunkSizeOverride);
 		}
 	}
 
@@ -569,10 +562,11 @@ namespace FPCGAsync
 	* @param Func - Signature: bool(int32, OutputType&). A function that has the index [0; NumIterations] and has to write to some data & return false if the result should be discarded. 
 	* @param bEnableTimeSlicing - If false, we will not stop until all the processing is done.
 	* @param ChunkSize - Size of the chunks to cut the input data with
+	* @param bAllowChunkSizeOverride - If true, ChunkSize can be overridden by 'pcg.AsyncOverrideChunkSize' CVar
 	* @returns true if the processing is done, false otherwise. Use this to know if you need to reschedule the task.
 	*/
 	template <typename OutputType, typename Func>
-	bool AsyncProcessing(FPCGAsyncState* AsyncState, int32 NumIterations, TArray<OutputType>& OutData, Func&& InFunc, const bool bEnableTimeSlicing, const int32 ChunkSize = 64)
+	bool AsyncProcessing(FPCGAsyncState* AsyncState, int32 NumIterations, TArray<OutputType>& OutData, Func&& InFunc, const bool bEnableTimeSlicing, const int32 ChunkSize = 64, bool bAllowChunkSizeOverride = true)
 	{
 		auto Initialize = [&OutData, NumIterations]()
 		{
@@ -596,7 +590,7 @@ namespace FPCGAsync
 			OutData.SetNum(Count);
 		};
 
-		return AsyncProcessingEx(AsyncState, NumIterations, Initialize, IterationInnerLoop, MoveData, Finished, bEnableTimeSlicing, ChunkSize);
+		return AsyncProcessingEx(AsyncState, NumIterations, Initialize, IterationInnerLoop, MoveData, Finished, bEnableTimeSlicing, ChunkSize, bAllowChunkSizeOverride);
 	}
 }
 
diff --git a/Engine/Plugins/PCG/Source/PCG/Public/Helpers/PCGAsyncState.h b/Engine/Plugins/PCG/Source/PCG/Public/Helpers/PCGAsyncState.h
--- a/Engine/Plugins/PCG/Source/PCG/Public/Helpers/PCGAsyncState.h
+++ b/Engine/Plugins/PCG/Source/PCG/Public/Helpers/PCGAsyncState.h
@@ -53,9 +53,15 @@ struct FPCGAsyncState
 	/** True if currently inside a PCGAsync scope - will prevent further async processing. */
 	bool bIsRunningAsyncCall = false;
 
+	/** True if we can run out of tick */
+	bool bIsRunningOutOfTick = false;
+
 	/** Returns true if we reached end time. */
 	PCG_API bool ShouldStop() const;
 
 	/** To be called at the end of an async processing. */
 	PCG_API void Reset();
+
+	/** To be called to change the state to started and provide an EndTime for timeslicing */
+	PCG_API void SetStarted(double InEndTime);
 };
\ No newline at end of file
diff --git a/Engine/Plugins/PCG/Source/PCG/Public/PCGContext.h b/Engine/Plugins/PCG/Source/PCG/Public/PCGContext.h
--- a/Engine/Plugins/PCG/Source/PCG/Public/PCGContext.h
+++ b/Engine/Plugins/PCG/Source/PCG/Public/PCGContext.h
@@ -52,6 +52,8 @@ struct PCG_API FPCGContext
 	GENERATED_BODY()
 
 	friend class FPCGGraphExecutor;
+	friend struct FPCGGraphActiveTask;
+	friend struct FPCGGraphTask;
 
 	virtual ~FPCGContext() = default;
 
diff --git a/Engine/Plugins/PCG/Source/PCG/Public/PCGGraph.h b/Engine/Plugins/PCG/Source/PCG/Public/PCGGraph.h
--- a/Engine/Plugins/PCG/Source/PCG/Public/PCGGraph.h
+++ b/Engine/Plugins/PCG/Source/PCG/Public/PCGGraph.h
@@ -464,6 +464,7 @@ public:
 	/** ~End UPCGGraphInterface interface */
 
 	// ~Begin UObject interface
+	virtual bool IsPostLoadThreadSafe() const override { return true; }
 	virtual void PostLoad() override;
 	virtual void PostDuplicate(bool bDuplicateForPIE) override;
 	virtual void PostEditImport() override;
diff --git a/Engine/Plugins/PCG/Source/PCG/Public/PCGSettings.h b/Engine/Plugins/PCG/Source/PCG/Public/PCGSettings.h
--- a/Engine/Plugins/PCG/Source/PCG/Public/PCGSettings.h
+++ b/Engine/Plugins/PCG/Source/PCG/Public/PCGSettings.h
@@ -222,6 +222,7 @@ public:
 	// ~End UPCGSettingsInterface interface
 
 	//~Begin UObject interface
+	virtual bool IsPostLoadThreadSafe() const override final { return true; }
 	virtual void PostLoad() override;
 	virtual void Serialize(FArchive& Ar) override;
 	virtual void PostSaveRoot(FObjectPostSaveRootContext ObjectSaveContext) override;
diff --git a/Engine/Plugins/PCG/Source/PCG/Public/PCGSubsystem.h b/Engine/Plugins/PCG/Source/PCG/Public/PCGSubsystem.h
--- a/Engine/Plugins/PCG/Source/PCG/Public/PCGSubsystem.h
+++ b/Engine/Plugins/PCG/Source/PCG/Public/PCGSubsystem.h
@@ -180,7 +180,10 @@ public:
 	void CancelAllGeneration();
 
 	/** Gets the output data for a given task */
-	bool GetOutputData(FPCGTaskId InTaskId, FPCGDataCollection& OutData, bool bClearDataOnGet = false);
+	bool GetOutputData(FPCGTaskId InTaskId, FPCGDataCollection& OutData);
+
+	/** Clears the output data for a given task. Should only be called on tasks with bNeedsManualClear set to true */
+	void ClearOutputData(FPCGTaskId InTaskId);
 
 	/** Register a new PCG Component or update it, will be added to the octree if it doesn't exists yet. Returns true if it was added/updated. Thread safe */
 	bool RegisterOrUpdatePCGComponent(UPCGComponent* InComponent, bool bDoActorMapping = true) { return ActorAndComponentMapping.RegisterOrUpdatePCGComponent(InComponent, bDoActorMapping); }
