Added option to get components reference in the Get Actor Data node

diff --git a/Engine/Plugins/PCG/Source/PCG/Private/Elements/PCGDataFromActor.cpp b/Engine/Plugins/PCG/Source/PCG/Private/Elements/PCGDataFromActor.cpp
--- a/Engine/Plugins/PCG/Source/PCG/Private/Elements/PCGDataFromActor.cpp
+++ b/Engine/Plugins/PCG/Source/PCG/Private/Elements/PCGDataFromActor.cpp
@@ -10,6 +10,7 @@
 #include "Data/PCGPointData.h"
 #include "Data/PCGSpatialData.h"
 #include "Data/PCGVolumeData.h"
+#include "Elements/PCGAddComponent.h"
 #include "Elements/PCGMergeElement.h"
 #include "Elements/PCGMergeAttributes.h"
 #include "Grid/PCGPartitionActor.h"
@@ -226,7 +227,7 @@ EPCGDataType UPCGDataFromActorSettings::GetCurrentPinTypes(const UPCGPin* InPin)
 		{
 			return EPCGDataType::Point;
 		}
-		else if (Mode == EPCGGetDataFromActorMode::GetActorReference)
+		else if (Mode == EPCGGetDataFromActorMode::GetActorReference || Mode == EPCGGetDataFromActorMode::GetComponentsReference)
 		{
 			return EPCGDataType::Param;
 		}
@@ -584,10 +585,24 @@ void FPCGDataFromActorElement::ProcessActors(FPCGContext* Context, const UPCGDat
 {
 	// Special case:
 	// If we're asking for single point with the merge single point data, we can do a more efficient process
-	if ((Settings->Mode == EPCGGetDataFromActorMode::GetSinglePoint || Settings->Mode == EPCGGetDataFromActorMode::GetActorReference) && Settings->bMergeSinglePointData && FoundActors.Num() > 1)
+	if (Settings->Mode == EPCGGetDataFromActorMode::GetSinglePoint && Settings->bMergeSinglePointData && FoundActors.Num() > 1)
 	{
 		MergeActorsIntoData(Context, Settings, FoundActors);
 	}
+	else if((Settings->Mode == EPCGGetDataFromActorMode::GetActorReference || Settings->Mode == EPCGGetDataFromActorMode::GetComponentsReference))
+	{
+		if(Settings->bMergeSinglePointData)
+		{
+			CreateReferenceData(Context, Settings, FoundActors);
+		}
+		else
+		{
+			for(AActor* Actor : FoundActors)
+			{
+				CreateReferenceData(Context, Settings, {Actor});
+			}
+		}
+	}
 	else
 	{
 		for (AActor* Actor : FoundActors)
@@ -597,115 +612,95 @@ void FPCGDataFromActorElement::ProcessActors(FPCGContext* Context, const UPCGDat
 	}
 }
 
-void FPCGDataFromActorElement::MergeActorsIntoData(FPCGContext* Context, const UPCGDataFromActorSettings* Settings, const TArray<AActor*>& FoundActors) const
+void FPCGDataFromActorElement::CreateReferenceData(FPCGContext* Context, const UPCGDataFromActorSettings* Settings, const TArray<AActor*>& Actors) const
 {
 	check(Context);
-	check(Settings && (Settings->Mode == EPCGGetDataFromActorMode::GetSinglePoint || Settings->Mode == EPCGGetDataFromActorMode::GetActorReference));
+	check(Settings && (Settings->Mode == EPCGGetDataFromActorMode::GetActorReference || Settings->Mode == EPCGGetDataFromActorMode::GetComponentsReference));
 
-	// At this point in time, the partition actors behave slightly differently, so if we are in the case where
-	// we have one or more partition actors, we'll go through the normal process and do post-processing to merge the point data instead.
-	const bool bContainsPartitionActors = Algo::AnyOf(FoundActors, [](const AActor* Actor) { return Cast<APCGPartitionActor>(Actor) != nullptr; });
-	const bool bCreatePointData = (Settings->Mode == EPCGGetDataFromActorMode::GetSinglePoint);
+	const bool bCreateActorReferences = (Settings->Mode == EPCGGetDataFromActorMode::GetActorReference || Settings->Mode == EPCGGetDataFromActorMode::GetComponentsReference);
+	const bool bCreateComponentReferences = (Settings->Mode == EPCGGetDataFromActorMode::GetComponentsReference);
 
-	if (!bContainsPartitionActors)
-	{
-		UPCGData* Data = nullptr;
-		UPCGPointData* PointData = nullptr;
-		UPCGParamData* ParamData = nullptr;
-		FPCGMetadataAttribute<FSoftObjectPath>* ActorReference = nullptr;
+	UPCGParamData* ParamData = FPCGContext::NewObject_AnyThread<UPCGParamData>(Context);
+	FPCGMetadataAttribute<FSoftObjectPath>* ActorReference = nullptr;
+	FPCGMetadataAttribute<FSoftObjectPath>* ComponentReference = nullptr;
 
-		if (bCreatePointData)
-		{
-			PointData = FPCGContext::NewObject_AnyThread<UPCGPointData>(Context);
-			Data = PointData;
-		}
-		else
-		{
-			ParamData = FPCGContext::NewObject_AnyThread<UPCGParamData>(Context);
-			ActorReference = ParamData->MutableMetadata()->FindOrCreateAttribute(PCGPointDataConstants::ActorReferenceAttribute, FSoftObjectPath(), /*bAllowsInterpolation=*/false, /*bOverrideParent=*/false, /*bOverwriteIfTypeMismatch=*/true);
-			Data = ParamData;
-		}
+	if (bCreateActorReferences)
+	{
+		ActorReference = ParamData->MutableMetadata()->FindOrCreateAttribute(PCGPointDataConstants::ActorReferenceAttribute, FSoftObjectPath(), /*bAllowsInterpolation=*/false, /*bOverrideParent=*/false, /*bOverwriteIfTypeMismatch=*/true);
+	}
 
-		bool bHasData = false;
-		bool bAnyAttributeNameWasSanitized = false;
+	if (bCreateComponentReferences)
+	{
+		ComponentReference = ParamData->MutableMetadata()->FindOrCreateAttribute(PCGAddComponentConstants::ComponentReferenceAttribute, FSoftObjectPath(), /*bAllowsInterpolation=*/false, /*bOverrideParent=*/false, /*bOverwriteIfTypeMismatch=*/true);
+	}
 
-		for (AActor* Actor : FoundActors)
+	for (AActor* Actor : Actors)
+	{
+		if (Actor)
 		{
-			if (Actor)
+			if (ComponentReference)
 			{
-				bool bAttributeNameWasSanitized = false;
-				if (PointData)
-				{
-					PointData->AddSinglePointFromActor(Actor, &bAttributeNameWasSanitized);
-				}
-				else
-				{
-					ActorReference->SetValue(ParamData->MutableMetadata()->AddEntry(), FSoftObjectPath(Actor));
-				}
+				TInlineComponentArray<UActorComponent*, 16> Components;
+				Actor->GetComponents(Components);
 
-				bAnyAttributeNameWasSanitized |= bAttributeNameWasSanitized;
-
-				bHasData = true;
-			}
-		}
+				FSoftObjectPath ActorPath(Actor);
 
-		if (bAnyAttributeNameWasSanitized && !Settings->bSilenceSanitizedAttributeNameWarnings)
-		{
-			PCGE_LOG(Warning, GraphAndLog, PCGDataFromActorConstants::TagNamesSanitizedWarning);
-		}
+				for (UActorComponent* Component : Components)
+				{
+					PCGMetadataEntryKey Entry = ParamData->MutableMetadata()->AddEntry();
+					ActorReference->SetValue(Entry, ActorPath);
+					ComponentReference->SetValue(Entry, FSoftObjectPath(Component));
 
-		if (bHasData)
-		{
-			FPCGTaggedData& TaggedData = Context->OutputData.TaggedData.Emplace_GetRef();
-			TaggedData.Data = Data;
+				}
+			}
+			else if (ActorReference)
+			{
+				ActorReference->SetValue(ParamData->MutableMetadata()->AddEntry(), FSoftObjectPath(Actor));
+			}
 		}
 	}
-	else // Stripped-down version of the normal code path with bParseActor = false
-	{
-		FPCGDataCollection DataToMerge;
-		const bool bParseActor = false;
-		bool bAnyAttributeNameWasSanitized = false;
 
-		for (AActor* Actor : FoundActors)
-		{
-			if (Actor)
-			{
-				bool bAttributeNameWasSanitized = false;
-				const EPCGDataType TargetDataType = (bCreatePointData ? EPCGDataType::Point : EPCGDataType::Param);
-				FPCGDataCollection Collection = UPCGComponent::CreateActorPCGDataCollection(Actor, Context->SourceComponent.Get(), TargetDataType, bParseActor, &bAttributeNameWasSanitized);
+	if (ParamData->MutableMetadata()->GetItemCountForChild() > 0)
+	{
+		FPCGTaggedData& TaggedData = Context->OutputData.TaggedData.Emplace_GetRef();
+		TaggedData.Data = ParamData;
+	}
+}
 
-				bAnyAttributeNameWasSanitized |= bAttributeNameWasSanitized;
+void FPCGDataFromActorElement::MergeActorsIntoData(FPCGContext* Context, const UPCGDataFromActorSettings* Settings, const TArray<AActor*>& FoundActors) const
+{
+	check(Context);
+	check(Settings && Settings->Mode == EPCGGetDataFromActorMode::GetSinglePoint);
 
-				DataToMerge.TaggedData += Collection.TaggedData;
-			}
-		}
+	// At this point in time, the partition actors behave slightly differently, so if we are in the case where
+	// we have one or more partition actors, we'll go through the normal process and do post-processing to merge the point data instead.
+	UPCGPointData* PointData = PointData = FPCGContext::NewObject_AnyThread<UPCGPointData>(Context);
+	UPCGData* Data = PointData;
 
-		if (bAnyAttributeNameWasSanitized && !Settings->bSilenceSanitizedAttributeNameWarnings)
-		{
-			PCGE_LOG(Warning, GraphAndLog, PCGDataFromActorConstants::TagNamesSanitizedWarning);
-		}
+	bool bHasData = false;
+	bool bAnyAttributeNameWasSanitized = false;
 
-		// Perform point data-to-point data merge
-		if (DataToMerge.TaggedData.Num() > 1)
+	for (AActor* Actor : FoundActors)
+	{
+		if (Actor)
 		{
-			FPCGMergeElement MergeElement;
-			FPCGMergeAttributesElement MergeAttributesElement;
-			IPCGElement* Element = (bCreatePointData ? static_cast<IPCGElement*>(&MergeElement) : static_cast<IPCGElement*>(&MergeAttributesElement));
+			bool bAttributeNameWasSanitized = false;
+			PointData->AddSinglePointFromActor(Actor, &bAttributeNameWasSanitized);
+			bAnyAttributeNameWasSanitized |= bAttributeNameWasSanitized;
 
-			FPCGContext* MergeContext = Element->Initialize(DataToMerge, Context->SourceComponent, nullptr);
-			MergeContext->AsyncState.NumAvailableTasks = Context->AsyncState.NumAvailableTasks;
+			bHasData = true;
+		}
+	}
 
-			while (!Element->Execute(MergeContext))
-			{
-			}
+	if (bAnyAttributeNameWasSanitized && !Settings->bSilenceSanitizedAttributeNameWarnings)
+	{
+		PCGE_LOG(Warning, GraphAndLog, PCGDataFromActorConstants::TagNamesSanitizedWarning);
+	}
 
-			Context->OutputData = MergeContext->OutputData;
-			delete MergeContext;
-		}
-		else if (DataToMerge.TaggedData.Num() == 1)
-		{
-			Context->OutputData.TaggedData = DataToMerge.TaggedData;
-		}
+	if (bHasData)
+	{
+		FPCGTaggedData& TaggedData = Context->OutputData.TaggedData.Emplace_GetRef();
+		TaggedData.Data = Data;
 	}
 }
 
@@ -830,8 +825,8 @@ void FPCGDataFromActorElement::ProcessActor(FPCGContext* Context, const UPCGData
 	}
 	else
 	{
-		const bool bParseActor = (Settings->Mode != EPCGGetDataFromActorMode::GetSinglePoint && Settings->Mode != EPCGGetDataFromActorMode::GetActorReference);
-		EPCGDataType DataFilter = (Settings->Mode == EPCGGetDataFromActorMode::GetActorReference ? EPCGDataType::Param : Settings->GetDataFilter());
+		const bool bParseActor = (Settings->Mode != EPCGGetDataFromActorMode::GetSinglePoint);
+		EPCGDataType DataFilter = Settings->GetDataFilter();
 		bool bAttributeNameWasSanitized = false;
 		FPCGDataCollection Collection = UPCGComponent::CreateActorPCGDataCollection(FoundActor, SourceComponent, DataFilter, bParseActor, &bAttributeNameWasSanitized);
 
diff --git a/Engine/Plugins/PCG/Source/PCG/Public/Elements/PCGDataFromActor.h b/Engine/Plugins/PCG/Source/PCG/Public/Elements/PCGDataFromActor.h
--- a/Engine/Plugins/PCG/Source/PCG/Public/Elements/PCGDataFromActor.h
+++ b/Engine/Plugins/PCG/Source/PCG/Public/Elements/PCGDataFromActor.h
@@ -19,7 +19,8 @@ enum class EPCGGetDataFromActorMode : uint8
 	GetDataFromProperty UMETA(DisplayName = "Get PCG Data From Property", Tooltip = "Gets a data collection from an actor property."),
 	GetDataFromPCGComponent UMETA(Tooltip = "Copy generated output from other PCG components on the found actor(s)."),
 	GetDataFromPCGComponentOrParseComponents UMETA(Tooltip = "Attempts to copy generated output from other PCG components on the found actor(s), otherwise, falls back to parsing actor components."),
-	GetActorReference UMETA(Tooltip = "Produces one entry per actor with only the actor reference.")
+	GetActorReference UMETA(Tooltip = "Produces one entry per actor with only the actor reference."),
+	GetComponentsReference UMETA(Tooltip = "Produces one entry per component within the actor selection.")
 };
 
 /** Builds a collection of PCG-compatible data from the selected actors. */
@@ -101,7 +102,7 @@ public:
 	int32 AllowedGrids = int32(EPCGHiGenGrid::Uninitialized);
 
 	/** Merges all the single point data outputs into a single point data. */
-	UPROPERTY(BlueprintReadWrite, EditAnywhere, Category = "Data Retrieval Settings", meta = (DisplayName = "Merge Simple Data", EditCondition = "Mode == EPCGGetDataFromActorMode::GetSinglePoint || Mode == EPCGGetDataFromActorMode::GetActorReference", EditConditionHides))
+	UPROPERTY(BlueprintReadWrite, EditAnywhere, Category = "Data Retrieval Settings", meta = (DisplayName = "Merge Simple Data", EditCondition = "Mode == EPCGGetDataFromActorMode::GetSinglePoint || Mode == EPCGGetDataFromActorMode::GetActorReference || Mode == EPCGGetDataFromActorMode::GetComponentsReference", EditConditionHides))
 	bool bMergeSinglePointData = false;
 
 	/** Provide pin names to match against the found component output pins. Data will automatically be wired to the expected pin if the name comparison succeeds. All unmatched pins will go into the standard out pin. */
@@ -153,4 +154,5 @@ protected:
 	virtual void ProcessActor(FPCGContext* Context, const UPCGDataFromActorSettings* Settings, AActor* FoundActor) const;
 
 	void MergeActorsIntoData(FPCGContext* Context, const UPCGDataFromActorSettings* Settings, const TArray<AActor*>& FoundActors) const;
+	void CreateReferenceData(FPCGContext* Context, const UPCGDataFromActorSettings* Settings, const TArray<AActor*>& Actors) const;
 };
