Fixed copy/paste error
#rnx

diff --git a/Engine/Plugins/PCG/Source/PCG/Private/Elements/PCGDataFromActor.cpp b/Engine/Plugins/PCG/Source/PCG/Private/Elements/PCGDataFromActor.cpp
--- a/Engine/Plugins/PCG/Source/PCG/Private/Elements/PCGDataFromActor.cpp
+++ b/Engine/Plugins/PCG/Source/PCG/Private/Elements/PCGDataFromActor.cpp
@@ -674,8 +674,7 @@ void FPCGDataFromActorElement::MergeActorsIntoData(FPCGContext* Context, const U
 
 	// At this point in time, the partition actors behave slightly differently, so if we are in the case where
 	// we have one or more partition actors, we'll go through the normal process and do post-processing to merge the point data instead.
-	UPCGPointData* PointData = PointData = FPCGContext::NewObject_AnyThread<UPCGPointData>(Context);
-	UPCGData* Data = PointData;
+	UPCGPointData* PointData = FPCGContext::NewObject_AnyThread<UPCGPointData>(Context);
 
 	bool bHasData = false;
 	bool bAnyAttributeNameWasSanitized = false;
@@ -700,7 +699,7 @@ void FPCGDataFromActorElement::MergeActorsIntoData(FPCGContext* Context, const U
 	if (bHasData)
 	{
 		FPCGTaggedData& TaggedData = Context->OutputData.TaggedData.Emplace_GetRef();
-		TaggedData.Data = Data;
+		TaggedData.Data = PointData;
 	}
 }
 
