[PCG] Fix crash when Debugging node caused by 34557887

#rnx

diff --git a/Engine/Plugins/PCG/Source/PCG/Private/Graph/PCGGraphExecutor.cpp b/Engine/Plugins/PCG/Source/PCG/Private/Graph/PCGGraphExecutor.cpp
--- a/Engine/Plugins/PCG/Source/PCG/Private/Graph/PCGGraphExecutor.cpp
+++ b/Engine/Plugins/PCG/Source/PCG/Private/Graph/PCGGraphExecutor.cpp
@@ -1074,12 +1074,13 @@ void FPCGGraphExecutor::OnTaskInputsReady(FPCGGraphTask& Task, TArray<FCachedRes
 		}
 
 		PCGGraphExecutor::TScopeLock ChildScopeLock(LiveTasksLock);
-		if (CachedResult && !CachedResult->bIsBypassed)
+		if (CachedResult)
 		{
 			check(CachedResult->TaskId != InvalidPCGTaskId);
 			OutCachedResults.Add(CachedResult);
 		}
-		else
+
+		if(!CachedResult || CachedResult->bIsBypassed)
 		{
 			ReadyTasks.Emplace(MoveTemp(Task));
 		}
@@ -1165,6 +1166,7 @@ void FPCGGraphExecutor::PrepareForExecute(FPCGGraphTask& Task, FCachedResult*& O
 			check(!CachingResults.Contains(Task.NodeId));
 			TUniquePtr<FCachedResult>& CachingResultPtr = CachingResults.Add(Task.NodeId, MakeUnique<FCachedResult>());
 			OutCachedResult = CachingResultPtr.Get();
+			OutCachedResult->TaskId = Task.NodeId;
 			OutCachedResult->Output = MoveTemp(LocalCachedResult.Output);
 		}
 
@@ -1208,9 +1210,6 @@ void FPCGGraphExecutor::PrepareForExecute(FPCGGraphTask& Task, FCachedResult*& O
 		// If the task is a post execute, then we can safely clear the data after getting it from the results.
 		OutCachedResult->bIsPostGraphTask = (Task.Element == GraphCompiler.GetSharedTrivialPostGraphElement());
 
-		// Fast-forward cached result to stored results
-		OutCachedResult->TaskId = Task.NodeId;
-								
 		bNeedToExecuteTasksEnded = true;
 
 		return;
@@ -1234,6 +1233,7 @@ void FPCGGraphExecutor::PrepareForExecute(FPCGGraphTask& Task, FCachedResult*& O
 #if WITH_EDITOR
 	if (bResultAlreadyInCache)
 	{
+		check(OutCachedResult->TaskId != InvalidPCGTaskId);
 		Task.bIsBypassed = true;
 		OutCachedResult->bIsBypassed = true;
 		Task.Context->OutputData = OutCachedResult->Output;
@@ -1797,7 +1797,11 @@ void FPCGGraphExecutor::ExecuteV1()
 			{
 				// Any paused tasks at that point should relinquish their resources
 				TSharedPtr<FPCGGraphActiveTask> ActiveTask = ActiveTasks[ActiveTaskIndex];
+#if WITH_EDITOR
+				check(ActiveTask->bIsExecutingTask || ActiveTask->bIsBypassed);
+#else
 				check(ActiveTask->bIsExecutingTask);
+#endif
 
 				// Any task that asks to be paused or now needs to run on the main thread but doesn't have that slot currently will be moved to the sleeping queue
 				const bool bTaskShouldBePutAside = (ActiveTask->Context->bIsPaused || (ActiveTaskIndex > 0 && ActiveTask->Element->CanExecuteOnlyOnMainThread(ActiveTask->Context.Get())));
