[PCG] Fix Cancellation regression with old execution path

#rnx

diff --git a/Engine/Plugins/PCG/Source/PCG/Private/Graph/PCGGraphExecutor.cpp b/Engine/Plugins/PCG/Source/PCG/Private/Graph/PCGGraphExecutor.cpp
--- a/Engine/Plugins/PCG/Source/PCG/Private/Graph/PCGGraphExecutor.cpp
+++ b/Engine/Plugins/PCG/Source/PCG/Private/Graph/PCGGraphExecutor.cpp
@@ -581,7 +581,12 @@ TSet<UPCGComponent*> FPCGGraphExecutor::Cancel(TFunctionRef<bool(TWeakObjectPtr<
 						ActiveTask->bWasCancelled = true;
 
 						CancelledActiveTasks.Add({ ActiveTask, MoveTemp(TaskHandle) });
-						ActiveTasks.RemoveAtSwap(ActiveTaskIndex);
+						
+						// Avoid removing if using old execution path as it doesn't keep a sharedptr to the executing task (which would have been a way of allowing removal here, but since that path is going away, lets keep it this way)
+						if (CurrentExecuteVersion == EExecuteVersion::V2)
+						{
+							ActiveTasks.RemoveAtSwap(ActiveTaskIndex);
+						}
 					}
 				}
 			}
@@ -597,12 +602,6 @@ TSet<UPCGComponent*> FPCGGraphExecutor::Cancel(TFunctionRef<bool(TWeakObjectPtr<
 				ActiveTask->Element->Abort(ActiveTask->Context.Get());
 												
 				bStableCancellationSet &= !CancelNextTasks(ActiveTask->NodeId, CancelledComponents);
-
-				// In V2 this is done inside the task
-				if (CurrentExecuteVersion == EExecuteVersion::V1)
-				{
-					PostTaskExecute(ActiveTask);
-				}
 			}
 		}
 
