[PCG] Graph Executor V2: Split ActiveTasks into SleepingTasks/ActiveTasks/ActiveTasksGameThreadOnly

- Don't allow shrinking on RemoveAtSwaps
- Maintain a TaskIndex so that Removal is faster

#rnx

diff --git a/Engine/Plugins/PCG/Source/PCG/Private/Graph/PCGGraphExecutor.cpp b/Engine/Plugins/PCG/Source/PCG/Private/Graph/PCGGraphExecutor.cpp
--- a/Engine/Plugins/PCG/Source/PCG/Private/Graph/PCGGraphExecutor.cpp
+++ b/Engine/Plugins/PCG/Source/PCG/Private/Graph/PCGGraphExecutor.cpp
@@ -115,6 +115,27 @@ namespace PCGGraphExecutor
 		}
 	}
 
+	// Needs to be called by owner of LiveTasksLock
+	void AddToActiveTaskArrayNoLock(TArray<TSharedPtr<FPCGGraphActiveTask>>& ActiveTaskArray, TSharedPtr<FPCGGraphActiveTask> ActiveTask)
+	{
+		ActiveTaskArray.Add(ActiveTask);
+		ActiveTask->TaskIndex = ActiveTaskArray.Num() - 1;
+	}
+
+	// Needs to be called by owner of LiveTasksLock
+	void RemoveAtFromActiveTaskArrayNoLock(TArray<TSharedPtr<FPCGGraphActiveTask>>& ActiveTaskArray, int32 Index)
+	{
+		if (ActiveTaskArray.IsValidIndex(Index))
+		{
+			ActiveTaskArray.RemoveAtSwap(Index, EAllowShrinking::No);
+
+			if (ActiveTaskArray.IsValidIndex(Index))
+			{
+				ActiveTaskArray[Index]->TaskIndex = Index;
+			}
+		}
+	}
+		
 	// Similar to existing UE::TScopeLock but allows optional TryLock in which case caller is responsible to 
 	// check if Lock succeeded or not with IsLocked(). This allows execution code to be skipped if non main thread
 	// tasks fail to get the lock instead of spinning for nothing
@@ -569,26 +590,32 @@ TSet<UPCGComponent*> FPCGGraphExecutor::Cancel(TFunctionRef<bool(TWeakObjectPtr<
 		{
 			TArray<TTuple<TSharedPtr<FPCGGraphActiveTask>, UE::Tasks::TTask<bool>>> CancelledActiveTasks;
 			{
-				PCGGraphExecutor::TScopeLock ScopeLock(LiveTasksLock);
-				for (int32 ActiveTaskIndex = ActiveTasks.Num() - 1; ActiveTaskIndex >= 0; --ActiveTaskIndex)
+				auto CancelActiveTasks = [CurrentExecuteVersion, &CancelledComponents, &CancelledActiveTasks](TArray<TSharedPtr<FPCGGraphActiveTask>>& InActiveTasks)
 				{
-					TSharedPtr<FPCGGraphActiveTask>& ActiveTask = ActiveTasks[ActiveTaskIndex];
-					if (ActiveTask->Context && CancelledComponents.Contains(ActiveTask->Context->SourceComponent.Get()))
+					for (int32 ActiveTaskIndex = InActiveTasks.Num() - 1; ActiveTaskIndex >= 0; --ActiveTaskIndex)
 					{
-						// While we have lock Task can't complete, but we can't wait on this task with the lock neither so we capture it here 
-						// and are going to wait on it outside of the LiveTaskLock
-						UE::Tasks::TTask<bool> TaskHandle = ActiveTask->ExecutingTask;
-						ActiveTask->bWasCancelled = true;
-
-						CancelledActiveTasks.Add({ ActiveTask, MoveTemp(TaskHandle) });
-						
-						// Avoid removing if using old execution path as it doesn't keep a sharedptr to the executing task (which would have been a way of allowing removal here, but since that path is going away, lets keep it this way)
-						if (CurrentExecuteVersion == EExecuteVersion::V2)
+						TSharedPtr<FPCGGraphActiveTask>& ActiveTask = InActiveTasks[ActiveTaskIndex];
+						if (ActiveTask->Context && CancelledComponents.Contains(ActiveTask->Context->SourceComponent.Get()))
 						{
-							ActiveTasks.RemoveAtSwap(ActiveTaskIndex);
+							// While we have lock Task can't complete, but we can't wait on this task with the lock neither so we capture it here 
+							// and are going to wait on it outside of the LiveTaskLock
+							UE::Tasks::TTask<bool> TaskHandle = ActiveTask->ExecutingTask;
+							ActiveTask->bWasCancelled = true;
+
+							CancelledActiveTasks.Add({ ActiveTask, MoveTemp(TaskHandle) });
+							
+							// Avoid removing if using old execution path as it doesn't keep a sharedptr to the executing task (which would have been a way of allowing removal here, but since that path is going away, lets keep it this way)
+							if (CurrentExecuteVersion == EExecuteVersion::V2)
+							{
+								InActiveTasks.RemoveAtSwap(ActiveTaskIndex);
+							}
 						}
 					}
-				}
+				};
+				
+				PCGGraphExecutor::TScopeLock ScopeLock(LiveTasksLock);
+				CancelActiveTasks(ActiveTasks);
+				CancelActiveTasks(ActiveTasksGameThreadOnly);
 			}
 
 			for (TTuple<TSharedPtr<FPCGGraphActiveTask>, UE::Tasks::TTask<bool>>& CancelledActiveTask : CancelledActiveTasks)
@@ -702,7 +729,7 @@ bool FPCGGraphExecutor::IsAnyGraphCurrentlyExecuting() const
 
 int32 FPCGGraphExecutor::GetNonScheduledRemainingTaskCount() const
 {
-	return Tasks.Num() + ReadyTasks.Num() + ActiveTasks.Num() + SleepingTasks.Num();
+	return Tasks.Num() + ReadyTasks.Num() + ActiveTasks.Num() + ActiveTasksGameThreadOnly.Num() + SleepingTasks.Num();
 }
 
 FPCGTaskId FPCGGraphExecutor::ScheduleGeneric(TFunction<bool()> InOperation, UPCGComponent* InSourceComponent, const TArray<FPCGTaskId>& TaskExecutionDependencies)
@@ -889,7 +916,7 @@ bool FPCGGraphExecutor::ProcessScheduledTasks()
 		PCGGraphExecutor::TScopeLock ScopeLock(TasksLock);
 		PCGGraphExecutor::TScopeLock ChildScopeLock(LiveTasksLock);
 		// This is a safeguard to check if we're in a stuck state
-		if (CachingResults.Num() == 0 && ReadyTasks.Num() == 0 && ActiveTasks.Num() == 0 && SleepingTasks.Num() == 0 && Tasks.Num() > 0)
+		if (CachingResults.Num() == 0 && ReadyTasks.Num() == 0 && ActiveTasks.Num() == 0 && ActiveTasksGameThreadOnly.Num() == 0 && SleepingTasks.Num() == 0 && Tasks.Num() > 0)
 		{
 			UE_LOG(LogPCG, Error, TEXT("PCG Graph executor error: tasks are in a deadlocked state. Will drop all tasks."));
 			ClearAllTasks();
@@ -940,14 +967,43 @@ void FPCGGraphExecutor::ExecuteTasksEnded()
 	}
 }
 
-void FPCGGraphExecutor::PostTaskExecute(TSharedPtr<FPCGGraphActiveTask> ActiveTaskPtr)
+void FPCGGraphExecutor::PostTaskExecute(TSharedPtr<FPCGGraphActiveTask> ActiveTaskPtr, bool bIsDone)
 {
 	TRACE_CPUPROFILER_EVENT_SCOPE(FPCGGraphExecutor::PostTaskExecute);
 	check(ActiveTaskPtr);
 	FPCGGraphActiveTask& ActiveTask = *ActiveTaskPtr;
 	check(ActiveTask.Context);
 	check(ActiveTask.bIsExecutingTask || ActiveTask.bWasCancelled);
-	
+
+	if (!bIsDone)
+	{
+		TRACE_CPUPROFILER_EVENT_SCOPE(FPCGGraphExecutor::PostTaskExecute::NotDone);
+		PCGGraphExecutor::TScopeLock ScopeLock(LiveTasksLock);
+		ActiveTask.StopExecuting();
+
+		TArray<TSharedPtr<FPCGGraphActiveTask>>& ActiveTaskArrayToRemove = ActiveTask.bIsGameThreadOnly ? ActiveTasksGameThreadOnly : ActiveTasks;
+		const int32 TaskRemoveIndex = ActiveTask.TaskIndex;
+		check(TaskRemoveIndex != INDEX_NONE && ActiveTaskArrayToRemove[TaskRemoveIndex] == ActiveTaskPtr);
+
+		if (ActiveTask.Context->bIsPaused)
+		{
+			PCGGraphExecutor::AddToActiveTaskArrayNoLock(SleepingTasks, ActiveTaskPtr);
+			PCGGraphExecutor::RemoveAtFromActiveTaskArrayNoLock(ActiveTaskArrayToRemove, TaskRemoveIndex);
+		}
+		else if (ActiveTask.bIsGameThreadOnly != ActiveTask.Element->CanExecuteOnlyOnMainThread(ActiveTask.Context.Get()))
+		{
+			ActiveTask.bIsGameThreadOnly = !ActiveTask.bIsGameThreadOnly;
+			TArray<TSharedPtr<FPCGGraphActiveTask>>& ActiveTaskArrayToAdd = ActiveTask.bIsGameThreadOnly ? ActiveTasksGameThreadOnly : ActiveTasks;
+
+			PCGGraphExecutor::AddToActiveTaskArrayNoLock(ActiveTaskArrayToAdd, ActiveTaskPtr);
+			PCGGraphExecutor::RemoveAtFromActiveTaskArrayNoLock(ActiveTaskArrayToRemove, TaskRemoveIndex);
+		}
+
+		// Task might have been moved to SleepingTasks or changed from ActiveTasks to ActiveTasksGameThreadOnly (or vice-versa) 
+		// but it isn't done so we return
+		return;
+	}
+
 	const bool bTaskWasCancelled = ActiveTask.bWasCancelled;
 #if WITH_EDITOR
 	const bool bTaskWasBypassed = ActiveTask.bIsBypassed;
@@ -1045,14 +1101,20 @@ void FPCGGraphExecutor::PostTaskExecute(TSharedPtr<FPCGGraphActiveTask> ActiveTa
 	// Erase from ActiveTasks
 	{
 		PCGGraphExecutor::TScopeLock ScopeLock(LiveTasksLock);
-		if (int32 TaskIndex = ActiveTasks.Find(ActiveTaskPtr); TaskIndex != INDEX_NONE)
+		TArray<TSharedPtr<FPCGGraphActiveTask>>& ActiveTaskArray = ActiveTask.bIsGameThreadOnly ? ActiveTasksGameThreadOnly : ActiveTasks;
+		const int32 TaskIndex = ActiveTask.TaskIndex >= 0 ? ActiveTask.TaskIndex : ActiveTaskArray.Find(ActiveTaskPtr);
+		if (ActiveTaskArray.IsValidIndex(TaskIndex))
 		{
+			check(ActiveTaskArray[TaskIndex] == ActiveTaskPtr);
 			// Remove current active task from list
-			ActiveTasks.RemoveAtSwap(TaskIndex);
+			PCGGraphExecutor::RemoveAtFromActiveTaskArrayNoLock(ActiveTaskArray, TaskIndex);
 		}
 
 		// Make sure we set this to false inside lock
 		ActiveTask.StopExecuting();
+
+		// Next Scheduling call needs to check if this task completion unblocked some sleeping task(s)
+		bNeedToCheckSleepingTasks = true;
 	}
 }
 
@@ -1292,7 +1354,7 @@ FPCGGraphExecutor::EExecuteVersion FPCGGraphExecutor::GetExecuteVersion() const
 	return ExecuteVersion;
 }
 
-bool FPCGGraphExecutor::ExecuteScheduling(double EndTime, TSharedPtr<FPCGGraphActiveTask>* OutMainThreadTask)
+bool FPCGGraphExecutor::ExecuteScheduling(double EndTime, TSharedPtr<FPCGGraphActiveTask>* OutMainThreadTask, bool bForceCheckSleepingTasks)
 {
 	bool bStateChanged = false;
 	const bool bIsInGameThread = IsInGameThread();
@@ -1304,6 +1366,30 @@ bool FPCGGraphExecutor::ExecuteScheduling(double EndTime, TSharedPtr<FPCGGraphAc
 		PCGGraphExecutor::TScopeLock ScopeLock(LiveTasksLock, !bIsInGameThread);
 		if (ScopeLock.IsLocked())
 		{
+			if (bNeedToCheckSleepingTasks || bForceCheckSleepingTasks)
+			{
+				bNeedToCheckSleepingTasks = false;
+
+				// First, wake up any sleeping tasks that can be reactivated
+				for (int32 SleepingTaskIndex = SleepingTasks.Num() - 1; SleepingTaskIndex >= 0; --SleepingTaskIndex)
+				{
+					TSharedPtr<FPCGGraphActiveTask>& SleepingTask = SleepingTasks[SleepingTaskIndex];
+
+					if (SleepingTask->Context->bIsPaused)
+					{
+						continue; // still sleeping
+					}
+
+					const bool bIsGameThreadOnly = SleepingTask->Element->CanExecuteOnlyOnMainThread(SleepingTask->Context.Get());
+					TArray<TSharedPtr<FPCGGraphActiveTask>>& ActiveTaskArray = bIsGameThreadOnly ? ActiveTasksGameThreadOnly : ActiveTasks;
+					SleepingTask->bIsGameThreadOnly = bIsGameThreadOnly;
+
+					PCGGraphExecutor::AddToActiveTaskArrayNoLock(ActiveTaskArray, SleepingTask);
+					PCGGraphExecutor::RemoveAtFromActiveTaskArrayNoLock(SleepingTasks, SleepingTaskIndex);
+					bStateChanged = true;
+				}
+			}
+
 			for (int32 ReadyTaskIndex = ReadyTasks.Num() - 1; ReadyTaskIndex >= 0; --ReadyTaskIndex)
 			{
 				FPCGGraphTask& Task = ReadyTasks[ReadyTaskIndex];
@@ -1324,7 +1410,7 @@ bool FPCGGraphExecutor::ExecuteScheduling(double EndTime, TSharedPtr<FPCGGraphAc
 					{
 						check(CachedResult->TaskId != InvalidPCGTaskId);
 						CachedResults.Add(MoveTemp(CachedResult));
-						ReadyTasks.RemoveAtSwap(ReadyTaskIndex);
+						ReadyTasks.RemoveAtSwap(ReadyTaskIndex, EAllowShrinking::No);
 						bStateChanged = true;
 						continue;
 					}
@@ -1333,7 +1419,12 @@ bool FPCGGraphExecutor::ExecuteScheduling(double EndTime, TSharedPtr<FPCGGraphAc
 				}
 
 				// Validate that we can start this task now
-				FPCGGraphActiveTask& ActiveTask = *ActiveTasks.Emplace_GetRef(MakeShared<FPCGGraphActiveTask>());
+				const bool bIsGameThreadOnly = Task.Element->CanExecuteOnlyOnMainThread(Task.Context);
+				TArray<TSharedPtr<FPCGGraphActiveTask>>& ActiveTaskArray = bIsGameThreadOnly ? ActiveTasksGameThreadOnly : ActiveTasks;
+
+				FPCGGraphActiveTask& ActiveTask = *ActiveTaskArray.Emplace_GetRef(MakeShared<FPCGGraphActiveTask>());
+				ActiveTask.TaskIndex = ActiveTaskArray.Num() - 1;
+				ActiveTask.bIsGameThreadOnly = bIsGameThreadOnly;
 				ActiveTask.Inputs = MoveTemp(Task.Inputs);
 				ActiveTask.Element = Task.Element;
 				ActiveTask.NodeId = Task.NodeId;
@@ -1344,7 +1435,7 @@ bool FPCGGraphExecutor::ExecuteScheduling(double EndTime, TSharedPtr<FPCGGraphAc
 				ActiveTask.bIsBypassed = Task.bIsBypassed;
 #endif
 				// Move the task up front if it needs to run on the main thread
-				ReadyTasks.RemoveAtSwap(ReadyTaskIndex);
+				ReadyTasks.RemoveAtSwap(ReadyTaskIndex, EAllowShrinking::No);
 				bStateChanged = true;
 			}
 		}
@@ -1361,32 +1452,27 @@ bool FPCGGraphExecutor::ExecuteScheduling(double EndTime, TSharedPtr<FPCGGraphAc
 			{
 				TRACE_CPUPROFILER_EVENT_SCOPE(FPCGGraphExecutor::Execute::SelectMainThreadTask);
 				// Execute main thread task
-				if (!ActiveTasks.IsEmpty())
+				if (!ActiveTasksGameThreadOnly.IsEmpty() || !ActiveTasks.IsEmpty())
 				{
-					// Sort tasks so that we might find an main thread executable task at index 0
-					ActiveTasks.Sort([](const TSharedPtr<FPCGGraphActiveTask>& Left, const TSharedPtr<FPCGGraphActiveTask>& Right)
+					auto FindTaskToExecute = [](const TSharedPtr<FPCGGraphActiveTask>& InActiveTask)
 					{
-#if WITH_EDITOR
-						if (Left->bIsBypassed != Right->bIsBypassed) { return Left->bIsBypassed; }
-#endif
-						if (Left->Context->bIsPaused != Right->Context->bIsPaused) { return Right->Context->bIsPaused; }
-
-						if (Left->bIsExecutingTask != Right->bIsExecutingTask) { return Right->bIsExecutingTask; }
-
-						const bool LeftMainThreadOnly = Left->Element->CanExecuteOnlyOnMainThread(Left->Context.Get());
-						const bool RightMainThreadOnly = Right->Element->CanExecuteOnlyOnMainThread(Right->Context.Get());
-
-						if (LeftMainThreadOnly != RightMainThreadOnly) { return LeftMainThreadOnly; }
-
-						return Left->NodeId < Right->NodeId;
-					});
+						return !InActiveTask->bIsExecutingTask && !InActiveTask->Context->bIsPaused;
+					};
 
-					// Sort should give us best target to execute
-					TSharedPtr<FPCGGraphActiveTask>& ActiveTask = ActiveTasks[0];
+					TSharedPtr<FPCGGraphActiveTask>* FoundActiveTask = ActiveTasksGameThreadOnly.FindByPredicate(FindTaskToExecute);
+					if (!FoundActiveTask)
+					{
+						FoundActiveTask = ActiveTasks.FindByPredicate(FindTaskToExecute);
+					}
 
-					// Skip already running task
-					if (!ActiveTask->bIsExecutingTask && !ActiveTask->Context->bIsPaused)
+					if (FoundActiveTask)
 					{
+						// Sort should give us best target to execute
+						TSharedPtr<FPCGGraphActiveTask>& ActiveTask = *FoundActiveTask;
+
+						// Skip already running task
+						check (!ActiveTask->bIsExecutingTask && !ActiveTask->Context->bIsPaused)
+						
 						ActiveTask->StartExecuting();
 						ActiveTask->Context->AsyncState.NumAvailableTasks = -1;
 						ActiveTask->Context->AsyncState.bIsRunningOnMainThread = true;
@@ -1441,15 +1527,7 @@ bool FPCGGraphExecutor::ExecuteScheduling(double EndTime, TSharedPtr<FPCGGraphAc
 							const bool bIsDone = ActiveTask->bWasCancelled || ActiveTask->Element->Execute(ActiveTask->Context.Get());
 							const bool bIsPaused = ActiveTask->Context->bIsPaused;
 
-							if (bIsDone)
-							{
-								PostTaskExecute(ActiveTask);
-							}
-							else
-							{
-								PCGGraphExecutor::TScopeLock ScopeLock(LiveTasksLock);
-								ActiveTask->StopExecuting();
-							}
+							PostTaskExecute(ActiveTask, bIsDone);
 
 							if (bIsDone || bIsPaused)
 							{
@@ -1489,17 +1567,20 @@ bool FPCGGraphExecutor::ExecuteScheduling(double EndTime, TSharedPtr<FPCGGraphAc
  * - [ReadyTasks] are processed through [FPCGGraphExecutor::ExecuteScheduling] which can happen on GameThread/WorkerThread
  * - On WorkerThreads [FPCGGraphExecutor::ExecuteScheduling] will try to acquire the Lock and if it can't will just skip through, on GameThread it will wait until it can grab the lock. (LiveTasksLock) (same reason as above)
  *	 
- * - [ReadyTasks] will be moved to [ActiveTasks] through [FPCGGraphExecutor::ExecuteScheduling]
+ * - [ReadyTasks] will be moved to [ActiveTasks/ActiveTasksGameThreadOnly] through [FPCGGraphExecutor::ExecuteScheduling]
  * - This will only happen if Task.bHasDonePrepareForExecute is true or if we can run [FPCGGraphExecutor::PrepareForExecute] on it. (this can always happen on GameThread but not always on WorkerThreads as explained before)
  * 
- * - [ActiveTasks] are also processed in [FPCGGraphExecutor::ExecuteScheduling]
- * - On the GameThread the [ActiveTasks] will get sorted so that tasks that can execute on the GameThread only will be first in the array
- * - After the sort, GameThread will always try to execute [ActiveTasks[0]], it will do this [FPCGGraphExecutor::ExecuteScheduling] while it exhausts its GameThread budget
+ * - [ActiveTasks/ActiveTasksGameThreadOnly] are also processed in [FPCGGraphExecutor::ExecuteScheduling]
+ * - The GameThread will always try to execute the first [ActiveTasksGameThreadOnly] or [ActiveTasks] if no tasks need to be executed on GameThread, it will do this [FPCGGraphExecutor::ExecuteScheduling] while it exhausts its GameThread budget
  * - Both GameThread/WorkerThread will also spawn new UE::Tasks for the remaining [ActiveTasks] that are allowed to run outside of the GameThread still happens in [FPCGGraphExecutor::ExecuteScheduling]
  * 
- * - When an ActiveTasks Execute call returns and it returns true, it means that the task is done. In this case we call [FPCGGraphExecutor::StoreResults] followed by [FPCGGraphExecutor::QueueNextTasks] this will potentially move some [Tasks] to [ReadyTasks]
+ * - When an [ActiveTasks/ActiveTasksGameThreadOnly] Execute call returns and it returns true, it means that the task is done. In this case we call [FPCGGraphExecutor::StoreResults] followed by [FPCGGraphExecutor::QueueNextTasks] this will potentially move some [Tasks] to [ReadyTasks]
  *   so that new tasks can get executed now that their inputs have been filled
  * 
+ * - After a call to Execute, if a [ActiveTasks/ActiveTasksGameThreadOnly] bPaused flag is true, it gets moved to [SleepingTasks]
+ * - [SleepingTasks] gets processed at least once a frame to see if some tasks are no longer paused and need to be moved back to [ActiveTasks/ActiveTasksGameThreadOnly]
+ * - It will also be processed again if some other [ActiveTasks/ActiveTasksGameThreadOnly] completes execution
+ * 
  * - Cache information: When we call [FPCGGraphExecutor::PrepareForExecute] is called on a task, if a Cache result is found for this task then we can directly call [FPCGGraphExecutor::StoreResults] and [FPCGGraphExecutor::QueueNextTasks] on it skipping execution [FPCGGraphExecutor::ProcessCachedResults]
  * - An exception to this is when the FPCGElement is marked for debugging then we move the task through the normal execution pipeline but set a flag on it so that the actual Execute call can be skipped
  * - [CachingResults] exists only to keep GC references to PCG Data that has been read from the Cache but hasn't yet been moved to the [TaskOutputs]
@@ -1527,30 +1608,22 @@ void FPCGGraphExecutor::ExecuteV2()
 
 	const double StartTime = FPlatformTime::Seconds();
 	const double EndTime = StartTime + GetTickBudgetInSeconds();
-		
+
 	while ((bFirstLoop || (FPlatformTime::Seconds() < EndTime)) && bContinueExecute)
 	{
-		bFirstLoop = false;
-		
 		TSharedPtr<FPCGGraphActiveTask> MainThreadTask;
-		bContinueExecute = ExecuteScheduling(EndTime, &MainThreadTask);
-								
+		bContinueExecute = ExecuteScheduling(EndTime, &MainThreadTask, /*bForceCheckSleepingTasks=*/bFirstLoop);
+		bFirstLoop = false;
+
 		if (MainThreadTask)
 		{
 			TRACE_CPUPROFILER_EVENT_SCOPE(FPCGGraphExecutor::Execute::ExecuteMainThreadTask);
 #if WITH_EDITOR
-			if (MainThreadTask->bIsBypassed || MainThreadTask->bWasCancelled || MainThreadTask->Element->Execute(MainThreadTask->Context.Get()))
+			const bool bIsDone = (MainThreadTask->bIsBypassed || MainThreadTask->bWasCancelled || MainThreadTask->Element->Execute(MainThreadTask->Context.Get()));
 #else
-			if (MainThreadTask->bWasCancelled || MainThreadTask->Element->Execute(MainThreadTask->Context.Get()))
+			const bool bIsDone = (MainThreadTask->bWasCancelled || MainThreadTask->Element->Execute(MainThreadTask->Context.Get()));
 #endif
-			{
-				PostTaskExecute(MainThreadTask);
-			}
-			else
-			{
-				PCGGraphExecutor::TScopeLock ScopeLock(LiveTasksLock);
-				MainThreadTask->StopExecuting();
-			}
+			PostTaskExecute(MainThreadTask, bIsDone);
 		}
 	}
 
@@ -1637,17 +1710,19 @@ void FPCGGraphExecutor::ExecuteV1()
 
 				if (!bIsMainThreadTask || bMainThreadAvailable)
 				{
-					ActiveTasks.Emplace(MoveTemp(SleepingTask));
+					PCGGraphExecutor::AddToActiveTaskArrayNoLock(ActiveTasks, SleepingTask);
 
 					// Move task up front if it has to run on the main thread
 					if (bIsMainThreadTask && ActiveTasks.Num() > 1)
 					{
 						ActiveTasks.Swap(0, ActiveTasks.Num() - 1);
+						ActiveTasks[0]->TaskIndex = 0;
+						ActiveTasks[ActiveTasks.Num() - 1]->TaskIndex = ActiveTasks.Num() - 1;
 					}
 
 					TaskDispatchBookKeeping(bIsMainThreadTask);
 
-					SleepingTasks.RemoveAtSwap(SleepingTaskIndex);
+					PCGGraphExecutor::RemoveAtFromActiveTaskArrayNoLock(SleepingTasks, SleepingTaskIndex);
 				}
 			}
 
@@ -1667,6 +1742,7 @@ void FPCGGraphExecutor::ExecuteV1()
 				if (!bIsMainThreadTask || bMainThreadAvailable)
 				{
 					FPCGGraphActiveTask& ActiveTask = *ActiveTasks.Emplace_GetRef(MakeShared<FPCGGraphActiveTask>());
+					ActiveTask.TaskIndex = ActiveTasks.Num() - 1;
 					ActiveTask.Inputs = MoveTemp(Task.Inputs);
 					ActiveTask.Element = Task.Element;
 					ActiveTask.NodeId = Task.NodeId;
@@ -1681,6 +1757,8 @@ void FPCGGraphExecutor::ExecuteV1()
 					if (bIsMainThreadTask && ActiveTasks.Num() > 1)
 					{
 						ActiveTasks.Swap(0, ActiveTasks.Num() - 1);
+						ActiveTasks[0]->TaskIndex = 0;
+						ActiveTasks[ActiveTasks.Num() - 1]->TaskIndex = ActiveTasks.Num() - 1;
 					}
 
 					TaskDispatchBookKeeping(bIsMainThreadTask);
@@ -1731,6 +1809,7 @@ void FPCGGraphExecutor::ExecuteV1()
 					ActiveTask.StartExecuting();
 					ActiveTask.ExecutingTask = UE::Tasks::Launch(UE_SOURCE_LOCATION, [&ActiveTask]()
 					{
+						TRACE_CPUPROFILER_EVENT_SCOPE(FPCGGraphExecutor::ExecuteAsyncTask);
 						return ActiveTask.Element->Execute(ActiveTask.Context.Get());
 					});
 				}
@@ -1777,14 +1856,14 @@ void FPCGGraphExecutor::ExecuteV1()
 
 					if (bTaskDone || ActiveTasks[ExecutionIndex]->bWasCancelled)
 					{
-						PostTaskExecute(ActiveTasks[ExecutionIndex]);
+						PostTaskExecute(ActiveTasks[ExecutionIndex], true);
 					}
 				}
 			}
 
 			if (bMainTaskDone)
 			{
-				PostTaskExecute(ActiveTasks[0]);
+				PostTaskExecute(ActiveTasks[0], true);
 			}
 		}
 
@@ -1812,8 +1891,8 @@ void FPCGGraphExecutor::ExecuteV1()
 					}
 					ActiveTask->Context->AsyncState.NumAvailableTasks = 0;
 
-					SleepingTasks.Emplace(ActiveTask);
-					ActiveTasks.RemoveAtSwap(ActiveTaskIndex);
+					PCGGraphExecutor::AddToActiveTaskArrayNoLock(SleepingTasks, ActiveTask);
+					PCGGraphExecutor::RemoveAtFromActiveTaskArrayNoLock(ActiveTasks, ActiveTaskIndex);
 				}
 
 				ActiveTask->StopExecuting();
@@ -1846,6 +1925,7 @@ void FPCGGraphExecutor::ClearAllTasks()
 
 	ReadyTasks.Reset();
 	ActiveTasks.Reset();
+	ActiveTasksGameThreadOnly.Reset();
 	SleepingTasks.Reset();
 }
 
@@ -1929,6 +2009,9 @@ void FPCGGraphExecutor::ProcessCachedResults(const TArray<FCachedResult*>& Cache
 			check(CachingResults.Contains(CachedResult->TaskId));
 			CachingResults.Remove(CachedResult->TaskId);
 		}
+
+		// Next Scheduling call needs to check if this task completion unblocked some sleeping task(s)
+		bNeedToCheckSleepingTasks = true;
 	}
 }
 
@@ -2338,7 +2421,7 @@ void FPCGGraphExecutor::CullInactiveDownstreamNodes(FPCGTaskId InCompletedTaskId
 							{
 								if (SuccessorTask->Inputs[InputIndex].TaskId == RemovedTaskId)
 								{
-									SuccessorTask->Inputs.RemoveAtSwap(InputIndex);
+									SuccessorTask->Inputs.RemoveAtSwap(InputIndex, EAllowShrinking::No);
 								}
 							}
 						}
@@ -2412,30 +2495,26 @@ void FPCGGraphExecutor::AddReferencedObjects(FReferenceCollector& Collector)
 			}
 		});
 
-		Algo::ForEach(ActiveTasks, [&Collector](auto& Task)
+		auto AddReferencesActiveTask = [&Collector](TSharedPtr<FPCGGraphActiveTask>& InActiveTask)
 		{
-			if (Task->bIsExecutingTask)
+			if (InActiveTask->bIsExecutingTask)
 			{
-				Collector.AddReferencedObjects(Task->ExecutingReferences);
+				Collector.AddReferencedObjects(InActiveTask->ExecutingReferences);
 				// @todo_pcg this is to allow referencing extra objects from the Context sub classes, should probably be part of the visiting pattern also
-				if (Task->Context)
+				if (InActiveTask->Context)
 				{
-					Task->Context->AddExtraStructReferencedObjects(Collector);
+					InActiveTask->Context->AddExtraStructReferencedObjects(Collector);
 				}
 			}
-			else if (Task->Context)
+			else if (InActiveTask->Context)
 			{
-				Task->Context->AddStructReferencedObjects(Collector);
+				InActiveTask->Context->AddStructReferencedObjects(Collector);
 			}
-		});
+		};
 
-		Algo::ForEach(SleepingTasks, [&Collector](auto& Task)
-		{
-			if (Task->Context)
-			{
-				Task->Context->AddStructReferencedObjects(Collector);
-			}
-		});
+		Algo::ForEach(ActiveTasks, AddReferencesActiveTask);
+		Algo::ForEach(ActiveTasksGameThreadOnly, AddReferencesActiveTask);
+		Algo::ForEach(SleepingTasks, AddReferencesActiveTask);
 
 		Algo::ForEach(CachingResults, [&Collector](auto& CachingResult)
 		{
diff --git a/Engine/Plugins/PCG/Source/PCG/Private/Graph/PCGGraphExecutor.h b/Engine/Plugins/PCG/Source/PCG/Private/Graph/PCGGraphExecutor.h
--- a/Engine/Plugins/PCG/Source/PCG/Private/Graph/PCGGraphExecutor.h
+++ b/Engine/Plugins/PCG/Source/PCG/Private/Graph/PCGGraphExecutor.h
@@ -149,10 +149,14 @@ struct FPCGGraphActiveTask : TSharedFromThis<FPCGGraphActiveTask>
 	// Those members need to be modified under the FPCGGraphExecutor::LiveTasksLock (unless we are running the old executor path)
 	UE::Tasks::TTask<bool> ExecutingTask;
 	bool bIsExecutingTask = false;
+
+	// Used to know if task should be in ActiveTasks or ActiveTasksGameThreadOnly
+	bool bIsGameThreadOnly = false;
+	// TaskIndex inside ActiveTasks/ActiveTasksGameThreadOnly/SleepingTasks for fast removal
+	int32 TaskIndex = INDEX_NONE;
+
 	static int32 NumExecuting;
 	TArray<TObjectPtr<const UObject>> ExecutingReferences;
-
-	TSharedPtr<UE::Tasks::FCancellationToken> CancelToken;
 };
 
 class FPCGGraphExecutor : public FGCObject
@@ -267,10 +271,10 @@ private:
 		bool bIsBypassed = false;
 	};
 		
-	void PostTaskExecute(TSharedPtr<FPCGGraphActiveTask> ActiveTask);
+	void PostTaskExecute(TSharedPtr<FPCGGraphActiveTask> ActiveTask, bool bIsDone);
 	bool ProcessScheduledTasks();
 	void ExecuteTasksEnded();
-	bool ExecuteScheduling(double EndTime, TSharedPtr<FPCGGraphActiveTask>* OutMainThreadTask = nullptr);
+	bool ExecuteScheduling(double EndTime, TSharedPtr<FPCGGraphActiveTask>* OutMainThreadTask = nullptr, bool bForceCheckSleepingTasks = false);
 
 	TSet<UPCGComponent*> Cancel(TFunctionRef<bool(TWeakObjectPtr<UPCGComponent>)> CancelFilter);
 	void ClearAllTasks();
@@ -346,13 +350,14 @@ private:
 	UE::FSpinLock LiveTasksLock;
 	TArray<FPCGGraphTask> ReadyTasks;
 	TArray<TSharedPtr<FPCGGraphActiveTask>> ActiveTasks;
+	TArray<TSharedPtr<FPCGGraphActiveTask>> ActiveTasksGameThreadOnly;
+	TArray<TSharedPtr<FPCGGraphActiveTask>> SleepingTasks;
+	bool bNeedToCheckSleepingTasks = false;
+
 	// Used to keep GC references to in flight caching results (not yet stored to output and might not be in cache anymore)
 	TMap<FPCGTaskId, TUniquePtr<FCachedResult>> CachingResults;
 	int32 NumWorkerTasks = 0;
 
-	// @todo_pcg: to remove when we remove ExecuteV1
-	TArray<TSharedPtr<FPCGGraphActiveTask>> SleepingTasks;
-	
 	/** Lock level 3 */
 	UE::FSpinLock CollectGCReferenceTasksLock;
 	TSet<TSharedPtr<FPCGGraphActiveTask>> CollectGCReferenceTasks;
