[PCG] GPU execution support v1

* First pass on PCG GPU. Uses ComputeFramework plugin to dispatch GPU work.
* Compiler finds subsets of ?GPU-enabled? nodes and collapses them to a single ComputeGraphElement. Upon execution this will build and compile the compute shaders and wait for GPU execution to terminate.
* Supported PCG types on the GPU (so far) include: Points, Landscape, Texture
  * PointData can be multi-data
  * Landscape & Texture must be single-data
* GPU nodes can exist alongside CPU nodes, edges that cross the CPU <-> GPU boundary are automatically bundled into Upload/Readback data interfaces respectively.
* Introduce 'Custom HLSL' element to easily author general nodes. Kernel types include:
  * Point Processor - Assumes a 1:1 relationship between the primary input pin & primary output pin. Automatically uploads point properties & metadata attributes, and copies them to the output data.
  * Point Generator - Generates a single PointData on the primary output pin.
  * Custom - General execution, can define buffer sizes and relationships from pin properties. User is responsible for the boiler plate involved in copying attributes/properties to output data.
* Attribute read/write supported via a static lookup table computed at graph compilation time (with static analysis of the compute kernels)
* Various helper functions provided for ease of authoring kernels

CVars
* pcg.GraphExecution.GPU.Enable - Enable GPU execution
* pcg.GraphExecution.GPU.EnableDebugging - Enable verbose GPU logging

Notes
* All kernel input pins forced to Single Connection, temporarily
* Inspection & Debug - need to be readback to CPU, temporary solution is to connect to a CPU gather node for now
* GPU execution is editor only. Graph compilation outside of editor build culls nodes that require the GPU, which will disconnect them from the graph.

diff --git a/Engine/Plugins/PCG/PCG.uplugin b/Engine/Plugins/PCG/PCG.uplugin
--- a/Engine/Plugins/PCG/PCG.uplugin
+++ b/Engine/Plugins/PCG/PCG.uplugin
@@ -35,6 +35,10 @@
 		{
 			"Name": "EditorScriptingUtilities",
 			"Enabled": true
+		},
+		{
+			"Name": "ComputeFramework",
+			"Enabled": true
 		}
 	]
 }
diff --git a/Engine/Plugins/PCG/Shaders/Private/PCGDataCollectionDataInterface.ush b/Engine/Plugins/PCG/Shaders/Private/PCGDataCollectionDataInterface.ush
new file mode 100644
--- /dev/null
+++ b/Engine/Plugins/PCG/Shaders/Private/PCGDataCollectionDataInterface.ush
@@ -0,0 +1,489 @@
+// Copyright Epic Games, Inc. All Rights Reserved.
+
+#include "/Engine/Private/Quaternion.ush"
+
+RWByteAddressBuffer {DataInterfaceName}_DataCollectionBuffer;
+
+// TODO: These defines could be code-gen'd to stay in sync with our PCGComputeConstants
+#define POINT_POSITION_ID 0
+#define POINT_ROTATION_ID 1
+#define POINT_SCALE_ID 2
+#define POINT_BOUNDS_MIN_ID 3
+#define POINT_BOUNDS_MAX_ID 4
+#define POINT_COLOR_ID 5
+#define POINT_DENSITY_ID 6
+#define POINT_SEED_ID 7
+#define POINT_STEEPNESS_ID 8
+#define POINT_METADATA_ENTRY_ID 9
+
+// #################### HEADER READERS ##########################
+
+uint GetNumData_{DataInterfaceName}()
+{
+	return {DataInterfaceName}_DataCollectionBuffer.Load(0);
+}
+
+uint ReadDataAddress_{DataInterfaceName}(uint DataIndex)
+{
+	return {DataInterfaceName}_DataCollectionBuffer.Load((DataIndex + 1) * 4);
+}
+
+uint ReadDataId_{DataInterfaceName}(uint DataAddress)
+{
+	return {DataInterfaceName}_DataCollectionBuffer.Load(DataAddress + 0);
+}
+
+uint ReadDataNumAttributes_{DataInterfaceName}(uint DataAddress)
+{
+	return {DataInterfaceName}_DataCollectionBuffer.Load(DataAddress + 4);
+}
+
+uint ReadDataPreambleSize_{DataInterfaceName}(uint DataAddress)
+{
+	return {DataInterfaceName}_DataCollectionBuffer.Load(DataAddress + 8);
+}
+
+// For PointData, this is just NumPoints
+uint ReadDataInfo_{DataInterfaceName}(uint DataAddress)
+{
+	return {DataInterfaceName}_DataCollectionBuffer.Load(DataAddress + 12);
+}
+
+// Returns (DataIndex, DataAddress, ElementIndex)
+uint3 GetThreadData_{DataInterfaceName}(uint ThreadIndex)
+{
+	int ElementIndex = ThreadIndex;
+
+	const uint NumData = GetNumData_{DataInterfaceName}();
+	for (uint DataIndex = 0; DataIndex < NumData; ++DataIndex)
+	{
+		const uint DataAddress = ReadDataAddress_{DataInterfaceName}(DataIndex);
+		const uint ElemCount = ReadDataInfo_{DataInterfaceName}(DataAddress);
+
+		if (ElementIndex < ElemCount)
+		{
+			return uint3(DataIndex, DataAddress, ElementIndex);
+		}
+
+		ElementIndex -= ElemCount;
+	}
+
+	return uint3(-1, -1, -1);
+}
+
+uint GetNumElements_{DataInterfaceName}()
+{
+	uint NumElements = 0;
+
+	const uint NumData = GetNumData_{DataInterfaceName}();
+	for (uint DataIndex = 0; DataIndex < NumData; ++DataIndex)
+	{
+		const uint DataAddress = ReadDataAddress_{DataInterfaceName}(DataIndex);
+		NumElements += ReadDataInfo_{DataInterfaceName}(DataAddress);
+	}
+
+	return NumElements;
+}
+
+uint ReadDataAttributeHeadersAddress_{DataInterfaceName}(uint DataAddress)
+{
+	return DataAddress + ReadDataPreambleSize_{DataInterfaceName}(DataAddress);
+}
+
+uint ReadAttributeHeaderAddress_{DataInterfaceName}(uint AttributeHeadersAddress, uint AttributeId)
+{
+	return AttributeHeadersAddress + AttributeId * 8;
+}
+
+uint ReadAttributeIdAndStride_{DataInterfaceName}(uint AttributeHeaderAddress)
+{
+	return {DataInterfaceName}_DataCollectionBuffer.Load(AttributeHeaderAddress);
+}
+
+uint ReadAttributeId_{DataInterfaceName}(uint AttributeHeaderAddress)
+{
+	return ReadAttributeIdAndStride_{DataInterfaceName}(AttributeHeaderAddress) >> 8;
+}
+
+uint ReadAttributeStride_{DataInterfaceName}(uint AttributeHeaderAddress)
+{
+	return ReadAttributeIdAndStride_{DataInterfaceName}(AttributeHeaderAddress) & (0xFFFFFFFF >> 24);
+}
+
+uint ReadAttributeAddress_{DataInterfaceName}(uint AttributeHeaderAddress)
+{
+	return {DataInterfaceName}_DataCollectionBuffer.Load(AttributeHeaderAddress + 4);
+}
+
+uint ReadAttributeAddress_{DataInterfaceName}(uint InDataAddress, uint InAttributeId)
+{
+	const uint AttributeHeadersAddress = ReadDataAttributeHeadersAddress_{DataInterfaceName}(InDataAddress);
+	const uint NumAttributes = ReadDataNumAttributes_{DataInterfaceName}(InDataAddress);
+
+	for (uint AttributeIndex = 0; AttributeIndex < NumAttributes; ++AttributeIndex)
+	{
+		const uint AttributeHeaderAddress = ReadAttributeHeaderAddress_{DataInterfaceName}(AttributeHeadersAddress, AttributeIndex);
+
+		if (ReadAttributeId_{DataInterfaceName}(AttributeHeaderAddress) == InAttributeId)
+		{
+			return ReadAttributeAddress_{DataInterfaceName}(AttributeHeaderAddress);
+		}
+	}
+
+	return -1;
+}
+
+// #################### HEADER WRITERS ##########################
+
+void WriteNumData_{DataInterfaceName}(uint NumData)
+{
+	{DataInterfaceName}_DataCollectionBuffer.Store(0, NumData);
+}
+
+void WriteDataAddress_{DataInterfaceName}(uint DataIndex, uint DataAddress)
+{
+	{DataInterfaceName}_DataCollectionBuffer.Store((DataIndex + 1) * 4, DataAddress);
+}
+
+void WriteDataId_{DataInterfaceName}(uint DataAddress, uint Id)
+{
+	{DataInterfaceName}_DataCollectionBuffer.Store(DataAddress + 0, Id);
+}
+
+void WriteDataNumAttributes_{DataInterfaceName}(uint DataAddress, uint NumAttributes)
+{
+	{DataInterfaceName}_DataCollectionBuffer.Store(DataAddress + 4, NumAttributes);
+}
+
+void WriteDataPreambleSize_{DataInterfaceName}(uint DataAddress, uint PreambleSize)
+{
+	{DataInterfaceName}_DataCollectionBuffer.Store(DataAddress + 8, PreambleSize);
+}
+
+// For PointData, this is just NumPoints
+void WriteDataInfo_{DataInterfaceName}(uint DataAddress, uint Info)
+{
+	{DataInterfaceName}_DataCollectionBuffer.Store(DataAddress + 12, Info);
+}
+
+void WriteAttributeIdAndStride_{DataInterfaceName}(uint AttributeHeaderAddress, uint IdAndStride)
+{
+	{DataInterfaceName}_DataCollectionBuffer.Store(AttributeHeaderAddress, IdAndStride);
+}
+
+void WriteAttributeAddress_{DataInterfaceName}(uint AttributeHeaderAddress, uint AttributeAddress)
+{
+	{DataInterfaceName}_DataCollectionBuffer.Store(AttributeHeaderAddress + 4, AttributeAddress);
+}
+
+uint LoadBuffer_{DataInterfaceName}(uint Address)
+{
+	return {DataInterfaceName}_DataCollectionBuffer.Load(Address);
+}
+
+void StoreBuffer_{DataInterfaceName}(uint Address, uint Value)
+{
+	return {DataInterfaceName}_DataCollectionBuffer.Store(Address, Value);
+}
+
+// #################### ATTRIBUTE GETTERS ##########################
+
+uint ComputeAddressHelper_{DataInterfaceName}(uint DataIndex, uint AttributeId, uint ElementIndex)
+{
+	const uint DataAddress = ReadDataAddress_{DataInterfaceName}(DataIndex);
+	const uint AttributeHeadersAddress = ReadDataAttributeHeadersAddress_{DataInterfaceName}(DataAddress);
+	const uint AttributeHeaderAddress = ReadAttributeHeaderAddress_{DataInterfaceName}(AttributeHeadersAddress, AttributeId);
+	return ReadAttributeAddress_{DataInterfaceName}(AttributeHeaderAddress);
+}
+
+bool GetBool_{DataInterfaceName}(uint DataIndex, uint AttributeId, uint ElementIndex)
+{
+	const uint AttributeAddress = ComputeAddressHelper_{DataInterfaceName}(DataIndex, AttributeId, ElementIndex);
+	return asbool({DataInterfaceName}_DataCollectionBuffer.Load(AttributeAddress + ElementIndex * 1 * 4));
+}
+
+int GetInt_{DataInterfaceName}(uint DataIndex, uint AttributeId, uint ElementIndex)
+{
+	const uint AttributeAddress = ComputeAddressHelper_{DataInterfaceName}(DataIndex, AttributeId, ElementIndex);
+	return asint({DataInterfaceName}_DataCollectionBuffer.Load(AttributeAddress + ElementIndex * 1 * 4));
+}
+
+float GetFloat_{DataInterfaceName}(uint DataIndex, uint AttributeId, uint ElementIndex)
+{
+	const uint AttributeAddress = ComputeAddressHelper_{DataInterfaceName}(DataIndex, AttributeId, ElementIndex);
+	return asfloat({DataInterfaceName}_DataCollectionBuffer.Load(AttributeAddress + ElementIndex * 1 * 4));
+}
+
+float2 GetFloat2_{DataInterfaceName}(uint DataIndex, uint AttributeId, uint ElementIndex)
+{
+	const uint AttributeAddress = ComputeAddressHelper_{DataInterfaceName}(DataIndex, AttributeId, ElementIndex);
+	return asfloat({DataInterfaceName}_DataCollectionBuffer.Load2(AttributeAddress + ElementIndex * 2 * 4));
+}
+
+float3 GetFloat3_{DataInterfaceName}(uint DataIndex, uint AttributeId, uint ElementIndex)
+{
+	const uint AttributeAddress = ComputeAddressHelper_{DataInterfaceName}(DataIndex, AttributeId, ElementIndex);
+	return asfloat({DataInterfaceName}_DataCollectionBuffer.Load3(AttributeAddress + ElementIndex * 3 * 4));
+}
+
+float4 GetFloat4_{DataInterfaceName}(uint DataIndex, uint AttributeId, uint ElementIndex)
+{
+	const uint AttributeAddress = ComputeAddressHelper_{DataInterfaceName}(DataIndex, AttributeId, ElementIndex);
+	return asfloat({DataInterfaceName}_DataCollectionBuffer.Load4(AttributeAddress + ElementIndex * 4 * 4));
+}
+
+float3 GetRotator_{DataInterfaceName}(uint DataIndex, uint AttributeId, uint ElementIndex)
+{
+	const uint AttributeAddress = ComputeAddressHelper_{DataInterfaceName}(DataIndex, AttributeId, ElementIndex);
+	return asfloat({DataInterfaceName}_DataCollectionBuffer.Load3(AttributeAddress + ElementIndex * 3 * 4));
+}
+
+float4 GetQuat_{DataInterfaceName}(uint DataIndex, uint AttributeId, uint ElementIndex)
+{
+	const uint AttributeAddress = ComputeAddressHelper_{DataInterfaceName}(DataIndex, AttributeId, ElementIndex);
+	return asfloat({DataInterfaceName}_DataCollectionBuffer.Load4(AttributeAddress + ElementIndex * 4 * 4));
+}
+
+// #################### ATTRIBUTE SETTERS ##########################
+
+void SetBool_{DataInterfaceName}(uint DataIndex, uint AttributeId, uint ElementIndex, bool Value)
+{
+	const uint AttributeAddress = ComputeAddressHelper_{DataInterfaceName}(DataIndex, AttributeId, ElementIndex);
+	{DataInterfaceName}_DataCollectionBuffer.Store(AttributeAddress + ElementIndex * 1 * 4, asuint(Value));
+}
+
+void SetInt_{DataInterfaceName}(uint DataIndex, uint AttributeId, uint ElementIndex, int Value)
+{
+	const uint AttributeAddress = ComputeAddressHelper_{DataInterfaceName}(DataIndex, AttributeId, ElementIndex);
+	{DataInterfaceName}_DataCollectionBuffer.Store(AttributeAddress + ElementIndex * 1 * 4, asuint(Value));
+}
+
+void SetFloat_{DataInterfaceName}(uint DataIndex, uint AttributeId, uint ElementIndex, float Value)
+{
+	const uint AttributeAddress = ComputeAddressHelper_{DataInterfaceName}(DataIndex, AttributeId, ElementIndex);
+	{DataInterfaceName}_DataCollectionBuffer.Store(AttributeAddress + ElementIndex * 1 * 4, asuint(Value));
+}
+
+void SetFloat2_{DataInterfaceName}(uint DataIndex, uint AttributeId, uint ElementIndex, float2 Value)
+{
+	const uint AttributeAddress = ComputeAddressHelper_{DataInterfaceName}(DataIndex, AttributeId, ElementIndex);
+	{DataInterfaceName}_DataCollectionBuffer.Store2(AttributeAddress + ElementIndex * 2 * 4, asuint(Value));
+}
+
+void SetFloat3_{DataInterfaceName}(uint DataIndex, uint AttributeId, uint ElementIndex, float3 Value)
+{
+	const uint AttributeAddress = ComputeAddressHelper_{DataInterfaceName}(DataIndex, AttributeId, ElementIndex);
+	{DataInterfaceName}_DataCollectionBuffer.Store3(AttributeAddress + ElementIndex * 3 * 4, asuint(Value));
+}
+
+void SetFloat4_{DataInterfaceName}(uint DataIndex, uint AttributeId, uint ElementIndex, float4 Value)
+{
+	const uint AttributeAddress = ComputeAddressHelper_{DataInterfaceName}(DataIndex, AttributeId, ElementIndex);
+	{DataInterfaceName}_DataCollectionBuffer.Store4(AttributeAddress + ElementIndex * 4 * 4, asuint(Value));
+}
+
+void SetRotator_{DataInterfaceName}(uint DataIndex, uint AttributeId, uint ElementIndex, float3 Value)
+{
+	const uint AttributeAddress = ComputeAddressHelper_{DataInterfaceName}(DataIndex, AttributeId, ElementIndex);
+	{DataInterfaceName}_DataCollectionBuffer.Store3(AttributeAddress + ElementIndex * 3 * 4, asuint(Value));
+}
+
+void SetQuat_{DataInterfaceName}(uint DataIndex, uint AttributeId, uint ElementIndex, float4 Value)
+{
+	const uint AttributeAddress = ComputeAddressHelper_{DataInterfaceName}(DataIndex, AttributeId, ElementIndex);
+	{DataInterfaceName}_DataCollectionBuffer.Store4(AttributeAddress + ElementIndex * 4 * 4, asuint(Value));
+}
+
+// #################### POINT ATTRIBUTE GETTERS ##########################
+
+float3 GetPosition_{DataInterfaceName}(uint DataIndex, uint ElementIndex)
+{
+	return GetFloat3_{DataInterfaceName}(DataIndex, POINT_POSITION_ID, ElementIndex);
+}
+
+float4 GetRotation_{DataInterfaceName}(uint DataIndex, uint ElementIndex)
+{
+	return GetFloat4_{DataInterfaceName}(DataIndex, POINT_ROTATION_ID, ElementIndex);
+}
+
+float3 GetScale_{DataInterfaceName}(uint DataIndex, uint ElementIndex)
+{
+	return GetFloat3_{DataInterfaceName}(DataIndex, POINT_SCALE_ID, ElementIndex);
+}
+
+float3 GetBoundsMin_{DataInterfaceName}(uint DataIndex, uint ElementIndex)
+{
+	return GetFloat3_{DataInterfaceName}(DataIndex, POINT_BOUNDS_MIN_ID, ElementIndex);
+}
+
+float3 GetBoundsMax_{DataInterfaceName}(uint DataIndex, uint ElementIndex)
+{
+	return GetFloat3_{DataInterfaceName}(DataIndex, POINT_BOUNDS_MAX_ID, ElementIndex);
+}
+
+float4 GetColor_{DataInterfaceName}(uint DataIndex, uint ElementIndex)
+{
+	return GetFloat4_{DataInterfaceName}(DataIndex, POINT_COLOR_ID, ElementIndex);
+}
+
+float GetDensity_{DataInterfaceName}(uint DataIndex, uint ElementIndex)
+{
+	return GetFloat_{DataInterfaceName}(DataIndex, POINT_DENSITY_ID, ElementIndex);
+}
+
+float GetSteepness_{DataInterfaceName}(uint DataIndex, uint ElementIndex)
+{
+	return GetFloat_{DataInterfaceName}(DataIndex, POINT_STEEPNESS_ID, ElementIndex);
+}
+
+int GetSeed_{DataInterfaceName}(uint DataIndex, uint ElementIndex)
+{
+	return GetInt_{DataInterfaceName}(DataIndex, POINT_SEED_ID, ElementIndex);
+}
+
+float4x4 GetPointTransform_{DataInterfaceName}(uint InDataIndex, uint InElementIndex)
+{
+	const FQuat Rotation = GetRotation_{DataInterfaceName}(InDataIndex, InElementIndex);
+	const half3x3 RotationMatrix = QuatToMatrix(Rotation);
+	const float3 Position = GetPosition_{DataInterfaceName}(InDataIndex, InElementIndex);
+	const float3 Scale = GetScale_{DataInterfaceName}(InDataIndex, InElementIndex);
+	
+	const float3 Axis0 = Scale.x * RotationMatrix[0];
+	const float3 Axis1 = Scale.y * RotationMatrix[1];
+	const float3 Axis2 = Scale.z * RotationMatrix[2];
+
+	return float4x4(
+		Axis0.x, Axis0.y, Axis0.z, Position.x,
+		Axis1.x, Axis1.y, Axis1.z, Position.y,
+		Axis2.x, Axis2.y, Axis2.z, Position.z,
+		(float3)0.0, 1.0);
+}
+
+// #################### POINT ATTRIBUTE SETTERS ##########################
+
+void SetPosition_{DataInterfaceName}(uint DataIndex, uint ElementIndex, float3 Position)
+{
+	SetFloat3_{DataInterfaceName}(DataIndex, POINT_POSITION_ID, ElementIndex, Position);
+}
+
+void SetRotation_{DataInterfaceName}(uint DataIndex, uint ElementIndex, float4 Rotation)
+{
+	SetFloat4_{DataInterfaceName}(DataIndex, POINT_ROTATION_ID, ElementIndex, Rotation);
+}
+
+void SetScale_{DataInterfaceName}(uint DataIndex, uint ElementIndex, float3 Scale)
+{
+	SetFloat3_{DataInterfaceName}(DataIndex, POINT_SCALE_ID, ElementIndex, Scale);
+}
+
+void SetBoundsMin_{DataInterfaceName}(uint DataIndex, uint ElementIndex, float3 BoundsMin)
+{
+	SetFloat3_{DataInterfaceName}(DataIndex, POINT_BOUNDS_MIN_ID, ElementIndex, BoundsMin);
+}
+
+void SetBoundsMax_{DataInterfaceName}(uint DataIndex, uint ElementIndex, float3 BoundsMax)
+{
+	SetFloat3_{DataInterfaceName}(DataIndex, POINT_BOUNDS_MAX_ID, ElementIndex, BoundsMax);
+}
+
+void SetColor_{DataInterfaceName}(uint DataIndex, uint ElementIndex, float4 Color)
+{
+	SetFloat4_{DataInterfaceName}(DataIndex, POINT_COLOR_ID, ElementIndex, Color);
+}
+
+void SetDensity_{DataInterfaceName}(uint DataIndex, uint ElementIndex, float Density)
+{
+	SetFloat_{DataInterfaceName}(DataIndex, POINT_DENSITY_ID, ElementIndex, Density);
+}
+
+void SetSteepness_{DataInterfaceName}(uint DataIndex, uint ElementIndex, float Steepness)
+{
+	SetFloat_{DataInterfaceName}(DataIndex, POINT_STEEPNESS_ID, ElementIndex, Steepness);
+}
+
+void SetSeed_{DataInterfaceName}(uint DataIndex, uint ElementIndex, int Seed)
+{
+	SetInt_{DataInterfaceName}(DataIndex, POINT_SEED_ID, ElementIndex, Seed);
+}
+
+void SetSeedFromPosition_{DataInterfaceName}(uint DataIndex, uint ElementIndex, float3 Position)
+{
+	const int3 PosInt = (int3)Position;
+	const int Seed = ((PosInt.x * 196314165U) + 907633515U) ^ ((PosInt.y * 73148459U) + 453816763U) ^ ((PosInt.z * 34731343U) + 453816743U);
+
+	SetInt_{DataInterfaceName}(DataIndex, POINT_SEED_ID, ElementIndex, Seed);
+}
+
+void SetPointTransform_{DataInterfaceName}(uint InDataIndex, uint InElementIndex, float4x4 Transform)
+{
+	const float3 Scale = float3(length(Transform._m00_m01_m02), length(Transform._m10_m11_m12), length(Transform._m20_m21_m22));
+	SetScale_{DataInterfaceName}(InDataIndex, InElementIndex, Scale);
+
+	// Assumes Transform axes have the correct handedness. We could do more work to fix this up too if required.
+	float3x3 RotationMatix = Transform;
+	// Required unfortunately, QuatFromMatrix expects orthonormalized.
+	RotationMatix[0] /= Scale.x;
+	RotationMatix[1] /= Scale.y;
+	RotationMatix[2] /= Scale.z;
+	const FQuat Rotation = QuatFromMatrix(RotationMatix);
+	SetRotation_{DataInterfaceName}(InDataIndex, InElementIndex, Rotation);
+
+	const float3 Position = Transform._m03_m13_m23;
+	SetPosition_{DataInterfaceName}(InDataIndex, InElementIndex, Position);
+}
+
+// Write single point data collection header
+void WriteSinglePointDataCollectionHeader_{DataInterfaceName}(uint InPointCount)
+{
+    const uint NumData = 1;
+    WriteNumData_{DataInterfaceName}(NumData);
+    
+	// TODO code gen the various constants here and below rather than hardcoding them.
+    const static uint AttributeStrides[9] = { 3, 4, 3, 3, 3, 4, 1, 1, 1 };
+    
+    for (uint DataIndex = 0; DataIndex < NumData; ++DataIndex)
+    {
+        const uint DataAddress = 8;
+        const uint Id = 0; // PCG_POINT_DATA_TYPE_ID
+        const uint NumAttributes = 9; // Number of point attributes
+        const uint PreambleSizeBytes = 16; // PCG_POINT_DATA_HEADER_PREAMBLE_SIZE_BYTES
+        const uint NumElements = InPointCount;
+        
+        WriteDataAddress_{DataInterfaceName}(DataIndex, DataAddress);
+        WriteDataId_{DataInterfaceName}(DataAddress, Id);
+        WriteDataNumAttributes_{DataInterfaceName}(DataAddress, NumAttributes);
+        WriteDataPreambleSize_{DataInterfaceName}(DataAddress, PreambleSizeBytes);
+        WriteDataInfo_{DataInterfaceName}(DataAddress, NumElements);
+        
+        const uint AttributeHeadersAddress = DataAddress + PreambleSizeBytes;
+        // TODO: These magic numbers should be code gen'd using the constants in PCGComputeConstants
+        uint CurrentAttributeAddress = AttributeHeadersAddress + 128 * 8; // MAX_NUM_ATTRS * ATTR_HEADER_SIZE_BYTES
+        
+        for (uint AttributeIndex = 0; AttributeIndex < NumAttributes; ++AttributeIndex)
+        {
+            const uint AttributeHeaderAddress = AttributeHeadersAddress + 8 * AttributeIndex;
+            const uint AttributeId = AttributeIndex; // Point attributes are [0, 8)
+            const uint AttributeStrideBytes = AttributeStrides[AttributeIndex] * 4;
+            const uint AttributeIdAndStride = (AttributeId << 8) + AttributeStrideBytes;
+            
+            WriteAttributeIdAndStride_{DataInterfaceName}(AttributeHeaderAddress, AttributeIdAndStride);
+            WriteAttributeAddress_{DataInterfaceName}(AttributeHeaderAddress, CurrentAttributeAddress);
+            
+            CurrentAttributeAddress += NumElements * AttributeStrideBytes;
+        }
+    }
+}
+
+// Initialize a single point with default values.
+void InitializePoint_{DataInterfaceName}(uint InDataIndex, uint InElementIndex)
+{
+	SetPosition_{DataInterfaceName}(InDataIndex, InElementIndex, 1.0f);
+	SetRotation_{DataInterfaceName}(InDataIndex, InElementIndex, float4(0.0f, 0.0f, 0.0f, 1.0f));
+	SetScale_{DataInterfaceName}(InDataIndex, InElementIndex, 1.0f);
+	SetBoundsMin_{DataInterfaceName}(InDataIndex, InElementIndex, -50.0f);
+	SetBoundsMax_{DataInterfaceName}(InDataIndex, InElementIndex, 50.0f);
+	SetColor_{DataInterfaceName}(InDataIndex, InElementIndex, 1.0f);
+	SetDensity_{DataInterfaceName}(InDataIndex, InElementIndex, 1.0f);
+	SetSeed_{DataInterfaceName}(InDataIndex, InElementIndex, 42);
+	SetSteepness_{DataInterfaceName}(InDataIndex, InElementIndex, 0.5f);
+}
diff --git a/Engine/Plugins/PCG/Shaders/Private/PCGLandscapeDataInterface.ush b/Engine/Plugins/PCG/Shaders/Private/PCGLandscapeDataInterface.ush
new file mode 100644
--- /dev/null
+++ b/Engine/Plugins/PCG/Shaders/Private/PCGLandscapeDataInterface.ush
@@ -0,0 +1,47 @@
+// Copyright Epic Games, Inc. All Rights Reserved.
+
+Texture2D {DataInterfaceName}_HeightTexture;
+SamplerState {DataInterfaceName}_HeightTextureSampler;
+float3 {DataInterfaceName}_HeightTextureLWCTile;
+float4x4 {DataInterfaceName}_HeightTextureWorldToUvTransform;
+float4 {DataInterfaceName}_HeightTextureUvScaleBias;
+float2 {DataInterfaceName}_HeightTextureWorldGridSize;
+float3 {DataInterfaceName}_SystemLWCTile;
+
+float2 GetHeightTextureUv_{DataInterfaceName}(float3 InWorldPos)
+{
+	const FLWCVector4 LwcWorldPos = MakeLWCVector4({DataInterfaceName}_SystemLWCTile, float4(InWorldPos, 1.0f));
+	const FLWCInverseMatrix LwcWorldToUv = MakeLWCInverseMatrix({DataInterfaceName}_HeightTextureLWCTile, {DataInterfaceName}_HeightTextureWorldToUvTransform);
+
+	float2 Uv = LWCMultiply(LwcWorldPos, LwcWorldToUv).xy;
+	Uv = Uv * {DataInterfaceName}_HeightTextureUvScaleBias.xy + {DataInterfaceName}_HeightTextureUvScaleBias.zw;
+
+	return Uv;
+}
+
+float GetHeight_{DataInterfaceName}(float3 InWorldPos)
+{
+	const float2 Uv = GetHeightTextureUv_{DataInterfaceName}(InWorldPos);
+
+	return {DataInterfaceName}_HeightTexture.SampleLevel({DataInterfaceName}_HeightTextureSampler, Uv, 0).x;
+}
+
+float3 GetNormal_{DataInterfaceName}(float3 InWorldPos)
+{
+	float3 Pos1 = InWorldPos;
+	float3 Pos2 = InWorldPos + float3(0, {DataInterfaceName}_HeightTextureWorldGridSize.y, 0);
+	float3 Pos3 = InWorldPos + float3({DataInterfaceName}_HeightTextureWorldGridSize.x, 0, 0);
+
+	const float2 Uv1 = GetHeightTextureUv_{DataInterfaceName}(Pos1);
+	const float2 Uv2 = GetHeightTextureUv_{DataInterfaceName}(Pos2);
+	const float2 Uv3 = GetHeightTextureUv_{DataInterfaceName}(Pos3);
+
+	Pos1.z = {DataInterfaceName}_HeightTexture.SampleLevel({DataInterfaceName}_HeightTextureSampler, Uv1, 0).x;
+	Pos2.z = {DataInterfaceName}_HeightTexture.SampleLevel({DataInterfaceName}_HeightTextureSampler, Uv2, 0).x;
+	Pos3.z = {DataInterfaceName}_HeightTexture.SampleLevel({DataInterfaceName}_HeightTextureSampler, Uv3, 0).x;
+
+	const float3 DeltaX = Pos2 - Pos1;
+	const float3 DeltaY = Pos3 - Pos1;
+
+	return normalize(cross(DeltaY, DeltaX));
+}
diff --git a/Engine/Plugins/PCG/Shaders/Private/PCGShaderUtils.ush b/Engine/Plugins/PCG/Shaders/Private/PCGShaderUtils.ush
new file mode 100644
--- /dev/null
+++ b/Engine/Plugins/PCG/Shaders/Private/PCGShaderUtils.ush
@@ -0,0 +1,98 @@
+// Copyright Epic Games, Inc. All Rights Reserved.
+
+#include "/Engine/Private/Common.ush"
+#include "/Engine/Public/Platform.ush"
+
+float3 CreateGrid2D(int ElementIndex, int NumPoints, int NumRows, float3 Min, float3 Max)
+{
+    if (NumRows == 0)
+	{
+		return Min;
+	}
+
+    const int NumCols = ceil((0.f + NumPoints) / NumRows);
+
+    if (NumCols == 0)
+	{
+		return Min;
+	}
+
+    const int IndexX = ElementIndex / NumRows;
+    const int IndexY = ElementIndex - NumRows * IndexX;
+
+    float3 Spacing = (Max - Min);
+    if (NumCols > 1)
+	{
+		Spacing.x /= (NumCols - 1);
+	}
+    if (NumRows > 1)
+	{
+		Spacing.y /= (NumRows - 1);
+	}
+
+    float3 Position = Min;
+    Position.x += IndexX * Spacing.x;
+    Position.y += IndexY * Spacing.y;
+
+    return Position;
+}
+
+float3 CreateGrid2D(int ElementIndex, int NumPoints, float3 Min, float3 Max)
+{
+	if (NumPoints <= 0)
+	{
+		return Min;
+	}
+
+    const int NumRows = ceil(sqrt(NumPoints));
+    return CreateGrid2D(ElementIndex, NumPoints, NumRows, Min, Max);
+}
+
+float3 CreateGrid3D(int ElementIndex, int NumPoints, int NumRows, int NumCols, float3 Min, float3 Max)
+{
+    if (NumRows == 0 || NumCols == 0)
+	{
+		return Min;
+	}
+
+    const int NumLayers = ceil((0.f + NumPoints) / (NumRows * NumCols));
+
+    if (NumLayers == 0)
+	{
+		return Min;
+	}
+
+    int3 Index;
+    Index.x = ElementIndex % NumCols;
+    Index.y = (ElementIndex / NumCols) % NumRows;
+    Index.z = ElementIndex / (NumRows * NumCols);
+
+    float3 Spacing = Max - Min;
+    if (NumCols > 1)
+	{
+		Spacing.x /= (NumCols - 1);
+	}
+    if (NumRows > 1)
+	{
+		Spacing.y /= (NumRows - 1);
+	}
+    if (NumLayers > 1)
+	{
+		Spacing.z /= (NumLayers - 1);
+	}
+
+    return Min + Index * Spacing;
+}
+
+float3 CreateGrid3D(int ElementIndex, int NumPoints, float3 Min, float3 Max)
+{
+	if (NumPoints <= 0)
+	{
+		return Min;
+	}
+
+    const int NumRows = round(pow(NumPoints, 1.f / 3.f));
+    const int NumCols = NumRows;
+
+    return CreateGrid3D(ElementIndex, NumPoints, NumRows, NumCols, Min, Max);
+}
diff --git a/Engine/Plugins/PCG/Shaders/Private/PCGTextureDataInterface.ush b/Engine/Plugins/PCG/Shaders/Private/PCGTextureDataInterface.ush
new file mode 100644
--- /dev/null
+++ b/Engine/Plugins/PCG/Shaders/Private/PCGTextureDataInterface.ush
@@ -0,0 +1,23 @@
+// Copyright Epic Games, Inc. All Rights Reserved.
+
+Texture2D {DataInterfaceName}_Texture;
+SamplerState {DataInterfaceName}_Sampler;
+
+float2 GetTexCoords_{DataInterfaceName}(float2 WorldPos, float2 Min, float2 Max)
+{
+	// Inverse lerp
+	const float2 Pos = WorldPos - Min;
+	const float2 Extents = Max - Min;
+
+	if (!all(Extents))
+	{
+		return (float2)0;
+	}
+
+	return Pos / Extents;
+}
+
+float4 Sample_{DataInterfaceName}(float2 TexCoords)
+{
+	return {DataInterfaceName}_Texture.SampleLevel({DataInterfaceName}_Sampler, TexCoords, 0);
+}
diff --git a/Engine/Plugins/PCG/Source/PCG/PCG.Build.cs b/Engine/Plugins/PCG/Source/PCG/PCG.Build.cs
--- a/Engine/Plugins/PCG/Source/PCG/PCG.Build.cs
+++ b/Engine/Plugins/PCG/Source/PCG/PCG.Build.cs
@@ -9,7 +9,8 @@ namespace UnrealBuildTool.Rules
 			PCHUsage = PCHUsageMode.UseExplicitOrSharedPCHs;
 
 			PublicDependencyModuleNames.AddRange(
-				new string[] {
+				new string[]
+				{
 					"Chaos",
 					"Core",
 					"CoreUObject",
@@ -26,10 +27,10 @@ namespace UnrealBuildTool.Rules
 			PrivateDependencyModuleNames.AddRange(
 				new string[]
 				{
+					"ComputeFramework",
 					"PCGCompute",
 					"Voronoi",
-				}
-			);
+				});
 
 			if (Target.bBuildEditor == true)
 			{
@@ -39,8 +40,7 @@ namespace UnrealBuildTool.Rules
 						"UnrealEd",
 						"Settings",
 						"SourceControl",
-					}
-				);
+					});
 			}
 
 			bDisableAutoRTFMInstrumentation = true;
diff --git a/Engine/Plugins/PCG/Source/PCG/Private/Compute/DataInterfaces/PCGComputeDataInterface.h b/Engine/Plugins/PCG/Source/PCG/Private/Compute/DataInterfaces/PCGComputeDataInterface.h
new file mode 100644
--- /dev/null
+++ b/Engine/Plugins/PCG/Source/PCG/Private/Compute/DataInterfaces/PCGComputeDataInterface.h
@@ -0,0 +1,35 @@
+// Copyright Epic Games, Inc. All Rights Reserved.
+
+#pragma once
+
+#include "ComputeFramework/ComputeDataInterface.h"
+
+#include "PCGComputeDataInterface.generated.h"
+
+UCLASS(Abstract, ClassGroup = (Procedural))
+class UPCGComputeDataInterface : public UComputeDataInterface
+{
+	GENERATED_BODY()
+
+public:
+	/** Register a downstream pin (and optionally compute graph element pin alias). */
+	virtual void AddDownstreamInputPin(FName InInputPinLabel, const FName* InOptionalInputPinLabelAlias = nullptr)
+	{
+		DownstreamInputPinLabelAliases.AddUnique(InOptionalInputPinLabelAlias ? *InOptionalInputPinLabelAlias : InInputPinLabel);
+	}
+
+	/** Set the output pin label and label alias that this data interface is associated with. */
+	virtual void SetOutputPin(FName InOutputPinLabel, const FName* InOptionalOutputPinLabelAlias = nullptr)
+	{
+		OutputPinLabel = InOutputPinLabel;
+		OutputPinLabelAlias = InOptionalOutputPinLabelAlias ? *InOptionalOutputPinLabelAlias : InOutputPinLabel;
+	}
+
+protected:
+	/** Label of output pin that this data interface is associated with. */
+	FName OutputPinLabel;
+	FName OutputPinLabelAlias;
+
+	/** Generated PCG data will be assigned these labels. */
+	TArray<FName> DownstreamInputPinLabelAliases;
+};
diff --git a/Engine/Plugins/PCG/Source/PCG/Private/Compute/DataInterfaces/PCGCustomKernelDataInterface.cpp b/Engine/Plugins/PCG/Source/PCG/Private/Compute/DataInterfaces/PCGCustomKernelDataInterface.cpp
new file mode 100644
--- /dev/null
+++ b/Engine/Plugins/PCG/Source/PCG/Private/Compute/DataInterfaces/PCGCustomKernelDataInterface.cpp
@@ -0,0 +1,183 @@
+// Copyright Epic Games, Inc. All Rights Reserved.
+
+#include "Compute/DataInterfaces/PCGCustomKernelDataInterface.h"
+
+#include "PCGComponent.h"
+#include "PCGModule.h"
+#include "PCGSettings.h"
+
+#include "RenderGraphBuilder.h"
+#include "RenderGraphResources.h"
+#include "ShaderCompilerCore.h"
+#include "ShaderParameterMetadataBuilder.h"
+#include "ComputeFramework/ComputeMetadataBuilder.h"
+#include "ComputeFramework/ShaderParamTypeDefinition.h"
+
+#include UE_INLINE_GENERATED_CPP_BY_NAME(PCGCustomKernelDataInterface)
+
+const TCHAR* UPCGCustomKernelDataInterface::NumThreadsReservedName = TEXT("NumThreads");
+
+void UPCGCustomKernelDataInterface::GetSupportedInputs(TArray<FShaderFunctionDefinition>& OutFunctions) const
+{
+	OutFunctions.AddDefaulted_GetRef()
+		.SetName(FString::Printf(TEXT("Get%s"), NumThreadsReservedName))
+		.AddReturnType(FShaderValueType::Get(EShaderFundamentalType::Int));
+
+	// A convenient way to serve component bounds to all kernels. Could be pulled out into a PCG context DI in the future.
+	OutFunctions.AddDefaulted_GetRef()
+		.SetName(TEXT("GetComponentBoundsMin"))
+		.AddReturnType(EShaderFundamentalType::Float, 3);
+
+	OutFunctions.AddDefaulted_GetRef()
+		.SetName(TEXT("GetComponentBoundsMax"))
+		.AddReturnType(EShaderFundamentalType::Float, 3);
+}
+
+BEGIN_SHADER_PARAMETER_STRUCT(FPCGKernelDataInterfaceParameters, )
+	SHADER_PARAMETER(FIntVector3, NumThreads)
+	SHADER_PARAMETER(FVector3f, ComponentBoundsMin)
+	SHADER_PARAMETER(FVector3f, ComponentBoundsMax)
+END_SHADER_PARAMETER_STRUCT()
+
+void UPCGCustomKernelDataInterface::GetShaderParameters(TCHAR const* UID, FShaderParametersMetadataBuilder& InOutBuilder, FShaderParametersMetadataAllocations& InOutAllocations) const
+{
+	InOutBuilder.AddNestedStruct<FPCGKernelDataInterfaceParameters>(UID);
+}
+
+void UPCGCustomKernelDataInterface::GetShaderHash(FString& InOutKey) const
+{
+	// UComputeGraph::BuildKernelSource hashes the result of GetHLSL()
+	// Only append additional hashes here if the HLSL contains any additional includes	
+}
+
+void UPCGCustomKernelDataInterface::GetHLSL(FString& OutHLSL, FString const& InDataInterfaceName) const
+{
+	const FString TypeName = FShaderValueType::Get(EShaderFundamentalType::Int, 3)->ToString();
+
+	if (ensure(!TypeName.IsEmpty()))
+	{
+		TMap<FString, FStringFormatArg> TemplateArgs =
+		{
+			{ TEXT("DataInterfaceName"), InDataInterfaceName },
+		};
+
+		// Add uniforms.
+		OutHLSL += FString::Printf(TEXT("%s %s_%s;\n"), 
+			*TypeName, 
+			*InDataInterfaceName, 
+			NumThreadsReservedName);
+			
+		// Add function getters.
+		OutHLSL += FString::Printf(TEXT("%s Get%s_%s()\n{\n\treturn %s_%s;\n}\n"), 
+			*TypeName,
+			NumThreadsReservedName,
+			*InDataInterfaceName, 
+			*InDataInterfaceName, 
+			NumThreadsReservedName);
+
+		// Add helpers (e.g. GetComponentBounds())
+		OutHLSL += FString::Format(TEXT(
+			"float3 {DataInterfaceName}_ComponentBoundsMin;\n"
+			"float3 {DataInterfaceName}_ComponentBoundsMax;\n"
+			"\n"
+			"float3 GetComponentBoundsMin_{DataInterfaceName}()\n"
+			"{\n"
+			"	return {DataInterfaceName}_ComponentBoundsMin;\n"
+			"}\n"
+			"\n"
+			"float3 GetComponentBoundsMax_{DataInterfaceName}()\n"
+			"{\n"
+			"	return {DataInterfaceName}_ComponentBoundsMax;\n"
+			"}\n"), TemplateArgs);
+	}
+}
+
+UComputeDataProvider* UPCGCustomKernelDataInterface::CreateDataProvider(TObjectPtr<UObject> InBinding, uint64 InInputMask, uint64 InOutputMask) const
+{
+	check(Settings);
+
+	UPCGDataBinding* Binding = CastChecked<UPCGDataBinding>(InBinding);
+	check(Binding->SourceComponent.IsValid() && Binding->SourceComponent.Get());
+
+	const FBox ComponentBounds = Binding->SourceComponent.Get()->GetGridBounds();
+
+	UPCGCustomComputeKernelDataProvider* Provider = NewObject<UPCGCustomComputeKernelDataProvider>();
+	Provider->SourceComponentBounds = ComponentBounds;
+	Provider->ThreadCount = Settings->ComputeKernelThreadCount(Binding);
+
+	return Provider;
+}
+
+FComputeDataProviderRenderProxy* UPCGCustomComputeKernelDataProvider::GetRenderProxy()
+{
+	TArray<int32> InvocationCounts;
+	int32 TotalThreadCount = 0;
+	
+	if (!GetInvocationThreadCounts(InvocationCounts, TotalThreadCount))
+	{
+		InvocationCounts.Reset();
+	}
+
+	return new FPCGCustomComputeKernelDataProviderProxy(MoveTemp(InvocationCounts), TotalThreadCount, SourceComponentBounds);
+}
+
+bool UPCGCustomComputeKernelDataProvider::GetInvocationThreadCounts(TArray<int32>& OutInvocationThreadCount, int32& OutTotalThreadCount) const
+{
+	OutInvocationThreadCount.Reset(1);
+	OutInvocationThreadCount.Add(ThreadCount);
+
+	OutTotalThreadCount = ThreadCount;
+	
+	return true;
+}
+
+bool FPCGCustomComputeKernelDataProviderProxy::IsValid(FValidationData const& InValidationData) const
+{
+	if (InValidationData.ParameterStructSize != sizeof(FParameters))
+	{
+		UE_LOG(LogPCG, Error, TEXT("Proxy invalid due to mismatching ParameterStructSize."));
+		return false;
+	}
+
+	if (InvocationThreadCounts.Num() == 0)
+	{
+		UE_LOG(LogPCG, Error, TEXT("Proxy invalid due to empty InvocationThreadCounts."));
+		return false;
+	}
+
+	return true;
+}
+
+int32 FPCGCustomComputeKernelDataProviderProxy::GetDispatchThreadCount(TArray<FIntVector>& InOutThreadCounts) const
+{
+	InOutThreadCounts.Reset(InvocationThreadCounts.Num());
+	for (const int32 Count : InvocationThreadCounts)
+	{
+		InOutThreadCounts.Add({Count, 1, 1});
+	}
+	return InOutThreadCounts.Num();
+}
+
+void FPCGCustomComputeKernelDataProviderProxy::GatherDispatchData(FDispatchData const& InDispatchData)
+{
+	const TStridedView<FParameters> ParameterArray = MakeStridedParameterView<FParameters>(InDispatchData);
+	for (int32 InvocationIndex = 0; InvocationIndex < InDispatchData.NumInvocations; ++InvocationIndex)
+	{
+		int32 NumThreads;
+		if (InDispatchData.bUnifiedDispatch)
+		{
+			NumThreads = TotalThreadCount;
+		}
+		else
+		{
+			NumThreads = InvocationThreadCounts[InvocationIndex];
+		}
+
+		FParameters& Parameters = ParameterArray[InvocationIndex];
+		Parameters.NumThreads = FIntVector(NumThreads, 1, 1);
+
+		// Set component bounds
+		Parameters.ComponentBoundsMin = (FVector3f)SourceComponentBounds.Min;
+		Parameters.ComponentBoundsMax = (FVector3f)SourceComponentBounds.Max;
+	}
+}
diff --git a/Engine/Plugins/PCG/Source/PCG/Private/Compute/DataInterfaces/PCGCustomKernelDataInterface.h b/Engine/Plugins/PCG/Source/PCG/Private/Compute/DataInterfaces/PCGCustomKernelDataInterface.h
new file mode 100644
--- /dev/null
+++ b/Engine/Plugins/PCG/Source/PCG/Private/Compute/DataInterfaces/PCGCustomKernelDataInterface.h
@@ -0,0 +1,79 @@
+// Copyright Epic Games, Inc. All Rights Reserved.
+
+#pragma once
+
+#include "Compute/PCGDataBinding.h"
+
+#include "ComputeFramework/ComputeDataInterface.h"
+#include "ComputeFramework/ComputeDataProvider.h"
+
+#include "PCGCustomKernelDataInterface.generated.h"
+
+class FPCGKernelDataInterfaceParameters;
+
+/** Interface for any meta data provided to the compute kernel, such as num threads. */
+UCLASS(ClassGroup = (Procedural))
+class PCG_API UPCGCustomKernelDataInterface : public UComputeDataInterface
+{
+	GENERATED_BODY()
+
+public:
+	static const TCHAR* NumThreadsReservedName;
+	
+	//~ Begin UComputeDataInterface Interface
+	TCHAR const* GetClassName() const override { return TEXT("CustomComputeKernelData"); }
+	bool CanSupportUnifiedDispatch() const override { return true; }
+	void GetSupportedInputs(TArray<FShaderFunctionDefinition>& OutFunctions) const override;
+	void GetShaderParameters(TCHAR const* UID, FShaderParametersMetadataBuilder& InOutBuilder, FShaderParametersMetadataAllocations& InOutAllocations) const override;
+	void GetShaderHash(FString& InOutKey) const override;
+	void GetHLSL(FString& OutHLSL, FString const& InDataInterfaceName) const override;
+	UComputeDataProvider* CreateDataProvider(TObjectPtr<UObject> InBinding, uint64 InInputMask, uint64 InOutputMask) const override;
+
+	// This DI will provide execution parameters like dispatch information.
+	bool IsExecutionInterface() const override { return true; }
+	//~ End UComputeDataInterface Interface
+
+	UPROPERTY()
+	const UPCGSettings* Settings;
+};
+
+/** Compute Framework Data Provider for each custom compute kernel. */
+UCLASS()
+class UPCGCustomComputeKernelDataProvider : public UComputeDataProvider
+{
+	GENERATED_BODY()
+
+public:
+	//~ Begin UComputeDataProvider Interface
+	FComputeDataProviderRenderProxy* GetRenderProxy() override;
+	//~ End UComputeDataProvider Interface
+
+	int ThreadCount = -1;
+	FBox SourceComponentBounds;
+
+protected:
+	bool GetInvocationThreadCounts(TArray<int32>& OutInvocationThreadCount, int32& OutTotalThreadCount) const;
+};
+
+class FPCGCustomComputeKernelDataProviderProxy : public FComputeDataProviderRenderProxy
+{
+public:
+	FPCGCustomComputeKernelDataProviderProxy(TArray<int32>&& InInvocationThreadCounts, int32 InTotalThreadCount, const FBox& InSourceComponentBounds)
+		: InvocationThreadCounts(MoveTemp(InInvocationThreadCounts))
+		, TotalThreadCount(InTotalThreadCount)
+		, SourceComponentBounds(InSourceComponentBounds)
+	{}
+
+	//~ Begin FComputeDataProviderRenderProxy Interface
+	bool IsValid(FValidationData const& InValidationData) const override;
+	int32 GetDispatchThreadCount(TArray<FIntVector>& InOutThreadCounts) const override;
+	void GatherDispatchData(FDispatchData const& InDispatchData) override;
+	//~ End FComputeDataProviderRenderProxy Interface
+
+protected:
+	using FParameters = FPCGKernelDataInterfaceParameters;
+
+	TArray<int32> InvocationThreadCounts;
+	int32 TotalThreadCount;
+	FBox SourceComponentBounds;
+};
diff --git a/Engine/Plugins/PCG/Source/PCG/Private/Compute/DataInterfaces/PCGDataCollectionDataInterface.cpp b/Engine/Plugins/PCG/Source/PCG/Private/Compute/DataInterfaces/PCGDataCollectionDataInterface.cpp
new file mode 100644
--- /dev/null
+++ b/Engine/Plugins/PCG/Source/PCG/Private/Compute/DataInterfaces/PCGDataCollectionDataInterface.cpp
@@ -0,0 +1,526 @@
+// Copyright Epic Games, Inc. All Rights Reserved.
+
+#include "Compute/DataInterfaces/PCGDataCollectionDataInterface.h"
+
+#include "PCGModule.h"
+#include "PCGSettings.h"
+#include "Compute/PCGDataBinding.h"
+
+#include "ComputeFramework/ShaderParamTypeDefinition.h"
+#include "RenderGraphBuilder.h"
+#include "RenderGraphResources.h"
+#include "ShaderCompilerCore.h"
+#include "ShaderParameterMetadataBuilder.h"
+
+void UPCGDataCollectionDataInterface::GetSupportedInputs(TArray<FShaderFunctionDefinition>& OutFunctions) const
+{
+	OutFunctions.AddDefaulted_GetRef()
+		.SetName(TEXT("LoadBuffer"))
+		.AddReturnType(EShaderFundamentalType::Uint)
+		.AddParam(EShaderFundamentalType::Uint);
+
+	OutFunctions.AddDefaulted_GetRef()
+		.SetName(TEXT("StoreBuffer"))
+		.AddParam(EShaderFundamentalType::Uint)
+		.AddParam(EShaderFundamentalType::Uint);
+
+	// Header Readers
+	{
+		OutFunctions.AddDefaulted_GetRef()
+			.SetName(TEXT("GetNumData"))
+			.AddReturnType(EShaderFundamentalType::Uint);
+
+		OutFunctions.AddDefaulted_GetRef()
+			.SetName(TEXT("ReadDataAddress"))
+			.AddReturnType(EShaderFundamentalType::Uint)
+			.AddParam(EShaderFundamentalType::Uint);
+
+		OutFunctions.AddDefaulted_GetRef()
+			.SetName(TEXT("ReadDataId"))
+			.AddReturnType(EShaderFundamentalType::Uint)
+			.AddParam(EShaderFundamentalType::Uint);
+
+		OutFunctions.AddDefaulted_GetRef()
+			.SetName(TEXT("ReadDataNumAttributes"))
+			.AddReturnType(EShaderFundamentalType::Uint)
+			.AddParam(EShaderFundamentalType::Uint);
+
+		OutFunctions.AddDefaulted_GetRef()
+			.SetName(TEXT("ReadDataPreambleSize"))
+			.AddReturnType(EShaderFundamentalType::Uint)
+			.AddParam(EShaderFundamentalType::Uint);
+
+		OutFunctions.AddDefaulted_GetRef()
+			.SetName(TEXT("ReadDataInfo"))
+			.AddReturnType(EShaderFundamentalType::Uint)
+			.AddParam(EShaderFundamentalType::Uint);
+
+		OutFunctions.AddDefaulted_GetRef()
+			.SetName(TEXT("ReadDataAttributeHeadersAddress"))
+			.AddReturnType(EShaderFundamentalType::Uint)
+			.AddParam(EShaderFundamentalType::Uint);
+
+		OutFunctions.AddDefaulted_GetRef()
+			.SetName(TEXT("GetThreadData"))
+			.AddReturnType(EShaderFundamentalType::Uint, 3)
+			.AddParam(EShaderFundamentalType::Uint);
+
+		OutFunctions.AddDefaulted_GetRef()
+			.SetName(TEXT("GetNumElements"))
+			.AddReturnType(EShaderFundamentalType::Uint);
+
+		OutFunctions.AddDefaulted_GetRef()
+			.SetName(TEXT("ReadAttributeHeaderAddress"))
+			.AddReturnType(EShaderFundamentalType::Uint)
+			.AddParam(EShaderFundamentalType::Uint)
+			.AddParam(EShaderFundamentalType::Uint);
+
+		OutFunctions.AddDefaulted_GetRef()
+			.SetName(TEXT("ReadAttributeIdAndStride"))
+			.AddReturnType(EShaderFundamentalType::Uint)
+			.AddParam(EShaderFundamentalType::Uint);
+
+		OutFunctions.AddDefaulted_GetRef()
+			.SetName(TEXT("ReadAttributeId"))
+			.AddReturnType(EShaderFundamentalType::Uint)
+			.AddParam(EShaderFundamentalType::Uint);
+
+		OutFunctions.AddDefaulted_GetRef()
+			.SetName(TEXT("ReadAttributeStride"))
+			.AddReturnType(EShaderFundamentalType::Uint)
+			.AddParam(EShaderFundamentalType::Uint);
+
+		OutFunctions.AddDefaulted_GetRef()
+			.SetName(TEXT("ReadAttributeAddress"))
+			.AddReturnType(EShaderFundamentalType::Uint)
+			.AddParam(EShaderFundamentalType::Uint);
+
+		OutFunctions.AddDefaulted_GetRef()
+			.SetName(TEXT("ReadAttributeAddress"))
+			.AddReturnType(EShaderFundamentalType::Uint)
+			.AddParam(EShaderFundamentalType::Uint)
+			.AddParam(EShaderFundamentalType::Uint);
+	}
+
+	// Header Writers
+	{
+		OutFunctions.AddDefaulted_GetRef()
+			.SetName(TEXT("WriteNumData"))
+			.AddParam(EShaderFundamentalType::Uint);
+
+		OutFunctions.AddDefaulted_GetRef()
+			.SetName(TEXT("WriteDataAddress"))
+			.AddParam(EShaderFundamentalType::Uint)
+			.AddParam(EShaderFundamentalType::Uint);
+
+		OutFunctions.AddDefaulted_GetRef()
+			.SetName(TEXT("WriteDataId"))
+			.AddParam(EShaderFundamentalType::Uint)
+			.AddParam(EShaderFundamentalType::Uint);
+
+		OutFunctions.AddDefaulted_GetRef()
+			.SetName(TEXT("WriteDataNumAttributes"))
+			.AddParam(EShaderFundamentalType::Uint)
+			.AddParam(EShaderFundamentalType::Uint);
+
+		OutFunctions.AddDefaulted_GetRef()
+			.SetName(TEXT("WriteDataPreambleSize"))
+			.AddParam(EShaderFundamentalType::Uint)
+			.AddParam(EShaderFundamentalType::Uint);
+
+		OutFunctions.AddDefaulted_GetRef()
+			.SetName(TEXT("WriteDataInfo"))
+			.AddParam(EShaderFundamentalType::Uint)
+			.AddParam(EShaderFundamentalType::Uint);
+
+		OutFunctions.AddDefaulted_GetRef()
+			.SetName(TEXT("WriteAttributeIdAndStride"))
+			.AddParam(EShaderFundamentalType::Uint)
+			.AddParam(EShaderFundamentalType::Uint);
+
+		OutFunctions.AddDefaulted_GetRef()
+			.SetName(TEXT("WriteAttributeAddress"))
+			.AddParam(EShaderFundamentalType::Uint)
+			.AddParam(EShaderFundamentalType::Uint);
+
+		OutFunctions.AddDefaulted_GetRef()
+			.SetName(TEXT("WriteHeader"));
+	}
+
+	// Attribute Getters
+	{
+		OutFunctions.AddDefaulted_GetRef()
+			.SetName(TEXT("GetBool"))
+			.AddReturnType(EShaderFundamentalType::Bool)
+			.AddParam(EShaderFundamentalType::Uint) // DataIndex
+			.AddParam(EShaderFundamentalType::Uint) // AttributeId
+			.AddParam(EShaderFundamentalType::Uint); // ElementIndex
+
+		OutFunctions.AddDefaulted_GetRef()
+			.SetName(TEXT("GetInt"))
+			.AddReturnType(EShaderFundamentalType::Int)
+			.AddParam(EShaderFundamentalType::Uint)
+			.AddParam(EShaderFundamentalType::Uint)
+			.AddParam(EShaderFundamentalType::Uint);
+
+		OutFunctions.AddDefaulted_GetRef()
+			.SetName(TEXT("GetFloat"))
+			.AddReturnType(EShaderFundamentalType::Float)
+			.AddParam(EShaderFundamentalType::Uint)
+			.AddParam(EShaderFundamentalType::Uint)
+			.AddParam(EShaderFundamentalType::Uint);
+
+		OutFunctions.AddDefaulted_GetRef()
+			.SetName(TEXT("GetFloat2"))
+			.AddReturnType(EShaderFundamentalType::Float, 2)
+			.AddParam(EShaderFundamentalType::Uint)
+			.AddParam(EShaderFundamentalType::Uint)
+			.AddParam(EShaderFundamentalType::Uint);
+
+		OutFunctions.AddDefaulted_GetRef()
+			.SetName(TEXT("GetFloat3"))
+			.AddReturnType(EShaderFundamentalType::Float, 3)
+			.AddParam(EShaderFundamentalType::Uint)
+			.AddParam(EShaderFundamentalType::Uint)
+			.AddParam(EShaderFundamentalType::Uint);
+
+		OutFunctions.AddDefaulted_GetRef()
+			.SetName(TEXT("GetFloat4"))
+			.AddReturnType(EShaderFundamentalType::Float, 4)
+			.AddParam(EShaderFundamentalType::Uint)
+			.AddParam(EShaderFundamentalType::Uint)
+			.AddParam(EShaderFundamentalType::Uint);
+
+		OutFunctions.AddDefaulted_GetRef()
+			.SetName(TEXT("GetRotator"))
+			.AddReturnType(EShaderFundamentalType::Float, 3)
+			.AddParam(EShaderFundamentalType::Uint)
+			.AddParam(EShaderFundamentalType::Uint)
+			.AddParam(EShaderFundamentalType::Uint);
+
+		OutFunctions.AddDefaulted_GetRef()
+			.SetName(TEXT("GetQuat"))
+			.AddReturnType(EShaderFundamentalType::Float, 4)
+			.AddParam(EShaderFundamentalType::Uint)
+			.AddParam(EShaderFundamentalType::Uint)
+			.AddParam(EShaderFundamentalType::Uint);
+
+		OutFunctions.AddDefaulted_GetRef()
+			.SetName(TEXT("GetTransform"))
+			.AddReturnType(EShaderFundamentalType::Float, 4, 4)
+			.AddParam(EShaderFundamentalType::Uint)
+			.AddParam(EShaderFundamentalType::Uint)
+			.AddParam(EShaderFundamentalType::Uint);
+	}
+
+	// Attribute Setters
+	{
+		OutFunctions.AddDefaulted_GetRef()
+			.SetName(TEXT("SetBool"))
+			.AddParam(EShaderFundamentalType::Uint) // DataIndex
+			.AddParam(EShaderFundamentalType::Uint) // AttributeId
+			.AddParam(EShaderFundamentalType::Uint) // ElementIndex
+			.AddParam(EShaderFundamentalType::Bool); // Value
+
+		OutFunctions.AddDefaulted_GetRef()
+			.SetName(TEXT("SetInt"))
+			.AddParam(EShaderFundamentalType::Uint)
+			.AddParam(EShaderFundamentalType::Uint)
+			.AddParam(EShaderFundamentalType::Uint)
+			.AddParam(EShaderFundamentalType::Int);
+
+		OutFunctions.AddDefaulted_GetRef()
+			.SetName(TEXT("SetFloat"))
+			.AddParam(EShaderFundamentalType::Uint)
+			.AddParam(EShaderFundamentalType::Uint)
+			.AddParam(EShaderFundamentalType::Uint)
+			.AddParam(EShaderFundamentalType::Float);
+
+		OutFunctions.AddDefaulted_GetRef()
+			.SetName(TEXT("SetFloat2"))
+			.AddParam(EShaderFundamentalType::Uint)
+			.AddParam(EShaderFundamentalType::Uint)
+			.AddParam(EShaderFundamentalType::Uint)
+			.AddParam(EShaderFundamentalType::Float, 2);
+
+		OutFunctions.AddDefaulted_GetRef()
+			.SetName(TEXT("SetFloat3"))
+			.AddParam(EShaderFundamentalType::Uint)
+			.AddParam(EShaderFundamentalType::Uint)
+			.AddParam(EShaderFundamentalType::Uint)
+			.AddParam(EShaderFundamentalType::Float, 3);
+
+		OutFunctions.AddDefaulted_GetRef()
+			.SetName(TEXT("SetFloat4"))
+			.AddParam(EShaderFundamentalType::Uint)
+			.AddParam(EShaderFundamentalType::Uint)
+			.AddParam(EShaderFundamentalType::Uint)
+			.AddParam(EShaderFundamentalType::Float, 4);
+
+		OutFunctions.AddDefaulted_GetRef()
+			.SetName(TEXT("SetRotator"))
+			.AddParam(EShaderFundamentalType::Uint)
+			.AddParam(EShaderFundamentalType::Uint)
+			.AddParam(EShaderFundamentalType::Uint)
+			.AddParam(EShaderFundamentalType::Float, 3);
+
+		OutFunctions.AddDefaulted_GetRef()
+			.SetName(TEXT("SetQuat"))
+			.AddParam(EShaderFundamentalType::Uint)
+			.AddParam(EShaderFundamentalType::Uint)
+			.AddParam(EShaderFundamentalType::Uint)
+			.AddParam(EShaderFundamentalType::Float, 4);
+
+		OutFunctions.AddDefaulted_GetRef()
+			.SetName(TEXT("SetTransform"))
+			.AddParam(EShaderFundamentalType::Uint)
+			.AddParam(EShaderFundamentalType::Uint)
+			.AddParam(EShaderFundamentalType::Uint)
+			.AddParam(EShaderFundamentalType::Float, 4, 4);
+	}
+
+	// Point Attribute Getters
+	{
+		OutFunctions.AddDefaulted_GetRef()
+			.SetName(TEXT("GetPosition"))
+			.AddReturnType(EShaderFundamentalType::Float, 3)
+			.AddParam(EShaderFundamentalType::Uint) // DataIndex
+			.AddParam(EShaderFundamentalType::Uint); // ElementIndex
+
+		OutFunctions.AddDefaulted_GetRef()
+			.SetName(TEXT("GetRotation"))
+			.AddReturnType(EShaderFundamentalType::Float, 4)
+			.AddParam(EShaderFundamentalType::Uint)
+			.AddParam(EShaderFundamentalType::Uint);
+
+		OutFunctions.AddDefaulted_GetRef()
+			.SetName(TEXT("GetScale"))
+			.AddReturnType(EShaderFundamentalType::Float, 3)
+			.AddParam(EShaderFundamentalType::Uint)
+			.AddParam(EShaderFundamentalType::Uint);
+
+		OutFunctions.AddDefaulted_GetRef()
+			.SetName(TEXT("GetBoundsMin"))
+			.AddReturnType(EShaderFundamentalType::Float, 3)
+			.AddParam(EShaderFundamentalType::Uint)
+			.AddParam(EShaderFundamentalType::Uint);
+
+		OutFunctions.AddDefaulted_GetRef()
+			.SetName(TEXT("GetBoundsMax"))
+			.AddReturnType(EShaderFundamentalType::Float, 3)
+			.AddParam(EShaderFundamentalType::Uint)
+			.AddParam(EShaderFundamentalType::Uint);
+
+		OutFunctions.AddDefaulted_GetRef()
+			.SetName(TEXT("GetColor"))
+			.AddReturnType(EShaderFundamentalType::Float, 4)
+			.AddParam(EShaderFundamentalType::Uint)
+			.AddParam(EShaderFundamentalType::Uint);
+
+		OutFunctions.AddDefaulted_GetRef()
+			.SetName(TEXT("GetDensity"))
+			.AddReturnType(EShaderFundamentalType::Float)
+			.AddParam(EShaderFundamentalType::Uint)
+			.AddParam(EShaderFundamentalType::Uint);
+
+		OutFunctions.AddDefaulted_GetRef()
+			.SetName(TEXT("GetSeed"))
+			.AddReturnType(EShaderFundamentalType::Int)
+			.AddParam(EShaderFundamentalType::Uint)
+			.AddParam(EShaderFundamentalType::Uint);
+
+		OutFunctions.AddDefaulted_GetRef()
+			.SetName(TEXT("GetSteepness"))
+			.AddReturnType(EShaderFundamentalType::Float)
+			.AddParam(EShaderFundamentalType::Uint)
+			.AddParam(EShaderFundamentalType::Uint);
+
+		OutFunctions.AddDefaulted_GetRef()
+			.SetName(TEXT("GetPointTransform"))
+			.AddReturnType(EShaderFundamentalType::Float, 4, 4)
+			.AddParam(EShaderFundamentalType::Uint)
+			.AddParam(EShaderFundamentalType::Uint);
+	}
+
+	// Point Attribute Setters
+	{
+		OutFunctions.AddDefaulted_GetRef()
+			.SetName(TEXT("SetPosition"))
+			.AddParam(EShaderFundamentalType::Uint) // DataIndex
+			.AddParam(EShaderFundamentalType::Uint) // ElementIndex
+			.AddParam(EShaderFundamentalType::Float, 3); // Value
+
+		OutFunctions.AddDefaulted_GetRef()
+			.SetName(TEXT("SetRotation"))
+			.AddParam(EShaderFundamentalType::Uint)
+			.AddParam(EShaderFundamentalType::Uint)
+			.AddParam(EShaderFundamentalType::Float, 4);
+
+		OutFunctions.AddDefaulted_GetRef()
+			.SetName(TEXT("SetScale"))
+			.AddParam(EShaderFundamentalType::Uint)
+			.AddParam(EShaderFundamentalType::Uint)
+			.AddParam(EShaderFundamentalType::Float, 3);
+
+		OutFunctions.AddDefaulted_GetRef()
+			.SetName(TEXT("SetBoundsMin"))
+			.AddParam(EShaderFundamentalType::Uint)
+			.AddParam(EShaderFundamentalType::Uint)
+			.AddParam(EShaderFundamentalType::Float, 3);
+
+		OutFunctions.AddDefaulted_GetRef()
+			.SetName(TEXT("SetBoundsMax"))
+			.AddParam(EShaderFundamentalType::Uint)
+			.AddParam(EShaderFundamentalType::Uint)
+			.AddParam(EShaderFundamentalType::Float, 3);
+
+		OutFunctions.AddDefaulted_GetRef()
+			.SetName(TEXT("SetColor"))
+			.AddParam(EShaderFundamentalType::Uint)
+			.AddParam(EShaderFundamentalType::Uint)
+			.AddParam(EShaderFundamentalType::Float, 4);
+
+		OutFunctions.AddDefaulted_GetRef()
+			.SetName(TEXT("SetDensity"))
+			.AddParam(EShaderFundamentalType::Uint)
+			.AddParam(EShaderFundamentalType::Uint)
+			.AddParam(EShaderFundamentalType::Float);
+
+		OutFunctions.AddDefaulted_GetRef()
+			.SetName(TEXT("SetSeed"))
+			.AddParam(EShaderFundamentalType::Uint)
+			.AddParam(EShaderFundamentalType::Uint)
+			.AddParam(EShaderFundamentalType::Int);
+
+		OutFunctions.AddDefaulted_GetRef()
+			.SetName(TEXT("SetSeedFromPosition"))
+			.AddParam(EShaderFundamentalType::Uint)
+			.AddParam(EShaderFundamentalType::Uint)
+			.AddParam(EShaderFundamentalType::Float, 3);
+
+		OutFunctions.AddDefaulted_GetRef()
+			.SetName(TEXT("SetSteepness"))
+			.AddParam(EShaderFundamentalType::Uint)
+			.AddParam(EShaderFundamentalType::Uint)
+			.AddParam(EShaderFundamentalType::Float);
+
+		OutFunctions.AddDefaulted_GetRef()
+			.SetName(TEXT("SetPointTransform"))
+			.AddParam(EShaderFundamentalType::Uint)
+			.AddParam(EShaderFundamentalType::Uint)
+			.AddParam(EShaderFundamentalType::Float, 4, 4);
+
+		OutFunctions.AddDefaulted_GetRef()
+			.SetName(TEXT("WriteSinglePointDataCollectionHeader"))
+			.AddParam(EShaderFundamentalType::Uint);
+
+		OutFunctions.AddDefaulted_GetRef()
+			.SetName(TEXT("InitializePoint"))
+			.AddParam(EShaderFundamentalType::Uint)
+			.AddParam(EShaderFundamentalType::Uint);
+	}
+}
+
+BEGIN_SHADER_PARAMETER_STRUCT(FPCGDataCollectionDataInterfaceParameters, )
+	SHADER_PARAMETER_RDG_BUFFER_UAV(RWByteAddressBuffer, DataCollectionBuffer)
+END_SHADER_PARAMETER_STRUCT()
+
+void UPCGDataCollectionDataInterface::GetShaderParameters(TCHAR const* UID, FShaderParametersMetadataBuilder& InOutBuilder, FShaderParametersMetadataAllocations& InOutAllocations) const
+{
+	InOutBuilder.AddNestedStruct<FPCGDataCollectionDataInterfaceParameters>(UID);
+}
+
+TCHAR const* UPCGDataCollectionDataInterface::TemplateFilePath = TEXT("/Plugin/PCG/Private/PCGDataCollectionDataInterface.ush");
+
+TCHAR const* UPCGDataCollectionDataInterface::GetShaderVirtualPath() const
+{
+	return TemplateFilePath;
+}
+
+void UPCGDataCollectionDataInterface::GetShaderHash(FString& InOutKey) const
+{
+	GetShaderFileHash(TemplateFilePath, EShaderPlatform::SP_PCD3D_SM5).AppendString(InOutKey);
+
+	// Quaternion library
+	GetShaderFileHash(TEXT("/Engine/Private/Quaternion.ush"), EShaderPlatform::SP_PCD3D_SM5).AppendString(InOutKey);
+}
+
+void UPCGDataCollectionDataInterface::GetHLSL(FString& OutHLSL, FString const& InDataInterfaceName) const
+{
+	TMap<FString, FStringFormatArg> TemplateArgs =
+	{
+		{ TEXT("DataInterfaceName"), InDataInterfaceName },
+	};
+
+	FString TemplateFile;
+	LoadShaderSourceFile(TemplateFilePath, EShaderPlatform::SP_PCD3D_SM5, &TemplateFile, nullptr);
+	OutHLSL += FString::Format(*TemplateFile, TemplateArgs);
+}
+
+UComputeDataProvider* UPCGDataCollectionDataInterface::CreateDataProvider(TObjectPtr<UObject> InBinding, uint64 InInputMask, uint64 InOutputMask) const
+{
+	UPCGDataBinding* Binding = CastChecked<UPCGDataBinding>(InBinding);
+
+	UPCGDataCollectionDataProvider* Provider = NewObject<UPCGDataCollectionDataProvider>();
+	Provider->Binding = Binding;
+	Provider->PinDesc = ProducerSettings->ComputeOutputPinDataDesc(OutputPinLabel, Binding);
+	Provider->SizeBytes = Provider->PinDesc.ComputePackedSize();
+	check(Provider->SizeBytes >= 4);
+
+	return Provider;
+}
+
+FComputeDataProviderRenderProxy* UPCGDataCollectionDataProvider::GetRenderProxy()
+{
+	FPCGDataCollectionDataProviderProxy* Proxy = new FPCGDataCollectionDataProviderProxy(Binding, PinDesc);
+
+	check(SizeBytes > 0);
+	Proxy->SizeBytes = SizeBytes;
+
+	return Proxy;
+}
+
+bool FPCGDataCollectionDataProviderProxy::IsValid(FValidationData const& InValidationData) const
+{
+	if (InValidationData.ParameterStructSize != sizeof(FParameters))
+	{
+		return false;
+	}
+
+	if (!Binding.IsValid())
+	{
+		UE_LOG(LogPCG, Error, TEXT("Proxy invalid due to missing data binding."));
+		return false;
+	}
+
+	return true;
+}
+
+void FPCGDataCollectionDataProviderProxy::GatherDispatchData(FDispatchData const& InDispatchData)
+{
+	check(BufferUAV);
+
+	const TStridedView<FParameters> ParameterArray = MakeStridedParameterView<FParameters>(InDispatchData);
+	for (int32 InvocationIndex = 0; InvocationIndex < ParameterArray.Num(); ++InvocationIndex)
+	{
+		FParameters& Parameters = ParameterArray[InvocationIndex];
+		Parameters.DataCollectionBuffer = BufferUAV;
+	}
+}
+
+void FPCGDataCollectionDataProviderProxy::AllocateResources(FRDGBuilder& GraphBuilder, FAllocationData const& InAllocationData)
+{
+	check(SizeBytes > 0 && (SizeBytes % sizeof(uint32)) == 0);
+
+	FRDGBufferDesc Desc = FRDGBufferDesc::CreateStructuredDesc(sizeof(uint32), SizeBytes / sizeof(uint32));
+	Buffer = GraphBuilder.CreateBuffer(Desc, TEXT("PCGDataCollectionBuffer"));
+	BufferUAV = GraphBuilder.CreateUAV(Buffer);
+
+	// Initialize with an empty data collection. The kernel may not run, for example if indirect dispatch args end up being 0. Ensure
+	// there is something meaningful to readback.
+	// TODO could have a statically-allocated resource rather than allocating & uploading here.
+	TArray<uint32> PackedDataCollection;
+	PinDesc.PackDataCollection(/*InDataCollection=*/{}, /*InPin=*/NAME_None, PackedDataCollection);
+
+	GraphBuilder.QueueBufferUpload(Buffer, PackedDataCollection.GetData(), PackedDataCollection.Num() * PackedDataCollection.GetTypeSize(), ERDGInitialDataFlags::None);
+}
diff --git a/Engine/Plugins/PCG/Source/PCG/Private/Compute/DataInterfaces/PCGDataCollectionDataInterface.h b/Engine/Plugins/PCG/Source/PCG/Private/Compute/DataInterfaces/PCGDataCollectionDataInterface.h
new file mode 100644
--- /dev/null
+++ b/Engine/Plugins/PCG/Source/PCG/Private/Compute/DataInterfaces/PCGDataCollectionDataInterface.h
@@ -0,0 +1,87 @@
+// Copyright Epic Games, Inc. All Rights Reserved.
+
+#pragma once
+
+#include "Compute/PCGDataForGPU.h"
+#include "Compute/DataInterfaces/PCGComputeDataInterface.h"
+
+#include "ComputeFramework/ComputeDataProvider.h"
+
+#include "PCGDataCollectionDataInterface.generated.h"
+
+class FPCGDataCollectionDataInterfaceParameters;
+class FRDGBuffer;
+class FRDGBufferSRV;
+class UPCGDataBinding;
+
+/** Compute Framework Data Interface for reading PCG data. */
+UCLASS(ClassGroup = (Procedural))
+class UPCGDataCollectionDataInterface : public UPCGComputeDataInterface
+{
+	GENERATED_BODY()
+
+public:
+	//~ Begin UComputeDataInterface Interface
+	TCHAR const* GetClassName() const override { return TEXT("PCGDataCollection"); }
+	/** Return true if the associated UComputeDataProvider holds data that can be combined into a single dispatch invocation. */
+	bool CanSupportUnifiedDispatch() const override { return false; } // I think this means compute shader can produce multiple buffers simultaneously?
+	// TODO don't allow writing to an input!
+	void GetSupportedInputs(TArray<FShaderFunctionDefinition>& OutFunctions) const override;
+	// TODO could differentiate later for SRV vs UAV.
+	void GetSupportedOutputs(TArray<FShaderFunctionDefinition>& OutFunctions) const override { GetSupportedInputs(OutFunctions); }
+	void GetShaderParameters(TCHAR const* UID, FShaderParametersMetadataBuilder& InOutBuilder, FShaderParametersMetadataAllocations& InOutAllocations) const override;
+	TCHAR const* GetShaderVirtualPath() const override; // TODO probably easier to just inline rather than external source?
+	void GetShaderHash(FString& InOutKey) const override;
+	void GetHLSL(FString& OutHLSL, FString const& InDataInterfaceName) const override;
+	UComputeDataProvider* CreateDataProvider(TObjectPtr<UObject> InBinding, uint64 InInputMask, uint64 InOutputMask) const override;
+	//~ End UComputeDataInterface Interface
+
+	/** Settings of node that produces this data, normally the upstream node. */
+	UPROPERTY()
+	const UPCGSettings* ProducerSettings;
+
+private:
+	static TCHAR const* TemplateFilePath;
+};
+
+/** Compute Framework Data Provider for reading a PCG Data Collection. */
+UCLASS()
+class UPCGDataCollectionDataProvider : public UComputeDataProvider
+{
+	GENERATED_BODY()
+
+public:
+	//~ Begin UComputeDataProvider Interface
+	FComputeDataProviderRenderProxy* GetRenderProxy() override;
+	//~ End UComputeDataProvider Interface
+
+	TWeakObjectPtr<UPCGDataBinding> Binding;
+	FPCGDataCollectionDesc PinDesc;
+	int SizeBytes = -1;
+};
+
+class FPCGDataCollectionDataProviderProxy : public FComputeDataProviderRenderProxy
+{
+public:
+	FPCGDataCollectionDataProviderProxy(TWeakObjectPtr<UPCGDataBinding> InBinding, const FPCGDataCollectionDesc& InPinDesc)
+		: Binding(InBinding)
+		, PinDesc(InPinDesc)
+	{}
+
+	//~ Begin FComputeDataProviderRenderProxy Interface
+	bool IsValid(FValidationData const& InValidationData) const override;
+	void GatherDispatchData(FDispatchData const& InDispatchData) override;
+	void AllocateResources(FRDGBuilder& GraphBuilder, FAllocationData const& InAllocationData) override;
+	//~ End FComputeDataProviderRenderProxy Interface
+
+	int SizeBytes = -1;
+
+protected:
+	using FParameters = FPCGDataCollectionDataInterfaceParameters;
+
+	TWeakObjectPtr<UPCGDataBinding> Binding;
+	FPCGDataCollectionDesc PinDesc;
+
+	FRDGBufferRef Buffer = nullptr;
+	FRDGBufferUAVRef BufferUAV = nullptr;
+};
diff --git a/Engine/Plugins/PCG/Source/PCG/Private/Compute/DataInterfaces/PCGDataCollectionReadbackDataInterface.cpp b/Engine/Plugins/PCG/Source/PCG/Private/Compute/DataInterfaces/PCGDataCollectionReadbackDataInterface.cpp
new file mode 100644
--- /dev/null
+++ b/Engine/Plugins/PCG/Source/PCG/Private/Compute/DataInterfaces/PCGDataCollectionReadbackDataInterface.cpp
@@ -0,0 +1,139 @@
+// Copyright Epic Games, Inc. All Rights Reserved.
+
+#include "Compute/DataInterfaces/PCGDataCollectionReadbackDataInterface.h"
+
+#include "PCGModule.h"
+#include "PCGSettings.h"
+#include "Compute/PCGDataBinding.h"
+
+#include "RenderGraphBuilder.h"
+#include "RenderGraphResources.h"
+
+UComputeDataProvider* UPCGDataCollectionReadbackDataInterface::CreateDataProvider(TObjectPtr<UObject> InBinding, uint64 InInputMask, uint64 InOutputMask) const
+{
+	UPCGDataBinding* Binding = CastChecked<UPCGDataBinding>(InBinding);
+
+	UPCGDataProviderDataCollectionReadback* Provider = NewObject<UPCGDataProviderDataCollectionReadback>();
+	Provider->Binding = Binding;
+	Provider->PinDesc = ProducerSettings->ComputeOutputPinDataDesc(OutputPinLabel, Binding);
+	// Use the aliased label as this is the output from the compute graph.
+	Provider->OutputPinLabelAlias = OutputPinLabelAlias;
+	Provider->SizeBytes = Provider->PinDesc.ComputePackedSize();
+	check(Provider->SizeBytes >= 4);
+
+	return Provider;
+}
+
+FComputeDataProviderRenderProxy* UPCGDataProviderDataCollectionReadback::GetRenderProxy()
+{
+	TWeakObjectPtr<UPCGDataProviderDataCollectionReadback> ThisWeakPtr(this);
+
+	auto ProcessReadbackData_RenderThread = [ThisWeakPtr](const void* InData, int InNumBytes)
+	{
+		TRACE_CPUPROFILER_EVENT_SCOPE(UPCGDataProviderDataCollectionReadback::ProcessReadbackData);
+
+		if (!ThisWeakPtr.IsValid() || !ThisWeakPtr.Get())
+		{
+			return;
+		}
+
+		UPCGDataProviderDataCollectionReadback* ThisDataProvider = ThisWeakPtr.Get();
+		check(ThisDataProvider);
+
+		// We should never find ourselves stomping existing data.
+		check(ThisDataProvider->RawReadbackData.IsEmpty());
+
+		if (InData && InNumBytes > 0)
+		{
+			// Copy the data to temp storage for the game thread to pick up.
+			ThisDataProvider->RawReadbackData.SetNumUninitialized(InNumBytes);
+			FMemory::Memcpy(ThisDataProvider->RawReadbackData.GetData(), InData, InNumBytes);
+		}
+		else
+		{
+			// Can happen currently if no threads dispatched.
+			ThisDataProvider->RawReadbackData.Reset();
+		}
+
+		ThisDataProvider->bReadbackComplete = true;
+
+		ThisDataProvider->OnReadbackComplete.Broadcast();
+	};
+
+	return new FPCGDataProviderDataCollectionReadbackProxy(Binding, PinDesc, SizeBytes, ProcessReadbackData_RenderThread);
+}
+
+bool UPCGDataProviderDataCollectionReadback::ProcessReadBackData()
+{
+	if (!ensure(bReadbackComplete))
+	{
+		// This should not be called until readback has been done.
+		return false;
+	}
+
+	if (RawReadbackData.IsEmpty())
+	{
+		// If there's no data, we leave an empty data collection and we're done.
+		return true;
+	}
+
+	if (Binding.IsValid())
+	{
+		// TODO these are not adding in order if multiple buffers are read back? Maybe that never happens here though.
+		PinDesc.UnpackDataCollection(RawReadbackData, OutputPinLabelAlias, Binding->OutputDataCollection);
+	}
+
+	RawReadbackData.Reset();
+
+	return true;
+}
+
+FPCGDataProviderDataCollectionReadbackProxy::FPCGDataProviderDataCollectionReadbackProxy(
+	TWeakObjectPtr<UPCGDataBinding> InBinding,
+	const FPCGDataCollectionDesc& InPinDesc,
+	int InSizeBytes,
+	FReadbackCallback InAsyncReadbackCallback_RenderThread)
+	: FPCGDataCollectionDataProviderProxy(InBinding, InPinDesc)
+{
+	SizeBytes = InSizeBytes;
+	AsyncReadbackCallback_RenderThread = InAsyncReadbackCallback_RenderThread;
+}
+
+bool FPCGDataProviderDataCollectionReadbackProxy::IsValid(FValidationData const& InValidationData) const
+{
+	if (SizeBytes <= 0)
+	{
+		UE_LOG(LogPCG, Error, TEXT("Proxy invalid due to invalid size."));
+		return false;
+	}
+
+	return FPCGDataCollectionDataProviderProxy::IsValid(InValidationData);
+}
+
+void FPCGDataProviderDataCollectionReadbackProxy::AllocateResources(FRDGBuilder& GraphBuilder, FAllocationData const& InAllocationData)
+{
+	check(SizeBytes > 0);
+
+	FRDGBufferDesc Desc = FRDGBufferDesc::CreateStructuredDesc(sizeof(uint32), SizeBytes / sizeof(uint32));
+	Desc.Usage = EBufferUsageFlags(Desc.Usage | BUF_SourceCopy);
+	Buffer = GraphBuilder.CreateBuffer(Desc, TEXT("PCGDataCollectionReadbackBuffer"));
+	BufferUAV = GraphBuilder.CreateUAV(Buffer);
+
+	// Initialize with an empty data collection. The kernel may not run, for example if indirect dispatch args end up being 0. Ensure
+	// there is something meaningful to readback.
+	// TODO factor out to avoid doing this work each time.
+	TArray<uint32> PackedDataCollection;
+	PinDesc.PackDataCollection(/*InDataCollection=*/{}, /*InPin=*/NAME_None, PackedDataCollection);
+
+	GraphBuilder.QueueBufferUpload(Buffer, PackedDataCollection.GetData(), PackedDataCollection.Num() * PackedDataCollection.GetTypeSize(), ERDGInitialDataFlags::None);
+}
+
+void FPCGDataProviderDataCollectionReadbackProxy::GetReadbackData(TArray<FReadbackData>& OutReadbackData) const
+{
+	FReadbackData Data;
+	Data.Buffer = Buffer;
+	Data.NumBytes = SizeBytes;
+	Data.ReadbackCallback_RenderThread = &AsyncReadbackCallback_RenderThread;
+
+	OutReadbackData.Add(MoveTemp(Data));
+}
diff --git a/Engine/Plugins/PCG/Source/PCG/Private/Compute/DataInterfaces/PCGDataCollectionReadbackDataInterface.h b/Engine/Plugins/PCG/Source/PCG/Private/Compute/DataInterfaces/PCGDataCollectionReadbackDataInterface.h
new file mode 100644
--- /dev/null
+++ b/Engine/Plugins/PCG/Source/PCG/Private/Compute/DataInterfaces/PCGDataCollectionReadbackDataInterface.h
@@ -0,0 +1,66 @@
+// Copyright Epic Games, Inc. All Rights Reserved.
+
+#pragma once
+
+#include "Compute/DataInterfaces/PCGDataCollectionDataInterface.h"
+
+#include "ComputeFramework/ComputeDataInterface.h"
+#include "ComputeFramework/ComputeDataProvider.h"
+
+#include "PCGDataCollectionReadbackDataInterface.generated.h"
+
+class UPCGDataBinding;
+
+/** Compute Framework Data Interface for reading back data from PCG data collection. */
+UCLASS(ClassGroup = (Procedural))
+class UPCGDataCollectionReadbackDataInterface : public UPCGDataCollectionDataInterface
+{
+	GENERATED_BODY()
+
+public:
+	//~ Begin UComputeDataInterface Interface
+	virtual bool GetRequiresReadback() const override { return true; }
+	UComputeDataProvider* CreateDataProvider(TObjectPtr<UObject> InBinding, uint64 InInputMask, uint64 InOutputMask) const override;
+	//~ End UComputeDataInterface Interface
+};
+
+/** Compute Framework Data Provider for writing a PCG Data Collection. */
+UCLASS()
+class UPCGDataProviderDataCollectionReadback : public UPCGDataCollectionDataProvider
+{
+	GENERATED_BODY()
+
+public:
+	//~ Begin UComputeDataProvider Interface
+	FComputeDataProviderRenderProxy* GetRenderProxy() override;
+	//~ End UComputeDataProvider Interface
+
+	/** Processed read back data. Should be called after readback complete. Called from game thread. Returns true if data processed successfully. */
+	bool ProcessReadBackData();
+
+	std::atomic<bool> bReadbackComplete = false;
+	TArray<uint8> RawReadbackData;
+
+	FName OutputPinLabelAlias;
+
+	DECLARE_EVENT(UPCGDataProviderDataCollectionReadback, FOnReadbackComplete);
+	FOnReadbackComplete& OnReadbackComplete_RenderThread() { return OnReadbackComplete; }
+
+private:
+	FOnReadbackComplete OnReadbackComplete;
+};
+
+class FPCGDataProviderDataCollectionReadbackProxy : public FPCGDataCollectionDataProviderProxy
+{
+public:
+	FPCGDataProviderDataCollectionReadbackProxy(TWeakObjectPtr<UPCGDataBinding> InBinding, const FPCGDataCollectionDesc& InPinDesc, int InSizeBytes, FReadbackCallback InAsyncReadbackCallback_RenderThread);
+
+	//~ Begin FComputeDataProviderRenderProxy Interface
+	bool IsValid(FValidationData const& InValidationData) const override;
+	void AllocateResources(FRDGBuilder& GraphBuilder, FAllocationData const& InAllocationData) override;
+	void GetReadbackData(TArray<FReadbackData>& OutReadbackData) const override;
+	//~ End FComputeDataProviderRenderProxy Interface
+
+private:
+	FReadbackCallback AsyncReadbackCallback_RenderThread;
+};
diff --git a/Engine/Plugins/PCG/Source/PCG/Private/Compute/DataInterfaces/PCGDataCollectionUploadDataInterface.cpp b/Engine/Plugins/PCG/Source/PCG/Private/Compute/DataInterfaces/PCGDataCollectionUploadDataInterface.cpp
new file mode 100644
--- /dev/null
+++ b/Engine/Plugins/PCG/Source/PCG/Private/Compute/DataInterfaces/PCGDataCollectionUploadDataInterface.cpp
@@ -0,0 +1,69 @@
+// Copyright Epic Games, Inc. All Rights Reserved.
+
+#include "Compute/DataInterfaces/PCGDataCollectionUploadDataInterface.h"
+
+#include "PCGEdge.h"
+#include "PCGModule.h"
+#include "PCGNode.h"
+#include "PCGSettings.h"
+#include "Compute/PCGComputeGraph.h"
+#include "Compute/PCGDataBinding.h"
+
+#include "RenderGraphBuilder.h"
+#include "RenderGraphResources.h"
+
+UComputeDataProvider* UPCGDataCollectionUploadDataInterface::CreateDataProvider(TObjectPtr<UObject> InBinding, uint64 InInputMask, uint64 InOutputMask) const
+{
+	UPCGDataBinding* Binding = CastChecked<UPCGDataBinding>(InBinding);
+
+	UPCGDataProviderDataCollectionUpload* Provider = NewObject<UPCGDataProviderDataCollectionUpload>();
+	Provider->Binding = Binding;
+
+	// Pick the data items from input data collection using any of the compute graph element virtual input pin labels.
+	// TODO: Ideally we could call ProducerSettings->ComputeOutputPinDataDesc() but some settings do not have associated nodes/pins.
+	check(!DownstreamInputPinLabelAliases.IsEmpty());
+	Provider->PinDesc = FPCGDataCollectionDesc::BuildFromInputDataCollectionAndInputPinLabel(
+		Binding->DataForGPU.InputDataCollection,
+		DownstreamInputPinLabelAliases[0],
+		Binding->Graph->GetAttributeLookupTable());
+
+	Provider->DownstreamInputPinLabels = DownstreamInputPinLabelAliases;
+
+	return Provider;
+}
+
+FComputeDataProviderRenderProxy* UPCGDataProviderDataCollectionUpload::GetRenderProxy()
+{
+	return new FPCGDataProviderDataCollectionUploadProxy(Binding, PinDesc, DownstreamInputPinLabels);
+}
+
+FPCGDataProviderDataCollectionUploadProxy::FPCGDataProviderDataCollectionUploadProxy(
+	TWeakObjectPtr<UPCGDataBinding> InBinding,
+	const FPCGDataCollectionDesc& InPinDesc,
+	const TArray<FName>& InDownstreamInputPinLabels)
+	: FPCGDataCollectionDataProviderProxy(InBinding, InPinDesc)
+{
+	DownstreamInputPinLabels = InDownstreamInputPinLabels;
+}
+
+void FPCGDataProviderDataCollectionUploadProxy::AllocateResources(FRDGBuilder& GraphBuilder, FAllocationData const& InAllocationData)
+{
+	if (!ensure(Binding.IsValid()))
+	{
+		return;
+	}
+
+	TArray<uint32> PackedDataCollection;
+
+	// Use any downstream input pin label to grab data from the collection.
+	check(!DownstreamInputPinLabels.IsEmpty());
+	PinDesc.PackDataCollection(Binding->DataForGPU.InputDataCollection, DownstreamInputPinLabels[0], PackedDataCollection);
+	
+	SizeBytes = PackedDataCollection.Num() * sizeof(uint32);
+
+	const FRDGBufferDesc Desc = FRDGBufferDesc::CreateStructuredDesc(sizeof(uint32), PackedDataCollection.Num());
+	Buffer = GraphBuilder.CreateBuffer(Desc, TEXT("PCGDataCollectionUploadBuffer"));
+	BufferUAV = GraphBuilder.CreateUAV(Buffer);
+
+	GraphBuilder.QueueBufferUpload(Buffer, PackedDataCollection.GetData(), PackedDataCollection.Num() * PackedDataCollection.GetTypeSize(), ERDGInitialDataFlags::None);
+}
diff --git a/Engine/Plugins/PCG/Source/PCG/Private/Compute/DataInterfaces/PCGDataCollectionUploadDataInterface.h b/Engine/Plugins/PCG/Source/PCG/Private/Compute/DataInterfaces/PCGDataCollectionUploadDataInterface.h
new file mode 100644
--- /dev/null
+++ b/Engine/Plugins/PCG/Source/PCG/Private/Compute/DataInterfaces/PCGDataCollectionUploadDataInterface.h
@@ -0,0 +1,49 @@
+// Copyright Epic Games, Inc. All Rights Reserved.
+
+// Based on OptimusDataInterfaceSkinnedMesh.h
+
+#pragma once
+
+#include "Compute/DataInterfaces/PCGDataCollectionDataInterface.h"
+
+#include "Compute/PCGDataBinding.h"
+
+#include "PCGDataCollectionUploadDataInterface.generated.h"
+
+/** Compute Framework Data Interface for reading PCG data. */
+UCLASS(ClassGroup = (Procedural))
+class UPCGDataCollectionUploadDataInterface : public UPCGDataCollectionDataInterface
+{
+	GENERATED_BODY()
+
+public:
+	//~ Begin UComputeDataInterface Interface
+	UComputeDataProvider* CreateDataProvider(TObjectPtr<UObject> InBinding, uint64 InInputMask, uint64 InOutputMask) const override;
+	//~ End UComputeDataInterface Interface
+};
+
+/** Compute Framework Data Provider for reading a PCG Data Collection. */
+UCLASS()
+class UPCGDataProviderDataCollectionUpload : public UPCGDataCollectionDataProvider
+{
+	GENERATED_BODY()
+
+public:
+	//~ Begin UComputeDataProvider Interface
+	FComputeDataProviderRenderProxy* GetRenderProxy() override;
+	//~ End UComputeDataProvider Interface
+
+	TArray<FName> DownstreamInputPinLabels;
+};
+
+class FPCGDataProviderDataCollectionUploadProxy : public FPCGDataCollectionDataProviderProxy
+{
+public:
+	FPCGDataProviderDataCollectionUploadProxy(TWeakObjectPtr<UPCGDataBinding> InBinding, const FPCGDataCollectionDesc& InPinDesc, const TArray<FName>& InDownstreamInputPinLabels);
+
+	//~ Begin FComputeDataProviderRenderProxy Interface
+	void AllocateResources(FRDGBuilder& GraphBuilder, FAllocationData const& InAllocationData) override;
+	//~ End FComputeDataProviderRenderProxy Interface
+
+	TArray<FName> DownstreamInputPinLabels;
+};
diff --git a/Engine/Plugins/PCG/Source/PCG/Private/Compute/DataInterfaces/PCGDebugDataInterface.cpp b/Engine/Plugins/PCG/Source/PCG/Private/Compute/DataInterfaces/PCGDebugDataInterface.cpp
new file mode 100644
--- /dev/null
+++ b/Engine/Plugins/PCG/Source/PCG/Private/Compute/DataInterfaces/PCGDebugDataInterface.cpp
@@ -0,0 +1,146 @@
+// Copyright Epic Games, Inc. All Rights Reserved.
+
+#include "Compute/DataInterfaces/PCGDebugDataInterface.h"
+
+#include "PCGModule.h"
+
+#include "RenderGraphBuilder.h"
+#include "RenderGraphResources.h"
+#include "ShaderCompilerCore.h"
+#include "ShaderParameterMetadataBuilder.h"
+#include "ComputeFramework/ComputeMetadataBuilder.h"
+#include "ComputeFramework/ShaderParameterMetadataAllocation.h"
+#include "ComputeFramework/ShaderParamTypeDefinition.h"
+
+#include UE_INLINE_GENERATED_CPP_BY_NAME(PCGDebugDataInterface)
+
+void UPCGDebugDataInterface::GetSupportedOutputs(TArray<FShaderFunctionDefinition>& OutFunctions) const
+{
+	OutFunctions.AddDefaulted_GetRef()
+		.SetName("WriteDebugValue")
+		.AddParam(EShaderFundamentalType::Uint)
+		.AddParam(EShaderFundamentalType::Float);
+}
+
+BEGIN_SHADER_PARAMETER_STRUCT(FPCGDebugDataInterfaceParameters, )
+	SHADER_PARAMETER_RDG_BUFFER_UAV(RWByteAddressBuffer, DebugBuffer)
+END_SHADER_PARAMETER_STRUCT()
+
+void UPCGDebugDataInterface::GetShaderParameters(TCHAR const* UID, FShaderParametersMetadataBuilder& InOutBuilder, FShaderParametersMetadataAllocations& InOutAllocations) const
+{
+	InOutBuilder.AddNestedStruct<FPCGDebugDataInterfaceParameters>(UID);
+}
+
+void UPCGDebugDataInterface::GetShaderHash(FString& InOutKey) const
+{
+	// UComputeGraph::BuildKernelSource hashes the result of GetHLSL()
+	// Only append additional hashes here if the HLSL contains any additional includes	
+}
+
+void UPCGDebugDataInterface::GetHLSL(FString& OutHLSL, FString const& InDataInterfaceName) const
+{
+	const FString WriteDebugValueImpl = (DebugBufferSize > 0)
+		? FString::Format(TEXT(
+			"\n{\n"
+			"	if (Index < {0})\n"
+			"	{\n"
+			"		return {1}_DebugBuffer.Store(Index * 4, Value);\n"
+			"	}\n"
+			"}\n"),
+			{ DebugBufferSize, *InDataInterfaceName })
+		: TEXT("{ /* No-Op */ }\n");
+
+	OutHLSL += FString::Format(TEXT(
+		"RWByteAddressBuffer {0}_DebugBuffer;\n"
+		"void WriteDebugValue_{0}(uint Index, float Value) {1}"),
+		{ *InDataInterfaceName, WriteDebugValueImpl });
+}
+
+bool UPCGDebugDataInterface::GetRequiresReadback() const
+{
+	return DebugBufferSize > 0;
+}
+
+UComputeDataProvider* UPCGDebugDataInterface::CreateDataProvider(TObjectPtr<UObject> InBinding, uint64 InInputMask, uint64 InOutputMask) const
+{
+	UPCGDebugDataProvider* Provider = NewObject<UPCGDebugDataProvider>();
+	Provider->DebugBufferSize = DebugBufferSize;
+	return Provider;
+}
+
+FComputeDataProviderRenderProxy* UPCGDebugDataProvider::GetRenderProxy()
+{
+	TWeakObjectPtr<UPCGDebugDataProvider> ThisWeakPtr(this);
+
+	auto ProcessReadbackData_RenderThread = [ThisWeakPtr](const void* InData, int InNumBytes)
+	{
+		if (!ThisWeakPtr.IsValid() || !ThisWeakPtr.Get() || !InData)
+		{
+			return;
+		}
+
+		check(ThisWeakPtr->DebugBufferSize * sizeof(float) == InNumBytes);
+
+		UPCGDebugDataProvider* DataProvider = ThisWeakPtr.Get();
+		
+		const float* DataAsFloat = static_cast<const float*>(InData);
+
+		UE_LOG(LogPCG, Warning, TEXT("Logging kernel debug..."));
+
+		for (int Index = 0; Index < InNumBytes / sizeof(float); ++Index)
+		{
+			UE_LOG(LogPCG, Warning, TEXT("\t%d:\t%f"), Index, DataAsFloat[Index]);
+		}
+
+		DataProvider->bOperationComplete = true;
+	};
+
+	return new FPCGDebugDataProviderProxy(DebugBufferSize, ProcessReadbackData_RenderThread);
+}
+
+bool FPCGDebugDataProviderProxy::IsValid(FValidationData const& InValidationData) const
+{
+	if (InValidationData.ParameterStructSize != sizeof(FParameters))
+	{
+		UE_LOG(LogPCG, Error, TEXT("Proxy invalid due to mismatching ParameterStructSize."));
+		return false;
+	}
+
+	return true;
+}
+
+void FPCGDebugDataProviderProxy::GatherDispatchData(FDispatchData const& InDispatchData)
+{
+	const TStridedView<FParameters> ParameterArray = MakeStridedParameterView<FParameters>(InDispatchData);
+	for (int32 InvocationIndex = 0; InvocationIndex < InDispatchData.NumInvocations; ++InvocationIndex)
+	{
+		FParameters& Parameters = ParameterArray[InvocationIndex];
+		Parameters.DebugBuffer = DebugBufferUAV;
+	}
+}
+
+void FPCGDebugDataProviderProxy::AllocateResources(FRDGBuilder& GraphBuilder, FAllocationData const& InAllocationData)
+{
+	if (DebugBufferSize > 0)
+	{
+		FRDGBufferDesc Desc = FRDGBufferDesc::CreateStructuredDesc(sizeof(float), DebugBufferSize);
+		Desc.Usage = EBufferUsageFlags(Desc.Usage | BUF_SourceCopy);
+		DebugBuffer = GraphBuilder.CreateBuffer(Desc, TEXT("PCGDebugBuffer"));
+		DebugBufferUAV = GraphBuilder.CreateUAV(DebugBuffer);
+
+		TArray<float> ZeroInitializedBuffer;
+		ZeroInitializedBuffer.SetNumZeroed(DebugBufferSize);
+
+		GraphBuilder.QueueBufferUpload(DebugBuffer, ZeroInitializedBuffer.GetData(), ZeroInitializedBuffer.Num() * sizeof(float), ERDGInitialDataFlags::None);
+	}
+}
+
+void FPCGDebugDataProviderProxy::GetReadbackData(TArray<FReadbackData>& OutReadbackData) const
+{
+	FReadbackData Data;
+	Data.Buffer = DebugBuffer;
+	Data.NumBytes = DebugBufferSize * sizeof(float);
+	Data.ReadbackCallback_RenderThread = &AsyncReadbackCallback_RenderThread;
+
+	OutReadbackData.Add(MoveTemp(Data));
+}
diff --git a/Engine/Plugins/PCG/Source/PCG/Private/Compute/DataInterfaces/PCGDebugDataInterface.h b/Engine/Plugins/PCG/Source/PCG/Private/Compute/DataInterfaces/PCGDebugDataInterface.h
new file mode 100644
--- /dev/null
+++ b/Engine/Plugins/PCG/Source/PCG/Private/Compute/DataInterfaces/PCGDebugDataInterface.h
@@ -0,0 +1,74 @@
+// Copyright Epic Games, Inc. All Rights Reserved.
+
+#pragma once
+
+#include "ComputeFramework/ComputeDataInterface.h"
+#include "ComputeFramework/ComputeDataProvider.h"
+
+#include "PCGDebugDataInterface.generated.h"
+
+class FPCGDebugDataInterfaceParameters;
+class FRDGBuffer;
+class FRDGBufferUAV;
+
+/** Interface for data about the kernel debug, such as debug value buffer. */
+UCLASS(ClassGroup = (Procedural))
+class PCG_API UPCGDebugDataInterface : public UComputeDataInterface
+{
+	GENERATED_BODY()
+
+public:
+	//~ Begin UComputeDataInterface Interface
+	TCHAR const* GetClassName() const override { return TEXT("KernelDebug"); }
+	void GetSupportedOutputs(TArray<FShaderFunctionDefinition>& OutFunctions) const override;
+	void GetShaderParameters(TCHAR const* UID, FShaderParametersMetadataBuilder& InOutBuilder, FShaderParametersMetadataAllocations& InOutAllocations) const override;
+	void GetShaderHash(FString& InOutKey) const override;
+	void GetHLSL(FString& OutHLSL, FString const& InDataInterfaceName) const override;
+	bool GetRequiresReadback() const override;
+	UComputeDataProvider* CreateDataProvider(TObjectPtr<UObject> InBinding, uint64 InInputMask, uint64 InOutputMask) const override;
+	//~ End UComputeDataInterface Interface
+
+	void SetDebugBufferSize(uint32 InDebugBufferSize) { DebugBufferSize = InDebugBufferSize; }
+
+private:
+	uint32 DebugBufferSize = 0;
+};
+
+/** Compute Framework Data Provider for each custom compute kernel. */
+UCLASS()
+class UPCGDebugDataProvider : public UComputeDataProvider
+{
+	GENERATED_BODY()
+
+public:
+	//~ Begin UComputeDataProvider Interface
+	FComputeDataProviderRenderProxy* GetRenderProxy() override;
+	//~ End UComputeDataProvider Interface
+
+	uint32 DebugBufferSize = 0;
+	std::atomic<bool> bOperationComplete = false;
+};
+
+class FPCGDebugDataProviderProxy : public FComputeDataProviderRenderProxy
+{
+public:
+	FPCGDebugDataProviderProxy(uint32 InDebugBufferSize, FReadbackCallback InAsyncReadbackCallback_RenderThread)
+		: DebugBufferSize(InDebugBufferSize)
+		, AsyncReadbackCallback_RenderThread(InAsyncReadbackCallback_RenderThread)
+	{}
+
+	//~ Begin FComputeDataProviderRenderProxy Interface
+	bool IsValid(FValidationData const& InValidationData) const override;
+	void GatherDispatchData(FDispatchData const& InDispatchData) override;
+	void AllocateResources(FRDGBuilder& GraphBuilder, FAllocationData const& InAllocationData) override;
+	void GetReadbackData(TArray<FReadbackData>& OutReadbackData) const override;
+	//~ End FComputeDataProviderRenderProxy Interface
+
+protected:
+	using FParameters = FPCGDebugDataInterfaceParameters;
+
+	FRDGBufferRef DebugBuffer = nullptr;
+	FRDGBufferUAVRef DebugBufferUAV = nullptr;
+	uint32 DebugBufferSize = 0;
+	FReadbackCallback AsyncReadbackCallback_RenderThread;
+};
diff --git a/Engine/Plugins/PCG/Source/PCG/Private/Compute/DataInterfaces/PCGLandscapeDataInterface.cpp b/Engine/Plugins/PCG/Source/PCG/Private/Compute/DataInterfaces/PCGLandscapeDataInterface.cpp
new file mode 100644
--- /dev/null
+++ b/Engine/Plugins/PCG/Source/PCG/Private/Compute/DataInterfaces/PCGLandscapeDataInterface.cpp
@@ -0,0 +1,439 @@
+// Copyright Epic Games, Inc. All Rights Reserved.
+
+#include "PCGLandscapeDataInterface.h"
+
+#include "PCGComponent.h"
+#include "PCGData.h"
+#include "Compute/PCGDataBinding.h"
+#include "Data/PCGLandscapeData.h"
+
+#include "EngineUtils.h"
+#include "GlobalRenderResources.h"
+#include "Landscape.h"
+#include "LandscapeHeightfieldCollisionComponent.h"
+#include "LandscapeInfo.h"
+#include "RenderGraphBuilder.h"
+#include "RenderResource.h"
+#include "RHIStaticStates.h"
+#include "RHIUtilities.h"
+#include "ShaderCompilerCore.h"
+#include "ShaderCore.h"
+#include "ShaderParameterMetadataBuilder.h"
+#include "ComputeFramework/ShaderParamTypeDefinition.h"
+#include "Containers/ResourceArray.h"
+#include "Misc/LargeWorldRenderPosition.h"
+
+#include UE_INLINE_GENERATED_CPP_BY_NAME(PCGLandscapeDataInterface)
+
+#define LOCTEXT_NAMESPACE "UPCGLandscapeDataInterface"
+
+/** RenderResource used to hold textures generated by this DI, pulled from the collision geometry of the terrain. */
+class FPCGLandscapeTextureResource : public FRenderResource
+{
+public:
+	FPCGLandscapeTextureResource(const FIntPoint& InCellCount): CellCount(InCellCount) {}
+	virtual ~FPCGLandscapeTextureResource() = default;
+
+	FPCGLandscapeTextureResource() = delete;
+	FPCGLandscapeTextureResource(const FPCGLandscapeTextureResource&) = delete;
+	FPCGLandscapeTextureResource(const FPCGLandscapeTextureResource&&) = delete;
+
+	virtual void InitRHI(FRHICommandListBase& RHICmdList) override
+	{
+		if (HeightBulkData.DataSize > 0)
+		{
+			HeightTexture.Initialize(TEXT("FLandscapeTextureResource_HeightTexture"), sizeof(float), CellCount.X, CellCount.Y, EPixelFormat::PF_R32_FLOAT, FTextureReadBuffer2D::DefaultTextureInitFlag, &HeightBulkData);
+		}
+
+		ReleaseSourceData();
+	}
+
+	virtual void ReleaseRHI() override { HeightTexture.Release(); }
+	void ReleaseSourceData() { HeightValues.Empty(); }
+	FRHITexture* GetHeightTexture() const { return HeightTexture.Buffer; }
+	FIntPoint GetDimensions() const { return CellCount; }
+
+	TArray<float>& EditHeightValues(int32 SampleCount)
+	{
+		const float DefaultHeight = 0.0f;
+		HeightValues.Reset(SampleCount);
+		HeightValues.Init(DefaultHeight, SampleCount);
+		HeightBulkData.Init(HeightValues.GetData(), HeightValues.Num() * HeightValues.GetTypeSize());
+
+		return HeightValues;
+	}
+
+private:
+	FTextureReadBuffer2D HeightTexture;
+	FIntPoint CellCount;
+
+	TArray<float> HeightValues;
+
+	struct FTextureBulkData : public FResourceBulkDataInterface
+	{
+		void Init(const void* InData, uint32 InDataSize) { Data = InData; DataSize = InDataSize; }
+		void Clear() { Data = nullptr; DataSize = 0; }
+		virtual const void* GetResourceBulkData() const override { return Data; }
+		virtual uint32 GetResourceBulkDataSize() const override { return DataSize; }
+		virtual void Discard() override { }
+		const void* Data = nullptr;
+		uint32 DataSize = 0;
+	};
+
+	FTextureBulkData* GetHeightBulkData() { return HeightBulkData.DataSize > 0 ? &HeightBulkData : nullptr; }
+	FTextureBulkData HeightBulkData;
+};
+
+FPCGLandscapeResource::FPCGLandscapeResource(const FResourceKey& InKey)
+	: ResourceKey(InKey)
+{
+	const ULandscapeInfo* LandscapeInfo = ResourceKey.Source->GetLandscapeInfo();
+	if (!LandscapeInfo)
+	{
+		return;
+	}
+
+	const int32 ComponentQuadCount = ResourceKey.Source->ComponentSizeQuads;
+	const FIntPoint RegionSpan = ResourceKey.MaxCaptureRegion - ResourceKey.MinCaptureRegion + FIntPoint(1, 1);
+	const FIntPoint CaptureQuadSpan = RegionSpan * ComponentQuadCount;
+	const FIntPoint CaptureVertexSpan = CaptureQuadSpan + FIntPoint(1, 1);
+	const int32 SampleCount = CaptureVertexSpan.X * CaptureVertexSpan.Y;
+
+	LandscapeTextures = new FPCGLandscapeTextureResource(CaptureVertexSpan);
+
+	TArray<float>* HeightValues = ResourceKey.bIncludesHeight ? &LandscapeTextures->EditHeightValues(SampleCount) : nullptr;
+
+	const FIntPoint RegionVertexBase = ResourceKey.MinCaptureRegion * ComponentQuadCount;
+
+	for (const FIntPoint& Region : ResourceKey.CapturedRegions)
+	{
+		auto FoundCollisionComponent = LandscapeInfo->XYtoCollisionComponentMap.Find(Region);
+		check(FoundCollisionComponent);
+
+		if (FoundCollisionComponent)
+		{
+			if (const ULandscapeHeightfieldCollisionComponent* CollisionComponent = *FoundCollisionComponent)
+			{
+				if (HeightValues)
+				{
+					const FIntPoint SectionBase = (Region - ResourceKey.MinCaptureRegion) * ComponentQuadCount;
+					CollisionComponent->FillHeightTile(*HeightValues, SectionBase.X + SectionBase.Y * CaptureVertexSpan.X, CaptureVertexSpan.X);
+				}
+			}
+		}
+	}
+
+	// Number of cells that are represented in our heights array.
+	CellCount = CaptureVertexSpan;
+
+	// Mapping to get the UV from 'cell space' which is relative to the entire terrain (not just the captured regions).
+	FVector2D UvScale(1.0f / CaptureVertexSpan.X, 1.0f / CaptureVertexSpan.Y);
+
+	UvScaleBias = FVector4(
+		UvScale.X,
+		UvScale.Y,
+		(0.5f - RegionVertexBase.X) * UvScale.X,
+		(0.5f - RegionVertexBase.Y) * UvScale.Y);
+
+	FTransform LandscapeTransform = ResourceKey.Source->GetTransform();
+	FLargeWorldRenderPosition LandscapeTransformOrigin(LandscapeTransform.GetLocation());
+
+	LandscapeLWCTile = LandscapeTransformOrigin.GetTile();
+	LandscapeTransform.SetLocation(FVector(LandscapeTransformOrigin.GetOffset()));
+
+	ActorToWorldTransform = LandscapeTransform.ToMatrixWithScale();
+	WorldToActorTransform = ActorToWorldTransform.Inverse();
+	TextureWorldGridSize = FVector2D(ResourceKey.Source->GetTransform().GetScale3D());
+}
+
+void FPCGLandscapeResource::Release()
+{
+	if (LandscapeTextures)
+	{
+		ENQUEUE_RENDER_COMMAND(BeginDestroyCommand)([RT_Resource = LandscapeTextures](FRHICommandListImmediate& RHICmdList)
+		{
+			check(RT_Resource);
+			RT_Resource->ReleaseResource();
+
+			// On some RHIs textures will push data on the RHI thread
+			// Therefore we are not 'released' until the RHI thread has processed all commands
+			RHICmdList.EnqueueLambda([RT_Resource](FRHICommandListImmediate& RHICmdList)
+			{
+				delete RT_Resource;
+			});
+		});
+
+		LandscapeTextures = nullptr;
+	}
+}
+
+void UPCGLandscapeDataInterface::GetSupportedInputs(TArray<FShaderFunctionDefinition>& OutFunctions) const
+{
+	OutFunctions.AddDefaulted_GetRef()
+		.SetName(TEXT("GetHeight"))
+		.AddReturnType(EShaderFundamentalType::Float)
+		.AddParam(EShaderFundamentalType::Float, 3);
+
+	OutFunctions.AddDefaulted_GetRef()
+		.SetName(TEXT("GetNormal"))
+		.AddReturnType(EShaderFundamentalType::Float, 3)
+		.AddParam(EShaderFundamentalType::Float, 3);
+
+	// TODO: Color, PhysicalMaterialIndex, etc.
+}
+
+BEGIN_SHADER_PARAMETER_STRUCT(FPCGLandscapeDataInterfaceParameters,)
+	SHADER_PARAMETER_TEXTURE(Texture2D, HeightTexture)
+	SHADER_PARAMETER_SAMPLER(SamplerState, HeightTextureSampler)
+	SHADER_PARAMETER(FVector3f, HeightTextureLWCTile)
+	SHADER_PARAMETER(FMatrix44f, HeightTextureWorldToUvTransform)
+	SHADER_PARAMETER(FVector4f, HeightTextureUvScaleBias)
+	SHADER_PARAMETER(FVector2f, HeightTextureWorldGridSize)
+	SHADER_PARAMETER(FVector3f, SystemLWCTile)
+END_SHADER_PARAMETER_STRUCT()
+
+void UPCGLandscapeDataInterface::GetShaderParameters(TCHAR const* UID, FShaderParametersMetadataBuilder& InOutBuilder, FShaderParametersMetadataAllocations& InOutAllocations) const
+{
+	InOutBuilder.AddNestedStruct<FPCGLandscapeDataInterfaceParameters>(UID);
+}
+
+TCHAR const* UPCGLandscapeDataInterface::TemplateFilePath = TEXT("/Plugin/PCG/Private/PCGLandscapeDataInterface.ush");
+
+TCHAR const* UPCGLandscapeDataInterface::GetShaderVirtualPath() const
+{
+	return TemplateFilePath;
+}
+
+void UPCGLandscapeDataInterface::GetShaderHash(FString& InOutKey) const
+{
+	GetShaderFileHash(TemplateFilePath, EShaderPlatform::SP_PCD3D_SM5).AppendString(InOutKey);
+}
+
+void UPCGLandscapeDataInterface::GetHLSL(FString& OutHLSL, FString const& InDataInterfaceName) const
+{
+	TMap<FString, FStringFormatArg> TemplateArgs =
+	{
+		{ TEXT("DataInterfaceName"), InDataInterfaceName },
+	};
+
+	FString TemplateFile;
+	LoadShaderSourceFile(TemplateFilePath, EShaderPlatform::SP_PCD3D_SM5, &TemplateFile, nullptr);
+	OutHLSL += FString::Format(*TemplateFile, TemplateArgs);
+}
+
+UComputeDataProvider* UPCGLandscapeDataInterface::CreateDataProvider(TObjectPtr<UObject> InBinding, uint64 InInputMask, uint64 InOutputMask) const
+{
+	UPCGDataBinding* Binding = CastChecked<UPCGDataBinding>(InBinding);
+	check(Binding->SourceComponent.IsValid() && Binding->SourceComponent.Get());
+
+	const UPCGComponent* SourceComponent = Binding->SourceComponent.Get();
+	const UWorld* World = SourceComponent->GetWorld();
+	const FBox ComponentBounds = SourceComponent->GetGridBounds();
+
+	auto TestLandscape = [World, &ComponentBounds](const ALandscape* InLandscape)
+	{
+		if (InLandscape->GetWorld() == World)
+		{
+			if (const ULandscapeInfo* LandscapeInfo = InLandscape->GetLandscapeInfo())
+			{
+				for (const auto& ComponentIt : LandscapeInfo->XYtoCollisionComponentMap)
+				{
+					if (ComponentBounds.IntersectXY(ComponentIt.Value->Bounds.GetBox()))
+					{
+						return true;
+					}
+				}
+			}
+		}
+
+		return false;
+	};
+
+	TObjectPtr<UPCGLandscapeDataProvider> DataProvider = NewObject<UPCGLandscapeDataProvider>();
+
+	// Take any input pin label alias to obtain the data from the input data collection.
+	check(!DownstreamInputPinLabelAliases.IsEmpty());
+	const TArray<FPCGTaggedData> TaggedDatas = Binding->DataForGPU.InputDataCollection.GetInputsByPin(DownstreamInputPinLabelAliases[0]);
+
+	if (!TaggedDatas.IsEmpty())
+	{
+		ensure(TaggedDatas.Num() == 1); // There should only be one landscape data
+
+		const UPCGLandscapeData* LandscapeData = CastChecked<UPCGLandscapeData>(TaggedDatas[0].Data);
+
+		for (TSoftObjectPtr<ALandscapeProxy> LandscapeProxyPtr : LandscapeData->Landscapes)
+		{
+			if (ALandscapeProxy* LandscapeProxy = LandscapeProxyPtr.Get())
+			{
+				if (ALandscape* Landscape = LandscapeProxy->GetLandscapeActor())
+				{
+					if (TestLandscape(Landscape))
+					{
+						DataProvider->Initialize(Landscape, ComponentBounds);
+						break;
+					}
+				}
+			}
+		}
+	}
+
+	return DataProvider;
+}
+
+void UPCGLandscapeDataProvider::Initialize(ALandscape* InLandscape, const FBox& Bounds)
+{
+	check(IsInGameThread());
+
+	const ULandscapeInfo* LandscapeInfo = InLandscape ? InLandscape->GetLandscapeInfo() : nullptr;
+
+	if (!ensure(LandscapeInfo))
+	{
+		return;
+	}
+
+	// We want to use the bounds of the system to figure out which cells of the landscape that we need to handle.
+	const int32 MaxLandscapeRegionCount = LandscapeInfo->XYtoCollisionComponentMap.Num();
+
+	const FVector3f LWCTile(0, 0, 0); // Used to offset in LWC for precision, default to 0-vector for now.
+	const FVector LWCTileOffset = FVector(LWCTile) * FLargeWorldRenderScalar::GetTileSize();
+	FBox SystemWorldBounds = Bounds;
+	SystemWorldBounds.Min += LWCTileOffset;
+	SystemWorldBounds.Max += LWCTileOffset;
+
+	const FTransform& LandscapeActorToWorld = InLandscape->LandscapeActorToWorld();
+	const FVector SystemMinInLandscape = LandscapeActorToWorld.InverseTransformPosition(SystemWorldBounds.Min);
+	const FVector SystemMaxInLandscape = LandscapeActorToWorld.InverseTransformPosition(SystemWorldBounds.Max);
+
+	FBox SystemBoundsInLandscape(
+		SystemMinInLandscape.ComponentMin(SystemMaxInLandscape),
+		SystemMinInLandscape.ComponentMax(SystemMaxInLandscape));
+
+	// Transform the above box into a range of integers covering the cells of the landscape first clamp it at 0.
+	SystemBoundsInLandscape.Min = SystemBoundsInLandscape.Min.ComponentMax(FVector::ZeroVector);
+	SystemBoundsInLandscape.Max = SystemBoundsInLandscape.Max.ComponentMax(FVector::ZeroVector);
+
+	// Next rescale based on the quad size.
+	const double QuadSizeScaleFactor = 1.0 / ((double)InLandscape->ComponentSizeQuads);
+	SystemBoundsInLandscape.Min *= QuadSizeScaleFactor;
+	SystemBoundsInLandscape.Max *= QuadSizeScaleFactor;
+
+	// Truncate to integers.
+	const FVector MaxIntValue = FVector(double(TNumericLimits<int32>::Max()));
+	SystemBoundsInLandscape.Min = SystemBoundsInLandscape.Min.ComponentMin(MaxIntValue);
+	SystemBoundsInLandscape.Max = SystemBoundsInLandscape.Max.ComponentMin(MaxIntValue);
+
+	const FIntRect SystemRect = FIntRect(
+		FIntPoint(FMath::FloorToInt(SystemBoundsInLandscape.Min.X), FMath::FloorToInt(SystemBoundsInLandscape.Min.Y)),
+		FIntPoint(FMath::CeilToInt(SystemBoundsInLandscape.Max.X), FMath::CeilToInt(SystemBoundsInLandscape.Max.Y)));
+
+	// For obnoxiously large system bounds we need to guard against potential overflow on the number of cells.
+	const int32 MaxSystemWidth = FMath::Clamp<int32>(SystemRect.Max.X - SystemRect.Min.X, 0, MaxLandscapeRegionCount);
+	const int32 MaxSystemHeight = FMath::Clamp<int32>(SystemRect.Max.Y - SystemRect.Min.Y, 0, MaxLandscapeRegionCount);
+
+	const int64 MaxSystemRegionCount64 = int64(MaxSystemWidth) * int64(MaxSystemHeight);
+	const int32 MaxSystemRegionCount = int32(FMath::Min<int64>(int64(TNumericLimits<int32>::Max()), MaxSystemRegionCount64));
+
+	const int32 MaxRegionCount = FMath::Min(MaxSystemRegionCount, MaxLandscapeRegionCount);
+
+	FPCGLandscapeResource::FResourceKey Key;
+	Key.Source = InLandscape;
+	Key.CapturedRegions.Reserve(MaxRegionCount);
+	Key.MinCaptureRegion = FIntPoint(TNumericLimits<int32>::Max(), TNumericLimits<int32>::Max());
+	Key.MaxCaptureRegion = FIntPoint(TNumericLimits<int32>::Min(), TNumericLimits<int32>::Min());
+	Key.bIncludesHeight = true;
+
+	auto AddRegion = [&](const FIntPoint& Region)
+	{
+		Key.CapturedRegions.Add(Region);
+		Key.MinCaptureRegion = Key.MinCaptureRegion.ComponentMin(Region);
+		Key.MaxCaptureRegion = Key.MaxCaptureRegion.ComponentMax(Region);
+	};
+
+	if (MaxSystemRegionCount > MaxLandscapeRegionCount)
+	{
+		for (const auto& LandscapeComponent : LandscapeInfo->XYtoCollisionComponentMap)
+		{
+			if (SystemRect.Contains(LandscapeComponent.Key))
+			{
+				AddRegion(LandscapeComponent.Key);
+			}
+		}
+	}
+	else
+	{
+		for (int32 GridY = SystemRect.Min.Y; GridY < SystemRect.Max.Y; ++GridY)
+		{
+			for (int32 GridX = SystemRect.Min.X; GridX < SystemRect.Max.X; ++GridX)
+			{
+				const FIntPoint CurrentRegion(GridX, GridY);
+				if (LandscapeInfo->XYtoCollisionComponentMap.Contains(CurrentRegion))
+				{
+					AddRegion(CurrentRegion);
+				}
+			}
+		}
+	}
+
+	Resource = FPCGLandscapeResource(Key);
+}
+
+UPCGLandscapeDataProvider::~UPCGLandscapeDataProvider()
+{
+	// Implementation note: Releasing the resource here assumes the Data Provider will always outlive the DataProviderProxies.
+	Resource.Release();
+}
+
+FComputeDataProviderRenderProxy* UPCGLandscapeDataProvider::GetRenderProxy()
+{
+	FPCGLandscapeDataProviderProxy* Proxy = new FPCGLandscapeDataProviderProxy(Resource);
+	return Proxy;
+}
+
+bool FPCGLandscapeDataProviderProxy::IsValid(FValidationData const& InValidationData) const
+{
+	if (InValidationData.ParameterStructSize != sizeof(FParameters))
+	{
+		return false;
+	}
+
+	return true;
+}
+
+void FPCGLandscapeDataProviderProxy::GatherDispatchData(FDispatchData const& InDispatchData)
+{
+	const TStridedView<FParameters> ParameterArray = MakeStridedParameterView<FParameters>(InDispatchData);
+	for (int32 InvocationIndex = 0; InvocationIndex < ParameterArray.Num(); ++InvocationIndex)
+	{
+		FParameters& Parameters = ParameterArray[InvocationIndex];
+
+		// Set Samplers
+		FRHISamplerState* BilinearSamplerState = TStaticSamplerState<SF_Bilinear, AM_Clamp, AM_Clamp, AM_Clamp>::GetRHI();
+		FRHISamplerState* PointClampedSampler = TStaticSamplerState<SF_Point, AM_Clamp, AM_Clamp, AM_Clamp>::GetRHI();
+		Parameters.HeightTextureSampler = RHIPixelFormatHasCapabilities(EPixelFormat::PF_R32_FLOAT, EPixelFormatCapabilities::TextureFilterable) ? BilinearSamplerState : PointClampedSampler;
+
+		// Set Textures
+		FRHITexture* HeightTexture = Resource.LandscapeTextures ? Resource.LandscapeTextures->GetHeightTexture() : (FRHITexture*)GBlackTexture->TextureRHI;
+		check(HeightTexture);
+
+		// TODO: Bindless resources to handle multiple landscapes
+		Parameters.HeightTexture = HeightTexture;
+		Parameters.HeightTextureLWCTile = Resource.LandscapeLWCTile;
+		Parameters.HeightTextureWorldToUvTransform = (FMatrix44f)Resource.WorldToActorTransform;
+		Parameters.HeightTextureUvScaleBias = (FVector4f)Resource.UvScaleBias;
+		Parameters.HeightTextureWorldGridSize = (FVector2f)Resource.TextureWorldGridSize;
+
+		// System tile for LWC
+		FVector3f LWCTile(0, 0, 0); // Used to offset in LWC for precision, default to 0-vector for now.
+		Parameters.SystemLWCTile = LWCTile;
+	}
+}
+
+void FPCGLandscapeDataProviderProxy::AllocateResources(FRDGBuilder& GraphBuilder, FAllocationData const& InAllocationData)
+{
+	if (Resource.LandscapeTextures)
+	{
+		Resource.LandscapeTextures->InitResource(GraphBuilder.RHICmdList);
+	}
+}
+
+#undef LOCTEXT_NAMESPACE
diff --git a/Engine/Plugins/PCG/Source/PCG/Private/Compute/DataInterfaces/PCGLandscapeDataInterface.h b/Engine/Plugins/PCG/Source/PCG/Private/Compute/DataInterfaces/PCGLandscapeDataInterface.h
new file mode 100644
--- /dev/null
+++ b/Engine/Plugins/PCG/Source/PCG/Private/Compute/DataInterfaces/PCGLandscapeDataInterface.h
@@ -0,0 +1,101 @@
+// Copyright Epic Games, Inc. All Rights Reserved.
+
+#pragma once
+
+#include "Compute/DataInterfaces/PCGComputeDataInterface.h"
+#include "ComputeFramework/ComputeDataProvider.h"
+
+#include "PCGLandscapeDataInterface.generated.h"
+
+class ALandscape;
+class FPCGLandscapeDataInterfaceParameters;
+class FPCGLandscapeTextureResource;
+
+/** Manages the resources created by this DI. */
+class FPCGLandscapeResource
+{
+public:
+	struct FResourceKey
+	{
+		TWeakObjectPtr<const ALandscape> Source = nullptr;
+		TArray<FIntPoint> CapturedRegions;
+		FIntPoint MinCaptureRegion = FIntPoint(ForceInitToZero);
+		FIntPoint MaxCaptureRegion = FIntPoint(ForceInitToZero);
+		bool bIncludesHeight = false;
+	};
+
+	FPCGLandscapeResource() {}
+	FPCGLandscapeResource(const FResourceKey& InKey);
+
+	/** Release the LandscapeTextures resource handle on the RHI and clear the pointer. */
+	void Release();
+
+	FPCGLandscapeTextureResource* LandscapeTextures = nullptr;
+	FVector3f LandscapeLWCTile = FVector3f::ZeroVector;
+	FMatrix ActorToWorldTransform = FMatrix::Identity;
+	FMatrix WorldToActorTransform = FMatrix::Identity;
+	FVector4 UvScaleBias = FVector4(1.0f, 1.0f, 0.0f, 0.0f);
+	FIntPoint CellCount = FIntPoint(ForceInitToZero);
+	FVector2D TextureWorldGridSize = FVector2D(1.0f, 1.0f);
+
+private:
+	FResourceKey ResourceKey;
+};
+
+/** Data Interface allowing sampling of a Landscape. */
+UCLASS(ClassGroup = (Procedural))
+class UPCGLandscapeDataInterface : public UPCGComputeDataInterface
+{
+	GENERATED_BODY()
+
+public:
+	//~ Begin UComputeDataInterface Interface
+	TCHAR const* GetClassName() const override { return TEXT("PCGLandscape"); }
+	void GetSupportedInputs(TArray<FShaderFunctionDefinition>& OutFunctions) const override;
+	void GetShaderParameters(TCHAR const* UID, FShaderParametersMetadataBuilder& InOutBuilder, FShaderParametersMetadataAllocations& InOutAllocations) const override;
+	TCHAR const* GetShaderVirtualPath() const override;
+	void GetShaderHash(FString& InOutKey) const override;
+	void GetHLSL(FString& OutHLSL, FString const& InDataInterfaceName) const override;
+	UComputeDataProvider* CreateDataProvider(TObjectPtr<UObject> InBinding, uint64 InInputMask, uint64 InOutputMask) const override;
+	//~ End UComputeDataInterface Interface
+
+private:
+	static TCHAR const* TemplateFilePath;
+};
+
+/** Compute Framework Data Provider for reading a landscape. */
+UCLASS()
+class UPCGLandscapeDataProvider : public UComputeDataProvider
+{
+	GENERATED_BODY()
+
+public:
+	void Initialize(ALandscape* InLandscape, const FBox& Bounds);
+	~UPCGLandscapeDataProvider();
+
+	//~ Begin UComputeDataProvider Interface
+	FComputeDataProviderRenderProxy* GetRenderProxy() override;
+	//~ End UComputeDataProvider Interface
+
+private:
+	FPCGLandscapeResource Resource;
+};
+
+class FPCGLandscapeDataProviderProxy : public FComputeDataProviderRenderProxy
+{
+public:
+	FPCGLandscapeDataProviderProxy(const FPCGLandscapeResource& InResource)
+		: Resource(InResource)
+	{}
+
+	//~ Begin FComputeDataProviderRenderProxy Interface
+	bool IsValid(FValidationData const& InValidationData) const override;
+	void GatherDispatchData(FDispatchData const& InDispatchData) override;
+	void AllocateResources(FRDGBuilder& GraphBuilder, FAllocationData const& InAllocationData) override;
+	//~ End FComputeDataProviderRenderProxy Interface
+
+protected:
+	using FParameters = FPCGLandscapeDataInterfaceParameters;
+
+	FPCGLandscapeResource Resource;
+};
diff --git a/Engine/Plugins/PCG/Source/PCG/Private/Compute/DataInterfaces/PCGTextureDataInterface.cpp b/Engine/Plugins/PCG/Source/PCG/Private/Compute/DataInterfaces/PCGTextureDataInterface.cpp
new file mode 100644
--- /dev/null
+++ b/Engine/Plugins/PCG/Source/PCG/Private/Compute/DataInterfaces/PCGTextureDataInterface.cpp
@@ -0,0 +1,129 @@
+// Copyright Epic Games, Inc. All Rights Reserved.
+
+#include "PCGTextureDataInterface.h"
+
+#include "PCGData.h"
+#include "Compute/PCGDataBinding.h"
+#include "Data/PCGTextureData.h"
+
+#include "GlobalRenderResources.h"
+#include "RHIResources.h"
+#include "RHIStaticStates.h"
+#include "ShaderCompilerCore.h"
+#include "ShaderCore.h"
+#include "ShaderParameterMetadataBuilder.h"
+#include "TextureResource.h"
+#include "ComputeFramework/ShaderParamTypeDefinition.h"
+#include "Engine/Texture.h"
+
+#include UE_INLINE_GENERATED_CPP_BY_NAME(PCGTextureDataInterface)
+
+void UPCGTextureDataInterface::GetSupportedInputs(TArray<FShaderFunctionDefinition>& OutFunctions) const
+{
+	OutFunctions.AddDefaulted_GetRef()
+		.SetName(TEXT("GetTexCoords"))
+		.AddReturnType(EShaderFundamentalType::Float, 2)
+		.AddParam(EShaderFundamentalType::Float, 2)
+		.AddParam(EShaderFundamentalType::Float, 2)
+		.AddParam(EShaderFundamentalType::Float, 2);
+
+	OutFunctions.AddDefaulted_GetRef()
+		.SetName(TEXT("Sample"))
+		.AddReturnType(EShaderFundamentalType::Float, 4)
+		.AddParam(EShaderFundamentalType::Float, 2);
+}
+
+BEGIN_SHADER_PARAMETER_STRUCT(FPCGTextureDataInterfaceParameters,)
+	SHADER_PARAMETER_TEXTURE(Texture2D, Texture)
+	SHADER_PARAMETER_SAMPLER(SamplerState, Sampler)
+END_SHADER_PARAMETER_STRUCT()
+
+void UPCGTextureDataInterface::GetShaderParameters(TCHAR const* UID, FShaderParametersMetadataBuilder& InOutBuilder, FShaderParametersMetadataAllocations& InOutAllocations) const
+{
+	InOutBuilder.AddNestedStruct<FPCGTextureDataInterfaceParameters>(UID);
+}
+
+TCHAR const* UPCGTextureDataInterface::TemplateFilePath = TEXT("/Plugin/PCG/Private/PCGTextureDataInterface.ush");
+
+TCHAR const* UPCGTextureDataInterface::GetShaderVirtualPath() const
+{
+	return TemplateFilePath;
+}
+
+void UPCGTextureDataInterface::GetShaderHash(FString& InOutKey) const
+{
+	GetShaderFileHash(TemplateFilePath, EShaderPlatform::SP_PCD3D_SM5).AppendString(InOutKey);
+}
+
+void UPCGTextureDataInterface::GetHLSL(FString& OutHLSL, FString const& InDataInterfaceName) const
+{
+	TMap<FString, FStringFormatArg> TemplateArgs =
+	{
+		{ TEXT("DataInterfaceName"), InDataInterfaceName },
+	};
+
+	FString TemplateFile;
+	LoadShaderSourceFile(TemplateFilePath, EShaderPlatform::SP_PCD3D_SM5, &TemplateFile, nullptr);
+	OutHLSL += FString::Format(*TemplateFile, TemplateArgs);
+}
+
+UComputeDataProvider* UPCGTextureDataInterface::CreateDataProvider(TObjectPtr<UObject> InBinding, uint64 InInputMask, uint64 InOutputMask) const
+{
+	UPCGDataBinding* Binding = CastChecked<UPCGDataBinding>(InBinding);
+
+	TObjectPtr<UPCGTextureDataProvider> DataProvider = NewObject<UPCGTextureDataProvider>();
+
+	// Take any input pin label alias to obtain the data from the input data collection.
+	check(!DownstreamInputPinLabelAliases.IsEmpty());
+	const TArray<FPCGTaggedData> TaggedDatas = Binding->DataForGPU.InputDataCollection.GetInputsByPin(DownstreamInputPinLabelAliases[0]);
+
+	if (!TaggedDatas.IsEmpty())
+	{
+		ensure(TaggedDatas.Num() == 1); // There should only be one texture data
+
+		const UPCGTextureData* TextureData = CastChecked<UPCGTextureData>(TaggedDatas[0].Data);
+		const UTexture* Texture = TextureData->Texture.IsValid() ? TextureData->Texture.Get() : nullptr;
+
+		// TODO: Support texture2d array, texture3d
+		// TODO: Bindless resources to handle multiple textures
+
+		if (Texture)
+		{
+			const FTextureResource* TextureResource = Texture->GetResource();
+
+			if (TextureResource && TextureResource->TextureRHI)
+			{
+				DataProvider->Texture = TextureResource->TextureRHI;
+			}
+		}
+	}
+
+	return DataProvider;
+}
+
+FComputeDataProviderRenderProxy* UPCGTextureDataProvider::GetRenderProxy()
+{
+	FPCGTextureDataProviderProxy* Proxy = new FPCGTextureDataProviderProxy(Texture);
+	return Proxy;
+}
+
+bool FPCGTextureDataProviderProxy::IsValid(FValidationData const& InValidationData) const
+{
+	if (InValidationData.ParameterStructSize != sizeof(FParameters))
+	{
+		return false;
+	}
+
+	return true;
+}
+
+void FPCGTextureDataProviderProxy::GatherDispatchData(FDispatchData const& InDispatchData)
+{
+	const TStridedView<FParameters> ParameterArray = MakeStridedParameterView<FParameters>(InDispatchData);
+	for (int32 InvocationIndex = 0; InvocationIndex < ParameterArray.Num(); ++InvocationIndex)
+	{
+		FParameters& Parameters = ParameterArray[InvocationIndex];
+		Parameters.Texture = Texture ? (FRHITexture*)Texture : (FRHITexture*)GBlackTexture->TextureRHI;
+		Parameters.Sampler = TStaticSamplerState<SF_Point, AM_Clamp, AM_Clamp, AM_Clamp>::GetRHI();
+	}
+}
diff --git a/Engine/Plugins/PCG/Source/PCG/Private/Compute/DataInterfaces/PCGTextureDataInterface.h b/Engine/Plugins/PCG/Source/PCG/Private/Compute/DataInterfaces/PCGTextureDataInterface.h
new file mode 100644
--- /dev/null
+++ b/Engine/Plugins/PCG/Source/PCG/Private/Compute/DataInterfaces/PCGTextureDataInterface.h
@@ -0,0 +1,65 @@
+// Copyright Epic Games, Inc. All Rights Reserved.
+
+#pragma once
+
+#include "Compute/DataInterfaces/PCGComputeDataInterface.h"
+
+#include "RHIResources.h"
+#include "ComputeFramework/ComputeDataProvider.h"
+
+#include "PCGTextureDataInterface.generated.h"
+
+class FPCGTextureDataInterfaceParameters;
+
+/** Data Interface allowing sampling of a texture. */
+UCLASS(ClassGroup = (Procedural))
+class UPCGTextureDataInterface : public UPCGComputeDataInterface
+{
+	GENERATED_BODY()
+
+public:
+	//~ Begin UComputeDataInterface Interface
+	TCHAR const* GetClassName() const override { return TEXT("PCGTexture"); }
+	void GetSupportedInputs(TArray<FShaderFunctionDefinition>& OutFunctions) const override;
+	void GetShaderParameters(TCHAR const* UID, FShaderParametersMetadataBuilder& InOutBuilder, FShaderParametersMetadataAllocations& InOutAllocations) const override;
+	TCHAR const* GetShaderVirtualPath() const override;
+	void GetShaderHash(FString& InOutKey) const override;
+	void GetHLSL(FString& OutHLSL, FString const& InDataInterfaceName) const override;
+	UComputeDataProvider* CreateDataProvider(TObjectPtr<UObject> InBinding, uint64 InInputMask, uint64 InOutputMask) const override;
+	//~ End UComputeDataInterface Interface
+
+private:
+	static TCHAR const* TemplateFilePath;
+};
+
+/** Compute Framework Data Provider for reading a texture. */
+UCLASS()
+class UPCGTextureDataProvider : public UComputeDataProvider
+{
+	GENERATED_BODY()
+
+public:
+	//~ Begin UComputeDataProvider Interface
+	FComputeDataProviderRenderProxy* GetRenderProxy() override;
+	//~ End UComputeDataProvider Interface
+
+	FTextureRHIRef Texture;
+};
+
+class FPCGTextureDataProviderProxy : public FComputeDataProviderRenderProxy
+{
+public:
+	FPCGTextureDataProviderProxy(FTextureRHIRef InTexture)
+		: Texture(InTexture)
+	{}
+
+	//~ Begin FComputeDataProviderRenderProxy Interface
+	bool IsValid(FValidationData const& InValidationData) const override;
+	void GatherDispatchData(FDispatchData const& InDispatchData) override;
+	//~ End FComputeDataProviderRenderProxy Interface
+
+protected:
+	using FParameters = FPCGTextureDataInterfaceParameters;
+
+	FTextureRHIRef Texture = nullptr;
+};
diff --git a/Engine/Plugins/PCG/Source/PCG/Private/Compute/Elements/PCGComputeGraphElement.cpp b/Engine/Plugins/PCG/Source/PCG/Private/Compute/Elements/PCGComputeGraphElement.cpp
new file mode 100644
--- /dev/null
+++ b/Engine/Plugins/PCG/Source/PCG/Private/Compute/Elements/PCGComputeGraphElement.cpp
@@ -0,0 +1,308 @@
+// Copyright Epic Games, Inc. All Rights Reserved.
+
+#include "Compute/Elements/PCGComputeGraphElement.h"
+
+#include "PCGComponent.h"
+#include "PCGModule.h"
+#include "PCGSubsystem.h"
+#include "Compute/DataInterfaces/PCGDataCollectionReadbackDataInterface.h"
+
+#include "ComputeWorkerInterface.h"
+#include "ComputeFramework/ComputeFramework.h"
+#include "ComputeFramework/ComputeKernelCompileResult.h"
+#include "GameFramework/Actor.h"
+#include "Logging/LogVerbosity.h"
+
+#define LOCTEXT_NAMESPACE "PCGComputeGraphElement"
+
+void FPCGComputeGraphContext::AddExtraStructReferencedObjects(FReferenceCollector& Collector)
+{
+	Collector.AddPropertyReferences(FComputeGraphInstance::StaticStruct(), &ComputeGraphInstance);
+}
+
+#if WITH_EDITOR
+bool FPCGComputeGraphElement::operator==(const FPCGComputeGraphElement& Other) const
+{
+	// Equivalence is same compute graph.
+	// TODO: A compute graph is currently generated for every compile, so the presence of GPU nodes breaks the current
+	// change detection. We could either cache compute graphs formed by subsets of GPU nodes that have not changed, or
+	// we could do a deep equality check for compute graphs here.
+	return Graph == Other.Graph;
+}
+#endif
+
+bool FPCGComputeGraphElement::ExecuteInternal(FPCGContext* InContext) const
+{
+	TRACE_CPUPROFILER_EVENT_SCOPE(FPCGComputeGraphElement::ExecuteInternal);
+	check(InContext);
+	FPCGComputeGraphContext* Context = static_cast<FPCGComputeGraphContext*>(InContext);
+
+	if (!ensure(Graph))
+	{
+		return true;
+	}
+
+	// Drive the execution of a compute graph. The stages are commented below and numbered by the sequence in which they are executed.
+	// The sequence appears out of order as calls so that paths that are executed multiple times (like checking for completion) are as
+	// short as possible.
+
+	// 7. Execution is is complete when any async readbacks are complete.
+	if (Context->bAllAsyncOperationsDone)
+	{
+		Context->bExecutionSuccess = true;
+
+		for (UComputeDataProvider* DataProvider : Context->ComputeGraphInstance.GetDataProviders())
+		{
+			if (UPCGDataProviderDataCollectionReadback* Readback = Cast<UPCGDataProviderDataCollectionReadback>(DataProvider))
+			{
+				// Process data for all readbacks, and track whether all succeeded.
+				const bool bProcessResult = Readback->ProcessReadBackData();
+				Context->bExecutionSuccess &= bProcessResult;
+			}
+		}
+
+		// Currently we don't output anything if processing any readback data processing failed.
+		if (ensure(Context->bExecutionSuccess) && ensure(Context->DataBinding))
+		{
+			Context->OutputData = Context->DataBinding->OutputDataCollection;
+		}
+
+		return true;
+	}
+
+	auto SleepUntilNextFrame = [Context]()
+	{
+		// TODO unsafe access to raw pointer, need cancellation lambda
+		Context->bIsPaused = true;
+		Context->SourceComponent->GetSubsystem()->RegisterBeginTickAction([Context]()
+		{
+			Context->bIsPaused = false;
+		});
+	};
+
+	// 3. If still compiling, try again next frame.
+	if (Graph->HasKernelResourcesPendingShaderCompilation())
+	{
+		UE_LOG(LogPCG, Log, TEXT("Deferring until next frame as the kernel has pending shader compilations."));
+		SleepUntilNextFrame();
+		return false;
+	}
+
+	// 6. Keep waiting for execution to complete.
+	if (Context->bGraphEnqueued)
+	{
+		// Likely we need a frame to pass in order to make progress.
+		SleepUntilNextFrame();
+		return true;
+	}
+
+	// 4. Initialize and parse incoming data for data sizes, attributes, etc that will drive buffer allocations and dispatch thread counts.
+	if (!Context->DataBinding)
+	{
+		UPCGDataBinding* DataBindingObject = FPCGContext::NewObject_AnyThread<UPCGDataBinding>(Context);
+		Context->DataBinding.Reset(DataBindingObject);
+
+		DataBindingObject->SourceComponent = Context->SourceComponent;
+		DataBindingObject->Graph = Graph.Get();
+
+		FPCGDataForGPU& DataForGPU = DataBindingObject->DataForGPU;
+
+		DataForGPU.InputDataCollection = Context->InputData;
+
+		// Link each input pin to the data collection, so that data providers can find the data.
+		for (TWeakObjectPtr<const UPCGPin>& InputPinPtr : Graph->PinsReceivingDataFromCPU)
+		{
+			if (const UPCGPin* InputPin = InputPinPtr.Get())
+			{
+				DataForGPU.InputPins.Add(InputPin);
+			}
+		}
+
+		DataForGPU.InputPinLabelAliases = Graph->InputPinLabelAliases;
+
+		Context->ComputeGraphInstance.CreateDataProviders(Graph.Get(), 0, Context->DataBinding.Get());
+
+		// Register all providers running async operations. TODO review if we should have a general API like "RunsAsyncOperations()"?
+		for (UComputeDataProvider* DataProvider : Context->ComputeGraphInstance.GetDataProviders())
+		{
+			if (UPCGDataProviderDataCollectionReadback* Readback = Cast<UPCGDataProviderDataCollectionReadback>(DataProvider))
+			{
+				Context->ProvidersRunningAsyncOperations.Add(Readback);
+
+				Readback->OnReadbackComplete_RenderThread().AddLambda([Context, Readback]()
+				{
+					FWriteScopeLock Lock(Context->ProvidersRunningAsyncOperationsLock);
+
+					const bool bEmptyBefore = Context->ProvidersRunningAsyncOperations.IsEmpty();
+
+					ensure(Context->ProvidersRunningAsyncOperations.Contains(Readback));
+					Context->ProvidersRunningAsyncOperations.Remove(Readback);
+
+					if (!bEmptyBefore && Context->ProvidersRunningAsyncOperations.IsEmpty())
+					{
+						Context->bAllAsyncOperationsDone = true;
+
+						Context->bIsPaused = false;
+					}
+				});
+			}
+		}
+	}
+
+	check(Context->DataBinding && InContext->SourceComponent.Get());
+
+	// 1. Prepare render resources. In editor, this will trigger shader compile if not compiled already.
+	if (!Graph->GetRenderProxy())
+	{
+		Graph->UpdateResources();
+
+		SleepUntilNextFrame();
+		return false;
+	}
+
+	// 2. Validate compilation
+	{
+		// Add any messages that may have occurred during compilation to visual logs.
+#if WITH_EDITOR
+		LogCompilationMessages(Context);
+#endif
+
+		// If there was any error then we should abort.
+		for (const TPair<TObjectKey<const UPCGNode>, TArray<FComputeKernelCompileMessage>>& NodeAndCompileMessages : Graph->KernelToCompileMessages)
+		{
+			for (const FComputeKernelCompileMessage& Message : NodeAndCompileMessages.Get<1>())
+			{
+				if (Message.Type == FComputeKernelCompileMessage::EMessageType::Error)
+				{
+					return true;
+				}
+			}
+		}
+	}
+
+	// 5. Enqueue work to be executed when the GPU processes the current frame.
+	Context->bGraphEnqueued = Context->ComputeGraphInstance.EnqueueWork(
+		Graph.Get(),
+		InContext->SourceComponent->GetScene(),
+		ComputeTaskExecutionGroup::EndOfFrameUpdate,
+		InContext->SourceComponent->GetOwner()->GetFName(),
+		FSimpleDelegate());
+
+	if (ensure(Context->bGraphEnqueued))
+	{
+		FReadScopeLock Lock(Context->ProvidersRunningAsyncOperationsLock);
+
+		if (!Context->ProvidersRunningAsyncOperations.IsEmpty())
+		{
+			// If we're running async operations (like readbacks), go to sleep and let them wake us up later.
+			Context->bIsPaused = true;
+			return false;
+		}
+		else
+		{
+			// No operations to wait for, so signal completion.
+			Context->bExecutionSuccess = true;
+			return true;
+		}
+	}
+	else
+	{
+		PCGE_LOG(Error, GraphAndLog, LOCTEXT("EnqueueFailed", "Compute graph enqueue failed, check log for errors."));
+		ResetAsyncOperations(InContext);
+		return true;
+	}
+}
+
+void FPCGComputeGraphElement::PostExecuteInternal(FPCGContext* InContext) const
+{
+	check(InContext);
+	FPCGComputeGraphContext* Context = static_cast<FPCGComputeGraphContext*>(InContext);
+
+	if (!ensure(Context->DataBinding))
+	{
+		return;
+	}
+
+#if WITH_EDITOR
+	if (Context->bExecutionSuccess)
+	{
+		for (TWeakObjectPtr<const UPCGNode> NodePtr : Context->DataBinding->Graph->KernelToNode)
+		{
+			const UPCGNode* Node = NodePtr.Get();
+			UPCGComponent* Component = Context->SourceComponent.Get();
+			if (Component && Context->Stack && Node)
+			{
+				// TODO calling this to register that node executed. Regarding inspection data, we need to create
+				// a readback and pipe the data to here.
+				Component->StoreInspectionData(Context->Stack, Node, /*InTimer=*/nullptr, {}, {}, /*bUsedCache*/false);
+			}
+		}
+	}
+#endif
+
+	Context->DataBinding->MarkAsGarbage();
+	Context->DataBinding = nullptr;
+}
+
+void FPCGComputeGraphElement::AbortInternal(FPCGContext* InContext) const
+{
+	ResetAsyncOperations(InContext);
+}
+
+void FPCGComputeGraphElement::ResetAsyncOperations(FPCGContext* InContext) const
+{
+	if (InContext)
+	{
+		FPCGComputeGraphContext* Context = static_cast<FPCGComputeGraphContext*>(InContext);
+
+		for (UComputeDataProvider* DataProvider : Context->ComputeGraphInstance.GetDataProviders())
+		{
+			if (UPCGDataProviderDataCollectionReadback* Readback = Cast<UPCGDataProviderDataCollectionReadback>(DataProvider))
+			{
+				Readback->OnReadbackComplete_RenderThread().Clear();
+			}
+		}
+
+		FWriteScopeLock Lock(Context->ProvidersRunningAsyncOperationsLock);
+		Context->ProvidersRunningAsyncOperations.Reset();
+	}
+}
+
+#if WITH_EDITOR
+void FPCGComputeGraphElement::LogCompilationMessages(FPCGComputeGraphContext* InContext) const
+{
+	if (InContext->SourceComponent.IsValid() && InContext->Stack)
+	{
+		for (const TPair<TObjectKey<const UPCGNode>, TArray<FComputeKernelCompileMessage>>& NodeAndCompileMessages : Graph->KernelToCompileMessages)
+		{
+			for (const FComputeKernelCompileMessage& Message : NodeAndCompileMessages.Get<1>())
+			{
+				// These messages already go to log. So just pick out the warnings and errors to display on graph. Need to convert
+				// message type.
+				ELogVerbosity::Type Verbosity = ELogVerbosity::All;
+				if (Message.Type == FComputeKernelCompileMessage::EMessageType::Warning)
+				{
+					Verbosity = ELogVerbosity::Warning;
+				}
+				else if (Message.Type == FComputeKernelCompileMessage::EMessageType::Error)
+				{
+					Verbosity = ELogVerbosity::Error;
+				}
+
+				if (Verbosity < ELogVerbosity::Log)
+				{
+					if (UPCGSubsystem* Subsystem = UPCGSubsystem::GetInstance(InContext->SourceComponent->GetWorld()))
+					{
+						FPCGStack StackWithNode = *InContext->Stack;
+						StackWithNode.PushFrame(NodeAndCompileMessages.Get<0>().ResolveObjectPtr());
+
+						Subsystem->GetNodeVisualLogsMutable().Log(StackWithNode, Verbosity, FText::FromString(Message.Text));
+					}
+				}
+			}
+		}
+	}
+}
+#endif
+
+#undef LOCTEXT_NAMESPACE
diff --git a/Engine/Plugins/PCG/Source/PCG/Private/Compute/Elements/PCGComputeGraphElement.h b/Engine/Plugins/PCG/Source/PCG/Private/Compute/Elements/PCGComputeGraphElement.h
new file mode 100644
--- /dev/null
+++ b/Engine/Plugins/PCG/Source/PCG/Private/Compute/Elements/PCGComputeGraphElement.h
@@ -0,0 +1,75 @@
+// Copyright Epic Games, Inc. All Rights Reserved.
+
+#pragma once
+
+#include "PCGElement.h"
+
+#include "PCGContext.h"
+#include "Compute/Elements/PCGCustomHLSL.h"
+
+#include "Compute/PCGComputeGraph.h"
+#include "Compute/PCGDataBinding.h"
+#include "ComputeFramework/ComputeGraphInstance.h"
+
+#include "UObject/StrongObjectPtr.h"
+
+struct FPCGComputeGraphContext : public FPCGContext
+{
+protected:
+	virtual void AddExtraStructReferencedObjects(FReferenceCollector& Collector) override;
+
+public:
+	TStrongObjectPtr<UPCGDataBinding> DataBinding = nullptr;
+
+	/** Data providers created from data interfaces and data bindings. */
+	FComputeGraphInstance ComputeGraphInstance;
+
+	/** Keep track of data providers that perform async operations and require multiple frames to complete. */
+	TSet<UComputeDataProvider*> ProvidersRunningAsyncOperations;
+	mutable FRWLock ProvidersRunningAsyncOperationsLock;
+
+	/** Graph enqueued (scheduled for execution by GPU). */
+	bool bGraphEnqueued = false;
+
+	/** All async operations complete and results processed. */
+	bool bAllAsyncOperationsDone = false;
+
+	/** Graph executed successfully. */
+	bool bExecutionSuccess = false;
+};
+
+/** Executes a CF graph. Created by the compiler when collapsing GPU nodes rather than by a settings/node. */
+class FPCGComputeGraphElement : public IPCGElement
+{
+public:
+	FPCGComputeGraphElement() {}
+
+#if WITH_EDITOR
+	//~Begin IPCGElement interface
+	virtual bool IsComputeGraphElement() const override { return true; }
+	//~End IPCGElement interface
+
+	/** Return true if the elements are identical, used for change detection. */
+	bool operator==(const FPCGComputeGraphElement& Other) const;
+#endif
+
+	TStrongObjectPtr<UPCGComputeGraph> Graph = nullptr;
+
+protected:
+	virtual FPCGContext* CreateContext() override { return new FPCGComputeGraphContext(); }
+	virtual bool ExecuteInternal(FPCGContext* InContext) const override;
+	virtual void PostExecuteInternal(FPCGContext* InContext) const override;
+	virtual void AbortInternal(FPCGContext* InContext) const override;;
+
+	// The calls to initialize the compute graph are not thread safe.
+	virtual bool CanExecuteOnlyOnMainThread(FPCGContext* Context) const override { return true; }
+
+	// TODO - need to accumulate dependencies from compute graph nodes.
+	virtual bool IsCacheable(const UPCGSettings* InSettings) const override { return false; }
+
+	void ResetAsyncOperations(FPCGContext* InContext) const;
+
+#if WITH_EDITOR
+	void LogCompilationMessages(FPCGComputeGraphContext* InContext) const;
+#endif
+};
diff --git a/Engine/Plugins/PCG/Source/PCG/Private/Compute/Elements/PCGCustomHLSL.cpp b/Engine/Plugins/PCG/Source/PCG/Private/Compute/Elements/PCGCustomHLSL.cpp
new file mode 100644
--- /dev/null
+++ b/Engine/Plugins/PCG/Source/PCG/Private/Compute/Elements/PCGCustomHLSL.cpp
@@ -0,0 +1,1159 @@
+// Copyright Epic Games, Inc. All Rights Reserved.
+
+#include "Compute/Elements/PCGCustomHLSL.h"
+
+#include "PCGContext.h"
+#include "PCGEdge.h"
+#include "PCGModule.h"
+#include "PCGPoint.h"
+#include "Compute/PCGComputeCommon.h"
+#include "Compute/PCGComputeGraph.h"
+#include "Compute/PCGDataBinding.h"
+#include "Data/PCGPointData.h"
+
+#include UE_INLINE_GENERATED_CPP_BY_NAME(PCGCustomHLSL)
+
+#define LOCTEXT_NAMESPACE "PCGCustomHLSLElement"
+
+namespace PCGHLSLElement
+{
+	const FString KernelAttributeToken = TEXT("@");
+
+	void ConvertObjectPathToShaderFilePath(FString& InOutPath)
+	{
+		// Shader compiler recognizes "/Engine/Generated/..." path as special. 
+		// It doesn't validate file suffix etc.
+		InOutPath = FString::Printf(TEXT("/Engine/Generated/UObject%s.ush"), *InOutPath);
+		// Shader compilation result parsing will break if it finds ':' where it doesn't expect.
+		InOutPath.ReplaceCharInline(TEXT(':'), TEXT('@'));
+	}
+
+	int GetElementCount(const UPCGData* InData)
+	{
+		if (const UPCGPointData* PointData = Cast<UPCGPointData>(InData))
+		{
+			return PointData->GetPoints().Num();
+		}
+		else if (InData)
+		{
+			// TODO - spline points, other things with multiple elements?
+			return 1;
+		}
+		else
+		{
+			return 0;
+		}
+	}
+
+	FString GetKernelAttributeKeyAsString(const FPCGKernelAttributeKey& Key)
+	{
+		return FString::Format(TEXT("{0}{1}"), { KernelAttributeToken, Key.Name.ToString()});
+	}
+}
+
+#if WITH_EDITOR
+void UPCGCustomHLSLSettings::PostLoad()
+{
+	Super::PostLoad();
+
+	UpdatePinSettings();
+}
+
+void UPCGCustomHLSLSettings::PostInitProperties()
+{
+	Super::PostInitProperties();
+
+	UpdatePinSettings();
+
+	UpdateDeclarations();
+}
+#endif
+
+bool UPCGCustomHLSLSettings::ShouldExecuteOnGPU() const
+{
+	return IsKernelValid();
+}
+
+#if WITH_EDITOR
+void UPCGCustomHLSLSettings::PostEditChangeProperty(FPropertyChangedEvent& PropertyChangedEvent)
+{
+	// Apply any pin setup before refreshing the node.
+	UpdatePinSettings();
+
+	Super::PostEditChangeProperty(PropertyChangedEvent);
+
+	UpdateDeclarations();
+}
+#endif
+
+FPCGElementPtr UPCGCustomHLSLSettings::CreateElement() const
+{
+	return MakeShared<FPCGCustomHLSLElement>();
+}
+
+int UPCGCustomHLSLSettings::GetProcessingElemCountForInputPin(const UPCGPin* InputPin, const UPCGDataBinding* Binding) const
+{
+	check(Binding);
+	const FPCGDataForGPU& DataForGPU = Binding->DataForGPU;
+
+	// Upper bound estimate of total number of data elements expected to arrive at this pin.
+	int ProcessingElemCount = 0;
+
+	if (DataForGPU.InputPins.Contains(InputPin))
+	{
+		FName PinLabel = InputPin->Properties.Label;
+		if (const FName* PinLabelAlias = DataForGPU.InputPinLabelAliases.Find(InputPin))
+		{
+			PinLabel = *PinLabelAlias;
+		}
+
+		// Data coming straight from CPU, know its size now so count it.
+		TArray<FPCGTaggedData> DataForPin = DataForGPU.InputDataCollection.GetInputsByPin(PinLabel);
+		for (const FPCGTaggedData& Data : DataForPin)
+		{
+			ProcessingElemCount += PCGHLSLElement::GetElementCount(Data.Data);
+		}
+	}
+	else
+	{
+		// Estimate (upper bound) element count by looking at incident connections.
+		for (const UPCGEdge* Edge : InputPin->Edges)
+		{
+			// InputPin is upstream output pin.
+			const UPCGPin* UpstreamOutputPin = Edge->InputPin;
+			if (!UpstreamOutputPin)
+			{
+				continue;
+			}
+
+			const UPCGSettings* UpstreamSettings = UpstreamOutputPin->Node ? UpstreamOutputPin->Node->GetSettings() : nullptr;
+			check(UpstreamSettings);
+
+			const FPCGDataCollectionDesc PinDesc = UpstreamSettings->ComputeOutputPinDataDesc(UpstreamOutputPin, Binding);
+
+			// Only support Point data right now.
+			ProcessingElemCount += PinDesc.ComputeDataElementCount(EPCGDataType::Point);
+		}
+	}
+
+	return ProcessingElemCount;
+}
+
+int UPCGCustomHLSLSettings::ComputeKernelThreadCount(const UPCGDataBinding* Binding) const
+{
+	int ThreadCount = 0;
+
+	if (KernelType == EPCGKernelType::PointGenerator)
+	{
+		// Point generator has fixed thread count.
+		ThreadCount = PointCount;
+	}
+	else if (KernelType == EPCGKernelType::PointProcessor)
+	{
+		// Processing volume depends on data arriving on primary pin.
+		const UPCGPin* InputPin = GetPointProcessingInputPin();
+		ThreadCount = GetProcessingElemCountForInputPin(InputPin, Binding);
+	}
+	else if (KernelType == EPCGKernelType::Custom)
+	{
+		auto GetFromFirstInput = [this, Binding]()
+		{
+			const UPCGPin* InputPin = GetPointProcessingInputPin();
+			return GetProcessingElemCountForInputPin(InputPin, Binding);
+		};
+
+		auto GetFromSecondInput = [this, Binding]()
+		{
+			const UPCGPin* InputPin = GetSecondPointProcessingInputPin();
+			return GetProcessingElemCountForInputPin(InputPin, Binding);
+		};
+
+		auto GetFromFirstInputXSecondInput = [this, Binding]()
+		{
+			const UPCGPin* FirstPin = GetPointProcessingInputPin();
+			const int FirstElemCount = GetProcessingElemCountForInputPin(FirstPin, Binding);
+			const UPCGPin* SecondPin = GetSecondPointProcessingInputPin();
+			const int SecondElemCount = GetProcessingElemCountForInputPin(SecondPin, Binding);
+			return FirstElemCount * SecondElemCount;
+		};
+
+		if (DispatchThreadCount == EPCGDispatchThreadCount::FromOutput)
+		{
+			const UPCGPin* OutputPin = GetFirstOutputPin();
+
+			if (OutputPin->Properties.BufferSizeMode == EPCGPinBufferSizeMode::FixedElementCount)
+			{
+				ThreadCount = OutputPin->Properties.FixedBufferElementCount;
+			}
+			else if (OutputPin->Properties.BufferSizeMode == EPCGPinBufferSizeMode::FromFirstPin)
+			{
+				ThreadCount = GetFromFirstInput();
+			}
+			else if (OutputPin->Properties.BufferSizeMode == EPCGPinBufferSizeMode::FromSecondPin)
+			{
+				ThreadCount = GetFromSecondInput();
+			}
+			else if (OutputPin->Properties.BufferSizeMode == EPCGPinBufferSizeMode::FirstPinXSecondPin)
+			{
+				ThreadCount = GetFromFirstInputXSecondInput();
+			}
+		}
+		else if (DispatchThreadCount == EPCGDispatchThreadCount::FromFirstInput)
+		{
+			ThreadCount = GetFromFirstInput();
+		}
+		else if (DispatchThreadCount == EPCGDispatchThreadCount::FromSecondInput)
+		{
+			ThreadCount = GetFromSecondInput();
+		}
+		else if (DispatchThreadCount == EPCGDispatchThreadCount::FirstInputXSecondInput)
+		{
+			ThreadCount = GetFromFirstInputXSecondInput();
+		}
+		else if (DispatchThreadCount == EPCGDispatchThreadCount::Fixed)
+		{
+			ThreadCount = FixedThreadCount;
+		}
+	}
+	else
+	{
+		checkNoEntry();
+	}
+
+	return ThreadCount;
+}
+
+TArray<FPCGKernelAttributeKey> UPCGCustomHLSLSettings::GetKernelAttributeKeys() const
+{
+	TArray<FPCGKernelAttributeKey> Keys;
+
+	const FString SourceToParse = ShaderFunctions + ShaderSource;
+
+	int NumExistingAttributes = PCGComputeConstants::NUM_RESERVED_ATTRS;
+	int AttributeSourceIndex = SourceToParse.Find(PCGHLSLElement::KernelAttributeToken);
+
+	while (AttributeSourceIndex != INDEX_NONE)
+	{
+		const int GetStringIndex = SourceToParse.Find(TEXT("Get"), ESearchCase::CaseSensitive, ESearchDir::FromEnd, AttributeSourceIndex);
+		const int SetStringIndex = SourceToParse.Find(TEXT("Set"), ESearchCase::CaseSensitive, ESearchDir::FromEnd, AttributeSourceIndex);
+		const int StartTypeStringIndex = FMath::Max(GetStringIndex, SetStringIndex) + 3;
+		const int EndTypeStringIndex = SourceToParse.Find(TEXT("("), ESearchCase::CaseSensitive, ESearchDir::FromStart, StartTypeStringIndex);
+
+		const FString TypeString = (StartTypeStringIndex != INDEX_NONE || EndTypeStringIndex != INDEX_NONE)
+			? SourceToParse.Mid(StartTypeStringIndex, EndTypeStringIndex - StartTypeStringIndex)
+			: TEXT("");
+
+		if (TypeString.IsEmpty())
+		{
+			UE_LOG(LogPCG, Error, TEXT("Read invalid attribute type in shader source.")); 
+		}
+
+		EPCGKernelAttributeType Type = EPCGKernelAttributeType::Bool;
+
+		if (TypeString == TEXT("Bool"))
+		{
+			Type = EPCGKernelAttributeType::Bool;
+		}
+		else if (TypeString == TEXT("Int"))
+		{
+			Type = EPCGKernelAttributeType::Int;
+		}
+		else if (TypeString == TEXT("Float"))
+		{
+			Type = EPCGKernelAttributeType::Float;
+		}
+		else if (TypeString == TEXT("Float2"))
+		{
+			Type = EPCGKernelAttributeType::Float2;
+		}
+		else if (TypeString == TEXT("Float3"))
+		{
+			Type = EPCGKernelAttributeType::Float3;
+		}
+		else if (TypeString == TEXT("Float4"))
+		{
+			Type = EPCGKernelAttributeType::Float4;
+		}
+		else if (TypeString == TEXT("Rotator"))
+		{
+			Type = EPCGKernelAttributeType::Rotator;
+		}
+		else if (TypeString == TEXT("Quat"))
+		{
+			Type = EPCGKernelAttributeType::Quat;
+		}
+		else if (TypeString == TEXT("Transform"))
+		{
+			Type = EPCGKernelAttributeType::Transform;
+		}
+		else
+		{
+			UE_LOG(LogPCG, Error, TEXT("Read invalid attribute type in shader source."));
+		}
+
+		FString AttrNameString = "";
+
+		// TODO: Copy the full attribute name at once with AttrNameString.Mid()
+		for (int I = AttributeSourceIndex + 1; I < SourceToParse.Len(); ++I)
+		{
+			const char Char = SourceToParse[I];
+
+			if (isalnum(Char))
+			{
+				AttrNameString += Char;
+			}
+			else
+			{
+				break;
+			}
+		}
+
+		Keys.AddUnique({ Type, FName(*AttrNameString) });
+
+		AttributeSourceIndex = SourceToParse.Find(PCGHLSLElement::KernelAttributeToken, ESearchCase::IgnoreCase, ESearchDir::FromStart, AttributeSourceIndex + 1);
+	}
+
+	return Keys;
+}
+
+FPCGDataCollectionDesc UPCGCustomHLSLSettings::ComputeOutputPinDataDesc(const UPCGPin* OutputPin, const UPCGDataBinding* Binding) const
+{
+	check(OutputPin);
+	FPCGDataCollectionDesc PinDesc;
+
+	// The primary output pin follows any rules prescribed by kernel type.
+	if (OutputPin == GetFirstOutputPin())
+	{
+		if (KernelType == EPCGKernelType::PointProcessor)
+		{
+			// First output pin passes through first input pin.
+			if (const UPCGPin* PointProcessingInputPin = GetPointProcessingInputPin())
+			{
+				PinDesc = ComputeInputPinDataDesc(PointProcessingInputPin, Binding);
+			}
+
+			return PinDesc;
+		}
+		else if (KernelType == EPCGKernelType::PointGenerator)
+		{
+			// Generators always produce a single point data with known point count.
+			PinDesc.DataDescs.Emplace(EPCGDataType::Point, PointCount);
+
+			// TODO: Until support for 'Create Attributes' exists, PointGenerators cannot manipulate attributes.
+			return PinDesc;
+		}
+	}
+
+	// No size set by kernel, fall back to pin settings.
+	if (OutputPin->Properties.BufferSizeMode == EPCGPinBufferSizeMode::FromFirstPin)
+	{
+		if (const UPCGPin* PointProcessingInputPin = GetPointProcessingInputPin())
+		{
+			PinDesc = ComputeInputPinDataDesc(PointProcessingInputPin, Binding);
+		}
+	}
+	else if (OutputPin->Properties.BufferSizeMode == EPCGPinBufferSizeMode::FromSecondPin)
+	{
+		if (const UPCGPin* PointProcessingInputPin = GetSecondPointProcessingInputPin())
+		{
+			PinDesc = ComputeInputPinDataDesc(PointProcessingInputPin, Binding);
+		}
+	}
+	else if (OutputPin->Properties.BufferSizeMode == EPCGPinBufferSizeMode::FirstPinXSecondPin)
+	{
+		const UPCGPin* FirstPin = GetPointProcessingInputPin();
+		const UPCGPin* SecondPin = GetSecondPointProcessingInputPin();
+
+		if (FirstPin && SecondPin)
+		{
+			const FPCGDataCollectionDesc FirstPinDesc = ComputeInputPinDataDesc(FirstPin, Binding);
+			const FPCGDataCollectionDesc SecondPinDesc = ComputeInputPinDataDesc(SecondPin, Binding);
+
+			int TotalSourceElementCount = 0;
+			for (const FPCGDataDesc& DataDesc : FirstPinDesc.DataDescs)
+			{
+				if (DataDesc.Type == EPCGDataType::Point)
+				{
+					TotalSourceElementCount += DataDesc.ElementCount;
+				}
+			}
+
+			int TotalOutputElementCount = 0;
+			for (const FPCGDataDesc& DataDesc : SecondPinDesc.DataDescs)
+			{
+				if (DataDesc.Type == EPCGDataType::Point)
+				{
+					TotalOutputElementCount += DataDesc.ElementCount * TotalSourceElementCount;
+				}
+			}
+
+			PinDesc.DataDescs.Emplace(EPCGDataType::Point, TotalOutputElementCount);
+		}
+	}
+
+	return PinDesc;
+}
+
+#if WITH_EDITOR
+EPCGChangeType UPCGCustomHLSLSettings::GetChangeTypeForProperty(const FName& InPropertyName) const
+{
+	EPCGChangeType ChangeType = Super::GetChangeTypeForProperty(InPropertyName);
+
+	if (InPropertyName == GET_MEMBER_NAME_CHECKED(UPCGCustomHLSLSettings, ShaderSource)
+		|| InPropertyName == GET_MEMBER_NAME_CHECKED(UPCGCustomHLSLSettings, ShaderFunctions))
+	{
+		ChangeType |= EPCGChangeType::ShaderSource;
+	}
+
+	// Any settings change to this node could change the compute graph.
+	ChangeType |= EPCGChangeType::Structural;
+
+	return ChangeType;
+}
+#endif
+
+const UPCGPin* UPCGCustomHLSLSettings::GetPointProcessingInputPin() const
+{
+	if (const UPCGNode* Node = Cast<UPCGNode>(GetOuter()))
+	{
+		const UPCGPin* FirstPin = !Node->GetInputPins().IsEmpty() ? Node->GetInputPins()[0] : nullptr;
+		if (FirstPin && FirstPin->Properties.AllowedTypes == EPCGDataType::Point)
+		{
+			return FirstPin;
+		}
+	}
+
+	return nullptr;
+}
+
+const UPCGPin* UPCGCustomHLSLSettings::GetSecondPointProcessingInputPin() const
+{
+	if (const UPCGNode* Node = Cast<UPCGNode>(GetOuter()))
+	{
+		bool bEncounteredPointPin = false;
+		for (const UPCGPin* Pin : Node->GetInputPins())
+		{
+			if (Pin->Properties.AllowedTypes == EPCGDataType::Point)
+			{
+				if (bEncounteredPointPin)
+				{
+					return Pin;
+				}
+
+				bEncounteredPointPin = true;
+			}
+		}
+	}
+
+	return nullptr;
+}
+
+const UPCGPin* UPCGCustomHLSLSettings::GetFirstOutputPin() const
+{
+	if (const UPCGNode* Node = Cast<UPCGNode>(GetOuter()))
+	{
+		return !Node->GetOutputPins().IsEmpty() ? Node->GetOutputPins()[0] : nullptr;
+	}
+
+	return nullptr;
+}
+
+const UPCGPin* UPCGCustomHLSLSettings::GetFirstPointOutputPin() const
+{
+	if (const UPCGNode* Node = Cast<UPCGNode>(GetOuter()))
+	{
+		const UPCGPin* FirstPin = !Node->GetOutputPins().IsEmpty() ? Node->GetOutputPins()[0] : nullptr;
+		if (FirstPin && FirstPin->Properties.AllowedTypes == EPCGDataType::Point)
+		{
+			return FirstPin;
+		}
+	}
+
+	return nullptr;
+}
+
+#if WITH_EDITOR
+void UPCGCustomHLSLSettings::UpdateDeclarations()
+{
+	// Reference: UOptimusNode_CustomComputeKernel::UpdatePreamble
+	InputDeclarations.Reset();
+	OutputDeclarations.Reset();
+
+	// Add constants category
+	{
+		if (KernelType == EPCGKernelType::PointGenerator)
+		{
+			const UPCGPin* PointProcessingOutputPin = GetFirstPointOutputPin();
+			InputDeclarations += FString::Format(
+				TEXT("// Constants\n")
+				TEXT("uint PointCount = {0};\n")
+				TEXT("\n"),
+				{ PointCount });
+		}
+	}
+
+	// Add resource indexing category
+	{
+		InputDeclarations += TEXT("// Resource Indexing\n");
+		InputDeclarations += TEXT("uint ThreadIndex;\n");
+
+		if (KernelType == EPCGKernelType::PointProcessor)
+		{
+			const UPCGPin* PointProcessingInputPin = GetPointProcessingInputPin();
+			const UPCGPin* PointProcessingOutputPin = GetFirstPointOutputPin();
+
+			if (PointProcessingInputPin && PointProcessingOutputPin)
+			{
+				InputDeclarations += FString::Format(
+					TEXT("uint {0}_DataIndex;\n")
+					TEXT("uint {1}_DataIndex;\n"),
+					{ PointProcessingInputPin->Properties.Label.ToString(),  PointProcessingOutputPin->Properties.Label.ToString() });
+			}
+		}
+		else if (KernelType == EPCGKernelType::PointGenerator)
+		{
+			if (const UPCGPin* PointProcessingOutputPin = GetFirstPointOutputPin())
+			{
+				InputDeclarations += FString::Format(
+					TEXT("uint {0}_DataIndex;\n"),
+					{ PointProcessingOutputPin->Properties.Label.ToString() });
+			}
+		}
+		else if (KernelType == EPCGKernelType::Custom)
+		{
+			auto EmitGetThreadElement = [&InInputDeclarations = InputDeclarations](const FPCGPinProperties& Properties)
+			{
+				InInputDeclarations += FString::Format(TEXT(
+					"// (DataIndex, DataAddress, ElementIndex). Returns -1 if invalid!\n"
+					"uint3 {0}_GetThreadData(uint ThreadIndex);\n"),
+					{ Properties.Label.ToString() });
+			};
+
+			for (const FPCGPinProperties& Properties : InputPinProperties())
+			{
+				EmitGetThreadElement(Properties);
+			}
+
+			for (const FPCGPinProperties& Properties : OutputPinProperties())
+			{
+				EmitGetThreadElement(Properties);
+			}
+		}
+		else
+		{
+			checkNoEntry();
+		}
+
+		InputDeclarations += TEXT("uint ElementIndex;\n");
+		InputDeclarations += TEXT("uint GetNumThreads();");
+	}
+
+	// Add debug category
+	{
+		if (bPrintShaderDebugValues)
+		{
+			InputDeclarations += FString::Format(
+				TEXT("\n\n")
+				TEXT("// Debug\n")
+				TEXT("void WriteDebugValue(uint Index, float Value); // Index in [0, {0}] (set from 'Debug Buffer Size' property)"),
+				{ DebugBufferSize - 1 });
+		}
+	}
+
+	// Per-pin input category
+	{
+		TArray<FString> DataPins;
+		TArray<FString> PointDataPins;
+		TArray<FString> SplineDataPins;
+		TArray<FString> LandscapeDataPins;
+		TArray<FString> TextureDataPins;
+		TArray<FString> RawBufferDataPins;
+
+		for (const FPCGPinProperties& Pin : InputPinProperties())
+		{
+			DataPins.Add(Pin.Label.ToString());
+
+			if (!!(Pin.AllowedTypes & EPCGDataType::Point))
+			{
+				PointDataPins.Add(Pin.Label.ToString());
+			}
+
+			if (!!(Pin.AllowedTypes & EPCGDataType::Spline))
+			{
+				SplineDataPins.Add(Pin.Label.ToString());
+			}
+
+			if (!!(Pin.AllowedTypes & EPCGDataType::Landscape))
+			{
+				LandscapeDataPins.Add(Pin.Label.ToString());
+			}
+
+			if (!!(Pin.AllowedTypes & EPCGDataType::Texture))
+			{
+				TextureDataPins.Add(Pin.Label.ToString());
+			}
+		}
+
+		InputDeclarations += TEXT("\n\n### HELPER FUNCTIONS ###\n\n");
+		InputDeclarations += TEXT("float3 GetComponentBoundsMin(); // World-space\n");
+		InputDeclarations += TEXT("float3 GetComponentBoundsMax();");
+
+		InputDeclarations += TEXT("\n\n");
+		InputDeclarations += TEXT("float3 CreateGrid2D(int ElementIndex, int NumPoints, float3 Min, float3 Max);\n");
+		InputDeclarations += TEXT("float3 CreateGrid2D(int ElementIndex, int NumPoints, int NumRows, float3 Min, float3 Max);\n");
+		InputDeclarations += TEXT("float3 CreateGrid3D(int ElementIndex, int NumPoints, float3 Min, float3 Max);\n");
+		InputDeclarations += TEXT("float3 CreateGrid3D(int ElementIndex, int NumPoints, int NumRows, int NumCols, float3 Min, float3 Max);");
+
+		if (!DataPins.IsEmpty())
+		{
+			InputDeclarations += TEXT("\n\n### DATA FUNCTIONS ###\n\n");
+			InputDeclarations += TEXT("// Valid pins: ") + FString::Join(DataPins, TEXT(", ")) + TEXT("\n");
+			InputDeclarations += TEXT("// Valid types: bool, int, float, float2, float3, float4, Rotator (float3), Quat (float4), Transform (float4x4)\n");
+			InputDeclarations += TEXT("\n");
+			InputDeclarations += TEXT("uint <pin>_GetNumData();\n");
+			InputDeclarations += TEXT("uint <pin>_GetNumElements();\n");
+			InputDeclarations += TEXT("<type> <pin>_Get<type>(uint DataIndex, uint AttributeId, uint ElementIndex);");
+		}
+
+		if (!PointDataPins.IsEmpty())
+		{
+			InputDeclarations += TEXT("\n\n### POINT DATA FUNCTIONS ###\n\n");
+			InputDeclarations += TEXT("// Valid pins: ") + FString::Join(PointDataPins, TEXT(", ")) + TEXT("\n");
+			InputDeclarations += TEXT("\n");
+			InputDeclarations += TEXT("uint <pin>_GetNumPoints(uint DataIndex);\n");
+			InputDeclarations += TEXT("float3 <pin>_GetPosition(uint DataIndex, uint ElementIndex);\n");
+			InputDeclarations += TEXT("float4 <pin>_GetRotation(uint DataIndex, uint ElementIndex);\n");
+			InputDeclarations += TEXT("float3 <pin>_GetScale(uint DataIndex, uint ElementIndex);\n");
+			InputDeclarations += TEXT("float3 <pin>_GetBoundsMin(uint DataIndex, uint ElementIndex);\n");
+			InputDeclarations += TEXT("float3 <pin>_GetBoundsMax(uint DataIndex, uint ElementIndex);\n");
+			InputDeclarations += TEXT("float3 <pin>_GetColor(uint DataIndex, uint ElementIndex);\n");
+			InputDeclarations += TEXT("float <pin>_GetDensity(uint DataIndex, uint ElementIndex);\n");
+			InputDeclarations += TEXT("int <pin>_GetSeed(uint DataIndex, uint ElementIndex);\n");
+			InputDeclarations += TEXT("float <pin>_GetSteepness(uint DataIndex, uint ElementIndex);\n");
+			InputDeclarations += TEXT("float4x4 <pin>_GetPointTransform(uint DataIndex, uint ElementIndex);");
+		}
+
+		if (!LandscapeDataPins.IsEmpty())
+		{
+			InputDeclarations += TEXT("\n\n### LANDSCAPE DATA FUNCTIONS ###\n\n");
+			InputDeclarations += TEXT("// Valid pins: ") + FString::Join(LandscapeDataPins, TEXT(", ")) + TEXT("\n");
+			InputDeclarations += TEXT("\n");
+			InputDeclarations += TEXT("float <pin>_GetHeight(float3 WorldPos);\n");
+			InputDeclarations += TEXT("float3 <pin>_GetNormal(float3 WorldPos);");
+		}
+
+		if (!TextureDataPins.IsEmpty())
+		{
+			InputDeclarations += TEXT("\n\n### TEXTURE DATA FUNCTIONS ###\n\n");
+			InputDeclarations += TEXT("// Valid pins: ") + FString::Join(TextureDataPins, TEXT(", ")) + TEXT("\n");
+			InputDeclarations += TEXT("\n");
+			InputDeclarations += TEXT("float2 <pin>_GetTexCoords(float2 WorldPos, float2 Min, float2 Max);\n");
+			InputDeclarations += TEXT("float4 <pin>_Sample(float2 TexCoords);");
+		}
+
+		if (!RawBufferDataPins.IsEmpty())
+		{
+			InputDeclarations += TEXT("\n\n### BYTE ADDRESS BUFFER DATA FUNCTIONS ###\n\n");
+			InputDeclarations += TEXT("// Valid pins: ") + FString::Join(RawBufferDataPins, TEXT(", ")) + TEXT("\n");
+			InputDeclarations += TEXT("\n");
+			InputDeclarations += TEXT("uint <pin>_ReadNumValues();\n");
+			InputDeclarations += TEXT("uint <pin>_ReadValue(uint Index);");
+		}
+	}
+
+	// Per-pin output category
+	{
+		TArray<FString> DataPins;
+		TArray<FString> PointDataPins;
+		TArray<FString> SplineDataPins;
+		TArray<FString> RawBufferDataPins;
+
+		for (const FPCGPinProperties& Pin : OutputPinProperties())
+		{
+			DataPins.Add(Pin.Label.ToString());
+
+			if (!!(Pin.AllowedTypes & EPCGDataType::Point))
+			{
+				PointDataPins.Add(Pin.Label.ToString());
+			}
+
+			if (!!(Pin.AllowedTypes & EPCGDataType::Spline))
+			{
+				SplineDataPins.Add(Pin.Label.ToString());
+			}
+		}
+
+		if (!DataPins.IsEmpty())
+		{
+			OutputDeclarations += TEXT("### DATA FUNCTIONS ###\n\n");
+			OutputDeclarations += TEXT("// Valid pins: ") + FString::Join(DataPins, TEXT(", ")) + TEXT("\n");
+			OutputDeclarations += TEXT("// Valid types: bool, int, float, float2, float3, float4, Rotator (float3), Quat (float4), Transform (float4x4)\n");
+			OutputDeclarations += TEXT("\n");
+			OutputDeclarations += TEXT("void <pin>_Set<type>(uint DataIndex, uint AttributeId, uint ElementIndex, <type> Value);");
+		}
+
+		if (!PointDataPins.IsEmpty())
+		{
+			OutputDeclarations += TEXT("\n\n### POINT DATA FUNCTIONS ###\n\n");
+			OutputDeclarations += TEXT("// Valid pins: ") + FString::Join(PointDataPins, TEXT(", ")) + TEXT("\n");
+			OutputDeclarations += TEXT("\n");
+			OutputDeclarations += TEXT("void <pin>_SetPosition(uint DataIndex, uint ElementIndex, float3 Position);\n");
+			OutputDeclarations += TEXT("void <pin>_SetRotation(uint DataIndex, uint ElementIndex, float4 Rotation);\n");
+			OutputDeclarations += TEXT("void <pin>_SetScale(uint DataIndex, uint ElementIndex, float3 Scale);\n");
+			OutputDeclarations += TEXT("void <pin>_SetBoundsMin(uint DataIndex, uint ElementIndex, float3 BoundsMin);\n");
+			OutputDeclarations += TEXT("void <pin>_SetBoundsMax(uint DataIndex, uint ElementIndex, float3 BoundsMax);\n");
+			OutputDeclarations += TEXT("void <pin>_SetColor(uint DataIndex, uint ElementIndex, float4 Color);\n");
+			OutputDeclarations += TEXT("void <pin>_SetDensity(uint DataIndex, uint ElementIndex, float Density);\n");
+			OutputDeclarations += TEXT("void <pin>_SetSeed(uint DataIndex, uint ElementIndex, int Seed);\n");
+			OutputDeclarations += TEXT("void <pin>_SetSeedFromPosition(uint DataIndex, uint ElementIndex, float3 Position);\n");
+			OutputDeclarations += TEXT("void <pin>_SetSteepness(uint DataIndex, uint ElementIndex, float Steepness);\n");
+			OutputDeclarations += TEXT("void <pin>_SetPointTransform(uint DataIndex, uint ElementIndex, float4x4 Transform);");
+		}
+
+		if (!RawBufferDataPins.IsEmpty())
+		{
+			OutputDeclarations += TEXT("\n\n### BYTE ADDRESS BUFFER DATA FUNCTIONS ###\n\n");
+			OutputDeclarations += TEXT("// Valid pins: ") + FString::Join(RawBufferDataPins, TEXT(", ")) + TEXT("\n");
+			OutputDeclarations += TEXT("\n");
+			OutputDeclarations += TEXT("uint <pin>_WriteValue(uint Index, uint Value);");
+		}
+	}
+}
+
+void UPCGCustomHLSLSettings::UpdatePinSettings()
+{
+	for (int PinIndex = 0; PinIndex < InputPins.Num(); ++PinIndex)
+	{
+		FPCGPinProperties& Properties = InputPins[PinIndex];
+		Properties.bDisplayBufferSizeSettings = false;
+
+		// Type Any is not allowed, default to Point
+		if (Properties.AllowedTypes == EPCGDataType::Any)
+		{
+			Properties.AllowedTypes = EPCGDataType::Point;
+		}
+
+		if (!!(Properties.AllowedTypes & EPCGDataType::Landscape)
+			|| !!(Properties.AllowedTypes & EPCGDataType::Texture))
+		{
+			// Don't allow multiple data on this pin because we do not support a dynamic number of textures/landscapes bound to a
+			// compute kernel.
+			Properties.bAllowMultipleData = false;
+			Properties.bAllowEditMultipleData = false;
+		}
+		else
+		{
+			Properties.bAllowEditMultipleData = true;
+		}
+
+		// TODO: We have work to do to allow dynamic merging of data. Also we will likely inject Gather
+		// nodes on the CPU side so that merging is handled CPU side where possible.
+		Properties.SetAllowMultipleConnections(false);
+		Properties.bAllowEditMultipleConnections = false;
+	}
+
+	for (int PinIndex = 0; PinIndex < OutputPins.Num(); ++PinIndex)
+	{
+		FPCGPinProperties& Properties = OutputPins[PinIndex];
+
+		// Type Any is not allowed, default to Point
+		if (Properties.AllowedTypes == EPCGDataType::Any)
+		{
+			Properties.AllowedTypes = EPCGDataType::Point;
+		}
+
+		// Primary pin settings driven by kernel (if not custom kernel type).
+		Properties.bDisplayBufferSizeSettings = PinIndex > 0 || KernelType == EPCGKernelType::Custom;
+
+		// Output pins should always allow multiple connections.
+		// TODO this could be hoisted up somewhere in the future.
+		Properties.bAllowEditMultipleConnections = false;
+
+		// Node is free to output any number of data.
+		Properties.bAllowEditMultipleData = false;
+	}
+}
+#endif
+
+bool UPCGCustomHLSLSettings::IsKernelValid(FPCGContext* InContext) const
+{
+	if (OutputPinProperties().IsEmpty())
+	{
+		if (InContext)
+		{
+			PCGE_LOG_C(Error, GraphAndLog, InContext, LOCTEXT("NoOutputs", "Kernels must have at least one output."));
+		}
+
+		return false;
+	}
+
+	for (const FPCGPinProperties& Properties : InputPinProperties())
+	{
+		if (Properties.AllowedTypes == EPCGDataType::Any)
+		{
+			PCGE_LOG_C(Error, GraphAndLog, InContext, FText::Format(LOCTEXT("InvalidAnyInput", "Custom kernels do not support inputs of type Any, found on pin {0}."), FText::FromName(Properties.Label)));
+			return false;
+		}
+	}
+
+	for (const FPCGPinProperties& Properties : OutputPinProperties())
+	{
+		if (Properties.AllowedTypes == EPCGDataType::Any)
+		{
+			PCGE_LOG_C(Error, GraphAndLog, InContext, FText::Format(LOCTEXT("InvalidAnyOutput", "Custom kernels do not support outputs of type Any, found on pin {0}."), FText::FromName(Properties.Label)));
+			return false;
+		}
+
+		if (!!(Properties.AllowedTypes & EPCGDataType::Landscape))
+		{
+			PCGE_LOG_C(Error, GraphAndLog, InContext, FText::Format(LOCTEXT("InvalidLSOutput", "Custom kernels do not support outputs of type Landscape, found on pin {0}."), FText::FromName(Properties.Label)));
+			return false;
+		}
+
+		if (!!(Properties.AllowedTypes & EPCGDataType::Texture))
+		{
+			PCGE_LOG_C(Error, GraphAndLog, InContext, FText::Format(LOCTEXT("InvalidTextureOutput", "Custom kernels do not support outputs of type Texture, found on pin {0}."), FText::FromName(Properties.Label)));
+			return false;
+		}
+	}
+
+	if (KernelType == EPCGKernelType::PointProcessor)
+	{
+		if (!GetPointProcessingInputPin())
+		{
+			if (InContext)
+			{
+				PCGE_LOG_C(Error, GraphAndLog, InContext, LOCTEXT("InvalidPPInput", "Point processing kernel requires a first input pin of type point."));
+			}
+
+			return false;
+		}
+
+		if (!GetFirstPointOutputPin())
+		{
+			if (InContext)
+			{
+				PCGE_LOG_C(Error, GraphAndLog, InContext, LOCTEXT("InvalidPPInput", "Point processing kernel requires a first output pin of type point."));
+			}
+
+			return false;
+		}
+	}
+
+	auto CheckPinLabels = [InContext](const TArray<FPCGPinProperties>& Pins)
+	{
+		TSet<FName> EncounteredLabels;
+		for (const FPCGPinProperties& Pin : Pins)
+		{
+			if (EncounteredLabels.Contains(Pin.Label))
+			{
+				if (InContext)
+				{
+					PCGE_LOG_C(Error, GraphAndLog, InContext, FText::Format(LOCTEXT("DuplicatedPinLabels", "Duplicate pin label '{0}', all labels must be unique."), FText::FromName(Pin.Label)));
+				}
+
+				return false;
+			}
+		}
+
+		return true;
+	};
+
+	if (!CheckPinLabels(InputPinProperties()) || !CheckPinLabels(OutputPinProperties()))
+	{
+		return false;
+	}
+
+	return true;
+}
+
+FString UPCGCustomHLSLSettings::GetCookedKernelSource(const TMap<FPCGKernelAttributeKey, int32>& GlobalAttributeLookupTable) const
+{
+	const FIntVector GroupSize = GetThreadGroupSize();
+
+	// FIXME: Create source range mappings so that we can go from error location to our source.
+	FString Source = (TEXT("    // User kernel\n") + ShaderSource);
+	FString Functions = (TEXT("// User kernel functions\n") + ShaderFunctions);
+
+#if PLATFORM_WINDOWS
+	// Remove old-school stuff.
+	Source.ReplaceInline(TEXT("\r"), TEXT(""));
+	Functions.ReplaceInline(TEXT("\r"), TEXT(""));
+#endif
+
+	Source.ReplaceInline(TEXT("\n"), TEXT("\n    ")); // Properly indent kernel source
+
+	for (const TPair<FPCGKernelAttributeKey, int32>& Pair : GlobalAttributeLookupTable)
+	{
+		const FString SourceDefinition = PCGHLSLElement::GetKernelAttributeKeyAsString(Pair.Key);
+		const FString AttributeIndexAsString = FString::FromInt(Pair.Value);
+
+		Source.ReplaceInline(*SourceDefinition, *AttributeIndexAsString);
+		Functions.ReplaceInline(*SourceDefinition, *AttributeIndexAsString);
+	}
+
+	FString ShaderPathName = GetPathName();
+	PCGHLSLElement::ConvertObjectPathToShaderFilePath(ShaderPathName);
+
+	const bool bHasKernelKeyword = Source.Contains(TEXT("KERNEL"), ESearchCase::CaseSensitive);
+
+	FString Includes;
+	{
+		// Add with caution: Pulling in external includes has the danger that 1) we may pull in more than we expect
+		// if they include additional things in the future, and 2) if their functions change it could break user's
+		// kernel source. The latter may be mitigated by branching our own PCG version of these if needed in the future.
+
+		// TODO bring these in via additional sources so that their contents get hashed. Also these could generate declarations.
+		Includes += TEXT("#include \"/Engine/Private/ComputeShaderUtils.ush\"\n");
+		Includes += TEXT("#include \"/Engine/Private/Quaternion.ush\"\n");
+		Includes += TEXT("#include \"/Plugin/PCG/Private/PCGShaderUtils.ush\"\n");
+	}
+
+	const FString KernelFunc = FString::Printf(
+		TEXT("[numthreads(%d, %d, %d)]\nvoid %s(uint3 GroupId : SV_GroupID, uint GroupIndex : SV_GroupIndex)"),
+		GroupSize.X, GroupSize.Y, GroupSize.Z, *GetKernelEntryPoint());
+
+	const FString UnWrappedDispatchThreadId = FString::Printf(
+		TEXT("GetUnWrappedDispatchThreadId(GroupId, GroupIndex, %d)"),
+		GroupSize.X * GroupSize.Y * GroupSize.Z
+	);
+
+	// Header writers initialize PCG data collection format headers in output buffers.
+	// TODO: Currently done from all threads. Only write header from global index 0? Or write sections from data local index 0?
+	FString HeaderWriters;
+
+	auto EmitHeaderWriterFromInputPin = [&HeaderWriters](const FPCGPinProperties& InOutputPinProps, const UPCGPin* InFromPin)
+	{
+		HeaderWriters += FString::Format(TEXT(
+			"    if (ThreadIndex >= GetNumThreads()) return;\n"
+			"    \n"
+			"    // Copy data collection header from pin {0} to pin {1}\n"
+			"    {\n"
+			"        const uint NumData = {0}_GetNumData();\n"
+			"        {1}_WriteNumData(NumData);\n"
+			"        \n"
+			"        for (uint DataIndex = 0; DataIndex < NumData; ++DataIndex)\n"
+			"        {\n"
+			"            const uint DataAddress = {0}_ReadDataAddress(DataIndex);\n"
+			"            const uint Id = {0}_ReadDataId(DataAddress);\n"
+			"            const uint NumAttributes = {0}_ReadDataNumAttributes(DataAddress);\n"
+			"            const uint PreambleSizeBytes = {0}_ReadDataPreambleSize(DataAddress);\n"
+			"            const uint NumElements = {0}_ReadDataInfo(DataAddress);\n"
+			"            \n"
+			"            {1}_WriteDataAddress(DataIndex, DataAddress);\n"
+			"            {1}_WriteDataId(DataAddress, Id);\n"
+			"            {1}_WriteDataNumAttributes(DataAddress, NumAttributes);\n"
+			"            {1}_WriteDataPreambleSize(DataAddress, PreambleSizeBytes);\n"
+			"            {1}_WriteDataInfo(DataAddress, NumElements);\n"
+			"            \n"
+			"            const uint AttributeHeadersAddress = {0}_ReadDataAttributeHeadersAddress(DataAddress);\n"
+			"            const uint PCG_MAX_NUM_ATTRIBUTES = {2};\n"
+			"            \n"
+			"            for (uint AttributeIndex = 0; AttributeIndex < PCG_MAX_NUM_ATTRIBUTES; ++AttributeIndex)\n"
+			"            {\n"
+			"                const uint AttributeHeaderAddress = {0}_ReadAttributeHeaderAddress(AttributeHeadersAddress, AttributeIndex);\n"
+			"                const uint AttributeIdAndStride = {0}_ReadAttributeIdAndStride(AttributeHeaderAddress);\n"
+			"                const uint AttributeAddress = {0}_ReadAttributeAddress(AttributeHeaderAddress);\n"
+			"                \n"
+			"                {1}_WriteAttributeIdAndStride(AttributeHeaderAddress, AttributeIdAndStride);\n"
+			"                {1}_WriteAttributeAddress(AttributeHeaderAddress, AttributeAddress);\n"
+			"            }\n"
+			"        }\n"
+			"    }\n"),
+			{ InFromPin->Properties.Label.ToString(), InOutputPinProps.Label.ToString(), PCGComputeConstants::MAX_NUM_ATTRS });
+	};
+
+	if (KernelType == EPCGKernelType::PointProcessor || KernelType == EPCGKernelType::Custom)
+	{
+		const UPCGPin* FirstPin = (KernelType == EPCGKernelType::PointProcessor)
+			? GetPointProcessingInputPin()
+			: CastChecked<UPCGNode>(GetOuter())->GetPassThroughInputPin();
+
+		const UPCGPin* SecondPin = GetSecondPointProcessingInputPin();
+
+		// Initialize all output headers.
+		for (const FPCGPinProperties& PinProps : OutputPinProperties())
+		{
+			// Only support Point pins for now.
+			if (PinProps.AllowedTypes != EPCGDataType::Point)
+			{
+				continue;
+			}
+
+			if (PinProps.BufferSizeMode == EPCGPinBufferSizeMode::FromFirstPin && FirstPin)
+			{
+				EmitHeaderWriterFromInputPin(PinProps, FirstPin);
+			}
+			else if (PinProps.BufferSizeMode == EPCGPinBufferSizeMode::FromSecondPin && SecondPin)
+			{
+				EmitHeaderWriterFromInputPin(PinProps, SecondPin);
+			}
+		}
+	}
+	else if (KernelType == EPCGKernelType::PointGenerator)
+	{
+		if (const UPCGNode* Node = Cast<UPCGNode>(GetOuter()))
+		{
+			const UPCGPin* FirstPin = (KernelType == EPCGKernelType::PointProcessor)
+				? GetPointProcessingInputPin()
+				: CastChecked<UPCGNode>(GetOuter())->GetPassThroughInputPin();
+			const UPCGPin* SecondPin = GetSecondPointProcessingInputPin();
+			const UPCGPin* PrimaryOutputPin = GetFirstPointOutputPin();
+			
+			auto EmitPointGenHeader = [&HeaderWriters, InPointCount = PointCount](const FPCGPinProperties& InOutputPinProps)
+			{
+				HeaderWriters += FString::Format(TEXT(
+					"    if (ThreadIndex >= GetNumThreads()) return;\n"
+					"    \n"
+					"    // Write data collection header for pin {0}\n"
+					"    {0}_WriteSinglePointDataCollectionHeader({1});\n"),
+					{ InOutputPinProps.Label.ToString(), InPointCount });
+			};
+
+			for (const UPCGPin* OutputPin : Node->GetOutputPins())
+			{
+				if (!OutputPin || OutputPin->Properties.AllowedTypes != EPCGDataType::Point)
+				{
+					// Point only for now.
+					continue;
+				}
+
+				const FPCGPinProperties& PinProps = OutputPin->Properties;
+
+				if (OutputPin == PrimaryOutputPin)
+				{
+					EmitPointGenHeader(PinProps);
+				}
+				else if (PinProps.BufferSizeMode == EPCGPinBufferSizeMode::FromFirstPin && FirstPin)
+				{
+					EmitHeaderWriterFromInputPin(PinProps, FirstPin);
+				}
+				else if (PinProps.BufferSizeMode == EPCGPinBufferSizeMode::FromSecondPin && SecondPin)
+				{
+					EmitHeaderWriterFromInputPin(PinProps, SecondPin);
+				}
+			}
+		}
+	}
+	else
+	{
+		checkNoEntry();
+	}
+
+	// Per-kernel-type preamble. Set up shader inputs and initialize output data.
+	FString KernelSpecificPreamble = TEXT("    // Kernel preamble\n");
+	
+	if (KernelType == EPCGKernelType::PointProcessor)
+	{
+		const UPCGPin* InputPin = GetPointProcessingInputPin();
+		const UPCGPin* OutputPin = GetFirstPointOutputPin();
+
+		if (InputPin && OutputPin)
+		{
+			KernelSpecificPreamble += FString::Format(TEXT(
+				"    const uint3 {0}_ThreadInfo = {0}_GetThreadData(ThreadIndex);\n"
+				"    const uint {0}_DataIndex = {0}_ThreadInfo[0];\n"
+				"    const uint {0}_DataAddress = {0}_ThreadInfo[1];\n"
+				"    const uint ElementIndex = {0}_ThreadInfo[2]; // Assumption - element index identical in input and output data.\n"),
+				{ InputPin->Properties.Label.ToString() });
+
+			KernelSpecificPreamble += FString::Format(TEXT(
+				"    const uint3 {0}_ThreadInfo = {0}_GetThreadData(ThreadIndex);\n"
+				"    const uint {0}_DataIndex = {0}_ThreadInfo[0];\n"
+				"    const uint {0}_DataAddress = {0}_ThreadInfo[1];\n"),
+				{ OutputPin->Properties.Label.ToString() });
+
+			// Automatically copy all attributes
+			KernelSpecificPreamble += FString::Format(TEXT(
+				"\n"
+				"    // Loop over all attribute headers, if the address is non-zero, then copy it from pin {0} to pin {1}.\n"
+				"    {\n"
+				"        const uint HeadersAddress = {0}_ReadDataAttributeHeadersAddress({0}_DataAddress);\n"
+				"        const uint NumAttributes = {0}_ReadDataNumAttributes({0}_DataAddress);\n"
+				"        uint NumAttributesProcessed = 0;\n"
+				"\n"
+				"        for (int AttributeIndex = 0; AttributeIndex < 128; ++AttributeIndex)\n"
+				"        {\n"
+				"            const uint HeaderAddress = {0}_ReadAttributeHeaderAddress(HeadersAddress, AttributeIndex);\n"
+				"            const uint Stride = {0}_ReadAttributeStride(HeaderAddress);\n"
+				"            const uint Address = {0}_ReadAttributeAddress(HeaderAddress);\n"
+				"\n"
+				"            if (Address != 0)\n"
+				"            {\n"
+				"                const uint BaseElementAddress = Address + ElementIndex * Stride;\n"
+				"\n"
+				"                for (int I = 0; I < Stride; I += 4)\n"
+				"                {\n"
+				"                    const uint ElementAddress = BaseElementAddress + I;\n"
+				"                    {1}_StoreBuffer(ElementAddress, {0}_LoadBuffer(ElementAddress));\n"
+				"                }\n"
+				"\n"
+				"                if (++NumAttributesProcessed >= NumAttributes) break; // We can early-out when we've looked at all the possible attributes\n"
+				"            }\n"
+				"        }\n"
+				"    }\n"),
+				{ InputPin->Properties.Label.ToString(), OutputPin->Properties.Label.ToString() });
+		}
+	}
+	else if (KernelType == EPCGKernelType::PointGenerator)
+	{
+		KernelSpecificPreamble += FString::Format(TEXT(
+			"    const uint PointCount = {0};\n"),
+			{ PointCount });
+
+		if (const UPCGPin* OutputPin = GetFirstPointOutputPin())
+		{
+			KernelSpecificPreamble += FString::Format(TEXT(
+				"    const uint3 {0}_ThreadInfo = {0}_GetThreadData(ThreadIndex);\n"
+				"    const uint {0}_DataIndex = {0}_ThreadInfo[0];\n"
+				"    const uint {0}_DataAddress = {0}_ThreadInfo[1];\n"
+				"    const uint ElementIndex = {0}_ThreadInfo[2];\n"),
+				{ OutputPin->Properties.Label.ToString() });
+
+			KernelSpecificPreamble += FString::Format(TEXT(
+				"    \n"
+				"    // Initialize all values to defaults for output pin {0}\n"
+				"    {0}_InitializePoint({0}_DataIndex, ElementIndex);\n"),
+				{ OutputPin->Properties.Label.ToString() });
+		}
+	}
+
+	FString Result;
+
+	if (bHasKernelKeyword)
+	{
+		Source.ReplaceInline(TEXT("KERNEL"), TEXT("void __kernel_func(uint ThreadIndex)"), ESearchCase::CaseSensitive);
+
+		Result = FString::Printf(TEXT(
+			"#line 1 \"%s\"\n" // ShaderPathName
+			"%s\n" // Includes
+			"%s\n" // Functions
+			"%s\n" // Source
+			"%s { __kernel_func(%s); }\n"), // KernelFunc, UnWrappedDispatchThreadId
+			*ShaderPathName, *Includes, *Functions, *Source, *KernelFunc, *UnWrappedDispatchThreadId);
+	}
+	else
+	{
+		Result = FString::Printf(TEXT(
+			"%s\n\n" // Includes
+			"%s\n\n" // Functions
+			"%s\n" // KernelFunc
+			"{\n"
+			"    const uint ThreadIndex = %s;\n" // UnWrappedDispatchThreadId
+			"%s\n" // HeaderWriters
+			"%s\n" // KernelSpecificPreamble
+			"#line 1 \"%s\"\n\n" // ShaderPathName
+			"%s\n" // Source
+			"}\n"),
+			*Includes, *Functions, *KernelFunc, *UnWrappedDispatchThreadId, *HeaderWriters, *KernelSpecificPreamble, *ShaderPathName, *Source);
+	}
+
+	if (bDumpCookedHLSL)
+	{
+		UE_LOG(LogPCG, Log, TEXT("Cooked HLSL:\n%s\n"), *Result);
+	}
+
+	return Result;
+}
+
+#undef LOCTEXT_NAMESPACE
diff --git a/Engine/Plugins/PCG/Source/PCG/Private/Compute/PCGComputeCommon.h b/Engine/Plugins/PCG/Source/PCG/Private/Compute/PCGComputeCommon.h
new file mode 100644
--- /dev/null
+++ b/Engine/Plugins/PCG/Source/PCG/Private/Compute/PCGComputeCommon.h
@@ -0,0 +1,27 @@
+// Copyright Epic Games, Inc. All Rights Reserved.
+
+#pragma once
+
+namespace PCGComputeConstants
+{
+	constexpr int MAX_NUM_ATTRS = 128;
+	constexpr int NUM_RESERVED_ATTRS = 32; // Reserved for point properties, spline accessors, etc.
+	constexpr int MAX_NUM_CUSTOM_ATTRS = MAX_NUM_ATTRS - NUM_RESERVED_ATTRS; // Reserved for custom attributes
+
+	constexpr int ATTRIBUTE_HEADER_SIZE_BYTES = 8;
+
+	constexpr int POINT_DATA_TYPE_ID = 0;
+	constexpr int POINT_DATA_HEADER_PREAMBLE_SIZE_BYTES = 16; // 4 bytes for Type, 4 bytes for NumAttrs, 4 bytes for the offset, 4 bytes for TypeInfo
+	constexpr int POINT_DATA_HEADER_SIZE_BYTES = POINT_DATA_HEADER_PREAMBLE_SIZE_BYTES + ATTRIBUTE_HEADER_SIZE_BYTES * MAX_NUM_ATTRS;
+
+	constexpr int NUM_POINT_PROPERTIES = 9;
+	constexpr int POINT_POSITION_ATTRIBUTE_ID = 0;
+	constexpr int POINT_ROTATION_ATTRIBUTE_ID = 1;
+	constexpr int POINT_SCALE_ATTRIBUTE_ID = 2;
+	constexpr int POINT_BOUNDS_MIN_ATTRIBUTE_ID = 3;
+	constexpr int POINT_BOUNDS_MAX_ATTRIBUTE_ID = 4;
+	constexpr int POINT_COLOR_ATTRIBUTE_ID = 5;
+	constexpr int POINT_DENSITY_ATTRIBUTE_ID = 6;
+	constexpr int POINT_SEED_ATTRIBUTE_ID = 7;
+	constexpr int POINT_STEEPNESS_ATTRIBUTE_ID = 8;
+}
diff --git a/Engine/Plugins/PCG/Source/PCG/Private/Compute/PCGComputeGraph.cpp b/Engine/Plugins/PCG/Source/PCG/Private/Compute/PCGComputeGraph.cpp
new file mode 100644
--- /dev/null
+++ b/Engine/Plugins/PCG/Source/PCG/Private/Compute/PCGComputeGraph.cpp
@@ -0,0 +1,23 @@
+// Copyright Epic Games, Inc. All Rights Reserved.
+
+#include "PCGComputeGraph.h"
+
+#include "PCGModule.h"
+#include "PCGNode.h"
+
+#include "ComputeFramework/ComputeKernelCompileResult.h"
+
+#include UE_INLINE_GENERATED_CPP_BY_NAME(PCGComputeGraph)
+
+void UPCGComputeGraph::OnKernelCompilationComplete(int32 InKernelIndex, FComputeKernelCompileResults const& InCompileResults)
+{
+	if (KernelToNode.IsValidIndex(InKernelIndex) && KernelToNode[InKernelIndex].Get())
+	{
+		KernelToCompileMessages.FindOrAdd(KernelToNode[InKernelIndex].Get()) = InCompileResults.Messages;
+	}
+	else
+	{
+		// We may in general have kernels with no corresponding node.
+		UE_LOG(LogPCG, Verbose, TEXT("Compilation message ignored for kernel index %d which has no associated node."), InKernelIndex);
+	}
+}
diff --git a/Engine/Plugins/PCG/Source/PCG/Private/Compute/PCGComputeGraph.h b/Engine/Plugins/PCG/Source/PCG/Private/Compute/PCGComputeGraph.h
new file mode 100644
--- /dev/null
+++ b/Engine/Plugins/PCG/Source/PCG/Private/Compute/PCGComputeGraph.h
@@ -0,0 +1,54 @@
+// Copyright Epic Games, Inc. All Rights Reserved.
+
+#pragma once
+
+#include "Compute/PCGDataForGPU.h"
+
+#include "ComputeFramework/ComputeGraph.h"
+#include "ComputeFramework/ComputeKernelCompileResult.h"
+
+#include "UObject/ObjectKey.h"
+
+#include "PCGComputeGraph.generated.h"
+
+class UPCGNode;
+class UPCGPin;
+
+UCLASS()
+class UPCGComputeGraph : public UComputeGraph
+{
+	GENERATED_BODY()
+
+public:
+	//~Begin UComputeGraph interface
+	void OnKernelCompilationComplete(int32 InKernelIndex, FComputeKernelCompileResults const& InCompileResults) override;
+	//~End UComputeGraph interface
+
+	/** Get the global attribute indices. */
+	const TMap<FPCGKernelAttributeKey, int32>& GetAttributeLookupTable() const { return GlobalAttributeLookupTable; }
+
+public:
+	TMap<TObjectKey<const UPCGNode>, TArray<FComputeKernelCompileMessage>> KernelToCompileMessages;
+	
+	/** Set of input pins at the CPU -> GPU border. */
+	UPROPERTY()
+	TArray<TWeakObjectPtr<const UPCGPin>> PinsReceivingDataFromCPU;
+
+	/** Pin label aliases, used for selecting data items corresponding to an input pin from the input data collection. */
+	UPROPERTY()
+	TMap<const UPCGPin*, FName> InputPinLabelAliases;
+
+	/** Mapping from upstream output pin to downstream pin alias, used to select data items originating from upstream pin from the input data collection. */
+	UPROPERTY()
+	TMap<const UPCGPin*, FName> OutputCPUPinToInputGPUPinAlias;
+
+	// Node corresponding to each kernel, useful for compilation feedback.
+	UPROPERTY()
+	TArray<TWeakObjectPtr<const UPCGNode>> KernelToNode;
+
+protected:
+	UPROPERTY()
+	TMap<FPCGKernelAttributeKey, int32 /* Attribute Index */> GlobalAttributeLookupTable;
+
+	friend class FPCGGraphCompiler;
+};
diff --git a/Engine/Plugins/PCG/Source/PCG/Private/Compute/PCGComputeKernelSource.h b/Engine/Plugins/PCG/Source/PCG/Private/Compute/PCGComputeKernelSource.h
new file mode 100644
--- /dev/null
+++ b/Engine/Plugins/PCG/Source/PCG/Private/Compute/PCGComputeKernelSource.h
@@ -0,0 +1,30 @@
+// Copyright Epic Games, Inc. All Rights Reserved.
+
+#pragma once
+
+#include "ComputeFramework/ComputeKernelSource.h"
+
+#include "PCGComputeKernelSource.generated.h"
+
+UCLASS()
+class PCG_API UPCGComputeKernelSource : public UComputeKernelSource
+{
+	GENERATED_BODY()
+
+public:
+	void SetSource(FString const& InSource)
+	{
+		Source = InSource;
+	}
+	
+	//~ Begin UComputeKernelSource Interface.
+	FString GetSource() const override
+	{
+		return Source;
+	}
+	//~ End UComputeKernelSource Interface.
+
+protected:
+	UPROPERTY()
+	FString Source;
+};
diff --git a/Engine/Plugins/PCG/Source/PCG/Private/Compute/PCGDataBinding.h b/Engine/Plugins/PCG/Source/PCG/Private/Compute/PCGDataBinding.h
new file mode 100644
--- /dev/null
+++ b/Engine/Plugins/PCG/Source/PCG/Private/Compute/PCGDataBinding.h
@@ -0,0 +1,29 @@
+// Copyright Epic Games, Inc. All Rights Reserved.
+
+#pragma once
+
+#include "PCGData.h"
+#include "Compute/PCGDataForGPU.h"
+
+#include "ComputeFramework/ComputeGraphInstance.h"
+
+#include "PCGDataBinding.generated.h"
+
+class UPCGComponent;
+class UPCGComputeGraph;
+
+UCLASS(Transient, Category = PCG)
+class UPCGDataBinding : public UObject
+{
+	GENERATED_BODY()
+
+public:
+	TWeakObjectPtr<UPCGComponent> SourceComponent;
+
+	FPCGDataCollection OutputDataCollection;
+
+	const UPCGComputeGraph* Graph = nullptr;
+
+	// Data collection pointer + set of pins that get data from the collection (cross barrier)
+	FPCGDataForGPU DataForGPU;
+};
diff --git a/Engine/Plugins/PCG/Source/PCG/Private/Compute/PCGDataForGPU.cpp b/Engine/Plugins/PCG/Source/PCG/Private/Compute/PCGDataForGPU.cpp
new file mode 100644
--- /dev/null
+++ b/Engine/Plugins/PCG/Source/PCG/Private/Compute/PCGDataForGPU.cpp
@@ -0,0 +1,915 @@
+// Copyright Epic Games, Inc. All Rights Reserved.
+
+#include "Compute/PCGDataForGPU.h"
+
+#include "PCGData.h"
+#include "PCGPoint.h"
+#include "Compute/PCGComputeCommon.h"
+#include "Data/PCGPointData.h"
+#include "Metadata/PCGMetadata.h"
+
+#include "Async/ParallelFor.h"
+
+using namespace PCGComputeConstants;
+
+namespace PCGDataForGPUConstants
+{
+	const static FPCGKernelAttributeDesc PointPropertyDescs[NUM_POINT_PROPERTIES] =
+	{
+		{ POINT_POSITION_ATTRIBUTE_ID,   EPCGKernelAttributeType::Float3 },
+		{ POINT_ROTATION_ATTRIBUTE_ID,   EPCGKernelAttributeType::Quat   },
+		{ POINT_SCALE_ATTRIBUTE_ID,      EPCGKernelAttributeType::Float3 },
+		{ POINT_BOUNDS_MIN_ATTRIBUTE_ID, EPCGKernelAttributeType::Float3 },
+		{ POINT_BOUNDS_MAX_ATTRIBUTE_ID, EPCGKernelAttributeType::Float3 },
+		{ POINT_COLOR_ATTRIBUTE_ID,      EPCGKernelAttributeType::Float4 },
+		{ POINT_DENSITY_ATTRIBUTE_ID,    EPCGKernelAttributeType::Float  },
+		{ POINT_SEED_ATTRIBUTE_ID,       EPCGKernelAttributeType::Int    },
+		{ POINT_STEEPNESS_ATTRIBUTE_ID,  EPCGKernelAttributeType::Float  }
+	};
+}
+
+namespace PCGDataForGPUHelpers
+{
+	EPCGKernelAttributeType GetAttributeTypeFromMetadataType(EPCGMetadataTypes MetadataType)
+	{
+		switch (MetadataType)
+		{
+		case EPCGMetadataTypes::Boolean:
+			return EPCGKernelAttributeType::Bool;
+		case EPCGMetadataTypes::Float:
+		case EPCGMetadataTypes::Double:
+			return EPCGKernelAttributeType::Float;
+		case EPCGMetadataTypes::Integer32:
+		case EPCGMetadataTypes::Integer64:
+			return EPCGKernelAttributeType::Int;
+		case EPCGMetadataTypes::Vector2:
+			return EPCGKernelAttributeType::Float2;
+		case EPCGMetadataTypes::Vector:
+			return EPCGKernelAttributeType::Float3;
+		case EPCGMetadataTypes::Rotator:
+			return EPCGKernelAttributeType::Rotator;
+		case EPCGMetadataTypes::Vector4:
+			return EPCGKernelAttributeType::Float4;
+		case EPCGMetadataTypes::Quaternion:
+			return EPCGKernelAttributeType::Quat;
+		case EPCGMetadataTypes::Transform:
+			return EPCGKernelAttributeType::Transform;
+		default:
+			checkNoEntry();
+			return EPCGKernelAttributeType::Float;
+		}
+	}
+
+	int GetAttributeTypeStrideBytes(EPCGKernelAttributeType Type)
+	{
+		switch (Type)
+		{
+		case EPCGKernelAttributeType::Bool:
+		case EPCGKernelAttributeType::Int:
+		case EPCGKernelAttributeType::Float:
+			return 4;
+		case EPCGKernelAttributeType::Float2:
+			return 8;
+		case EPCGKernelAttributeType::Float3:
+		case EPCGKernelAttributeType::Rotator:
+			return 12;
+		case EPCGKernelAttributeType::Float4:
+		case EPCGKernelAttributeType::Quat:
+			return 16;
+		case EPCGKernelAttributeType::Transform:
+			return 64;
+		default:
+			checkNoEntry();
+			return 0;
+		}
+	}
+
+	bool PackAttributeHelper(const FPCGMetadataAttributeBase* InAttributeBase, const FPCGKernelAttributeDesc& InAttributeDesc, PCGMetadataEntryKey InEntryKey, TArray<uint32>& OutPackedDataCollection, uint32 ElementIndex)
+	{
+		check(InAttributeBase);
+
+		const PCGMetadataValueKey ValueKey = InAttributeBase->GetValueKey(InEntryKey);
+		const int16 TypeId = InAttributeBase->GetTypeId();
+		const int StrideBytes = PCGDataForGPUHelpers::GetAttributeTypeStrideBytes(InAttributeDesc.Type);
+
+		switch (TypeId)
+		{
+		case PCG::Private::MetadataTypes<bool>::Id:
+		{
+			const FPCGMetadataAttribute<bool>* Attribute = static_cast<const FPCGMetadataAttribute<bool>*>(InAttributeBase);
+			const bool Value = Attribute->GetValue(ValueKey);
+			check(StrideBytes == 4);
+			OutPackedDataCollection[ElementIndex + 0] = Value;
+			break;
+		}
+		case PCG::Private::MetadataTypes<float>::Id:
+		{
+			const FPCGMetadataAttribute<float>* Attribute = static_cast<const FPCGMetadataAttribute<float>*>(InAttributeBase);
+			const float Value = Attribute->GetValue(ValueKey);
+			check(StrideBytes == 4);
+			OutPackedDataCollection[ElementIndex + 0] = FMath::AsUInt(Value);
+			break;
+		}
+		case PCG::Private::MetadataTypes<double>::Id:
+		{
+			const FPCGMetadataAttribute<double>* Attribute = static_cast<const FPCGMetadataAttribute<double>*>(InAttributeBase);
+			const double Value = Attribute->GetValue(ValueKey);
+			check(StrideBytes == 4);
+			OutPackedDataCollection[ElementIndex + 0] = FMath::AsUInt(static_cast<float>(Value));
+			break;
+		}
+		case PCG::Private::MetadataTypes<int32>::Id:
+		{
+			const FPCGMetadataAttribute<int32>* Attribute = static_cast<const FPCGMetadataAttribute<int32>*>(InAttributeBase);
+			const int32 Value = Attribute->GetValue(ValueKey);
+			check(StrideBytes == 4);
+			OutPackedDataCollection[ElementIndex + 0] = Value;
+			break;
+		}
+		case PCG::Private::MetadataTypes<int64>::Id:
+		{
+			const FPCGMetadataAttribute<int64>* Attribute = static_cast<const FPCGMetadataAttribute<int64>*>(InAttributeBase);
+			const int64 Value = Attribute->GetValue(ValueKey);
+			check(StrideBytes == 4);
+			OutPackedDataCollection[ElementIndex + 0] = Value;
+			break;
+		}
+		case PCG::Private::MetadataTypes<FVector2D>::Id:
+		{
+			const FPCGMetadataAttribute<FVector2D>* Attribute = static_cast<const FPCGMetadataAttribute<FVector2D>*>(InAttributeBase);
+			const FVector2D Value = Attribute->GetValue(ValueKey);
+			check(StrideBytes == 8);
+			OutPackedDataCollection[ElementIndex + 0] = FMath::AsUInt(static_cast<float>(Value.X));
+			OutPackedDataCollection[ElementIndex + 1] = FMath::AsUInt(static_cast<float>(Value.Y));
+			break;
+		}
+		case PCG::Private::MetadataTypes<FRotator>::Id:
+		{
+			const FPCGMetadataAttribute<FRotator>* Attribute = static_cast<const FPCGMetadataAttribute<FRotator>*>(InAttributeBase);
+			const FRotator Value = Attribute->GetValue(ValueKey);
+			check(StrideBytes == 12);
+			OutPackedDataCollection[ElementIndex + 0] = FMath::AsUInt(static_cast<float>(Value.Pitch));
+			OutPackedDataCollection[ElementIndex + 1] = FMath::AsUInt(static_cast<float>(Value.Yaw));
+			OutPackedDataCollection[ElementIndex + 2] = FMath::AsUInt(static_cast<float>(Value.Roll));
+			break;
+		}
+		case PCG::Private::MetadataTypes<FVector>::Id:
+		{
+			const FPCGMetadataAttribute<FVector>* Attribute = static_cast<const FPCGMetadataAttribute<FVector>*>(InAttributeBase);
+			const FVector Value = Attribute->GetValue(ValueKey);
+			check(StrideBytes == 12);
+			OutPackedDataCollection[ElementIndex + 0] = FMath::AsUInt(static_cast<float>(Value.X));
+			OutPackedDataCollection[ElementIndex + 1] = FMath::AsUInt(static_cast<float>(Value.Y));
+			OutPackedDataCollection[ElementIndex + 2] = FMath::AsUInt(static_cast<float>(Value.Z));
+			break;
+		}
+		case PCG::Private::MetadataTypes<FVector4>::Id:
+		{
+			const FPCGMetadataAttribute<FVector4>* Attribute = static_cast<const FPCGMetadataAttribute<FVector4>*>(InAttributeBase);
+			const FVector4 Value = Attribute->GetValue(ValueKey);
+			check(StrideBytes == 16);
+			OutPackedDataCollection[ElementIndex + 0] = FMath::AsUInt(static_cast<float>(Value.X));
+			OutPackedDataCollection[ElementIndex + 1] = FMath::AsUInt(static_cast<float>(Value.Y));
+			OutPackedDataCollection[ElementIndex + 2] = FMath::AsUInt(static_cast<float>(Value.Z));
+			OutPackedDataCollection[ElementIndex + 3] = FMath::AsUInt(static_cast<float>(Value.W));
+			break;
+		}
+		case PCG::Private::MetadataTypes<FQuat>::Id:
+		{
+			const FPCGMetadataAttribute<FQuat>* Attribute = static_cast<const FPCGMetadataAttribute<FQuat>*>(InAttributeBase);
+			const FQuat Value = Attribute->GetValue(ValueKey);
+			check(StrideBytes == 16);
+			OutPackedDataCollection[ElementIndex + 0] = FMath::AsUInt(static_cast<float>(Value.X));
+			OutPackedDataCollection[ElementIndex + 1] = FMath::AsUInt(static_cast<float>(Value.Y));
+			OutPackedDataCollection[ElementIndex + 2] = FMath::AsUInt(static_cast<float>(Value.Z));
+			OutPackedDataCollection[ElementIndex + 3] = FMath::AsUInt(static_cast<float>(Value.W));
+			break;
+		}
+		case PCG::Private::MetadataTypes<FTransform>::Id:
+		{
+			const FPCGMetadataAttribute<FTransform>* Attribute = static_cast<const FPCGMetadataAttribute<FTransform>*>(InAttributeBase);
+			const FTransform Value = Attribute->GetValue(ValueKey);
+			const FMatrix Matrix = Value.ToMatrixWithScale();
+			check(StrideBytes == 64);
+			OutPackedDataCollection[ElementIndex + 0] = FMath::AsUInt(static_cast<float>(Matrix.M[0][0]));
+			OutPackedDataCollection[ElementIndex + 1] = FMath::AsUInt(static_cast<float>(Matrix.M[0][1]));
+			OutPackedDataCollection[ElementIndex + 2] = FMath::AsUInt(static_cast<float>(Matrix.M[0][2]));
+			OutPackedDataCollection[ElementIndex + 3] = FMath::AsUInt(static_cast<float>(Matrix.M[0][3]));
+			OutPackedDataCollection[ElementIndex + 4] = FMath::AsUInt(static_cast<float>(Matrix.M[1][0]));
+			OutPackedDataCollection[ElementIndex + 5] = FMath::AsUInt(static_cast<float>(Matrix.M[1][1]));
+			OutPackedDataCollection[ElementIndex + 6] = FMath::AsUInt(static_cast<float>(Matrix.M[1][2]));
+			OutPackedDataCollection[ElementIndex + 7] = FMath::AsUInt(static_cast<float>(Matrix.M[1][3]));
+			OutPackedDataCollection[ElementIndex + 8] = FMath::AsUInt(static_cast<float>(Matrix.M[2][0]));
+			OutPackedDataCollection[ElementIndex + 9] = FMath::AsUInt(static_cast<float>(Matrix.M[2][1]));
+			OutPackedDataCollection[ElementIndex + 10] = FMath::AsUInt(static_cast<float>(Matrix.M[2][2]));
+			OutPackedDataCollection[ElementIndex + 11] = FMath::AsUInt(static_cast<float>(Matrix.M[2][3]));
+			OutPackedDataCollection[ElementIndex + 12] = FMath::AsUInt(static_cast<float>(Matrix.M[3][0]));
+			OutPackedDataCollection[ElementIndex + 13] = FMath::AsUInt(static_cast<float>(Matrix.M[3][1]));
+			OutPackedDataCollection[ElementIndex + 14] = FMath::AsUInt(static_cast<float>(Matrix.M[3][2]));
+			OutPackedDataCollection[ElementIndex + 15] = FMath::AsUInt(static_cast<float>(Matrix.M[3][3]));
+			break;
+		}
+		default:
+			return false;
+		}
+
+		return true;
+	}
+
+	FPCGMetadataAttributeBase* CreateAttributeFromAttributeDesc(UPCGMetadata* Metadata, const FPCGKernelAttributeDesc& AttributeDesc)
+	{
+		switch (AttributeDesc.Type)
+		{
+		case EPCGKernelAttributeType::Bool:
+		{
+			return Metadata->FindOrCreateAttribute<bool>(AttributeDesc.Name);
+		}
+		case EPCGKernelAttributeType::Int:
+		{
+			return Metadata->FindOrCreateAttribute<int>(AttributeDesc.Name);
+		}
+		case EPCGKernelAttributeType::Float:
+		{
+			return Metadata->FindOrCreateAttribute<float>(AttributeDesc.Name);
+		}
+		case EPCGKernelAttributeType::Float2:
+		{
+			return Metadata->FindOrCreateAttribute<FVector2D>(AttributeDesc.Name);
+		}
+		case EPCGKernelAttributeType::Float3:
+		{
+			return Metadata->FindOrCreateAttribute<FVector>(AttributeDesc.Name);
+		}
+		case EPCGKernelAttributeType::Float4:
+		{
+			return Metadata->FindOrCreateAttribute<FVector4>(AttributeDesc.Name);
+		}
+		case EPCGKernelAttributeType::Rotator:
+		{
+			return Metadata->FindOrCreateAttribute<FRotator>(AttributeDesc.Name);
+		}
+		case EPCGKernelAttributeType::Quat:
+		{
+			return Metadata->FindOrCreateAttribute<FQuat>(AttributeDesc.Name);
+		}
+		case EPCGKernelAttributeType::Transform:
+		{
+			return Metadata->FindOrCreateAttribute<FTransform>(AttributeDesc.Name);
+		}
+		default:
+			return nullptr;
+		}
+	}
+
+	bool UnpackAttributeHelper(const void* InPackedData, uint32 ElementIndex, FPCGMetadataAttributeBase* AttributeBase, const FPCGKernelAttributeDesc& AttributeDesc, PCGMetadataEntryKey EntryKey)
+	{
+		check(InPackedData && AttributeBase);
+
+		const float* DataAsFloat = static_cast<const float*>(InPackedData);
+		const int* DataAsInt = static_cast<const int*>(InPackedData);
+
+		switch (AttributeDesc.Type)
+		{
+		case EPCGKernelAttributeType::Bool:
+		{
+			FPCGMetadataAttribute<bool>* Attribute = static_cast<FPCGMetadataAttribute<bool>*>(AttributeBase);
+
+			const bool Value = static_cast<bool>(DataAsFloat[ElementIndex]);
+			Attribute->SetValue(EntryKey, Value);
+			break;
+		}
+		case EPCGKernelAttributeType::Int:
+		{
+			FPCGMetadataAttribute<int>* Attribute = static_cast<FPCGMetadataAttribute<int>*>(AttributeBase);
+
+			const int Value = DataAsInt[ElementIndex];
+			Attribute->SetValue(EntryKey, Value);
+			break;
+		}
+		case EPCGKernelAttributeType::Float:
+		{
+			FPCGMetadataAttribute<float>* Attribute = static_cast<FPCGMetadataAttribute<float>*>(AttributeBase);
+
+			const float Value = DataAsFloat[ElementIndex];
+			Attribute->SetValue(EntryKey, Value);
+			break;
+		}
+		case EPCGKernelAttributeType::Float2:
+		{
+			FPCGMetadataAttribute<FVector2D>* Attribute = static_cast<FPCGMetadataAttribute<FVector2D>*>(AttributeBase);
+
+			FVector2D Value;
+			Value.X = DataAsFloat[ElementIndex + 0];
+			Value.Y = DataAsFloat[ElementIndex + 1];
+			Attribute->SetValue(EntryKey, Value);
+			break;
+		}
+		case EPCGKernelAttributeType::Float3:
+		{
+			FPCGMetadataAttribute<FVector>* Attribute = static_cast<FPCGMetadataAttribute<FVector>*>(AttributeBase);
+
+			FVector Value;
+			Value.X = DataAsFloat[ElementIndex + 0];
+			Value.Y = DataAsFloat[ElementIndex + 1];
+			Value.Z = DataAsFloat[ElementIndex + 2];
+			Attribute->SetValue(EntryKey, Value);
+			break;
+		}
+		case EPCGKernelAttributeType::Float4:
+		{
+			FPCGMetadataAttribute<FVector4>* Attribute = static_cast<FPCGMetadataAttribute<FVector4>*>(AttributeBase);
+
+			FVector4 Value;
+			Value.X = DataAsFloat[ElementIndex + 0];
+			Value.Y = DataAsFloat[ElementIndex + 1];
+			Value.Z = DataAsFloat[ElementIndex + 2];
+			Value.W = DataAsFloat[ElementIndex + 3];
+			Attribute->SetValue(EntryKey, Value);
+			break;
+		}
+		case EPCGKernelAttributeType::Rotator:
+		{
+			FPCGMetadataAttribute<FRotator>* Attribute = static_cast<FPCGMetadataAttribute<FRotator>*>(AttributeBase);
+
+			FRotator Value;
+			Value.Pitch = DataAsFloat[ElementIndex + 0];
+			Value.Yaw = DataAsFloat[ElementIndex + 1];
+			Value.Roll = DataAsFloat[ElementIndex + 2];
+			Attribute->SetValue(EntryKey, Value);
+			break;
+		}
+		case EPCGKernelAttributeType::Quat:
+		{
+			FPCGMetadataAttribute<FQuat>* Attribute = static_cast<FPCGMetadataAttribute<FQuat>*>(AttributeBase);
+
+			FQuat Value;
+			Value.X = DataAsFloat[ElementIndex + 0];
+			Value.Y = DataAsFloat[ElementIndex + 1];
+			Value.Z = DataAsFloat[ElementIndex + 2];
+			Value.W = DataAsFloat[ElementIndex + 3];
+			Attribute->SetValue(EntryKey, Value);
+			break;
+		}
+		case EPCGKernelAttributeType::Transform:
+		{
+			FPCGMetadataAttribute<FTransform>* Attribute = static_cast<FPCGMetadataAttribute<FTransform>*>(AttributeBase);
+
+			FMatrix Matrix;
+			Matrix.M[0][0] = DataAsFloat[ElementIndex + 0];
+			Matrix.M[0][1] = DataAsFloat[ElementIndex + 1];
+			Matrix.M[0][2] = DataAsFloat[ElementIndex + 2];
+			Matrix.M[0][3] = DataAsFloat[ElementIndex + 3];
+			Matrix.M[1][0] = DataAsFloat[ElementIndex + 4];
+			Matrix.M[1][1] = DataAsFloat[ElementIndex + 5];
+			Matrix.M[1][2] = DataAsFloat[ElementIndex + 6];
+			Matrix.M[1][3] = DataAsFloat[ElementIndex + 7];
+			Matrix.M[2][0] = DataAsFloat[ElementIndex + 8];
+			Matrix.M[2][1] = DataAsFloat[ElementIndex + 9];
+			Matrix.M[2][2] = DataAsFloat[ElementIndex + 10];
+			Matrix.M[2][3] = DataAsFloat[ElementIndex + 11];
+			Matrix.M[3][0] = DataAsFloat[ElementIndex + 12];
+			Matrix.M[3][1] = DataAsFloat[ElementIndex + 13];
+			Matrix.M[3][2] = DataAsFloat[ElementIndex + 14];
+			Matrix.M[3][3] = DataAsFloat[ElementIndex + 15];
+
+			const FTransform Value(Matrix);
+			Attribute->SetValue(EntryKey, Value);
+			break;
+		}
+		default:
+			return false;
+		}
+
+		return true;
+	}
+}
+
+uint32 GetTypeHash(const FPCGKernelAttributeKey& In)
+{
+	return HashCombine(GetTypeHash(In.Type), GetTypeHash(In.Name));
+}
+
+FPCGDataDesc::FPCGDataDesc(EPCGDataType InType, int InElementCount)
+	: Type(InType)
+	, ElementCount(InElementCount)
+{
+	InitializeAttributeDescs(nullptr);
+}
+
+FPCGDataDesc::FPCGDataDesc(const UPCGData* Data, const TMap<FPCGKernelAttributeKey, int32>& GlobalAttributeLookupTable)
+{
+	check(Data);
+
+	Type = Data->GetDataType();
+
+	if (Type == EPCGDataType::Point)
+	{
+		const UPCGPointData* PointData = CastChecked<UPCGPointData>(Data);
+		ElementCount = PointData->GetPoints().Num();
+	}
+	else { /* TODO: More types! */ }
+
+	InitializeAttributeDescs(Data->ConstMetadata(), GlobalAttributeLookupTable);
+}
+
+uint32 FPCGDataDesc::ComputePackedSize() const
+{
+	uint32 DataSizeBytes = 0;
+
+	if (Type == EPCGDataType::Point)
+	{
+		DataSizeBytes += POINT_DATA_HEADER_SIZE_BYTES;
+	}
+	else
+	{
+		// TODO: Support more types
+		checkNoEntry();
+	}
+
+	for (const FPCGKernelAttributeDesc& AttributeDesc : AttributeDescs)
+	{
+		DataSizeBytes += PCGDataForGPUHelpers::GetAttributeTypeStrideBytes(AttributeDesc.Type) * ElementCount;
+	}
+
+	return DataSizeBytes;
+}
+
+void FPCGDataDesc::InitializeAttributeDescs(const UPCGMetadata* Metadata, const TMap<FPCGKernelAttributeKey, int32>& GlobalAttributeLookupTable)
+{
+	if (Type == EPCGDataType::Point)
+	{
+		AttributeDescs.Append(PCGDataForGPUConstants::PointPropertyDescs, NUM_POINT_PROPERTIES);
+	}
+	else { /* TODO: More types! */ }
+
+	if (Metadata)
+	{
+		TArray<FName> AttributeNames;
+		TArray<EPCGMetadataTypes> AttributeTypes;
+		Metadata->GetAttributes(AttributeNames, AttributeTypes);
+
+		TArray<FPCGKernelAttributeKey> DelayedAttributeKeys; // Attribute keys that don't exist in the global lookup table must be delayed so we can append them at the end.
+		int NumAttributesFromLUT = 0; // Keep track of how many attributes come from the LUT. This will help give us the starting index for our delayed attributes.
+
+		for (int CustomAttributeIndex = 0; CustomAttributeIndex < AttributeNames.Num(); ++CustomAttributeIndex)
+		{
+			const FName AttributeName = AttributeNames[CustomAttributeIndex];
+			const EPCGKernelAttributeType AttributeType = PCGDataForGPUHelpers::GetAttributeTypeFromMetadataType(AttributeTypes[CustomAttributeIndex]);
+			const FPCGKernelAttributeKey AttributeKey = { AttributeType, AttributeName };
+
+			// Ignore excess attributes.
+			if (CustomAttributeIndex >= MAX_NUM_CUSTOM_ATTRS)
+			{
+				// TODO: Would be nice to include the pin label for debug purposes
+				UE_LOG(LogPCG, Warning, TEXT("Attempted to exceed max number of custom attributes (%d). Additional attributes will be ignored."), MAX_NUM_CUSTOM_ATTRS);
+				break;
+			}
+
+			if (AttributeName == NAME_None)
+			{
+				// TODO: Would be nice to include the pin label for debug purposes
+				UE_LOG(LogPCG, Warning, TEXT("'%s' is not a valid name for a kernel attribute. Attribute will be skipped."), *AttributeName.ToString());
+				continue;
+			}
+
+			if (const int32* AttributeId = GlobalAttributeLookupTable.Find(AttributeKey))
+			{
+				AttributeDescs.Emplace(*AttributeId, AttributeType, AttributeName);
+				++NumAttributesFromLUT;
+			}
+			else
+			{
+				DelayedAttributeKeys.Add(AttributeKey);
+			}
+		}
+
+		for (int DelayedAttributeIndex = 0; DelayedAttributeIndex < DelayedAttributeKeys.Num(); ++DelayedAttributeIndex)
+		{
+			const FPCGKernelAttributeKey& AttributeKey = DelayedAttributeKeys[DelayedAttributeIndex];
+			AttributeDescs.Emplace(NUM_RESERVED_ATTRS + DelayedAttributeIndex + NumAttributesFromLUT, AttributeKey.Type, AttributeKey.Name);
+		}
+	}
+}
+
+FPCGDataCollectionDesc FPCGDataCollectionDesc::BuildFromInputDataCollectionAndInputPinLabel(
+	const FPCGDataCollection& InDataCollection,
+	FName InputPinLabel,
+	const TMap<FPCGKernelAttributeKey, int32>& InAttributeLookupTable)
+{
+	FPCGDataCollectionDesc Desc;
+
+	TArray<FPCGTaggedData> DataForPin = InDataCollection.GetInputsByPin(InputPinLabel);
+
+	for (const FPCGTaggedData& Data : DataForPin)
+	{
+		// Only support Point data right now
+		if (!Data.Data || Data.Data->GetDataType() != EPCGDataType::Point)
+		{
+			continue;
+		}
+
+		Desc.DataDescs.Emplace(Data.Data, InAttributeLookupTable);
+	}
+
+	return Desc;
+}
+
+uint32 FPCGDataCollectionDesc::ComputePackedSize(TArray<uint32>* OutDataAddresses) const
+{
+	TRACE_CPUPROFILER_EVENT_SCOPE(FPCGDataCollectionDesc::ComputePackedSize);
+
+	const int NumData = DataDescs.Num();
+	// Calculation: sizeof(NumDatas) + (sizeof(DataAddress) * NumData)
+	const uint32 CollectionHeaderSizeBytes = sizeof(uint32) + (sizeof(uint32) * NumData);
+	uint32 TotalCollectionSizeBytes = CollectionHeaderSizeBytes;
+
+	if (OutDataAddresses)
+	{
+		OutDataAddresses->SetNumUninitialized(NumData);
+	}
+
+	for (int DataIndex = 0; DataIndex < NumData; ++DataIndex)
+	{
+		const FPCGDataDesc& DataDesc = DataDescs[DataIndex];
+		const uint32 DataSize = DataDesc.ComputePackedSize();
+
+		if (OutDataAddresses)
+		{
+			(*OutDataAddresses)[DataIndex] = TotalCollectionSizeBytes;
+		}
+
+		TotalCollectionSizeBytes += DataSize;
+	}
+
+	return TotalCollectionSizeBytes;
+}
+
+void FPCGDataCollectionDesc::PackDataCollection(const FPCGDataCollection& InDataCollection, FName InPin, TArray<uint32>& OutPackedDataCollection) const
+{
+	const TArray<FPCGTaggedData> InputData = InDataCollection.GetInputsByPin(InPin);
+	const uint32 NumData = InputData.Num();
+
+	TArray<uint32> DataAddresses;
+	const uint32 PackedDataCollectionSizeBytes = ComputePackedSize(&DataAddresses);
+
+	OutPackedDataCollection.SetNumZeroed(PackedDataCollectionSizeBytes / sizeof(uint32));
+	OutPackedDataCollection[0] = NumData;
+
+	for (uint32 DataIndex = 0; DataIndex < NumData; ++DataIndex)
+	{
+		TRACE_CPUPROFILER_EVENT_SCOPE(FPCGDataCollectionDesc::PackDataItem);
+
+		// Write the data addresses
+		const uint32 CurrentDataAddress = DataAddresses[DataIndex];
+		const uint32 CurrentDataIndex = CurrentDataAddress / sizeof(uint32);
+		OutPackedDataCollection[DataIndex + 1] = CurrentDataAddress;
+
+		// DataHeader: (TypeId, NumAttrs, AttrHeaderStartOffset, TypeInfo), Attr0 Header, Attr1 Header, ..., Attr255 Header
+		// Data: Attr0, Attr1, ...
+		if (const UPCGPointData* PointData = Cast<UPCGPointData>(InputData[DataIndex].Data))
+		{
+			const UPCGMetadata* Metadata = PointData->ConstMetadata();
+			const TArray<FPCGPoint>& Points = PointData->GetPoints();
+			const uint32 NumElements = Points.Num();
+
+			const TArray<FPCGKernelAttributeDesc>& AttributeDescs = DataDescs[DataIndex].AttributeDescs;
+			const uint32 NumAttributes = AttributeDescs.Num();
+
+			OutPackedDataCollection[CurrentDataIndex + 0] = /*PointDataTypeId=*/POINT_DATA_TYPE_ID;
+			OutPackedDataCollection[CurrentDataIndex + 1] = NumAttributes;
+			OutPackedDataCollection[CurrentDataIndex + 2] = POINT_DATA_HEADER_PREAMBLE_SIZE_BYTES;
+			OutPackedDataCollection[CurrentDataIndex + 3] = NumElements; // TypeInfo for PointData is just NumPoints
+
+			const uint32 BaseAttributeHeaderAddress = CurrentDataAddress + POINT_DATA_HEADER_PREAMBLE_SIZE_BYTES;
+			uint32 CurrentAttributeAddress = CurrentDataAddress + POINT_DATA_HEADER_SIZE_BYTES;
+
+			for (const FPCGKernelAttributeDesc& AttributeDesc : AttributeDescs)
+			{
+				const uint32 AttributeId = AttributeDesc.Index;
+				const uint32 AttributeStrideBytes = PCGDataForGPUHelpers::GetAttributeTypeStrideBytes(AttributeDesc.Type);
+				const uint32 AttributeNumComponents = AttributeStrideBytes / sizeof(uint32); // E.g. float3 has 3 components
+
+				// Pack Position (24 bits for AttributeId, 8 bits for Stride)
+				const uint32 PackedIdAndStride = (AttributeId << 8) + AttributeStrideBytes;
+				const uint32 AttributeIndex = CurrentAttributeAddress / sizeof(uint32);
+
+				const uint32 AttributeHeaderIndex = (BaseAttributeHeaderAddress + AttributeId * ATTRIBUTE_HEADER_SIZE_BYTES) / sizeof(uint32);
+				OutPackedDataCollection[AttributeHeaderIndex + 0] = PackedIdAndStride;
+				OutPackedDataCollection[AttributeHeaderIndex + 1] = CurrentAttributeAddress;
+
+				const FPCGMetadataAttributeBase* AttributeBase = (AttributeId >= NUM_RESERVED_ATTRS) ? Metadata->GetConstAttribute(AttributeDesc.Name) : nullptr;
+
+				for (uint32 ElementIndex = 0; ElementIndex < NumElements; ++ElementIndex)
+				{
+					const uint32 PackedDataElementIndex = AttributeIndex + (ElementIndex * AttributeNumComponents);
+
+					if (AttributeBase) // Pack attribute
+					{
+						ensure(PCGDataForGPUHelpers::PackAttributeHelper(AttributeBase, AttributeDesc, Points[ElementIndex].MetadataEntry, OutPackedDataCollection, PackedDataElementIndex));
+					}
+					else // Pack property
+					{
+						switch (AttributeId)
+						{
+						case POINT_POSITION_ATTRIBUTE_ID:
+						{
+							const FVector Position = Points[ElementIndex].Transform.GetLocation();
+
+							OutPackedDataCollection[PackedDataElementIndex + 0] = FMath::AsUInt(static_cast<float>(Position.X));
+							OutPackedDataCollection[PackedDataElementIndex + 1] = FMath::AsUInt(static_cast<float>(Position.Y));
+							OutPackedDataCollection[PackedDataElementIndex + 2] = FMath::AsUInt(static_cast<float>(Position.Z));
+							break;
+						}
+						case POINT_ROTATION_ATTRIBUTE_ID:
+						{
+							const FQuat Rotation = Points[ElementIndex].Transform.GetRotation();
+
+							OutPackedDataCollection[PackedDataElementIndex + 0] = FMath::AsUInt(static_cast<float>(Rotation.X));
+							OutPackedDataCollection[PackedDataElementIndex + 1] = FMath::AsUInt(static_cast<float>(Rotation.Y));
+							OutPackedDataCollection[PackedDataElementIndex + 2] = FMath::AsUInt(static_cast<float>(Rotation.Z));
+							OutPackedDataCollection[PackedDataElementIndex + 3] = FMath::AsUInt(static_cast<float>(Rotation.W));
+							break;
+						}
+						case POINT_SCALE_ATTRIBUTE_ID:
+						{
+							const FVector Scale = Points[ElementIndex].Transform.GetScale3D();
+
+							OutPackedDataCollection[PackedDataElementIndex + 0] = FMath::AsUInt(static_cast<float>(Scale.X));
+							OutPackedDataCollection[PackedDataElementIndex + 1] = FMath::AsUInt(static_cast<float>(Scale.Y));
+							OutPackedDataCollection[PackedDataElementIndex + 2] = FMath::AsUInt(static_cast<float>(Scale.Z));
+							break;
+						}
+						case POINT_BOUNDS_MIN_ATTRIBUTE_ID:
+						{
+							const FVector& BoundsMin = Points[ElementIndex].BoundsMin;
+
+							OutPackedDataCollection[PackedDataElementIndex + 0] = FMath::AsUInt(static_cast<float>(BoundsMin.X));
+							OutPackedDataCollection[PackedDataElementIndex + 1] = FMath::AsUInt(static_cast<float>(BoundsMin.Y));
+							OutPackedDataCollection[PackedDataElementIndex + 2] = FMath::AsUInt(static_cast<float>(BoundsMin.Z));
+							break;
+						}
+						case POINT_BOUNDS_MAX_ATTRIBUTE_ID:
+						{
+							const FVector& BoundsMax = Points[ElementIndex].BoundsMax;
+
+							OutPackedDataCollection[PackedDataElementIndex + 0] = FMath::AsUInt(static_cast<float>(BoundsMax.X));
+							OutPackedDataCollection[PackedDataElementIndex + 1] = FMath::AsUInt(static_cast<float>(BoundsMax.Y));
+							OutPackedDataCollection[PackedDataElementIndex + 2] = FMath::AsUInt(static_cast<float>(BoundsMax.Z));
+							break;
+						}
+						case POINT_COLOR_ATTRIBUTE_ID:
+						{
+							const FVector4& Color = Points[ElementIndex].Color;
+
+							OutPackedDataCollection[PackedDataElementIndex + 0] = FMath::AsUInt(static_cast<float>(Color.X));
+							OutPackedDataCollection[PackedDataElementIndex + 1] = FMath::AsUInt(static_cast<float>(Color.Y));
+							OutPackedDataCollection[PackedDataElementIndex + 2] = FMath::AsUInt(static_cast<float>(Color.Z));
+							OutPackedDataCollection[PackedDataElementIndex + 3] = FMath::AsUInt(static_cast<float>(Color.W));
+							break;
+						}
+						case POINT_DENSITY_ATTRIBUTE_ID:
+						{
+							const float Density = Points[ElementIndex].Density;
+
+							OutPackedDataCollection[PackedDataElementIndex + 0] = FMath::AsUInt(Density);
+							break;
+						}
+						case POINT_SEED_ATTRIBUTE_ID:
+						{
+							const int Seed = Points[ElementIndex].Seed;
+
+							OutPackedDataCollection[PackedDataElementIndex + 0] = Seed;
+							break;
+						}
+						case POINT_STEEPNESS_ATTRIBUTE_ID:
+						{
+							const float Steepness = Points[ElementIndex].Steepness;
+
+							OutPackedDataCollection[PackedDataElementIndex + 0] = FMath::AsUInt(Steepness);
+							break;
+						}
+						default:
+							checkNoEntry();
+							break;
+						}
+					}
+				}
+
+				CurrentAttributeAddress += NumElements * AttributeNumComponents * 4;
+			}
+		}
+		else { /* TODO: Support non-point data. */ }
+	}
+}
+
+void FPCGDataCollectionDesc::UnpackDataCollection(const TArray<uint8>& InPackedData, FName InPin, FPCGDataCollection& OutDataCollection) const
+{
+	const void* PackedData = InPackedData.GetData();
+	const float* DataAsFloat = static_cast<const float*>(PackedData);
+	const uint32* DataAsUint = static_cast<const uint32*>(PackedData);
+	const int* DataAsInt = static_cast<const int*>(PackedData);
+
+	const uint32 NumPackedFloats = InPackedData.Num() / 4;
+	const uint32 NumData = DataAsUint[0];
+
+	check(NumData == DataDescs.Num());
+
+	TArray<FPCGTaggedData>& OutData = OutDataCollection.TaggedData;
+
+	for (uint32 DataIndex = 0; DataIndex < NumData; ++DataIndex)
+	{
+		TRACE_CPUPROFILER_EVENT_SCOPE(FPCGDataCollectionDesc::UnpackDataItem);
+
+		const uint32 CurrentDataAddress = DataAsUint[DataIndex + 1];
+		const uint32 CurrentDataIndex = CurrentDataAddress / sizeof(uint32);
+		const uint32 TypeId =                      DataAsUint[CurrentDataIndex + 0];
+		const uint32 NumAttributes =               DataAsUint[CurrentDataIndex + 1];
+		const uint32 DataHeaderPreambleSizeBytes = DataAsUint[CurrentDataIndex + 2];
+		const uint32 NumElements =                 DataAsUint[CurrentDataIndex + 3];
+
+		const TArray<FPCGKernelAttributeDesc>& AttributeDescs = DataDescs[DataIndex].AttributeDescs;
+		check(NumAttributes == AttributeDescs.Num());
+
+		if (TypeId == POINT_DATA_TYPE_ID)
+		{
+			UPCGPointData* OutPointData = nullptr;
+			{
+				TRACE_CPUPROFILER_EVENT_SCOPE(InitalizeOutput);
+				OutPointData = NewObject<UPCGPointData>();
+				OutPointData->GetMutablePoints().SetNumUninitialized(NumElements);
+			}
+
+			UPCGMetadata* Metadata = OutPointData->MutableMetadata();
+			TArray<FPCGPoint>& OutPoints = OutPointData->GetMutablePoints();
+
+			{
+				// This can take an age as it populates the memory hierarchy.
+				TRACE_CPUPROFILER_EVENT_SCOPE(MetadataEntry);
+				ParallelFor(NumElements, [&OutPoints](int32 ElementIndex)
+				{
+					OutPoints[ElementIndex].MetadataEntry = -1;
+				});
+			}
+
+			FPCGTaggedData& OutTaggedData = OutData.Emplace_GetRef();
+			OutTaggedData.Data = OutPointData;
+			OutTaggedData.Pin = InPin;
+
+			const uint32 AttributeHeadersIndex = CurrentDataIndex + POINT_DATA_HEADER_PREAMBLE_SIZE_BYTES / sizeof(uint32);
+
+			// Loop over attributes.
+			for (const FPCGKernelAttributeDesc& AttributeDesc : AttributeDescs)
+			{
+				TRACE_CPUPROFILER_EVENT_SCOPE(WriteAttribute);
+
+				const uint32 AttributeId = AttributeDesc.Index;
+				const uint32 AttributeNumComponents = PCGDataForGPUHelpers::GetAttributeTypeStrideBytes(AttributeDesc.Type) / sizeof(uint32);
+				const uint32 AttributeHeaderIndex = AttributeHeadersIndex + AttributeDesc.Index * ATTRIBUTE_HEADER_SIZE_BYTES / sizeof(uint32);
+				const uint32 AttributeIndex = DataAsUint[AttributeHeaderIndex + 1] / sizeof(uint32);
+
+				FPCGMetadataAttributeBase* AttributeBase = (AttributeId >= NUM_RESERVED_ATTRS) ? PCGDataForGPUHelpers::CreateAttributeFromAttributeDesc(Metadata, AttributeDesc) : nullptr;
+
+				// 2. Parse each element in the attribute.
+
+				ParallelFor(NumElements, [&](int32 ElementIndex)
+				{
+					const uint32 PackedDataElementIndex = AttributeIndex + ElementIndex * AttributeNumComponents;
+					check(PackedDataElementIndex + AttributeNumComponents <= NumPackedFloats);
+
+					if (AttributeBase) // Unpack attribute
+					{
+						Metadata->InitializeOnSet(OutPoints[ElementIndex].MetadataEntry);
+						ensure(PCGDataForGPUHelpers::UnpackAttributeHelper(PackedData, PackedDataElementIndex, AttributeBase, AttributeDesc, OutPoints[ElementIndex].MetadataEntry));
+					}
+					else // Unpack property
+					{
+						// We tried hoisting this decision to a lambda but it didn't appear to help.
+						switch (AttributeId)
+						{
+						case POINT_POSITION_ATTRIBUTE_ID:
+						{
+							const FVector Location = FVector
+							(
+								DataAsFloat[PackedDataElementIndex + 0],
+								DataAsFloat[PackedDataElementIndex + 1],
+								DataAsFloat[PackedDataElementIndex + 2]
+							);
+
+							OutPoints[ElementIndex].Transform.SetLocation(Location);
+							break;
+						}
+						case POINT_ROTATION_ATTRIBUTE_ID:
+						{
+							FQuat Rotation = FQuat
+							(
+								DataAsFloat[PackedDataElementIndex + 0],
+								DataAsFloat[PackedDataElementIndex + 1],
+								DataAsFloat[PackedDataElementIndex + 2],
+								DataAsFloat[PackedDataElementIndex + 3]
+							);
+
+							// Normalize here with default tolerance (zero quat will return identity).
+							OutPoints[ElementIndex].Transform.SetRotation(Rotation.GetNormalized());
+							break;
+						}
+						case POINT_SCALE_ATTRIBUTE_ID:
+						{
+							const FVector Scale = FVector
+							(
+								DataAsFloat[PackedDataElementIndex + 0],
+								DataAsFloat[PackedDataElementIndex + 1],
+								DataAsFloat[PackedDataElementIndex + 2]
+							);
+
+							OutPoints[ElementIndex].Transform.SetScale3D(Scale);
+							break;
+						}
+						case POINT_BOUNDS_MIN_ATTRIBUTE_ID:
+						{
+							const FVector BoundsMin = FVector
+							(
+								DataAsFloat[PackedDataElementIndex + 0],
+								DataAsFloat[PackedDataElementIndex + 1],
+								DataAsFloat[PackedDataElementIndex + 2]
+							);
+
+							OutPoints[ElementIndex].BoundsMin = BoundsMin;
+							break;
+						}
+						case POINT_BOUNDS_MAX_ATTRIBUTE_ID:
+						{
+							const FVector BoundsMax = FVector
+							(
+								DataAsFloat[PackedDataElementIndex + 0],
+								DataAsFloat[PackedDataElementIndex + 1],
+								DataAsFloat[PackedDataElementIndex + 2]
+							);
+
+							OutPoints[ElementIndex].BoundsMax = BoundsMax;
+							break;
+						}
+						case POINT_COLOR_ATTRIBUTE_ID:
+						{
+							const FVector4 Color = FVector4
+							(
+								DataAsFloat[PackedDataElementIndex + 0],
+								DataAsFloat[PackedDataElementIndex + 1],
+								DataAsFloat[PackedDataElementIndex + 2],
+								DataAsFloat[PackedDataElementIndex + 2]
+							);
+
+							OutPoints[ElementIndex].Color = Color;
+							break;
+						}
+						case POINT_DENSITY_ATTRIBUTE_ID:
+						{
+							const float Density = DataAsFloat[PackedDataElementIndex];
+
+							OutPoints[ElementIndex].Density = Density;
+							break;
+						}
+						case POINT_SEED_ATTRIBUTE_ID:
+						{
+							const int Seed = DataAsInt[PackedDataElementIndex];
+
+							OutPoints[ElementIndex].Seed = Seed;
+							break;
+						}
+						case POINT_STEEPNESS_ATTRIBUTE_ID:
+						{
+							const float Steepness = DataAsFloat[PackedDataElementIndex];
+
+							OutPoints[ElementIndex].Steepness = Steepness;
+							break;
+						}
+						default:
+							checkNoEntry();
+							break;
+						}
+					}
+				});
+			}
+		}
+		else { /* TODO: Support non-point data. */ }
+	}
+}
+
+uint32 FPCGDataCollectionDesc::ComputeDataElementCount(EPCGDataType InDataType) const
+{
+	uint32 ElementCount = 0;
+
+	for (const FPCGDataDesc& DataDesc : DataDescs)
+	{
+		if (DataDesc.Type == InDataType)
+		{
+			ElementCount += DataDesc.ElementCount;
+		}
+	}
+
+	return ElementCount;
+}
+
+void FPCGDataCollectionDesc::Combine(const FPCGDataCollectionDesc& Other)
+{
+	DataDescs.Append(Other.DataDescs);
+}
diff --git a/Engine/Plugins/PCG/Source/PCG/Private/Graph/PCGGraphCompiler.cpp b/Engine/Plugins/PCG/Source/PCG/Private/Graph/PCGGraphCompiler.cpp
--- a/Engine/Plugins/PCG/Source/PCG/Private/Graph/PCGGraphCompiler.cpp
+++ b/Engine/Plugins/PCG/Source/PCG/Private/Graph/PCGGraphCompiler.cpp
@@ -2,19 +2,35 @@
 
 #include "Graph/PCGGraphCompiler.h"
 
+#include "PCGComponent.h"
 #include "PCGEdge.h"
 #include "PCGGraph.h"
 #include "PCGInputOutputSettings.h"
 #include "PCGModule.h"
 #include "PCGPin.h"
 #include "PCGSubgraph.h"
+#include "Compute/PCGComputeCommon.h"
+#include "Compute/PCGComputeKernelSource.h"
+#include "Compute/DataInterfaces/PCGComputeDataInterface.h"
+#include "Compute/DataInterfaces/PCGCustomKernelDataInterface.h"
+#include "Compute/DataInterfaces/PCGDataCollectionDataInterface.h"
+#include "Compute/DataInterfaces/PCGDataCollectionReadbackDataInterface.h"
+#include "Compute/DataInterfaces/PCGDataCollectionUploadDataInterface.h"
+#include "Compute/DataInterfaces/PCGDebugDataInterface.h"
+#include "Compute/DataInterfaces/PCGLandscapeDataInterface.h"
+#include "Compute/DataInterfaces/PCGTextureDataInterface.h"
+#include "Compute/Elements/PCGComputeGraphElement.h"
+#include "Compute/Elements/PCGCustomHLSL.h"
+#include "Elements/PCGGather.h"
 #include "Elements/PCGHiGenGridSize.h"
 #include "Elements/PCGReroute.h"
 #include "Graph/PCGGraphExecutor.h"
 #include "Graph/PCGPinDependencyExpression.h"
 
+#include "ComputeFramework/ComputeKernel.h"
 #include "HAL/IConsoleManager.h"
 #include "Misc/ScopeRWLock.h"
+#include "Shader/ShaderTypes.h"
 
 namespace PCGGraphCompiler
 {
@@ -22,6 +38,18 @@ namespace PCGGraphCompiler
 		TEXT("pcg.GraphExecution.TaskStaticCulling"),
 		true,
 		TEXT("Enable static culling of tasks which considers static branches, generation grid size, trivial nodes and more."));
+
+	TAutoConsoleVariable<bool> CVarEnableGPUExecution(
+		TEXT("pcg.GraphExecution.GPU.Enable"),
+		true,
+		TEXT("Whether to emit compatible nodes as compute graphs to execute on the GPU."));
+
+#if WITH_EDITOR
+	TAutoConsoleVariable<bool> CVarEnableGPUDebugging(
+		TEXT("pcg.GraphExecution.GPU.EnableDebugging"),
+		false,
+		TEXT("Enable verbose logging of GPU compilation and execution."));
+#endif
 }
 
 TArray<FPCGGraphTask> FPCGGraphCompiler::CompileGraph(UPCGGraph* InGraph, FPCGTaskId& NextId, FPCGStackContext& InOutStackContext)
@@ -1047,6 +1075,32 @@ void FPCGGraphCompiler::CompileTopGraph(UPCGGraph* InGraph, uint32 GenerationGri
 		}
 	}
 
+	// TODO: GPU execution is editor only for now. Need to cook kernels for standalone.
+#if WITH_EDITOR
+	if (PCGGraphCompiler::CVarEnableGPUExecution.GetValueOnAnyThread())
+	{
+		CreateGPUNodes(InGraph, CompiledTasks);
+	}
+	else
+#endif
+	{
+		// GPU not supported. Cull all the GPU compatible nodes without adding passthrough wires. This is destructive!
+		// TODO: In the future when we have nodes that can target both CPU and GPU, this could just force CPU (and throw
+		// errors if nodes strictly require GPU).
+		int CountBefore = CompiledTasks.Num();
+
+		CullTasks(CompiledTasks, /*bAddPassthroughWires=*/false, [](const FPCGGraphTask& InTask)
+		{
+			const UPCGSettings* Settings = InTask.Node ? InTask.Node->GetSettings() : nullptr;
+			return Settings && Settings->ShouldExecuteOnGPU();
+		});
+
+		if (CountBefore > CompiledTasks.Num())
+		{
+			UE_LOG(LogPCG, Warning, TEXT("One or more GPU nodes were culled from the graph!"));
+		}
+	}
+
 	// Post culling - remove any stacks that are no longer part of execution. Besides being tidy this also helps
 	// debug tools discern which stacks were executed or not.
 	PostCullStackCleanup(CompiledTasks, StackContext);
@@ -1149,6 +1203,27 @@ FPCGElementPtr FPCGGraphCompiler::GetSharedTrivialElement()
 	return SharedTrivialElement;
 }
 
+FPCGElementPtr FPCGGraphCompiler::GetSharedGatherElement()
+{
+	{
+		FReadScopeLock Lock(SharedGatherElementLock);
+
+		if (SharedGatherElement)
+		{
+			return SharedGatherElement;
+		}
+	}
+
+	FWriteScopeLock Lock(SharedGatherElementLock);
+
+	if (!SharedGatherElement)
+	{
+		SharedGatherElement = MakeShared<FPCGGatherElement>();
+	}
+
+	return SharedGatherElement;
+}
+
 FPCGElementPtr FPCGGraphCompiler::GetSharedTrivialPostGraphElement()
 {
 	{
@@ -1241,3 +1316,986 @@ void FPCGGraphCompiler::RemoveFromCacheRecursive(UPCGGraph* InGraph)
 	}
 }
 #endif // WITH_EDITOR
+
+void FPCGGraphCompiler::CollectGPUNodeSubsets(const TArray<FPCGGraphTask>& InCompiledTasks, const TMap<FPCGTaskId, TArray<FPCGTaskId>>& InTaskSuccessors, const TSet<FPCGTaskId>& InGPUCompatibleTaskIds, TArray<TSet<FPCGTaskId>>& OutNodeSubsetsToConvertToCFGraph)
+{
+	// Populate initial sets of tasks that are ready to consume vs ones that currently blocked.
+	TSet<FPCGTaskId> ReadyTaskIds;
+	TSet<FPCGTaskId> RemainingTaskIds;
+	ReadyTaskIds.Reserve(InCompiledTasks.Num());
+	RemainingTaskIds.Reserve(InCompiledTasks.Num());
+
+	for (FPCGTaskId TaskId = 0; TaskId < InCompiledTasks.Num(); ++TaskId)
+	{
+		if (InCompiledTasks[TaskId].Inputs.IsEmpty())
+		{
+			ReadyTaskIds.Add(TaskId);
+		}
+		else
+		{
+			RemainingTaskIds.Add(TaskId);
+		}
+	}
+
+	// Queue all successors of InTaskId that are ready to go (all upstream input tasks have been processed).
+	auto QueueSuccessors = [&InTaskSuccessors, &ReadyTaskIds, &RemainingTaskIds, &InCompiledTasks](int InTaskId)
+	{
+		bool bQueuedTask = false;
+
+		// Queue up any successors that are ready to go.
+		if (const TArray<FPCGTaskId>* Successors = InTaskSuccessors.Find(InTaskId))
+		{
+			for (FPCGTaskId Successor : *Successors)
+			{
+				const bool bSuccessorQueued = ReadyTaskIds.Contains(Successor);
+
+				// All successors should either be already queued, or waiting to be queued.
+				check(bSuccessorQueued || RemainingTaskIds.Contains(Successor));
+
+				if (!bSuccessorQueued)
+				{
+					bool bSuccessorReady = true;
+					for (const FPCGGraphTaskInput& Input : InCompiledTasks[Successor].Inputs)
+					{
+						if (ReadyTaskIds.Contains(Input.TaskId) || RemainingTaskIds.Contains(Input.TaskId))
+						{
+							bSuccessorReady = false;
+						}
+					}
+
+					if (bSuccessorReady)
+					{
+						ReadyTaskIds.Add(Successor);
+						RemainingTaskIds.Remove(Successor);
+						bQueuedTask = true;
+					}
+				}
+			}
+		}
+
+		return bQueuedTask;
+	};
+	
+	// Local to loops below but pulled out for performance.
+	TArray<FPCGTaskId> FoundReadyTaskIds;
+	TSet<FPCGTaskId> GPUSubsetTaskIds;
+	FoundReadyTaskIds.Reserve(InCompiledTasks.Num());
+	GPUSubsetTaskIds.Reserve(InCompiledTasks.Num());
+
+	// Build subsets of nodes that are GPU compatible and can be dispatched together.
+	while (!ReadyTaskIds.IsEmpty() || !RemainingTaskIds.IsEmpty())
+	{
+		// Consume as many CPU nodes as we can.
+		bool bQueuedTasks = true;
+		while (bQueuedTasks)
+		{
+			FoundReadyTaskIds.Reset();
+
+			for (FPCGTaskId ReadyTaskId : ReadyTaskIds)
+			{
+				if (!InGPUCompatibleTaskIds.Contains(ReadyTaskId))
+				{
+					FoundReadyTaskIds.Add(ReadyTaskId);
+				}
+			}
+
+			bQueuedTasks = false;
+
+			for (FPCGTaskId ReadyCPUTaskId : FoundReadyTaskIds)
+			{
+				ReadyTaskIds.Remove(ReadyCPUTaskId);
+				bQueuedTasks |= QueueSuccessors(ReadyCPUTaskId);
+			}
+		}
+
+		GPUSubsetTaskIds.Reset();
+
+		int StackIndex = INDEX_NONE;
+
+		// Now the opposite - consume as many GPU nodes as we can and accumulate them into a set that will be compiled into a compute graph.
+		bQueuedTasks = !ReadyTaskIds.IsEmpty();
+		while (bQueuedTasks)
+		{
+			FoundReadyTaskIds.Reset();
+
+			for (FPCGTaskId ReadyTaskId : ReadyTaskIds)
+			{
+				if (InGPUCompatibleTaskIds.Contains(ReadyTaskId))
+				{
+					// For now don't mix tasks from different execution stacks (in and out of subgraphs for instance) into one compute graph.
+					if ((StackIndex == INDEX_NONE) || InCompiledTasks[ReadyTaskId].StackIndex == StackIndex)
+					{
+						StackIndex = InCompiledTasks[ReadyTaskId].StackIndex;
+						FoundReadyTaskIds.Add(ReadyTaskId);
+					}
+				}
+			}
+
+			bQueuedTasks = false;
+
+			for (FPCGTaskId ReadyGPUTaskId : FoundReadyTaskIds)
+			{
+				GPUSubsetTaskIds.Add(ReadyGPUTaskId);
+				ReadyTaskIds.Remove(ReadyGPUTaskId);
+				bQueuedTasks |= QueueSuccessors(ReadyGPUTaskId);
+			}
+		}
+
+		if (!GPUSubsetTaskIds.IsEmpty())
+		{
+			OutNodeSubsetsToConvertToCFGraph.Add(MoveTemp(GPUSubsetTaskIds));
+		}
+	}
+}
+
+void FPCGGraphCompiler::CreateGatherTasksAtGPUInputs(const TSet<FPCGTaskId>& InGPUCompatibleTaskIds, TArray<FPCGGraphTask>& InOutCompiledTasks)
+{
+	using FOriginalInputPinKey = TPair<FPCGTaskId /* Original GPU task */, FName /* Input pin label */>;
+
+	// These are local to loop below but hoisted here for efficiency.
+	TSet<FOriginalInputPinKey> EncounteredInputPins; // TODO this is heavyweight, could use simple array.
+	TMap<FOriginalInputPinKey, FPCGTaskId /* Gather task */> InputPinToGatherTask;
+
+	// Add all compute graph task inputs and outputs.
+	for (FPCGTaskId GPUTaskId : InGPUCompatibleTaskIds)
+	{
+		EncounteredInputPins.Reset();
+		InputPinToGatherTask.Reset();
+
+		// First pass - create gather tasks for all original input pins which have more than one incident edge.
+		// This is so we can gather on the CPU (much more efficient than going it on the GPU).
+		for (int InputIndex = 0; InputIndex < InOutCompiledTasks[GPUTaskId].Inputs.Num(); ++InputIndex)
+		{
+			// Helper to get current input. We avoid simply taking a local reference as InOutCompiledTasks can be modified below.
+			auto CurrentInput = [&InOutCompiledTasks, InputIndex, GPUTaskId]() -> FPCGGraphTaskInput&
+			{
+				return InOutCompiledTasks[GPUTaskId].Inputs[InputIndex];
+			};
+
+			if (!CurrentInput().DownstreamPin.IsSet())
+			{
+				continue;
+			}
+
+			const FOriginalInputPinKey PinKey = { GPUTaskId, CurrentInput().DownstreamPin->Label };
+
+			// If already created a gather task, then nothing more to do for this pin.
+			if (const FPCGTaskId* GatherTaskId = InputPinToGatherTask.Find(PinKey))
+			{
+				continue;
+			}
+
+			// If we're encountering pin for first time, register it.
+			if (!EncounteredInputPins.Contains(PinKey))
+			{
+				EncounteredInputPins.Add(PinKey);
+				continue;
+			}
+
+			// Second time we've encountered this input pin - create a gather element because we need one edge connected to
+			// each virtual input pin, so that we can obtain the data items from the input data collection using the unique
+			// virtual pin label at execution time.
+			const FPCGTaskId GatherTaskId = InOutCompiledTasks.Num();
+			FPCGGraphTask& GatherTask = InOutCompiledTasks.Emplace_GetRef();
+			GatherTask.NodeId = GatherTaskId;
+			GatherTask.ParentId = InOutCompiledTasks[GPUTaskId].ParentId;
+			GatherTask.Element = GetSharedGatherElement();
+
+			InputPinToGatherTask.Add(PinKey, GatherTaskId);
+		}
+
+		EncounteredInputPins.Reset();
+
+		// Second pass - wire up the newly added gather tasks once we have the full picture of which edges are affected.
+		for (int InputIndex = 0; InputIndex < InOutCompiledTasks[GPUTaskId].Inputs.Num(); ++InputIndex)
+		{
+			auto CurrentInput = [&InOutCompiledTasks, InputIndex, GPUTaskId]() -> FPCGGraphTaskInput&
+			{
+				return InOutCompiledTasks[GPUTaskId].Inputs[InputIndex];
+			};
+
+			if (!CurrentInput().DownstreamPin.IsSet())
+			{
+				continue;
+			}
+
+			const FOriginalInputPinKey PinKey = { GPUTaskId, CurrentInput().DownstreamPin->Label };
+
+			if (const FPCGTaskId* GatherTaskId = InputPinToGatherTask.Find(PinKey))
+			{
+				// Wire the upstream output pin to the gather task.
+				FPCGGraphTaskInput& WireUpstreamNodeToGather = InOutCompiledTasks[*GatherTaskId].Inputs.Add_GetRef(CurrentInput());
+				if (WireUpstreamNodeToGather.DownstreamPin.IsSet())
+				{
+					WireUpstreamNodeToGather.DownstreamPin->Label = PCGPinConstants::DefaultInputLabel;
+				}
+
+				if (!EncounteredInputPins.Contains(PinKey))
+				{
+					// First time we're encountering this input pin, wire it to the gather task.
+					EncounteredInputPins.Add(PinKey);
+
+					CurrentInput().TaskId = *GatherTaskId;
+					if (CurrentInput().UpstreamPin.IsSet())
+					{
+						CurrentInput().UpstreamPin->Label = PCGPinConstants::DefaultOutputLabel;
+					}
+				}
+				else
+				{
+					// Input pin already encountered, already wired to gather task. Remove this input.
+					InOutCompiledTasks[GPUTaskId].Inputs.RemoveAt(InputIndex);
+					--InputIndex;
+				}
+			}
+		}
+	}
+}
+
+void FPCGGraphCompiler::WireGPUGraphNode(
+	FPCGTaskId InGPUGraphTaskId,
+	const TSet<FPCGTaskId>& InCollapsedTasks,
+	const TSet<FPCGTaskId>& InGPUCompatibleTaskIds,
+	TArray<FPCGGraphTask>& InOutCompiledTasks,
+	const TMap<FPCGTaskId, TArray<FPCGTaskId>>& InTaskSuccessors,
+	FOriginalToVirtualPin& OutOriginalToVirtualPin,
+	TMap<const UPCGPin*, FName>& OutOutputCPUPinToVirtualPin)
+{
+	FPCGGraphTask& GPUGraphTask = InOutCompiledTasks[InGPUGraphTaskId];
+
+	// Used to construct unique input/output labels, ultimately consumed in graph executor in BuildTaskInput and PostExecute for input/output respectively.
+	int InputCount = 0;
+	int OutputCount = 0;
+
+	// Add all compute graph task inputs and outputs.
+	for (FPCGTaskId GPUTaskId : InCollapsedTasks)
+	{
+		// First find CPU to GPU edges and wire in the GPU graph node inputs.
+		for (const FPCGGraphTaskInput& Input : InOutCompiledTasks[GPUTaskId].Inputs)
+		{
+			if (InGPUCompatibleTaskIds.Contains(Input.TaskId))
+			{
+				continue;
+			}
+
+			FPCGGraphTaskInput& AddedInput = GPUGraphTask.Inputs.Add_GetRef(Input);
+
+			// TODO is pinless fine with skipping?
+			if (AddedInput.DownstreamPin.IsSet())
+			{
+				const FName VirtualLabel = *FString::Format(TEXT("{0}-VirtualIn{1}"), { AddedInput.DownstreamPin->Label.ToString(), InputCount });
+				const bool bIsInputPin = true;
+				OutOriginalToVirtualPin.Add({ GPUTaskId, AddedInput.DownstreamPin->Label, bIsInputPin }, VirtualLabel);
+				AddedInput.DownstreamPin->Label = VirtualLabel;
+
+				++InputCount;
+
+				if (const UPCGNode* UpstreamNode = InOutCompiledTasks[Input.TaskId].Node)
+				{
+					if (const UPCGPin* OutputPin = UpstreamNode->GetOutputPin(AddedInput.UpstreamPin->Label))
+					{
+						OutOutputCPUPinToVirtualPin.Add(OutputPin, VirtualLabel);
+					}
+				}
+			}
+		}
+
+		if (!InTaskSuccessors.Contains(GPUTaskId))
+		{
+			continue;
+		}
+
+		// Next consider GPU to CPU edges to wire in the GPU graph node outputs.
+		for (FPCGTaskId Successor : InTaskSuccessors[GPUTaskId])
+		{
+			if (InGPUCompatibleTaskIds.Contains(Successor))
+			{
+				continue;
+			}
+
+			// Rewire inputs of this downstream CPU node to the outputs of the compute graph task.
+			FPCGGraphTask& DownstreamCPUNode = InOutCompiledTasks[Successor];
+
+			// Order matters here! We can never reorder inputs as it will impact execution.
+			const int InputCountBefore = DownstreamCPUNode.Inputs.Num();
+			for (int SuccessorInputIndex = 0; SuccessorInputIndex < InputCountBefore; ++SuccessorInputIndex)
+			{
+				// Implementation note: we modify the Inputs array in this loop, so don't take a reference to the current element.
+				
+				// Skip irrelevant edges.
+				if (DownstreamCPUNode.Inputs[SuccessorInputIndex].TaskId != GPUTaskId)
+				{
+					continue;
+				}
+
+				// Wire downstream CPU node to compute graph task.
+				FPCGGraphTaskInput InputCopy = DownstreamCPUNode.Inputs[SuccessorInputIndex];
+
+				InputCopy.TaskId = InGPUGraphTaskId;
+
+				if (DownstreamCPUNode.Inputs[SuccessorInputIndex].UpstreamPin.IsSet())
+				{
+					const FNodePin PinKey = { GPUTaskId, InputCopy.UpstreamPin->Label, /*Pin is input*/false };
+					if (const FName* FoundVirtualPinLabel = OutOriginalToVirtualPin.Find(PinKey))
+					{
+						// Wire to the existing virtual output pin.
+						InputCopy.UpstreamPin->Label = *FoundVirtualPinLabel;
+					}
+					else
+					{
+						const FName VirtualLabel = *FString::Format(TEXT("{0}-VirtualOut{1}"), { InputCopy.UpstreamPin->Label.ToString(), OutputCount });
+						OutOriginalToVirtualPin.Add(PinKey, VirtualLabel);
+
+						InputCopy.UpstreamPin->Label = VirtualLabel;
+
+						++OutputCount;
+					}
+				}
+
+				DownstreamCPUNode.Inputs.Add(MoveTemp(InputCopy));
+			}
+		}
+	}
+}
+
+void FPCGGraphCompiler::BuildGPUGraphTask(
+	UPCGGraph* InGraph,
+	FPCGTaskId InGPUGraphTaskId,
+	const TSet<FPCGTaskId>& InCollapsedTasks,
+	const TMap<FPCGTaskId, TArray<FPCGTaskId>>& InTaskSuccessors,
+	TArray<FPCGGraphTask>& InOutCompiledTasks,
+	const FOriginalToVirtualPin& InOriginalToVirtualPin,
+	const TMap<const UPCGPin*, FName>& InOutputCPUPinToVirtualPin)
+{
+	TSharedPtr<FPCGComputeGraphElement> Element = MakeShared<FPCGComputeGraphElement>();
+	InOutCompiledTasks[InGPUGraphTaskId].Element = Element;
+
+	const FName GraphName = MakeUniqueObjectName(InGraph, UPCGComputeGraph::StaticClass(), InGraph->GetFName());
+
+	UPCGComputeGraph* ComputeGraph = NewObject<UPCGComputeGraph>(InGraph, GraphName);
+	ComputeGraph->OutputCPUPinToInputGPUPinAlias = InOutputCPUPinToVirtualPin;
+	Element->Graph.Reset(ComputeGraph);
+
+	// Not incredibly useful for us - DG adds GetComponentSource()->GetComponentClass() object which allows it to bind at execution time by class.
+	// But execution code requires it currently.
+	Element->Graph->Bindings.Add(UPCGDataBinding::StaticClass());
+
+	// Create data interfaces which allow kernels to read or write data. Each data interface is associated with a node output pin.
+	// For CPU->GPU edges, an upload data interface is created. For GPU->CPU edges, a readback data interface is created.
+	auto CreateDataInterface = [&InCollapsedTasks, &InOutCompiledTasks, ComputeGraph](FPCGTaskId InTaskId, bool bRequiresReadback, const FPCGPinProperties& InOutputPinProperties) -> UPCGComputeDataInterface*
+	{
+		const bool bUpstreamIsGPUTask = InCollapsedTasks.Contains(InTaskId);
+
+		EPCGDataType PinType = InOutputPinProperties.AllowedTypes;
+
+		// Dynamically typed pins could have a different type
+		const UPCGSettings* Settings = InOutCompiledTasks[InTaskId].Node ? InOutCompiledTasks[InTaskId].Node->GetSettings() : nullptr;
+		const UPCGPin* Pin = InOutCompiledTasks[InTaskId].Node ? InOutCompiledTasks[InTaskId].Node->GetOutputPin(InOutputPinProperties.Label) : nullptr;
+		if (Settings && Pin)
+		{
+			PinType = Settings->GetCurrentPinTypes(Pin);
+		}
+
+		UPCGComputeDataInterface* DataInterface = nullptr;
+
+		switch (PinType)
+		{
+		case EPCGDataType::Point:
+		case EPCGDataType::PointOrParam:
+		{
+			UPCGDataCollectionDataInterface* DataInterfacePCGData = nullptr;
+
+			if (bUpstreamIsGPUTask)
+			{
+				if (bRequiresReadback)
+				{
+					// GPU -> CPU
+					DataInterfacePCGData = NewObject<UPCGDataCollectionReadbackDataInterface>(ComputeGraph);
+				}
+				else
+				{
+					// GPU -> GPU
+					DataInterfacePCGData = NewObject<UPCGDataCollectionDataInterface>(ComputeGraph);
+				}
+			}
+			else
+			{
+				// CPU -> GPU
+				DataInterfacePCGData = NewObject<UPCGDataCollectionUploadDataInterface>(ComputeGraph);
+			}
+
+			check(DataInterfacePCGData);
+			DataInterfacePCGData->ProducerSettings = InOutCompiledTasks[InTaskId].Node ? InOutCompiledTasks[InTaskId].Node->GetSettings() : nullptr;
+
+			DataInterface = DataInterfacePCGData;
+			break;
+		}
+		case EPCGDataType::Texture:
+		{
+			DataInterface = NewObject<UPCGTextureDataInterface>(ComputeGraph);
+			break;
+		}
+		case EPCGDataType::Landscape:
+		{
+			DataInterface = NewObject<UPCGLandscapeDataInterface>(ComputeGraph);
+			break;
+		}
+		default:
+			ensure(false);
+			break;
+		}
+
+		if (DataInterface)
+		{
+			DataInterface->SetOutputPin(InOutputPinProperties.Label);
+		}
+
+		return DataInterface;
+	};
+
+	TMap<TPair</* Node task*/FPCGTaskId, /*Node output pin*/FName>, UPCGComputeDataInterface*> OutputPinDataInterfaces;
+
+	// Create all the output data interfaces.
+	for (FPCGTaskId TaskId : InCollapsedTasks)
+	{
+		// Create DIs for all output pins, because the kernels currently need their outputs to be bound to valid resources.
+		if (const UPCGSettings* Settings = InOutCompiledTasks[TaskId].Node ? InOutCompiledTasks[TaskId].Node->GetSettings() : nullptr)
+		{
+			for (const FPCGPinProperties& OutputPinProperties : Settings->AllOutputPinProperties())
+			{
+				if (OutputPinDataInterfaces.Contains({ TaskId, OutputPinProperties.Label }))
+				{
+					ensure(false);
+					continue;
+				}
+
+				bool bRequiresReadback = false;
+				if (const TArray<FPCGTaskId>* Successors = InTaskSuccessors.Find(TaskId))
+				{
+					for (FPCGTaskId Successor : *Successors)
+					{
+						for (const FPCGGraphTaskInput& Input : InOutCompiledTasks[Successor].Inputs)
+						{
+							if (Input.UpstreamPin.IsSet() && (*Input.UpstreamPin == OutputPinProperties) && !InCollapsedTasks.Contains(Successor))
+							{
+								bRequiresReadback = true;
+								break;
+							}
+						}
+					}
+				}
+
+				if (UPCGComputeDataInterface* OutputDI = CreateDataInterface(TaskId, bRequiresReadback, OutputPinProperties))
+				{
+					OutputDI->SetOutputPin(OutputPinProperties.Label);
+					ComputeGraph->DataInterfaces.Add(OutputDI);
+					OutputPinDataInterfaces.Add({ TaskId, OutputPinProperties.Label }, OutputDI);
+
+					for (const FPCGKernelAttributeKey& Key : Settings->GetKernelAttributeKeys())
+					{
+						if (!ComputeGraph->GlobalAttributeLookupTable.Find(Key))
+						{
+							ComputeGraph->GlobalAttributeLookupTable.Add(Key, ComputeGraph->GlobalAttributeLookupTable.Num() + PCGComputeConstants::NUM_RESERVED_ATTRS);
+						}
+					}
+				}
+			}
+		}
+
+		// Create any DIs for upstream CPU nodes.
+		for (const FPCGGraphTaskInput& Input : InOutCompiledTasks[TaskId].Inputs)
+		{
+			// Only deal with upstream CPU tasks.
+			if (InCollapsedTasks.Contains(Input.TaskId))
+			{
+				continue;
+			}
+
+			if (!Input.DownstreamPin.IsSet())
+			{
+				continue;
+			}
+
+			if (const UPCGNode* DownstreamNode = InOutCompiledTasks[TaskId].Node)
+			{
+				if (const UPCGPin* InputPin = DownstreamNode->GetInputPin(Input.DownstreamPin->Label))
+				{
+					ComputeGraph->PinsReceivingDataFromCPU.Add(InputPin);
+				}
+			}
+
+			if (!Input.UpstreamPin.IsSet())
+			{
+				continue;
+			}
+
+			if (OutputPinDataInterfaces.Contains({ Input.TaskId, Input.UpstreamPin->Label }))
+			{
+				// Skip if already created.
+				continue;
+			}
+
+			if (UPCGComputeDataInterface* OutputDI = CreateDataInterface(Input.TaskId, /*bRequiresReadback=*/false, *Input.UpstreamPin))
+			{
+				OutputDI->SetOutputPin(Input.UpstreamPin->Label);
+				ComputeGraph->DataInterfaces.Add(OutputDI);
+				OutputPinDataInterfaces.Add({ Input.TaskId, Input.UpstreamPin->Label }, OutputDI);
+			}
+		}
+	}
+
+	TSet<FPCGTaskId> RemainingTasks = InCollapsedTasks;
+
+	while (!RemainingTasks.IsEmpty())
+	{
+		// Find a ready task
+		FPCGTaskId TaskId = InvalidPCGTaskId;
+		for (FPCGTaskId RemainingTask : RemainingTasks)
+		{
+			// TODO: use 'QueueSuccessors' pattern rather than brute force searching for ready tasks
+			bool bReady = true;
+			for (const FPCGGraphTaskInput& Input : InOutCompiledTasks[RemainingTask].Inputs)
+			{
+				if (RemainingTasks.Contains(Input.TaskId))
+				{
+					bReady = false;
+					break;
+				}
+			}
+
+			if (bReady)
+			{
+				TaskId = RemainingTask;
+				break;
+			}
+		}
+
+		if (TaskId == InvalidPCGTaskId)
+		{
+			// Tasks in RemainingTasks are unreachable?
+			ensure(false);
+			break;
+		}
+
+		RemainingTasks.Remove(TaskId);
+
+		const UPCGNode* Node = InOutCompiledTasks[TaskId].Node;
+
+		const UPCGCustomHLSLSettings* Settings = Cast<UPCGCustomHLSLSettings>(Node ? Node->GetSettings() : nullptr);
+		check(Settings && Settings->bEnabled && Settings->ShouldExecuteOnGPU());
+
+		// For every usage of a DI, get the original (non-aliased) pin label.
+		TMap<TPair<FPCGTaskId, UComputeDataInterface*>, FName> DataInterfaceUsageToPinLabel;
+
+		TArray<int> InputDataInterfaceIndices;
+		TArray<int> OutputDataInterfaceIndices;
+		InputDataInterfaceIndices.Reserve(Settings->InputPinProperties().Num());
+		OutputDataInterfaceIndices.Reserve(Settings->OutputPinProperties().Num());
+
+		// Add DIs (PCG -> CF transcoding).
+
+		for (const FPCGGraphTaskInput& Input : InOutCompiledTasks[TaskId].Inputs)
+		{
+			if (!Input.UpstreamPin.IsSet())
+			{
+				// Execution-only dependencies not supported currently. Unclear if this should ever be supported for GPU graphs.
+				// Writes followed by reads will be protected via barriers added by RDG.
+				continue;
+			}
+
+			UPCGComputeDataInterface* UpstreamDI = nullptr;
+			if (UPCGComputeDataInterface** FoundUpstreamDI = OutputPinDataInterfaces.Find({ Input.TaskId, Input.UpstreamPin->Label }))
+			{
+				UpstreamDI = *FoundUpstreamDI;
+			}
+
+			if (!UpstreamDI)
+			{
+				ensure(false);
+				continue;
+			}
+
+			const int Index = ComputeGraph->DataInterfaces.Find(UpstreamDI);
+			if (Index == INDEX_NONE)
+			{
+				ensure(false);
+				continue;
+			}
+
+			InputDataInterfaceIndices.Add(Index);
+
+			FName DownstreamInputPinLabel = Input.DownstreamPin->Label;
+
+			DataInterfaceUsageToPinLabel.Add({ TaskId, UpstreamDI }, DownstreamInputPinLabel);
+
+			const bool bIsInputPin = true;
+			UpstreamDI->AddDownstreamInputPin(DownstreamInputPinLabel, InOriginalToVirtualPin.Find({ TaskId, DownstreamInputPinLabel, bIsInputPin }));
+		}
+
+		// Always create a DI for every output pin, so kernel always has something to write to.
+		for (const FPCGPinProperties& OutputPinProperties : Settings->AllOutputPinProperties())
+		{
+			UPCGComputeDataInterface** FoundDI = OutputPinDataInterfaces.Find({ TaskId, OutputPinProperties.Label });
+			if (!ensure(FoundDI) || !ensure(*FoundDI))
+			{
+				continue;
+			}
+
+			const int Index = ComputeGraph->DataInterfaces.Find(*FoundDI);
+			if (Index == INDEX_NONE)
+			{
+				ensure(false);
+				continue;
+			}
+
+			OutputDataInterfaceIndices.Add(Index);
+			DataInterfaceUsageToPinLabel.Add({ TaskId, *FoundDI }, OutputPinProperties.Label);
+		}
+
+		// Make sure every downstream input pin is registered with the upstream DI.
+		if (const TArray<FPCGTaskId>* Successors = InTaskSuccessors.Find(TaskId))
+		{
+			for (FPCGTaskId Successor : *Successors)
+			{
+				for (const FPCGGraphTaskInput& SuccessorInput : InOutCompiledTasks[Successor].Inputs)
+				{
+					if (SuccessorInput.TaskId != TaskId)
+					{
+						continue;
+					}
+
+					if (!SuccessorInput.UpstreamPin.IsSet() || !SuccessorInput.DownstreamPin.IsSet())
+					{
+						continue;
+					}
+
+					FName OutputPinLabel = SuccessorInput.UpstreamPin->Label;
+
+					// DIs for output pins should have all been created.
+					UPCGComputeDataInterface** FoundOutputDI = OutputPinDataInterfaces.Find({ TaskId, OutputPinLabel });
+					if (!ensure(FoundOutputDI) || !ensure(*FoundOutputDI))
+					{
+						continue;
+					}
+					
+					// If the map has an entry for this task and output pin label, then its a GPU->CPU readback.
+					const bool bIsInputPin = false;
+					(*FoundOutputDI)->SetOutputPin(OutputPinLabel, InOriginalToVirtualPin.Find({ TaskId, OutputPinLabel, bIsInputPin }));
+				}
+			}
+		}
+
+		UPCGCustomKernelDataInterface* KernelDI = NewObject<UPCGCustomKernelDataInterface>(ComputeGraph);
+		KernelDI->Settings = Settings;
+		const int KernelDIIndex = ComputeGraph->DataInterfaces.Num();
+		ComputeGraph->DataInterfaces.Add(KernelDI);
+
+		// TODO add graph data interface (graph params). Reference: UOptimusGraphDataInterface.
+		//Element->Graph->DataInterfaces.Add(NewObject< UPCGDataCollectionDataInterface>(InGraph));
+		//Element->Graph->DataInterfaceToBinding.Add(0);
+
+		// TODO once we support cooking for different platforms/configs, don't create the interface if logging is not present.
+		int DebugDIIndex = INDEX_NONE;
+		if (Settings->bPrintShaderDebugValues)
+		{
+			UPCGDebugDataInterface* DebugDI = NewObject<UPCGDebugDataInterface>(ComputeGraph);
+			DebugDI->SetDebugBufferSize(Settings->DebugBufferSize);
+
+			DebugDIIndex = ComputeGraph->DataInterfaces.Num();
+			ComputeGraph->DataInterfaces.Add(DebugDI);
+		}
+
+		// Now that all data interfaces added, create the (trivial) binding mapping. All map to primary binding, index 0.
+		ComputeGraph->DataInterfaceToBinding.SetNumZeroed(ComputeGraph->DataInterfaces.Num());
+
+		struct FInterfaceBinding
+		{
+			const UComputeDataInterface* DataInterface;
+			int32 DataInterfaceBindingIndex;
+			FString BindingFunctionName;
+			FString BindingFunctionNamespace;
+		};
+
+		struct FKernelWithDataBindings
+		{
+			UComputeKernel* Kernel;
+			TArray<FInterfaceBinding> InputDataBindings;
+			TArray<FInterfaceBinding> OutputDataBindings;
+		};
+
+		FKernelWithDataBindings KernelWithBindings;
+
+		KernelWithBindings.Kernel = NewObject<UComputeKernel>(ComputeGraph);
+		const int KernelIndex = ComputeGraph->KernelInvocations.Num();
+		ComputeGraph->KernelInvocations.Add(KernelWithBindings.Kernel);
+		ComputeGraph->KernelToNode.Add(Node);
+		
+		auto SetupAllInputBindings = [&KernelWithBindings, ComputeGraph](int InDataInterfaceIndex)
+		{
+			const UComputeDataInterface* DataInterface = ComputeGraph->DataInterfaces[InDataInterfaceIndex];
+			TArray<FShaderFunctionDefinition> Functions;
+			DataInterface->GetSupportedInputs(Functions);
+
+			for (int FuncIndex = 0; FuncIndex < Functions.Num(); ++FuncIndex)
+			{
+				FInterfaceBinding& Binding = KernelWithBindings.InputDataBindings.Emplace_GetRef();
+				Binding.DataInterface = DataInterface;
+				Binding.BindingFunctionName = Functions[FuncIndex].Name;
+				Binding.BindingFunctionNamespace = TEXT("");
+				Binding.DataInterfaceBindingIndex = FuncIndex;
+			}
+		};
+
+		auto SetupAllOutputBindings = [&KernelWithBindings, ComputeGraph](int InDataInterfaceIndex)
+		{
+			const UComputeDataInterface* DataInterface = ComputeGraph->DataInterfaces[InDataInterfaceIndex];
+			TArray<FShaderFunctionDefinition> Functions;
+			DataInterface->GetSupportedOutputs(Functions);
+
+			for (int FuncIndex = 0; FuncIndex < Functions.Num(); ++FuncIndex)
+			{
+				FInterfaceBinding& Binding = KernelWithBindings.OutputDataBindings.Emplace_GetRef();
+				Binding.DataInterface = DataInterface;
+				Binding.BindingFunctionName = Functions[FuncIndex].Name;
+				Binding.BindingFunctionNamespace = TEXT("");
+				Binding.DataInterfaceBindingIndex = FuncIndex;
+			}
+		};
+
+		// Bind data interfaces.
+		for (int InputDataInterfaceIndex : InputDataInterfaceIndices)
+		{
+			SetupAllInputBindings(InputDataInterfaceIndex);
+		}
+
+		SetupAllInputBindings(KernelDIIndex);
+
+		for (int OutputDataInterfaceIndex : OutputDataInterfaceIndices)
+		{
+			SetupAllOutputBindings(OutputDataInterfaceIndex);
+		}
+
+		if (DebugDIIndex != INDEX_NONE)
+		{
+			SetupAllOutputBindings(DebugDIIndex);
+		}
+
+		{
+			UPCGComputeKernelSource* KernelSource = NewObject<UPCGComputeKernelSource>(KernelWithBindings.Kernel); // is outer to kernel fine?
+			KernelWithBindings.Kernel->KernelSource = KernelSource;
+			KernelSource->EntryPoint = Settings->GetKernelEntryPoint();
+			KernelSource->GroupSize = Settings->GetThreadGroupSize();
+
+			KernelSource->SetSource(Settings->GetCookedKernelSource(ComputeGraph->GlobalAttributeLookupTable));
+
+#if WITH_EDITOR
+			if (PCGGraphCompiler::CVarEnableGPUDebugging.GetValueOnAnyThread())
+			{
+				UE_LOG(LogPCG, Warning, TEXT("ATTRIBUTE LOOK-UP TABLE [%s]"), *Settings->GetDefaultNodeTitle().ToString());
+
+				for (const TPair<FPCGKernelAttributeKey, int32 /* Attribute Index */>& Pair : ComputeGraph->GlobalAttributeLookupTable)
+				{
+					const FString TypeString = UEnum::GetValueAsString(Pair.Key.Type);
+					const FString NameString = Pair.Key.Name.ToString();
+					const FString IndexString = FString::FromInt(Pair.Value);
+
+					UE_LOG(LogPCG, Warning, TEXT("%s: %s (%s)"), *IndexString, *NameString, *TypeString);
+				}
+			}
+#endif
+
+			// Add functions for external inputs/outputs which must be fulfilled by DIs
+			for (FInterfaceBinding& Binding : KernelWithBindings.InputDataBindings)
+			{
+				TArray<FShaderFunctionDefinition> Functions;
+				Binding.DataInterface->GetSupportedInputs(Functions);
+				check(Functions.IsValidIndex(Binding.DataInterfaceBindingIndex));
+
+				FShaderFunctionDefinition FuncDef = Functions[Binding.DataInterfaceBindingIndex];
+				for (FShaderParamTypeDefinition& ParamType : FuncDef.ParamTypes)
+				{
+					// Making sure parameter has type declaration generated
+					ParamType.ResetTypeDeclaration();
+				}
+
+				KernelSource->ExternalInputs.Emplace(FuncDef);
+			}
+
+			for (FInterfaceBinding& Binding : KernelWithBindings.OutputDataBindings)
+			{
+				TArray<FShaderFunctionDefinition> Functions;
+				Binding.DataInterface->GetSupportedOutputs(Functions);
+				check(Functions.IsValidIndex(Binding.DataInterfaceBindingIndex));
+
+				FShaderFunctionDefinition FuncDef = Functions[Binding.DataInterfaceBindingIndex];
+				for (FShaderParamTypeDefinition& ParamType : FuncDef.ParamTypes)
+				{
+					// Making sure parameter has type declaration generated
+					ParamType.ResetTypeDeclaration();
+				}
+
+				KernelSource->ExternalOutputs.Emplace(FuncDef);
+			}
+		}
+
+		auto AddAllEdgesForKernel = [&KernelWithBindings, ComputeGraph, TaskId, &DataInterfaceUsageToPinLabel](int32 InKernelIndex, bool bInEdgesAreInputs)
+		{
+			TArray<FInterfaceBinding>& Bindings = bInEdgesAreInputs ? KernelWithBindings.InputDataBindings : KernelWithBindings.OutputDataBindings;
+
+			// Add all graph edges for bindings. This is somewhat odd but likely what we'll do vs the more finegrained interface
+			// definitions in optimus, but we need to see.
+			for (int BindingIndex = 0; BindingIndex < Bindings.Num(); ++BindingIndex)
+			{
+				FInterfaceBinding& Binding = Bindings[BindingIndex];
+
+				FComputeGraphEdge& Edge = ComputeGraph->GraphEdges.Emplace_GetRef();
+				Edge.KernelIndex = InKernelIndex;
+				Edge.KernelBindingIndex = BindingIndex;
+				Edge.DataInterfaceIndex = ComputeGraph->DataInterfaces.IndexOfByPredicate([&Binding](const UComputeDataInterface* In) { return Binding.DataInterface == In; });
+				check(Edge.DataInterfaceIndex != INDEX_NONE);
+				Edge.DataInterfaceBindingIndex = Binding.DataInterfaceBindingIndex;
+				Edge.bKernelInput = bInEdgesAreInputs;
+
+				UComputeDataInterface* DataInterface = ComputeGraph->DataInterfaces[Edge.DataInterfaceIndex];
+				check(DataInterface);
+
+				if (FName* PinLabel = DataInterfaceUsageToPinLabel.Find({ TaskId, DataInterface }))
+				{
+					TArray<FShaderFunctionDefinition> DataInterfaceFunctions;
+					if (bInEdgesAreInputs)
+					{
+						DataInterface->GetSupportedInputs(DataInterfaceFunctions);
+					}
+					else
+					{
+						DataInterface->GetSupportedOutputs(DataInterfaceFunctions);
+					}
+
+					Edge.BindingFunctionNameOverride = FString::Format(
+						TEXT("{0}_{1}"),
+						{ PinLabel->ToString(), DataInterfaceFunctions[Edge.DataInterfaceBindingIndex].Name }
+					);
+				}
+			}
+		};
+
+		AddAllEdgesForKernel(KernelIndex, /*bInEdgesAreInputs=*/true);
+		AddAllEdgesForKernel(KernelIndex, /*bInEdgesAreInputs=*/false);
+	}
+
+	// Register all virtual pin aliases with the corresponding pins for downstream usage.
+	for (const TPair<TTuple<FPCGTaskId, FName, bool>, FName>& Mapping : InOriginalToVirtualPin)
+	{
+		const bool bIsInputPin = Mapping.Get<0>().Get<2>();
+		if (bIsInputPin)
+		{
+			const FPCGTaskId TaskId = Mapping.Get<0>().Get<0>();
+
+			if (const UPCGNode* Node = InOutCompiledTasks[TaskId].Node)
+			{
+				const FName& OriginalLabel = Mapping.Get<0>().Get<1>();
+
+				if (const UPCGPin* Pin = Node->GetInputPin(OriginalLabel))
+				{
+					const FName& VirtualLabel = Mapping.Get<1>();
+					ComputeGraph->InputPinLabelAliases.FindOrAdd(Pin) = VirtualLabel;
+				}
+			}
+		}
+	}
+
+	// Kick off shader compilation (if needed).
+	ComputeGraph->UpdateResources();
+}
+
+void FPCGGraphCompiler::CreateGPUNodes(UPCGGraph* InGraph, TArray<FPCGGraphTask>& InOutCompiledTasks)
+{
+	TSet<FPCGTaskId> GPUCompatibleTaskIds;
+	GPUCompatibleTaskIds.Reserve(InOutCompiledTasks.Num());
+	for (FPCGTaskId TaskId = 0; TaskId < InOutCompiledTasks.Num(); ++TaskId)
+	{
+		const UPCGNode* Node = InOutCompiledTasks[TaskId].Node;
+		const UPCGSettings* Settings = Node ? Node->GetSettings() : nullptr;
+		if (Settings && Settings->ShouldExecuteOnGPU() && Settings->bEnabled)
+		{
+			GPUCompatibleTaskIds.Add(TaskId);
+		}
+	}
+
+	if (GPUCompatibleTaskIds.IsEmpty())
+	{
+		// Nothing to do for this graph.
+		return;
+	}
+
+	TMap<FPCGTaskId, TArray<FPCGTaskId>> TaskSuccessors;
+	TaskSuccessors.Reserve(InOutCompiledTasks.Num());
+	for (FPCGTaskId TaskId = 0; TaskId < InOutCompiledTasks.Num(); ++TaskId)
+	{
+		for (int InputIndex = 0; InputIndex < InOutCompiledTasks[TaskId].Inputs.Num(); ++InputIndex)
+		{
+			TaskSuccessors.FindOrAdd(InOutCompiledTasks[TaskId].Inputs[InputIndex].TaskId).Add(TaskId);
+		}
+	}
+
+	// For input pins at CPU -> GPU boundary, inject gather elements to pre-combine data on CPU side
+	// before passing to GPU.
+	CreateGatherTasksAtGPUInputs(GPUCompatibleTaskIds, InOutCompiledTasks);
+
+	TArray<TSet<FPCGTaskId>> NodeSubsetsToConvertToCFGraph;
+	CollectGPUNodeSubsets(InOutCompiledTasks, TaskSuccessors, GPUCompatibleTaskIds, NodeSubsetsToConvertToCFGraph);
+
+	// Do actual collapsing now, one subset at a time. Each collapse will do all fixup of task ids? That will invalidate
+	// ids in NodeSubsetsToConvertToCFGraph, so may need remap table. But can ignore this for now.
+	for (TSet<FPCGTaskId>& NodeSubsetToConvertToCFGraph : NodeSubsetsToConvertToCFGraph)
+	{
+		if (NodeSubsetToConvertToCFGraph.IsEmpty())
+		{
+			ensure(false);
+			continue;
+		}
+
+		// Add a new compute graph task. Then kill the original GPU tasks.
+		const FPCGTaskId ComputeGraphTaskId = InOutCompiledTasks.Num();
+		FPCGGraphTask& ComputeGraphTask = InOutCompiledTasks.Emplace_GetRef();
+		ComputeGraphTask.NodeId = ComputeGraphTaskId;
+
+		// All nodes in subset will be from same stack/parent, so assign from any.
+		for (FPCGTaskId GPUTaskId : NodeSubsetToConvertToCFGraph)
+		{
+			ComputeGraphTask.ParentId = InOutCompiledTasks[GPUTaskId].ParentId;
+			ComputeGraphTask.StackIndex = InOutCompiledTasks[GPUTaskId].StackIndex;
+			break;
+		}
+
+		// Mapping from task ID & pin label to a virtual pin label. Compute graphs are executed within a generated element,
+		// and the input and output pins of this element must have unique virtual pin labels so that we can parse the data that
+		// PCG provides through the input data collection correctly, and route the output data to the downstream pins correctly.
+		FOriginalToVirtualPin OriginalToVirtualPin;
+
+		TMap<const UPCGPin*, FName> OutputCPUPinToVirtualPin;
+
+		// Wire in the compute graph task, side by side with the individual GPU tasks, which will be culled below.
+		WireGPUGraphNode(
+			ComputeGraphTaskId,
+			NodeSubsetToConvertToCFGraph,
+			GPUCompatibleTaskIds,
+			InOutCompiledTasks,
+			TaskSuccessors,
+			OriginalToVirtualPin,
+			OutputCPUPinToVirtualPin);
+
+		// Generate a compute graph from all of the individual GPU tasks.
+		BuildGPUGraphTask(
+			InGraph,
+			ComputeGraphTaskId,
+			NodeSubsetToConvertToCFGraph,
+			TaskSuccessors,
+			InOutCompiledTasks,
+			OriginalToVirtualPin,
+			OutputCPUPinToVirtualPin);
+	}
+
+	// Now cull all the GPU compatible nodes. The compute graph task are already wired in so we're fine to just delete.
+	CullTasks(InOutCompiledTasks, /*bAddPassthroughWires=*/false, [&GPUCompatibleTaskIds](const FPCGGraphTask& InTask) { return GPUCompatibleTaskIds.Contains(InTask.NodeId); });
+}
diff --git a/Engine/Plugins/PCG/Source/PCG/Private/Graph/PCGGraphCompiler.h b/Engine/Plugins/PCG/Source/PCG/Private/Graph/PCGGraphCompiler.h
--- a/Engine/Plugins/PCG/Source/PCG/Private/Graph/PCGGraphCompiler.h
+++ b/Engine/Plugins/PCG/Source/PCG/Private/Graph/PCGGraphCompiler.h
@@ -40,6 +40,46 @@ public:
 
 	/** Returns the post graph element, which is used to determine results caching behavior */
 	FPCGElementPtr GetSharedTrivialPostGraphElement();
+	
+	/** Used to track new unique virtual pins created on generated compute graph elements. */
+	using FNodePin = TTuple<FPCGTaskId, /*Pin label*/FName, /*Pin is input*/bool>;
+	using FOriginalToVirtualPin = TMap<FNodePin, /*Virtual pin label*/FName>;
+
+	/** Outputs sets of task IDs, where each set is GPU nodes that can be compiled into a compute graph and dispatched together. */
+	void CollectGPUNodeSubsets(
+		const TArray<FPCGGraphTask>& InCompiledTasks,
+		const TMap<FPCGTaskId, TArray<FPCGTaskId>>& InTaskSuccessors,
+		const TSet<FPCGTaskId>& InGPUCompatibleTaskIds,
+		TArray<TSet<FPCGTaskId>>& OutNodeSubsetsToConvertToCFGraph);
+
+	/** For GPU node inputs that have multiple incident edges, bundle them into a single edge. This is to avoid an inefficient
+	* gather operation on the GPU, and allows data interfaces to pick their data from the compute graph element input data collection
+	* using unique virtual input pin labels. */
+	void CreateGatherTasksAtGPUInputs(const TSet<FPCGTaskId>& InGPUCompatibleTaskIds, TArray<FPCGGraphTask>& InOutCompiledTasks);
+
+	/** Wires in a compute graph element alongside each set of GPU compatible nodes. The tasks for each node will be culled later. */
+	void WireGPUGraphNode(
+		FPCGTaskId InGPUGraphTaskId,
+		const TSet<FPCGTaskId>& InCollapsedTasks,
+		const TSet<FPCGTaskId>& InGPUCompatibleTaskIds,
+		TArray<FPCGGraphTask>& InOutCompiledTasks,
+		const TMap<FPCGTaskId,
+		TArray<FPCGTaskId>>& InTaskSuccessors,
+		FOriginalToVirtualPin& OutOriginalToVirtualPin,
+		TMap<const UPCGPin*, FName>& OutOutputCPUPinToVirtualPin);
+
+	/** Compiles a compute graph. */
+	void BuildGPUGraphTask(
+		UPCGGraph* InGraph,
+		FPCGTaskId InGPUGraphTaskId,
+		const TSet<FPCGTaskId>& InCollapsedTasks,
+		const TMap<FPCGTaskId, TArray<FPCGTaskId>>& InTaskSuccessors,
+		TArray<FPCGGraphTask>& InOutCompiledTasks,
+		const FOriginalToVirtualPin& InOriginalToVirtualPin,
+		const TMap<const UPCGPin*, FName>& InOutputCPUPinToVirtualPin);
+
+	/** Finds connected subgraphs of GPU - enabled nodes that can be dispatched together and replaces each one with a compute graph. */
+	void CreateGPUNodes(UPCGGraph* InGraph, TArray<FPCGGraphTask>& InOutCompiledTasks);
 
 private:
 	TArray<FPCGGraphTask> CompileGraph(UPCGGraph* InGraph, FPCGTaskId& NextId, FPCGStackContext& InOutStackContext);
@@ -50,6 +90,9 @@ private:
 	/** Returns the trivial element object shared by all tasks that need it. */
 	FPCGElementPtr GetSharedTrivialElement();
 
+	/** Returns the gather element object shared by all tasks that need it. */
+	FPCGElementPtr GetSharedGatherElement();
+
 	/** Propagates grid sizes through a graph's compiled tasks. */
 	static void ResolveGridSizes(
 		EPCGHiGenGrid GenerationGrid,
@@ -97,8 +140,10 @@ private:
 	TMap<UPCGGraph*, TMap<uint32, FPCGStackContext>> TopGraphToStackContextMap;
 
 	FPCGElementPtr SharedTrivialElement;
+	FPCGElementPtr SharedGatherElement;
 	FPCGElementPtr SharedTrivialPostGraphElement;
 	mutable FRWLock SharedTrivialElementLock;
+	mutable FRWLock SharedGatherElementLock;
 
 #if WITH_EDITOR
 	void RemoveFromCacheRecursive(UPCGGraph* InGraph);
diff --git a/Engine/Plugins/PCG/Source/PCG/Private/Graph/PCGGraphExecutor.cpp b/Engine/Plugins/PCG/Source/PCG/Private/Graph/PCGGraphExecutor.cpp
--- a/Engine/Plugins/PCG/Source/PCG/Private/Graph/PCGGraphExecutor.cpp
+++ b/Engine/Plugins/PCG/Source/PCG/Private/Graph/PCGGraphExecutor.cpp
@@ -12,6 +12,7 @@
 #include "PCGPin.h"
 #include "PCGSubsystem.h"
 #include "PCGWorldActor.h"
+#include "Compute/Elements/PCGComputeGraphElement.h"
 #include "Grid/PCGPartitionActor.h"
 #include "Helpers/PCGActorHelpers.h"
 #include "Helpers/PCGHelpers.h"
@@ -224,6 +225,12 @@ bool FPCGGraphTask::IsApproximatelyEqual(const FPCGGraphTask& Other) const
 			const PCGGraphExecutor::FPCGGridLinkageElement& OtherLinkageElement = static_cast<const PCGGraphExecutor::FPCGGridLinkageElement&>(*Other.Element);
 			bElementsMatch = (LinkageElement == OtherLinkageElement);
 		}
+		else if (Element->IsComputeGraphElement() && Other.Element->IsComputeGraphElement())
+		{
+			const FPCGComputeGraphElement& ComputeGraphElement = static_cast<const FPCGComputeGraphElement&>(*Element);
+			const FPCGComputeGraphElement& OtherComputeGraphElement = static_cast<const FPCGComputeGraphElement&>(*Other.Element);
+			bElementsMatch = (ComputeGraphElement == OtherComputeGraphElement);
+		}
 		else
 		{
 			ensureMsgf(false, TEXT("Graph compilation emitted an element type that is not a trivial element or a grid linkage element. Element comparison will fail. ")
diff --git a/Engine/Plugins/PCG/Source/PCG/Private/PCGModule.cpp b/Engine/Plugins/PCG/Source/PCG/Private/PCGModule.cpp
--- a/Engine/Plugins/PCG/Source/PCG/Private/PCGModule.cpp
+++ b/Engine/Plugins/PCG/Source/PCG/Private/PCGModule.cpp
@@ -5,14 +5,18 @@
 #include "PCGContext.h"
 #include "PCGElement.h"
 
+#include "Interfaces/IPluginManager.h"
 #include "Modules/ModuleManager.h"
 
 #if WITH_EDITOR
 #include "Elements/PCGDifferenceElement.h"
-#include "ISettingsModule.h"
-#include "ShowFlags.h"
 #include "Tests/Determinism/PCGDeterminismNativeTests.h"
 #include "Tests/Determinism/PCGDifferenceDeterminismTest.h"
+
+#include "ISettingsModule.h"
+#include "ShaderCore.h"
+#include "ShowFlags.h"
+#include "Misc/Paths.h"
 #endif
 
 #define LOCTEXT_NAMESPACE "FPCGModule"
@@ -25,6 +29,9 @@ void FPCGModule::StartupModule()
 	RegisterNativeElementDeterminismTests();
 
 	FEngineShowFlags::RegisterCustomShowFlag(PCGEngineShowFlags::Debug, /*DefaultEnabled=*/true, EShowFlagGroup::SFG_Developer, LOCTEXT("ShowFlagDisplayName", "PCG Debug"));
+
+	const FString PluginShaderDir = FPaths::Combine(IPluginManager::Get().FindPlugin(TEXT("PCG"))->GetBaseDir(), TEXT("Shaders"));
+	AddShaderSourceDirectoryMapping(TEXT("/Plugin/PCG"), PluginShaderDir);
 }
 
 void FPCGModule::ShutdownModule()
diff --git a/Engine/Plugins/PCG/Source/PCG/Private/PCGPin.cpp b/Engine/Plugins/PCG/Source/PCG/Private/PCGPin.cpp
--- a/Engine/Plugins/PCG/Source/PCG/Private/PCGPin.cpp
+++ b/Engine/Plugins/PCG/Source/PCG/Private/PCGPin.cpp
@@ -91,7 +91,9 @@ bool FPCGPinProperties::operator==(const FPCGPinProperties& Other) const
 		bAllowMultipleData == Other.bAllowMultipleData &&
 		Usage == Other.Usage &&
 		PinStatus == Other.PinStatus &&
-		bInvisiblePin == Other.bInvisiblePin;
+		bInvisiblePin == Other.bInvisiblePin &&
+		BufferSizeMode == Other.BufferSizeMode &&
+		FixedBufferElementCount == Other.FixedBufferElementCount;
 }
 
 void FPCGPinProperties::PostSerialize(const FArchive& Ar)
@@ -109,6 +111,25 @@ void FPCGPinProperties::PostSerialize(const FArchive& Ar)
 #endif // WITH_EDITOR
 }
 
+#if WITH_EDITOR
+bool FPCGPinProperties::CanEditChange(const FEditPropertyChain& PropertyChain) const
+{
+	if (FProperty* Property = PropertyChain.GetActiveNode()->GetValue())
+	{
+		if (Property->GetFName() == GET_MEMBER_NAME_CHECKED(FPCGPinProperties, bAllowMultipleData))
+		{
+			return bAllowEditMultipleData;
+		}
+		else if (Property->GetFName() == GET_MEMBER_NAME_CHECKED(FPCGPinProperties, bAllowMultipleConnections))
+		{
+			return bAllowMultipleData && bAllowEditMultipleConnections;
+		}
+	}
+
+	return true;
+}
+#endif // WITH_EDITOR
+
 UPCGPin::UPCGPin(const FObjectInitializer& ObjectInitializer)
 	: Super(ObjectInitializer)
 {
diff --git a/Engine/Plugins/PCG/Source/PCG/Private/PCGSettings.cpp b/Engine/Plugins/PCG/Source/PCG/Private/PCGSettings.cpp
--- a/Engine/Plugins/PCG/Source/PCG/Private/PCGSettings.cpp
+++ b/Engine/Plugins/PCG/Source/PCG/Private/PCGSettings.cpp
@@ -11,6 +11,9 @@
 #include "PCGPin.h"
 #include "PCGSubgraph.h"
 #include "PCGSubsystem.h"
+#include "Compute/PCGComputeGraph.h"
+#include "Compute/PCGDataBinding.h"
+#include "Compute/PCGDataForGPU.h"
 #include "Elements/PCGAddTag.h"
 #include "Helpers/PCGHelpers.h"
 #include "Helpers/PCGSettingsHelpers.h"
@@ -187,6 +190,54 @@ void UPCGSettingsInterface::SetEnabled(bool bInEnabled)
 	}
 }
 
+FPCGDataCollectionDesc UPCGSettings::ComputeInputPinDataDesc(const UPCGPin* InputPin, const UPCGDataBinding* Binding) const
+{
+	check(InputPin && Binding);
+
+	FPCGDataCollectionDesc PinDesc;
+
+	// Grab data from all incident edges.
+	for (const UPCGEdge* Edge : InputPin->Edges)
+	{
+		// InputPin is upstream output pin.
+		const UPCGPin* UpstreamOutputPin = Edge->InputPin;
+		if (!UpstreamOutputPin)
+		{
+			continue;
+		}
+
+		const UPCGSettings* UpstreamSettings = UpstreamOutputPin->Node ? UpstreamOutputPin->Node->GetSettings() : nullptr;
+		check(UpstreamSettings);
+
+		// Add data from connected upstream output pin.
+		PinDesc.Combine(UpstreamSettings->ComputeOutputPinDataDesc(UpstreamOutputPin, Binding));
+	}
+
+	return PinDesc;
+}
+
+FPCGDataCollectionDesc UPCGSettings::ComputeOutputPinDataDesc(const FName& OutputPinLabel, const UPCGDataBinding* Binding) const
+{
+	const UPCGNode* Node = CastChecked<UPCGNode>(GetOuter());
+	return ComputeOutputPinDataDesc(Node->GetOutputPin(OutputPinLabel), Binding);
+}
+
+FPCGDataCollectionDesc UPCGSettings::ComputeOutputPinDataDesc(const UPCGPin* OutputPin, const UPCGDataBinding* Binding) const
+{
+	// This base class implementation will be called on upstream CPU nodes. Get the downstream pin alias in order to
+	// pick the relevant data items out of the compute graph element's input data collection.
+	if (const FName* FoundPinAlias = Binding->Graph->OutputCPUPinToInputGPUPinAlias.Find(OutputPin))
+	{
+		return FPCGDataCollectionDesc::BuildFromInputDataCollectionAndInputPinLabel(
+			Binding->DataForGPU.InputDataCollection,
+			*FoundPinAlias,
+			Binding->Graph->GetAttributeLookupTable());
+	}
+
+	ensure(false);
+	return {};
+}
+
 uint32 UPCGSettings::GetTypeNameHash() const
 {
 	return GetTypeHash(GetClass()->GetFName());
@@ -591,9 +642,18 @@ EPCGChangeType UPCGSettings::GetChangeTypeForProperty(FPropertyChangedEvent& Pro
 
 EPCGChangeType UPCGSettings::GetChangeTypeForProperty(const FName& InPropertyName) const
 {
+	EPCGChangeType ChangeType = EPCGChangeType::Settings;
+
+	if (InPropertyName == GET_MEMBER_NAME_CHECKED(UPCGSettings, bExecuteOnGPU))
+	{
+		ChangeType |= EPCGChangeType::Structural;
+	}
+
 PRAGMA_DISABLE_DEPRECATION_WARNINGS
-	return EPCGChangeType::Settings | (IsStructuralProperty(InPropertyName) ? EPCGChangeType::Structural : EPCGChangeType::None);
+	ChangeType |= IsStructuralProperty(InPropertyName) ? EPCGChangeType::Structural : EPCGChangeType::None;
 PRAGMA_ENABLE_DEPRECATION_WARNINGS
+
+	return ChangeType;
 }
 
 void UPCGSettings::DirtyCache()
diff --git a/Engine/Plugins/PCG/Source/PCG/Private/PCGSubsystem.cpp b/Engine/Plugins/PCG/Source/PCG/Private/PCGSubsystem.cpp
--- a/Engine/Plugins/PCG/Source/PCG/Private/PCGSubsystem.cpp
+++ b/Engine/Plugins/PCG/Source/PCG/Private/PCGSubsystem.cpp
@@ -208,6 +208,11 @@ UPCGSubsystem* UPCGSubsystem::GetInstance(UWorld* World)
 	}
 }
 
+void UPCGSubsystem::RegisterBeginTickAction(FTickAction&& Action)
+{
+	BeginTickActions.Emplace(Action);
+}
+
 #if WITH_EDITOR
 UPCGSubsystem* UPCGSubsystem::GetActiveEditorInstance()
 {
@@ -254,6 +259,8 @@ void UPCGSubsystem::Tick(float DeltaSeconds)
 {
 	Super::Tick(DeltaSeconds);
 
+	ExecuteBeginTickActions();
+
 #if WITH_EDITOR
 	PerActorConstructionScriptSourceComponents.Empty();
 #endif
@@ -1752,6 +1759,17 @@ void UPCGSubsystem::CreatePartitionActorsWithinBounds(UPCGComponent* InComponent
 
 #endif // WITH_EDITOR
 
+void UPCGSubsystem::ExecuteBeginTickActions()
+{
+	TArray<FTickAction> Actions = MoveTemp(BeginTickActions);
+	BeginTickActions.Reset();
+
+	for (FTickAction& Action : Actions)
+	{
+		Action();
+	}
+}
+
 void UPCGSubsystem::FlushCache()
 {
 	if (GraphExecutor)
diff --git a/Engine/Plugins/PCG/Source/PCG/Public/Compute/Elements/PCGCustomHLSL.h b/Engine/Plugins/PCG/Source/PCG/Public/Compute/Elements/PCGCustomHLSL.h
new file mode 100644
--- /dev/null
+++ b/Engine/Plugins/PCG/Source/PCG/Public/Compute/Elements/PCGCustomHLSL.h
@@ -0,0 +1,151 @@
+// Copyright Epic Games, Inc. All Rights Reserved.
+
+#pragma once
+
+#include "PCGSettings.h"
+
+#include "PCGCustomHLSL.generated.h"
+
+class UPCGPin;
+
+/** Type of kernel allows us to make decisions about execution automatically, streamlining authoring. */
+UENUM()
+enum class EPCGKernelType : uint8
+{
+	PointProcessor UMETA(Tooltip = "Kernel executes on each point in first input pin."),
+	PointGenerator UMETA(Tooltip = "Kernel executes for fixed number of points, configurable on node."),
+	Custom UMETA(Tooltip = "Execution thread counts and output buffer sizes configurable on node."),
+};
+
+/** Total number of threads that will be dispatched for this kernel. */
+UENUM()
+enum class EPCGDispatchThreadCount : uint8
+{
+	FromOutput UMETA(DisplayName = "Match First Output Pin", Tooltip = "One thread per pin data element."),
+	FromFirstInput UMETA(DisplayName = "Match First Input Pin", Tooltip = "One thread per pin data element."),
+	FromSecondInput UMETA(DisplayName = "Match Second Input Pin", Tooltip = "One thread per pin data element."),
+	FirstInputXSecondInput UMETA(DisplayName = "First Input X Second Input"),
+	Fixed UMETA(DisplayName = "Fixed Thread Count"),
+
+	// TODO: Could support inlined expression evaluation for dispatch thread count.
+};
+
+/** Prouduces a HLSL compute shader which will be executed on the GPU. */
+UCLASS(MinimalAPI, BlueprintType, ClassGroup = (Procedural))
+class UPCGCustomHLSLSettings : public UPCGSettings
+{
+	GENERATED_BODY()
+
+public:
+	//~Begin UObject interface
+#if WITH_EDITOR
+	virtual void PostLoad() override;
+	virtual void PostInitProperties() override;
+#endif
+	//~End UObject interface
+
+	//~Begin UPCGSettings interface
+	virtual TArray<FPCGPinProperties> InputPinProperties() const override { return InputPins; }
+	virtual TArray<FPCGPinProperties> OutputPinProperties() const override { return OutputPins; }
+	virtual bool ShouldExecuteOnGPU() const override;
+	virtual bool IsInputPinRequiredByExecution(const UPCGPin* InPin) const { return true; }
+#if WITH_EDITOR
+	virtual bool DisplayExecuteOnGPUSetting() const override { return false; }
+	virtual FName GetDefaultNodeName() const override { return FName(TEXT("CustomHLSL")); }
+	virtual FText GetDefaultNodeTitle() const override { return NSLOCTEXT("PCGCustomHLSLElement", "NodeTitle", "Custom HLSL"); }
+	virtual FText GetNodeTooltipText() const override { return NSLOCTEXT("PCGCustomHLSLElement", "NodeTooltip", "Prouduces a HLSL compute shader which will be executed on the GPU."); }
+	virtual EPCGSettingsType GetType() const override { return EPCGSettingsType::GPU; }
+#endif
+
+	virtual TArray<FPCGKernelAttributeKey> GetKernelAttributeKeys() const override;
+	virtual FPCGDataCollectionDesc ComputeOutputPinDataDesc(const UPCGPin* OutputPin, const UPCGDataBinding* Binding) const override;
+	virtual int ComputeKernelThreadCount(const UPCGDataBinding* Binding) const override;
+
+protected:
+#if WITH_EDITOR
+	virtual void PostEditChangeProperty(FPropertyChangedEvent& PropertyChangedEvent) override;
+	virtual EPCGChangeType GetChangeTypeForProperty(const FName& InPropertyName) const override;
+#endif
+	virtual FPCGElementPtr CreateElement() const override;
+	//~End UPCGSettings interface
+
+#if WITH_EDITOR
+	void UpdateDeclarations();
+	void UpdatePinSettings();
+#endif
+
+public:
+	bool IsKernelValid(FPCGContext* InContext = nullptr) const;
+	
+	FString GetCookedKernelSource(const TMap<FPCGKernelAttributeKey, int>& GlobalAttributeLookupTable) const;
+	FString GetKernelEntryPoint() const { return TEXT("Main"); }
+	FIntVector GetThreadGroupSize() const { return FIntVector(64, 1, 1); }
+
+	int GetPointCount() const { return PointCount; }
+	int GetFixedThreadCount() const { return FixedThreadCount; }
+
+	const UPCGPin* GetPointProcessingInputPin() const;
+	const UPCGPin* GetSecondPointProcessingInputPin() const;
+	const UPCGPin* GetFirstOutputPin() const;
+	const UPCGPin* GetFirstPointOutputPin() const;
+
+	int GetProcessingElemCountForInputPin(const UPCGPin* InputPin, const UPCGDataBinding* Binding) const;
+	virtual const UPCGPin* GetExecutionPin() const { return GetPointProcessingInputPin(); }
+
+protected:
+	UPROPERTY(BlueprintReadWrite, EditAnywhere, Category = "Settings")
+	EPCGKernelType KernelType = EPCGKernelType::PointProcessor;
+
+	UPROPERTY(BlueprintReadWrite, EditAnywhere, Category = "Settings", meta = (EditCondition = "KernelType == EPCGKernelType::PointGenerator", EditConditionHides))
+	int PointCount = 256;
+
+	UPROPERTY(BlueprintReadWrite, EditAnywhere, Category = "Settings", meta = (EditCondition = "KernelType == EPCGKernelType::Custom", EditConditionHides))
+	EPCGDispatchThreadCount DispatchThreadCount = EPCGDispatchThreadCount::FromOutput;
+
+	UPROPERTY(BlueprintReadWrite, EditAnywhere, Category = "Settings", meta = (EditCondition = "KernelType == EPCGKernelType::Custom && DispatchThreadCount != EPCGDispatchThreadCount::Fixed", EditConditionHides))
+	int ThreadCountMultiplier = 1;
+
+	UPROPERTY(BlueprintReadWrite, EditAnywhere, Category = "Settings", meta = (EditCondition = "KernelType == EPCGKernelType::Custom && DispatchThreadCount == EPCGDispatchThreadCount::Fixed", EditConditionHides))
+	int FixedThreadCount = 1;
+
+public:
+	/** Dump the cooked HLSL into the log after it is generated. */
+	UPROPERTY(EditAnywhere, Category = "Debug")
+	bool bDumpCookedHLSL = false;
+
+	/** Enable use of 'WriteDebugValue(uint Index, float Value)' function in your kernel. Allows you to write float values to a buffer for logging on the CPU. */
+	UPROPERTY(EditAnywhere, Category = "Debug")
+	bool bPrintShaderDebugValues = false;
+
+	/** Size (in number of floats) of the shader debug print buffer. */
+	UPROPERTY(EditAnywhere, Category = "Debug", meta = (EditCondition="bPrintShaderDebugValues", EditConditionHides))
+	int DebugBufferSize = 16;
+
+public:
+	UPROPERTY(BlueprintReadWrite, EditAnywhere, Category = "Settings")
+	TArray<FPCGPinProperties> InputPins = Super::DefaultPointInputPinProperties();
+
+	UPROPERTY(BlueprintReadWrite, EditAnywhere, Category = "Settings")
+	TArray<FPCGPinProperties> OutputPins = Super::DefaultPointOutputPinProperties();
+
+protected:
+	UPROPERTY(Transient, VisibleAnywhere, Category = "Settings|Declarations", meta = (MultiLine = true))
+	FString InputDeclarations;
+
+	UPROPERTY(Transient, VisibleAnywhere, Category = "Settings|Declarations", meta = (MultiLine = true))
+	FString OutputDeclarations;
+
+	UPROPERTY(BlueprintReadWrite, EditAnywhere, Category = "Settings", meta = (MultiLine = true, Tooltip = ""))
+	FString ShaderFunctions;
+
+	UPROPERTY(BlueprintReadWrite, EditAnywhere, Category = "Settings", meta = (MultiLine = true, Tooltip = ""))
+	FString ShaderSource;
+};
+
+class FPCGCustomHLSLElement : public IPCGElement
+{
+protected:
+	// During compilation, custom HLSL nodes will be replaced with a compute graph element, and this should never actually execute.
+	// Left ensure here for validation.
+	virtual bool ExecuteInternal(FPCGContext* Context) const override { ensure(false); return true; }
+};
diff --git a/Engine/Plugins/PCG/Source/PCG/Public/Compute/PCGDataForGPU.h b/Engine/Plugins/PCG/Source/PCG/Public/Compute/PCGDataForGPU.h
new file mode 100644
--- /dev/null
+++ b/Engine/Plugins/PCG/Source/PCG/Public/Compute/PCGDataForGPU.h
@@ -0,0 +1,98 @@
+// Copyright Epic Games, Inc. All Rights Reserved.
+
+#pragma once
+
+#include "PCGCommon.h"
+#include "PCGData.h"
+
+#include "PCGDataForGPU.generated.h"
+
+class UPCGMetadata;
+class UPCGPin;
+class UPCGSettings;
+
+UENUM()
+enum class EPCGKernelAttributeType : uint8
+{
+	Bool,
+	Int,
+	Float,
+	Float2,
+	Float3,
+	Float4,
+	Rotator,
+	Quat,
+	Transform
+};
+
+USTRUCT()
+struct FPCGKernelAttributeKey
+{
+	GENERATED_BODY()
+
+	UPROPERTY(VisibleAnywhere, Category = "Settings")
+	EPCGKernelAttributeType Type = EPCGKernelAttributeType::Float;
+
+	UPROPERTY(VisibleAnywhere, Category = "Settings")
+	FName Name = NAME_None;
+
+	bool operator==(const FPCGKernelAttributeKey& Other) const { return Type == Other.Type && Name == Other.Name; }
+	friend uint32 GetTypeHash(const FPCGKernelAttributeKey& In);
+};
+
+struct FPCGKernelAttributeDesc
+{
+	int32 Index = INDEX_NONE;
+	EPCGKernelAttributeType Type = EPCGKernelAttributeType::Float;
+	FName Name = NAME_None;
+};
+
+struct FPCGDataDesc
+{
+	FPCGDataDesc(EPCGDataType InType, int32 InElementCount);
+	FPCGDataDesc(const UPCGData* Data, const TMap<FPCGKernelAttributeKey, int32>& GlobalAttributeLookupTable);
+
+	uint32 ComputePackedSize() const;
+
+	EPCGDataType Type = EPCGDataType::Point;
+	TArray<FPCGKernelAttributeDesc> AttributeDescs;
+	int32 ElementCount = 0;
+
+private:
+	void InitializeAttributeDescs(const UPCGMetadata* Metadata, const TMap<FPCGKernelAttributeKey, int32>& GlobalAttributeLookupTable = {});
+};
+
+struct FPCGDataCollectionDesc
+{
+	static FPCGDataCollectionDesc BuildFromInputDataCollectionAndInputPinLabel(
+		const FPCGDataCollection& InDataCollection,
+		FName InputPinLabel,
+		const TMap<FPCGKernelAttributeKey, int32>& InAttributeLookupTable);
+
+	/** Computes the size (in bytes) of the data collection after packing. Also produces the offset (in bytes) for each data in the packed collection. */
+	uint32 ComputePackedSize(TArray<uint32>* OutDataAddresses = nullptr) const;
+
+	/** Pack a data collection into the GPU data format. DataDescs defines which attributes are packed. */
+	void PackDataCollection(const FPCGDataCollection& InDataCollection, FName InPin, TArray<uint32>& OutPackedDataCollection) const;
+
+	/** Unpack a buffer of 8-bit uints to a data collection. */
+	void UnpackDataCollection(const TArray<uint8>& InPackedData, FName InPin, FPCGDataCollection& OutDataCollection) const;
+
+	/** Compute total number of processing elements of the given type. */
+	uint32 ComputeDataElementCount(EPCGDataType InDataType) const;
+
+	/** Aggregate another data description. */
+	void Combine(const FPCGDataCollectionDesc& Other);
+
+	TArray<FPCGDataDesc> DataDescs;
+};
+
+struct FPCGDataForGPU
+{
+	// All the node input pins that have edges that cross from CPU to GPU.
+	TSet<const UPCGPin*> InputPins;
+	TMap<const UPCGPin*, FName> InputPinLabelAliases;
+
+	// Since the compute graph is collapsed to a single element, all data crossing from CPU to GPU is in a single collection.
+	FPCGDataCollection InputDataCollection;
+};
diff --git a/Engine/Plugins/PCG/Source/PCG/Public/PCGCommon.h b/Engine/Plugins/PCG/Source/PCG/Public/PCGCommon.h
--- a/Engine/Plugins/PCG/Source/PCG/Public/PCGCommon.h
+++ b/Engine/Plugins/PCG/Source/PCG/Public/PCGCommon.h
@@ -47,6 +47,8 @@ enum class EPCGChangeType : uint8
 	Structural = 1 << 5,
 	/** Anything related to generation grids - changing grid size or adding/removing grid size nodes. */
 	GenerationGrid = 1 << 6,
+	/** Change to any shader source code. */
+	ShaderSource = 1 << 7,
 };
 ENUM_CLASS_FLAGS(EPCGChangeType);
 
diff --git a/Engine/Plugins/PCG/Source/PCG/Public/PCGElement.h b/Engine/Plugins/PCG/Source/PCG/Public/PCGElement.h
--- a/Engine/Plugins/PCG/Source/PCG/Public/PCGElement.h
+++ b/Engine/Plugins/PCG/Source/PCG/Public/PCGElement.h
@@ -103,6 +103,9 @@ public:
 	/** Is this element used for marshalling data across higen grids. Used as rudimentary RTTI for element object comparisons (editor only). */
 	virtual bool IsGridLinkage() const { return false; }
 
+	/** Is this element used for dispatching compute graphs. Used as rudimentary RTTI for element object comparisons (editor only). */
+	virtual bool IsComputeGraphElement() const { return false; }
+
 	/** Note: must be called from the main thread. */
 	void DebugDisplay(FPCGContext* Context) const;
 #endif
diff --git a/Engine/Plugins/PCG/Source/PCG/Public/PCGPin.h b/Engine/Plugins/PCG/Source/PCG/Public/PCGPin.h
--- a/Engine/Plugins/PCG/Source/PCG/Public/PCGPin.h
+++ b/Engine/Plugins/PCG/Source/PCG/Public/PCGPin.h
@@ -34,6 +34,16 @@ enum class EPCGPinStatus : uint8
 	Advanced
 };
 
+/** Method for computing the size of a pin on a GPU node. */
+UENUM()
+enum class EPCGPinBufferSizeMode : uint8
+{
+	FromFirstPin UMETA(DisplayName = "Match First Input Pin"),
+	FromSecondPin UMETA(DisplayName = "Match Second Input Pin"),
+	FirstPinXSecondPin UMETA(DisplayName = "First Input Pin X Second Input Pin"),
+	FixedElementCount,
+};
+
 USTRUCT(BlueprintType, meta=(HasNativeBreak="/Script/PCG.PCGBlueprintPinHelpers.BreakPinProperty", HasNativeMake="/Script/PCG.PCGBlueprintPinHelpers.MakePinProperty"))
 struct PCG_API FPCGPinProperties
 {
@@ -68,6 +78,13 @@ struct PCG_API FPCGPinProperties
 	FText Tooltip;
 #endif
 
+	/** Compute graphs use this to calculate the buffer size of output pins. */
+	UPROPERTY(BlueprintReadWrite, EditAnywhere, Category = Settings, meta = (EditCondition = "bDisplayBufferSizeSettings && AllowedTypes == EPCGDataType::Point", EditConditionHides))
+	EPCGPinBufferSizeMode BufferSizeMode = EPCGPinBufferSizeMode::FromFirstPin;
+
+	UPROPERTY(BlueprintReadWrite, EditAnywhere, Category = Settings, meta = (EditCondition = "bDisplayBufferSizeSettings && AllowedTypes == EPCGDataType::Point && BufferSizeMode == EPCGPinBufferSizeMode::FixedElementCount", EditConditionHides))
+	int FixedBufferElementCount = 4;
+
 	// Multiple connections are only possible if we support multi data.
 	bool AllowsMultipleConnections() const { return bAllowMultipleData && bAllowMultipleConnections; }
 
@@ -86,6 +103,21 @@ struct PCG_API FPCGPinProperties
 	// Convert the bIsAdvanced boolean to PinStatus for deprecation purposes.
 	void PostSerialize(const FArchive& Ar);
 
+#if WITH_EDITOR
+	bool CanEditChange(const FEditPropertyChain& PropertyChain) const;
+#endif
+
+#if WITH_EDITORONLY_DATA
+	UPROPERTY(Transient)
+	bool bDisplayBufferSizeSettings = true;
+
+	UPROPERTY(Transient)
+	bool bAllowEditMultipleData = true;
+	
+	UPROPERTY(Transient)
+	bool bAllowEditMultipleConnections = true;
+#endif // WITH_EDITORONLY_DATA
+
 private:
 #if WITH_EDITORONLY_DATA
 	/* Advanced pin will be hidden by default in the UI and will be shown only if the user extend the node (in the UI) to see advanced pins. */
@@ -93,7 +125,7 @@ private:
 	bool bAdvancedPin_DEPRECATED = false;
 #endif // WITH_EDITORONLY_DATA
 
-	UPROPERTY(BlueprintReadWrite, EditAnywhere, Category = Settings, meta = (AllowPrivateAccess = "true", EditCondition = "bAllowMultipleData", DisplayAfter = "bAllowMultipleData"))
+	UPROPERTY(BlueprintReadWrite, EditAnywhere, Category = Settings, meta = (AllowPrivateAccess = "true", DisplayAfter = "bAllowMultipleData"))
 	bool bAllowMultipleConnections = true;
 };
 
@@ -103,6 +135,7 @@ struct TStructOpsTypeTraits<FPCGPinProperties> : public TStructOpsTypeTraitsBase
 	enum
 	{
 		WithPostSerialize = true,
+		WithCanEditChange = true,
 	};
 };
 
diff --git a/Engine/Plugins/PCG/Source/PCG/Public/PCGSettings.h b/Engine/Plugins/PCG/Source/PCG/Public/PCGSettings.h
--- a/Engine/Plugins/PCG/Source/PCG/Public/PCGSettings.h
+++ b/Engine/Plugins/PCG/Source/PCG/Public/PCGSettings.h
@@ -6,6 +6,7 @@
 #include "PCGDebug.h"
 #include "PCGElement.h"
 #include "PCGPin.h"
+#include "Compute/PCGDataForGPU.h"
 #include "Elements/PCGActorSelector.h"
 #include "Tests/Determinism/PCGDeterminismSettings.h"
 
@@ -21,6 +22,7 @@ struct FPropertyChangedEvent;
 class UPCGGraph;
 class UPCGNode;
 class UPCGSettings;
+class UPCGDataBinding;
 
 using FPCGSettingsAndCulling = TPair<TSoftObjectPtr<const UPCGSettings>, bool>;
 using FPCGSelectionKeyToSettingsMap = TMap<FPCGSelectionKey, TArray<FPCGSettingsAndCulling>>;
@@ -62,7 +64,8 @@ enum class EPCGSettingsType : uint8
 	ControlFlow,
 	PointOps,
 	GraphParameters,
-	Reroute
+	Reroute,
+	GPU,
 };
 
 #if WITH_EDITOR
@@ -260,6 +263,9 @@ public:
 	
 	bool UseSeed() const { return bUseSeed; }
 
+	/** Whether this node should be executed on the GPU. */
+	virtual bool ShouldExecuteOnGPU() const { return bExecuteOnGPU; }
+
 	// Get the seed, combined with optional PCGComponent seed
 	int GetSeed(const UPCGComponent* InSourceComponent = nullptr) const;
 
@@ -398,6 +404,19 @@ public:
 	// Holds the original settings used to duplicate this object if it was overridden
 	const UPCGSettings* OriginalSettings = nullptr;
 
+	/** Get a list of the attributes read or written by this node. */
+	virtual TArray<FPCGKernelAttributeKey> GetKernelAttributeKeys() const { return {}; }
+
+	/** Compute a description of all data arriving on InputPin. */
+	FPCGDataCollectionDesc ComputeInputPinDataDesc(const UPCGPin* InputPin, const UPCGDataBinding* Binding) const;
+
+	/** Compute a description of data that will be output from OutputPinLabel/OutputPin. */
+	FPCGDataCollectionDesc ComputeOutputPinDataDesc(const FName& OutputPinLabel, const UPCGDataBinding* Binding) const;
+	virtual FPCGDataCollectionDesc ComputeOutputPinDataDesc(const UPCGPin* OutputPin, const UPCGDataBinding* Binding) const;
+
+	/** Compute how many threads should be dispatched to execute this node on the GPU. */
+	virtual int ComputeKernelThreadCount(const UPCGDataBinding* Binding) const { return 0; };
+
 protected:
 	// Returns an array of all the input pin properties. You should not add manually a "params" pin, it is handled automatically by FillOverridableParamsPins
 	virtual TArray<FPCGPinProperties> InputPinProperties() const;
@@ -436,8 +455,16 @@ protected:
 
 	// Passthrough for the simpler method, to avoid modifying the child settings already overriding this method.
 	virtual bool CanEditChange(const FProperty* InProperty) const override;
+
+	/** Whether to display GPU execution option in node settings UI. */
+	UFUNCTION()
+	virtual bool DisplayExecuteOnGPUSetting() const { return ShouldExecuteOnGPU(); }
 #endif
 
+	/** Whether this node should be executed on the GPU. */
+	UPROPERTY(BlueprintReadWrite, EditAnywhere, Category = GPU, meta = (EditCondition = "DisplayExecuteOnGPUSetting()", EditConditionHides, HideEditConditionToggle))
+	bool bExecuteOnGPU = false;
+
 	// By default, settings won't use a seed. Set this bool to true in the child ctor to allow edition and use it.
 	UPROPERTY(VisibleAnywhere, Transient, Category = Settings, meta = (EditCondition = false, EditConditionHides))
 	bool bUseSeed = false;
diff --git a/Engine/Plugins/PCG/Source/PCG/Public/PCGSubsystem.h b/Engine/Plugins/PCG/Source/PCG/Public/PCGSubsystem.h
--- a/Engine/Plugins/PCG/Source/PCG/Public/PCGSubsystem.h
+++ b/Engine/Plugins/PCG/Source/PCG/Public/PCGSubsystem.h
@@ -74,6 +74,10 @@ public:
 	/** Will return the subsystem from the World if it exists and if it is initialized */
 	static UPCGSubsystem* GetInstance(UWorld* World);
 
+	/** Adds an action that will be executed once at the beginning of this subsystem's next Tick(). */
+	using FTickAction = TFunction<void()>;
+	void RegisterBeginTickAction(FTickAction&& Action);
+
 #if WITH_EDITOR
 	/** Returns PIE world if it is active, otherwise returns editor world. */
 	static UPCGSubsystem* GetActiveEditorInstance();
@@ -334,6 +338,8 @@ private:
 #endif // WITH_EDITOR
 	
 private:
+	void ExecuteBeginTickActions();
+
 	APCGWorldActor* PCGWorldActor = nullptr;
 	FPCGGraphExecutor* GraphExecutor = nullptr;
 	FPCGRuntimeGenScheduler* RuntimeGenScheduler = nullptr;
@@ -344,6 +350,9 @@ private:
 	TArray<FPCGStack> ExecutedStacks;
 	mutable FRWLock ExecutedStacksLock;
 
+	/** Functions will be executed at the beginning of the tick and then removed from this array. */
+	TArray<FTickAction> BeginTickActions;
+
 #if WITH_EDITOR
 	FCriticalSection PCGWorldActorLock;
 		
diff --git a/Engine/Plugins/PCG/Source/PCGEditor/PCGEditor.Build.cs b/Engine/Plugins/PCG/Source/PCGEditor/PCGEditor.Build.cs
--- a/Engine/Plugins/PCG/Source/PCGEditor/PCGEditor.Build.cs
+++ b/Engine/Plugins/PCG/Source/PCGEditor/PCGEditor.Build.cs
@@ -9,34 +9,27 @@ namespace UnrealBuildTool.Rules
 			PCHUsage = PCHUsageMode.UseExplicitOrSharedPCHs;
 
 			PublicDependencyModuleNames.AddRange(
-			new string[] {
-			"Core",
-			"Projects",
-			"Engine",
-			"CoreUObject",
-			"PlacementMode",
-		});
+				new string[]
+				{
+					"Core",
+					"Projects",
+					"Engine",
+					"CoreUObject",
+					"PlacementMode",
+				});
 
 			if (Target.WithAutomationTests)
 			{
-
 				PublicDependencyModuleNames.AddRange(
 					new string[]
 					{
 						"LevelEditor"
-					}
-				);
+					});
 			}
 
 			PrivateDependencyModuleNames.AddRange(
 				new string[]
 				{
-				
-				}
-			);
-
-			PrivateDependencyModuleNames.AddRange(
-				new string[]{
 					"AppFramework",
 					"ApplicationCore",
 					"AssetDefinition",
@@ -58,6 +51,7 @@ namespace UnrealBuildTool.Rules
 					"KismetWidgets",
 					"PCG",
 					"PropertyEditor",
+					"RenderCore",
 					"Slate",
 					"SlateCore",
 					"SourceControl",
@@ -68,13 +62,7 @@ namespace UnrealBuildTool.Rules
 					"TypedElementRuntime",
 					"UnrealEd",
 					"LevelEditor"
-				}
-			);
-
-			PrivateIncludePaths.AddRange(
-				new string[] {
 				});
 		}
-
 	}
 }
diff --git a/Engine/Plugins/PCG/Source/PCGEditor/Private/PCGEditor.cpp b/Engine/Plugins/PCG/Source/PCGEditor/Private/PCGEditor.cpp
--- a/Engine/Plugins/PCG/Source/PCGEditor/Private/PCGEditor.cpp
+++ b/Engine/Plugins/PCG/Source/PCGEditor/Private/PCGEditor.cpp
@@ -49,6 +49,7 @@
 #include "PropertyEditorModule.h"
 #include "SNodePanel.h"
 #include "ScopedTransaction.h"
+#include "ShaderCore.h"
 #include "SourceCodeNavigation.h"
 #include "ToolMenu.h"
 #include "ToolMenuEntry.h"
@@ -201,6 +202,8 @@ void FPCGEditor::Initialize(const EToolkitMode::Type InMode, const TSharedPtr<cl
 
 	InitAssetEditor(InMode, InToolkitHost, PCGGraphEditorAppName, StandaloneDefaultLayout, /*bCreateDefaultStandaloneMenu=*/ true, /*bCreateDefaultToolbar=*/ true, InPCGGraph);
 
+	PCGGraphBeingEdited->OnGraphChangedDelegate.AddRaw(this, &FPCGEditor::OnGraphChanged);
+
 	// Hook to map change / delete actor to refresh debug object selection list, to help prevent it going stale.
 	FLevelEditorModule& LevelEditor = FModuleManager::LoadModuleChecked<FLevelEditorModule>("LevelEditor");
 	LevelEditor.OnMapChanged().AddRaw(this, &FPCGEditor::OnMapChanged);
@@ -920,6 +923,8 @@ void FPCGEditor::OnForceGraphRegeneration_Clicked()
 			}
 
 			ChangeType |= EPCGChangeType::GenerationGrid;
+
+			ChangeType |= EPCGChangeType::ShaderSource;
 		}
 
 		PCGGraphBeingEdited->ForceNotificationForEditor(ChangeType);
@@ -2847,6 +2852,8 @@ void FPCGEditor::OnClose()
 
 	if (PCGGraphBeingEdited)
 	{
+		PCGGraphBeingEdited->OnGraphChangedDelegate.RemoveAll(this);
+
 		if (PCGGraphBeingEdited->IsInspecting())
 		{
 			PCGGraphBeingEdited->DisableInspection();
@@ -3068,6 +3075,16 @@ void FPCGEditor::UnregisterDelegatesForWorld(UWorld* World)
 	}
 }
 
+void FPCGEditor::OnGraphChanged(UPCGGraphInterface* InGraph, EPCGChangeType ChangeType)
+{
+	if (!!(ChangeType & EPCGChangeType::ShaderSource))
+	{
+		// Flush the shader file cache in case we are editing engine or data interface shaders.
+		// We could make the user do this manually, but that makes iterating on data interfaces really painful.
+		FlushShaderFileCache();
+	}
+}
+
 void FPCGEditor::OnMapChanged(UWorld* InWorld, EMapChangeType InMapChangedType)
 {
 	if (InMapChangedType != EMapChangeType::SaveMap)
diff --git a/Engine/Plugins/PCG/Source/PCGEditor/Private/PCGEditor.h b/Engine/Plugins/PCG/Source/PCGEditor/Private/PCGEditor.h
--- a/Engine/Plugins/PCG/Source/PCGEditor/Private/PCGEditor.h
+++ b/Engine/Plugins/PCG/Source/PCGEditor/Private/PCGEditor.h
@@ -339,6 +339,7 @@ private:
 	void RegisterDelegatesForWorld(UWorld* World);
 	void UnregisterDelegatesForWorld(UWorld* World);
 
+	void OnGraphChanged(UPCGGraphInterface* InGraph, EPCGChangeType ChangeType);
 	void OnMapChanged(UWorld* InWorld, EMapChangeType InMapChangedType);
 	void OnPostPIEStarted(bool bIsSimulating);
 	void OnEndPIE(bool bIsSimulating);
