[PCG] Fix Mac compile error and static analysis issue

#rnx

diff --git a/Engine/Plugins/PCG/Source/PCG/Private/Compute/Elements/PCGCustomHLSL.cpp b/Engine/Plugins/PCG/Source/PCG/Private/Compute/Elements/PCGCustomHLSL.cpp
--- a/Engine/Plugins/PCG/Source/PCG/Private/Compute/Elements/PCGCustomHLSL.cpp
+++ b/Engine/Plugins/PCG/Source/PCG/Private/Compute/Elements/PCGCustomHLSL.cpp
@@ -994,9 +994,7 @@ FString UPCGCustomHLSLSettings::GetCookedKernelSource(const TMap<FPCGKernelAttri
 	{
 		if (const UPCGNode* Node = Cast<UPCGNode>(GetOuter()))
 		{
-			const UPCGPin* FirstPin = (KernelType == EPCGKernelType::PointProcessor)
-				? GetPointProcessingInputPin()
-				: CastChecked<UPCGNode>(GetOuter())->GetPassThroughInputPin();
+			const UPCGPin* FirstPin = Node->GetPassThroughInputPin();
 			const UPCGPin* SecondPin = GetSecondPointProcessingInputPin();
 			const UPCGPin* PrimaryOutputPin = GetFirstPointOutputPin();
 			
diff --git a/Engine/Plugins/PCG/Source/PCG/Private/Compute/PCGDataForGPU.cpp b/Engine/Plugins/PCG/Source/PCG/Private/Compute/PCGDataForGPU.cpp
--- a/Engine/Plugins/PCG/Source/PCG/Private/Compute/PCGDataForGPU.cpp
+++ b/Engine/Plugins/PCG/Source/PCG/Private/Compute/PCGDataForGPU.cpp
@@ -16,15 +16,15 @@ namespace PCGDataForGPUConstants
 {
 	const static FPCGKernelAttributeDesc PointPropertyDescs[NUM_POINT_PROPERTIES] =
 	{
-		{ POINT_POSITION_ATTRIBUTE_ID,   EPCGKernelAttributeType::Float3 },
-		{ POINT_ROTATION_ATTRIBUTE_ID,   EPCGKernelAttributeType::Quat   },
-		{ POINT_SCALE_ATTRIBUTE_ID,      EPCGKernelAttributeType::Float3 },
-		{ POINT_BOUNDS_MIN_ATTRIBUTE_ID, EPCGKernelAttributeType::Float3 },
-		{ POINT_BOUNDS_MAX_ATTRIBUTE_ID, EPCGKernelAttributeType::Float3 },
-		{ POINT_COLOR_ATTRIBUTE_ID,      EPCGKernelAttributeType::Float4 },
-		{ POINT_DENSITY_ATTRIBUTE_ID,    EPCGKernelAttributeType::Float  },
-		{ POINT_SEED_ATTRIBUTE_ID,       EPCGKernelAttributeType::Int    },
-		{ POINT_STEEPNESS_ATTRIBUTE_ID,  EPCGKernelAttributeType::Float  }
+		FPCGKernelAttributeDesc(POINT_POSITION_ATTRIBUTE_ID,   EPCGKernelAttributeType::Float3, NAME_None),
+		FPCGKernelAttributeDesc(POINT_ROTATION_ATTRIBUTE_ID,   EPCGKernelAttributeType::Quat,   NAME_None),
+		FPCGKernelAttributeDesc(POINT_SCALE_ATTRIBUTE_ID,      EPCGKernelAttributeType::Float3, NAME_None),
+		FPCGKernelAttributeDesc(POINT_BOUNDS_MIN_ATTRIBUTE_ID, EPCGKernelAttributeType::Float3, NAME_None),
+		FPCGKernelAttributeDesc(POINT_BOUNDS_MAX_ATTRIBUTE_ID, EPCGKernelAttributeType::Float3, NAME_None),
+		FPCGKernelAttributeDesc(POINT_COLOR_ATTRIBUTE_ID,      EPCGKernelAttributeType::Float4, NAME_None),
+		FPCGKernelAttributeDesc(POINT_DENSITY_ATTRIBUTE_ID,    EPCGKernelAttributeType::Float,  NAME_None),
+		FPCGKernelAttributeDesc(POINT_SEED_ATTRIBUTE_ID,       EPCGKernelAttributeType::Int,    NAME_None),
+		FPCGKernelAttributeDesc(POINT_STEEPNESS_ATTRIBUTE_ID,  EPCGKernelAttributeType::Float,  NAME_None)
 	};
 }
 
diff --git a/Engine/Plugins/PCG/Source/PCG/Public/Compute/PCGDataForGPU.h b/Engine/Plugins/PCG/Source/PCG/Public/Compute/PCGDataForGPU.h
--- a/Engine/Plugins/PCG/Source/PCG/Public/Compute/PCGDataForGPU.h
+++ b/Engine/Plugins/PCG/Source/PCG/Public/Compute/PCGDataForGPU.h
@@ -42,6 +42,13 @@ struct FPCGKernelAttributeKey
 
 struct FPCGKernelAttributeDesc
 {
+	explicit FPCGKernelAttributeDesc(int32 InIndex, EPCGKernelAttributeType InType, FName InName)
+		: Index(InIndex)
+		, Type(InType)
+		, Name(InName)
+	{
+	}
+
 	int32 Index = INDEX_NONE;
 	EPCGKernelAttributeType Type = EPCGKernelAttributeType::Float;
 	FName Name = NAME_None;
