[PCG] Add command to refresh all runtime gen components
* 'pcg.RuntimeGeneration.Refresh'

diff --git a/Engine/Plugins/PCG/Source/PCG/Private/PCGSubsystem.cpp b/Engine/Plugins/PCG/Source/PCG/Private/PCGSubsystem.cpp
--- a/Engine/Plugins/PCG/Source/PCG/Private/PCGSubsystem.cpp
+++ b/Engine/Plugins/PCG/Source/PCG/Private/PCGSubsystem.cpp
@@ -114,6 +114,17 @@ namespace PCGSubsystemConsole
 				PCGSubsystem->DestroyAllPCGWorldActors();
 			}
 		}));
+
+	static FAutoConsoleCommand CommandRefreshRuntimeGen(
+		TEXT("pcg.RuntimeGeneration.Refresh"),
+		TEXT("Cleans up and re-generates all GenerateAtRuntime PCG components, including their partition actors."),
+		FConsoleCommandDelegate::CreateLambda([]()
+		{
+			if (UPCGSubsystem* PCGSubsystem = UPCGSubsystem::GetSubsystemForCurrentWorld())
+			{
+				PCGSubsystem->RefreshAllRuntimeGenComponents(EPCGChangeType::GenerationGrid);
+			}
+		}));
 #endif
 }
 
@@ -908,6 +919,14 @@ void UPCGSubsystem::RefreshRuntimeGenComponent(UPCGComponent* RuntimeComponent,
 	}
 }
 
+void UPCGSubsystem::RefreshAllRuntimeGenComponents(EPCGChangeType ChangeType)
+{
+	for (UPCGComponent* Component : GetAllRegisteredComponents())
+	{
+		RefreshRuntimeGenComponent(Component, ChangeType);
+	}
+}
+
 bool UPCGSubsystem::IsGraphCurrentlyExecuting(UPCGGraph* Graph)
 {
 	check(GraphExecutor);
diff --git a/Engine/Plugins/PCG/Source/PCG/Public/PCGSubsystem.h b/Engine/Plugins/PCG/Source/PCG/Public/PCGSubsystem.h
--- a/Engine/Plugins/PCG/Source/PCG/Public/PCGSubsystem.h
+++ b/Engine/Plugins/PCG/Source/PCG/Public/PCGSubsystem.h
@@ -169,9 +169,12 @@ public:
 	/** Cancels currently running generation on given graph */
 	void CancelGeneration(UPCGGraph* Graph);
 
-	/** Asks the runtime generation scheduler to refresh a given GenerateAtRuntime component. bRemovePartitionActors will also perform a full cleanup of PAs and local components. */
+	/** Asks the runtime generation scheduler to refresh a given GenerateAtRuntime component. ChangeType should be 'GenerationGrid' to perform a full cleanup of PAs and local components. */
 	void RefreshRuntimeGenComponent(UPCGComponent* RuntimeComponent, EPCGChangeType ChangeType = EPCGChangeType::None);
 
+	/** Asks the runtime generation scheduler to refresh all GenerateAtRuntime components. ChangeType should be 'GenerationGrid' to perform a full cleanup of PAs and local components. */
+	void RefreshAllRuntimeGenComponents(EPCGChangeType ChangeType = EPCGChangeType::None);
+
 	FPCGRuntimeGenScheduler* GetRuntimeGenScheduler() const { return RuntimeGenScheduler; }
 
 	/** Returns true if there are any tasks for this graph currently scheduled or executing. */
diff --git a/Engine/Plugins/PCG/Source/PCGEditor/Private/PCGEditorModule.cpp b/Engine/Plugins/PCG/Source/PCGEditor/Private/PCGEditorModule.cpp
--- a/Engine/Plugins/PCG/Source/PCGEditor/Private/PCGEditorModule.cpp
+++ b/Engine/Plugins/PCG/Source/PCGEditor/Private/PCGEditorModule.cpp
@@ -566,6 +566,19 @@ void FPCGEditorModule::PopulateMenuActions(FMenuBuilder& MenuBuilder)
 		false,
 		FSlateIcon());
 
+	MenuBuilder.AddMenuEntry(
+		LOCTEXT("RefreshRuntimeGen", "Refresh all runtime gen components"),
+		LOCTEXT("RefreshRuntimeGen_Tooltip", "Cleans up and re-generates all GenerateAtRuntime PCG components, including their partition actors."),
+		FSlateIcon(),
+		FUIAction(
+			FExecuteAction::CreateLambda([]()
+			{
+				if (UPCGSubsystem* PCGSubsystem = UPCGSubsystem::GetSubsystemForCurrentWorld())
+				{
+					PCGSubsystem->RefreshAllRuntimeGenComponents(EPCGChangeType::GenerationGrid);
+				}
+			})),
+		NAME_None);
 }
 
 void FPCGEditorModule::RegisterSettings()
