Added custom PCG ISM descriptors so we can manage our own defaults.
Added ability to not allow component class changes.
Updated ISM params on debug draw.

diff --git a/Engine/Plugins/PCG/Source/PCG/Private/Elements/PCGStaticMeshSpawner.cpp b/Engine/Plugins/PCG/Source/PCG/Private/Elements/PCGStaticMeshSpawner.cpp
--- a/Engine/Plugins/PCG/Source/PCG/Private/Elements/PCGStaticMeshSpawner.cpp
+++ b/Engine/Plugins/PCG/Source/PCG/Private/Elements/PCGStaticMeshSpawner.cpp
@@ -434,6 +434,8 @@ void FPCGStaticMeshSpawnerElement::SpawnStaticMeshInstances(FPCGStaticMeshSpawne
 	const UPCGStaticMeshSpawnerSettings* Settings = Context->GetInputSettings<UPCGStaticMeshSpawnerSettings>();
 	check(Settings);
 
+	Params.bAllowDescriptorChanges = Settings->bAllowDescriptorChanges;
+
 	UPCGManagedISMComponent* MISMC = UPCGActorHelpers::GetOrCreateManagedISMC(TargetActor, Context->SourceComponent.Get(), Settings->UID, Params);
 
 	check(MISMC);
diff --git a/Engine/Plugins/PCG/Source/PCG/Private/Helpers/PCGActorHelpers.cpp b/Engine/Plugins/PCG/Source/PCG/Private/Helpers/PCGActorHelpers.cpp
--- a/Engine/Plugins/PCG/Source/PCG/Private/Helpers/PCGActorHelpers.cpp
+++ b/Engine/Plugins/PCG/Source/PCG/Private/Helpers/PCGActorHelpers.cpp
@@ -50,10 +50,37 @@ UPCGManagedISMComponent* UPCGActorHelpers::GetOrCreateManagedISMC(AActor* InTarg
 
 	TRACE_CPUPROFILER_EVENT_SCOPE(UPCGActorHelpers::GetOrCreateManagedISMC);
 
+	FISMComponentDescriptor Descriptor = InParams.Descriptor;
+
+	// If the component class is invalid, default to HISM.
+	// TODO: should this be part of the descriptor changes?
+	if (!Descriptor.ComponentClass)
+	{
+		Descriptor.ComponentClass = UHierarchicalInstancedStaticMeshComponent::StaticClass();
+	}
+
+	if (InParams.bAllowDescriptorChanges)
+	{
+		if (Descriptor.ComponentClass == UHierarchicalInstancedStaticMeshComponent::StaticClass())
+		{
+			// Done as in InstancedStaticMesh.cpp
+#if WITH_EDITOR
+			const bool bMeshHasNaniteData = StaticMesh->IsNaniteEnabled();
+#else
+			const bool bMeshHasNaniteData = StaticMesh->GetRenderData() && StaticMesh->GetRenderData()->HasValidNaniteData();
+#endif
+
+			if (bMeshHasNaniteData)
+			{
+				Descriptor.ComponentClass = UInstancedStaticMeshComponent::StaticClass();
+			}
+		}
+	}
+
 	{
 		TRACE_CPUPROFILER_EVENT_SCOPE(UPCGActorHelpers::GetOrCreateManagedISMC::FindMatchingMISMC);
 		UPCGManagedISMComponent* MatchingResource = nullptr;
-		InSourceComponent->ForEachManagedResource([&MatchingResource, &InParams, &InTargetActor, SettingsUID](UPCGManagedResource* InResource)
+		InSourceComponent->ForEachManagedResource([&MatchingResource, &InParams, &InTargetActor, &Descriptor, SettingsUID](UPCGManagedResource* InResource)
 		{
 			// Early out if already found a match
 			if (MatchingResource)
@@ -73,7 +100,7 @@ UPCGManagedISMComponent* UPCGActorHelpers::GetOrCreateManagedISMC(AActor* InTarg
 					if (IsValid(ISMC) &&
 						ISMC->GetOwner() == InTargetActor &&
 						ISMC->NumCustomDataFloats == InParams.NumCustomDataFloats &&
-						Resource->GetDescriptor() == InParams.Descriptor)
+						Resource->GetDescriptor() == Descriptor)
 					{
 						MatchingResource = Resource;
 					}
@@ -93,37 +120,15 @@ UPCGManagedISMComponent* UPCGActorHelpers::GetOrCreateManagedISMC(AActor* InTarg
 	// No matching ISM component found, let's create a new one
 	InTargetActor->Modify(!InSourceComponent->IsInPreviewMode());
 
-	// Done as in InstancedStaticMesh.cpp
-#if WITH_EDITOR
-	const bool bMeshHasNaniteData = StaticMesh->IsNaniteEnabled();
-#else
-	const bool bMeshHasNaniteData = StaticMesh->GetRenderData() && StaticMesh->GetRenderData()->HasValidNaniteData();
-#endif
-
 	FString ComponentName;
-	TSubclassOf<UInstancedStaticMeshComponent> ComponentClass = InParams.Descriptor.ComponentClass;
-
-	// If the component class in invalid, default to HISM.
-	if (!ComponentClass)
-	{
-		ComponentClass = UHierarchicalInstancedStaticMeshComponent::StaticClass();
-	}
 
 	// It's potentially less efficient to put nanite meshes inside of HISMs so decay those to ISM in this case.
 	// Note the equality here, not a IsA because we do not want to change derived types either
-	if (ComponentClass == UHierarchicalInstancedStaticMeshComponent::StaticClass())
+	if (Descriptor.ComponentClass == UHierarchicalInstancedStaticMeshComponent::StaticClass())
 	{
-		if (bMeshHasNaniteData)
-		{
-			ComponentClass = UInstancedStaticMeshComponent::StaticClass();
-		}
-		else
-		{
-			ComponentName = TEXT("HISM_");
-		}
+		ComponentName = TEXT("HISM_");
 	}
-
-	if (ComponentClass == UInstancedStaticMeshComponent::StaticClass())
+	else if (Descriptor.ComponentClass == UInstancedStaticMeshComponent::StaticClass())
 	{
 		ComponentName = TEXT("ISM_");
 	}
@@ -131,8 +136,8 @@ UPCGManagedISMComponent* UPCGActorHelpers::GetOrCreateManagedISMC(AActor* InTarg
 	ComponentName += StaticMesh->GetName();
 
 	const EObjectFlags ObjectFlags = (InSourceComponent->IsInPreviewMode() ? RF_Transient : RF_NoFlags);
-	UInstancedStaticMeshComponent* ISMC = NewObject<UInstancedStaticMeshComponent>(InTargetActor, ComponentClass, MakeUniqueObjectName(InTargetActor, ComponentClass, FName(ComponentName)), ObjectFlags);
-	InParams.Descriptor.InitComponent(ISMC);
+	UInstancedStaticMeshComponent* ISMC = NewObject<UInstancedStaticMeshComponent>(InTargetActor, Descriptor.ComponentClass, MakeUniqueObjectName(InTargetActor, Descriptor.ComponentClass, FName(ComponentName)), ObjectFlags);
+	Descriptor.InitComponent(ISMC);
 	ISMC->SetNumCustomDataFloats(InParams.NumCustomDataFloats);
 
 	ISMC->RegisterComponent();
@@ -145,7 +150,7 @@ UPCGManagedISMComponent* UPCGActorHelpers::GetOrCreateManagedISMC(AActor* InTarg
 	// Create managed resource on source component
 	UPCGManagedISMComponent* Resource = NewObject<UPCGManagedISMComponent>(InSourceComponent);
 	Resource->SetComponent(ISMC);
-	Resource->SetDescriptor(InParams.Descriptor);
+	Resource->SetDescriptor(Descriptor);
 	if (InTargetActor->GetRootComponent())
 	{
 		Resource->SetRootLocation(InTargetActor->GetRootComponent()->GetComponentLocation());
diff --git a/Engine/Plugins/PCG/Source/PCG/Private/MeshSelectors/PCGISMDescriptor.cpp b/Engine/Plugins/PCG/Source/PCG/Private/MeshSelectors/PCGISMDescriptor.cpp
new file mode 100644
--- /dev/null
+++ b/Engine/Plugins/PCG/Source/PCG/Private/MeshSelectors/PCGISMDescriptor.cpp
@@ -0,0 +1,30 @@
+// Copyright Epic Games, Inc. All Rights Reserved.
+
+#include "MeshSelectors/PCGISMDescriptor.h"
+
+#include "Components/InstancedStaticMeshComponent.h"
+
+#include UE_INLINE_GENERATED_CPP_BY_NAME(PCGISMDescriptor)
+
+FPCGSoftISMComponentDescriptor::FPCGSoftISMComponentDescriptor()
+	: FSoftISMComponentDescriptor()
+{
+	// Override defaults from base class
+	ComponentClass = UInstancedStaticMeshComponent::StaticClass();
+	Mobility = EComponentMobility::Static;
+	bGenerateOverlapEvents = 0;
+	BodyInstance.SetCollisionProfileName(UCollisionProfile::NoCollision_ProfileName);
+}
+
+bool FPCGSoftISMComponentDescriptor::SerializeFromMismatchedTag(const FPropertyTag& Tag, FStructuredArchive::FSlot Slot)
+{
+	if (Tag.GetType().IsStruct(FSoftISMComponentDescriptor::StaticStruct()->GetFName()))
+	{
+		FSoftISMComponentDescriptor TempDescriptor;
+		FSoftISMComponentDescriptor::StaticStruct()->SerializeItem(Slot, &TempDescriptor, nullptr);
+		*static_cast<FSoftISMComponentDescriptor*>(this) = TempDescriptor;
+		return true;
+	}
+
+	return false;
+}
\ No newline at end of file
diff --git a/Engine/Plugins/PCG/Source/PCG/Private/MeshSelectors/PCGMeshSelectorByAttribute.cpp b/Engine/Plugins/PCG/Source/PCG/Private/MeshSelectors/PCGMeshSelectorByAttribute.cpp
--- a/Engine/Plugins/PCG/Source/PCG/Private/MeshSelectors/PCGMeshSelectorByAttribute.cpp
+++ b/Engine/Plugins/PCG/Source/PCG/Private/MeshSelectors/PCGMeshSelectorByAttribute.cpp
@@ -24,7 +24,7 @@ namespace PCGMeshSelectorAttribute
 	// Returns variation based on mesh, material overrides and reverse culling
 	FPCGMeshInstanceList& GetInstanceList(
 		TArray<FPCGMeshInstanceList>& InstanceLists,
-		const FSoftISMComponentDescriptor& TemplateDescriptor,
+		const FPCGSoftISMComponentDescriptor& TemplateDescriptor,
 		TSoftObjectPtr<UStaticMesh> Mesh,
 		const TArray<TSoftObjectPtr<UMaterialInterface>>& MaterialOverrides,
 		bool bReverseCulling,
@@ -170,7 +170,7 @@ bool UPCGMeshSelectorByAttribute::SelectInstances(
 		TRACE_CPUPROFILER_EVENT_SCOPE(UPCGMeshSelectorByAttribute::SelectEntries::BuildingPartition);
 
 		TArray<FPCGObjectPropertyOverrideDescription> PropertyOverrides = Settings->StaticMeshComponentPropertyOverrides;
-		const FString StaticMeshPropertyString = GET_MEMBER_NAME_CHECKED(FSoftISMComponentDescriptor, StaticMesh).ToString();
+		const FString StaticMeshPropertyString = GET_MEMBER_NAME_CHECKED(FPCGSoftISMComponentDescriptor, StaticMesh).ToString();
 
 		// Add the Static mesh override to the list only if not already provided
 		if (!Algo::AnyOf(PropertyOverrides, [&StaticMeshPropertyString](const FPCGObjectPropertyOverrideDescription& PropertyOverride) { return PropertyOverride.PropertyTarget == StaticMeshPropertyString; }))
@@ -180,12 +180,12 @@ bool UPCGMeshSelectorByAttribute::SelectInstances(
 			PropertyOverrides.Emplace(MeshSelector, StaticMeshPropertyString);
 		}
 
-		// Validate all the selectors are actual FSoftISMComponentDescriptor properties
+		// Validate all the selectors are actual FPCGSoftISMComponentDescriptor properties
 		TArray<FPCGAttributePropertySelector> ValidSelectorOverrides;
 
 		for (const FPCGObjectPropertyOverrideDescription& PropertyOverride : PropertyOverrides)
 		{
-			if (FSoftISMComponentDescriptor::StaticStruct()->FindPropertyByName(FName(PropertyOverride.PropertyTarget)))
+			if (FPCGSoftISMComponentDescriptor::StaticStruct()->FindPropertyByName(FName(PropertyOverride.PropertyTarget)))
 			{
 				ValidSelectorOverrides.Emplace_GetRef() = FPCGAttributePropertySelector::CreateFromOtherSelector<FPCGAttributePropertySelector>(PropertyOverride.InputSource);
 			}
@@ -205,7 +205,7 @@ bool UPCGMeshSelectorByAttribute::SelectInstances(
 		Context.OverriddenDescriptors.Reserve(Context.AttributeOverridePartition.Num());
 		for (int I = 0; I < Context.AttributeOverridePartition.Num(); ++I)
 		{
-			FSoftISMComponentDescriptor& Descriptor = Context.OverriddenDescriptors.Add_GetRef(TemplateDescriptor);
+			FPCGSoftISMComponentDescriptor& Descriptor = Context.OverriddenDescriptors.Add_GetRef(TemplateDescriptor);
 
 			// If partition is empty (which can happen, esp. for the default partition on the default value, we'll just skip it here.
 			if (Context.AttributeOverridePartition[I].IsEmpty())
@@ -243,7 +243,7 @@ bool UPCGMeshSelectorByAttribute::SelectInstances(
 		{
 			const int32 PartitionIndex = CurrentPartitionIndex++;
 			TArray<int32>& Partition = Context.AttributeOverridePartition[PartitionIndex];
-			const FSoftISMComponentDescriptor& CurrentPartitionDescriptor = Context.OverriddenDescriptors[PartitionIndex];
+			const FPCGSoftISMComponentDescriptor& CurrentPartitionDescriptor = Context.OverriddenDescriptors[PartitionIndex];
 			
 			if (Partition.IsEmpty() || CurrentPartitionDescriptor.StaticMesh.IsNull())
 			{
diff --git a/Engine/Plugins/PCG/Source/PCG/Private/MeshSelectors/PCGMeshSelectorWeighted.cpp b/Engine/Plugins/PCG/Source/PCG/Private/MeshSelectors/PCGMeshSelectorWeighted.cpp
--- a/Engine/Plugins/PCG/Source/PCG/Private/MeshSelectors/PCGMeshSelectorWeighted.cpp
+++ b/Engine/Plugins/PCG/Source/PCG/Private/MeshSelectors/PCGMeshSelectorWeighted.cpp
@@ -330,7 +330,7 @@ void UPCGMeshSelectorWeighted::PostEditChangeProperty(FPropertyChangedEvent& Pro
 {
 	Super::PostEditChangeProperty(PropertyChangedEvent);
 
-	if (PropertyChangedEvent.GetPropertyName() == GET_MEMBER_NAME_CHECKED(FSoftISMComponentDescriptor, StaticMesh))
+	if (PropertyChangedEvent.GetPropertyName() == GET_MEMBER_NAME_CHECKED(FPCGSoftISMComponentDescriptor, StaticMesh))
 	{
 		RefreshDisplayNames();
 	}
diff --git a/Engine/Plugins/PCG/Source/PCG/Private/MeshSelectors/PCGMeshSelectorWeightedByCategory.cpp b/Engine/Plugins/PCG/Source/PCG/Private/MeshSelectors/PCGMeshSelectorWeightedByCategory.cpp
--- a/Engine/Plugins/PCG/Source/PCG/Private/MeshSelectors/PCGMeshSelectorWeightedByCategory.cpp
+++ b/Engine/Plugins/PCG/Source/PCG/Private/MeshSelectors/PCGMeshSelectorWeightedByCategory.cpp
@@ -342,7 +342,7 @@ void UPCGMeshSelectorWeightedByCategory::PostEditChangeProperty(FPropertyChanged
 {
 	Super::PostEditChangeProperty(PropertyChangedEvent);
 
-	if (PropertyChangedEvent.GetPropertyName() == GET_MEMBER_NAME_CHECKED(FSoftISMComponentDescriptor, StaticMesh))
+	if (PropertyChangedEvent.GetPropertyName() == GET_MEMBER_NAME_CHECKED(FPCGSoftISMComponentDescriptor, StaticMesh))
 	{
 		RefreshDisplayNames();
 	}
diff --git a/Engine/Plugins/PCG/Source/PCG/Public/Elements/PCGStaticMeshSpawner.h b/Engine/Plugins/PCG/Source/PCG/Public/Elements/PCGStaticMeshSpawner.h
--- a/Engine/Plugins/PCG/Source/PCG/Public/Elements/PCGStaticMeshSpawner.h
+++ b/Engine/Plugins/PCG/Source/PCG/Public/Elements/PCGStaticMeshSpawner.h
@@ -66,6 +66,10 @@ public:
 	UPROPERTY(BlueprintReadOnly, VisibleAnywhere, Instanced, Category = MeshSelector)
 	TObjectPtr<UPCGMeshSelectorBase> MeshSelectorParameters;
 
+	/** Allows PCG to make some changes on the descriptors as situation arises (using ISM instead of HISM for nanite meshes, etc.) */
+	UPROPERTY(BlueprintReadOnly, EditAnywhere, Category = MeshSelector)
+	bool bAllowDescriptorChanges = true;
+
 	/** Defines the method of custom data packing for spawned (H)ISMCs */
 	UPROPERTY(BlueprintReadOnly, EditAnywhere, Category = InstanceDataPacker)
 	TSubclassOf<UPCGInstanceDataPackerBase> InstanceDataPackerType;
diff --git a/Engine/Plugins/PCG/Source/PCG/Public/Elements/PCGStaticMeshSpawnerContext.h b/Engine/Plugins/PCG/Source/PCG/Public/Elements/PCGStaticMeshSpawnerContext.h
--- a/Engine/Plugins/PCG/Source/PCG/Public/Elements/PCGStaticMeshSpawnerContext.h
+++ b/Engine/Plugins/PCG/Source/PCG/Public/Elements/PCGStaticMeshSpawnerContext.h
@@ -69,7 +69,7 @@ struct FPCGStaticMeshSpawnerContext : public FPCGContext, public IPCGAsyncLoadin
 
 	// Used for mesh property overrides
 	TArray<TArray<int32>> AttributeOverridePartition;
-	TArray<FSoftISMComponentDescriptor> OverriddenDescriptors;
+	TArray<FPCGSoftISMComponentDescriptor> OverriddenDescriptors;
 
 	// Keeping track of all touched resources to allow for correct cleanup on abort
 	TArray<TWeakObjectPtr<UPCGManagedISMComponent>> TouchedResources;
diff --git a/Engine/Plugins/PCG/Source/PCG/Public/Helpers/PCGActorHelpers.h b/Engine/Plugins/PCG/Source/PCG/Public/Helpers/PCGActorHelpers.h
--- a/Engine/Plugins/PCG/Source/PCG/Public/Helpers/PCGActorHelpers.h
+++ b/Engine/Plugins/PCG/Source/PCG/Public/Helpers/PCGActorHelpers.h
@@ -30,13 +30,14 @@ struct FPCGISMCBuilderParameters
 {
 	FISMComponentDescriptor Descriptor;
 	int32 NumCustomDataFloats = 0;
+	bool bAllowDescriptorChanges = true;
 
 	friend inline uint32 GetTypeHash(const FPCGISMCBuilderParameters& Key)
 	{
 		return HashCombine(GetTypeHash(Key.Descriptor), 1 + Key.NumCustomDataFloats);
 	}
 
-	inline bool operator==(const FPCGISMCBuilderParameters& Other) const { return Descriptor == Other.Descriptor && NumCustomDataFloats == Other.NumCustomDataFloats; }
+	inline bool operator==(const FPCGISMCBuilderParameters& Other) const { return Descriptor == Other.Descriptor && NumCustomDataFloats == Other.NumCustomDataFloats && bAllowDescriptorChanges == Other.bAllowDescriptorChanges; }
 };
 
 struct FPCGSplineMeshComponentBuilderParameters
diff --git a/Engine/Plugins/PCG/Source/PCG/Public/MeshSelectors/PCGISMDescriptor.h b/Engine/Plugins/PCG/Source/PCG/Public/MeshSelectors/PCGISMDescriptor.h
new file mode 100644
--- /dev/null
+++ b/Engine/Plugins/PCG/Source/PCG/Public/MeshSelectors/PCGISMDescriptor.h
@@ -0,0 +1,29 @@
+// Copyright Epic Games, Inc. All Rights Reserved.
+
+#pragma once
+
+#include "Engine/CollisionProfile.h"
+#include "ISMPartition/ISMComponentDescriptor.h"
+
+#include "PCGISMDescriptor.generated.h"
+
+/** Convenience PCG-side component descriptor so we can adjust defaults to the most common use cases. */
+USTRUCT()
+struct FPCGSoftISMComponentDescriptor : public FSoftISMComponentDescriptor
+{
+	GENERATED_BODY()
+
+	PCG_API FPCGSoftISMComponentDescriptor();
+
+public:
+	PCG_API bool SerializeFromMismatchedTag(const FPropertyTag& Tag, FStructuredArchive::FSlot Slot);
+};
+
+template<>
+struct TStructOpsTypeTraits<FPCGSoftISMComponentDescriptor> : public TStructOpsTypeTraitsBase2<FPCGSoftISMComponentDescriptor>
+{
+	enum
+	{
+		WithStructuredSerializeFromMismatchedTag = true,
+	};
+};
\ No newline at end of file
diff --git a/Engine/Plugins/PCG/Source/PCG/Public/MeshSelectors/PCGMeshSelectorBase.h b/Engine/Plugins/PCG/Source/PCG/Public/MeshSelectors/PCGMeshSelectorBase.h
--- a/Engine/Plugins/PCG/Source/PCG/Public/MeshSelectors/PCGMeshSelectorBase.h
+++ b/Engine/Plugins/PCG/Source/PCG/Public/MeshSelectors/PCGMeshSelectorBase.h
@@ -3,10 +3,10 @@
 #pragma once
 
 #include "PCGPoint.h"
+#include "MeshSelectors/PCGISMDescriptor.h"
 #include "Metadata/PCGMetadata.h"
 
 #include "Engine/CollisionProfile.h"
-#include "ISMPartition/ISMComponentDescriptor.h"
 
 #include "PCGMeshSelectorBase.generated.h"
 
@@ -26,13 +26,13 @@ struct FPCGMeshInstanceList
 
 	FPCGMeshInstanceList() = default;
 
-	explicit FPCGMeshInstanceList(const FSoftISMComponentDescriptor& InDescriptor)
+	explicit FPCGMeshInstanceList(const FPCGSoftISMComponentDescriptor& InDescriptor)
 		: Descriptor(InDescriptor)
 		, AttributePartitionIndex(INDEX_NONE)
 	{}
 
 	UPROPERTY(EditAnywhere, Category = Settings)
-	FSoftISMComponentDescriptor Descriptor;
+	FPCGSoftISMComponentDescriptor Descriptor;
 
 	UPROPERTY(BlueprintReadWrite, EditAnywhere, Category = Settings)
 	TArray<FTransform> Instances;
diff --git a/Engine/Plugins/PCG/Source/PCG/Public/MeshSelectors/PCGMeshSelectorByAttribute.h b/Engine/Plugins/PCG/Source/PCG/Public/MeshSelectors/PCGMeshSelectorByAttribute.h
--- a/Engine/Plugins/PCG/Source/PCG/Public/MeshSelectors/PCGMeshSelectorByAttribute.h
+++ b/Engine/Plugins/PCG/Source/PCG/Public/MeshSelectors/PCGMeshSelectorByAttribute.h
@@ -2,7 +2,8 @@
 
 #pragma once
 
-#include "PCGMeshSelectorBase.h"
+#include "MeshSelectors/PCGMeshSelectorBase.h"
+#include "MeshSelectors/PCGISMDescriptor.h"
 
 #include "PCGMeshSelectorByAttribute.generated.h"
 
@@ -28,7 +29,7 @@ public:
 	FName AttributeName;
 
 	UPROPERTY(EditAnywhere, Category = MeshSelector)
-	FSoftISMComponentDescriptor TemplateDescriptor;
+	FPCGSoftISMComponentDescriptor TemplateDescriptor;
 
 	UPROPERTY(BlueprintReadWrite, EditAnywhere, Category = MeshSelector, meta = (InlineEditConditionToggle))
 	bool bUseAttributeMaterialOverrides = false;
@@ -63,4 +64,4 @@ public:
 	UPROPERTY()
 	int32 WorldPositionOffsetDisableDistance_DEPRECATED = 0;
 #endif // WITH_EDITORONLY_DATA
-};
+};
\ No newline at end of file
diff --git a/Engine/Plugins/PCG/Source/PCG/Public/MeshSelectors/PCGMeshSelectorWeighted.h b/Engine/Plugins/PCG/Source/PCG/Public/MeshSelectors/PCGMeshSelectorWeighted.h
--- a/Engine/Plugins/PCG/Source/PCG/Public/MeshSelectors/PCGMeshSelectorWeighted.h
+++ b/Engine/Plugins/PCG/Source/PCG/Public/MeshSelectors/PCGMeshSelectorWeighted.h
@@ -2,7 +2,8 @@
 
 #pragma once
 
-#include "PCGMeshSelectorBase.h"
+#include "MeshSelectors/PCGMeshSelectorBase.h"
+#include "MeshSelectors/PCGISMDescriptor.h"
 
 #include "PCGMeshSelectorWeighted.generated.h"
 
@@ -29,7 +30,7 @@ struct PCG_API FPCGMeshSelectorWeightedEntry
 #endif // WITH_EDITOR
 
 	UPROPERTY(EditAnywhere, Category = Settings)
-	FSoftISMComponentDescriptor Descriptor;
+	FPCGSoftISMComponentDescriptor Descriptor;
 
 	UPROPERTY(BlueprintReadWrite, EditAnywhere, Category = Settings, meta = (ClampMin = "0"))
 	int Weight = 1;
@@ -37,9 +38,7 @@ struct PCG_API FPCGMeshSelectorWeightedEntry
 #if WITH_EDITORONLY_DATA
 	UPROPERTY(Transient, VisibleAnywhere, Category = Settings)
 	FName DisplayName = NAME_None;
-#endif
 
-#if WITH_EDITORONLY_DATA
 	UPROPERTY()
 	TSoftObjectPtr<UStaticMesh> Mesh_DEPRECATED;
 
diff --git a/Engine/Plugins/PCG/Source/PCGEditor/Private/DataVisualizations/PCGSpatialDataVisualization.cpp b/Engine/Plugins/PCG/Source/PCGEditor/Private/DataVisualizations/PCGSpatialDataVisualization.cpp
--- a/Engine/Plugins/PCG/Source/PCGEditor/Private/DataVisualizations/PCGSpatialDataVisualization.cpp
+++ b/Engine/Plugins/PCG/Source/PCGEditor/Private/DataVisualizations/PCGSpatialDataVisualization.cpp
@@ -166,9 +166,23 @@ void IPCGSpatialDataVisualization::ExecuteDebugDisplay(FPCGContext* Context, con
 	Params[0].NumCustomDataFloats = NumCustomData;
 	Params[0].Descriptor.StaticMesh = Mesh;
 	Params[0].Descriptor.OverrideMaterials = Materials;
+	Params[0].Descriptor.Mobility = EComponentMobility::Static;
 	Params[0].Descriptor.BodyInstance.SetCollisionProfileName(UCollisionProfile::NoCollision_ProfileName);
+
 	// Note: In the future we may consider enabling culling for performance reasons, but for now culling disabled.
 	Params[0].Descriptor.InstanceStartCullDistance = Params[0].Descriptor.InstanceEndCullDistance = 0;
+	// Additional performance switches
+	Params[0].Descriptor.bAffectDistanceFieldLighting = false;
+	Params[0].Descriptor.bAffectDynamicIndirectLighting = false;
+	Params[0].Descriptor.bAffectDynamicIndirectLightingWhileHidden = false;
+	Params[0].Descriptor.bCastContactShadow = false;
+	Params[0].Descriptor.bCastDynamicShadow = false;
+	Params[0].Descriptor.bCastShadow = false;	
+	Params[0].Descriptor.bCastStaticShadow = false;
+	Params[0].Descriptor.bGenerateOverlapEvents = false;
+	Params[0].Descriptor.bIncludeInHLOD = false;
+	Params[0].Descriptor.bReceivesDecals = false;
+	Params[0].Descriptor.bVisibleInRayTracing = false;
 	
 	// If the root actor we're binding to is movable, then the ISMC should be movable by default
 	if (USceneComponent* SceneComponent = TargetActor->GetRootComponent())
