[PCG] Improve Set values on metadata, by adding a preparation phase when we know how many items will be written

#rnx

diff --git a/Engine/Plugins/PCG/Source/PCG/Private/Elements/Metadata/PCGMetadataOpElementBase.cpp b/Engine/Plugins/PCG/Source/PCG/Private/Elements/Metadata/PCGMetadataOpElementBase.cpp
--- a/Engine/Plugins/PCG/Source/PCG/Private/Elements/Metadata/PCGMetadataOpElementBase.cpp
+++ b/Engine/Plugins/PCG/Source/PCG/Private/Elements/Metadata/PCGMetadataOpElementBase.cpp
@@ -24,7 +24,12 @@ namespace PCGMetadataBase
 	TAutoConsoleVariable<int> CVarMetadataOperationChunkSize(
 		TEXT("pcg.MetadataOperationChunkSize"),
 		256,
-		TEXT("Metadata operations chunk size."));
+		TEXT("Metadata operations chunk size."));	
+	
+	TAutoConsoleVariable<bool> CVarMetadataOperationReserveValues(
+		TEXT("pcg.MetadataOperationReserveValues"),
+		true,
+		TEXT("Metadata operations reserve values."));
 }
 
 void UPCGMetadataSettingsBase::PostLoad()
diff --git a/Engine/Plugins/PCG/Source/PCG/Private/Metadata/Accessors/PCGAttributeAccessorKeys.cpp b/Engine/Plugins/PCG/Source/PCG/Private/Metadata/Accessors/PCGAttributeAccessorKeys.cpp
--- a/Engine/Plugins/PCG/Source/PCG/Private/Metadata/Accessors/PCGAttributeAccessorKeys.cpp
+++ b/Engine/Plugins/PCG/Source/PCG/Private/Metadata/Accessors/PCGAttributeAccessorKeys.cpp
@@ -79,12 +79,12 @@ void FPCGAttributeAccessorKeysEntries::InitializeFromMetadata(const UPCGMetadata
 	}
 }
 
-bool FPCGAttributeAccessorKeysEntries::GetMetadataEntryKeys(int32 InStart, TArrayView<PCGMetadataEntryKey*>& OutEntryKeys)
+bool FPCGAttributeAccessorKeysEntries::GetMetadataEntryKeys(int32 InStart, TArrayView<PCGMetadataEntryKey*> OutEntryKeys)
 {
 	return PCGAttributeAccessorKeys::GetKeys(Entries, InStart, OutEntryKeys, [](PCGMetadataEntryKey& Key) -> PCGMetadataEntryKey* { return &Key; });
 }
 
-bool FPCGAttributeAccessorKeysEntries::GetMetadataEntryKeys(int32 InStart, TArrayView<const PCGMetadataEntryKey*>& OutEntryKeys) const
+bool FPCGAttributeAccessorKeysEntries::GetMetadataEntryKeys(int32 InStart, TArrayView<const PCGMetadataEntryKey*> OutEntryKeys) const
 {
 	return PCGAttributeAccessorKeys::GetKeys(Entries, InStart, OutEntryKeys, [](const PCGMetadataEntryKey& Key) -> const PCGMetadataEntryKey* { return &Key; });
 }
@@ -109,32 +109,32 @@ FPCGAttributeAccessorKeysPoints::FPCGAttributeAccessorKeysPoints(const FPCGPoint
 	: FPCGAttributeAccessorKeysPoints(TArrayView<const FPCGPoint>(&InPoint, 1))
 {}
 
-bool FPCGAttributeAccessorKeysPoints::GetPointKeys(int32 InStart, TArrayView<FPCGPoint*>& OutPoints)
+bool FPCGAttributeAccessorKeysPoints::GetPointKeys(int32 InStart, TArrayView<FPCGPoint*> OutPoints)
 {
 	return PCGAttributeAccessorKeys::GetKeys(Points, InStart, OutPoints, [](FPCGPoint& Point) -> FPCGPoint* { return &Point; });
 }
 
-bool FPCGAttributeAccessorKeysPoints::GetPointKeys(int32 InStart, TArrayView<const FPCGPoint*>& OutPoints) const
+bool FPCGAttributeAccessorKeysPoints::GetPointKeys(int32 InStart, TArrayView<const FPCGPoint*> OutPoints) const
 {
 	return PCGAttributeAccessorKeys::GetKeys(Points, InStart, OutPoints, [](const FPCGPoint& Point) -> const FPCGPoint* { return &Point; });
 }
 
-bool FPCGAttributeAccessorKeysPoints::GetGenericObjectKeys(int32 InStart, TArrayView<void*>& OutObjects)
+bool FPCGAttributeAccessorKeysPoints::GetGenericObjectKeys(int32 InStart, TArrayView<void*> OutObjects)
 {
 	return PCGAttributeAccessorKeys::GetKeys(Points, InStart, OutObjects, [](FPCGPoint& Point) -> void* { return &Point; });
 }
 
-bool FPCGAttributeAccessorKeysPoints::GetGenericObjectKeys(int32 InStart, TArrayView<const void*>& OutObjects) const
+bool FPCGAttributeAccessorKeysPoints::GetGenericObjectKeys(int32 InStart, TArrayView<const void*> OutObjects) const
 {
 	return PCGAttributeAccessorKeys::GetKeys(Points, InStart, OutObjects, [](const FPCGPoint& Point) -> const void* { return &Point; });
 }
 
-bool FPCGAttributeAccessorKeysPoints::GetMetadataEntryKeys(int32 InStart, TArrayView<PCGMetadataEntryKey*>& OutEntryKeys)
+bool FPCGAttributeAccessorKeysPoints::GetMetadataEntryKeys(int32 InStart, TArrayView<PCGMetadataEntryKey*> OutEntryKeys)
 {
 	return PCGAttributeAccessorKeys::GetKeys(Points, InStart, OutEntryKeys, [](FPCGPoint& Point) -> PCGMetadataEntryKey* { return &(Point.MetadataEntry); });
 }
 
-bool FPCGAttributeAccessorKeysPoints::GetMetadataEntryKeys(int32 InStart, TArrayView<const PCGMetadataEntryKey*>& OutEntryKeys) const
+bool FPCGAttributeAccessorKeysPoints::GetMetadataEntryKeys(int32 InStart, TArrayView<const PCGMetadataEntryKey*> OutEntryKeys) const
 {
 	return PCGAttributeAccessorKeys::GetKeys(Points, InStart, OutEntryKeys, [](const FPCGPoint& Point) -> const PCGMetadataEntryKey* { return &(Point.MetadataEntry); });
-}
+}
\ No newline at end of file
diff --git a/Engine/Plugins/PCG/Source/PCG/Public/Elements/Metadata/PCGMetadataOpElementBase.h b/Engine/Plugins/PCG/Source/PCG/Public/Elements/Metadata/PCGMetadataOpElementBase.h
--- a/Engine/Plugins/PCG/Source/PCG/Public/Elements/Metadata/PCGMetadataOpElementBase.h
+++ b/Engine/Plugins/PCG/Source/PCG/Public/Elements/Metadata/PCGMetadataOpElementBase.h
@@ -22,6 +22,7 @@ namespace PCGMetadataBase
 {
 	extern TAutoConsoleVariable<bool> CVarMetadataOperationInMT;
 	extern TAutoConsoleVariable<int> CVarMetadataOperationChunkSize;
+	extern TAutoConsoleVariable<bool> CVarMetadataOperationReserveValues;
 }
 
 namespace PCGMetadataSettingsBaseConstants
@@ -431,12 +432,21 @@ inline bool FPCGMetadataElementBase::DoNAryOp(PCGMetadataOps::FOperationData& In
 
 	InOperationData.Validate<NbInputs, NbOutputs>();
 
-	EPCGAttributeAccessorFlags Flags = EPCGAttributeAccessorFlags::AllowBroadcast;
+	EPCGAttributeAccessorFlags Flags = EPCGAttributeAccessorFlags::AllowBroadcastAndConstructible;
 
 	// First set the default value (only on first pass)
 	PCG::Private::NAryOperation::Options Options{ Flags, Flags | EPCGAttributeAccessorFlags::AllowSetDefaultValue, true };
 	if (!InOperationData.Context->AsyncState.bStarted)
 	{
+		if (PCGMetadataBase::CVarMetadataOperationReserveValues.GetValueOnAnyThread())
+		{
+			for (int32 j = 0; j < NbOutputs; ++j)
+			{
+				// We can't re-use entry keys yet, it can be dangerous in some situations where some points share their entry key.
+				InOperationData.OutputAccessors[j]->Prepare(*InOperationData.OutputKeys[j], InOperationData.NumberOfElementsToProcess, /*bCanReuseEntryKeys=*/false);
+			}
+		}
+
 		PCG::Private::NAryOperation::Operation<InputTypes...>(InOperationData, /*StartIndex=*/0, /*Range=*/1, Options, InCallbacks);
 	}
 
diff --git a/Engine/Plugins/PCG/Source/PCG/Public/Metadata/Accessors/IPCGAttributeAccessor.h b/Engine/Plugins/PCG/Source/PCG/Public/Metadata/Accessors/IPCGAttributeAccessor.h
--- a/Engine/Plugins/PCG/Source/PCG/Public/Metadata/Accessors/IPCGAttributeAccessor.h
+++ b/Engine/Plugins/PCG/Source/PCG/Public/Metadata/Accessors/IPCGAttributeAccessor.h
@@ -27,6 +27,11 @@ enum class EPCGAttributeAccessorFlags
 	// USE WITH CAUTION
 	AllowSetDefaultValue = 1 << 3,
 
+	// New writes usually create a new metadata entry key when we write. In most cases, that's not mandatory, so use this flag to re-use an existing key.
+	// Only useful for writing to attributes.
+	// USE WITH CAUTION
+	AllowReuseMetadataEntryKey = 1 << 4,
+
 	AllowBroadcastAndConstructible = AllowBroadcast | AllowConstructible
 };
 ENUM_CLASS_FLAGS(EPCGAttributeAccessorFlags);
@@ -133,6 +138,8 @@ public:
 		return false;
 	}
 
+	virtual void Prepare(IPCGAttributeAccessorKeys& Keys, int32 Count, const bool bCanReuseEntryKeys) {}
+
 	int16 GetUnderlyingType() const { return UnderlyingType; }
 	bool IsReadOnly() const { return bReadOnly; }
 
diff --git a/Engine/Plugins/PCG/Source/PCG/Public/Metadata/Accessors/PCGAttributeAccessor.h b/Engine/Plugins/PCG/Source/PCG/Public/Metadata/Accessors/PCGAttributeAccessor.h
--- a/Engine/Plugins/PCG/Source/PCG/Public/Metadata/Accessors/PCGAttributeAccessor.h
+++ b/Engine/Plugins/PCG/Source/PCG/Public/Metadata/Accessors/PCGAttributeAccessor.h
@@ -59,61 +59,115 @@ public:
 
 	bool SetRangeImpl(TArrayView<const T> InValues, int32 Index, IPCGAttributeAccessorKeys& Keys, EPCGAttributeAccessorFlags Flags)
 	{
-		TArray<PCGMetadataEntryKey*, TInlineAllocator<256>> EntryKeys;
+		if (EnumHasAnyFlags(Flags, EPCGAttributeAccessorFlags::AllowSetDefaultValue) && Keys.GetNum() == 1)
+		{
+			PCGMetadataEntryKey* EntryKey = nullptr;
+			Keys.GetKey(EntryKey);
+
+			if (EntryKey && *EntryKey == PCGInvalidEntryKey)
+			{
+				check(!InValues.IsEmpty())
+				Attribute->SetDefaultValue(InValues[0]);
+				return true;
+			}
+		}
+
+		TArray<PCGMetadataEntryKey*, TInlineAllocator<512>> EntryKeys;
 		EntryKeys.SetNumUninitialized(InValues.Num());
 		TArrayView<PCGMetadataEntryKey*> EntryKeysView(EntryKeys);
+
+		int32 Start = Index;
+
+		if (!bWasPrepared || PCG::Private::MetadataTraits<T>::CompressData)
+		{
+			if (!Prepare(Keys, InValues.Num(), /*bCanReuseEntryKeys=*/EnumHasAnyFlags(Flags, EPCGAttributeAccessorFlags::AllowReuseMetadataEntryKey), Index, /*bPreallocateValues=*/false, &EntryKeysView, &Start))
+			{
+				return false;
+			}
+
+			{
+				//TRACE_CPUPROFILER_EVENT_SCOPE(FPCGAttributeAccessor::SetRangeImpl::SetValues);
+				Attribute->SetValues(EntryKeys, InValues);
+			}
+		}
+		else
+		{
+			//TRACE_CPUPROFILER_EVENT_SCOPE(FPCGAttributeAccessor::SetRangeImpl::SetValues);
+			Attribute->SetValues_TryLockless(EntryKeysView, InValues, Start);
+		}
+
+		return true;
+	}
+
+	virtual bool IsAttribute() const override { return true; }
+
+	bool Prepare(IPCGAttributeAccessorKeys& Keys, int32 Count, const bool bCanReuseEntryKeys, int32 Index, bool bPreallocateValues, TArrayView<PCGMetadataEntryKey*>* OutEntryKeys = nullptr, int32* OutStartIndex = nullptr)
+	{
+		//TRACE_CPUPROFILER_EVENT_SCOPE(FPCGAttributeAccessor::Prepare);
+
+		check(!OutEntryKeys || OutEntryKeys->Num() == Count);
+
+		TArray<PCGMetadataEntryKey*, TInlineAllocator<512>> EntryKeys;
+		TArrayView<PCGMetadataEntryKey*> EntryKeysView;
+		if (!OutEntryKeys)
+		{
+			EntryKeys.SetNumUninitialized(Count);
+			EntryKeysView = TArrayView<PCGMetadataEntryKey*>(EntryKeys);
+		}
+		else
+		{
+			EntryKeysView = *OutEntryKeys;
+		}
+
 		if (!Keys.GetKeys<PCGMetadataEntryKey>(Index, EntryKeysView))
 		{
 			return false;
 		}
 
-		int LastDefaultKeyIndex = INDEX_NONE;
-
-		TArray<PCGMetadataEntryKey*, TInlineAllocator<256>> EntriesToSet;
-		EntriesToSet.Reserve(EntryKeys.Num());
+		TArray<PCGMetadataEntryKey*, TInlineAllocator<512>> EntriesToSet;
+		EntriesToSet.Reserve(Count);
 
 		// Implementation note: this is a stripped down version of UPCGMetadata::InitializeOnSet
-		for(int EntryIndex = 0; EntryIndex < EntryKeys.Num(); ++EntryIndex)
+		for (int EntryIndex = 0; EntryIndex < EntryKeysView.Num(); ++EntryIndex)
 		{
-			PCGMetadataEntryKey& EntryKey = *EntryKeys[EntryIndex];
-			if (EntryKey == PCGInvalidEntryKey)
-			{
-				if (!(Flags & EPCGAttributeAccessorFlags::AllowSetDefaultValue))
-				{
-					EntriesToSet.Add(&EntryKey);
-				}
-				else
-				{
-					LastDefaultKeyIndex = EntryIndex;
-				}
-			}
-			else if (EntryKey < Metadata->GetItemKeyCountForParent())
+			PCGMetadataEntryKey& EntryKey = *EntryKeysView[EntryIndex];
+			if (EntryKey == PCGInvalidEntryKey || (EntryKey < Metadata->GetItemKeyCountForParent() && !bCanReuseEntryKeys))
 			{
 				EntriesToSet.Add(&EntryKey);
 			}
 		}
 
+		if (!EntriesToSet.IsEmpty())
 		{
-			//TRACE_CPUPROFILER_EVENT_SCOPE(FPCGAttributeAccessor::SetRangeImpl::AddEntriesInPlace);
+			TRACE_CPUPROFILER_EVENT_SCOPE(FPCGAttributeAccessor::Prepare::AddEntriesInPlace);
 			Metadata->AddEntriesInPlace(EntriesToSet);
 		}
 
+		if (bPreallocateValues && ensure(!bWasPrepared))
 		{
-			//TRACE_CPUPROFILER_EVENT_SCOPE(FPCGAttributeAccessor::SetRangeImpl::SetValues);
-			Attribute->SetValues(EntryKeys, InValues);
+			const int32 StartIndex = Attribute->PreallocateValues(EntryKeysView, /*bLockless=*/false);
+			if (OutStartIndex)
+			{
+				*OutStartIndex = StartIndex;
+			}
 		}
 
-		if (LastDefaultKeyIndex != INDEX_NONE)
+		return true;
+	}
+
+	virtual void Prepare(IPCGAttributeAccessorKeys& Keys, int32 Count, const bool bCanReuseEntryKeys) override
+	{
+		if (!ensure(!bWasPrepared))
 		{
-			Attribute->SetDefaultValue(InValues[LastDefaultKeyIndex]);
+			return;
 		}
 
-		return true;
+		Attribute->Prepare(Count);
+		bWasPrepared = Prepare(Keys, Count, bCanReuseEntryKeys, /*Index=*/0, /*bPreallocateValues=*/true);
 	}
 
-	virtual bool IsAttribute() const override { return true; }
-
 private:
 	FPCGMetadataAttribute<T>* Attribute = nullptr;
 	UPCGMetadata* Metadata = nullptr;
+	bool bWasPrepared = false;
 };
diff --git a/Engine/Plugins/PCG/Source/PCG/Public/Metadata/Accessors/PCGAttributeAccessorKeys.h b/Engine/Plugins/PCG/Source/PCG/Public/Metadata/Accessors/PCGAttributeAccessorKeys.h
--- a/Engine/Plugins/PCG/Source/PCG/Public/Metadata/Accessors/PCGAttributeAccessorKeys.h
+++ b/Engine/Plugins/PCG/Source/PCG/Public/Metadata/Accessors/PCGAttributeAccessorKeys.h
@@ -15,7 +15,7 @@ class UPCGMetadata;
 namespace PCGAttributeAccessorKeys
 {
 	template <typename T, typename Container, typename Func>
-	bool GetKeys(Container& InContainer, int32 InStart, TArrayView<T*>& OutItems, Func&& Transform)
+	bool GetKeys(Container& InContainer, int32 InStart, TArrayView<T*> OutItems, Func&& Transform)
 	{
 		if (InContainer.Num() == 0)
 		{
@@ -64,11 +64,11 @@ public:
 	* @return true if it succeeded, false otherwise. (like num == 0,unsupported type or read only)
 	*/
 	template <typename ObjectType>
-	bool GetKeys(int32 InStart, TArrayView<ObjectType*>& OutKeys);
+	bool GetKeys(int32 InStart, TArrayView<ObjectType*> OutKeys);
 
 	// Same function but const.
 	template <typename ObjectType>
-	bool GetKeys(int32 InStart, TArrayView<const ObjectType*>& OutKeys) const;
+	bool GetKeys(int32 InStart, TArrayView<const ObjectType*> OutKeys) const;
 
 	/**
 	* Retrieve in the given argument pointer of the wanted type.
@@ -118,14 +118,14 @@ public:
 	bool IsReadOnly() const { return bIsReadOnly; }
 
 protected:
-	virtual bool GetPointKeys(int32 InStart, TArrayView<FPCGPoint*>& OutPoints) { return false; }
-	virtual bool GetPointKeys(int32 InStart, TArrayView<const FPCGPoint*>& OutPoints) const { return false; }
+	virtual bool GetPointKeys(int32 InStart, TArrayView<FPCGPoint*> OutPoints) { return false; }
+	virtual bool GetPointKeys(int32 InStart, TArrayView<const FPCGPoint*> OutPoints) const { return false; }
 
-	virtual bool GetGenericObjectKeys(int32 InStart, TArrayView<void*>& OutObjects) { return false; }
-	virtual bool GetGenericObjectKeys(int32 InStart, TArrayView<const void*>& OutObjects) const { return false; }
+	virtual bool GetGenericObjectKeys(int32 InStart, TArrayView<void*> OutObjects) { return false; }
+	virtual bool GetGenericObjectKeys(int32 InStart, TArrayView<const void*> OutObjects) const { return false; }
 
-	virtual bool GetMetadataEntryKeys(int32 InStart, TArrayView<PCGMetadataEntryKey*>& OutEntryKeys) { return false; }
-	virtual bool GetMetadataEntryKeys(int32 InStart, TArrayView<const PCGMetadataEntryKey*>& OutEntryKeys) const { return false; }
+	virtual bool GetMetadataEntryKeys(int32 InStart, TArrayView<PCGMetadataEntryKey*> OutEntryKeys) { return false; }
+	virtual bool GetMetadataEntryKeys(int32 InStart, TArrayView<const PCGMetadataEntryKey*> OutEntryKeys) const { return false; }
 
 	bool bIsReadOnly = false;
 };
@@ -152,8 +152,8 @@ public:
 	virtual int32 GetNum() const override { return Entries.Num(); }
 
 protected:
-	virtual bool GetMetadataEntryKeys(int32 InStart, TArrayView<PCGMetadataEntryKey*>& OutEntryKeys) override;
-	virtual bool GetMetadataEntryKeys(int32 InStart, TArrayView<const PCGMetadataEntryKey*>& OutEntryKeys) const override;
+	virtual bool GetMetadataEntryKeys(int32 InStart, TArrayView<PCGMetadataEntryKey*> OutEntryKeys) override;
+	virtual bool GetMetadataEntryKeys(int32 InStart, TArrayView<const PCGMetadataEntryKey*> OutEntryKeys) const override;
 
 	void InitializeFromMetadata(const UPCGMetadata* Metadata);
 
@@ -178,14 +178,14 @@ public:
 	virtual int32 GetNum() const override { return Points.Num(); }
 
 protected:
-	virtual bool GetPointKeys(int32 InStart, TArrayView<FPCGPoint*>& OutPoints) override;
-	virtual bool GetPointKeys(int32 InStart, TArrayView<const FPCGPoint*>& OutPoints) const override;
+	virtual bool GetPointKeys(int32 InStart, TArrayView<FPCGPoint*> OutPoints) override;
+	virtual bool GetPointKeys(int32 InStart, TArrayView<const FPCGPoint*> OutPoints) const override;
 
-	virtual bool GetGenericObjectKeys(int32 InStart, TArrayView<void*>& OutObjects) override;
-	virtual bool GetGenericObjectKeys(int32 InStart, TArrayView<const void*>& OutObjects) const override;
+	virtual bool GetGenericObjectKeys(int32 InStart, TArrayView<void*> OutObjects) override;
+	virtual bool GetGenericObjectKeys(int32 InStart, TArrayView<const void*> OutObjects) const override;
 
-	virtual bool GetMetadataEntryKeys(int32 InStart, TArrayView<PCGMetadataEntryKey*>& OutEntryKeys) override;
-	virtual bool GetMetadataEntryKeys(int32 InStart, TArrayView<const PCGMetadataEntryKey*>& OutEntryKeys) const override;
+	virtual bool GetMetadataEntryKeys(int32 InStart, TArrayView<PCGMetadataEntryKey*> OutEntryKeys) override;
+	virtual bool GetMetadataEntryKeys(int32 InStart, TArrayView<const PCGMetadataEntryKey*> OutEntryKeys) const override;
 
 	TArrayView<FPCGPoint> Points;
 };
@@ -222,12 +222,12 @@ public:
 	virtual int32 GetNum() const override { return Objects.Num(); }
 
 protected:
-	virtual bool GetGenericObjectKeys(int32 InStart, TArrayView<void*>& OutObjects) override
+	virtual bool GetGenericObjectKeys(int32 InStart, TArrayView<void*> OutObjects) override
 	{
 		return PCGAttributeAccessorKeys::GetKeys(Objects, InStart, OutObjects, [](ObjectType& Obj) -> ObjectType* { return &Obj; });
 	}
 
-	virtual bool GetGenericObjectKeys(int32 InStart, TArrayView<const void*>& OutObjects) const override
+	virtual bool GetGenericObjectKeys(int32 InStart, TArrayView<const void*> OutObjects) const override
 	{
 		return PCGAttributeAccessorKeys::GetKeys(Objects, InStart, OutObjects, [](const ObjectType& Obj) -> const ObjectType* { return &Obj; });
 	}
@@ -267,7 +267,7 @@ public:
 	virtual int32 GetNum() const override { return 1; }
 
 protected:
-	virtual bool GetGenericObjectKeys(int32 InStart, TArrayView<void*>& OutObjects) override
+	virtual bool GetGenericObjectKeys(int32 InStart, TArrayView<void*> OutObjects) override
 	{
 		if (Ptr == nullptr)
 		{
@@ -282,7 +282,7 @@ protected:
 		return true;
 	}
 
-	virtual bool GetGenericObjectKeys(int32 InStart, TArrayView<const void*>& OutObjects) const override
+	virtual bool GetGenericObjectKeys(int32 InStart, TArrayView<const void*> OutObjects) const override
 	{
 		if (Ptr == nullptr)
 		{
@@ -324,12 +324,12 @@ public:
 	virtual int32 GetNum() const override { return Ptrs.Num(); }
 
 protected:
-	virtual bool GetGenericObjectKeys(int32 InStart, TArrayView<void*>& OutObjects) override
+	virtual bool GetGenericObjectKeys(int32 InStart, TArrayView<void*> OutObjects) override
 	{
 		return PCGAttributeAccessorKeys::GetKeys(Ptrs, InStart, OutObjects, [](void* Ptr) -> void* { return Ptr; });
 	}
 
-	virtual bool GetGenericObjectKeys(int32 InStart, TArrayView<const void*>& OutObjects) const override
+	virtual bool GetGenericObjectKeys(int32 InStart, TArrayView<const void*> OutObjects) const override
 	{
 		return PCGAttributeAccessorKeys::GetKeys(Ptrs, InStart, OutObjects, [](const void* Ptr) -> const void* { return Ptr; });
 	}
@@ -338,7 +338,7 @@ protected:
 };
 
 template <typename ObjectType>
-inline bool IPCGAttributeAccessorKeys::GetKeys(int32 InStart, TArrayView<ObjectType*>& OutKeys)
+inline bool IPCGAttributeAccessorKeys::GetKeys(int32 InStart, TArrayView<ObjectType*> OutKeys)
 {
 	TRACE_CPUPROFILER_EVENT_SCOPE(IPCGAttributeAccessorKeys::GetKeys);
 
@@ -366,7 +366,7 @@ inline bool IPCGAttributeAccessorKeys::GetKeys(int32 InStart, TArrayView<ObjectT
 }
 
 template <typename ObjectType>
-inline bool IPCGAttributeAccessorKeys::GetKeys(int32 InStart, TArrayView<const ObjectType*>& OutKeys) const
+inline bool IPCGAttributeAccessorKeys::GetKeys(int32 InStart, TArrayView<const ObjectType*> OutKeys) const
 {
 	TRACE_CPUPROFILER_EVENT_SCOPE(IPCGAttributeAccessorKeys::GetKeys);
 
diff --git a/Engine/Plugins/PCG/Source/PCG/Public/Metadata/PCGMetadataAttributeTpl.h b/Engine/Plugins/PCG/Source/PCG/Public/Metadata/PCGMetadataAttributeTpl.h
--- a/Engine/Plugins/PCG/Source/PCG/Public/Metadata/PCGMetadataAttributeTpl.h
+++ b/Engine/Plugins/PCG/Source/PCG/Public/Metadata/PCGMetadataAttributeTpl.h
@@ -813,6 +813,91 @@ public:
 		DefaultValue = Value;
 	}
 
+	void Prepare(int32 Count)
+	{
+		EntryToValueKeyMap.Reserve(EntryToValueKeyMap.Num() + Count);
+		if constexpr (!PCG::Private::MetadataTraits<T>::CompressData)
+		{
+			Values.Reserve(Values.Num() + Count);
+		}
+	}
+
+	int32 PreallocateValues(TArrayView<PCGMetadataEntryKey*> EntryKeys, bool bLockless)
+	{
+		int32 StartIndex = INDEX_NONE;
+
+		if constexpr (!PCG::Private::MetadataTraits<T>::CompressData)
+		{
+			if (!bLockless)
+			{
+				ValueLock.WriteLock();
+			}
+
+			StartIndex = Values.Num();
+
+			if constexpr (std::is_trivially_copyable_v<T>)
+			{
+				Values.SetNumUninitialized(Values.Num() + EntryKeys.Num(), /*AllowShrinking=*/EAllowShrinking::No);
+			}
+			else
+			{
+				Values.SetNum(Values.Num() + EntryKeys.Num(), /*AllowShrinking=*/EAllowShrinking::No);
+			}
+
+			if (!bLockless)
+			{
+				ValueLock.WriteUnlock();
+			}
+		}
+
+		if (!bLockless)
+		{
+			EntryMapLock.WriteLock();
+		}
+
+		EntryToValueKeyMap.Reserve(EntryToValueKeyMap.Num() + EntryKeys.Num());
+
+		if constexpr (!PCG::Private::MetadataTraits<T>::CompressData)
+		{
+			for (int32 i = 0; i < EntryKeys.Num(); ++i)
+			{
+				const PCGMetadataValueKey ValueKey = StartIndex + i + ValueKeyOffset;
+				EntryToValueKeyMap.Emplace(*EntryKeys[i], ValueKey);
+			}
+		}
+
+		if (!bLockless)
+		{
+			EntryMapLock.WriteUnlock();
+		}
+
+		return StartIndex;
+	}
+
+	void SetValues_TryLockless(TArrayView<PCGMetadataEntryKey*> EntryKeys, TArrayView<const T> InValues, int32 StartIndex)
+	{
+		if constexpr (PCG::Private::MetadataTraits<T>::CompressData)
+		{
+			SetValues(EntryKeys, InValues);
+		}
+		else
+		{
+			check(StartIndex != INDEX_NONE && Values.IsValidIndex(StartIndex + InValues.Num() - 1));
+
+			if constexpr (std::is_trivially_copyable_v<T>)
+			{
+				FMemory::Memcpy(Values.GetData() + StartIndex, InValues.GetData(), InValues.Num() * sizeof(T));
+			}
+			else
+			{
+				for (int32 i = 0; i < InValues.Num(); ++i)
+				{
+					Values[StartIndex + i] = InValues[i];
+				}
+			}
+		}
+	}
+
 protected:
 	/** Code related to computing compared values (min, max, sub, add) */
 	template<typename IT = T, typename TEnableIf<PCG::Private::MetadataTraits<IT>::CanMinMax>::Type* = nullptr>
