[PCG] Prevent APCGPartitionActor copy/paste as they are generated data

#rnx

diff --git a/Engine/Plugins/PCG/Source/PCG/Public/Grid/PCGPartitionActor.h b/Engine/Plugins/PCG/Source/PCG/Public/Grid/PCGPartitionActor.h
--- a/Engine/Plugins/PCG/Source/PCG/Public/Grid/PCGPartitionActor.h
+++ b/Engine/Plugins/PCG/Source/PCG/Public/Grid/PCGPartitionActor.h
@@ -41,6 +41,8 @@ public:
 	virtual void PostUnregisterAllComponents() override;
 #if WITH_EDITOR
 	virtual AActor* GetSceneOutlinerParent() const override;
+	virtual bool ShouldExport() override { return false; }
+	virtual bool ShouldImport(FStringView ActorPropString, bool IsMovingLevel) override { return false; }
 #endif
 	//~End AActor Interface
 
