[PCG] Fix PCGStamp actor naming in World Partition levels (need to be globally unique so use Requested name mode instead of Requested_Fatal)

#rnx

diff --git a/Engine/Plugins/PCG/Source/PCG/Private/PCGComponent.cpp b/Engine/Plugins/PCG/Source/PCG/Private/PCGComponent.cpp
--- a/Engine/Plugins/PCG/Source/PCG/Private/PCGComponent.cpp
+++ b/Engine/Plugins/PCG/Source/PCG/Private/PCGComponent.cpp
@@ -58,6 +58,7 @@
 #include "Editor.h"
 #include "EditorActorFolders.h"
 #include "ScopedTransaction.h"
+#include "Editor/EditorEngine.h"
 #include "Editor/Transactor.h"
 #endif
 
@@ -976,7 +977,9 @@ AActor* UPCGComponent::ClearPCGLink(UClass* TemplateActorClass)
 	
 	FActorSpawnParameters ActorSpawnParams;
 	UClass* SpawnClass = TemplateActorClass ? TemplateActorClass : AActor::StaticClass();
-	ActorSpawnParams.Name = MakeUniqueObjectName(GetOwner()->GetLevel(), SpawnClass, TEXT("PCGStamp"));
+	const FString DefaultName(TEXT("PCGStamp"));
+	ActorSpawnParams.Name = *DefaultName;
+	ActorSpawnParams.NameMode = FActorSpawnParameters::ESpawnActorNameMode::Requested;
 	ActorSpawnParams.OverrideLevel = GetOwner()->GetLevel();
 
 	UPCGActorHelpers::FSpawnDefaultActorParams SpawnDefaultActorParams(World, SpawnClass, GetOwner()->GetTransform(), ActorSpawnParams);
@@ -987,6 +990,9 @@ AActor* UPCGComponent::ClearPCGLink(UClass* TemplateActorClass)
 
 	// First create a new actor that will be the new owner of all the resources
 	AActor* NewActor = UPCGActorHelpers::SpawnDefaultActor(SpawnDefaultActorParams);
+#if WITH_EDITOR
+	FActorLabelUtilities::SetActorLabelUnique(NewActor, DefaultName);
+#endif
 
 	// Then move all resources linked to this component to this actor
 	bool bHasMovedResources = MoveResourcesToNewActor(NewActor, /*bCreateChild=*/false);
@@ -1149,7 +1155,9 @@ bool UPCGComponent::MoveResourcesToNewActor(AActor* InNewActor, bool bCreateChil
 	if (bCreateChild)
 	{
 		FActorSpawnParameters ActorSpawnParams;
-		ActorSpawnParams.Name = MakeUniqueObjectName(Owner->GetLevel(), NewActor->GetClass(), TEXT("PCGStampChild"));
+		const FString DefaultName(TEXT("PCGStampChild"));
+		ActorSpawnParams.Name = *DefaultName;
+		ActorSpawnParams.NameMode = FActorSpawnParameters::ESpawnActorNameMode::Requested;
 		ActorSpawnParams.OverrideLevel = Owner->GetLevel();
 
 		UPCGActorHelpers::FSpawnDefaultActorParams SpawnDefaultActorParams(GetWorld(), NewActor->GetClass(), Owner->GetTransform(), ActorSpawnParams);
@@ -1159,6 +1167,9 @@ bool UPCGComponent::MoveResourcesToNewActor(AActor* InNewActor, bool bCreateChil
 #endif
 
 		NewActor = UPCGActorHelpers::SpawnDefaultActor(SpawnDefaultActorParams);
+#if WITH_EDITOR
+		FActorLabelUtilities::SetActorLabelUnique(NewActor, DefaultName);
+#endif
 		NewActor->AttachToActor(InNewActor, FAttachmentTransformRules::KeepWorldTransform);
 		check(NewActor);
 	}
