[PCG] Fix PCGPartitionActor issues in non-WP maps
- Missing GC Call preventing PCGPartitionActor from being recreated after deletion (name clash)
- Fix crash when saving unsaved map and trying to create a PCGPartitionActor while level is registering its components (PCGPartitionActor might exist already but not have been registered to the ACM yet)

#rnx

diff --git a/Engine/Plugins/PCG/Source/PCG/Private/PCGSubsystem.cpp b/Engine/Plugins/PCG/Source/PCG/Private/PCGSubsystem.cpp
--- a/Engine/Plugins/PCG/Source/PCG/Private/PCGSubsystem.cpp
+++ b/Engine/Plugins/PCG/Source/PCG/Private/PCGSubsystem.cpp
@@ -1194,6 +1194,13 @@ APCGPartitionActor* UPCGSubsystem::FindOrCreatePCGPartitionActor(const FPCGGridD
 #endif
 	}
 
+	// If FindOrCreatePCGPartitionActor is called on a Level while it is not fully registered.
+	// We can't create the actor as it may already exist but not have been registered yet.
+	if (!World->PersistentLevel->bAreComponentsCurrentlyRegistered)
+	{
+		return nullptr;
+	}
+	
 	if (!bCanCreateActor)
 	{
 		return nullptr;
@@ -1561,6 +1568,9 @@ void UPCGSubsystem::DeleteSerializedPartitionActors(bool bOnlyDeleteUnused, bool
 		{
 			ObjectTools::CleanupAfterSuccessfulDelete(PackagesToCleanup.Array(), /*bPerformanceReferenceCheck=*/true);
 		}
+
+		// Non World Partition Levels might have deleted actors without saving anything and we need to GC so that Partition Actors can be created again (avoid name clash)
+		CollectGarbage(GARBAGE_COLLECTION_KEEPFLAGS);
 	}
 }
 
