[PCG] Fix crash when creating a Dummy transaction on top of an existing transaction, Cancelling a non index 0 transaction isn't supported

- Still need to figure out a way for those actors to be properly referenced as dirty so they don't get GCed before getting saved

#rnx

diff --git a/Engine/Plugins/PCG/Source/PCG/Private/Helpers/PCGActorHelpers.cpp b/Engine/Plugins/PCG/Source/PCG/Private/Helpers/PCGActorHelpers.cpp
--- a/Engine/Plugins/PCG/Source/PCG/Private/Helpers/PCGActorHelpers.cpp
+++ b/Engine/Plugins/PCG/Source/PCG/Private/Helpers/PCGActorHelpers.cpp
@@ -17,6 +17,7 @@
 #include UE_INLINE_GENERATED_CPP_BY_NAME(PCGActorHelpers)
 
 #if WITH_EDITOR
+#include "Editor.h"
 #include "Engine/Level.h"
 #include "ScopedTransaction.h"
 #include "Subsystems/ActorEditorContextSubsystem.h"
@@ -357,7 +358,7 @@ bool UPCGActorHelpers::DeleteActors(UWorld* World, const TArray<TSoftObjectPtr<A
 		// the next GC to collect those packages
 		//
 		// The fix here is to create a dummy transaction so that the deleted actors are tracked but since we don't want to actually push a transaction, we cancel it after the DestroyActor calls
-		FScopedTransaction DummyTransaction(NSLOCTEXT("PCGActorHelpers", "DummyTransaction", "DummyTransaction"), World && !World->IsGameWorld());
+		FScopedTransaction DummyTransaction(NSLOCTEXT("PCGActorHelpers", "DummyTransaction", "DummyTransaction"), World && !World->IsGameWorld() && ActorsToDelete.Num() > 0 && !GEditor->IsTransactionActive());
 #endif
 
 		// Not in editor, really unlikely to happen but might be slow
