[PCG] Allow data visualizations to override ALV double-click behavior
* SplineDataVisualization ALV double-click focuses control point in viewport

diff --git a/Engine/Plugins/PCG/Source/PCG/Public/PCGDataVisualization.h b/Engine/Plugins/PCG/Source/PCG/Public/PCGDataVisualization.h
--- a/Engine/Plugins/PCG/Source/PCG/Public/PCGDataVisualization.h
+++ b/Engine/Plugins/PCG/Source/PCG/Public/PCGDataVisualization.h
@@ -40,6 +40,7 @@ struct FPCGTableVisualizerInfo
 	TArray<FPCGTableVisualizerColumnInfo> ColumnInfos;
 	EPCGTableVisualizerColumnSortingMode SortingMode = EPCGTableVisualizerColumnSortingMode::Ascending;
 	FName SortingColumn = NAME_None;
+	TFunction<void(const UPCGData*, int)> DoubleClickCallback = nullptr;
 	TSharedPtr<const IPCGAttributeAccessorKeys> AccessorKeys;
 };
 
diff --git a/Engine/Plugins/PCG/Source/PCGEditor/Private/DataVisualizations/PCGSpatialDataVisualization.cpp b/Engine/Plugins/PCG/Source/PCGEditor/Private/DataVisualizations/PCGSpatialDataVisualization.cpp
--- a/Engine/Plugins/PCG/Source/PCGEditor/Private/DataVisualizations/PCGSpatialDataVisualization.cpp
+++ b/Engine/Plugins/PCG/Source/PCGEditor/Private/DataVisualizations/PCGSpatialDataVisualization.cpp
@@ -303,6 +303,20 @@ FPCGTableVisualizerInfo IPCGSpatialDataVisualization::GetTableVisualizerInfo(con
 	// Column Sorting
 	Info.SortingColumn = NAME_Index;
 
+	// Double Click Behavior
+	Info.DoubleClickCallback = [](const UPCGData* Data, int Index)
+	{
+		if (const UPCGPointData* PointData = Cast<UPCGPointData>(Data))
+		{
+			const TArray<FPCGPoint>& Points = PointData->GetPoints();
+			check(Points.IsValidIndex(Index));
+
+			const FPCGPoint& Point = Points[Index];
+			const FBox BoundingBox = Point.GetLocalBounds().TransformBy(Point.Transform.ToMatrixWithScale());
+			GEditor->MoveViewportCamerasToBox(BoundingBox, /*bActiveViewportOnly=*/true, /*DrawDebugBoxTimeInSeconds=*/2.5f);
+		}
+	};
+
 	// Accessor Keys
 	Info.AccessorKeys = TSharedPtr<const IPCGAttributeAccessorKeys>(PCGAttributeAccessorHelpers::CreateConstKeys(PointData, FPCGAttributePropertyInputSelector()).Release());
 
diff --git a/Engine/Plugins/PCG/Source/PCGEditor/Private/DataVisualizations/PCGSplineDataVisualization.cpp b/Engine/Plugins/PCG/Source/PCGEditor/Private/DataVisualizations/PCGSplineDataVisualization.cpp
--- a/Engine/Plugins/PCG/Source/PCGEditor/Private/DataVisualizations/PCGSplineDataVisualization.cpp
+++ b/Engine/Plugins/PCG/Source/PCGEditor/Private/DataVisualizations/PCGSplineDataVisualization.cpp
@@ -6,6 +6,35 @@
 #include "Data/PCGSplineData.h"
 #include "Data/PCGSplineStruct.h"
 
+namespace PCGSplineDataVisualizationConstants
+{
+	const FVector HalfExtents = FVector(50.0); // Set the points to fill up 1 meter of space by default.
+}
+
+FPCGTableVisualizerInfo IPCGSplineDataVisualization::GetTableVisualizerInfo(const UPCGData* Data) const
+{
+	FPCGTableVisualizerInfo Info = IPCGSpatialDataVisualization::GetTableVisualizerInfo(Data);
+
+	Info.DoubleClickCallback = [](const UPCGData* Data, int Index)
+	{
+		if (const UPCGSplineData* SplineData = Cast<UPCGSplineData>(Data))
+		{
+			const TArray<FInterpCurvePointVector>& Positions = SplineData->SplineStruct.SplineCurves.Position.Points;
+			const TArray<FInterpCurvePointVector>& Scales = SplineData->SplineStruct.SplineCurves.Scale.Points;
+
+			check(Positions.IsValidIndex(Index) && Scales.IsValidIndex(Index));
+
+			const FVector& Position = SplineData->GetTransform().TransformPosition(Positions[Index].OutVal);
+			const FVector HalfExtent = Scales[Index].OutVal * PCGSplineDataVisualizationConstants::HalfExtents;
+
+			FBox BoundingBox(Position + HalfExtent, Position - HalfExtent);
+			GEditor->MoveViewportCamerasToBox(BoundingBox, /*bActiveViewportOnly=*/true, /*DrawDebugBoxTimeInSeconds=*/2.5f);
+		}
+	};
+
+	return Info;
+}
+
 const UPCGPointData* IPCGSplineDataVisualization::CollapseToDebugPointData(FPCGContext* Context, const UPCGData* Data) const
 {
 	if (const UPCGSplineData* SplineData = Cast<UPCGSplineData>(Data))
@@ -20,7 +49,7 @@ const UPCGPointData* IPCGSplineDataVisualization::CollapseToDebugPointData(FPCGC
 		{
 			FPCGPoint& Point = Points.Emplace_GetRef();
 			Point.Transform = SplineData->SplineStruct.GetTransformAtSplineInputKey(ControlPointIndex, ESplineCoordinateSpace::World);
-			Point.SetExtents(FVector::One() * 50.0f); // Set the points to fill up 1 meter of space by default.
+			Point.SetExtents(PCGSplineDataVisualizationConstants::HalfExtents);
 		}
 
 		return PointData;
diff --git a/Engine/Plugins/PCG/Source/PCGEditor/Private/DataVisualizations/PCGSplineDataVisualization.h b/Engine/Plugins/PCG/Source/PCGEditor/Private/DataVisualizations/PCGSplineDataVisualization.h
--- a/Engine/Plugins/PCG/Source/PCGEditor/Private/DataVisualizations/PCGSplineDataVisualization.h
+++ b/Engine/Plugins/PCG/Source/PCGEditor/Private/DataVisualizations/PCGSplineDataVisualization.h
@@ -12,6 +12,12 @@ struct FPCGContext;
 class PCGEDITOR_API IPCGSplineDataVisualization : public IPCGSpatialDataVisualization
 {
 public:
+	// ~Begin IPCGDataVisualization interface
+	virtual FPCGTableVisualizerInfo GetTableVisualizerInfo(const UPCGData* Data) const override;
+	// ~End IPCGDataVisualization interface
+
+	// ~Begin IPCGSpatialDataVisualization interface
 	/** Overrides collapse behavior to show the spline control points. */
 	virtual const UPCGPointData* CollapseToDebugPointData(FPCGContext* Context, const UPCGData* Data) const override;
+	// ~End IPCGSpatialDataVisualization interface
 };
diff --git a/Engine/Plugins/PCG/Source/PCGEditor/Private/SPCGEditorGraphAttributeListView.cpp b/Engine/Plugins/PCG/Source/PCGEditor/Private/SPCGEditorGraphAttributeListView.cpp
--- a/Engine/Plugins/PCG/Source/PCGEditor/Private/SPCGEditorGraphAttributeListView.cpp
+++ b/Engine/Plugins/PCG/Source/PCGEditor/Private/SPCGEditorGraphAttributeListView.cpp
@@ -682,17 +682,7 @@ void SPCGEditorGraphAttributeListView::RefreshAttributeList()
 		{
 			PCGListviewItemPtr ListViewItem = MakeShared<FPCGListViewItem>();
 			ListViewItem->Index = Index;
-
-			// TODO: It would be preferable to have a custom override for double-click behavior instead of relying on collapsed point data.
-			ListViewItem->DoubleClickCallback = [DataPtr = DataStrongPtr.Get()](int Index)
-			{
-				if (const UPCGPointData* PointData = Cast<UPCGPointData>(DataPtr))
-				{
-					const FPCGPoint& Point = PointData->GetPoints()[Index];
-					const FBox BoundingBox = Point.GetLocalBounds().TransformBy(Point.Transform.ToMatrixWithScale());
-					GEditor->MoveViewportCamerasToBox(BoundingBox, true, 2.5f);
-				}
-			};
+			ListViewItem->DoubleClickCallback = TableVisualizerInfo.DoubleClickCallback;
 
 			ListViewItems.Add(ListViewItem);
 		}
@@ -980,7 +970,7 @@ void SPCGEditorGraphAttributeListView::OnItemDoubleClicked(PCGListviewItemPtr It
 
 	if (Item->DoubleClickCallback)
 	{
-		Item->DoubleClickCallback(Item->Index);
+		Item->DoubleClickCallback(DataStrongPtr.Get(), Item->Index);
 	}
 }
 
diff --git a/Engine/Plugins/PCG/Source/PCGEditor/Private/SPCGEditorGraphAttributeListView.h b/Engine/Plugins/PCG/Source/PCGEditor/Private/SPCGEditorGraphAttributeListView.h
--- a/Engine/Plugins/PCG/Source/PCGEditor/Private/SPCGEditorGraphAttributeListView.h
+++ b/Engine/Plugins/PCG/Source/PCGEditor/Private/SPCGEditorGraphAttributeListView.h
@@ -45,7 +45,7 @@ namespace PCGEditorGraphAttributeListView
 struct FPCGListViewItem
 {
 	int32 Index = INDEX_NONE;
-	TFunction<void(int)> DoubleClickCallback = nullptr;
+	TFunction<void(const UPCGData*, int)> DoubleClickCallback = nullptr;
 };
 
 struct FPCGColumnData
