[PCG] PCGCreateSplineMesh: can only run on main thread

- override CanExecuteOnlyOnMainThread to return true

#rnx

diff --git a/Engine/Plugins/PCG/Source/PCG/Public/Elements/PCGCreateSplineMesh.h b/Engine/Plugins/PCG/Source/PCG/Public/Elements/PCGCreateSplineMesh.h
--- a/Engine/Plugins/PCG/Source/PCG/Public/Elements/PCGCreateSplineMesh.h
+++ b/Engine/Plugins/PCG/Source/PCG/Public/Elements/PCGCreateSplineMesh.h
@@ -59,8 +59,11 @@ struct FPCGCreateSplineMeshContext : public FPCGContext, public IPCGAsyncLoading
 
 class FPCGCreateSplineMeshElement : public IPCGElement
 {
-protected:
+public:
+	virtual bool CanExecuteOnlyOnMainThread(FPCGContext* Context) const override { return true; }
 	virtual bool IsCacheable(const UPCGSettings* InSettings) const override { return false; }
+
+protected:
 	virtual FPCGContext* CreateContext() override;
 	virtual bool PrepareDataInternal(FPCGContext* InContext) const override;
 	virtual bool ExecuteInternal(FPCGContext* Context) const override;
