[PCG] Fix typo in localization key

#rnx

diff --git a/Engine/Plugins/PCG/Source/PCG/Private/Compute/Elements/PCGCustomHLSL.cpp b/Engine/Plugins/PCG/Source/PCG/Private/Compute/Elements/PCGCustomHLSL.cpp
--- a/Engine/Plugins/PCG/Source/PCG/Private/Compute/Elements/PCGCustomHLSL.cpp
+++ b/Engine/Plugins/PCG/Source/PCG/Private/Compute/Elements/PCGCustomHLSL.cpp
@@ -833,7 +833,7 @@ bool UPCGCustomHLSLSettings::IsKernelValid(FPCGContext* InContext) const
 		{
 			if (InContext)
 			{
-				PCGE_LOG_C(Error, GraphAndLog, InContext, LOCTEXT("InvalidPPInput", "Point processing kernel requires a first output pin of type point."));
+				PCGE_LOG_C(Error, GraphAndLog, InContext, LOCTEXT("InvalidPPOutput", "Point processing kernel requires a first output pin of type point."));
 			}
 
 			return false;
