[PCG] SplineSampler: Fix Divide by Zero (when NumIterations < MinIterationPerDispatch)

#rnx

diff --git a/Engine/Plugins/PCG/Source/PCG/Private/Elements/PCGSplineSampler.cpp b/Engine/Plugins/PCG/Source/PCG/Private/Elements/PCGSplineSampler.cpp
--- a/Engine/Plugins/PCG/Source/PCG/Private/Elements/PCGSplineSampler.cpp
+++ b/Engine/Plugins/PCG/Source/PCG/Private/Elements/PCGSplineSampler.cpp
@@ -1212,6 +1212,8 @@ namespace PCGSplineSamplerHelpers
 		{
 			NumDispatch = FMath::Min(Context->AsyncState.NumAvailableTasks, NumDispatch);
 		}
+		NumDispatch = FMath::Max(1, NumDispatch);
+
 		const int32 NumIterationsPerDispatch = NumIterations / NumDispatch;
 
 		TArray<TArray<TTuple<FTransform, FVector, float>>> InteriorSplinePointData;
