[PCG] Separate out GPU compiler into a separate class

#rnx

diff --git a/Engine/Plugins/PCG/Source/PCG/Private/Compute/PCGComputeGraph.h b/Engine/Plugins/PCG/Source/PCG/Private/Compute/PCGComputeGraph.h
--- a/Engine/Plugins/PCG/Source/PCG/Private/Compute/PCGComputeGraph.h
+++ b/Engine/Plugins/PCG/Source/PCG/Private/Compute/PCGComputeGraph.h
@@ -50,5 +50,5 @@ protected:
 	UPROPERTY()
 	TMap<FPCGKernelAttributeKey, int32 /* Attribute Index */> GlobalAttributeLookupTable;
 
-	friend class FPCGGraphCompiler;
+	friend class FPCGGraphCompilerGPU;
 };
diff --git a/Engine/Plugins/PCG/Source/PCG/Private/Graph/PCGGraphCompiler.cpp b/Engine/Plugins/PCG/Source/PCG/Private/Graph/PCGGraphCompiler.cpp
--- a/Engine/Plugins/PCG/Source/PCG/Private/Graph/PCGGraphCompiler.cpp
+++ b/Engine/Plugins/PCG/Source/PCG/Private/Graph/PCGGraphCompiler.cpp
@@ -9,47 +9,30 @@
 #include "PCGModule.h"
 #include "PCGPin.h"
 #include "PCGSubgraph.h"
-#include "Compute/PCGComputeCommon.h"
-#include "Compute/PCGComputeKernelSource.h"
-#include "Compute/DataInterfaces/PCGComputeDataInterface.h"
-#include "Compute/DataInterfaces/PCGCustomKernelDataInterface.h"
-#include "Compute/DataInterfaces/PCGDataCollectionDataInterface.h"
-#include "Compute/DataInterfaces/PCGDataCollectionReadbackDataInterface.h"
-#include "Compute/DataInterfaces/PCGDataCollectionUploadDataInterface.h"
-#include "Compute/DataInterfaces/PCGDebugDataInterface.h"
-#include "Compute/DataInterfaces/PCGLandscapeDataInterface.h"
-#include "Compute/DataInterfaces/PCGTextureDataInterface.h"
-#include "Compute/Elements/PCGComputeGraphElement.h"
-#include "Compute/Elements/PCGCustomHLSL.h"
 #include "Elements/PCGGather.h"
 #include "Elements/PCGHiGenGridSize.h"
 #include "Elements/PCGReroute.h"
+#include "Graph/PCGGraphCompilerGPU.h"
 #include "Graph/PCGGraphExecutor.h"
 #include "Graph/PCGPinDependencyExpression.h"
 
-#include "ComputeFramework/ComputeKernel.h"
 #include "HAL/IConsoleManager.h"
 #include "Misc/ScopeRWLock.h"
-#include "Shader/ShaderTypes.h"
+
+FPCGElementPtr FPCGGraphCompiler::SharedGatherElement;
+FRWLock FPCGGraphCompiler::SharedGatherElementLock;
 
 namespace PCGGraphCompiler
 {
-	TAutoConsoleVariable<bool> CVarEnableTaskStaticCulling(
+	static TAutoConsoleVariable<bool> CVarEnableTaskStaticCulling(
 		TEXT("pcg.GraphExecution.TaskStaticCulling"),
 		true,
 		TEXT("Enable static culling of tasks which considers static branches, generation grid size, trivial nodes and more."));
 
-	TAutoConsoleVariable<bool> CVarEnableGPUExecution(
+	static TAutoConsoleVariable<bool> CVarEnableGPUExecution(
 		TEXT("pcg.GraphExecution.GPU.Enable"),
 		true,
 		TEXT("Whether to emit compatible nodes as compute graphs to execute on the GPU."));
-
-#if WITH_EDITOR
-	TAutoConsoleVariable<bool> CVarEnableGPUDebugging(
-		TEXT("pcg.GraphExecution.GPU.EnableDebugging"),
-		false,
-		TEXT("Enable verbose logging of GPU compilation and execution."));
-#endif
 }
 
 TArray<FPCGGraphTask> FPCGGraphCompiler::CompileGraph(UPCGGraph* InGraph, FPCGTaskId& NextId, FPCGStackContext& InOutStackContext)
@@ -1079,7 +1062,7 @@ void FPCGGraphCompiler::CompileTopGraph(UPCGGraph* InGraph, uint32 GenerationGri
 #if WITH_EDITOR
 	if (PCGGraphCompiler::CVarEnableGPUExecution.GetValueOnAnyThread())
 	{
-		CreateGPUNodes(InGraph, CompiledTasks);
+		FPCGGraphCompilerGPU::CreateGPUNodes(InGraph, CompiledTasks);
 	}
 	else
 #endif
@@ -1316,986 +1299,3 @@ void FPCGGraphCompiler::RemoveFromCacheRecursive(UPCGGraph* InGraph)
 	}
 }
 #endif // WITH_EDITOR
-
-void FPCGGraphCompiler::CollectGPUNodeSubsets(const TArray<FPCGGraphTask>& InCompiledTasks, const TMap<FPCGTaskId, TArray<FPCGTaskId>>& InTaskSuccessors, const TSet<FPCGTaskId>& InGPUCompatibleTaskIds, TArray<TSet<FPCGTaskId>>& OutNodeSubsetsToConvertToCFGraph)
-{
-	// Populate initial sets of tasks that are ready to consume vs ones that currently blocked.
-	TSet<FPCGTaskId> ReadyTaskIds;
-	TSet<FPCGTaskId> RemainingTaskIds;
-	ReadyTaskIds.Reserve(InCompiledTasks.Num());
-	RemainingTaskIds.Reserve(InCompiledTasks.Num());
-
-	for (FPCGTaskId TaskId = 0; TaskId < InCompiledTasks.Num(); ++TaskId)
-	{
-		if (InCompiledTasks[TaskId].Inputs.IsEmpty())
-		{
-			ReadyTaskIds.Add(TaskId);
-		}
-		else
-		{
-			RemainingTaskIds.Add(TaskId);
-		}
-	}
-
-	// Queue all successors of InTaskId that are ready to go (all upstream input tasks have been processed).
-	auto QueueSuccessors = [&InTaskSuccessors, &ReadyTaskIds, &RemainingTaskIds, &InCompiledTasks](int InTaskId)
-	{
-		bool bQueuedTask = false;
-
-		// Queue up any successors that are ready to go.
-		if (const TArray<FPCGTaskId>* Successors = InTaskSuccessors.Find(InTaskId))
-		{
-			for (FPCGTaskId Successor : *Successors)
-			{
-				const bool bSuccessorQueued = ReadyTaskIds.Contains(Successor);
-
-				// All successors should either be already queued, or waiting to be queued.
-				check(bSuccessorQueued || RemainingTaskIds.Contains(Successor));
-
-				if (!bSuccessorQueued)
-				{
-					bool bSuccessorReady = true;
-					for (const FPCGGraphTaskInput& Input : InCompiledTasks[Successor].Inputs)
-					{
-						if (ReadyTaskIds.Contains(Input.TaskId) || RemainingTaskIds.Contains(Input.TaskId))
-						{
-							bSuccessorReady = false;
-						}
-					}
-
-					if (bSuccessorReady)
-					{
-						ReadyTaskIds.Add(Successor);
-						RemainingTaskIds.Remove(Successor);
-						bQueuedTask = true;
-					}
-				}
-			}
-		}
-
-		return bQueuedTask;
-	};
-	
-	// Local to loops below but pulled out for performance.
-	TArray<FPCGTaskId> FoundReadyTaskIds;
-	TSet<FPCGTaskId> GPUSubsetTaskIds;
-	FoundReadyTaskIds.Reserve(InCompiledTasks.Num());
-	GPUSubsetTaskIds.Reserve(InCompiledTasks.Num());
-
-	// Build subsets of nodes that are GPU compatible and can be dispatched together.
-	while (!ReadyTaskIds.IsEmpty() || !RemainingTaskIds.IsEmpty())
-	{
-		// Consume as many CPU nodes as we can.
-		bool bQueuedTasks = true;
-		while (bQueuedTasks)
-		{
-			FoundReadyTaskIds.Reset();
-
-			for (FPCGTaskId ReadyTaskId : ReadyTaskIds)
-			{
-				if (!InGPUCompatibleTaskIds.Contains(ReadyTaskId))
-				{
-					FoundReadyTaskIds.Add(ReadyTaskId);
-				}
-			}
-
-			bQueuedTasks = false;
-
-			for (FPCGTaskId ReadyCPUTaskId : FoundReadyTaskIds)
-			{
-				ReadyTaskIds.Remove(ReadyCPUTaskId);
-				bQueuedTasks |= QueueSuccessors(ReadyCPUTaskId);
-			}
-		}
-
-		GPUSubsetTaskIds.Reset();
-
-		int StackIndex = INDEX_NONE;
-
-		// Now the opposite - consume as many GPU nodes as we can and accumulate them into a set that will be compiled into a compute graph.
-		bQueuedTasks = !ReadyTaskIds.IsEmpty();
-		while (bQueuedTasks)
-		{
-			FoundReadyTaskIds.Reset();
-
-			for (FPCGTaskId ReadyTaskId : ReadyTaskIds)
-			{
-				if (InGPUCompatibleTaskIds.Contains(ReadyTaskId))
-				{
-					// For now don't mix tasks from different execution stacks (in and out of subgraphs for instance) into one compute graph.
-					if ((StackIndex == INDEX_NONE) || InCompiledTasks[ReadyTaskId].StackIndex == StackIndex)
-					{
-						StackIndex = InCompiledTasks[ReadyTaskId].StackIndex;
-						FoundReadyTaskIds.Add(ReadyTaskId);
-					}
-				}
-			}
-
-			bQueuedTasks = false;
-
-			for (FPCGTaskId ReadyGPUTaskId : FoundReadyTaskIds)
-			{
-				GPUSubsetTaskIds.Add(ReadyGPUTaskId);
-				ReadyTaskIds.Remove(ReadyGPUTaskId);
-				bQueuedTasks |= QueueSuccessors(ReadyGPUTaskId);
-			}
-		}
-
-		if (!GPUSubsetTaskIds.IsEmpty())
-		{
-			OutNodeSubsetsToConvertToCFGraph.Add(MoveTemp(GPUSubsetTaskIds));
-		}
-	}
-}
-
-void FPCGGraphCompiler::CreateGatherTasksAtGPUInputs(const TSet<FPCGTaskId>& InGPUCompatibleTaskIds, TArray<FPCGGraphTask>& InOutCompiledTasks)
-{
-	using FOriginalInputPinKey = TPair<FPCGTaskId /* Original GPU task */, FName /* Input pin label */>;
-
-	// These are local to loop below but hoisted here for efficiency.
-	TSet<FOriginalInputPinKey> EncounteredInputPins; // TODO this is heavyweight, could use simple array.
-	TMap<FOriginalInputPinKey, FPCGTaskId /* Gather task */> InputPinToGatherTask;
-
-	// Add all compute graph task inputs and outputs.
-	for (FPCGTaskId GPUTaskId : InGPUCompatibleTaskIds)
-	{
-		EncounteredInputPins.Reset();
-		InputPinToGatherTask.Reset();
-
-		// First pass - create gather tasks for all original input pins which have more than one incident edge.
-		// This is so we can gather on the CPU (much more efficient than going it on the GPU).
-		for (int InputIndex = 0; InputIndex < InOutCompiledTasks[GPUTaskId].Inputs.Num(); ++InputIndex)
-		{
-			// Helper to get current input. We avoid simply taking a local reference as InOutCompiledTasks can be modified below.
-			auto CurrentInput = [&InOutCompiledTasks, InputIndex, GPUTaskId]() -> FPCGGraphTaskInput&
-			{
-				return InOutCompiledTasks[GPUTaskId].Inputs[InputIndex];
-			};
-
-			if (!CurrentInput().DownstreamPin.IsSet())
-			{
-				continue;
-			}
-
-			const FOriginalInputPinKey PinKey = { GPUTaskId, CurrentInput().DownstreamPin->Label };
-
-			// If already created a gather task, then nothing more to do for this pin.
-			if (const FPCGTaskId* GatherTaskId = InputPinToGatherTask.Find(PinKey))
-			{
-				continue;
-			}
-
-			// If we're encountering pin for first time, register it.
-			if (!EncounteredInputPins.Contains(PinKey))
-			{
-				EncounteredInputPins.Add(PinKey);
-				continue;
-			}
-
-			// Second time we've encountered this input pin - create a gather element because we need one edge connected to
-			// each virtual input pin, so that we can obtain the data items from the input data collection using the unique
-			// virtual pin label at execution time.
-			const FPCGTaskId GatherTaskId = InOutCompiledTasks.Num();
-			FPCGGraphTask& GatherTask = InOutCompiledTasks.Emplace_GetRef();
-			GatherTask.NodeId = GatherTaskId;
-			GatherTask.ParentId = InOutCompiledTasks[GPUTaskId].ParentId;
-			GatherTask.Element = GetSharedGatherElement();
-
-			InputPinToGatherTask.Add(PinKey, GatherTaskId);
-		}
-
-		EncounteredInputPins.Reset();
-
-		// Second pass - wire up the newly added gather tasks once we have the full picture of which edges are affected.
-		for (int InputIndex = 0; InputIndex < InOutCompiledTasks[GPUTaskId].Inputs.Num(); ++InputIndex)
-		{
-			auto CurrentInput = [&InOutCompiledTasks, InputIndex, GPUTaskId]() -> FPCGGraphTaskInput&
-			{
-				return InOutCompiledTasks[GPUTaskId].Inputs[InputIndex];
-			};
-
-			if (!CurrentInput().DownstreamPin.IsSet())
-			{
-				continue;
-			}
-
-			const FOriginalInputPinKey PinKey = { GPUTaskId, CurrentInput().DownstreamPin->Label };
-
-			if (const FPCGTaskId* GatherTaskId = InputPinToGatherTask.Find(PinKey))
-			{
-				// Wire the upstream output pin to the gather task.
-				FPCGGraphTaskInput& WireUpstreamNodeToGather = InOutCompiledTasks[*GatherTaskId].Inputs.Add_GetRef(CurrentInput());
-				if (WireUpstreamNodeToGather.DownstreamPin.IsSet())
-				{
-					WireUpstreamNodeToGather.DownstreamPin->Label = PCGPinConstants::DefaultInputLabel;
-				}
-
-				if (!EncounteredInputPins.Contains(PinKey))
-				{
-					// First time we're encountering this input pin, wire it to the gather task.
-					EncounteredInputPins.Add(PinKey);
-
-					CurrentInput().TaskId = *GatherTaskId;
-					if (CurrentInput().UpstreamPin.IsSet())
-					{
-						CurrentInput().UpstreamPin->Label = PCGPinConstants::DefaultOutputLabel;
-					}
-				}
-				else
-				{
-					// Input pin already encountered, already wired to gather task. Remove this input.
-					InOutCompiledTasks[GPUTaskId].Inputs.RemoveAt(InputIndex);
-					--InputIndex;
-				}
-			}
-		}
-	}
-}
-
-void FPCGGraphCompiler::WireGPUGraphNode(
-	FPCGTaskId InGPUGraphTaskId,
-	const TSet<FPCGTaskId>& InCollapsedTasks,
-	const TSet<FPCGTaskId>& InGPUCompatibleTaskIds,
-	TArray<FPCGGraphTask>& InOutCompiledTasks,
-	const TMap<FPCGTaskId, TArray<FPCGTaskId>>& InTaskSuccessors,
-	FOriginalToVirtualPin& OutOriginalToVirtualPin,
-	TMap<const UPCGPin*, FName>& OutOutputCPUPinToVirtualPin)
-{
-	FPCGGraphTask& GPUGraphTask = InOutCompiledTasks[InGPUGraphTaskId];
-
-	// Used to construct unique input/output labels, ultimately consumed in graph executor in BuildTaskInput and PostExecute for input/output respectively.
-	int InputCount = 0;
-	int OutputCount = 0;
-
-	// Add all compute graph task inputs and outputs.
-	for (FPCGTaskId GPUTaskId : InCollapsedTasks)
-	{
-		// First find CPU to GPU edges and wire in the GPU graph node inputs.
-		for (const FPCGGraphTaskInput& Input : InOutCompiledTasks[GPUTaskId].Inputs)
-		{
-			if (InGPUCompatibleTaskIds.Contains(Input.TaskId))
-			{
-				continue;
-			}
-
-			FPCGGraphTaskInput& AddedInput = GPUGraphTask.Inputs.Add_GetRef(Input);
-
-			// TODO is pinless fine with skipping?
-			if (AddedInput.DownstreamPin.IsSet())
-			{
-				const FName VirtualLabel = *FString::Format(TEXT("{0}-VirtualIn{1}"), { AddedInput.DownstreamPin->Label.ToString(), InputCount });
-				const bool bIsInputPin = true;
-				OutOriginalToVirtualPin.Add({ GPUTaskId, AddedInput.DownstreamPin->Label, bIsInputPin }, VirtualLabel);
-				AddedInput.DownstreamPin->Label = VirtualLabel;
-
-				++InputCount;
-
-				if (const UPCGNode* UpstreamNode = InOutCompiledTasks[Input.TaskId].Node)
-				{
-					if (const UPCGPin* OutputPin = UpstreamNode->GetOutputPin(AddedInput.UpstreamPin->Label))
-					{
-						OutOutputCPUPinToVirtualPin.Add(OutputPin, VirtualLabel);
-					}
-				}
-			}
-		}
-
-		if (!InTaskSuccessors.Contains(GPUTaskId))
-		{
-			continue;
-		}
-
-		// Next consider GPU to CPU edges to wire in the GPU graph node outputs.
-		for (FPCGTaskId Successor : InTaskSuccessors[GPUTaskId])
-		{
-			if (InGPUCompatibleTaskIds.Contains(Successor))
-			{
-				continue;
-			}
-
-			// Rewire inputs of this downstream CPU node to the outputs of the compute graph task.
-			FPCGGraphTask& DownstreamCPUNode = InOutCompiledTasks[Successor];
-
-			// Order matters here! We can never reorder inputs as it will impact execution.
-			const int InputCountBefore = DownstreamCPUNode.Inputs.Num();
-			for (int SuccessorInputIndex = 0; SuccessorInputIndex < InputCountBefore; ++SuccessorInputIndex)
-			{
-				// Implementation note: we modify the Inputs array in this loop, so don't take a reference to the current element.
-				
-				// Skip irrelevant edges.
-				if (DownstreamCPUNode.Inputs[SuccessorInputIndex].TaskId != GPUTaskId)
-				{
-					continue;
-				}
-
-				// Wire downstream CPU node to compute graph task.
-				FPCGGraphTaskInput InputCopy = DownstreamCPUNode.Inputs[SuccessorInputIndex];
-
-				InputCopy.TaskId = InGPUGraphTaskId;
-
-				if (DownstreamCPUNode.Inputs[SuccessorInputIndex].UpstreamPin.IsSet())
-				{
-					const FNodePin PinKey = { GPUTaskId, InputCopy.UpstreamPin->Label, /*Pin is input*/false };
-					if (const FName* FoundVirtualPinLabel = OutOriginalToVirtualPin.Find(PinKey))
-					{
-						// Wire to the existing virtual output pin.
-						InputCopy.UpstreamPin->Label = *FoundVirtualPinLabel;
-					}
-					else
-					{
-						const FName VirtualLabel = *FString::Format(TEXT("{0}-VirtualOut{1}"), { InputCopy.UpstreamPin->Label.ToString(), OutputCount });
-						OutOriginalToVirtualPin.Add(PinKey, VirtualLabel);
-
-						InputCopy.UpstreamPin->Label = VirtualLabel;
-
-						++OutputCount;
-					}
-				}
-
-				DownstreamCPUNode.Inputs.Add(MoveTemp(InputCopy));
-			}
-		}
-	}
-}
-
-void FPCGGraphCompiler::BuildGPUGraphTask(
-	UPCGGraph* InGraph,
-	FPCGTaskId InGPUGraphTaskId,
-	const TSet<FPCGTaskId>& InCollapsedTasks,
-	const TMap<FPCGTaskId, TArray<FPCGTaskId>>& InTaskSuccessors,
-	TArray<FPCGGraphTask>& InOutCompiledTasks,
-	const FOriginalToVirtualPin& InOriginalToVirtualPin,
-	const TMap<const UPCGPin*, FName>& InOutputCPUPinToVirtualPin)
-{
-	TSharedPtr<FPCGComputeGraphElement> Element = MakeShared<FPCGComputeGraphElement>();
-	InOutCompiledTasks[InGPUGraphTaskId].Element = Element;
-
-	const FName GraphName = MakeUniqueObjectName(InGraph, UPCGComputeGraph::StaticClass(), InGraph->GetFName());
-
-	UPCGComputeGraph* ComputeGraph = NewObject<UPCGComputeGraph>(InGraph, GraphName);
-	ComputeGraph->OutputCPUPinToInputGPUPinAlias = InOutputCPUPinToVirtualPin;
-	Element->Graph.Reset(ComputeGraph);
-
-	// Not incredibly useful for us - DG adds GetComponentSource()->GetComponentClass() object which allows it to bind at execution time by class.
-	// But execution code requires it currently.
-	Element->Graph->Bindings.Add(UPCGDataBinding::StaticClass());
-
-	// Create data interfaces which allow kernels to read or write data. Each data interface is associated with a node output pin.
-	// For CPU->GPU edges, an upload data interface is created. For GPU->CPU edges, a readback data interface is created.
-	auto CreateDataInterface = [&InCollapsedTasks, &InOutCompiledTasks, ComputeGraph](FPCGTaskId InTaskId, bool bRequiresReadback, const FPCGPinProperties& InOutputPinProperties) -> UPCGComputeDataInterface*
-	{
-		const bool bUpstreamIsGPUTask = InCollapsedTasks.Contains(InTaskId);
-
-		EPCGDataType PinType = InOutputPinProperties.AllowedTypes;
-
-		// Dynamically typed pins could have a different type
-		const UPCGSettings* Settings = InOutCompiledTasks[InTaskId].Node ? InOutCompiledTasks[InTaskId].Node->GetSettings() : nullptr;
-		const UPCGPin* Pin = InOutCompiledTasks[InTaskId].Node ? InOutCompiledTasks[InTaskId].Node->GetOutputPin(InOutputPinProperties.Label) : nullptr;
-		if (Settings && Pin)
-		{
-			PinType = Settings->GetCurrentPinTypes(Pin);
-		}
-
-		UPCGComputeDataInterface* DataInterface = nullptr;
-
-		switch (PinType)
-		{
-		case EPCGDataType::Point:
-		case EPCGDataType::PointOrParam:
-		{
-			UPCGDataCollectionDataInterface* DataInterfacePCGData = nullptr;
-
-			if (bUpstreamIsGPUTask)
-			{
-				if (bRequiresReadback)
-				{
-					// GPU -> CPU
-					DataInterfacePCGData = NewObject<UPCGDataCollectionReadbackDataInterface>(ComputeGraph);
-				}
-				else
-				{
-					// GPU -> GPU
-					DataInterfacePCGData = NewObject<UPCGDataCollectionDataInterface>(ComputeGraph);
-				}
-			}
-			else
-			{
-				// CPU -> GPU
-				DataInterfacePCGData = NewObject<UPCGDataCollectionUploadDataInterface>(ComputeGraph);
-			}
-
-			check(DataInterfacePCGData);
-			DataInterfacePCGData->ProducerSettings = InOutCompiledTasks[InTaskId].Node ? InOutCompiledTasks[InTaskId].Node->GetSettings() : nullptr;
-
-			DataInterface = DataInterfacePCGData;
-			break;
-		}
-		case EPCGDataType::Texture:
-		{
-			DataInterface = NewObject<UPCGTextureDataInterface>(ComputeGraph);
-			break;
-		}
-		case EPCGDataType::Landscape:
-		{
-			DataInterface = NewObject<UPCGLandscapeDataInterface>(ComputeGraph);
-			break;
-		}
-		default:
-			ensure(false);
-			break;
-		}
-
-		if (DataInterface)
-		{
-			DataInterface->SetOutputPin(InOutputPinProperties.Label);
-		}
-
-		return DataInterface;
-	};
-
-	TMap<TPair</* Node task*/FPCGTaskId, /*Node output pin*/FName>, UPCGComputeDataInterface*> OutputPinDataInterfaces;
-
-	// Create all the output data interfaces.
-	for (FPCGTaskId TaskId : InCollapsedTasks)
-	{
-		// Create DIs for all output pins, because the kernels currently need their outputs to be bound to valid resources.
-		if (const UPCGSettings* Settings = InOutCompiledTasks[TaskId].Node ? InOutCompiledTasks[TaskId].Node->GetSettings() : nullptr)
-		{
-			for (const FPCGPinProperties& OutputPinProperties : Settings->AllOutputPinProperties())
-			{
-				if (OutputPinDataInterfaces.Contains({ TaskId, OutputPinProperties.Label }))
-				{
-					ensure(false);
-					continue;
-				}
-
-				bool bRequiresReadback = false;
-				if (const TArray<FPCGTaskId>* Successors = InTaskSuccessors.Find(TaskId))
-				{
-					for (FPCGTaskId Successor : *Successors)
-					{
-						for (const FPCGGraphTaskInput& Input : InOutCompiledTasks[Successor].Inputs)
-						{
-							if (Input.UpstreamPin.IsSet() && (*Input.UpstreamPin == OutputPinProperties) && !InCollapsedTasks.Contains(Successor))
-							{
-								bRequiresReadback = true;
-								break;
-							}
-						}
-					}
-				}
-
-				if (UPCGComputeDataInterface* OutputDI = CreateDataInterface(TaskId, bRequiresReadback, OutputPinProperties))
-				{
-					OutputDI->SetOutputPin(OutputPinProperties.Label);
-					ComputeGraph->DataInterfaces.Add(OutputDI);
-					OutputPinDataInterfaces.Add({ TaskId, OutputPinProperties.Label }, OutputDI);
-
-					for (const FPCGKernelAttributeKey& Key : Settings->GetKernelAttributeKeys())
-					{
-						if (!ComputeGraph->GlobalAttributeLookupTable.Find(Key))
-						{
-							ComputeGraph->GlobalAttributeLookupTable.Add(Key, ComputeGraph->GlobalAttributeLookupTable.Num() + PCGComputeConstants::NUM_RESERVED_ATTRS);
-						}
-					}
-				}
-			}
-		}
-
-		// Create any DIs for upstream CPU nodes.
-		for (const FPCGGraphTaskInput& Input : InOutCompiledTasks[TaskId].Inputs)
-		{
-			// Only deal with upstream CPU tasks.
-			if (InCollapsedTasks.Contains(Input.TaskId))
-			{
-				continue;
-			}
-
-			if (!Input.DownstreamPin.IsSet())
-			{
-				continue;
-			}
-
-			if (const UPCGNode* DownstreamNode = InOutCompiledTasks[TaskId].Node)
-			{
-				if (const UPCGPin* InputPin = DownstreamNode->GetInputPin(Input.DownstreamPin->Label))
-				{
-					ComputeGraph->PinsReceivingDataFromCPU.Add(InputPin);
-				}
-			}
-
-			if (!Input.UpstreamPin.IsSet())
-			{
-				continue;
-			}
-
-			if (OutputPinDataInterfaces.Contains({ Input.TaskId, Input.UpstreamPin->Label }))
-			{
-				// Skip if already created.
-				continue;
-			}
-
-			if (UPCGComputeDataInterface* OutputDI = CreateDataInterface(Input.TaskId, /*bRequiresReadback=*/false, *Input.UpstreamPin))
-			{
-				OutputDI->SetOutputPin(Input.UpstreamPin->Label);
-				ComputeGraph->DataInterfaces.Add(OutputDI);
-				OutputPinDataInterfaces.Add({ Input.TaskId, Input.UpstreamPin->Label }, OutputDI);
-			}
-		}
-	}
-
-	TSet<FPCGTaskId> RemainingTasks = InCollapsedTasks;
-
-	while (!RemainingTasks.IsEmpty())
-	{
-		// Find a ready task
-		FPCGTaskId TaskId = InvalidPCGTaskId;
-		for (FPCGTaskId RemainingTask : RemainingTasks)
-		{
-			// TODO: use 'QueueSuccessors' pattern rather than brute force searching for ready tasks
-			bool bReady = true;
-			for (const FPCGGraphTaskInput& Input : InOutCompiledTasks[RemainingTask].Inputs)
-			{
-				if (RemainingTasks.Contains(Input.TaskId))
-				{
-					bReady = false;
-					break;
-				}
-			}
-
-			if (bReady)
-			{
-				TaskId = RemainingTask;
-				break;
-			}
-		}
-
-		if (TaskId == InvalidPCGTaskId)
-		{
-			// Tasks in RemainingTasks are unreachable?
-			ensure(false);
-			break;
-		}
-
-		RemainingTasks.Remove(TaskId);
-
-		const UPCGNode* Node = InOutCompiledTasks[TaskId].Node;
-
-		const UPCGCustomHLSLSettings* Settings = Cast<UPCGCustomHLSLSettings>(Node ? Node->GetSettings() : nullptr);
-		check(Settings && Settings->bEnabled && Settings->ShouldExecuteOnGPU());
-
-		// For every usage of a DI, get the original (non-aliased) pin label.
-		TMap<TPair<FPCGTaskId, UComputeDataInterface*>, FName> DataInterfaceUsageToPinLabel;
-
-		TArray<int> InputDataInterfaceIndices;
-		TArray<int> OutputDataInterfaceIndices;
-		InputDataInterfaceIndices.Reserve(Settings->InputPinProperties().Num());
-		OutputDataInterfaceIndices.Reserve(Settings->OutputPinProperties().Num());
-
-		// Add DIs (PCG -> CF transcoding).
-
-		for (const FPCGGraphTaskInput& Input : InOutCompiledTasks[TaskId].Inputs)
-		{
-			if (!Input.UpstreamPin.IsSet())
-			{
-				// Execution-only dependencies not supported currently. Unclear if this should ever be supported for GPU graphs.
-				// Writes followed by reads will be protected via barriers added by RDG.
-				continue;
-			}
-
-			UPCGComputeDataInterface* UpstreamDI = nullptr;
-			if (UPCGComputeDataInterface** FoundUpstreamDI = OutputPinDataInterfaces.Find({ Input.TaskId, Input.UpstreamPin->Label }))
-			{
-				UpstreamDI = *FoundUpstreamDI;
-			}
-
-			if (!UpstreamDI)
-			{
-				ensure(false);
-				continue;
-			}
-
-			const int Index = ComputeGraph->DataInterfaces.Find(UpstreamDI);
-			if (Index == INDEX_NONE)
-			{
-				ensure(false);
-				continue;
-			}
-
-			InputDataInterfaceIndices.Add(Index);
-
-			FName DownstreamInputPinLabel = Input.DownstreamPin->Label;
-
-			DataInterfaceUsageToPinLabel.Add({ TaskId, UpstreamDI }, DownstreamInputPinLabel);
-
-			const bool bIsInputPin = true;
-			UpstreamDI->AddDownstreamInputPin(DownstreamInputPinLabel, InOriginalToVirtualPin.Find({ TaskId, DownstreamInputPinLabel, bIsInputPin }));
-		}
-
-		// Always create a DI for every output pin, so kernel always has something to write to.
-		for (const FPCGPinProperties& OutputPinProperties : Settings->AllOutputPinProperties())
-		{
-			UPCGComputeDataInterface** FoundDI = OutputPinDataInterfaces.Find({ TaskId, OutputPinProperties.Label });
-			if (!ensure(FoundDI) || !ensure(*FoundDI))
-			{
-				continue;
-			}
-
-			const int Index = ComputeGraph->DataInterfaces.Find(*FoundDI);
-			if (Index == INDEX_NONE)
-			{
-				ensure(false);
-				continue;
-			}
-
-			OutputDataInterfaceIndices.Add(Index);
-			DataInterfaceUsageToPinLabel.Add({ TaskId, *FoundDI }, OutputPinProperties.Label);
-		}
-
-		// Make sure every downstream input pin is registered with the upstream DI.
-		if (const TArray<FPCGTaskId>* Successors = InTaskSuccessors.Find(TaskId))
-		{
-			for (FPCGTaskId Successor : *Successors)
-			{
-				for (const FPCGGraphTaskInput& SuccessorInput : InOutCompiledTasks[Successor].Inputs)
-				{
-					if (SuccessorInput.TaskId != TaskId)
-					{
-						continue;
-					}
-
-					if (!SuccessorInput.UpstreamPin.IsSet() || !SuccessorInput.DownstreamPin.IsSet())
-					{
-						continue;
-					}
-
-					FName OutputPinLabel = SuccessorInput.UpstreamPin->Label;
-
-					// DIs for output pins should have all been created.
-					UPCGComputeDataInterface** FoundOutputDI = OutputPinDataInterfaces.Find({ TaskId, OutputPinLabel });
-					if (!ensure(FoundOutputDI) || !ensure(*FoundOutputDI))
-					{
-						continue;
-					}
-					
-					// If the map has an entry for this task and output pin label, then its a GPU->CPU readback.
-					const bool bIsInputPin = false;
-					(*FoundOutputDI)->SetOutputPin(OutputPinLabel, InOriginalToVirtualPin.Find({ TaskId, OutputPinLabel, bIsInputPin }));
-				}
-			}
-		}
-
-		UPCGCustomKernelDataInterface* KernelDI = NewObject<UPCGCustomKernelDataInterface>(ComputeGraph);
-		KernelDI->Settings = Settings;
-		const int KernelDIIndex = ComputeGraph->DataInterfaces.Num();
-		ComputeGraph->DataInterfaces.Add(KernelDI);
-
-		// TODO add graph data interface (graph params). Reference: UOptimusGraphDataInterface.
-		//Element->Graph->DataInterfaces.Add(NewObject< UPCGDataCollectionDataInterface>(InGraph));
-		//Element->Graph->DataInterfaceToBinding.Add(0);
-
-		// TODO once we support cooking for different platforms/configs, don't create the interface if logging is not present.
-		int DebugDIIndex = INDEX_NONE;
-		if (Settings->bPrintShaderDebugValues)
-		{
-			UPCGDebugDataInterface* DebugDI = NewObject<UPCGDebugDataInterface>(ComputeGraph);
-			DebugDI->SetDebugBufferSize(Settings->DebugBufferSize);
-
-			DebugDIIndex = ComputeGraph->DataInterfaces.Num();
-			ComputeGraph->DataInterfaces.Add(DebugDI);
-		}
-
-		// Now that all data interfaces added, create the (trivial) binding mapping. All map to primary binding, index 0.
-		ComputeGraph->DataInterfaceToBinding.SetNumZeroed(ComputeGraph->DataInterfaces.Num());
-
-		struct FInterfaceBinding
-		{
-			const UComputeDataInterface* DataInterface;
-			int32 DataInterfaceBindingIndex;
-			FString BindingFunctionName;
-			FString BindingFunctionNamespace;
-		};
-
-		struct FKernelWithDataBindings
-		{
-			UComputeKernel* Kernel;
-			TArray<FInterfaceBinding> InputDataBindings;
-			TArray<FInterfaceBinding> OutputDataBindings;
-		};
-
-		FKernelWithDataBindings KernelWithBindings;
-
-		KernelWithBindings.Kernel = NewObject<UComputeKernel>(ComputeGraph);
-		const int KernelIndex = ComputeGraph->KernelInvocations.Num();
-		ComputeGraph->KernelInvocations.Add(KernelWithBindings.Kernel);
-		ComputeGraph->KernelToNode.Add(Node);
-		
-		auto SetupAllInputBindings = [&KernelWithBindings, ComputeGraph](int InDataInterfaceIndex)
-		{
-			const UComputeDataInterface* DataInterface = ComputeGraph->DataInterfaces[InDataInterfaceIndex];
-			TArray<FShaderFunctionDefinition> Functions;
-			DataInterface->GetSupportedInputs(Functions);
-
-			for (int FuncIndex = 0; FuncIndex < Functions.Num(); ++FuncIndex)
-			{
-				FInterfaceBinding& Binding = KernelWithBindings.InputDataBindings.Emplace_GetRef();
-				Binding.DataInterface = DataInterface;
-				Binding.BindingFunctionName = Functions[FuncIndex].Name;
-				Binding.BindingFunctionNamespace = TEXT("");
-				Binding.DataInterfaceBindingIndex = FuncIndex;
-			}
-		};
-
-		auto SetupAllOutputBindings = [&KernelWithBindings, ComputeGraph](int InDataInterfaceIndex)
-		{
-			const UComputeDataInterface* DataInterface = ComputeGraph->DataInterfaces[InDataInterfaceIndex];
-			TArray<FShaderFunctionDefinition> Functions;
-			DataInterface->GetSupportedOutputs(Functions);
-
-			for (int FuncIndex = 0; FuncIndex < Functions.Num(); ++FuncIndex)
-			{
-				FInterfaceBinding& Binding = KernelWithBindings.OutputDataBindings.Emplace_GetRef();
-				Binding.DataInterface = DataInterface;
-				Binding.BindingFunctionName = Functions[FuncIndex].Name;
-				Binding.BindingFunctionNamespace = TEXT("");
-				Binding.DataInterfaceBindingIndex = FuncIndex;
-			}
-		};
-
-		// Bind data interfaces.
-		for (int InputDataInterfaceIndex : InputDataInterfaceIndices)
-		{
-			SetupAllInputBindings(InputDataInterfaceIndex);
-		}
-
-		SetupAllInputBindings(KernelDIIndex);
-
-		for (int OutputDataInterfaceIndex : OutputDataInterfaceIndices)
-		{
-			SetupAllOutputBindings(OutputDataInterfaceIndex);
-		}
-
-		if (DebugDIIndex != INDEX_NONE)
-		{
-			SetupAllOutputBindings(DebugDIIndex);
-		}
-
-		{
-			UPCGComputeKernelSource* KernelSource = NewObject<UPCGComputeKernelSource>(KernelWithBindings.Kernel); // is outer to kernel fine?
-			KernelWithBindings.Kernel->KernelSource = KernelSource;
-			KernelSource->EntryPoint = Settings->GetKernelEntryPoint();
-			KernelSource->GroupSize = Settings->GetThreadGroupSize();
-
-			KernelSource->SetSource(Settings->GetCookedKernelSource(ComputeGraph->GlobalAttributeLookupTable));
-
-#if WITH_EDITOR
-			if (PCGGraphCompiler::CVarEnableGPUDebugging.GetValueOnAnyThread())
-			{
-				UE_LOG(LogPCG, Warning, TEXT("ATTRIBUTE LOOK-UP TABLE [%s]"), *Settings->GetDefaultNodeTitle().ToString());
-
-				for (const TPair<FPCGKernelAttributeKey, int32 /* Attribute Index */>& Pair : ComputeGraph->GlobalAttributeLookupTable)
-				{
-					const FString TypeString = UEnum::GetValueAsString(Pair.Key.Type);
-					const FString NameString = Pair.Key.Name.ToString();
-					const FString IndexString = FString::FromInt(Pair.Value);
-
-					UE_LOG(LogPCG, Warning, TEXT("%s: %s (%s)"), *IndexString, *NameString, *TypeString);
-				}
-			}
-#endif
-
-			// Add functions for external inputs/outputs which must be fulfilled by DIs
-			for (FInterfaceBinding& Binding : KernelWithBindings.InputDataBindings)
-			{
-				TArray<FShaderFunctionDefinition> Functions;
-				Binding.DataInterface->GetSupportedInputs(Functions);
-				check(Functions.IsValidIndex(Binding.DataInterfaceBindingIndex));
-
-				FShaderFunctionDefinition FuncDef = Functions[Binding.DataInterfaceBindingIndex];
-				for (FShaderParamTypeDefinition& ParamType : FuncDef.ParamTypes)
-				{
-					// Making sure parameter has type declaration generated
-					ParamType.ResetTypeDeclaration();
-				}
-
-				KernelSource->ExternalInputs.Emplace(FuncDef);
-			}
-
-			for (FInterfaceBinding& Binding : KernelWithBindings.OutputDataBindings)
-			{
-				TArray<FShaderFunctionDefinition> Functions;
-				Binding.DataInterface->GetSupportedOutputs(Functions);
-				check(Functions.IsValidIndex(Binding.DataInterfaceBindingIndex));
-
-				FShaderFunctionDefinition FuncDef = Functions[Binding.DataInterfaceBindingIndex];
-				for (FShaderParamTypeDefinition& ParamType : FuncDef.ParamTypes)
-				{
-					// Making sure parameter has type declaration generated
-					ParamType.ResetTypeDeclaration();
-				}
-
-				KernelSource->ExternalOutputs.Emplace(FuncDef);
-			}
-		}
-
-		auto AddAllEdgesForKernel = [&KernelWithBindings, ComputeGraph, TaskId, &DataInterfaceUsageToPinLabel](int32 InKernelIndex, bool bInEdgesAreInputs)
-		{
-			TArray<FInterfaceBinding>& Bindings = bInEdgesAreInputs ? KernelWithBindings.InputDataBindings : KernelWithBindings.OutputDataBindings;
-
-			// Add all graph edges for bindings. This is somewhat odd but likely what we'll do vs the more finegrained interface
-			// definitions in optimus, but we need to see.
-			for (int BindingIndex = 0; BindingIndex < Bindings.Num(); ++BindingIndex)
-			{
-				FInterfaceBinding& Binding = Bindings[BindingIndex];
-
-				FComputeGraphEdge& Edge = ComputeGraph->GraphEdges.Emplace_GetRef();
-				Edge.KernelIndex = InKernelIndex;
-				Edge.KernelBindingIndex = BindingIndex;
-				Edge.DataInterfaceIndex = ComputeGraph->DataInterfaces.IndexOfByPredicate([&Binding](const UComputeDataInterface* In) { return Binding.DataInterface == In; });
-				check(Edge.DataInterfaceIndex != INDEX_NONE);
-				Edge.DataInterfaceBindingIndex = Binding.DataInterfaceBindingIndex;
-				Edge.bKernelInput = bInEdgesAreInputs;
-
-				UComputeDataInterface* DataInterface = ComputeGraph->DataInterfaces[Edge.DataInterfaceIndex];
-				check(DataInterface);
-
-				if (FName* PinLabel = DataInterfaceUsageToPinLabel.Find({ TaskId, DataInterface }))
-				{
-					TArray<FShaderFunctionDefinition> DataInterfaceFunctions;
-					if (bInEdgesAreInputs)
-					{
-						DataInterface->GetSupportedInputs(DataInterfaceFunctions);
-					}
-					else
-					{
-						DataInterface->GetSupportedOutputs(DataInterfaceFunctions);
-					}
-
-					Edge.BindingFunctionNameOverride = FString::Format(
-						TEXT("{0}_{1}"),
-						{ PinLabel->ToString(), DataInterfaceFunctions[Edge.DataInterfaceBindingIndex].Name }
-					);
-				}
-			}
-		};
-
-		AddAllEdgesForKernel(KernelIndex, /*bInEdgesAreInputs=*/true);
-		AddAllEdgesForKernel(KernelIndex, /*bInEdgesAreInputs=*/false);
-	}
-
-	// Register all virtual pin aliases with the corresponding pins for downstream usage.
-	for (const TPair<TTuple<FPCGTaskId, FName, bool>, FName>& Mapping : InOriginalToVirtualPin)
-	{
-		const bool bIsInputPin = Mapping.Get<0>().Get<2>();
-		if (bIsInputPin)
-		{
-			const FPCGTaskId TaskId = Mapping.Get<0>().Get<0>();
-
-			if (const UPCGNode* Node = InOutCompiledTasks[TaskId].Node)
-			{
-				const FName& OriginalLabel = Mapping.Get<0>().Get<1>();
-
-				if (const UPCGPin* Pin = Node->GetInputPin(OriginalLabel))
-				{
-					const FName& VirtualLabel = Mapping.Get<1>();
-					ComputeGraph->InputPinLabelAliases.FindOrAdd(Pin) = VirtualLabel;
-				}
-			}
-		}
-	}
-
-	// Kick off shader compilation (if needed).
-	ComputeGraph->UpdateResources();
-}
-
-void FPCGGraphCompiler::CreateGPUNodes(UPCGGraph* InGraph, TArray<FPCGGraphTask>& InOutCompiledTasks)
-{
-	TSet<FPCGTaskId> GPUCompatibleTaskIds;
-	GPUCompatibleTaskIds.Reserve(InOutCompiledTasks.Num());
-	for (FPCGTaskId TaskId = 0; TaskId < InOutCompiledTasks.Num(); ++TaskId)
-	{
-		const UPCGNode* Node = InOutCompiledTasks[TaskId].Node;
-		const UPCGSettings* Settings = Node ? Node->GetSettings() : nullptr;
-		if (Settings && Settings->ShouldExecuteOnGPU() && Settings->bEnabled)
-		{
-			GPUCompatibleTaskIds.Add(TaskId);
-		}
-	}
-
-	if (GPUCompatibleTaskIds.IsEmpty())
-	{
-		// Nothing to do for this graph.
-		return;
-	}
-
-	TMap<FPCGTaskId, TArray<FPCGTaskId>> TaskSuccessors;
-	TaskSuccessors.Reserve(InOutCompiledTasks.Num());
-	for (FPCGTaskId TaskId = 0; TaskId < InOutCompiledTasks.Num(); ++TaskId)
-	{
-		for (int InputIndex = 0; InputIndex < InOutCompiledTasks[TaskId].Inputs.Num(); ++InputIndex)
-		{
-			TaskSuccessors.FindOrAdd(InOutCompiledTasks[TaskId].Inputs[InputIndex].TaskId).Add(TaskId);
-		}
-	}
-
-	// For input pins at CPU -> GPU boundary, inject gather elements to pre-combine data on CPU side
-	// before passing to GPU.
-	CreateGatherTasksAtGPUInputs(GPUCompatibleTaskIds, InOutCompiledTasks);
-
-	TArray<TSet<FPCGTaskId>> NodeSubsetsToConvertToCFGraph;
-	CollectGPUNodeSubsets(InOutCompiledTasks, TaskSuccessors, GPUCompatibleTaskIds, NodeSubsetsToConvertToCFGraph);
-
-	// Do actual collapsing now, one subset at a time. Each collapse will do all fixup of task ids? That will invalidate
-	// ids in NodeSubsetsToConvertToCFGraph, so may need remap table. But can ignore this for now.
-	for (TSet<FPCGTaskId>& NodeSubsetToConvertToCFGraph : NodeSubsetsToConvertToCFGraph)
-	{
-		if (NodeSubsetToConvertToCFGraph.IsEmpty())
-		{
-			ensure(false);
-			continue;
-		}
-
-		// Add a new compute graph task. Then kill the original GPU tasks.
-		const FPCGTaskId ComputeGraphTaskId = InOutCompiledTasks.Num();
-		FPCGGraphTask& ComputeGraphTask = InOutCompiledTasks.Emplace_GetRef();
-		ComputeGraphTask.NodeId = ComputeGraphTaskId;
-
-		// All nodes in subset will be from same stack/parent, so assign from any.
-		for (FPCGTaskId GPUTaskId : NodeSubsetToConvertToCFGraph)
-		{
-			ComputeGraphTask.ParentId = InOutCompiledTasks[GPUTaskId].ParentId;
-			ComputeGraphTask.StackIndex = InOutCompiledTasks[GPUTaskId].StackIndex;
-			break;
-		}
-
-		// Mapping from task ID & pin label to a virtual pin label. Compute graphs are executed within a generated element,
-		// and the input and output pins of this element must have unique virtual pin labels so that we can parse the data that
-		// PCG provides through the input data collection correctly, and route the output data to the downstream pins correctly.
-		FOriginalToVirtualPin OriginalToVirtualPin;
-
-		TMap<const UPCGPin*, FName> OutputCPUPinToVirtualPin;
-
-		// Wire in the compute graph task, side by side with the individual GPU tasks, which will be culled below.
-		WireGPUGraphNode(
-			ComputeGraphTaskId,
-			NodeSubsetToConvertToCFGraph,
-			GPUCompatibleTaskIds,
-			InOutCompiledTasks,
-			TaskSuccessors,
-			OriginalToVirtualPin,
-			OutputCPUPinToVirtualPin);
-
-		// Generate a compute graph from all of the individual GPU tasks.
-		BuildGPUGraphTask(
-			InGraph,
-			ComputeGraphTaskId,
-			NodeSubsetToConvertToCFGraph,
-			TaskSuccessors,
-			InOutCompiledTasks,
-			OriginalToVirtualPin,
-			OutputCPUPinToVirtualPin);
-	}
-
-	// Now cull all the GPU compatible nodes. The compute graph task are already wired in so we're fine to just delete.
-	CullTasks(InOutCompiledTasks, /*bAddPassthroughWires=*/false, [&GPUCompatibleTaskIds](const FPCGGraphTask& InTask) { return GPUCompatibleTaskIds.Contains(InTask.NodeId); });
-}
diff --git a/Engine/Plugins/PCG/Source/PCG/Private/Graph/PCGGraphCompiler.h b/Engine/Plugins/PCG/Source/PCG/Private/Graph/PCGGraphCompiler.h
--- a/Engine/Plugins/PCG/Source/PCG/Private/Graph/PCGGraphCompiler.h
+++ b/Engine/Plugins/PCG/Source/PCG/Private/Graph/PCGGraphCompiler.h
@@ -40,46 +40,12 @@ public:
 
 	/** Returns the post graph element, which is used to determine results caching behavior */
 	FPCGElementPtr GetSharedTrivialPostGraphElement();
-	
-	/** Used to track new unique virtual pins created on generated compute graph elements. */
-	using FNodePin = TTuple<FPCGTaskId, /*Pin label*/FName, /*Pin is input*/bool>;
-	using FOriginalToVirtualPin = TMap<FNodePin, /*Virtual pin label*/FName>;
 
-	/** Outputs sets of task IDs, where each set is GPU nodes that can be compiled into a compute graph and dispatched together. */
-	void CollectGPUNodeSubsets(
-		const TArray<FPCGGraphTask>& InCompiledTasks,
-		const TMap<FPCGTaskId, TArray<FPCGTaskId>>& InTaskSuccessors,
-		const TSet<FPCGTaskId>& InGPUCompatibleTaskIds,
-		TArray<TSet<FPCGTaskId>>& OutNodeSubsetsToConvertToCFGraph);
-
-	/** For GPU node inputs that have multiple incident edges, bundle them into a single edge. This is to avoid an inefficient
-	* gather operation on the GPU, and allows data interfaces to pick their data from the compute graph element input data collection
-	* using unique virtual input pin labels. */
-	void CreateGatherTasksAtGPUInputs(const TSet<FPCGTaskId>& InGPUCompatibleTaskIds, TArray<FPCGGraphTask>& InOutCompiledTasks);
-
-	/** Wires in a compute graph element alongside each set of GPU compatible nodes. The tasks for each node will be culled later. */
-	void WireGPUGraphNode(
-		FPCGTaskId InGPUGraphTaskId,
-		const TSet<FPCGTaskId>& InCollapsedTasks,
-		const TSet<FPCGTaskId>& InGPUCompatibleTaskIds,
-		TArray<FPCGGraphTask>& InOutCompiledTasks,
-		const TMap<FPCGTaskId,
-		TArray<FPCGTaskId>>& InTaskSuccessors,
-		FOriginalToVirtualPin& OutOriginalToVirtualPin,
-		TMap<const UPCGPin*, FName>& OutOutputCPUPinToVirtualPin);
-
-	/** Compiles a compute graph. */
-	void BuildGPUGraphTask(
-		UPCGGraph* InGraph,
-		FPCGTaskId InGPUGraphTaskId,
-		const TSet<FPCGTaskId>& InCollapsedTasks,
-		const TMap<FPCGTaskId, TArray<FPCGTaskId>>& InTaskSuccessors,
-		TArray<FPCGGraphTask>& InOutCompiledTasks,
-		const FOriginalToVirtualPin& InOriginalToVirtualPin,
-		const TMap<const UPCGPin*, FName>& InOutputCPUPinToVirtualPin);
+	/** Returns the gather element object shared by all tasks that need it. */
+	static FPCGElementPtr GetSharedGatherElement();
 
-	/** Finds connected subgraphs of GPU - enabled nodes that can be dispatched together and replaces each one with a compute graph. */
-	void CreateGPUNodes(UPCGGraph* InGraph, TArray<FPCGGraphTask>& InOutCompiledTasks);
+	/** Culls tasks based on a given lambda. Never culls the first (input) task in the array. */
+	static void CullTasks(TArray<FPCGGraphTask>& InOutCompiledTasks, bool bAddPassthroughWires, TFunctionRef<bool(const FPCGGraphTask&)> CullTask);
 
 private:
 	TArray<FPCGGraphTask> CompileGraph(UPCGGraph* InGraph, FPCGTaskId& NextId, FPCGStackContext& InOutStackContext);
@@ -90,9 +56,6 @@ private:
 	/** Returns the trivial element object shared by all tasks that need it. */
 	FPCGElementPtr GetSharedTrivialElement();
 
-	/** Returns the gather element object shared by all tasks that need it. */
-	FPCGElementPtr GetSharedGatherElement();
-
 	/** Propagates grid sizes through a graph's compiled tasks. */
 	static void ResolveGridSizes(
 		EPCGHiGenGrid GenerationGrid,
@@ -122,9 +85,6 @@ private:
 	/** Culls nodes that are inactive for e.g. missing required inputs or on a statically inactive branch. */
 	static void CullTasksStaticInactive(TArray<FPCGGraphTask>& InOutCompiledTasks);
 
-	/** Culls tasks based on a given lambda. Never culls the first (input) task in the array. */
-	static void CullTasks(TArray<FPCGGraphTask>& InOutCompiledTasks, bool bAddPassthroughWires, TFunctionRef<bool(const FPCGGraphTask&)> CullTask);
-
 	/** Remove any stack frames that are not used by any task. */
 	static void PostCullStackCleanup(TArray<FPCGGraphTask>& InCompiledTasks, FPCGStackContext& InOutStackContext);
 
@@ -140,10 +100,12 @@ private:
 	TMap<UPCGGraph*, TMap<uint32, FPCGStackContext>> TopGraphToStackContextMap;
 
 	FPCGElementPtr SharedTrivialElement;
-	FPCGElementPtr SharedGatherElement;
-	FPCGElementPtr SharedTrivialPostGraphElement;
 	mutable FRWLock SharedTrivialElementLock;
-	mutable FRWLock SharedGatherElementLock;
+
+	FPCGElementPtr SharedTrivialPostGraphElement;
+
+	static FPCGElementPtr SharedGatherElement;
+	static FRWLock SharedGatherElementLock;
 
 #if WITH_EDITOR
 	void RemoveFromCacheRecursive(UPCGGraph* InGraph);
diff --git a/Engine/Plugins/PCG/Source/PCG/Private/Graph/PCGGraphCompilerGPU.cpp b/Engine/Plugins/PCG/Source/PCG/Private/Graph/PCGGraphCompilerGPU.cpp
new file mode 100644
--- /dev/null
+++ b/Engine/Plugins/PCG/Source/PCG/Private/Graph/PCGGraphCompilerGPU.cpp
@@ -0,0 +1,1021 @@
+// Copyright Epic Games, Inc. All Rights Reserved.
+
+#include "Graph/PCGGraphCompilerGPU.h"
+
+#include "PCGGraph.h"
+#include "PCGModule.h"
+#include "PCGPin.h"
+#include "Compute/PCGComputeCommon.h"
+#include "Compute/PCGComputeKernelSource.h"
+#include "Compute/DataInterfaces/PCGComputeDataInterface.h"
+#include "Compute/DataInterfaces/PCGCustomKernelDataInterface.h"
+#include "Compute/DataInterfaces/PCGDataCollectionDataInterface.h"
+#include "Compute/DataInterfaces/PCGDataCollectionReadbackDataInterface.h"
+#include "Compute/DataInterfaces/PCGDataCollectionUploadDataInterface.h"
+#include "Compute/DataInterfaces/PCGDebugDataInterface.h"
+#include "Compute/DataInterfaces/PCGLandscapeDataInterface.h"
+#include "Compute/DataInterfaces/PCGTextureDataInterface.h"
+#include "Compute/Elements/PCGComputeGraphElement.h"
+#include "Compute/Elements/PCGCustomHLSL.h"
+#include "Graph/PCGGraphCompiler.h"
+#include "Graph/PCGGraphExecutor.h"
+
+#include "ComputeFramework/ComputeKernel.h"
+#include "HAL/IConsoleManager.h"
+#include "Shader/ShaderTypes.h"
+
+namespace PCGGraphCompilerGPU
+{
+#if WITH_EDITOR
+	static TAutoConsoleVariable<bool> CVarEnableGPUDebugging(
+		TEXT("pcg.GraphExecution.GPU.EnableDebugging"),
+		false,
+		TEXT("Enable verbose logging of GPU compilation and execution."));
+#endif
+}
+
+void FPCGGraphCompilerGPU::CollectGPUNodeSubsets(const TArray<FPCGGraphTask>& InCompiledTasks, const TMap<FPCGTaskId, TArray<FPCGTaskId>>& InTaskSuccessors, const TSet<FPCGTaskId>& InGPUCompatibleTaskIds, TArray<TSet<FPCGTaskId>>& OutNodeSubsetsToConvertToCFGraph)
+{
+	// Populate initial sets of tasks that are ready to consume vs ones that currently blocked.
+	TSet<FPCGTaskId> ReadyTaskIds;
+	TSet<FPCGTaskId> RemainingTaskIds;
+	ReadyTaskIds.Reserve(InCompiledTasks.Num());
+	RemainingTaskIds.Reserve(InCompiledTasks.Num());
+
+	for (FPCGTaskId TaskId = 0; TaskId < InCompiledTasks.Num(); ++TaskId)
+	{
+		if (InCompiledTasks[TaskId].Inputs.IsEmpty())
+		{
+			ReadyTaskIds.Add(TaskId);
+		}
+		else
+		{
+			RemainingTaskIds.Add(TaskId);
+		}
+	}
+
+	// Queue all successors of InTaskId that are ready to go (all upstream input tasks have been processed).
+	auto QueueSuccessors = [&InTaskSuccessors, &ReadyTaskIds, &RemainingTaskIds, &InCompiledTasks](int InTaskId)
+	{
+		bool bQueuedTask = false;
+
+		// Queue up any successors that are ready to go.
+		if (const TArray<FPCGTaskId>* Successors = InTaskSuccessors.Find(InTaskId))
+		{
+			for (FPCGTaskId Successor : *Successors)
+			{
+				const bool bSuccessorQueued = ReadyTaskIds.Contains(Successor);
+
+				// All successors should either be already queued, or waiting to be queued.
+				check(bSuccessorQueued || RemainingTaskIds.Contains(Successor));
+
+				if (!bSuccessorQueued)
+				{
+					bool bSuccessorReady = true;
+					for (const FPCGGraphTaskInput& Input : InCompiledTasks[Successor].Inputs)
+					{
+						if (ReadyTaskIds.Contains(Input.TaskId) || RemainingTaskIds.Contains(Input.TaskId))
+						{
+							bSuccessorReady = false;
+						}
+					}
+
+					if (bSuccessorReady)
+					{
+						ReadyTaskIds.Add(Successor);
+						RemainingTaskIds.Remove(Successor);
+						bQueuedTask = true;
+					}
+				}
+			}
+		}
+
+		return bQueuedTask;
+	};
+
+	// Local to loops below but pulled out for performance.
+	TArray<FPCGTaskId> FoundReadyTaskIds;
+	TSet<FPCGTaskId> GPUSubsetTaskIds;
+	FoundReadyTaskIds.Reserve(InCompiledTasks.Num());
+	GPUSubsetTaskIds.Reserve(InCompiledTasks.Num());
+
+	// Build subsets of nodes that are GPU compatible and can be dispatched together.
+	while (!ReadyTaskIds.IsEmpty() || !RemainingTaskIds.IsEmpty())
+	{
+		// Consume as many CPU nodes as we can.
+		bool bQueuedTasks = true;
+		while (bQueuedTasks)
+		{
+			FoundReadyTaskIds.Reset();
+
+			for (FPCGTaskId ReadyTaskId : ReadyTaskIds)
+			{
+				if (!InGPUCompatibleTaskIds.Contains(ReadyTaskId))
+				{
+					FoundReadyTaskIds.Add(ReadyTaskId);
+				}
+			}
+
+			bQueuedTasks = false;
+
+			for (FPCGTaskId ReadyCPUTaskId : FoundReadyTaskIds)
+			{
+				ReadyTaskIds.Remove(ReadyCPUTaskId);
+				bQueuedTasks |= QueueSuccessors(ReadyCPUTaskId);
+			}
+		}
+
+		GPUSubsetTaskIds.Reset();
+
+		int StackIndex = INDEX_NONE;
+
+		// Now the opposite - consume as many GPU nodes as we can and accumulate them into a set that will be compiled into a compute graph.
+		bQueuedTasks = !ReadyTaskIds.IsEmpty();
+		while (bQueuedTasks)
+		{
+			FoundReadyTaskIds.Reset();
+
+			for (FPCGTaskId ReadyTaskId : ReadyTaskIds)
+			{
+				if (InGPUCompatibleTaskIds.Contains(ReadyTaskId))
+				{
+					// For now don't mix tasks from different execution stacks (in and out of subgraphs for instance) into one compute graph.
+					if ((StackIndex == INDEX_NONE) || InCompiledTasks[ReadyTaskId].StackIndex == StackIndex)
+					{
+						StackIndex = InCompiledTasks[ReadyTaskId].StackIndex;
+						FoundReadyTaskIds.Add(ReadyTaskId);
+					}
+				}
+			}
+
+			bQueuedTasks = false;
+
+			for (FPCGTaskId ReadyGPUTaskId : FoundReadyTaskIds)
+			{
+				GPUSubsetTaskIds.Add(ReadyGPUTaskId);
+				ReadyTaskIds.Remove(ReadyGPUTaskId);
+				bQueuedTasks |= QueueSuccessors(ReadyGPUTaskId);
+			}
+		}
+
+		if (!GPUSubsetTaskIds.IsEmpty())
+		{
+			OutNodeSubsetsToConvertToCFGraph.Add(MoveTemp(GPUSubsetTaskIds));
+		}
+	}
+}
+
+void FPCGGraphCompilerGPU::CreateGatherTasksAtGPUInputs(const TSet<FPCGTaskId>& InGPUCompatibleTaskIds, TArray<FPCGGraphTask>& InOutCompiledTasks)
+{
+	using FOriginalInputPinKey = TPair<FPCGTaskId /* Original GPU task */, FName /* Input pin label */>;
+
+	// These are local to loop below but hoisted here for efficiency.
+	TSet<FOriginalInputPinKey> EncounteredInputPins; // TODO this is heavyweight, could use simple array.
+	TMap<FOriginalInputPinKey, FPCGTaskId /* Gather task */> InputPinToGatherTask;
+
+	// Add all compute graph task inputs and outputs.
+	for (FPCGTaskId GPUTaskId : InGPUCompatibleTaskIds)
+	{
+		EncounteredInputPins.Reset();
+		InputPinToGatherTask.Reset();
+
+		// First pass - create gather tasks for all original input pins which have more than one incident edge.
+		// This is so we can gather on the CPU (much more efficient than going it on the GPU).
+		for (int InputIndex = 0; InputIndex < InOutCompiledTasks[GPUTaskId].Inputs.Num(); ++InputIndex)
+		{
+			// Helper to get current input. We avoid simply taking a local reference as InOutCompiledTasks can be modified below.
+			auto CurrentInput = [&InOutCompiledTasks, InputIndex, GPUTaskId]() -> FPCGGraphTaskInput&
+			{
+				return InOutCompiledTasks[GPUTaskId].Inputs[InputIndex];
+			};
+
+			if (!CurrentInput().DownstreamPin.IsSet())
+			{
+				continue;
+			}
+
+			const FOriginalInputPinKey PinKey = { GPUTaskId, CurrentInput().DownstreamPin->Label };
+
+			// If already created a gather task, then nothing more to do for this pin.
+			if (const FPCGTaskId* GatherTaskId = InputPinToGatherTask.Find(PinKey))
+			{
+				continue;
+			}
+
+			// If we're encountering pin for first time, register it.
+			if (!EncounteredInputPins.Contains(PinKey))
+			{
+				EncounteredInputPins.Add(PinKey);
+				continue;
+			}
+
+			// Second time we've encountered this input pin - create a gather element because we need one edge connected to
+			// each virtual input pin, so that we can obtain the data items from the input data collection using the unique
+			// virtual pin label at execution time.
+			const FPCGTaskId GatherTaskId = InOutCompiledTasks.Num();
+			FPCGGraphTask& GatherTask = InOutCompiledTasks.Emplace_GetRef();
+			GatherTask.NodeId = GatherTaskId;
+			GatherTask.ParentId = InOutCompiledTasks[GPUTaskId].ParentId;
+			GatherTask.Element = FPCGGraphCompiler::GetSharedGatherElement();
+
+			InputPinToGatherTask.Add(PinKey, GatherTaskId);
+		}
+
+		EncounteredInputPins.Reset();
+
+		// Second pass - wire up the newly added gather tasks once we have the full picture of which edges are affected.
+		for (int InputIndex = 0; InputIndex < InOutCompiledTasks[GPUTaskId].Inputs.Num(); ++InputIndex)
+		{
+			auto CurrentInput = [&InOutCompiledTasks, InputIndex, GPUTaskId]() -> FPCGGraphTaskInput&
+			{
+				return InOutCompiledTasks[GPUTaskId].Inputs[InputIndex];
+			};
+
+			if (!CurrentInput().DownstreamPin.IsSet())
+			{
+				continue;
+			}
+
+			const FOriginalInputPinKey PinKey = { GPUTaskId, CurrentInput().DownstreamPin->Label };
+
+			if (const FPCGTaskId* GatherTaskId = InputPinToGatherTask.Find(PinKey))
+			{
+				// Wire the upstream output pin to the gather task.
+				FPCGGraphTaskInput& WireUpstreamNodeToGather = InOutCompiledTasks[*GatherTaskId].Inputs.Add_GetRef(CurrentInput());
+				if (WireUpstreamNodeToGather.DownstreamPin.IsSet())
+				{
+					WireUpstreamNodeToGather.DownstreamPin->Label = PCGPinConstants::DefaultInputLabel;
+				}
+
+				if (!EncounteredInputPins.Contains(PinKey))
+				{
+					// First time we're encountering this input pin, wire it to the gather task.
+					EncounteredInputPins.Add(PinKey);
+
+					CurrentInput().TaskId = *GatherTaskId;
+					if (CurrentInput().UpstreamPin.IsSet())
+					{
+						CurrentInput().UpstreamPin->Label = PCGPinConstants::DefaultOutputLabel;
+					}
+				}
+				else
+				{
+					// Input pin already encountered, already wired to gather task. Remove this input.
+					InOutCompiledTasks[GPUTaskId].Inputs.RemoveAt(InputIndex);
+					--InputIndex;
+				}
+			}
+		}
+	}
+}
+
+void FPCGGraphCompilerGPU::WireGPUGraphNode(
+	FPCGTaskId InGPUGraphTaskId,
+	const TSet<FPCGTaskId>& InCollapsedTasks,
+	const TSet<FPCGTaskId>& InGPUCompatibleTaskIds,
+	TArray<FPCGGraphTask>& InOutCompiledTasks,
+	const TMap<FPCGTaskId, TArray<FPCGTaskId>>& InTaskSuccessors,
+	FOriginalToVirtualPin& OutOriginalToVirtualPin,
+	TMap<const UPCGPin*, FName>& OutOutputCPUPinToVirtualPin)
+{
+	FPCGGraphTask& GPUGraphTask = InOutCompiledTasks[InGPUGraphTaskId];
+
+	// Used to construct unique input/output labels, ultimately consumed in graph executor in BuildTaskInput and PostExecute for input/output respectively.
+	int InputCount = 0;
+	int OutputCount = 0;
+
+	// Add all compute graph task inputs and outputs.
+	for (FPCGTaskId GPUTaskId : InCollapsedTasks)
+	{
+		// First find CPU to GPU edges and wire in the GPU graph node inputs.
+		for (const FPCGGraphTaskInput& Input : InOutCompiledTasks[GPUTaskId].Inputs)
+		{
+			if (InGPUCompatibleTaskIds.Contains(Input.TaskId))
+			{
+				continue;
+			}
+
+			FPCGGraphTaskInput& AddedInput = GPUGraphTask.Inputs.Add_GetRef(Input);
+
+			// TODO is pinless fine with skipping?
+			if (AddedInput.DownstreamPin.IsSet())
+			{
+				const FName VirtualLabel = *FString::Format(TEXT("{0}-VirtualIn{1}"), { AddedInput.DownstreamPin->Label.ToString(), InputCount });
+				const bool bIsInputPin = true;
+				OutOriginalToVirtualPin.Add({ GPUTaskId, AddedInput.DownstreamPin->Label, bIsInputPin }, VirtualLabel);
+				AddedInput.DownstreamPin->Label = VirtualLabel;
+
+				++InputCount;
+
+				if (const UPCGNode* UpstreamNode = InOutCompiledTasks[Input.TaskId].Node)
+				{
+					if (const UPCGPin* OutputPin = UpstreamNode->GetOutputPin(AddedInput.UpstreamPin->Label))
+					{
+						OutOutputCPUPinToVirtualPin.Add(OutputPin, VirtualLabel);
+					}
+				}
+			}
+		}
+
+		if (!InTaskSuccessors.Contains(GPUTaskId))
+		{
+			continue;
+		}
+
+		// Next consider GPU to CPU edges to wire in the GPU graph node outputs.
+		for (FPCGTaskId Successor : InTaskSuccessors[GPUTaskId])
+		{
+			if (InGPUCompatibleTaskIds.Contains(Successor))
+			{
+				continue;
+			}
+
+			// Rewire inputs of this downstream CPU node to the outputs of the compute graph task.
+			FPCGGraphTask& DownstreamCPUNode = InOutCompiledTasks[Successor];
+
+			// Order matters here! We can never reorder inputs as it will impact execution.
+			const int InputCountBefore = DownstreamCPUNode.Inputs.Num();
+			for (int SuccessorInputIndex = 0; SuccessorInputIndex < InputCountBefore; ++SuccessorInputIndex)
+			{
+				// Implementation note: we modify the Inputs array in this loop, so don't take a reference to the current element.
+
+				// Skip irrelevant edges.
+				if (DownstreamCPUNode.Inputs[SuccessorInputIndex].TaskId != GPUTaskId)
+				{
+					continue;
+				}
+
+				// Wire downstream CPU node to compute graph task.
+				FPCGGraphTaskInput InputCopy = DownstreamCPUNode.Inputs[SuccessorInputIndex];
+
+				InputCopy.TaskId = InGPUGraphTaskId;
+
+				if (DownstreamCPUNode.Inputs[SuccessorInputIndex].UpstreamPin.IsSet())
+				{
+					const FNodePin PinKey = { GPUTaskId, InputCopy.UpstreamPin->Label, /*Pin is input*/false };
+					if (const FName* FoundVirtualPinLabel = OutOriginalToVirtualPin.Find(PinKey))
+					{
+						// Wire to the existing virtual output pin.
+						InputCopy.UpstreamPin->Label = *FoundVirtualPinLabel;
+					}
+					else
+					{
+						const FName VirtualLabel = *FString::Format(TEXT("{0}-VirtualOut{1}"), { InputCopy.UpstreamPin->Label.ToString(), OutputCount });
+						OutOriginalToVirtualPin.Add(PinKey, VirtualLabel);
+
+						InputCopy.UpstreamPin->Label = VirtualLabel;
+
+						++OutputCount;
+					}
+				}
+
+				DownstreamCPUNode.Inputs.Add(MoveTemp(InputCopy));
+			}
+		}
+	}
+}
+
+void FPCGGraphCompilerGPU::BuildGPUGraphTask(
+	UPCGGraph* InGraph,
+	FPCGTaskId InGPUGraphTaskId,
+	const TSet<FPCGTaskId>& InCollapsedTasks,
+	const TMap<FPCGTaskId, TArray<FPCGTaskId>>& InTaskSuccessors,
+	TArray<FPCGGraphTask>& InOutCompiledTasks,
+	const FOriginalToVirtualPin& InOriginalToVirtualPin,
+	const TMap<const UPCGPin*, FName>& InOutputCPUPinToVirtualPin)
+{
+	TSharedPtr<FPCGComputeGraphElement> Element = MakeShared<FPCGComputeGraphElement>();
+	InOutCompiledTasks[InGPUGraphTaskId].Element = Element;
+
+	const FName GraphName = MakeUniqueObjectName(InGraph, UPCGComputeGraph::StaticClass(), InGraph->GetFName());
+
+	UPCGComputeGraph* ComputeGraph = NewObject<UPCGComputeGraph>(InGraph, GraphName);
+	ComputeGraph->OutputCPUPinToInputGPUPinAlias = InOutputCPUPinToVirtualPin;
+	Element->Graph.Reset(ComputeGraph);
+
+	// Not incredibly useful for us - DG adds GetComponentSource()->GetComponentClass() object which allows it to bind at execution time by class.
+	// But execution code requires it currently.
+	Element->Graph->Bindings.Add(UPCGDataBinding::StaticClass());
+
+	// Create data interfaces which allow kernels to read or write data. Each data interface is associated with a node output pin.
+	// For CPU->GPU edges, an upload data interface is created. For GPU->CPU edges, a readback data interface is created.
+	auto CreateDataInterface = [&InCollapsedTasks, &InOutCompiledTasks, ComputeGraph](FPCGTaskId InTaskId, bool bRequiresReadback, const FPCGPinProperties& InOutputPinProperties) -> UPCGComputeDataInterface*
+	{
+		const bool bUpstreamIsGPUTask = InCollapsedTasks.Contains(InTaskId);
+
+		EPCGDataType PinType = InOutputPinProperties.AllowedTypes;
+
+		// Dynamically typed pins could have a different type
+		const UPCGSettings* Settings = InOutCompiledTasks[InTaskId].Node ? InOutCompiledTasks[InTaskId].Node->GetSettings() : nullptr;
+		const UPCGPin* Pin = InOutCompiledTasks[InTaskId].Node ? InOutCompiledTasks[InTaskId].Node->GetOutputPin(InOutputPinProperties.Label) : nullptr;
+		if (Settings && Pin)
+		{
+			PinType = Settings->GetCurrentPinTypes(Pin);
+		}
+
+		UPCGComputeDataInterface* DataInterface = nullptr;
+
+		switch (PinType)
+		{
+		case EPCGDataType::Point:
+		case EPCGDataType::PointOrParam:
+		{
+			UPCGDataCollectionDataInterface* DataInterfacePCGData = nullptr;
+
+			if (bUpstreamIsGPUTask)
+			{
+				if (bRequiresReadback)
+				{
+					// GPU -> CPU
+					DataInterfacePCGData = NewObject<UPCGDataCollectionReadbackDataInterface>(ComputeGraph);
+				}
+				else
+				{
+					// GPU -> GPU
+					DataInterfacePCGData = NewObject<UPCGDataCollectionDataInterface>(ComputeGraph);
+				}
+			}
+			else
+			{
+				// CPU -> GPU
+				DataInterfacePCGData = NewObject<UPCGDataCollectionUploadDataInterface>(ComputeGraph);
+			}
+
+			check(DataInterfacePCGData);
+			DataInterfacePCGData->ProducerSettings = InOutCompiledTasks[InTaskId].Node ? InOutCompiledTasks[InTaskId].Node->GetSettings() : nullptr;
+
+			DataInterface = DataInterfacePCGData;
+			break;
+		}
+		case EPCGDataType::Texture:
+		{
+			DataInterface = NewObject<UPCGTextureDataInterface>(ComputeGraph);
+			break;
+		}
+		case EPCGDataType::Landscape:
+		{
+			DataInterface = NewObject<UPCGLandscapeDataInterface>(ComputeGraph);
+			break;
+		}
+		default:
+			ensure(false);
+			break;
+		}
+
+		if (DataInterface)
+		{
+			DataInterface->SetOutputPin(InOutputPinProperties.Label);
+		}
+
+		return DataInterface;
+	};
+
+	TMap<TPair</* Node task*/FPCGTaskId, /*Node output pin*/FName>, UPCGComputeDataInterface*> OutputPinDataInterfaces;
+
+	// Create all the output data interfaces.
+	for (FPCGTaskId TaskId : InCollapsedTasks)
+	{
+		// Create DIs for all output pins, because the kernels currently need their outputs to be bound to valid resources.
+		if (const UPCGSettings* Settings = InOutCompiledTasks[TaskId].Node ? InOutCompiledTasks[TaskId].Node->GetSettings() : nullptr)
+		{
+			for (const FPCGPinProperties& OutputPinProperties : Settings->AllOutputPinProperties())
+			{
+				if (OutputPinDataInterfaces.Contains({ TaskId, OutputPinProperties.Label }))
+				{
+					ensure(false);
+					continue;
+				}
+
+				bool bRequiresReadback = false;
+				if (const TArray<FPCGTaskId>* Successors = InTaskSuccessors.Find(TaskId))
+				{
+					for (FPCGTaskId Successor : *Successors)
+					{
+						for (const FPCGGraphTaskInput& Input : InOutCompiledTasks[Successor].Inputs)
+						{
+							if (Input.UpstreamPin.IsSet() && (*Input.UpstreamPin == OutputPinProperties) && !InCollapsedTasks.Contains(Successor))
+							{
+								bRequiresReadback = true;
+								break;
+							}
+						}
+					}
+				}
+
+				if (UPCGComputeDataInterface* OutputDI = CreateDataInterface(TaskId, bRequiresReadback, OutputPinProperties))
+				{
+					OutputDI->SetOutputPin(OutputPinProperties.Label);
+					ComputeGraph->DataInterfaces.Add(OutputDI);
+					OutputPinDataInterfaces.Add({ TaskId, OutputPinProperties.Label }, OutputDI);
+
+					for (const FPCGKernelAttributeKey& Key : Settings->GetKernelAttributeKeys())
+					{
+						if (!ComputeGraph->GlobalAttributeLookupTable.Find(Key))
+						{
+							ComputeGraph->GlobalAttributeLookupTable.Add(Key, ComputeGraph->GlobalAttributeLookupTable.Num() + PCGComputeConstants::NUM_RESERVED_ATTRS);
+						}
+					}
+				}
+			}
+		}
+
+		// Create any DIs for upstream CPU nodes.
+		for (const FPCGGraphTaskInput& Input : InOutCompiledTasks[TaskId].Inputs)
+		{
+			// Only deal with upstream CPU tasks.
+			if (InCollapsedTasks.Contains(Input.TaskId))
+			{
+				continue;
+			}
+
+			if (!Input.DownstreamPin.IsSet())
+			{
+				continue;
+			}
+
+			if (const UPCGNode* DownstreamNode = InOutCompiledTasks[TaskId].Node)
+			{
+				if (const UPCGPin* InputPin = DownstreamNode->GetInputPin(Input.DownstreamPin->Label))
+				{
+					ComputeGraph->PinsReceivingDataFromCPU.Add(InputPin);
+				}
+			}
+
+			if (!Input.UpstreamPin.IsSet())
+			{
+				continue;
+			}
+
+			if (OutputPinDataInterfaces.Contains({ Input.TaskId, Input.UpstreamPin->Label }))
+			{
+				// Skip if already created.
+				continue;
+			}
+
+			if (UPCGComputeDataInterface* OutputDI = CreateDataInterface(Input.TaskId, /*bRequiresReadback=*/false, *Input.UpstreamPin))
+			{
+				OutputDI->SetOutputPin(Input.UpstreamPin->Label);
+				ComputeGraph->DataInterfaces.Add(OutputDI);
+				OutputPinDataInterfaces.Add({ Input.TaskId, Input.UpstreamPin->Label }, OutputDI);
+			}
+		}
+	}
+
+	TSet<FPCGTaskId> RemainingTasks = InCollapsedTasks;
+
+	while (!RemainingTasks.IsEmpty())
+	{
+		// Find a ready task
+		FPCGTaskId TaskId = InvalidPCGTaskId;
+		for (FPCGTaskId RemainingTask : RemainingTasks)
+		{
+			// TODO: use 'QueueSuccessors' pattern rather than brute force searching for ready tasks
+			bool bReady = true;
+			for (const FPCGGraphTaskInput& Input : InOutCompiledTasks[RemainingTask].Inputs)
+			{
+				if (RemainingTasks.Contains(Input.TaskId))
+				{
+					bReady = false;
+					break;
+				}
+			}
+
+			if (bReady)
+			{
+				TaskId = RemainingTask;
+				break;
+			}
+		}
+
+		if (TaskId == InvalidPCGTaskId)
+		{
+			// Tasks in RemainingTasks are unreachable?
+			ensure(false);
+			break;
+		}
+
+		RemainingTasks.Remove(TaskId);
+
+		const UPCGNode* Node = InOutCompiledTasks[TaskId].Node;
+
+		const UPCGCustomHLSLSettings* Settings = Cast<UPCGCustomHLSLSettings>(Node ? Node->GetSettings() : nullptr);
+		check(Settings && Settings->bEnabled && Settings->ShouldExecuteOnGPU());
+
+		// For every usage of a DI, get the original (non-aliased) pin label.
+		TMap<TPair<FPCGTaskId, UComputeDataInterface*>, FName> DataInterfaceUsageToPinLabel;
+
+		TArray<int> InputDataInterfaceIndices;
+		TArray<int> OutputDataInterfaceIndices;
+		InputDataInterfaceIndices.Reserve(Settings->InputPinProperties().Num());
+		OutputDataInterfaceIndices.Reserve(Settings->OutputPinProperties().Num());
+
+		// Add DIs (PCG -> CF transcoding).
+
+		for (const FPCGGraphTaskInput& Input : InOutCompiledTasks[TaskId].Inputs)
+		{
+			if (!Input.UpstreamPin.IsSet())
+			{
+				// Execution-only dependencies not supported currently. Unclear if this should ever be supported for GPU graphs.
+				// Writes followed by reads will be protected via barriers added by RDG.
+				continue;
+			}
+
+			UPCGComputeDataInterface* UpstreamDI = nullptr;
+			if (UPCGComputeDataInterface** FoundUpstreamDI = OutputPinDataInterfaces.Find({ Input.TaskId, Input.UpstreamPin->Label }))
+			{
+				UpstreamDI = *FoundUpstreamDI;
+			}
+
+			if (!UpstreamDI)
+			{
+				ensure(false);
+				continue;
+			}
+
+			const int Index = ComputeGraph->DataInterfaces.Find(UpstreamDI);
+			if (Index == INDEX_NONE)
+			{
+				ensure(false);
+				continue;
+			}
+
+			InputDataInterfaceIndices.Add(Index);
+
+			FName DownstreamInputPinLabel = Input.DownstreamPin->Label;
+
+			DataInterfaceUsageToPinLabel.Add({ TaskId, UpstreamDI }, DownstreamInputPinLabel);
+
+			const bool bIsInputPin = true;
+			UpstreamDI->AddDownstreamInputPin(DownstreamInputPinLabel, InOriginalToVirtualPin.Find({ TaskId, DownstreamInputPinLabel, bIsInputPin }));
+		}
+
+		// Always create a DI for every output pin, so kernel always has something to write to.
+		for (const FPCGPinProperties& OutputPinProperties : Settings->AllOutputPinProperties())
+		{
+			UPCGComputeDataInterface** FoundDI = OutputPinDataInterfaces.Find({ TaskId, OutputPinProperties.Label });
+			if (!ensure(FoundDI) || !ensure(*FoundDI))
+			{
+				continue;
+			}
+
+			const int Index = ComputeGraph->DataInterfaces.Find(*FoundDI);
+			if (Index == INDEX_NONE)
+			{
+				ensure(false);
+				continue;
+			}
+
+			OutputDataInterfaceIndices.Add(Index);
+			DataInterfaceUsageToPinLabel.Add({ TaskId, *FoundDI }, OutputPinProperties.Label);
+		}
+
+		// Make sure every downstream input pin is registered with the upstream DI.
+		if (const TArray<FPCGTaskId>* Successors = InTaskSuccessors.Find(TaskId))
+		{
+			for (FPCGTaskId Successor : *Successors)
+			{
+				for (const FPCGGraphTaskInput& SuccessorInput : InOutCompiledTasks[Successor].Inputs)
+				{
+					if (SuccessorInput.TaskId != TaskId)
+					{
+						continue;
+					}
+
+					if (!SuccessorInput.UpstreamPin.IsSet() || !SuccessorInput.DownstreamPin.IsSet())
+					{
+						continue;
+					}
+
+					FName OutputPinLabel = SuccessorInput.UpstreamPin->Label;
+
+					// DIs for output pins should have all been created.
+					UPCGComputeDataInterface** FoundOutputDI = OutputPinDataInterfaces.Find({ TaskId, OutputPinLabel });
+					if (!ensure(FoundOutputDI) || !ensure(*FoundOutputDI))
+					{
+						continue;
+					}
+
+					// If the map has an entry for this task and output pin label, then its a GPU->CPU readback.
+					const bool bIsInputPin = false;
+					(*FoundOutputDI)->SetOutputPin(OutputPinLabel, InOriginalToVirtualPin.Find({ TaskId, OutputPinLabel, bIsInputPin }));
+				}
+			}
+		}
+
+		UPCGCustomKernelDataInterface* KernelDI = NewObject<UPCGCustomKernelDataInterface>(ComputeGraph);
+		KernelDI->Settings = Settings;
+		const int KernelDIIndex = ComputeGraph->DataInterfaces.Num();
+		ComputeGraph->DataInterfaces.Add(KernelDI);
+
+		// TODO add graph data interface (graph params). Reference: UOptimusGraphDataInterface.
+		//Element->Graph->DataInterfaces.Add(NewObject< UPCGDataCollectionDataInterface>(InGraph));
+		//Element->Graph->DataInterfaceToBinding.Add(0);
+
+		// TODO once we support cooking for different platforms/configs, don't create the interface if logging is not present.
+		int DebugDIIndex = INDEX_NONE;
+		if (Settings->bPrintShaderDebugValues)
+		{
+			UPCGDebugDataInterface* DebugDI = NewObject<UPCGDebugDataInterface>(ComputeGraph);
+			DebugDI->SetDebugBufferSize(Settings->DebugBufferSize);
+
+			DebugDIIndex = ComputeGraph->DataInterfaces.Num();
+			ComputeGraph->DataInterfaces.Add(DebugDI);
+		}
+
+		// Now that all data interfaces added, create the (trivial) binding mapping. All map to primary binding, index 0.
+		ComputeGraph->DataInterfaceToBinding.SetNumZeroed(ComputeGraph->DataInterfaces.Num());
+
+		struct FInterfaceBinding
+		{
+			const UComputeDataInterface* DataInterface;
+			int32 DataInterfaceBindingIndex;
+			FString BindingFunctionName;
+			FString BindingFunctionNamespace;
+		};
+
+		struct FKernelWithDataBindings
+		{
+			UComputeKernel* Kernel;
+			TArray<FInterfaceBinding> InputDataBindings;
+			TArray<FInterfaceBinding> OutputDataBindings;
+		};
+
+		FKernelWithDataBindings KernelWithBindings;
+
+		KernelWithBindings.Kernel = NewObject<UComputeKernel>(ComputeGraph);
+		const int KernelIndex = ComputeGraph->KernelInvocations.Num();
+		ComputeGraph->KernelInvocations.Add(KernelWithBindings.Kernel);
+		ComputeGraph->KernelToNode.Add(Node);
+
+		auto SetupAllInputBindings = [&KernelWithBindings, ComputeGraph](int InDataInterfaceIndex)
+		{
+			const UComputeDataInterface* DataInterface = ComputeGraph->DataInterfaces[InDataInterfaceIndex];
+			TArray<FShaderFunctionDefinition> Functions;
+			DataInterface->GetSupportedInputs(Functions);
+
+			for (int FuncIndex = 0; FuncIndex < Functions.Num(); ++FuncIndex)
+			{
+				FInterfaceBinding& Binding = KernelWithBindings.InputDataBindings.Emplace_GetRef();
+				Binding.DataInterface = DataInterface;
+				Binding.BindingFunctionName = Functions[FuncIndex].Name;
+				Binding.BindingFunctionNamespace = TEXT("");
+				Binding.DataInterfaceBindingIndex = FuncIndex;
+			}
+		};
+
+		auto SetupAllOutputBindings = [&KernelWithBindings, ComputeGraph](int InDataInterfaceIndex)
+		{
+			const UComputeDataInterface* DataInterface = ComputeGraph->DataInterfaces[InDataInterfaceIndex];
+			TArray<FShaderFunctionDefinition> Functions;
+			DataInterface->GetSupportedOutputs(Functions);
+
+			for (int FuncIndex = 0; FuncIndex < Functions.Num(); ++FuncIndex)
+			{
+				FInterfaceBinding& Binding = KernelWithBindings.OutputDataBindings.Emplace_GetRef();
+				Binding.DataInterface = DataInterface;
+				Binding.BindingFunctionName = Functions[FuncIndex].Name;
+				Binding.BindingFunctionNamespace = TEXT("");
+				Binding.DataInterfaceBindingIndex = FuncIndex;
+			}
+		};
+
+		// Bind data interfaces.
+		for (int InputDataInterfaceIndex : InputDataInterfaceIndices)
+		{
+			SetupAllInputBindings(InputDataInterfaceIndex);
+		}
+
+		SetupAllInputBindings(KernelDIIndex);
+
+		for (int OutputDataInterfaceIndex : OutputDataInterfaceIndices)
+		{
+			SetupAllOutputBindings(OutputDataInterfaceIndex);
+		}
+
+		if (DebugDIIndex != INDEX_NONE)
+		{
+			SetupAllOutputBindings(DebugDIIndex);
+		}
+
+		{
+			UPCGComputeKernelSource* KernelSource = NewObject<UPCGComputeKernelSource>(KernelWithBindings.Kernel); // is outer to kernel fine?
+			KernelWithBindings.Kernel->KernelSource = KernelSource;
+			KernelSource->EntryPoint = Settings->GetKernelEntryPoint();
+			KernelSource->GroupSize = Settings->GetThreadGroupSize();
+
+			KernelSource->SetSource(Settings->GetCookedKernelSource(ComputeGraph->GlobalAttributeLookupTable));
+
+#if WITH_EDITOR
+			if (PCGGraphCompilerGPU::CVarEnableGPUDebugging.GetValueOnAnyThread())
+			{
+				UE_LOG(LogPCG, Warning, TEXT("ATTRIBUTE LOOK-UP TABLE [%s]"), *Settings->GetDefaultNodeTitle().ToString());
+
+				for (const TPair<FPCGKernelAttributeKey, int32 /* Attribute Index */>& Pair : ComputeGraph->GlobalAttributeLookupTable)
+				{
+					const FString TypeString = UEnum::GetValueAsString(Pair.Key.Type);
+					const FString NameString = Pair.Key.Name.ToString();
+					const FString IndexString = FString::FromInt(Pair.Value);
+
+					UE_LOG(LogPCG, Warning, TEXT("%s: %s (%s)"), *IndexString, *NameString, *TypeString);
+				}
+			}
+#endif
+
+			// Add functions for external inputs/outputs which must be fulfilled by DIs
+			for (FInterfaceBinding& Binding : KernelWithBindings.InputDataBindings)
+			{
+				TArray<FShaderFunctionDefinition> Functions;
+				Binding.DataInterface->GetSupportedInputs(Functions);
+				check(Functions.IsValidIndex(Binding.DataInterfaceBindingIndex));
+
+				FShaderFunctionDefinition FuncDef = Functions[Binding.DataInterfaceBindingIndex];
+				for (FShaderParamTypeDefinition& ParamType : FuncDef.ParamTypes)
+				{
+					// Making sure parameter has type declaration generated
+					ParamType.ResetTypeDeclaration();
+				}
+
+				KernelSource->ExternalInputs.Emplace(FuncDef);
+			}
+
+			for (FInterfaceBinding& Binding : KernelWithBindings.OutputDataBindings)
+			{
+				TArray<FShaderFunctionDefinition> Functions;
+				Binding.DataInterface->GetSupportedOutputs(Functions);
+				check(Functions.IsValidIndex(Binding.DataInterfaceBindingIndex));
+
+				FShaderFunctionDefinition FuncDef = Functions[Binding.DataInterfaceBindingIndex];
+				for (FShaderParamTypeDefinition& ParamType : FuncDef.ParamTypes)
+				{
+					// Making sure parameter has type declaration generated
+					ParamType.ResetTypeDeclaration();
+				}
+
+				KernelSource->ExternalOutputs.Emplace(FuncDef);
+			}
+		}
+
+		auto AddAllEdgesForKernel = [&KernelWithBindings, ComputeGraph, TaskId, &DataInterfaceUsageToPinLabel](int32 InKernelIndex, bool bInEdgesAreInputs)
+		{
+			TArray<FInterfaceBinding>& Bindings = bInEdgesAreInputs ? KernelWithBindings.InputDataBindings : KernelWithBindings.OutputDataBindings;
+
+			// Add all graph edges for bindings. This is somewhat odd but likely what we'll do vs the more finegrained interface
+			// definitions in optimus, but we need to see.
+			for (int BindingIndex = 0; BindingIndex < Bindings.Num(); ++BindingIndex)
+			{
+				FInterfaceBinding& Binding = Bindings[BindingIndex];
+
+				FComputeGraphEdge& Edge = ComputeGraph->GraphEdges.Emplace_GetRef();
+				Edge.KernelIndex = InKernelIndex;
+				Edge.KernelBindingIndex = BindingIndex;
+				Edge.DataInterfaceIndex = ComputeGraph->DataInterfaces.IndexOfByPredicate([&Binding](const UComputeDataInterface* In) { return Binding.DataInterface == In; });
+				check(Edge.DataInterfaceIndex != INDEX_NONE);
+				Edge.DataInterfaceBindingIndex = Binding.DataInterfaceBindingIndex;
+				Edge.bKernelInput = bInEdgesAreInputs;
+
+				UComputeDataInterface* DataInterface = ComputeGraph->DataInterfaces[Edge.DataInterfaceIndex];
+				check(DataInterface);
+
+				if (FName* PinLabel = DataInterfaceUsageToPinLabel.Find({ TaskId, DataInterface }))
+				{
+					TArray<FShaderFunctionDefinition> DataInterfaceFunctions;
+					if (bInEdgesAreInputs)
+					{
+						DataInterface->GetSupportedInputs(DataInterfaceFunctions);
+					}
+					else
+					{
+						DataInterface->GetSupportedOutputs(DataInterfaceFunctions);
+					}
+
+					Edge.BindingFunctionNameOverride = FString::Format(
+						TEXT("{0}_{1}"),
+						{ PinLabel->ToString(), DataInterfaceFunctions[Edge.DataInterfaceBindingIndex].Name }
+					);
+				}
+			}
+		};
+
+		AddAllEdgesForKernel(KernelIndex, /*bInEdgesAreInputs=*/true);
+		AddAllEdgesForKernel(KernelIndex, /*bInEdgesAreInputs=*/false);
+	}
+
+	// Register all virtual pin aliases with the corresponding pins for downstream usage.
+	for (const TPair<TTuple<FPCGTaskId, FName, bool>, FName>& Mapping : InOriginalToVirtualPin)
+	{
+		const bool bIsInputPin = Mapping.Get<0>().Get<2>();
+		if (bIsInputPin)
+		{
+			const FPCGTaskId TaskId = Mapping.Get<0>().Get<0>();
+
+			if (const UPCGNode* Node = InOutCompiledTasks[TaskId].Node)
+			{
+				const FName& OriginalLabel = Mapping.Get<0>().Get<1>();
+
+				if (const UPCGPin* Pin = Node->GetInputPin(OriginalLabel))
+				{
+					const FName& VirtualLabel = Mapping.Get<1>();
+					ComputeGraph->InputPinLabelAliases.FindOrAdd(Pin) = VirtualLabel;
+				}
+			}
+		}
+	}
+
+	// Kick off shader compilation (if needed).
+	ComputeGraph->UpdateResources();
+}
+
+void FPCGGraphCompilerGPU::CreateGPUNodes(UPCGGraph* InGraph, TArray<FPCGGraphTask>& InOutCompiledTasks)
+{
+	TSet<FPCGTaskId> GPUCompatibleTaskIds;
+	GPUCompatibleTaskIds.Reserve(InOutCompiledTasks.Num());
+	for (FPCGTaskId TaskId = 0; TaskId < InOutCompiledTasks.Num(); ++TaskId)
+	{
+		const UPCGNode* Node = InOutCompiledTasks[TaskId].Node;
+		const UPCGSettings* Settings = Node ? Node->GetSettings() : nullptr;
+		if (Settings && Settings->ShouldExecuteOnGPU() && Settings->bEnabled)
+		{
+			GPUCompatibleTaskIds.Add(TaskId);
+		}
+	}
+
+	if (GPUCompatibleTaskIds.IsEmpty())
+	{
+		// Nothing to do for this graph.
+		return;
+	}
+
+	TMap<FPCGTaskId, TArray<FPCGTaskId>> TaskSuccessors;
+	TaskSuccessors.Reserve(InOutCompiledTasks.Num());
+	for (FPCGTaskId TaskId = 0; TaskId < InOutCompiledTasks.Num(); ++TaskId)
+	{
+		for (int InputIndex = 0; InputIndex < InOutCompiledTasks[TaskId].Inputs.Num(); ++InputIndex)
+		{
+			TaskSuccessors.FindOrAdd(InOutCompiledTasks[TaskId].Inputs[InputIndex].TaskId).Add(TaskId);
+		}
+	}
+
+	// For input pins at CPU -> GPU boundary, inject gather elements to pre-combine data on CPU side
+	// before passing to GPU.
+	CreateGatherTasksAtGPUInputs(GPUCompatibleTaskIds, InOutCompiledTasks);
+
+	TArray<TSet<FPCGTaskId>> NodeSubsetsToConvertToCFGraph;
+	CollectGPUNodeSubsets(InOutCompiledTasks, TaskSuccessors, GPUCompatibleTaskIds, NodeSubsetsToConvertToCFGraph);
+
+	// Do actual collapsing now, one subset at a time. Each collapse will do all fixup of task ids? That will invalidate
+	// ids in NodeSubsetsToConvertToCFGraph, so may need remap table. But can ignore this for now.
+	for (TSet<FPCGTaskId>& NodeSubsetToConvertToCFGraph : NodeSubsetsToConvertToCFGraph)
+	{
+		if (NodeSubsetToConvertToCFGraph.IsEmpty())
+		{
+			ensure(false);
+			continue;
+		}
+
+		// Add a new compute graph task. Then kill the original GPU tasks.
+		const FPCGTaskId ComputeGraphTaskId = InOutCompiledTasks.Num();
+		FPCGGraphTask& ComputeGraphTask = InOutCompiledTasks.Emplace_GetRef();
+		ComputeGraphTask.NodeId = ComputeGraphTaskId;
+
+		// All nodes in subset will be from same stack/parent, so assign from any.
+		for (FPCGTaskId GPUTaskId : NodeSubsetToConvertToCFGraph)
+		{
+			ComputeGraphTask.ParentId = InOutCompiledTasks[GPUTaskId].ParentId;
+			ComputeGraphTask.StackIndex = InOutCompiledTasks[GPUTaskId].StackIndex;
+			break;
+		}
+
+		// Mapping from task ID & pin label to a virtual pin label. Compute graphs are executed within a generated element,
+		// and the input and output pins of this element must have unique virtual pin labels so that we can parse the data that
+		// PCG provides through the input data collection correctly, and route the output data to the downstream pins correctly.
+		FOriginalToVirtualPin OriginalToVirtualPin;
+
+		TMap<const UPCGPin*, FName> OutputCPUPinToVirtualPin;
+
+		// Wire in the compute graph task, side by side with the individual GPU tasks, which will be culled below.
+		WireGPUGraphNode(
+			ComputeGraphTaskId,
+			NodeSubsetToConvertToCFGraph,
+			GPUCompatibleTaskIds,
+			InOutCompiledTasks,
+			TaskSuccessors,
+			OriginalToVirtualPin,
+			OutputCPUPinToVirtualPin);
+
+		// Generate a compute graph from all of the individual GPU tasks.
+		BuildGPUGraphTask(
+			InGraph,
+			ComputeGraphTaskId,
+			NodeSubsetToConvertToCFGraph,
+			TaskSuccessors,
+			InOutCompiledTasks,
+			OriginalToVirtualPin,
+			OutputCPUPinToVirtualPin);
+	}
+
+	// Now cull all the GPU compatible nodes. The compute graph task are already wired in so we're fine to just delete.
+	FPCGGraphCompiler::CullTasks(InOutCompiledTasks, /*bAddPassthroughWires=*/false, [&GPUCompatibleTaskIds](const FPCGGraphTask& InTask)
+	{
+		return GPUCompatibleTaskIds.Contains(InTask.NodeId);
+	});
+}
diff --git a/Engine/Plugins/PCG/Source/PCG/Private/Graph/PCGGraphCompilerGPU.h b/Engine/Plugins/PCG/Source/PCG/Private/Graph/PCGGraphCompilerGPU.h
new file mode 100644
--- /dev/null
+++ b/Engine/Plugins/PCG/Source/PCG/Private/Graph/PCGGraphCompilerGPU.h
@@ -0,0 +1,53 @@
+// Copyright Epic Games, Inc. All Rights Reserved.
+
+#pragma once
+
+#include "PCGCommon.h"
+
+class UPCGGraph;
+class UPCGPin;
+struct FPCGGraphTask;
+
+class FPCGGraphCompilerGPU
+{
+public:
+	/** Used to track new unique virtual pins created on generated compute graph elements. */
+	using FNodePin = TTuple<FPCGTaskId, /*Pin label*/FName, /*Pin is input*/bool>;
+	using FOriginalToVirtualPin = TMap<FNodePin, /*Virtual pin label*/FName>;
+	
+	/** Outputs sets of task IDs, where each set is GPU nodes that can be compiled into a compute graph and dispatched together. */
+	static void CollectGPUNodeSubsets(
+		const TArray<FPCGGraphTask>& InCompiledTasks,
+		const TMap<FPCGTaskId, TArray<FPCGTaskId>>& InTaskSuccessors,
+		const TSet<FPCGTaskId>& InGPUCompatibleTaskIds,
+		TArray<TSet<FPCGTaskId>>& OutNodeSubsetsToConvertToCFGraph);
+	
+	/** For GPU node inputs that have multiple incident edges, bundle them into a single edge. This is to avoid an inefficient
+	* gather operation on the GPU, and allows data interfaces to pick their data from the compute graph element input data collection
+	* using unique virtual input pin labels. */
+	static void CreateGatherTasksAtGPUInputs(const TSet<FPCGTaskId>& InGPUCompatibleTaskIds, TArray<FPCGGraphTask>& InOutCompiledTasks);
+	
+	/** Wires in a compute graph element alongside each set of GPU compatible nodes. The tasks for each node will be culled later. */
+	static void WireGPUGraphNode(
+		FPCGTaskId InGPUGraphTaskId,
+		const TSet<FPCGTaskId>& InCollapsedTasks,
+		const TSet<FPCGTaskId>& InGPUCompatibleTaskIds,
+		TArray<FPCGGraphTask>& InOutCompiledTasks,
+		const TMap<FPCGTaskId,
+		TArray<FPCGTaskId>>&InTaskSuccessors,
+		FOriginalToVirtualPin& OutOriginalToVirtualPin,
+		TMap<const UPCGPin*, FName>& OutOutputCPUPinToVirtualPin);
+	
+	/** Compiles a compute graph. */
+	static void BuildGPUGraphTask(
+		UPCGGraph* InGraph,
+		FPCGTaskId InGPUGraphTaskId,
+		const TSet<FPCGTaskId>& InCollapsedTasks,
+		const TMap<FPCGTaskId, TArray<FPCGTaskId>>& InTaskSuccessors,
+		TArray<FPCGGraphTask>& InOutCompiledTasks,
+		const FOriginalToVirtualPin& InOriginalToVirtualPin,
+		const TMap<const UPCGPin*, FName>& InOutputCPUPinToVirtualPin);
+	
+	/** Finds connected subgraphs of GPU - enabled nodes that can be dispatched together and replaces each one with a compute graph. */
+	static void CreateGPUNodes(UPCGGraph* InGraph, TArray<FPCGGraphTask>& InOutCompiledTasks);
+};
