[PCG] Always release landscape texture in PCGLandscapeDataInterface

#rnx

diff --git a/Engine/Plugins/PCG/Source/PCG/Private/Compute/DataInterfaces/PCGLandscapeDataInterface.cpp b/Engine/Plugins/PCG/Source/PCG/Private/Compute/DataInterfaces/PCGLandscapeDataInterface.cpp
--- a/Engine/Plugins/PCG/Source/PCG/Private/Compute/DataInterfaces/PCGLandscapeDataInterface.cpp
+++ b/Engine/Plugins/PCG/Source/PCG/Private/Compute/DataInterfaces/PCGLandscapeDataInterface.cpp
@@ -84,6 +84,35 @@ private:
 	FTextureBulkData HeightBulkData;
 };
 
+/** Manages the resources created by this DI. */
+class FPCGLandscapeResource
+{
+public:
+	struct FResourceKey
+	{
+		TWeakObjectPtr<const ALandscape> Source = nullptr;
+		TArray<FIntPoint> CapturedRegions;
+		FIntPoint MinCaptureRegion = FIntPoint(ForceInitToZero);
+		FIntPoint MaxCaptureRegion = FIntPoint(ForceInitToZero);
+		bool bIncludesHeight = false;
+	};
+
+	FPCGLandscapeResource() {}
+	FPCGLandscapeResource(const FResourceKey& InKey);
+	~FPCGLandscapeResource();
+
+	FPCGLandscapeTextureResource* LandscapeTextures = nullptr;
+	FVector3f LandscapeLWCTile = FVector3f::ZeroVector;
+	FMatrix ActorToWorldTransform = FMatrix::Identity;
+	FMatrix WorldToActorTransform = FMatrix::Identity;
+	FVector4 UvScaleBias = FVector4(1.0f, 1.0f, 0.0f, 0.0f);
+	FIntPoint CellCount = FIntPoint(ForceInitToZero);
+	FVector2D TextureWorldGridSize = FVector2D(1.0f, 1.0f);
+
+private:
+	FResourceKey ResourceKey;
+};
+
 FPCGLandscapeResource::FPCGLandscapeResource(const FResourceKey& InKey)
 	: ResourceKey(InKey)
 {
@@ -146,10 +175,12 @@ FPCGLandscapeResource::FPCGLandscapeResource(const FResourceKey& InKey)
 	TextureWorldGridSize = FVector2D(ResourceKey.Source->GetTransform().GetScale3D());
 }
 
-void FPCGLandscapeResource::Release()
+FPCGLandscapeResource::~FPCGLandscapeResource()
 {
 	if (LandscapeTextures)
 	{
+		// Release the LandscapeTextures resource handle on the RHI and clear the pointer.
+
 		ENQUEUE_RENDER_COMMAND(BeginDestroyCommand)([RT_Resource = LandscapeTextures](FRHICommandListImmediate& RHICmdList)
 		{
 			check(RT_Resource);
@@ -374,18 +405,12 @@ void UPCGLandscapeDataProvider::Initialize(ALandscape* InLandscape, const FBox&
 		}
 	}
 
-	Resource = FPCGLandscapeResource(Key);
-}
-
-UPCGLandscapeDataProvider::~UPCGLandscapeDataProvider()
-{
-	// Implementation note: Releasing the resource here assumes the Data Provider will always outlive the DataProviderProxies.
-	Resource.Release();
+	Resource = MakeUnique<FPCGLandscapeResource>(Key);
 }
 
 FComputeDataProviderRenderProxy* UPCGLandscapeDataProvider::GetRenderProxy()
 {
-	FPCGLandscapeDataProviderProxy* Proxy = new FPCGLandscapeDataProviderProxy(Resource);
+	FPCGLandscapeDataProviderProxy* Proxy = new FPCGLandscapeDataProviderProxy(Resource.Get());
 	return Proxy;
 }
 
@@ -412,15 +437,26 @@ void FPCGLandscapeDataProviderProxy::GatherDispatchData(FDispatchData const& InD
 		Parameters.HeightTextureSampler = RHIPixelFormatHasCapabilities(EPixelFormat::PF_R32_FLOAT, EPixelFormatCapabilities::TextureFilterable) ? BilinearSamplerState : PointClampedSampler;
 
 		// Set Textures
-		FRHITexture* HeightTexture = Resource.LandscapeTextures ? Resource.LandscapeTextures->GetHeightTexture() : (FRHITexture*)GBlackTexture->TextureRHI;
+		FRHITexture* HeightTexture = (Resource && Resource->LandscapeTextures) ? Resource->LandscapeTextures->GetHeightTexture() : (FRHITexture*)GBlackTexture->TextureRHI;
 		check(HeightTexture);
 
 		// TODO: Bindless resources to handle multiple landscapes
 		Parameters.HeightTexture = HeightTexture;
-		Parameters.HeightTextureLWCTile = Resource.LandscapeLWCTile;
-		Parameters.HeightTextureWorldToUvTransform = (FMatrix44f)Resource.WorldToActorTransform;
-		Parameters.HeightTextureUvScaleBias = (FVector4f)Resource.UvScaleBias;
-		Parameters.HeightTextureWorldGridSize = (FVector2f)Resource.TextureWorldGridSize;
+
+		if (Resource)
+		{
+			Parameters.HeightTextureLWCTile = Resource->LandscapeLWCTile;
+			Parameters.HeightTextureWorldToUvTransform = (FMatrix44f)Resource->WorldToActorTransform;
+			Parameters.HeightTextureUvScaleBias = (FVector4f)Resource->UvScaleBias;
+			Parameters.HeightTextureWorldGridSize = (FVector2f)Resource->TextureWorldGridSize;
+		}
+		else
+		{
+			Parameters.HeightTextureLWCTile = FVector3f::ZeroVector;
+			Parameters.HeightTextureWorldToUvTransform = (FMatrix44f)FMatrix::Identity;
+			Parameters.HeightTextureUvScaleBias = (FVector4f)FVector4(1.0f, 1.0f, 0.0f, 0.0f);
+			Parameters.HeightTextureWorldGridSize = (FVector2f)FVector2D(1.0f, 1.0f);
+		}
 
 		// System tile for LWC
 		FVector3f LWCTile(0, 0, 0); // Used to offset in LWC for precision, default to 0-vector for now.
@@ -430,9 +466,9 @@ void FPCGLandscapeDataProviderProxy::GatherDispatchData(FDispatchData const& InD
 
 void FPCGLandscapeDataProviderProxy::AllocateResources(FRDGBuilder& GraphBuilder, FAllocationData const& InAllocationData)
 {
-	if (Resource.LandscapeTextures)
+	if (Resource && Resource->LandscapeTextures)
 	{
-		Resource.LandscapeTextures->InitResource(GraphBuilder.RHICmdList);
+		Resource->LandscapeTextures->InitResource(GraphBuilder.RHICmdList);
 	}
 }
 
diff --git a/Engine/Plugins/PCG/Source/PCG/Private/Compute/DataInterfaces/PCGLandscapeDataInterface.h b/Engine/Plugins/PCG/Source/PCG/Private/Compute/DataInterfaces/PCGLandscapeDataInterface.h
--- a/Engine/Plugins/PCG/Source/PCG/Private/Compute/DataInterfaces/PCGLandscapeDataInterface.h
+++ b/Engine/Plugins/PCG/Source/PCG/Private/Compute/DataInterfaces/PCGLandscapeDataInterface.h
@@ -9,38 +9,7 @@
 
 class ALandscape;
 class FPCGLandscapeDataInterfaceParameters;
-class FPCGLandscapeTextureResource;
-
-/** Manages the resources created by this DI. */
-class FPCGLandscapeResource
-{
-public:
-	struct FResourceKey
-	{
-		TWeakObjectPtr<const ALandscape> Source = nullptr;
-		TArray<FIntPoint> CapturedRegions;
-		FIntPoint MinCaptureRegion = FIntPoint(ForceInitToZero);
-		FIntPoint MaxCaptureRegion = FIntPoint(ForceInitToZero);
-		bool bIncludesHeight = false;
-	};
-
-	FPCGLandscapeResource() {}
-	FPCGLandscapeResource(const FResourceKey& InKey);
-
-	/** Release the LandscapeTextures resource handle on the RHI and clear the pointer. */
-	void Release();
-
-	FPCGLandscapeTextureResource* LandscapeTextures = nullptr;
-	FVector3f LandscapeLWCTile = FVector3f::ZeroVector;
-	FMatrix ActorToWorldTransform = FMatrix::Identity;
-	FMatrix WorldToActorTransform = FMatrix::Identity;
-	FVector4 UvScaleBias = FVector4(1.0f, 1.0f, 0.0f, 0.0f);
-	FIntPoint CellCount = FIntPoint(ForceInitToZero);
-	FVector2D TextureWorldGridSize = FVector2D(1.0f, 1.0f);
-
-private:
-	FResourceKey ResourceKey;
-};
+class FPCGLandscapeResource;
 
 /** Data Interface allowing sampling of a Landscape. */
 UCLASS(ClassGroup = (Procedural))
@@ -71,20 +40,19 @@ class UPCGLandscapeDataProvider : public UComputeDataProvider
 
 public:
 	void Initialize(ALandscape* InLandscape, const FBox& Bounds);
-	~UPCGLandscapeDataProvider();
 
 	//~ Begin UComputeDataProvider Interface
 	FComputeDataProviderRenderProxy* GetRenderProxy() override;
 	//~ End UComputeDataProvider Interface
 
 private:
-	FPCGLandscapeResource Resource;
+	TUniquePtr<FPCGLandscapeResource> Resource;
 };
 
 class FPCGLandscapeDataProviderProxy : public FComputeDataProviderRenderProxy
 {
 public:
-	FPCGLandscapeDataProviderProxy(const FPCGLandscapeResource& InResource)
+	FPCGLandscapeDataProviderProxy(const FPCGLandscapeResource* InResource)
 		: Resource(InResource)
 	{}
 
@@ -97,5 +65,5 @@ public:
 protected:
 	using FParameters = FPCGLandscapeDataInterfaceParameters;
 
-	FPCGLandscapeResource Resource;
+	const FPCGLandscapeResource* Resource = nullptr;
 };
