Added component selection and the extensible data registry to support better filtering.

diff --git a/Engine/Plugins/PCG/Source/PCG/Private/Data/PCGGetDataFunctionRegistry.cpp b/Engine/Plugins/PCG/Source/PCG/Private/Data/PCGGetDataFunctionRegistry.cpp
new file mode 100644
--- /dev/null
+++ b/Engine/Plugins/PCG/Source/PCG/Private/Data/PCGGetDataFunctionRegistry.cpp
@@ -0,0 +1,297 @@
+// Copyright Epic Games, Inc. All Rights Reserved.
+
+#include "Data/PCGGetDataFunctionRegistry.h"
+
+#include "PCGComponent.h"
+#include "PCGContext.h"
+#include "PCGData.h"
+#include "PCGGraph.h"
+#include "Data/PCGCollisionShapeData.h"
+#include "Data/PCGIntersectionData.h"
+#include "Data/PCGLandscapeData.h"
+#include "Data/PCGLandscapeSplineData.h"
+#include "Data/PCGPointData.h"
+#include "Data/PCGPrimitiveData.h"
+#include "Data/PCGProjectionData.h"
+#include "Data/PCGSpatialData.h"
+#include "Data/PCGSplineData.h"
+#include "Data/PCGUnionData.h"
+#include "Data/PCGVolumeData.h"
+#include "Grid/PCGPartitionActor.h"
+#include "Helpers/PCGActorHelpers.h"
+#include "Helpers/PCGBlueprintHelpers.h"
+#include "Helpers/PCGHelpers.h"
+
+#include "LandscapeComponent.h"
+#include "LandscapeProxy.h"
+#include "LandscapeSplinesComponent.h"
+#include "Components/ActorComponent.h"
+#include "Components/BillboardComponent.h"
+#include "Components/InstancedStaticMeshComponent.h"
+#include "Components/ShapeComponent.h"
+#include "Components/SplineComponent.h"
+#include "Components/SplineMeshComponent.h"
+#include "GameFramework/Actor.h"
+#include "GameFramework/Volume.h"
+
+FPCGGetDataFunctionRegistry::FFunctionHandle FPCGGetDataFunctionRegistry::RegisterDataFromActorFunction(const FPCGGetDataFunctionRegistry::FDataFromActorFunction& InFunction)
+{
+	check(IsInGameThread());
+	ActorParsingFunctions.Emplace(InFunction, NextFunctionHandle);
+	return NextFunctionHandle++;
+}
+
+void FPCGGetDataFunctionRegistry::UnregisterDataFromActorFunction(FPCGGetDataFunctionRegistry::FFunctionHandle InFunctionHandle)
+{
+	check(IsInGameThread());
+	int FunctionIndex = ActorParsingFunctions.IndexOfByPredicate([InFunctionHandle](const TPair<FDataFromActorFunction, FFunctionHandle>& InPair) { return InPair.Value == InFunctionHandle; });
+	if (FunctionIndex != INDEX_NONE)
+	{
+		ActorParsingFunctions.RemoveAt(FunctionIndex);
+	}
+}
+
+FPCGGetDataFunctionRegistry::FFunctionHandle FPCGGetDataFunctionRegistry::RegisterDataFromComponentFunction(const FPCGGetDataFunctionRegistry::FDataFromComponentFunction& InFunction)
+{
+	check(IsInGameThread());
+	ComponentParsingFunctions.Emplace(InFunction, NextFunctionHandle);
+	return NextFunctionHandle++;
+}
+
+void FPCGGetDataFunctionRegistry::UnregisterDataFromComponentFunction(FPCGGetDataFunctionRegistry::FFunctionHandle InFunctionHandle)
+{
+	check(IsInGameThread());
+	int FunctionIndex = ComponentParsingFunctions.IndexOfByPredicate([InFunctionHandle](const TPair<FDataFromComponentFunction, FFunctionHandle>& InPair) { return InPair.Value == InFunctionHandle; });
+	if (FunctionIndex != INDEX_NONE)
+	{
+		ComponentParsingFunctions.RemoveAt(FunctionIndex);
+	}
+}
+
+int FPCGGetDataFunctionRegistry::GetDataFromActor(FPCGContext* InContext, const FPCGGetDataFunctionRegistryParams& InParams, AActor* InActor, FPCGGetDataFunctionRegistryOutput& Output) const
+{
+	int DataCountBefore = Output.Collection.TaggedData.Num();
+
+	// Try custom getters first
+	for (const TPair<FDataFromActorFunction, FFunctionHandle>& ActorParsingPair : ActorParsingFunctions)
+	{
+		if (ActorParsingPair.Key(InContext, InParams, InActor, Output))
+		{
+			break;
+		}
+	}
+
+	if (Output.Collection.TaggedData.Num() != DataCountBefore)
+	{
+		return Output.Collection.TaggedData.Num() - DataCountBefore;
+	}
+
+	// Default actor parsing, if any
+	return DefaultDataFromActor(InContext, InParams, InActor, Output);
+}
+
+int FPCGGetDataFunctionRegistry::DefaultDataFromActor(FPCGContext* InContext, const FPCGGetDataFunctionRegistryParams& InParams, AActor* Actor, FPCGGetDataFunctionRegistryOutput& Output) const
+{
+	check(Actor);
+
+	// Early out if the actor gets rejected by the component selector
+	if (InParams.ComponentSelector && !InParams.ComponentSelector->FilterActor(Actor))
+	{
+		return 0;
+	}
+
+	const EPCGDataType InDataFilter = InParams.DataTypeFilter;
+	const bool bParseActor = InParams.bParseActor;
+	const UPCGComponent* Component = InParams.SourceComponent;
+
+	auto NameTagsToStringTags = [](const FName& InName) { return InName.ToString(); };
+	TSet<FString> ActorTags;
+	Algo::Transform(Actor->Tags, ActorTags, NameTagsToStringTags);
+
+	// Fill in collection based on the data on the given actor.
+	// Some actor types we will forego full parsing to build strictly on the actor existence, such as partition actors, volumes and landscape
+	// TODO: review the !bParseActor cases - it might make sense to have just a point for a partition actor, even if we preintersect it.
+	APCGPartitionActor* PartitionActor = (!!(InDataFilter & EPCGDataType::Spatial) || !!(InDataFilter & EPCGDataType::Volume)) ? Cast<APCGPartitionActor>(Actor) : nullptr;
+	ALandscapeProxy* LandscapeActor = !!(InDataFilter & EPCGDataType::Landscape) ? Cast<ALandscapeProxy>(Actor) : nullptr;
+	AVolume* VolumeActor = !!(InDataFilter & EPCGDataType::Volume) ? Cast<AVolume>(Actor) : nullptr;
+	if (!bParseActor && !!(InDataFilter & EPCGDataType::Point))
+	{
+		UPCGPointData* Data = FPCGContext::NewObject_AnyThread<UPCGPointData>(InContext);
+		Data->InitializeFromActor(Actor, &Output.bSanitizedTagAttributeNames);
+
+		FPCGTaggedData& TaggedData = Output.Collection.TaggedData.Emplace_GetRef();
+		TaggedData.Data = Data;
+		TaggedData.Tags = ActorTags;
+	}
+	else if (PartitionActor)
+	{
+		check(!Component || Component->GetOwner() == Actor); // Invalid processing otherwise because of the this usage
+
+		UPCGVolumeData* VolumeData = FPCGContext::NewObject_AnyThread<UPCGVolumeData>(InContext);
+		UPCGSpatialData* Result = VolumeData;
+		if (InDataFilter == EPCGDataType::Volume)
+		{
+			VolumeData->Initialize(PCGHelpers::GetGridBounds(Actor, Component));
+		}
+		else
+		{
+			VolumeData->Initialize(PartitionActor->GetFixedBounds());
+
+			UPCGComponent* OriginalComponent = Component ? PartitionActor->GetOriginalComponent(Component) : nullptr;
+			// Important note: we do NOT call the collection version here, as we want to have a union if that's the case
+			const UPCGSpatialData* OriginalComponentSpatialData = OriginalComponent ? Cast<const UPCGSpatialData>(OriginalComponent->GetActorPCGData()) : nullptr;
+
+			if (OriginalComponentSpatialData)
+			{
+				Result = Result->IntersectWith(nullptr, OriginalComponentSpatialData);
+			}
+		}
+
+		FPCGTaggedData& TaggedData = Output.Collection.TaggedData.Emplace_GetRef();
+		TaggedData.Data = Result;
+		// No need to keep partition actor tags, though we might want to push PCG grid GUID at some point
+	}
+	else if (LandscapeActor)
+	{
+		UPCGLandscapeData* Data = FPCGContext::NewObject_AnyThread<UPCGLandscapeData>(InContext);
+		const UPCGGraph* PCGGraph = Component ? Component->GetGraph() : nullptr;
+
+		FPCGLandscapeDataProps LandscapeDataProps;
+		LandscapeDataProps.bGetHeightOnly = false;
+		LandscapeDataProps.bGetLayerWeights = (!PCGGraph || PCGGraph->bLandscapeUsesMetadata);
+
+		Data->Initialize({ LandscapeActor }, PCGHelpers::GetGridBounds(Actor, Component), LandscapeDataProps);
+
+		FPCGTaggedData& TaggedData = Output.Collection.TaggedData.Emplace_GetRef();
+		TaggedData.Data = Data;
+		TaggedData.Tags = ActorTags;
+	}
+	else if (VolumeActor)
+	{
+		UPCGVolumeData* Data = FPCGContext::NewObject_AnyThread<UPCGVolumeData>(InContext);
+		Data->Initialize(VolumeActor);
+
+		FPCGTaggedData& TaggedData = Output.Collection.TaggedData.Emplace_GetRef();
+		TaggedData.Data = Data;
+		TaggedData.Tags = ActorTags;
+	}
+	else
+	{
+		// Use component-based approach
+		int AddedData = 0;
+		TInlineComponentArray<UActorComponent*> ActorComponents;
+		Actor->GetComponents(ActorComponents);
+
+		for (UActorComponent* ActorComponent : ActorComponents)
+		{
+			AddedData += GetDataFromComponent(InContext, InParams, ActorComponent, Output);
+		}
+		
+		// IMPLEMENTATION NOTE: API BREAK HERE, we normally would have returned a single point here if there were no valid components
+		return AddedData;
+	}
+
+	return 1;
+}
+
+int FPCGGetDataFunctionRegistry::GetDataFromComponent(FPCGContext* InContext, const FPCGGetDataFunctionRegistryParams& InParams, UActorComponent* InComponent, FPCGGetDataFunctionRegistryOutput& Output) const
+{
+	check(InComponent);
+
+	if (InParams.ComponentSelector && !InParams.ComponentSelector->FilterComponent(InComponent))
+	{
+		return 0;
+	}
+
+	int DataCountBefore = Output.Collection.TaggedData.Num();
+
+	for (const TPair<FDataFromComponentFunction, FFunctionHandle>& ComponentParsingPair : ComponentParsingFunctions)
+	{
+		if (ComponentParsingPair.Key(InContext, InParams, InComponent, Output))
+		{
+			break;
+		}
+	}
+
+	if (Output.Collection.TaggedData.Num() != DataCountBefore)
+	{
+		return Output.Collection.TaggedData.Num() - DataCountBefore;
+	}
+
+	// Default component parsing, if any
+	return DefaultDataFromComponent(InContext, InParams, InComponent, Output);
+}
+
+int FPCGGetDataFunctionRegistry::DefaultDataFromComponent(FPCGContext* InContext, const FPCGGetDataFunctionRegistryParams& InParams, UActorComponent* InComponent, FPCGGetDataFunctionRegistryOutput& Output) const
+{
+	check(InComponent);
+
+	if (InParams.bIgnorePCGGeneratedComponents && InComponent->ComponentTags.Contains(PCGHelpers::DefaultPCGTag))
+	{
+		return 0;
+	}
+
+	auto NameTagsToStringTags = [](const FName& InName) { return InName.ToString(); };
+	const EPCGDataType InDataFilter = InParams.DataTypeFilter;
+
+	// Landscape spline component
+	// Splines component
+	// Shape component
+	// Primitive component, but ignore billboard components
+	FPCGTaggedData TaggedData;
+
+	if (!TaggedData.Data && !!(InDataFilter & EPCGDataType::Spline))
+	{
+		if (ULandscapeSplinesComponent* LandscapeSplineComponent = Cast<ULandscapeSplinesComponent>(InComponent))
+		{
+			UPCGLandscapeSplineData* SplineData = FPCGContext::NewObject_AnyThread<UPCGLandscapeSplineData>(InContext);
+			SplineData->Initialize(LandscapeSplineComponent);
+			TaggedData.Data = SplineData;
+		}
+		else if (USplineComponent* SplineComponent = Cast<USplineComponent>(InComponent))
+		{
+			UPCGSplineData* SplineData = FPCGContext::NewObject_AnyThread<UPCGSplineData>(InContext);
+			SplineData->Initialize(SplineComponent);
+			TaggedData.Data = SplineData;
+		}
+	}
+
+	if (!TaggedData.Data && !!(InDataFilter & EPCGDataType::Primitive))
+	{
+		if (UShapeComponent* ShapeComponent = Cast<UShapeComponent>(InComponent); ShapeComponent && UPCGCollisionShapeData::IsSupported(ShapeComponent))
+		{
+			UPCGCollisionShapeData* ShapeData = FPCGContext::NewObject_AnyThread<UPCGCollisionShapeData>(InContext);
+			ShapeData->Initialize(ShapeComponent);
+			TaggedData.Data = ShapeData;
+		}
+		else if (UBillboardComponent* BillboardComponent = Cast<UBillboardComponent>(InComponent))
+		{
+			// Exception: skip the billboard component
+		}
+		else if (UPrimitiveComponent* PrimitiveComponent = Cast<UPrimitiveComponent>(InComponent))
+		{
+			UPCGPrimitiveData* PrimitiveData = FPCGContext::NewObject_AnyThread<UPCGPrimitiveData>(InContext);
+			PrimitiveData->Initialize(PrimitiveComponent);
+			TaggedData.Data = PrimitiveData;
+		}
+	}
+
+	if (TaggedData.Data)
+	{
+		Algo::Transform(InComponent->ComponentTags, TaggedData.Tags, NameTagsToStringTags);
+
+		if (InParams.bAddActorTags && InComponent->GetOwner())
+		{
+			TSet<FString> ActorTags;
+			Algo::Transform(InComponent->GetOwner()->Tags, ActorTags, NameTagsToStringTags);
+			TaggedData.Tags.Append(ActorTags);
+		}
+
+		Output.Collection.TaggedData.Add(MoveTemp(TaggedData));
+		return 1;
+	}	
+	else
+	{
+		return 0;
+	}
+}
\ No newline at end of file
diff --git a/Engine/Plugins/PCG/Source/PCG/Private/Elements/PCGActorSelector.cpp b/Engine/Plugins/PCG/Source/PCG/Private/Elements/PCGActorSelector.cpp
--- a/Engine/Plugins/PCG/Source/PCG/Private/Elements/PCGActorSelector.cpp
+++ b/Engine/Plugins/PCG/Source/PCG/Private/Elements/PCGActorSelector.cpp
@@ -7,6 +7,7 @@
 #include "Grid/PCGPartitionActor.h"
 #include "Helpers/PCGActorHelpers.h"
 
+#include "Components/ActorComponent.h"
 #include "GameFramework/Actor.h"
 
 #include UE_INLINE_GENERATED_CPP_BY_NAME(PCGActorSelector)
@@ -203,6 +204,42 @@ namespace PCGActorSelector
 		return FoundActors;
 	}
 
+	TArray<AActor*> FilterActors(const FPCGComponentSelectorSettings& Settings, TArrayView<AActor*> InputActors)
+	{
+		TArray<AActor*> FilteredActors;
+		FilteredActors.Reserve(InputActors.Num());
+
+		for (AActor* InputActor : InputActors)
+		{
+			if (InputActor && Settings.FilterActor(InputActor))
+			{
+				FilteredActors.Add(InputActor);
+			}
+		}
+
+		return FilteredActors;
+	}
+
+	TArray<AActor*> FindActors(const FPCGActorSelectorSettings* ActorSettings, const FPCGComponentSelectorSettings* ComponentSettings, const UPCGComponent* InComponent, const TFunction<bool(const AActor*)>& BoundsCheck, const TFunction<bool(const AActor*)>& SelfIgnoreCheck, TArrayView<AActor*> InputActors)
+	{
+		TArray<AActor*> FoundActors;
+		if (ActorSettings)
+		{
+			FoundActors = FindActors(*ActorSettings, InComponent, BoundsCheck, SelfIgnoreCheck, InputActors);
+		}
+		else
+		{
+			FoundActors = InputActors;
+		}
+
+		if (!FoundActors.IsEmpty() && ComponentSettings)
+		{
+			FoundActors = FilterActors(*ComponentSettings, FoundActors);
+		}
+
+		return FoundActors;
+	}
+
 	AActor* FindActor(const FPCGActorSelectorSettings& InSettings, UPCGComponent* InComponent, const TFunction<bool(const AActor*)>& BoundsCheck, const TFunction<bool(const AActor*)>& SelfIgnoreCheck, TArrayView<AActor*> InputActors)
 	{
 		// In order to make sure we don't try to select multiple, we'll do a copy of the settings here.
@@ -529,4 +566,82 @@ FPCGActorSelectorSettings FPCGActorSelectorSettings::ReconstructFromKey(const FP
 	Result.ActorSelectionClass = InKey.SelectionClass;
 
 	return Result;
+}
+
+bool FPCGComponentSelectorSettings::FilterComponent(UActorComponent* InComponent) const
+{
+	check(InComponent);
+
+	if (!ComponentList.IsEmpty() && !ComponentList.Contains(InComponent))
+	{
+		return false;
+	}
+	else
+	{
+		return (ComponentSelection == EPCGComponentSelection::ByTag) ? (ComponentSelectionTag == NAME_None || InComponent->ComponentTags.Contains(ComponentSelectionTag)) :
+			(ComponentSelectionClass == nullptr || ComponentSelectionClass == UActorComponent::StaticClass() || InComponent->GetClass()->IsChildOf(ComponentSelectionClass));
+	}
+}
+
+bool FPCGComponentSelectorSettings::FilterActor(AActor* InActor) const
+{
+	check(InActor);
+	TInlineComponentArray<UActorComponent*> ActorComponents;
+	InActor->GetComponents(ActorComponents);
+
+	for (UActorComponent* ActorComponent : ActorComponents)
+	{
+		if (FilterComponent(ActorComponent))
+		{
+			return true;
+		}
+	}
+
+	return false;
+}
+
+TArray<UActorComponent*> FPCGComponentSelectorSettings::FilterComponents(TArrayView<UActorComponent*> InComponents) const
+{
+	TArrayView<UActorComponent*> Components;
+	TArray<UActorComponent*> InPreFilteredComponents;
+	if (!ComponentList.IsEmpty())
+	{
+		Algo::CopyIf(InComponents, InPreFilteredComponents, [this](UActorComponent* Component) { return ComponentList.Contains(Component); });
+		Components = TArrayView<UActorComponent*>(InPreFilteredComponents);
+	}
+	else
+	{
+		Components = InComponents;
+	}
+
+	TArray<UActorComponent*> FilteredComponents;
+	FilteredComponents.Reserve(Components.Num());
+
+	if ((ComponentSelection == EPCGComponentSelection::ByTag && ComponentSelectionTag == NAME_None) ||
+		(ComponentSelection == EPCGComponentSelection::ByClass && (ComponentSelectionClass == nullptr || ComponentSelectionClass == UActorComponent::StaticClass())))
+	{
+		FilteredComponents = Components;
+	}
+	else if (ComponentSelection == EPCGComponentSelection::ByTag)
+	{
+		for (UActorComponent* Component : Components)
+		{
+			if (Component && Component->ComponentTags.Contains(ComponentSelectionTag))
+			{
+				FilteredComponents.Add(Component);
+			}
+		}
+	}
+	else if (ComponentSelection == EPCGComponentSelection::ByClass)
+	{
+		for (UActorComponent* Component : Components)
+		{
+			if (Component && Component->GetClass()->IsChildOf(ComponentSelectionClass))
+			{
+				FilteredComponents.Add(Component);
+			}
+		}
+	}
+
+	return FilteredComponents;
 }
\ No newline at end of file
diff --git a/Engine/Plugins/PCG/Source/PCG/Private/Elements/PCGDataFromActor.cpp b/Engine/Plugins/PCG/Source/PCG/Private/Elements/PCGDataFromActor.cpp
--- a/Engine/Plugins/PCG/Source/PCG/Private/Elements/PCGDataFromActor.cpp
+++ b/Engine/Plugins/PCG/Source/PCG/Private/Elements/PCGDataFromActor.cpp
@@ -408,13 +408,29 @@ bool FPCGDataFromActorElement::ExecuteInternal(FPCGContext* InContext) const
 		{
 			TRACE_CPUPROFILER_EVENT_SCOPE(FPCGDataFromActorElement::Execute::FindActors);
 			TArray<AActor*> ActorsFromInput;
+			TArray<UActorComponent*> ComponentsFromInput;
 
 			if (Settings->ActorSelector.ActorFilter == EPCGActorFilter::FromInput)
 			{
-				Algo::Transform(Context->PathsToObjectsAndDataIndex, ActorsFromInput, [](const TTuple<FSoftObjectPath, int32, int32>& InPath) { return Cast<AActor>(InPath.Get<0>().ResolveObject()); });
+				for (const TTuple<FSoftObjectPath, int32, int32>& InPath : Context->PathsToObjectsAndDataIndex)
+				{
+					UObject* Object = InPath.Get<0>().ResolveObject();
+					if (AActor* Actor = Cast<AActor>(Object))
+					{
+						ActorsFromInput.AddUnique(Actor);
+					}
+					else if (UActorComponent* ActorComponent = Cast<UActorComponent>(Object))
+					{
+						ComponentsFromInput.Add(ActorComponent);
+						ActorsFromInput.AddUnique(ActorComponent->GetOwner());
+					}
+				}
 			}
 
-			Context->FoundActors = PCGActorSelector::FindActors(Settings->ActorSelector, Context->SourceComponent.Get(), BoundsCheck, SelfIgnoreCheck, ActorsFromInput);
+			Context->ComponentSelector = Settings->ComponentSelector;
+			Context->ComponentSelector.ComponentList = std::move(ComponentsFromInput);
+
+			Context->FoundActors = PCGActorSelector::FindActors(&Settings->ActorSelector, &Context->ComponentSelector, Context->SourceComponent.Get(), BoundsCheck, SelfIgnoreCheck, ActorsFromInput);
 			Context->bPerformedQuery = true;
 
 #if WITH_EDITOR
@@ -509,6 +525,22 @@ bool FPCGDataFromActorElement::ExecuteInternal(FPCGContext* InContext) const
 	return true;
 }
 
+bool ShouldIgnorePCGComponent(FPCGDataFromActorContext* Context, const UPCGDataFromActorSettings* Settings, UPCGComponent* Component, const UPCGComponent* SourceOriginalComponent)
+{
+	check(Context && Settings);
+	if (!Context->ComponentSelector.FilterComponent(Component))
+	{
+		return true;
+	}
+
+	const UPCGComponent* OriginalComponent = Component ? Component->GetOriginalComponent() : nullptr;
+	const AActor* SourceOriginalOwner = SourceOriginalComponent->GetOwner();
+
+	return !OriginalComponent
+		|| OriginalComponent == SourceOriginalComponent
+		|| (Settings->ActorSelector.bIgnoreSelfAndChildren && OriginalComponent->GetOwner() == SourceOriginalOwner);
+}
+
 void FPCGDataFromActorElement::GatherWaitTasks(AActor* FoundActor, FPCGContext* InContext, TArray<FPCGTaskId>& OutWaitTasks) const
 {
 	if (!FoundActor)
@@ -544,10 +576,8 @@ void FPCGDataFromActorElement::GatherWaitTasks(AActor* FoundActor, FPCGContext*
 
 	for (UPCGComponent* Component : PCGComponents)
 	{
-		const UPCGComponent* OriginalComponent = Component ? Component->GetOriginalComponent() : nullptr;
-
 		// Avoid waiting on our own execution (including local components).
-		if (!OriginalComponent || OriginalComponent == SourceOriginalComponent || (Settings->ActorSelector.bIgnoreSelfAndChildren && OriginalComponent->GetOwner() == SourceOwner))
+		if (ShouldIgnorePCGComponent(Context, Settings, Component, SourceOriginalComponent))
 		{
 			continue;
 		}
@@ -612,9 +642,10 @@ void FPCGDataFromActorElement::ProcessActors(FPCGContext* Context, const UPCGDat
 	}
 }
 
-void FPCGDataFromActorElement::CreateReferenceData(FPCGContext* Context, const UPCGDataFromActorSettings* Settings, const TArray<AActor*>& Actors) const
+void FPCGDataFromActorElement::CreateReferenceData(FPCGContext* InContext, const UPCGDataFromActorSettings* Settings, const TArray<AActor*>& Actors) const
 {
-	check(Context);
+	check(InContext);
+	FPCGDataFromActorContext* Context = static_cast<FPCGDataFromActorContext*>(InContext);
 	check(Settings && (Settings->Mode == EPCGGetDataFromActorMode::GetActorReference || Settings->Mode == EPCGGetDataFromActorMode::GetComponentsReference));
 
 	const bool bCreateActorReferences = (Settings->Mode == EPCGGetDataFromActorMode::GetActorReference || Settings->Mode == EPCGGetDataFromActorMode::GetComponentsReference);
@@ -636,28 +667,34 @@ void FPCGDataFromActorElement::CreateReferenceData(FPCGContext* Context, const U
 
 	for (AActor* Actor : Actors)
 	{
-		if (Actor)
+		if (!Actor)
 		{
-			if (ComponentReference)
-			{
-				TInlineComponentArray<UActorComponent*, 16> Components;
-				Actor->GetComponents(Components);
+			continue;
+		}
 
-				FSoftObjectPath ActorPath(Actor);
+		if (ComponentReference)
+		{
+			TInlineComponentArray<UActorComponent*, 16> Components;
+			Actor->GetComponents(Components);
 
-				for (UActorComponent* Component : Components)
-				{
-					PCGMetadataEntryKey Entry = ParamData->MutableMetadata()->AddEntry();
-					ActorReference->SetValue(Entry, ActorPath);
-					ComponentReference->SetValue(Entry, FSoftObjectPath(Component));
+			FSoftObjectPath ActorPath(Actor);
 
-				}
-			}
-			else if (ActorReference)
+			for (UActorComponent* Component : Components)
 			{
-				ActorReference->SetValue(ParamData->MutableMetadata()->AddEntry(), FSoftObjectPath(Actor));
+				if(!Context->ComponentSelector.FilterComponent(Component))
+				{
+					continue;
+				}
+
+				PCGMetadataEntryKey Entry = ParamData->MutableMetadata()->AddEntry();
+				ActorReference->SetValue(Entry, ActorPath);
+				ComponentReference->SetValue(Entry, FSoftObjectPath(Component));
 			}
 		}
+		else if (ActorReference)
+		{
+			ActorReference->SetValue(ParamData->MutableMetadata()->AddEntry(), FSoftObjectPath(Actor));
+		}
 	}
 
 	if (ParamData->MutableMetadata()->GetItemCountForChild() > 0)
@@ -703,9 +740,10 @@ void FPCGDataFromActorElement::MergeActorsIntoData(FPCGContext* Context, const U
 	}
 }
 
-void FPCGDataFromActorElement::ProcessActor(FPCGContext* Context, const UPCGDataFromActorSettings* Settings, AActor* FoundActor) const
+void FPCGDataFromActorElement::ProcessActor(FPCGContext* InContext, const UPCGDataFromActorSettings* Settings, AActor* FoundActor) const
 {
-	check(Context);
+	check(InContext);
+	FPCGDataFromActorContext* Context = static_cast<FPCGDataFromActorContext*>(InContext);
 	check(Settings);
 
 	UPCGComponent* SourceComponent = Context->SourceComponent.IsValid() ? Context->SourceComponent.Get() : nullptr;
@@ -732,14 +770,11 @@ void FPCGDataFromActorElement::ProcessActor(FPCGContext* Context, const UPCGData
 			Settings->bComponentsMustOverlapSelf,
 			Settings->bComponentsMustOverlapSelf ? SourceComponent->GetGridBounds() : FBox());
 
+		// Remove any PCG components that are filtered from the component selector &
 		// Remove any PCG components that don't belong to an external execution context (i.e. share the same original component), or that share a common root actor.
-		PCGComponents.RemoveAllSwap([Settings, SourceOwner, SourceOriginalComponent](UPCGComponent* Component)
+		PCGComponents.RemoveAllSwap([Context, Settings, SourceOriginalComponent](UPCGComponent* Component)
 		{
-			const UPCGComponent* OriginalComponent = Component ? Component->GetOriginalComponent() : nullptr;
-
-			return !OriginalComponent
-				|| OriginalComponent == SourceOriginalComponent
-				|| (Settings->ActorSelector.bIgnoreSelfAndChildren && OriginalComponent->GetOwner() == SourceOwner);
+			return ShouldIgnorePCGComponent(Context, Settings, Component, SourceOriginalComponent);
 		});
 
 		for (UPCGComponent* Component : PCGComponents)
@@ -760,7 +795,7 @@ void FPCGDataFromActorElement::ProcessActor(FPCGContext* Context, const UPCGData
 	{
 		if (!PCGComponents.IsEmpty())
 		{
-			PCGE_LOG(Log, GraphAndLog, FText::Format(LOCTEXT("ActorHasNoGeneratedData", "Actor '{0}' does not have any previously generated data"), FText::FromName(FoundActor->GetFName())));
+			PCGE_LOG(Log, GraphAndLog, FText::Format(LOCTEXT("ActorHasNoGeneratedData", "Actor '{0}' does not have any previously generated data, or all its components were filtered out."), FText::FromName(FoundActor->GetFName())));
 		}
 
 		return;
@@ -771,6 +806,25 @@ void FPCGDataFromActorElement::ProcessActor(FPCGContext* Context, const UPCGData
 		return;
 	}
 
+	auto GetActorDataCollection = [SourceComponent, Context, Settings](AActor* Actor, EPCGDataType DataType, bool bParseActor) -> FPCGDataCollection
+	{
+		FPCGGetDataFunctionRegistryParams GDFRParams;
+		GDFRParams.SourceComponent = SourceComponent;
+		GDFRParams.ComponentSelector = &Context->ComponentSelector;
+		GDFRParams.bParseActor = bParseActor;
+		GDFRParams.DataTypeFilter = DataType;
+
+		FPCGGetDataFunctionRegistryOutput GDFROutput;
+		FPCGModule::ConstGetDataFunctionRegistry().GetDataFromActor(Context, GDFRParams, Actor, GDFROutput);
+
+		if (GDFROutput.bSanitizedTagAttributeNames && !Settings->bSilenceSanitizedAttributeNameWarnings)
+		{
+			PCGE_LOG_C(Warning, GraphAndLog, Context, PCGDataFromActorConstants::TagNamesSanitizedWarning);
+		}
+
+		return GDFROutput.Collection;
+	};
+
 	TArray<FPCGTaggedData>& Outputs = Context->OutputData.TaggedData;
 
 	if (bHasGeneratedPCGData)
@@ -805,7 +859,6 @@ void FPCGDataFromActorElement::ProcessActor(FPCGContext* Context, const UPCGData
 		// Soft object pointer to UPCGData
 		// Array of pcg data -> all on the default pin
 		// Map of pcg data -> use key for pin? might not be robust
-		// FPCGDataCollection
 		if (const FStructProperty* StructProperty = CastField<FStructProperty>(FoundProperty))
 		{
 			if (StructProperty->Struct == FPCGDataCollection::StaticStruct())
@@ -825,31 +878,17 @@ void FPCGDataFromActorElement::ProcessActor(FPCGContext* Context, const UPCGData
 	else
 	{
 		const bool bParseActor = (Settings->Mode != EPCGGetDataFromActorMode::GetSinglePoint);
-		EPCGDataType DataFilter = Settings->GetDataFilter();
-		bool bAttributeNameWasSanitized = false;
-		FPCGDataCollection Collection = UPCGComponent::CreateActorPCGDataCollection(FoundActor, SourceComponent, DataFilter, bParseActor, &bAttributeNameWasSanitized);
-
-		if (bAttributeNameWasSanitized && !Settings->bSilenceSanitizedAttributeNameWarnings)
-		{
-			PCGE_LOG(Warning, GraphAndLog, PCGDataFromActorConstants::TagNamesSanitizedWarning);
-		}
-
-		Outputs += Collection.TaggedData;
+		FPCGDataCollection ActorDataCollection = GetActorDataCollection(FoundActor, Settings->GetDataFilter(), bParseActor);
+		Outputs += ActorDataCollection.TaggedData;
 	}
 
 	// Finally, if we're in a case where we need to output the single point data too, let's do it now.
 	if (Settings->bAlsoOutputSinglePointData && (Settings->Mode == EPCGGetDataFromActorMode::GetDataFromPCGComponent || Settings->Mode == EPCGGetDataFromActorMode::GetDataFromPCGComponentOrParseComponents))
 	{
 		const bool bParseActor = false;
-		bool bAttributeNameWasSanitized = false;
-		FPCGDataCollection Collection = UPCGComponent::CreateActorPCGDataCollection(FoundActor, SourceComponent, EPCGDataType::Any, bParseActor, &bAttributeNameWasSanitized);
-
-		if (bAttributeNameWasSanitized && !Settings->bSilenceSanitizedAttributeNameWarnings)
-		{
-			PCGE_LOG(Warning, GraphAndLog, PCGDataFromActorConstants::TagNamesSanitizedWarning);
-		}
+		FPCGDataCollection SinglePointActorDataCollection = GetActorDataCollection(FoundActor, EPCGDataType::Any, bParseActor);
 
-		for (const FPCGTaggedData& SinglePointData : Collection.TaggedData)
+		for (const FPCGTaggedData& SinglePointData : SinglePointActorDataCollection.TaggedData)
 		{
 			FPCGTaggedData& OutSinglePoint = Outputs.Add_GetRef(SinglePointData);
 			OutSinglePoint.Pin = PCGDataFromActorConstants::SinglePointPinLabel;
diff --git a/Engine/Plugins/PCG/Source/PCG/Private/PCGComponent.cpp b/Engine/Plugins/PCG/Source/PCG/Private/PCGComponent.cpp
--- a/Engine/Plugins/PCG/Source/PCG/Private/PCGComponent.cpp
+++ b/Engine/Plugins/PCG/Source/PCG/Private/PCGComponent.cpp
@@ -8,21 +8,12 @@
 #include "PCGGraph.h"
 #include "PCGInputOutputSettings.h"
 #include "PCGManagedResource.h"
-#include "PCGParamData.h"
 #include "PCGPin.h"
 #include "PCGSubsystem.h"
-#include "Data/PCGCollisionShapeData.h"
-#include "Data/PCGDifferenceData.h"
 #include "Data/PCGIntersectionData.h"
 #include "Data/PCGLandscapeData.h"
-#include "Data/PCGLandscapeSplineData.h"
-#include "Data/PCGPointData.h"
-#include "Data/PCGPrimitiveData.h"
-#include "Data/PCGProjectionData.h"
 #include "Data/PCGSpatialData.h"
-#include "Data/PCGSplineData.h"
 #include "Data/PCGUnionData.h"
-#include "Data/PCGVolumeData.h"
 #include "Graph/PCGStackContext.h"
 #include "Grid/PCGPartitionActor.h"
 #include "Helpers/PCGActorHelpers.h"
@@ -35,21 +26,13 @@
 #include "Utils/PCGGraphExecutionLogging.h"
 
 #include "CoreGlobals.h"
-#include "LandscapeComponent.h"
 #include "LandscapeProxy.h"
-#include "Algo/AllOf.h"
 #include "Algo/AnyOf.h"
 #include "Algo/Transform.h"
-#include "Components/BillboardComponent.h"
 #include "Components/InstancedStaticMeshComponent.h"
-#include "Components/ShapeComponent.h"
-#include "Components/SplineComponent.h"
-#include "Components/SplineMeshComponent.h"
 #include "Engine/Engine.h"
 #include "Engine/Level.h"
-#include "GameFramework/Volume.h"
 #include "Kismet/GameplayStatics.h"
-#include "LandscapeSplinesComponent.h"
 #include "UObject/Package.h"
 
 #include UE_INLINE_GENERATED_CPP_BY_NAME(PCGComponent)
@@ -2862,245 +2845,15 @@ FPCGDataCollection UPCGComponent::CreateActorPCGDataCollection(AActor* Actor, co
 {
 	TRACE_CPUPROFILER_EVENT_SCOPE(UPCGComponent::CreateActorPCGData);
 
-	if (bOutOptionalSanitizedTagAttributeName)
-	{
-		*bOutOptionalSanitizedTagAttributeName = false;
-	}
-
-	FPCGDataCollection Collection;
-
-	if (!Actor)
-	{
-		return Collection;
-	}
-
-	auto NameTagsToStringTags = [](const FName& InName) { return InName.ToString(); };
-	TSet<FString> ActorTags;
-	Algo::Transform(Actor->Tags, ActorTags, NameTagsToStringTags);
-
-	// Fill in collection based on the data on the given actor.
-	// Some actor types we will forego full parsing to build strictly on the actor existence, such as partition actors, volumes and landscape
-	// TODO: add factory for extensibility
-	// TODO: review the !bParseActor cases - it might make sense to have just a point for a partition actor, even if we preintersect it.
-	APCGPartitionActor* PartitionActor = (!!(InDataFilter & EPCGDataType::Spatial) || !!(InDataFilter & EPCGDataType::Volume)) ? Cast<APCGPartitionActor>(Actor) : nullptr;
-	ALandscapeProxy* LandscapeActor = !!(InDataFilter & EPCGDataType::Landscape) ? Cast<ALandscapeProxy>(Actor) : nullptr;
-	AVolume* VolumeActor = !!(InDataFilter & EPCGDataType::Volume) ? Cast<AVolume>(Actor) : nullptr;
-	if (!bParseActor && !!(InDataFilter & EPCGDataType::Point))
-	{
-		UPCGPointData* Data = NewObject<UPCGPointData>();
-		Data->InitializeFromActor(Actor, bOutOptionalSanitizedTagAttributeName);
-
-		FPCGTaggedData& TaggedData = Collection.TaggedData.Emplace_GetRef();
-		TaggedData.Data = Data;
-		TaggedData.Tags = ActorTags;
-	}
-	else if (!bParseActor && !!(InDataFilter & EPCGDataType::Param))
-	{
-		UPCGParamData* Data = NewObject<UPCGParamData>();
-		check(Data && Data->MutableMetadata());
-		
-		if (FPCGMetadataAttribute<FSoftObjectPath>* ActorReferenceAttribute = Data->MutableMetadata()->FindOrCreateAttribute(PCGPointDataConstants::ActorReferenceAttribute, FSoftObjectPath(), /*bAllowsInterpolation=*/false, /*bOverrideParent=*/false, /*bOverwriteIfTypeMismatch=*/true))
-		{
-			ActorReferenceAttribute->SetValue(Data->MutableMetadata()->AddEntry(), FSoftObjectPath(Actor));
-		}
-
-		FPCGTaggedData& TaggedData = Collection.TaggedData.Emplace_GetRef();
-		TaggedData.Data = Data;
-		TaggedData.Tags = ActorTags;
-	}
-	else if (PartitionActor)
-	{
-		check(!Component || Component->GetOwner() == Actor); // Invalid processing otherwise because of the this usage
-
-		UPCGVolumeData* VolumeData = NewObject<UPCGVolumeData>();
-		UPCGSpatialData* Result = VolumeData;
-		if (InDataFilter == EPCGDataType::Volume)
-		{
-			VolumeData->Initialize(PCGHelpers::GetGridBounds(Actor, Component));
-		}
-		else
-		{
-			VolumeData->Initialize(PartitionActor->GetFixedBounds());
-
-			UPCGComponent* OriginalComponent = Component ? PartitionActor->GetOriginalComponent(Component) : nullptr;
-			// Important note: we do NOT call the collection version here, as we want to have a union if that's the case
-			const UPCGSpatialData* OriginalComponentSpatialData = OriginalComponent ? Cast<const UPCGSpatialData>(OriginalComponent->GetActorPCGData()) : nullptr;
-
-			if (OriginalComponentSpatialData)
-			{
-				Result = Result->IntersectWith(nullptr, OriginalComponentSpatialData);
-			}
-		}
-
-		FPCGTaggedData& TaggedData = Collection.TaggedData.Emplace_GetRef();
-		TaggedData.Data = Result;
-		// No need to keep partition actor tags, though we might want to push PCG grid GUID at some point
-	}
-	else if (LandscapeActor)
-	{
-		UPCGLandscapeData* Data = NewObject<UPCGLandscapeData>();
-		const UPCGGraph* PCGGraph = Component ? Component->GetGraph() : nullptr;
-
-		FPCGLandscapeDataProps LandscapeDataProps;
-		LandscapeDataProps.bGetHeightOnly = false;
-		LandscapeDataProps.bGetLayerWeights = (!PCGGraph || PCGGraph->bLandscapeUsesMetadata);
-
-		Data->Initialize({ LandscapeActor }, PCGHelpers::GetGridBounds(Actor, Component), LandscapeDataProps);
-
-		FPCGTaggedData& TaggedData = Collection.TaggedData.Emplace_GetRef();
-		TaggedData.Data = Data;
-		TaggedData.Tags = ActorTags;
-	}
-	else if (VolumeActor)
-	{
-		UPCGVolumeData* Data = NewObject<UPCGVolumeData>();
-		Data->Initialize(VolumeActor);
-
-		FPCGTaggedData& TaggedData = Collection.TaggedData.Emplace_GetRef();
-		TaggedData.Data = Data;
-		TaggedData.Tags = ActorTags;
-	}
-	else // Prepare data on a component basis
-	{
-		using PrimitiveComponentArray = TInlineComponentArray<UPrimitiveComponent*, 4>;
-		PrimitiveComponentArray Primitives;
-
-		auto RemoveDuplicatesFromPrimitives = [&Primitives](const auto& InComponents)
-		{
-			Primitives.RemoveAll([&InComponents](UPrimitiveComponent* Component)
-			{
-				return InComponents.Contains(Component);
-			});
-		};
-
-		auto RemovePCGGeneratedEntries = [](auto& InComponents)
-		{
-			for (int32 Index = InComponents.Num() - 1; Index >= 0; --Index)
-			{
-				if (InComponents[Index]->ComponentTags.Contains(PCGHelpers::DefaultPCGTag))
-				{
-					InComponents.RemoveAtSwap(Index);
-				}
-			}
-		};
-
-		auto RemoveSplineMeshComponents = [](PrimitiveComponentArray& InComponents)
-		{
-			for (int32 Index = InComponents.Num() - 1; Index >= 0; --Index)
-			{
-				if (InComponents[Index]->IsA<USplineMeshComponent>())
-				{
-					InComponents.RemoveAtSwap(Index);
-				}
-			}
-		};
-
-		Actor->GetComponents(Primitives);
-		RemovePCGGeneratedEntries(Primitives);
-
-		TInlineComponentArray<ULandscapeSplinesComponent*, 4> LandscapeSplines;
-		Actor->GetComponents(LandscapeSplines);
-		RemovePCGGeneratedEntries(LandscapeSplines);
-		RemoveDuplicatesFromPrimitives(LandscapeSplines);
-
-		TInlineComponentArray<USplineComponent*, 4> Splines;
-		Actor->GetComponents(Splines);
-		RemovePCGGeneratedEntries(Splines);
-		RemoveDuplicatesFromPrimitives(Splines);
-
-		// If we have a better representation than the spline mesh components, we shouldn't create them
-		if (!LandscapeSplines.IsEmpty() || !Splines.IsEmpty())
-		{
-			RemoveSplineMeshComponents(Primitives);
-		}
-
-		TInlineComponentArray<UShapeComponent*, 4> Shapes;
-		Actor->GetComponents(Shapes);
-		RemovePCGGeneratedEntries(Shapes);
-		RemoveDuplicatesFromPrimitives(Shapes);
-
-		if (!!(InDataFilter & EPCGDataType::Spline))
-		{
-			for (ULandscapeSplinesComponent* SplineComponent : LandscapeSplines)
-			{
-				UPCGLandscapeSplineData* SplineData = NewObject<UPCGLandscapeSplineData>();
-				SplineData->Initialize(SplineComponent);
-
-				FPCGTaggedData& TaggedData = Collection.TaggedData.Emplace_GetRef();
-				TaggedData.Data = SplineData;
-				Algo::Transform(SplineComponent->ComponentTags, TaggedData.Tags, NameTagsToStringTags);
-				TaggedData.Tags.Append(ActorTags);
-			}
-
-			for (USplineComponent* SplineComponent : Splines)
-			{
-				UPCGSplineData* SplineData = NewObject<UPCGSplineData>();
-				SplineData->Initialize(SplineComponent);
-
-				FPCGTaggedData& TaggedData = Collection.TaggedData.Emplace_GetRef();
-				TaggedData.Data = SplineData;
-				Algo::Transform(SplineComponent->ComponentTags, TaggedData.Tags, NameTagsToStringTags);
-				TaggedData.Tags.Append(ActorTags);
-			}
-		}
+	FPCGGetDataFunctionRegistryParams Params;
+	Params.SourceComponent = Component;
+	Params.bParseActor = bParseActor;
+	Params.DataTypeFilter = InDataFilter;
 
-		if (!!(InDataFilter & EPCGDataType::Primitive))
-		{
-			for (UShapeComponent* ShapeComponent : Shapes)
-			{
-				UPCGSpatialData* Data = nullptr;
-				if (UPCGCollisionShapeData::IsSupported(ShapeComponent))
-				{
-					UPCGCollisionShapeData* ShapeData = NewObject<UPCGCollisionShapeData>();
-					ShapeData->Initialize(ShapeComponent);
-
-					Data = ShapeData;
-				}
-				else
-				{
-					UPCGPrimitiveData* ShapeData = NewObject<UPCGPrimitiveData>();
-					ShapeData->Initialize(ShapeComponent);
-
-					Data = ShapeData;
-				}
-
-				FPCGTaggedData& TaggedData = Collection.TaggedData.Emplace_GetRef();
-				TaggedData.Data = Data;
-				Algo::Transform(ShapeComponent->ComponentTags, TaggedData.Tags, NameTagsToStringTags);
-				TaggedData.Tags.Append(ActorTags);
-			}
-
-			for (UPrimitiveComponent* PrimitiveComponent : Primitives)
-			{
-				// Exception: skip the billboard component
-				if (Cast<UBillboardComponent>(PrimitiveComponent))
-				{
-					continue;
-				}
-
-				UPCGPrimitiveData* PrimitiveData = NewObject<UPCGPrimitiveData>();
-				PrimitiveData->Initialize(PrimitiveComponent);
-
-				FPCGTaggedData& TaggedData = Collection.TaggedData.Emplace_GetRef();
-				TaggedData.Data = PrimitiveData;
-				Algo::Transform(PrimitiveComponent->ComponentTags, TaggedData.Tags, NameTagsToStringTags);
-				TaggedData.Tags.Append(ActorTags);
-			}
-		}
-	}
-
-	// Finally, if it's not a special actor and there are not parsed components, then return a single point at the actor position
-	if (Collection.TaggedData.IsEmpty() && !!(InDataFilter & EPCGDataType::Point))
-	{
-		UPCGPointData* Data = NewObject<UPCGPointData>();
-		Data->InitializeFromActor(Actor, bOutOptionalSanitizedTagAttributeName);
-
-		FPCGTaggedData& TaggedData = Collection.TaggedData.Emplace_GetRef();
-		TaggedData.Data = Data;
-		TaggedData.Tags = ActorTags;
-	}
+	FPCGGetDataFunctionRegistryOutput Output;
+	FPCGModule::ConstGetDataFunctionRegistry().GetDataFromActor(/*Context=*/nullptr, Params, Actor, Output);
 
-	return Collection;
+	return Output.Collection;
 }
 
 void UPCGComponent::RefreshSchedulingPolicy()
diff --git a/Engine/Plugins/PCG/Source/PCG/Public/Data/PCGGetDataFunctionRegistry.h b/Engine/Plugins/PCG/Source/PCG/Public/Data/PCGGetDataFunctionRegistry.h
new file mode 100644
--- /dev/null
+++ b/Engine/Plugins/PCG/Source/PCG/Public/Data/PCGGetDataFunctionRegistry.h
@@ -0,0 +1,72 @@
+// Copyright Epic Games, Inc. All Rights Reserved.
+
+#pragma once
+
+#include "PCGCommon.h"
+#include "PCGData.h"
+
+class AActor;
+class UActorComponent;
+class UPCGComponent;
+struct FPCGContext;
+struct FPCGComponentSelectorSettings;
+
+/** Struct to hold selection criteria */
+struct FPCGGetDataFunctionRegistryParams
+{
+	const UPCGComponent* SourceComponent = nullptr;
+	const FPCGComponentSelectorSettings* ComponentSelector = nullptr;
+	EPCGDataType DataTypeFilter = EPCGDataType::Any;
+	bool bParseActor = true;
+	bool bIgnorePCGGeneratedComponents = true;
+	bool bAddActorTags = true;
+};
+
+struct FPCGGetDataFunctionRegistryOutput
+{
+	FPCGDataCollection Collection;
+	bool bSanitizedTagAttributeNames = false;
+};
+
+/** Registry to hold actor & component to PCG data construction mapping */
+struct PCG_API FPCGGetDataFunctionRegistry
+{
+public:
+	using FDataFromActorFunction = TFunction<bool(FPCGContext*, const FPCGGetDataFunctionRegistryParams&, AActor*, FPCGGetDataFunctionRegistryOutput&)>;
+	using FDataFromComponentFunction = TFunction<bool(FPCGContext*, const FPCGGetDataFunctionRegistryParams&, UActorComponent*, FPCGGetDataFunctionRegistryOutput&)>;
+	using FFunctionHandle = uint64;
+
+	FPCGGetDataFunctionRegistry() = default;
+	FPCGGetDataFunctionRegistry(const FPCGGetDataFunctionRegistry&) = delete;
+	FPCGGetDataFunctionRegistry(FPCGGetDataFunctionRegistry&&) = default;
+	FPCGGetDataFunctionRegistry& operator=(const FPCGGetDataFunctionRegistry&) = delete;
+	FPCGGetDataFunctionRegistry& operator=(FPCGGetDataFunctionRegistry&&) = default;
+	virtual ~FPCGGetDataFunctionRegistry() = default;
+
+	/** Register an actor to PCG data function. */
+	FFunctionHandle RegisterDataFromActorFunction(const FDataFromActorFunction& InFunction);
+
+	/** Unregisters an actor to PCG data function. */
+	void UnregisterDataFromActorFunction(FFunctionHandle InFunctionHandle);
+
+	/** Appends actor data to the given collection and returns the number of data added. */
+	int GetDataFromActor(FPCGContext* InContext, const FPCGGetDataFunctionRegistryParams& InParams, AActor* InActor, FPCGGetDataFunctionRegistryOutput& OutDataCollection) const;
+
+	/** Register a component to PCG data function. */
+	FFunctionHandle RegisterDataFromComponentFunction(const FDataFromComponentFunction& InFunction);
+
+	/** Unregisters a component to PCG data function. */
+	void UnregisterDataFromComponentFunction(FFunctionHandle InFunctionHandle);
+
+	/** Append component data to the given collection, returns the number of data added. */
+	int GetDataFromComponent(FPCGContext* InContext, const FPCGGetDataFunctionRegistryParams& InParams, UActorComponent* InComponent, FPCGGetDataFunctionRegistryOutput& OutDataCollection) const;
+
+private:
+	int DefaultDataFromActor(FPCGContext* InContext, const FPCGGetDataFunctionRegistryParams& InParams, AActor* InActor, FPCGGetDataFunctionRegistryOutput& OutDataCollection) const;
+	int DefaultDataFromComponent(FPCGContext* InContext, const FPCGGetDataFunctionRegistryParams& InParams, UActorComponent* InComponent, FPCGGetDataFunctionRegistryOutput& OutDataCollection) const;
+
+	FFunctionHandle NextFunctionHandle = 0;
+
+	TArray<TPair<FDataFromActorFunction, FFunctionHandle>> ActorParsingFunctions;
+	TArray<TPair<FDataFromComponentFunction, FFunctionHandle>> ComponentParsingFunctions;
+};
\ No newline at end of file
diff --git a/Engine/Plugins/PCG/Source/PCG/Public/Elements/PCGActorSelector.h b/Engine/Plugins/PCG/Source/PCG/Public/Elements/PCGActorSelector.h
--- a/Engine/Plugins/PCG/Source/PCG/Public/Elements/PCGActorSelector.h
+++ b/Engine/Plugins/PCG/Source/PCG/Public/Elements/PCGActorSelector.h
@@ -10,6 +10,7 @@
 #include "PCGActorSelector.generated.h"
 
 class AActor;
+class UActorComponent;
 class UPCGComponent;
 class UWorld;
 struct FPCGActorSelectorSettings;
@@ -25,6 +26,14 @@ enum class EPCGActorSelection : uint8
 	Unknown UMETA(Hidden)
 };
 
+UENUM()
+enum class EPCGComponentSelection : uint8
+{
+	ByTag,
+	ByClass,
+	Unknown UMETA(Hidden)
+};
+
 UENUM()
 enum class EPCGActorFilter : uint8
 {
@@ -171,8 +180,46 @@ struct PCG_API FPCGActorSelectorSettings
 	static FPCGActorSelectorSettings ReconstructFromKey(const FPCGSelectionKey& InKey);
 };
 
+USTRUCT(BlueprintType)
+struct PCG_API FPCGComponentSelectorSettings
+{
+	GENERATED_BODY()
+
+	/** How to select when filtering actors. */
+	UPROPERTY(BlueprintReadWrite, EditAnywhere, Category = "Component Selector Settings", meta = (EditCondition = "bShowComponentSelection", EditConditionHides))
+	EPCGComponentSelection ComponentSelection = EPCGComponentSelection::ByTag;
+
+	/** Tag to match against when filtering actors. */
+	UPROPERTY(BlueprintReadWrite, EditAnywhere, Category = "Component Selector Settings", meta = (EditCondition = "bShowComponentSelection && ComponentSelection==EPCGComponentSelection::ByTag", EditConditionHides))
+	FName ComponentSelectionTag;
+
+	/** Actor class to match against when filtering actors. */
+	UPROPERTY(BlueprintReadWrite, EditAnywhere, Category = "Component Selector Settings", meta = (EditCondition = "bShowComponentSelection && bShowComponentSelectionClass && ComponentSelection==EPCGComponentSelection::ByClass", EditConditionHides, AllowAbstract = "true"))
+	TSubclassOf<UActorComponent> ComponentSelectionClass;
+
+	UPROPERTY(Transient)
+	bool bShowComponentSelection = true;
+
+	UPROPERTY(Transient)
+	bool bShowComponentSelectionClass = true;
+
+	TArray<UActorComponent*> ComponentList;
+
+	bool FilterComponent(UActorComponent* InComponent) const;
+	bool FilterActor(AActor* InActor) const;
+	TArray<UActorComponent*> FilterComponents(TArrayView<UActorComponent*> InComponents) const;
+};
+
 namespace PCGActorSelector
 {
+	// Simple actor filtering
 	PCG_API TArray<AActor*> FindActors(const FPCGActorSelectorSettings& Settings, const UPCGComponent* InComponent, const TFunction<bool(const AActor*)>& BoundsCheck, const TFunction<bool(const AActor*)>& SelfIgnoreCheck, TArrayView<AActor*> InputActors = TArrayView<AActor*>());
 	PCG_API AActor* FindActor(const FPCGActorSelectorSettings& InSettings, UPCGComponent* InComponent, const TFunction<bool(const AActor*)>& BoundsCheck, const TFunction<bool(const AActor*)>& SelfIgnoreCheck, TArrayView<AActor*> InputActors = TArrayView<AActor*>());
+
+	// Actor + Component filtering
+	PCG_API TArray<AActor*> FindActors(const FPCGActorSelectorSettings* ActorSettings, const FPCGComponentSelectorSettings* ComponentSettings, const UPCGComponent* InComponent, const TFunction<bool(const AActor*)>& BoundsCheck, const TFunction<bool(const AActor*)>& SelfIgnoreCheck, TArrayView<AActor*> InputActors = TArrayView<AActor*>());
+	PCG_API AActor* FindActor(const FPCGActorSelectorSettings* ActorSettings, const FPCGComponentSelectorSettings* ComponentSettings, const UPCGComponent* InComponent, const TFunction<bool(const AActor*)>& BoundsCheck, const TFunction<bool(const AActor*)>& SelfIgnoreCheck, TArrayView<AActor*> InputActors = TArrayView<AActor*>());
+	
+	// Additional filter methods
+	PCG_API TArray<AActor*> FilterActors(const FPCGComponentSelectorSettings& ComponentSettings, TArrayView<AActor*> InActors);
 }
diff --git a/Engine/Plugins/PCG/Source/PCG/Public/Elements/PCGDataFromActor.h b/Engine/Plugins/PCG/Source/PCG/Public/Elements/PCGDataFromActor.h
--- a/Engine/Plugins/PCG/Source/PCG/Public/Elements/PCGDataFromActor.h
+++ b/Engine/Plugins/PCG/Source/PCG/Public/Elements/PCGDataFromActor.h
@@ -81,6 +81,10 @@ public:
 	UPROPERTY(BlueprintReadWrite, EditAnywhere, Category = Settings, meta = (PCG_Overridable, ShowOnlyInnerProperties))
 	FPCGActorSelectorSettings ActorSelector;
 
+	/** Describes which components to select for the data collection. */
+	UPROPERTY(BlueprintReadWrite, EditAnywhere, Category = Settings, meta = (PCG_Overridable, ShowOnlyInnerProperties))
+	FPCGComponentSelectorSettings ComponentSelector;
+
 	/** Describes what kind of data we will collect from the found actor(s). */
 	UPROPERTY(BlueprintReadWrite, EditAnywhere, Category = "Data Retrieval Settings", meta = (EditCondition = "DisplayModeSettings()", EditConditionHides, HideEditConditionToggle))
 	EPCGGetDataFromActorMode Mode = EPCGGetDataFromActorMode::ParseActorComponents;
@@ -131,6 +135,7 @@ public:
 struct FPCGDataFromActorContext : public FPCGLoadObjectsFromPathContext
 {
 	TArray<AActor*> FoundActors;
+	FPCGComponentSelectorSettings ComponentSelector;
 	bool bPerformedQuery = false;
 
 #if WITH_EDITOR
diff --git a/Engine/Plugins/PCG/Source/PCG/Public/PCGModule.h b/Engine/Plugins/PCG/Source/PCG/Public/PCGModule.h
--- a/Engine/Plugins/PCG/Source/PCG/Public/PCGModule.h
+++ b/Engine/Plugins/PCG/Source/PCG/Public/PCGModule.h
@@ -3,6 +3,7 @@
 #pragma once
 
 #include "PCGDataVisualizationRegistry.h"
+#include "Data/PCGGetDataFunctionRegistry.h"
 
 #include "Modules/ModuleInterface.h"
 #include "Stats/Stats.h"
@@ -42,9 +43,14 @@ public:
 	virtual bool SupportsDynamicReloading() override { return true; }
 	//~ End IModuleInterface implementation
 
-#if WITH_EDITOR
 	PCG_API static FPCGModule& GetPCGModuleChecked();
+	static const FPCGGetDataFunctionRegistry& ConstGetDataFunctionRegistry() { return GetPCGModuleChecked().GetDataFunctionRegistry; }
+	static FPCGGetDataFunctionRegistry& MutableGetDataFunctionRegistry() { return GetPCGModuleChecked().GetDataFunctionRegistry; }
+
+private:
+	FPCGGetDataFunctionRegistry GetDataFunctionRegistry;
 
+#if WITH_EDITOR
 private:
 	void RegisterNativeElementDeterminismTests();
 	void DeregisterNativeElementDeterminismTests();
