[PCG] Fix PCG Partition Actor Undo / Recreate name clash crash

#rnx

diff --git a/Engine/Plugins/PCG/Source/PCG/Private/Helpers/PCGActorHelpers.cpp b/Engine/Plugins/PCG/Source/PCG/Private/Helpers/PCGActorHelpers.cpp
--- a/Engine/Plugins/PCG/Source/PCG/Private/Helpers/PCGActorHelpers.cpp
+++ b/Engine/Plugins/PCG/Source/PCG/Private/Helpers/PCGActorHelpers.cpp
@@ -358,7 +358,8 @@ bool UPCGActorHelpers::DeleteActors(UWorld* World, const TArray<TSoftObjectPtr<A
 		// the next GC to collect those packages
 		//
 		// The fix here is to create a dummy transaction so that the deleted actors are tracked but since we don't want to actually push a transaction, we cancel it after the DestroyActor calls
-		FScopedTransaction DummyTransaction(NSLOCTEXT("PCGActorHelpers", "DummyTransaction", "DummyTransaction"), World && !World->IsGameWorld() && ActorsToDelete.Num() > 0 && !GEditor->IsTransactionActive());
+		const bool bShouldTransact = World && !World->IsGameWorld() && ActorsToDelete.Num() > 0 && !GEditor->IsTransactionActive() && !GIsTransacting;
+		FScopedTransaction DummyTransaction(NSLOCTEXT("PCGActorHelpers", "DummyTransaction", "DummyTransaction"), bShouldTransact);
 #endif
 
 		// Not in editor, really unlikely to happen but might be slow
diff --git a/Engine/Plugins/PCG/Source/PCG/Private/PCGSubsystem.cpp b/Engine/Plugins/PCG/Source/PCG/Private/PCGSubsystem.cpp
--- a/Engine/Plugins/PCG/Source/PCG/Private/PCGSubsystem.cpp
+++ b/Engine/Plugins/PCG/Source/PCG/Private/PCGSubsystem.cpp
@@ -1200,6 +1200,14 @@ APCGPartitionActor* UPCGSubsystem::FindOrCreatePCGPartitionActor(const FPCGGridD
 	{
 		return nullptr;
 	}
+
+#if WITH_EDITOR
+	// Do not try and create while executing a Undo/Redo because actor might already be in the process of being re-created by the transaction
+	if (GIsTransacting)
+	{
+		return nullptr;
+	}
+#endif
 	
 	if (!bCanCreateActor)
 	{
@@ -1231,6 +1239,25 @@ APCGPartitionActor* UPCGSubsystem::FindOrCreatePCGPartitionActor(const FPCGGridD
 
 	// Specify EDL we want to use if any for spawning this actor
 	FScopedOverrideSpawningLevelMountPointObject EDLScope(ExternalDataLayerAsset);
+
+	// Handle the case where the actor already exists, but is in the undo stack (was deleted) (copied from ActorPartitionSubsystem, we should probably merge back into it at some point) 
+	if (SpawnParams.NameMode == FActorSpawnParameters::ESpawnActorNameMode::Required_Fatal)
+	{
+		if (UObject* ExistingObject = StaticFindObject(nullptr, World->PersistentLevel, *SpawnParams.Name.ToString()))
+		{
+			AActor* ExistingActor = CastChecked<AActor>(ExistingObject);
+			// This actor is expected to be invalid
+			check(!IsValidChecked(ExistingActor));
+			ExistingActor->Modify();
+
+			// Don't go through AActor::Rename here because we aren't changing outers (the actor's level). We just want to rename that actor 
+			// out of the way so we can spawn the new one in the exact same package, keeping the package name intact.
+			ExistingActor->UObject::Rename(nullptr, nullptr, REN_DontCreateRedirectors | REN_DoNotDirty | REN_NonTransactional);
+
+			// Reuse ActorGuid so that ActorDesc can be updated on save
+			SpawnParams.OverrideActorGuid = ExistingActor->GetActorGuid();
+		}
+	}
 #endif
 
 	const FVector CellCenter(FVector(GridCoords.X + 0.5, GridCoords.Y + 0.5, GridCoords.Z + 0.5) * GridDescriptor.GetGridSize());
