[PCG] Make the module info as input accepting attribute names

* Also fixing internal graphs

#rnx

#virtualized

diff --git a/Engine/Plugins/PCG/Source/PCG/Private/Elements/Grammar/PCGSlicingBase.cpp b/Engine/Plugins/PCG/Source/PCG/Private/Elements/Grammar/PCGSlicingBase.cpp
--- a/Engine/Plugins/PCG/Source/PCG/Private/Elements/Grammar/PCGSlicingBase.cpp
+++ b/Engine/Plugins/PCG/Source/PCG/Private/Elements/Grammar/PCGSlicingBase.cpp
@@ -42,27 +42,15 @@ FPCGSlicingBaseElement::FPCGModulesInfoMap FPCGSlicingBaseElement::GetModulesInf
 			return ModulesInfo;
 		}
 
-		static const FName SymbolAttributeName = PCGSlicingBaseConstants::SymbolAttributeName;
-		static const FName SizeAttributeName = GET_MEMBER_NAME_CHECKED(FPCGSlicingModule, Size);
-		static const FName ScalableAttributeName = GET_MEMBER_NAME_CHECKED(FPCGSlicingModule, bScalable);
-		static const FName DebugColorAttributeName = GET_MEMBER_NAME_CHECKED(FPCGSlicingModule, DebugColor);
-		static const FName ExtraScalableAttributeName = TEXT("Scalable"); // To also support the name without a "b".
+		const TUniquePtr<const IPCGAttributeAccessorKeys> Keys = PCGAttributeAccessorHelpers::CreateConstKeys(ParamData, FPCGAttributePropertySelector::CreateAttributeSelector(InSettings->ModulesInfoAttributeNames.SymbolAttributeName));
 
-		const TUniquePtr<const IPCGAttributeAccessorKeys> Keys = PCGAttributeAccessorHelpers::CreateConstKeys(ParamData, FPCGAttributePropertySelector::CreateAttributeSelector(SymbolAttributeName));
-
-		const TUniquePtr<const IPCGAttributeAccessor> SymbolAccessor = PCGAttributeAccessorHelpers::CreateConstAccessor(ParamData, FPCGAttributePropertySelector::CreateAttributeSelector(SymbolAttributeName));
-		const TUniquePtr<const IPCGAttributeAccessor> SizeAccessor = PCGAttributeAccessorHelpers::CreateConstAccessor(ParamData, FPCGAttributePropertySelector::CreateAttributeSelector(SizeAttributeName));
-		TUniquePtr<const IPCGAttributeAccessor> ScalableAccessor = PCGAttributeAccessorHelpers::CreateConstAccessor(ParamData, FPCGAttributePropertySelector::CreateAttributeSelector(ScalableAttributeName));
-		const TUniquePtr<const IPCGAttributeAccessor> DebugColorAccessor = PCGAttributeAccessorHelpers::CreateConstAccessor(ParamData, FPCGAttributePropertySelector::CreateAttributeSelector(DebugColorAttributeName));
-
-		if (!ScalableAccessor.IsValid())
+		auto GetRange = [InContext, &Keys, ParamData]<typename T>(TArray<T>& OutValues, const FName AttributeName, const bool bProvided = true)
 		{
-			// Retry with Scalable without a b.
-			ScalableAccessor = PCGAttributeAccessorHelpers::CreateConstAccessor(ParamData, FPCGAttributePropertySelector::CreateAttributeSelector(ExtraScalableAttributeName));
-		}
+			if (!bProvided)
+			{
+				return true;
+			}
 
-		auto GetRange = [InContext, &Keys, ParamData]<typename T>(TArray<T>& OutValues, const FName AttributeName)
-		{
 			const TUniquePtr<const IPCGAttributeAccessor> Accessor = PCGAttributeAccessorHelpers::CreateConstAccessor(ParamData, FPCGAttributePropertySelector::CreateAttributeSelector(AttributeName));
 			if (!Accessor || !Keys)
 			{
@@ -86,10 +74,10 @@ FPCGSlicingBaseElement::FPCGModulesInfoMap FPCGSlicingBaseElement::GetModulesInf
 		TArray<bool> Scalables;
 		TArray<FVector4> DebugColors;
 
-		if (!GetRange(Symbols, SymbolAttributeName)
-			|| !GetRange(Sizes, SizeAttributeName)
-			|| !GetRange(Scalables, ScalableAttributeName)
-			|| !GetRange(DebugColors, DebugColorAttributeName))
+		if (!GetRange(Symbols, InSettings->ModulesInfoAttributeNames.SymbolAttributeName)
+			|| !GetRange(Sizes, InSettings->ModulesInfoAttributeNames.SizeAttributeName)
+			|| !GetRange(Scalables, InSettings->ModulesInfoAttributeNames.ScalableAttributeName, InSettings->ModulesInfoAttributeNames.bProvideScalable)
+			|| !GetRange(DebugColors, InSettings->ModulesInfoAttributeNames.DebugColorAttributeName, InSettings->ModulesInfoAttributeNames.bProvideDebugColor))
 		{
 			return ModulesInfo;
 		}
@@ -108,8 +96,9 @@ FPCGSlicingBaseElement::FPCGModulesInfoMap FPCGSlicingBaseElement::GetModulesInf
 			FPCGSlicingModule& Module = ModulesInfo.Emplace(Symbols[i]);
 			Module.Symbol = Symbols[i];
 			Module.Size = Sizes[i];
-			Module.bScalable = Scalables[i];
-			Module.DebugColor = DebugColors[i];
+
+			Module.bScalable = InSettings->ModulesInfoAttributeNames.bProvideScalable ? Scalables[i] : false;
+			Module.DebugColor = InSettings->ModulesInfoAttributeNames.bProvideDebugColor ? DebugColors[i] : FVector4::One();
 		}
 
 		OutModuleInfoParamData = ParamData;
diff --git a/Engine/Plugins/PCG/Source/PCG/Public/Elements/Grammar/PCGSegmentSlicer.h b/Engine/Plugins/PCG/Source/PCG/Public/Elements/Grammar/PCGSegmentSlicer.h
--- a/Engine/Plugins/PCG/Source/PCG/Public/Elements/Grammar/PCGSegmentSlicer.h
+++ b/Engine/Plugins/PCG/Source/PCG/Public/Elements/Grammar/PCGSegmentSlicer.h
@@ -50,25 +50,25 @@ public:
 	UPROPERTY(BlueprintReadWrite, EditAnywhere, Category = Settings)
 	bool bAcceptIncompleteSlicing = false;
 
-	UPROPERTY(BlueprintReadWrite, EditAnywhere, Category = "Settings|ExtraAttributes", meta = (PCG_Overridable, InlineEditConditionToggle))
+	UPROPERTY(BlueprintReadWrite, EditAnywhere, Category = "Settings|Extra Output Attributes", meta = (PCG_Overridable, InlineEditConditionToggle))
 	bool bOutputModuleIndexAttribute = false;
 
 	/** Name of the module index output attribute name. */
-	UPROPERTY(BlueprintReadWrite, EditAnywhere, Category = "Settings|ExtraAttributes", meta = (PCG_Overridable, EditCondition = "bOutputModuleIndexAttribute"))
+	UPROPERTY(BlueprintReadWrite, EditAnywhere, Category = "Settings|Extra Output Attributes", meta = (PCG_Overridable, EditCondition = "bOutputModuleIndexAttribute"))
 	FName ModuleIndexAttributeName = TEXT("ModuleIndex");
 
-	UPROPERTY(BlueprintReadWrite, EditAnywhere, Category = "Settings|ExtraAttributes", meta = (PCG_Overridable, InlineEditConditionToggle))
+	UPROPERTY(BlueprintReadWrite, EditAnywhere, Category = "Settings|Extra Output Attributes", meta = (PCG_Overridable, InlineEditConditionToggle))
 	bool bOutputExtremityAttribute = false;
 
 	/** Name of the Extremity output attribute name. */
-	UPROPERTY(BlueprintReadWrite, EditAnywhere, Category = "Settings|ExtraAttributes", meta = (PCG_Overridable, EditCondition = "bOutputExtremityAttribute"))
+	UPROPERTY(BlueprintReadWrite, EditAnywhere, Category = "Settings|Extra Output Attributes", meta = (PCG_Overridable, EditCondition = "bOutputExtremityAttribute"))
 	FName ExtremityAttributeName = TEXT("Extremity");
 
-	UPROPERTY(BlueprintReadWrite, EditAnywhere, Category = "Settings|ExtraAttributes", meta = (PCG_Overridable, InlineEditConditionToggle))
+	UPROPERTY(BlueprintReadWrite, EditAnywhere, Category = "Settings|Extra Output Attributes", meta = (PCG_Overridable, InlineEditConditionToggle))
 	bool bOutputExtremityNeighborIndexAttribute = false;
 
 	/** Name of the extremity neighbor index output attribute name. */
-	UPROPERTY(BlueprintReadWrite, EditAnywhere, Category = "Settings|ExtraAttributes", meta = (PCG_Overridable, EditCondition = "bOutputExtremityNeighborIndexAttribute"))
+	UPROPERTY(BlueprintReadWrite, EditAnywhere, Category = "Settings|Extra Output Attributes", meta = (PCG_Overridable, EditCondition = "bOutputExtremityNeighborIndexAttribute"))
 	FName ExtremityNeighborIndexAttributeName = TEXT("ExtremityNeighborIndex");
 
 };
diff --git a/Engine/Plugins/PCG/Source/PCG/Public/Elements/Grammar/PCGSlicingBase.h b/Engine/Plugins/PCG/Source/PCG/Public/Elements/Grammar/PCGSlicingBase.h
--- a/Engine/Plugins/PCG/Source/PCG/Public/Elements/Grammar/PCGSlicingBase.h
+++ b/Engine/Plugins/PCG/Source/PCG/Public/Elements/Grammar/PCGSlicingBase.h
@@ -38,10 +38,39 @@ namespace PCGSlicingBaseConstants
 	const FName ModulesInfoPinLabel = TEXT("ModulesInfo");
 	const FName SymbolAttributeName = TEXT("Symbol");
 	const FName SizeAttributeName = TEXT("Size");
-	const FName ScalableAttributeName = TEXT("bScalable");
+	const FName ScalableAttributeName = TEXT("Scalable");
 	const FName DebugColorAttributeName = TEXT("DebugColor");
 }
 
+USTRUCT(BlueprintType)
+struct FPCGSlicingModuleAttributeNames
+{
+	GENERATED_BODY()
+
+public:
+	/** Mandatory. Expected type: FName. */
+	UPROPERTY(BlueprintReadWrite, EditAnywhere, Category = "")
+	FName SymbolAttributeName = PCGSlicingBaseConstants::SymbolAttributeName;
+
+	/** Mandatory. Expected type: double. */
+	UPROPERTY(BlueprintReadWrite, EditAnywhere, Category = "")
+	FName SizeAttributeName = PCGSlicingBaseConstants::SizeAttributeName;
+
+	UPROPERTY(BlueprintReadWrite, EditAnywhere, Category = "", meta = (InlineEditConditionToggle))
+	bool bProvideScalable = false;
+
+	/** Optional. Expected type: bool. If disabled, default value will be false. */
+	UPROPERTY(BlueprintReadWrite, EditAnywhere, Category = "", meta = (EditCondition = "bProvideScalable"))
+	FName ScalableAttributeName = PCGSlicingBaseConstants::ScalableAttributeName;
+
+	UPROPERTY(BlueprintReadWrite, EditAnywhere, Category = "", meta = (InlineEditConditionToggle))
+	bool bProvideDebugColor = false;
+
+	/** Optional. Expected type: Vector4. If disabled, default value will be (1.0, 1.0, 1.0, 1.0). */
+	UPROPERTY(BlueprintReadWrite, EditAnywhere, Category = "", meta = (EditCondition = "bProvideDebugColor"))
+	FName DebugColorAttributeName = PCGSlicingBaseConstants::DebugColorAttributeName;
+};
+
 UCLASS(MinimalAPI, Abstract, ClassGroup = (Procedural))
 class UPCGSlicingBaseSettings : public UPCGSettings
 {
@@ -55,7 +84,7 @@ public:
 	//~End UPCGSettings interface
 
 public:
-	/** Set it to true to pass the info as attribute set. Module info needs this structure (type + name): { FName Symbol, double Size, bool bScalable, FLinearColor DebugColor }. Types don't need to match exactly, must be able to be converted to the expected type. */
+	/** Set it to true to pass the info as attribute set. */
 	UPROPERTY(BlueprintReadWrite, EditAnywhere, Category = Settings)
 	bool bModuleInfoAsInput = false;
 
@@ -73,6 +102,10 @@ public:
 	UPROPERTY(BlueprintReadWrite, EditAnywhere, Category = Settings, meta = (EditCondition = "!bModuleInfoAsInput", EditConditionHides, DisplayAfter = bModuleInfoAsInput))
 	TArray<FPCGSlicingModule> ModulesInfo;
 
+	/** Fixed array of modules used for the slicing. */
+	UPROPERTY(BlueprintReadWrite, EditAnywhere, Category = Settings, meta = (EditCondition = "bModuleInfoAsInput", EditConditionHides, DisplayAfter = bModuleInfoAsInput, DisplayName = "Attribute Names for Module Info"))
+	FPCGSlicingModuleAttributeNames ModulesInfoAttributeNames;
+
 	/** An encoded string that represents how to apply a set of rules to a series of defined modules. */
 	UPROPERTY(BlueprintReadWrite, EditAnywhere, Category = Settings, meta = (EditCondition = "!bGrammarAsAttribute", EditConditionHides, DisplayAfter = bGrammarAsAttribute, PCG_Overridable))
 	FString Grammar;
@@ -82,32 +115,32 @@ public:
 	FPCGAttributePropertyInputSelector GrammarAttribute;
 
 	/** Do a match and set with the incoming module info, only if the module info is passed as input. */
-	UPROPERTY(BlueprintReadWrite, EditAnywhere, Category = "Settings|ExtraAttributes", meta = (EditCondition = "bModuleInfoAsInput", PCG_Overridable))
+	UPROPERTY(BlueprintReadWrite, EditAnywhere, Category = "Settings|Extra Output Attributes", meta = (EditCondition = "bModuleInfoAsInput", EditConditionHides, PCG_Overridable))
 	bool bForwardAttributesFromModulesInfo = false;
 
 	/** Name of the Symbol output attribute name. */
-	UPROPERTY(BlueprintReadWrite, EditAnywhere, Category = "Settings|ExtraAttributes", meta = (PCG_Overridable))
+	UPROPERTY(BlueprintReadWrite, EditAnywhere, Category = "Settings|Extra Output Attributes", meta = (PCG_Overridable))
 	FName SymbolAttributeName = PCGSlicingBaseConstants::SymbolAttributeName;
 
-	UPROPERTY(BlueprintReadWrite, EditAnywhere, Category = "Settings|ExtraAttributes", meta = (PCG_Overridable, InlineEditConditionToggle))
+	UPROPERTY(BlueprintReadWrite, EditAnywhere, Category = "Settings|Extra Output Attributes", meta = (PCG_Overridable, InlineEditConditionToggle))
 	bool bOutputSizeAttribute = true;
 
 	/** Name of the Size output attribute name, ignored if match and set from module info is true. */
-	UPROPERTY(BlueprintReadWrite, EditAnywhere, Category = "Settings|ExtraAttributes", meta = (PCG_Overridable, EditCondition = "bOutputSizeAttribute"))
+	UPROPERTY(BlueprintReadWrite, EditAnywhere, Category = "Settings|Extra Output Attributes", meta = (PCG_Overridable, EditCondition = "bOutputSizeAttribute"))
 	FName SizeAttributeName = PCGSlicingBaseConstants::SizeAttributeName;
 
-	UPROPERTY(BlueprintReadWrite, EditAnywhere, Category = "Settings|ExtraAttributes", meta = (PCG_Overridable, InlineEditConditionToggle))
+	UPROPERTY(BlueprintReadWrite, EditAnywhere, Category = "Settings|Extra Output Attributes", meta = (PCG_Overridable, InlineEditConditionToggle))
 	bool bOutputScalableAttribute = true;
 
 	/** Name of the Scalable output attribute name, ignored if match and set from module info is true. */
-	UPROPERTY(BlueprintReadWrite, EditAnywhere, Category = "Settings|ExtraAttributes", meta = (PCG_Overridable, EditCondition = "bOutputScalableAttribute"))
+	UPROPERTY(BlueprintReadWrite, EditAnywhere, Category = "Settings|Extra Output Attributes", meta = (PCG_Overridable, EditCondition = "bOutputScalableAttribute"))
 	FName ScalableAttributeName = PCGSlicingBaseConstants::ScalableAttributeName;
 
-	UPROPERTY(BlueprintReadWrite, EditAnywhere, Category = "Settings|ExtraAttributes", meta = (PCG_Overridable, InlineEditConditionToggle))
+	UPROPERTY(BlueprintReadWrite, EditAnywhere, Category = "Settings|Extra Output Attributes", meta = (PCG_Overridable, InlineEditConditionToggle))
 	bool bOutputDebugColorAttribute = false;
 
 	/** Name of the Debug Color output attribute name, ignored if match and set from module info is true. */
-	UPROPERTY(BlueprintReadWrite, EditAnywhere, Category = "Settings|ExtraAttributes", meta = (PCG_Overridable, EditCondition = "bOutputDebugColorAttribute"))
+	UPROPERTY(BlueprintReadWrite, EditAnywhere, Category = "Settings|Extra Output Attributes", meta = (PCG_Overridable, EditCondition = "bOutputDebugColorAttribute"))
 	FName DebugColorAttributeName = PCGSlicingBaseConstants::DebugColorAttributeName;
 };
 
diff --git a/Engine/Plugins/PCG/Source/PCG/Public/Elements/Grammar/PCGVolumeSlicer.h b/Engine/Plugins/PCG/Source/PCG/Public/Elements/Grammar/PCGVolumeSlicer.h
--- a/Engine/Plugins/PCG/Source/PCG/Public/Elements/Grammar/PCGVolumeSlicer.h
+++ b/Engine/Plugins/PCG/Source/PCG/Public/Elements/Grammar/PCGVolumeSlicer.h
@@ -38,11 +38,11 @@ public:
 	UPROPERTY(BlueprintReadWrite, EditAnywhere, Category = Settings, meta = (EditCondition = "bExtrudeVectorAsAttribute", EditConditionHides, DisplayAfter = bExtrudeVectorAsAttribute, PCG_Overridable))
 	FPCGAttributePropertyInputSelector ExtrudeVectorAttribute;
 
-	UPROPERTY(BlueprintReadWrite, EditAnywhere, Category = "Settings|ExtraAttributes", meta = (PCG_Overridable, InlineEditConditionToggle))
+	UPROPERTY(BlueprintReadWrite, EditAnywhere, Category = "Settings|Extra Output Attributes", meta = (PCG_Overridable, InlineEditConditionToggle))
 	bool bOutputSplineIndexAttribute = true;
 
 	/** Name of the spline index output attribute name. */
-	UPROPERTY(BlueprintReadWrite, EditAnywhere, Category = "Settings|ExtraAttributes", meta = (PCG_Overridable, EditCondition = "bOutputSplineIndexAttribute"))
+	UPROPERTY(BlueprintReadWrite, EditAnywhere, Category = "Settings|Extra Output Attributes", meta = (PCG_Overridable, EditCondition = "bOutputSplineIndexAttribute"))
 	FName SplineIndexAttributeName = TEXT("SplineIndex");
 };
 
