[PCG] Samplers PrepareForExecute needs to run on GT until we have proper pre-execution cache of UPCGComponent::GetPCGActorData()

#rnx

diff --git a/Engine/Plugins/PCG/Source/PCG/Public/Elements/PCGSplineSampler.h b/Engine/Plugins/PCG/Source/PCG/Public/Elements/PCGSplineSampler.h
--- a/Engine/Plugins/PCG/Source/PCG/Public/Elements/PCGSplineSampler.h
+++ b/Engine/Plugins/PCG/Source/PCG/Public/Elements/PCGSplineSampler.h
@@ -280,6 +280,10 @@ public:
 class FPCGSplineSamplerElement : public IPCGElement
 {
 public:
+	// Returns true when Context is nullptr which happens in FPCGGraphExecutor::OnTaskInputsReady
+	// We can't know for sure at that point if we need to call UPCGComponent::GetActorPCGData() in our GetDependenciesCrc so we need to make sure it happens on MainThread
+	virtual bool CanExecuteOnlyOnMainThread(FPCGContext* Context) const override { return Context == nullptr; }
+
 	virtual void GetDependenciesCrc(const FPCGDataCollection& InInput, const UPCGSettings* InSettings, UPCGComponent* InComponent, FPCGCrc& OutCrc) const override;
 	// Worth computing a full CRC in case we can halt change propagation/re-executions
 	virtual bool ShouldComputeFullOutputDataCrc(FPCGContext* Context) const override { return true; }
diff --git a/Engine/Plugins/PCG/Source/PCG/Public/Elements/PCGSurfaceSampler.h b/Engine/Plugins/PCG/Source/PCG/Public/Elements/PCGSurfaceSampler.h
--- a/Engine/Plugins/PCG/Source/PCG/Public/Elements/PCGSurfaceSampler.h
+++ b/Engine/Plugins/PCG/Source/PCG/Public/Elements/PCGSurfaceSampler.h
@@ -157,6 +157,10 @@ public:
 class FPCGSurfaceSamplerElement : public TPCGTimeSlicedElementBase<PCGSurfaceSampler::FSurfaceSamplerExecutionState, PCGSurfaceSampler::FSurfaceSamplerIterationState>
 {
 public:
+	// Returns true when Context is nullptr which happens in FPCGGraphExecutor::OnTaskInputsReady
+	// We can't know for sure at that point if we need to call UPCGComponent::GetActorPCGData() in our GetDependenciesCrc so we need to make sure it happens on MainThread
+	virtual bool CanExecuteOnlyOnMainThread(FPCGContext* Context) const override { return Context == nullptr; }
+
 	virtual void GetDependenciesCrc(const FPCGDataCollection& InInput, const UPCGSettings* InSettings, UPCGComponent* InComponent, FPCGCrc& OutCrc) const override;
 	// Might be sampling landscape or other external data, worth computing a full CRC in case we can halt change propagation/re-executions
 	virtual bool ShouldComputeFullOutputDataCrc(FPCGContext* Context) const override { return true; }
diff --git a/Engine/Plugins/PCG/Source/PCG/Public/Elements/PCGVolumeSampler.h b/Engine/Plugins/PCG/Source/PCG/Public/Elements/PCGVolumeSampler.h
--- a/Engine/Plugins/PCG/Source/PCG/Public/Elements/PCGVolumeSampler.h
+++ b/Engine/Plugins/PCG/Source/PCG/Public/Elements/PCGVolumeSampler.h
@@ -94,6 +94,10 @@ protected:
 class FPCGVolumeSamplerElement : public TPCGTimeSlicedElementBase<PCGVolumeSampler::FVolumeSamplerExecutionState, PCGVolumeSampler::FVolumeSamplerIterationState>
 {
 public:
+	// Returns true when Context is nullptr which happens in FPCGGraphExecutor::OnTaskInputsReady
+	// We can't know for sure at that point if we need to call UPCGComponent::GetActorPCGData() in our GetDependenciesCrc so we need to make sure it happens on MainThread
+	virtual bool CanExecuteOnlyOnMainThread(FPCGContext* Context) const override { return Context == nullptr; }
+
 	virtual void GetDependenciesCrc(const FPCGDataCollection& InInput, const UPCGSettings* InSettings, UPCGComponent* InComponent, FPCGCrc& OutCrc) const override;
 	// Might be sampling external data like brush, worth computing a full CRC in case we can halt change propagation/re-executions
 	virtual bool ShouldComputeFullOutputDataCrc(FPCGContext* Context) const override { return true; }
