[PCG] Reinstate missing GPU node setup validation errors

#rnx

diff --git a/Engine/Plugins/PCG/Source/PCG/Private/Compute/Elements/PCGCustomHLSL.cpp b/Engine/Plugins/PCG/Source/PCG/Private/Compute/Elements/PCGCustomHLSL.cpp
--- a/Engine/Plugins/PCG/Source/PCG/Private/Compute/Elements/PCGCustomHLSL.cpp
+++ b/Engine/Plugins/PCG/Source/PCG/Private/Compute/Elements/PCGCustomHLSL.cpp
@@ -791,7 +791,11 @@ bool UPCGCustomHLSLSettings::IsKernelValid(FPCGContext* InContext) const
 	{
 		if (Properties.AllowedTypes == EPCGDataType::Any)
 		{
-			PCGE_LOG_C(Error, GraphAndLog, InContext, FText::Format(LOCTEXT("InvalidAnyInput", "Custom kernels do not support inputs of type Any, found on pin {0}."), FText::FromName(Properties.Label)));
+			if (InContext)
+			{
+				PCGE_LOG_C(Error, GraphAndLog, InContext, FText::Format(LOCTEXT("InvalidAnyInput", "Custom kernels do not support inputs of type Any, found on pin {0}."), FText::FromName(Properties.Label)));
+			}
+
 			return false;
 		}
 	}
@@ -800,19 +804,31 @@ bool UPCGCustomHLSLSettings::IsKernelValid(FPCGContext* InContext) const
 	{
 		if (Properties.AllowedTypes == EPCGDataType::Any)
 		{
-			PCGE_LOG_C(Error, GraphAndLog, InContext, FText::Format(LOCTEXT("InvalidAnyOutput", "Custom kernels do not support outputs of type Any, found on pin {0}."), FText::FromName(Properties.Label)));
+			if (InContext)
+			{
+				PCGE_LOG_C(Error, GraphAndLog, InContext, FText::Format(LOCTEXT("InvalidAnyOutput", "Custom kernels do not support outputs of type Any, found on pin {0}."), FText::FromName(Properties.Label)));
+			}
+
 			return false;
 		}
 
 		if (!!(Properties.AllowedTypes & EPCGDataType::Landscape))
 		{
-			PCGE_LOG_C(Error, GraphAndLog, InContext, FText::Format(LOCTEXT("InvalidLSOutput", "Custom kernels do not support outputs of type Landscape, found on pin {0}."), FText::FromName(Properties.Label)));
+			if (InContext)
+			{
+				PCGE_LOG_C(Error, GraphAndLog, InContext, FText::Format(LOCTEXT("InvalidLSOutput", "Custom kernels do not support outputs of type Landscape, found on pin {0}."), FText::FromName(Properties.Label)));
+			}
+
 			return false;
 		}
 
 		if (!!(Properties.AllowedTypes & EPCGDataType::Texture))
 		{
-			PCGE_LOG_C(Error, GraphAndLog, InContext, FText::Format(LOCTEXT("InvalidTextureOutput", "Custom kernels do not support outputs of type Texture, found on pin {0}."), FText::FromName(Properties.Label)));
+			if (InContext)
+			{
+				PCGE_LOG_C(Error, GraphAndLog, InContext, FText::Format(LOCTEXT("InvalidTextureOutput", "Custom kernels do not support outputs of type Texture, found on pin {0}."), FText::FromName(Properties.Label)));
+			}
+
 			return false;
 		}
 	}
@@ -1154,4 +1170,19 @@ FString UPCGCustomHLSLSettings::GetCookedKernelSource(const TMap<FPCGKernelAttri
 	return Result;
 }
 
+bool FPCGCustomHLSLElement::ExecuteInternal(FPCGContext* Context) const
+{
+	check(Context);
+
+	const UPCGCustomHLSLSettings* Settings = Context->GetInputSettings<UPCGCustomHLSLSettings>();
+	check(Settings);
+
+	// Calling IsKernelValid with the Context will emit warnings/errors on the node.
+	const bool bIsValid = Settings->IsKernelValid(Context);
+
+	ensureMsgf(!bIsValid, TEXT("Custom HLSL element should only execute on CPU if the kernel is invalid."));
+
+	return true;
+}
+
 #undef LOCTEXT_NAMESPACE
diff --git a/Engine/Plugins/PCG/Source/PCG/Public/Compute/Elements/PCGCustomHLSL.h b/Engine/Plugins/PCG/Source/PCG/Public/Compute/Elements/PCGCustomHLSL.h
--- a/Engine/Plugins/PCG/Source/PCG/Public/Compute/Elements/PCGCustomHLSL.h
+++ b/Engine/Plugins/PCG/Source/PCG/Public/Compute/Elements/PCGCustomHLSL.h
@@ -145,7 +145,6 @@ protected:
 class FPCGCustomHLSLElement : public IPCGElement
 {
 protected:
-	// During compilation, custom HLSL nodes will be replaced with a compute graph element, and this should never actually execute.
-	// Left ensure here for validation.
-	virtual bool ExecuteInternal(FPCGContext* Context) const override { ensure(false); return true; }
+	// This will only be called if the custom HLSL node is not set up correctly (valid nodes are replaced with a compute graph element).
+	virtual bool ExecuteInternal(FPCGContext* Context) const override;
 };
