[PCG] Refactor Copy Attributes

* Match Copy Points API
* Sunset Transfer Attributes, which will be replaced by [REDACTED_PERSON]
* Update tests so they test more types

diff --git a/Engine/Plugins/PCG/Config/BasePCG.ini b/Engine/Plugins/PCG/Config/BasePCG.ini
--- a/Engine/Plugins/PCG/Config/BasePCG.ini
+++ b/Engine/Plugins/PCG/Config/BasePCG.ini
@@ -68,6 +68,8 @@
 +PropertyRedirects=(OldName="PCGPropertyToParamDataSettings.bExtractObjectAndStruct",NewName="PCGPropertyToParamDataSettings.bForceObjectAndStructExtraction")
 +PropertyRedirects=(OldName="PCGAttributePropertySelectorBlueprintHelpers.CopyAndFixSource.Selector",NewName="PCGAttributePropertySelectorBlueprintHelpers.CopyAndFixSource.OutputSelector")
 +PropertyRedirects=(OldName="PCGAttributePropertySelectorBlueprintHelpers.CopyAndFixSource.InSelector",NewName="PCGAttributePropertySelectorBlueprintHelpers.CopyAndFixSource.InputSelector")
++PropertyRedirects=(OldName="PCGAttributeTransferSettings.SourceAttributeProperty",NewName="PCGAttributeTransferSettings.InputSource")
++PropertyRedirects=(OldName="PCGAttributeTransferSettings.TargetAttributeProperty",NewName="PCGAttributeTransferSettings.OutputTarget")
 
 +ClassRedirects=(OldName="PCGPointSamplerSettings",NewName="/Script/PCG.PCGSelectPointsSettings")
 +ClassRedirects=(OldName="PCGDensityNoiseSettings",NewName="/Script/PCG.PCGAttributeNoiseSettings")
@@ -82,6 +84,7 @@
 +ClassRedirects=(OldName="PCGAttributeFilterNamesSettings",NewName="/Script/PCG.PCGDeleteAttributesSettings")
 +ClassRedirects=(OldName="PCGPropertyToParamDataSettings",NewName="/Script/PCG.PCGGetActorPropertySettings")
 +ClassRedirects=(OldName="PCGSortPointsSettings",NewName="/Script/PCG.PCGSortAttributesSettings")
++ClassRedirects=(OldName="PCGMetadataOperationSettings",NewName="/Script/PCG.PCGCopyAttributesSettings")
 
 +StructRedirects=(OldName="PCGPointFilterThresholdSettings",NewName="/Script/PCG.PCGAttributeFilterThresholdSettings")
 +StructRedirects=(OldName="PCGActorPropertyOverrideDescription",NewName="/Script/PCG.PCGObjectPropertyOverrideDescription")
diff --git a/Engine/Plugins/PCG/Source/PCG/Private/Elements/Metadata/PCGMetadataElement.cpp b/Engine/Plugins/PCG/Source/PCG/Private/Elements/Metadata/PCGMetadataElement.cpp
deleted file mode 100644
diff --git a/Engine/Plugins/PCG/Source/PCG/Private/Elements/PCGAttributeTransferElement.cpp b/Engine/Plugins/PCG/Source/PCG/Private/Elements/PCGAttributeTransferElement.cpp
deleted file mode 100644
diff --git a/Engine/Plugins/PCG/Source/PCG/Private/Elements/PCGCopyAttributes.cpp b/Engine/Plugins/PCG/Source/PCG/Private/Elements/PCGCopyAttributes.cpp
new file mode 100644
--- /dev/null
+++ b/Engine/Plugins/PCG/Source/PCG/Private/Elements/PCGCopyAttributes.cpp
@@ -0,0 +1,287 @@
+// Copyright Epic Games, Inc. All Rights Reserved.
+
+#include "Elements/PCGCopyAttributes.h"
+
+#include "PCGContext.h"
+#include "PCGCustomVersion.h"
+#include "PCGEdge.h"
+#include "PCGParamData.h"
+#include "Data/PCGPointData.h"
+#include "Data/PCGSpatialData.h"
+#include "Elements/PCGCopyPoints.h"
+#include "Helpers/PCGMetadataHelpers.h"
+
+// Keep loctext namespace to its previous value to avoid to redo loc for this.
+#define LOCTEXT_NAMESPACE "PCGMetadataElement"
+
+
+UPCGCopyAttributesSettings::UPCGCopyAttributesSettings()
+{
+	// Previous default object was: None for source and target attribute, density for point property
+	// and Property to attribute
+	// Recreate the same default
+	InputSource.SetPointProperty(EPCGPointProperties::Density);
+	OutputTarget.SetAttributeName(NAME_None);
+}
+
+FString UPCGCopyAttributesSettings::GetAdditionalTitleInformation() const
+{
+#if WITH_EDITOR
+	if (bCopyAllAttributes)
+	{
+		return LOCTEXT("NoteTitleAllAttributes", "All Attributes").ToString();
+	}
+
+	if (IsPropertyOverriddenByPin(GET_MEMBER_NAME_CHECKED(UPCGCopyAttributesSettings, InputSource)) || IsPropertyOverriddenByPin(GET_MEMBER_NAME_CHECKED(UPCGCopyAttributesSettings, OutputTarget)))
+	{
+		return FString();
+	}
+#endif
+
+	return FString::Printf(TEXT("%s -> %s"), *InputSource.GetDisplayText().ToString(), *OutputTarget.GetDisplayText().ToString());
+}
+
+#if WITH_EDITOR
+EPCGChangeType UPCGCopyAttributesSettings::GetChangeTypeForProperty(const FName& InPropertyName) const
+{
+	return Super::GetChangeTypeForProperty(InPropertyName) | EPCGChangeType::Cosmetic;
+}
+#endif // WITH_EDITOR
+
+FPCGElementPtr UPCGCopyAttributesSettings::CreateElement() const
+{
+	return MakeShared<FPCGCopyAttributesElement>();
+}
+
+#if WITH_EDITOR
+FText UPCGCopyAttributesSettings::GetNodeTooltipText() const
+{
+	return LOCTEXT("NodeTooltip", "Copy from the Input Source to Output Target attribute.");
+}
+
+void UPCGCopyAttributesSettings::ApplyDeprecation(UPCGNode* InOutNode)
+{
+	if (DataVersion < FPCGCustomVersion::UpdateAttributePropertyInputSelector
+		&& OutputTarget.GetSelection() == EPCGAttributePropertySelection::Attribute
+		&& OutputTarget.GetAttributeName() == PCGMetadataAttributeConstants::SourceAttributeName)
+	{
+		// Previous behavior of the output target for this node was:
+		// None => LastCreated
+		OutputTarget.SetAttributeName(PCGMetadataAttributeConstants::LastCreatedAttributeName);
+	}
+
+	if (bWasDeprecated)
+	{
+		UPCGPin* TargetPin = InOutNode->GetInputPin(PCGCopyPointsConstants::TargetPointsLabel);
+		UPCGPin* SourcePin = InOutNode->GetInputPin(PCGCopyPointsConstants::SourcePointsLabel);
+
+		check(TargetPin && SourcePin);
+
+		if (!SourcePin->IsConnected())
+		{
+			for (UPCGEdge* Edge : TargetPin->Edges)
+			{
+				if (!Edge || !Edge->InputPin)
+				{
+					continue;
+				}
+
+				Edge->InputPin->AddEdgeTo(SourcePin);
+			}
+		}
+	}
+
+	bWasDeprecated = false;
+
+	Super::ApplyDeprecation(InOutNode);
+}
+
+void UPCGCopyAttributesSettings::ApplyDeprecationBeforeUpdatePins(UPCGNode* InOutNode, TArray<TObjectPtr<UPCGPin>>& InputPins, TArray<TObjectPtr<UPCGPin>>& OutputPins)
+{
+	if (!InOutNode)
+	{
+		return;
+	}
+
+	bWasDeprecated = InputPins.ContainsByPredicate([](const TObjectPtr<UPCGPin>& Pin) -> bool { return Pin && Pin->Properties.Label == PCGPinConstants::DefaultInputLabel; });
+
+	InOutNode->RenameInputPin(PCGPinConstants::DefaultInputLabel, PCGCopyPointsConstants::TargetPointsLabel, /*bInBroadcastUpdate=*/false);
+	InOutNode->RenameInputPin(TEXT("Attribute"), PCGCopyPointsConstants::SourcePointsLabel, /*bInBroadcastUpdate=*/false);
+}
+
+EPCGDataType UPCGCopyAttributesSettings::GetCurrentPinTypes(const UPCGPin* InPin) const
+{
+	if (!InPin || !InPin->IsOutputPin())
+	{
+		return Super::GetCurrentPinTypes(InPin);
+	}
+
+	EPCGDataType OutputType = GetTypeUnionOfIncidentEdges(PCGCopyPointsConstants::TargetPointsLabel);
+	return OutputType != EPCGDataType::None ? OutputType : EPCGDataType::Any;
+}
+
+#endif // WITH_EDITOR
+
+TArray<FPCGPinProperties> UPCGCopyAttributesSettings::InputPinProperties() const
+{
+	TArray<FPCGPinProperties> Properties;
+	Properties.Emplace_GetRef(PCGCopyPointsConstants::TargetPointsLabel, EPCGDataType::Any).SetRequiredPin();
+	Properties.Emplace_GetRef(PCGCopyPointsConstants::SourcePointsLabel, EPCGDataType::Any).SetRequiredPin();
+
+	return Properties;
+}
+
+TArray<FPCGPinProperties> UPCGCopyAttributesSettings::OutputPinProperties() const
+{
+	TArray<FPCGPinProperties> Properties;
+	Properties.Emplace(PCGPinConstants::DefaultOutputLabel, EPCGDataType::Any);
+
+	return Properties;
+}
+
+bool FPCGCopyAttributesElement::ExecuteInternal(FPCGContext* Context) const
+{
+	TRACE_CPUPROFILER_EVENT_SCOPE(FPCGMetadataOperationElement::Execute);
+
+	const UPCGCopyAttributesSettings* Settings = Context->GetInputSettings<UPCGCopyAttributesSettings>();
+	check(Settings);
+
+	TArray<FPCGTaggedData> Sources = Context->InputData.GetInputsByPin(PCGCopyPointsConstants::SourcePointsLabel);
+	TArray<FPCGTaggedData> Targets = Context->InputData.GetInputsByPin(PCGCopyPointsConstants::TargetPointsLabel);
+
+	TArray<FPCGTaggedData>& Outputs = Context->OutputData.TaggedData;
+
+	const int32 NumSources = Sources.Num();
+	const int32 NumTargets = Targets.Num();
+
+	if (NumSources == 0 || NumTargets == 0)
+	{
+		// Nothing to do
+		return true;
+	}
+
+	int32 NumIterations = 0;
+	switch (Settings->Operation)
+	{
+	case EPCGCopyAttributesOperation::CopyEachSourceOnEveryTarget:
+		NumIterations = NumSources * NumTargets;
+		break;
+	case EPCGCopyAttributesOperation::MergeSourcesAndCopyToAllTargets:
+		NumIterations = NumTargets;
+		break;
+	case EPCGCopyAttributesOperation::CopyEachSourceToEachTargetRespectively:
+	{
+		if (NumSources != NumTargets && NumSources != 1 && NumTargets != 1)
+		{
+			PCGLog::LogErrorOnGraph(FText::Format(LOCTEXT("MismatchNum", "Num Sources ({0}) mismatches with Num Targets ({1}). Only supports N:N, 1:N and N:1 operation."), NumSources, NumTargets), Context);
+			Outputs = Sources;
+			return true;
+		}
+
+		NumIterations = NumTargets;
+		break;
+	}
+	}
+
+	for (int32 i = 0; i < NumIterations; ++i)
+	{
+		const int32 TargetIndex = i % NumTargets;
+		const FPCGTaggedData& Target = Targets[TargetIndex];
+
+		FPCGTaggedData& Output = Outputs.Add_GetRef(Target);
+		const UPCGData* TargetData = Target.Data;
+
+		bool bSuccess = false;
+		UPCGData* OutputData = nullptr;
+
+		auto DoCopy = [this, TargetData, &Sources, Context, Settings, &bSuccess, &OutputData](int32 SourceIndex)
+		{
+			const FPCGTaggedData& Source = Sources[SourceIndex];
+			const UPCGData* SourceData = Source.Data;
+
+			const bool bIsSameData = TargetData == SourceData;
+
+			if (!TargetData || !SourceData)
+			{
+				PCGE_LOG(Error, GraphAndLog, LOCTEXT("InvalidInputData", "Invalid input data"));
+				return;
+			}
+
+			const UPCGMetadata* SourceMetadata = SourceData->ConstMetadata();
+
+			if (!SourceMetadata)
+			{
+				PCGE_LOG(Warning, GraphAndLog, LOCTEXT("MissingMetadata", "Source has no metadata"));
+				return;
+			}
+
+			// Early out when trying to copy all attributes from self, since nothing will happen
+			if (bIsSameData && Settings->bCopyAllAttributes)
+			{
+				PCGE_LOG(Verbose, LogOnly, LOCTEXT("TrivialAllCopy", "Copying all attributes on itself is a trivial operation."));
+				return;
+			}
+
+			// Early out when trying to copy same value to self
+			if (bIsSameData && Settings->InputSource == Settings->OutputTarget)
+			{
+				PCGE_LOG(Verbose, LogOnly, LOCTEXT("TrivialCopy", "Copying attribute to itself is a trivial operation."));
+				return;
+			}
+
+			if (!OutputData)
+			{
+				OutputData = TargetData->DuplicateData(Context);
+			}
+
+			if (Settings->bCopyAllAttributes)
+			{
+				bSuccess |= PCGMetadataHelpers::CopyAllAttributes(SourceData, OutputData, Context);
+			}
+			else
+			{
+				PCGMetadataHelpers::FPCGCopyAttributeParams Params{};
+				Params.SourceData = SourceData;
+				Params.TargetData = OutputData;
+				Params.InputSource = Settings->InputSource;
+				Params.OutputTarget = Settings->OutputTarget;
+				Params.OptionalContext = Context;
+				Params.bSameOrigin = bIsSameData;
+
+				bSuccess |= PCGMetadataHelpers::CopyAttribute(Params);
+			}
+		};
+
+		switch (Settings->Operation)
+		{
+		case EPCGCopyAttributesOperation::CopyEachSourceOnEveryTarget:
+		{
+			DoCopy(i / NumTargets);
+			break;
+		}
+		case EPCGCopyAttributesOperation::MergeSourcesAndCopyToAllTargets:
+		{
+			for (int32 j = 0; j < NumSources; ++j)
+			{
+				DoCopy(j);
+			}
+			break;
+		}
+		case EPCGCopyAttributesOperation::CopyEachSourceToEachTargetRespectively:
+		default:
+		{
+			DoCopy(FMath::Min(i, NumSources - 1));
+			break;
+		}
+		}
+
+		if (bSuccess && OutputData)
+		{
+			Output.Data = OutputData;
+		}
+	}
+
+	return true;
+}
+
+#undef LOCTEXT_NAMESPACE
diff --git a/Engine/Plugins/PCG/Source/PCG/Private/Helpers/PCGMetadataHelpers.cpp b/Engine/Plugins/PCG/Source/PCG/Private/Helpers/PCGMetadataHelpers.cpp
--- a/Engine/Plugins/PCG/Source/PCG/Private/Helpers/PCGMetadataHelpers.cpp
+++ b/Engine/Plugins/PCG/Source/PCG/Private/Helpers/PCGMetadataHelpers.cpp
@@ -139,7 +139,6 @@ namespace PCGMetadataHelpers
 			}
 
 			// We need accessors if we have a multi entry source attribute or we have extractors
-			const bool bIsMultiEntries = SourceData->IsA<UPCGParamData>() && SourceMetadata->GetLocalItemCount() > 1;
 			const bool bInputHasAnyExtra = !InputSource.GetExtraNames().IsEmpty();
 			const bool bOutputHasAnyExtra = !OutputTarget.GetExtraNames().IsEmpty();
 			const bool bSourceIsAttribute = InputSource.GetSelection() == EPCGAttributePropertySelection::Attribute;
@@ -147,7 +146,7 @@ namespace PCGMetadataHelpers
 			// Cast is only required if it is on an output attribute that has no extra (that we will create)
 			const bool bOutputTypeCast = bTargetIsAttribute && !bOutputHasAnyExtra && (RequestedOutputType != EPCGMetadataTypes::Unknown);
 
-			const bool bNeedAccessors = bIsMultiEntries || bInputHasAnyExtra || bOutputHasAnyExtra || !bSourceIsAttribute || !bTargetIsAttribute || bOutputTypeCast;
+			const bool bNeedAccessors = bInputHasAnyExtra || bOutputHasAnyExtra || !bSourceIsAttribute || !bTargetIsAttribute || bOutputTypeCast;
 
 			// If no accessor, copy over the attribute
 			if (!bNeedAccessors)
@@ -162,10 +161,71 @@ namespace PCGMetadataHelpers
 				// Presence of attribute was already checked before, this should not return null
 				check(SourceAttribute);
 
-				if (FPCGMetadataAttributeBase* NewAttr = TargetMetadata->CopyAttribute(SourceAttribute, LocalDestinationAttribute, /*bKeepParent=*/false, /*bCopyEntries=*/true,/*bCopyValues=*/true))
+				if (FPCGMetadataAttributeBase* NewAttr = TargetMetadata->CopyAttribute(SourceAttribute, LocalDestinationAttribute, /*bKeepParent=*/false, /*bCopyEntries=*/bSameOrigin, /*bCopyValues=*/true))
 				{
 					// To keep the previous behavior for single value attributes, we force the copied attribute to have its default value set to the first entry (only done if there is just a single value).
-					NewAttr->SetDefaultValueToFirstEntry();
+					const bool bHasJustASingleEntryOrNoEntry = SourceAttribute->GetNumberOfEntriesWithParents() <= 1;
+					if (bHasJustASingleEntryOrNoEntry)
+					{
+						NewAttr->SetDefaultValueToFirstEntry();
+					}
+
+					// And finally, we create our source and target keys, to get all the metadata entry keys, to remap if we are not the same origin and that we have multiple entries.
+					if (!bSameOrigin && !bHasJustASingleEntryOrNoEntry)
+					{
+						TUniquePtr<const IPCGAttributeAccessorKeys> SourceKeys = PCGAttributeAccessorHelpers::CreateConstKeys(SourceData, InputSource);
+						TUniquePtr<IPCGAttributeAccessorKeys> TargetKeys = PCGAttributeAccessorHelpers::CreateKeys(TargetData, OutputTarget);
+
+						// They must exist, so we can check.
+						check(SourceKeys && TargetKeys);
+
+						if (SourceKeys->GetNum() > 0 && TargetKeys->GetNum() > 0)
+						{
+							TArray<const PCGMetadataEntryKey*> AllSourceEntryKeysPtr;
+							TArray<PCGMetadataEntryKey*> AllTargetEntryKeysPtr;
+
+							AllSourceEntryKeysPtr.SetNumUninitialized(FMath::Min(SourceKeys->GetNum(), TargetKeys->GetNum()));
+							AllTargetEntryKeysPtr.SetNumUninitialized(TargetKeys->GetNum());
+
+							if (SourceKeys->GetKeys<PCGMetadataEntryKey>(0, AllSourceEntryKeysPtr) && TargetKeys->GetKeys<PCGMetadataEntryKey>(0, AllTargetEntryKeysPtr))
+							{
+								// Gather all the value keys
+								TArray<PCGMetadataEntryKey> AllSourceEntryKeys;
+								Algo::Transform(AllSourceEntryKeysPtr, AllSourceEntryKeys, [](const PCGMetadataEntryKey* KeyPtr) { return *KeyPtr; });
+								TArray<PCGMetadataValueKey> ValueKeys;
+								ValueKeys.Reserve(AllTargetEntryKeysPtr.Num());
+								SourceAttribute->GetValueKeys(TArrayView<const PCGMetadataEntryKey>(AllSourceEntryKeys), ValueKeys);
+
+								// Extends values keys to match target entry keys size. It will loop on value keys.
+								const int32 ValueKeysSize = ValueKeys.Num();
+								for (int32 i = ValueKeysSize; i < AllTargetEntryKeysPtr.Num(); ++i)
+								{
+									ValueKeys.Add(ValueKeys[i % ValueKeysSize]);
+								}
+
+								// Make sure that the Target has some metadata entry
+								// Implementation note: this is a stripped down version of UPCGMetadata::InitializeOnSet
+								TArray<PCGMetadataEntryKey*> AllTargetEntryKeysPtrTemp;
+								AllTargetEntryKeysPtrTemp.Reserve(AllTargetEntryKeysPtr.Num());
+								for (int EntryIndex = 0; EntryIndex < AllTargetEntryKeysPtr.Num(); ++EntryIndex)
+								{
+									PCGMetadataEntryKey& EntryKey = *AllTargetEntryKeysPtr[EntryIndex];
+									if (EntryKey == PCGInvalidEntryKey || (EntryKey < TargetMetadata->GetItemKeyCountForParent()))
+									{
+										AllTargetEntryKeysPtrTemp.Add(&EntryKey);
+									}
+								}
+
+								if (!AllTargetEntryKeysPtrTemp.IsEmpty())
+								{
+									TRACE_CPUPROFILER_EVENT_SCOPE(FPCGAttributeAccessor::Prepare::AddEntriesInPlace);
+									TargetMetadata->AddEntriesInPlace(AllTargetEntryKeysPtrTemp);
+								}
+
+								NewAttr->SetValuesFromValueKeys(AllTargetEntryKeysPtr, ValueKeys, /*bResetValueOnDefaultValueKey=*/true);
+							}
+						}
+					}
 				}
 				else
 				{
diff --git a/Engine/Plugins/PCG/Source/PCG/Private/Tests/Elements/PCGCopyAttributeTest.cpp b/Engine/Plugins/PCG/Source/PCG/Private/Tests/Elements/PCGCopyAttributeTest.cpp
--- a/Engine/Plugins/PCG/Source/PCG/Private/Tests/Elements/PCGCopyAttributeTest.cpp
+++ b/Engine/Plugins/PCG/Source/PCG/Private/Tests/Elements/PCGCopyAttributeTest.cpp
@@ -11,7 +11,8 @@
 #include "Metadata/PCGMetadata.h"
 #include "Metadata/PCGMetadataAttributeTpl.h"
 
-#include "Elements/Metadata/PCGMetadataElement.h"
+#include "Elements/PCGCopyAttributes.h"
+#include "Elements/PCGCopyPoints.h"
 
 IMPLEMENT_CUSTOM_SIMPLE_AUTOMATION_TEST(FPCGCopyAttributeTests_Points_PropertyToProperty, FPCGTestBaseClass, "Plugins.PCG.CopyAttribute.Points.PropertyToProperty", PCGTestsCommon::TestFlags)
 IMPLEMENT_CUSTOM_SIMPLE_AUTOMATION_TEST(FPCGCopyAttributeTests_Points_PropertyToAttribute, FPCGTestBaseClass, "Plugins.PCG.CopyAttribute.Points.PropertyToAttribute", PCGTestsCommon::TestFlags)
@@ -28,7 +29,7 @@ IMPLEMENT_CUSTOM_SIMPLE_AUTOMATION_TEST(FPCGCopyAttributeTests_Params_CopyingAll
 
 namespace PCGCopyAttributeTests
 {
-	static const FName AttributeName = TEXT("Double");
+	static const FName AttributeName = TEXT("InputAttr");
 
 	UPCGPointData* CreateInputPointData(FPCGContext* Context, const int NumPoints)
 	{
@@ -37,10 +38,7 @@ namespace PCGCopyAttributeTests
 		UPCGPointData* NewPointData = NewObject<UPCGPointData>();
 		NewPointData->SetFlags(RF_Transient);
 
-		check(NewPointData && NewPointData->Metadata);
-
-		// Create an attribute
-		FPCGMetadataAttribute<double>* NewAttribute = NewPointData->Metadata->CreateAttribute<double>(AttributeName, /*DefaultValue=*/0.0, /*bAllowInterpolation=*/true, /*bOverrideParent=*/false);
+		check(NewPointData);
 
 		TArray<FPCGPoint>& Points = NewPointData->GetMutablePoints();
 		Points.SetNum(NumPoints);
@@ -48,49 +46,47 @@ namespace PCGCopyAttributeTests
 		{
 			// Store the index in the density
 			Points[i].Density = i;
-
-			 // And also in NewAttribute (offset by 5 to differientiate from density)
-			NewPointData->Metadata->InitializeOnSet(Points[i].MetadataEntry);
-			NewAttribute->SetValue(Points[i].MetadataEntry, i + 5);
 		}
 
-		FPCGTaggedData& InputData = Context->InputData.TaggedData.Emplace_GetRef();
-		InputData.Data = NewPointData;
-		InputData.Pin = PCGPinConstants::DefaultInputLabel;
-
 		return NewPointData;
 	}
 
-	UPCGParamData* CreateInputParamData(FPCGContext* Context, const int NumEntries)
+	UPCGParamData* CreateInputParamData(FPCGContext* Context)
 	{
 		check(Context);
 
 		UPCGParamData* NewParamData = NewObject<UPCGParamData>();
 		NewParamData->SetFlags(RF_Transient);
 
-		check(NewParamData && NewParamData->Metadata);
-
-		// Create an attribute
-		FPCGMetadataAttribute<double>* NewAttribute = NewParamData->Metadata->CreateAttribute<double>(AttributeName, /*DefaultValue=*/0.0, /*bAllowInterpolation=*/true, /*bOverrideParent=*/false);
+		check(NewParamData);
+		return NewParamData;
+	}
 
-		for (int i = 0; i < NumEntries; ++i)
-		{
-			// Store the index offset by one in the attribute
-			NewAttribute->SetValue(NewParamData->Metadata->AddEntry(), i + 1);
-		}
+	void ConnectToSource(FPCGContext* Context, const UPCGData* InData)
+	{
+		FPCGTaggedData& InputData = Context->InputData.TaggedData.Emplace_GetRef();
+		InputData.Data = InData;
+		InputData.Pin = PCGCopyPointsConstants::SourcePointsLabel;
+	}
 
+	void ConnectToTarget(FPCGContext* Context, const UPCGData* InData)
+	{
 		FPCGTaggedData& InputData = Context->InputData.TaggedData.Emplace_GetRef();
-		InputData.Data = NewParamData;
-		InputData.Pin = PCGPinConstants::DefaultInputLabel;
+		InputData.Data = InData;
+		InputData.Pin = PCGCopyPointsConstants::TargetPointsLabel;
+	}
 
-		return NewParamData;
+	void ConnectToSourceAndTarget(FPCGContext* Context, const UPCGData* InData)
+	{
+		ConnectToSource(Context, InData);
+		ConnectToTarget(Context, InData);
 	}
 }
 
 bool FPCGCopyAttributeTests_Points_PropertyToProperty::RunTest(const FString& Parameters)
 {
 	PCGTestsCommon::FTestData TestData;
-	UPCGMetadataOperationSettings* Settings = PCGTestsCommon::GenerateSettings<UPCGMetadataOperationSettings>(TestData);
+	UPCGCopyAttributesSettings* Settings = PCGTestsCommon::GenerateSettings<UPCGCopyAttributesSettings>(TestData);
 	check(Settings);
 
 	constexpr int NumOfPoints = 20;
@@ -101,6 +97,7 @@ bool FPCGCopyAttributeTests_Points_PropertyToProperty::RunTest(const FString& Pa
 
 	TUniquePtr<FPCGContext> Context = TestData.InitializeTestContext();
 	const UPCGPointData* InputPointData = PCGCopyAttributeTests::CreateInputPointData(Context.Get(), NumOfPoints);
+	PCGCopyAttributeTests::ConnectToSourceAndTarget(Context.Get(), InputPointData);
 
 	FPCGElementPtr TestElement = TestData.Settings->GetElement();
 	while (!TestElement->Execute(Context.Get())) {}
@@ -126,7 +123,7 @@ bool FPCGCopyAttributeTests_Points_PropertyToProperty::RunTest(const FString& Pa
 bool FPCGCopyAttributeTests_Points_PropertyToAttribute::RunTest(const FString& Parameters)
 {
 	PCGTestsCommon::FTestData TestData;
-	UPCGMetadataOperationSettings* Settings = PCGTestsCommon::GenerateSettings<UPCGMetadataOperationSettings>(TestData);
+	UPCGCopyAttributesSettings* Settings = PCGTestsCommon::GenerateSettings<UPCGCopyAttributesSettings>(TestData);
 	check(Settings);
 
 	constexpr int NumOfPoints = 20;
@@ -138,6 +135,7 @@ bool FPCGCopyAttributeTests_Points_PropertyToAttribute::RunTest(const FString& P
 
 	TUniquePtr<FPCGContext> Context = TestData.InitializeTestContext();
 	const UPCGPointData* InputPointData = PCGCopyAttributeTests::CreateInputPointData(Context.Get(), NumOfPoints);
+	PCGCopyAttributeTests::ConnectToSourceAndTarget(Context.Get(), InputPointData);
 
 	FPCGElementPtr TestElement = TestData.Settings->GetElement();
 	while (!TestElement->Execute(Context.Get())) {}
@@ -167,17 +165,19 @@ bool FPCGCopyAttributeTests_Points_PropertyToAttribute::RunTest(const FString& P
 bool FPCGCopyAttributeTests_Points_AttributeToProperty::RunTest(const FString& Parameters)
 {
 	PCGTestsCommon::FTestData TestData;
-	UPCGMetadataOperationSettings* Settings = PCGTestsCommon::GenerateSettings<UPCGMetadataOperationSettings>(TestData);
+	UPCGCopyAttributesSettings* Settings = PCGTestsCommon::GenerateSettings<UPCGCopyAttributesSettings>(TestData);
 	check(Settings);
 
 	constexpr int NumOfPoints = 20;
 
-	// Write Attribute in Position.Y
+	// Write Attribute in Position
 	Settings->InputSource.SetAttributeName(PCGCopyAttributeTests::AttributeName);
-	Settings->OutputTarget.Update(TEXT("$Position.Y"));
+	Settings->OutputTarget.Update(TEXT("$Position"));
 
 	TUniquePtr<FPCGContext> Context = TestData.InitializeTestContext();
-	const UPCGPointData* InputPointData = PCGCopyAttributeTests::CreateInputPointData(Context.Get(), NumOfPoints);
+	UPCGPointData* InputPointData = PCGCopyAttributeTests::CreateInputPointData(Context.Get(), NumOfPoints);
+	PCGTestsCommon::CreateAndFillRandomAttribute<FVector>(InputPointData, PCGCopyAttributeTests::AttributeName, FVector::ZeroVector, NumOfPoints);
+	PCGCopyAttributeTests::ConnectToSourceAndTarget(Context.Get(), InputPointData);
 
 	FPCGElementPtr TestElement = TestData.Settings->GetElement();
 	while (!TestElement->Execute(Context.Get())) {}
@@ -189,14 +189,14 @@ bool FPCGCopyAttributeTests_Points_AttributeToProperty::RunTest(const FString& P
 	UTEST_EQUAL("There is the right number of points in output", OutputData->GetPoints().Num(), NumOfPoints);
 
 	check(OutputData && OutputData->Metadata);
-	const FPCGMetadataAttribute<double>* InputAttribute = OutputData->Metadata->GetConstTypedAttribute<double>(PCGCopyAttributeTests::AttributeName);
+	const FPCGMetadataAttribute<FVector>* InputAttribute = OutputData->Metadata->GetConstTypedAttribute<FVector>(PCGCopyAttributeTests::AttributeName);
 	UTEST_NOT_NULL("Input Attribute exists in the output data", InputAttribute);
 
 	// Attribute value was copied in Position.Y correctly
 	for (int i = 0; i < NumOfPoints; ++i)
 	{
 		const FPCGPoint& OutputPoint = OutputData->GetPoints()[i];
-		UTEST_EQUAL(*FString::Printf(TEXT("Position.Y has the same value than the input attribute for point %d"), i), OutputPoint.Transform.GetLocation().Y, InputAttribute->GetValueFromItemKey(OutputPoint.MetadataEntry));
+		UTEST_EQUAL(*FString::Printf(TEXT("Position has the same value than the input attribute for point %d"), i), OutputPoint.Transform.GetLocation(), InputAttribute->GetValueFromItemKey(OutputPoint.MetadataEntry));
 	}
 
 	return true;
@@ -205,18 +205,20 @@ bool FPCGCopyAttributeTests_Points_AttributeToProperty::RunTest(const FString& P
 bool FPCGCopyAttributeTests_Points_AttributeToAttribute::RunTest(const FString& Parameters)
 {
 	PCGTestsCommon::FTestData TestData;
-	UPCGMetadataOperationSettings* Settings = PCGTestsCommon::GenerateSettings<UPCGMetadataOperationSettings>(TestData);
+	UPCGCopyAttributesSettings* Settings = PCGTestsCommon::GenerateSettings<UPCGCopyAttributesSettings>(TestData);
 	check(Settings);
 
 	constexpr int NumOfPoints = 20;
 	const FName OutputAttributeName = TEXT("OutputAttr");
 
-	// Write Input attribute to Output attribute
+	// Write Input attribute to Output attribute, testing with strings
 	Settings->InputSource.SetAttributeName(PCGCopyAttributeTests::AttributeName);
 	Settings->OutputTarget.SetAttributeName(OutputAttributeName);
 
 	TUniquePtr<FPCGContext> Context = TestData.InitializeTestContext();
-	const UPCGPointData* InputPointData = PCGCopyAttributeTests::CreateInputPointData(Context.Get(), NumOfPoints);
+	UPCGPointData* InputPointData = PCGCopyAttributeTests::CreateInputPointData(Context.Get(), NumOfPoints);
+	PCGTestsCommon::CreateAndFillRandomAttribute<FString>(InputPointData, PCGCopyAttributeTests::AttributeName, FString(), NumOfPoints);
+	PCGCopyAttributeTests::ConnectToSourceAndTarget(Context.Get(), InputPointData);
 
 	FPCGElementPtr TestElement = TestData.Settings->GetElement();
 	while (!TestElement->Execute(Context.Get())) {}
@@ -229,9 +231,9 @@ bool FPCGCopyAttributeTests_Points_AttributeToAttribute::RunTest(const FString&
 
 	check(OutputData && OutputData->Metadata);
 
-	const FPCGMetadataAttribute<double>* InputAttribute = OutputData->Metadata->GetConstTypedAttribute<double>(PCGCopyAttributeTests::AttributeName);
+	const FPCGMetadataAttribute<FString>* InputAttribute = OutputData->Metadata->GetConstTypedAttribute<FString>(PCGCopyAttributeTests::AttributeName);
 	UTEST_NOT_NULL("Input Attribute exists in the output data", InputAttribute);
-	const FPCGMetadataAttribute<double>* OutputAttribute = OutputData->Metadata->GetConstTypedAttribute<double>(OutputAttributeName);
+	const FPCGMetadataAttribute<FString>* OutputAttribute = OutputData->Metadata->GetConstTypedAttribute<FString>(OutputAttributeName);
 	UTEST_NOT_NULL("Output Attribute exists in the output data", OutputAttribute);
 
 	check(InputAttribute && OutputAttribute);
@@ -251,7 +253,7 @@ bool FPCGCopyAttributeTests_Points_AttributeToAttribute::RunTest(const FString&
 bool FPCGCopyAttributeTests_Points_CopyingToItself::RunTest(const FString& Parameters)
 {
 	PCGTestsCommon::FTestData TestData;
-	UPCGMetadataOperationSettings* Settings = PCGTestsCommon::GenerateSettings<UPCGMetadataOperationSettings>(TestData);
+	UPCGCopyAttributesSettings* Settings = PCGTestsCommon::GenerateSettings<UPCGCopyAttributesSettings>(TestData);
 	check(Settings);
 
 	constexpr int NumOfPoints = 20;
@@ -261,7 +263,10 @@ bool FPCGCopyAttributeTests_Points_CopyingToItself::RunTest(const FString& Param
 	Settings->OutputTarget.SetAttributeName(PCGCopyAttributeTests::AttributeName);
 
 	TUniquePtr<FPCGContext> Context = TestData.InitializeTestContext();
+	// Need constness for UTEST_EQUAL....
 	const UPCGPointData* InputPointData = PCGCopyAttributeTests::CreateInputPointData(Context.Get(), NumOfPoints);
+	PCGTestsCommon::CreateAndFillRandomAttribute<double>(const_cast<UPCGPointData*>(InputPointData), PCGCopyAttributeTests::AttributeName, 0.0, NumOfPoints);
+	PCGCopyAttributeTests::ConnectToSourceAndTarget(Context.Get(), InputPointData);
 
 	FPCGElementPtr TestElement = TestData.Settings->GetElement();
 	while (!TestElement->Execute(Context.Get())) {}
@@ -278,7 +283,7 @@ bool FPCGCopyAttributeTests_Points_CopyingToItself::RunTest(const FString& Param
 bool FPCGCopyAttributeTests_Points_CopyingAllToItself::RunTest(const FString& Parameters)
 {
 	PCGTestsCommon::FTestData TestData;
-	UPCGMetadataOperationSettings* Settings = PCGTestsCommon::GenerateSettings<UPCGMetadataOperationSettings>(TestData);
+	UPCGCopyAttributesSettings* Settings = PCGTestsCommon::GenerateSettings<UPCGCopyAttributesSettings>(TestData);
 	check(Settings);
 
 	constexpr int NumOfPoints = 20;
@@ -287,6 +292,8 @@ bool FPCGCopyAttributeTests_Points_CopyingAllToItself::RunTest(const FString& Pa
 
 	TUniquePtr<FPCGContext> Context = TestData.InitializeTestContext();
 	const UPCGPointData* InputPointData = PCGCopyAttributeTests::CreateInputPointData(Context.Get(), NumOfPoints);
+	PCGTestsCommon::CreateAndFillRandomAttribute<double>(const_cast<UPCGPointData*>(InputPointData), PCGCopyAttributeTests::AttributeName, 0.0, NumOfPoints);
+	PCGCopyAttributeTests::ConnectToSourceAndTarget(Context.Get(), InputPointData);
 
 	FPCGElementPtr TestElement = TestData.Settings->GetElement();
 	while (!TestElement->Execute(Context.Get())) {}
@@ -303,7 +310,7 @@ bool FPCGCopyAttributeTests_Points_CopyingAllToItself::RunTest(const FString& Pa
 bool FPCGCopyAttributeTests_Params_SingleValue::RunTest(const FString& Parameters)
 {
 	PCGTestsCommon::FTestData TestData;
-	UPCGMetadataOperationSettings* Settings = PCGTestsCommon::GenerateSettings<UPCGMetadataOperationSettings>(TestData);
+	UPCGCopyAttributesSettings* Settings = PCGTestsCommon::GenerateSettings<UPCGCopyAttributesSettings>(TestData);
 	check(Settings);
 
 	const FName OutputAttributeName = TEXT("OutputAttr");
@@ -313,7 +320,9 @@ bool FPCGCopyAttributeTests_Params_SingleValue::RunTest(const FString& Parameter
 	Settings->OutputTarget.SetAttributeName(OutputAttributeName);
 
 	TUniquePtr<FPCGContext> Context = TestData.InitializeTestContext();
-	const UPCGParamData* InputParamData = PCGCopyAttributeTests::CreateInputParamData(Context.Get(), 1);
+	const UPCGParamData* InputParamData = PCGCopyAttributeTests::CreateInputParamData(Context.Get());
+	PCGTestsCommon::CreateAndFillRandomAttribute<FRotator>(const_cast<UPCGParamData*>(InputParamData), PCGCopyAttributeTests::AttributeName, FRotator::ZeroRotator, 1);
+	PCGCopyAttributeTests::ConnectToSourceAndTarget(Context.Get(), InputParamData);
 
 	FPCGElementPtr TestElement = TestData.Settings->GetElement();
 	while (!TestElement->Execute(Context.Get())) {}
@@ -326,9 +335,9 @@ bool FPCGCopyAttributeTests_Params_SingleValue::RunTest(const FString& Parameter
 
 	check(OutputData && OutputData->Metadata);
 
-	const FPCGMetadataAttribute<double>* InputAttribute = OutputData->Metadata->GetConstTypedAttribute<double>(PCGCopyAttributeTests::AttributeName);
+	const FPCGMetadataAttribute<FRotator>* InputAttribute = OutputData->Metadata->GetConstTypedAttribute<FRotator>(PCGCopyAttributeTests::AttributeName);
 	UTEST_NOT_NULL("Input Attribute exists in the output data", InputAttribute);
-	const FPCGMetadataAttribute<double>* OutputAttribute = OutputData->Metadata->GetConstTypedAttribute<double>(OutputAttributeName);
+	const FPCGMetadataAttribute<FRotator>* OutputAttribute = OutputData->Metadata->GetConstTypedAttribute<FRotator>(OutputAttributeName);
 	UTEST_NOT_NULL("Output Attribute exists in the output data", OutputAttribute);
 
 	check(InputAttribute && OutputAttribute);
@@ -342,7 +351,7 @@ bool FPCGCopyAttributeTests_Params_SingleValue::RunTest(const FString& Parameter
 bool FPCGCopyAttributeTests_Params_MultiValue::RunTest(const FString& Parameters)
 {
 	PCGTestsCommon::FTestData TestData;
-	UPCGMetadataOperationSettings* Settings = PCGTestsCommon::GenerateSettings<UPCGMetadataOperationSettings>(TestData);
+	UPCGCopyAttributesSettings* Settings = PCGTestsCommon::GenerateSettings<UPCGCopyAttributesSettings>(TestData);
 	check(Settings);
 
 	constexpr int NumEntries = 20;
@@ -353,7 +362,9 @@ bool FPCGCopyAttributeTests_Params_MultiValue::RunTest(const FString& Parameters
 	Settings->OutputTarget.SetAttributeName(OutputAttributeName);
 
 	TUniquePtr<FPCGContext> Context = TestData.InitializeTestContext();
-	const UPCGParamData* InputParamData = PCGCopyAttributeTests::CreateInputParamData(Context.Get(), NumEntries);
+	const UPCGParamData* InputParamData = PCGCopyAttributeTests::CreateInputParamData(Context.Get());
+	PCGTestsCommon::CreateAndFillRandomAttribute<FSoftObjectPath>(const_cast<UPCGParamData*>(InputParamData), PCGCopyAttributeTests::AttributeName, FSoftObjectPath(), NumEntries);
+	PCGCopyAttributeTests::ConnectToSourceAndTarget(Context.Get(), InputParamData);
 
 	FPCGElementPtr TestElement = TestData.Settings->GetElement();
 	while (!TestElement->Execute(Context.Get())) {}
@@ -366,9 +377,9 @@ bool FPCGCopyAttributeTests_Params_MultiValue::RunTest(const FString& Parameters
 
 	check(OutputData && OutputData->Metadata);
 
-	const FPCGMetadataAttribute<double>* InputAttribute = OutputData->Metadata->GetConstTypedAttribute<double>(PCGCopyAttributeTests::AttributeName);
+	const FPCGMetadataAttribute<FSoftObjectPath>* InputAttribute = OutputData->Metadata->GetConstTypedAttribute<FSoftObjectPath>(PCGCopyAttributeTests::AttributeName);
 	UTEST_NOT_NULL("Input Attribute exists in the output data", InputAttribute);
-	const FPCGMetadataAttribute<double>* OutputAttribute = OutputData->Metadata->GetConstTypedAttribute<double>(OutputAttributeName);
+	const FPCGMetadataAttribute<FSoftObjectPath>* OutputAttribute = OutputData->Metadata->GetConstTypedAttribute<FSoftObjectPath>(OutputAttributeName);
 	UTEST_NOT_NULL("Output Attribute exists in the output data", OutputAttribute);
 
 	check(InputAttribute && OutputAttribute);
@@ -387,7 +398,7 @@ bool FPCGCopyAttributeTests_Params_MultiValue::RunTest(const FString& Parameters
 bool FPCGCopyAttributeTests_Params_CopyingToItself::RunTest(const FString& Parameters)
 {
 	PCGTestsCommon::FTestData TestData;
-	UPCGMetadataOperationSettings* Settings = PCGTestsCommon::GenerateSettings<UPCGMetadataOperationSettings>(TestData);
+	UPCGCopyAttributesSettings* Settings = PCGTestsCommon::GenerateSettings<UPCGCopyAttributesSettings>(TestData);
 	check(Settings);
 
 	// Write Input attribute to Output attribute
@@ -395,7 +406,9 @@ bool FPCGCopyAttributeTests_Params_CopyingToItself::RunTest(const FString& Param
 	Settings->OutputTarget.SetAttributeName(PCGCopyAttributeTests::AttributeName);
 
 	TUniquePtr<FPCGContext> Context = TestData.InitializeTestContext();
-	const UPCGParamData* InputParamData = PCGCopyAttributeTests::CreateInputParamData(Context.Get(), 1);
+	const UPCGParamData* InputParamData = PCGCopyAttributeTests::CreateInputParamData(Context.Get());
+	PCGTestsCommon::CreateAndFillRandomAttribute<double>(const_cast<UPCGParamData*>(InputParamData), PCGCopyAttributeTests::AttributeName, 0.0, 1);
+	PCGCopyAttributeTests::ConnectToSourceAndTarget(Context.Get(), InputParamData);
 
 	FPCGElementPtr TestElement = TestData.Settings->GetElement();
 	while (!TestElement->Execute(Context.Get())) {}
@@ -412,13 +425,15 @@ bool FPCGCopyAttributeTests_Params_CopyingToItself::RunTest(const FString& Param
 bool FPCGCopyAttributeTests_Params_CopyingAllToItself::RunTest(const FString& Parameters)
 {
 	PCGTestsCommon::FTestData TestData;
-	UPCGMetadataOperationSettings* Settings = PCGTestsCommon::GenerateSettings<UPCGMetadataOperationSettings>(TestData);
+	UPCGCopyAttributesSettings* Settings = PCGTestsCommon::GenerateSettings<UPCGCopyAttributesSettings>(TestData);
 	check(Settings);
 
 	Settings->bCopyAllAttributes = true;
 
 	TUniquePtr<FPCGContext> Context = TestData.InitializeTestContext();
-	const UPCGParamData* InputParamData = PCGCopyAttributeTests::CreateInputParamData(Context.Get(), 1);
+	const UPCGParamData* InputParamData = PCGCopyAttributeTests::CreateInputParamData(Context.Get());
+	PCGTestsCommon::CreateAndFillRandomAttribute<double>(const_cast<UPCGParamData*>(InputParamData), PCGCopyAttributeTests::AttributeName, 0.0, 1);
+	PCGCopyAttributeTests::ConnectToSourceAndTarget(Context.Get(), InputParamData);
 
 	FPCGElementPtr TestElement = TestData.Settings->GetElement();
 	while (!TestElement->Execute(Context.Get())) {}
diff --git a/Engine/Plugins/PCG/Source/PCG/Public/Elements/Metadata/PCGMetadataElement.h b/Engine/Plugins/PCG/Source/PCG/Public/Elements/Metadata/PCGMetadataElement.h
--- a/Engine/Plugins/PCG/Source/PCG/Public/Elements/Metadata/PCGMetadataElement.h
+++ b/Engine/Plugins/PCG/Source/PCG/Public/Elements/Metadata/PCGMetadataElement.h
@@ -2,57 +2,22 @@
 
 #pragma once
 
-#include "PCGSettings.h"
-#include "Metadata/PCGAttributePropertySelector.h"
+#include "Elements/PCGCopyAttributes.h"
 
 #include "PCGMetadataElement.generated.h"
 
+/** Deprecated, use UPCGCopyAttributesSettings */
 UCLASS(MinimalAPI, BlueprintType, ClassGroup = (Procedural))
-class UPCGMetadataOperationSettings : public UPCGSettings
+class UE_DEPRECATED(5.5, "Use UPCGCopyAttributeSettings") UPCGMetadataOperationSettings : public UPCGCopyAttributesSettings
 {
 	GENERATED_BODY()
 
 public:
-	UPCGMetadataOperationSettings();
-
-	//~Begin UObject interface
-	virtual void PostLoad() override;
-	//~End UObject interface
-
-	//~Begin UPCGSettings interface
+	UPCGMetadataOperationSettings()
+		: UPCGCopyAttributesSettings()
+	{
 #if WITH_EDITOR
-	virtual FName GetDefaultNodeName() const override { return FName(TEXT("CopyAttribute")); }
-	virtual FText GetDefaultNodeTitle() const override { return NSLOCTEXT("PCGMetadataOperationSettings", "NodeTitle", "Copy Attribute"); }
-	virtual FText GetNodeTooltipText() const;
-	virtual EPCGSettingsType GetType() const override { return EPCGSettingsType::Metadata; }
-	virtual bool IsPinUsedByNodeExecution(const UPCGPin* InPin) const override;
-	virtual void ApplyDeprecation(UPCGNode* InOutNode) override;
-	virtual bool HasDynamicPins() const override { return true; }
-#endif
-	virtual FString GetAdditionalTitleInformation() const override;
-
-protected:
-#if WITH_EDITOR
-	virtual EPCGChangeType GetChangeTypeForProperty(const FName& InPropertyName) const override;
-#endif
-	virtual TArray<FPCGPinProperties> InputPinProperties() const override;
-	virtual TArray<FPCGPinProperties> OutputPinProperties() const override;
-	virtual FPCGElementPtr CreateElement() const override;
-	//~End UPCGSettings interface
-
-public:
-	UPROPERTY(BlueprintReadWrite, EditAnywhere, Category = Settings, meta = (EditCondition="!bCopyAllAttributes", PCG_Overridable))
-	FPCGAttributePropertyInputSelector InputSource;
-
-	UPROPERTY(BlueprintReadWrite, EditAnywhere, Category = Settings, meta = (EditCondition = "!bCopyAllAttributes", PCG_Overridable))
-	FPCGAttributePropertyOutputSelector OutputTarget;
-
-	UPROPERTY(BlueprintReadWrite, EditAnywhere, Category = Settings, meta = (PCG_Overridable))
-	bool bCopyAllAttributes = false;
-};
-
-class FPCGMetadataOperationElement : public IPCGElement
-{
-protected:
-	virtual bool ExecuteInternal(FPCGContext* Context) const override;
-};
+		bExposeToLibrary = false;
+#endif // WITH_EDITOR
+	}
+};
\ No newline at end of file
diff --git a/Engine/Plugins/PCG/Source/PCG/Public/Elements/PCGAttributeTransferElement.h b/Engine/Plugins/PCG/Source/PCG/Public/Elements/PCGAttributeTransferElement.h
--- a/Engine/Plugins/PCG/Source/PCG/Public/Elements/PCGAttributeTransferElement.h
+++ b/Engine/Plugins/PCG/Source/PCG/Public/Elements/PCGAttributeTransferElement.h
@@ -2,78 +2,25 @@
 
 #pragma once
 
-#include "PCGSettings.h"
+#include "Elements/PCGCopyAttributes.h"
 
 #include "PCGAttributeTransferElement.generated.h"
 
-class UPCGPointData;
-class UPCGSpatialData;
-
-namespace PCGAttributeTransferConstants
-{
-	const FName NodeName = TEXT("TransferAttribute");
-	const FText NodeTitle = NSLOCTEXT("PCGAttributeTransferElement", "NodeTitle", "Transfer Attribute");
-	const FName SourceLabel = TEXT("Source");
-	const FName TargetLabel = TEXT("Target");
-}
-
-/**
-* Transfer an attribute from a source metadata to a target data.
-* Only support Spatial to Spatial and Points to Points, and they need to match.
-*  - For Spatial data, number of entries in the metadata should be the same between source and target.
-*  - For Point data, number of points should be the same between source and target.
-* 
-* To do the same but with a Source param data, use CreateAttribute.
-*/
 UCLASS(MinimalAPI, BlueprintType, ClassGroup = (Procedural))
-class UPCGAttributeTransferSettings : public UPCGSettings
+class UE_DEPRECATED(5.5, "Use UPCGCopyAttributeSettings") UPCGAttributeTransferSettings : public UPCGCopyAttributesSettings
 {
 	GENERATED_BODY()
 
 public:
-	// ~Begin UObject interface
-	virtual void PostLoad() override;
-	// ~End UObject interface
-
-	//~Begin UPCGSettings interface
+	UPCGAttributeTransferSettings()
+		: UPCGCopyAttributesSettings()
+	{
 #if WITH_EDITOR
-	virtual FName GetDefaultNodeName() const override;
-	virtual FText GetDefaultNodeTitle() const override;
-	virtual EPCGSettingsType GetType() const override { return EPCGSettingsType::Metadata; }
-	virtual FText GetNodeTooltipText() const override;
-	virtual bool HasDynamicPins() const override { return true; }
-	virtual void ApplyDeprecation(UPCGNode* InOutNode) override;
-#endif
-	virtual EPCGDataType GetCurrentPinTypes(const UPCGPin* InPin) const override;
-	virtual FString GetAdditionalTitleInformation() const override;
+		bExposeToLibrary = false;
+#endif // WITH_EDITOR
 
-protected:
-#if WITH_EDITOR
-	virtual EPCGChangeType GetChangeTypeForProperty(const FName& InPropertyName) const override { return Super::GetChangeTypeForProperty(InPropertyName) | EPCGChangeType::Cosmetic; }
-#endif
-	virtual TArray<FPCGPinProperties> InputPinProperties() const override;
-	virtual TArray<FPCGPinProperties> OutputPinProperties() const override;
-	virtual FPCGElementPtr CreateElement() const override;
-	//~End UPCGSettings interface
-
-public:
-	UPROPERTY(BlueprintReadWrite, EditAnywhere, Category = Settings, meta = (PCG_Overridable))
-	FPCGAttributePropertyInputSelector SourceAttributeProperty;
-
-	UPROPERTY(BlueprintReadWrite, EditAnywhere, Category = Settings, meta = (PCG_Overridable))
-	FPCGAttributePropertyOutputSelector TargetAttributeProperty;
-
-#if WITH_EDITORONLY_DATA
-	UPROPERTY()
-	FName SourceAttributeName_DEPRECATED;
-
-	UPROPERTY()
-	FName TargetAttributeName_DEPRECATED;
-#endif
-};
-
-class FPCGAttributeTransferElement : public IPCGElement
-{
-protected:
-	virtual bool ExecuteInternal(FPCGContext* Context) const override;
+		// Reset default values to the ones of Attribute Transfer, not inherited from CopyAttributes
+		InputSource = FPCGAttributePropertyInputSelector();
+		OutputTarget = FPCGAttributePropertyOutputSelector();
+	}
 };
diff --git a/Engine/Plugins/PCG/Source/PCG/Public/Elements/PCGCopyAttributes.h b/Engine/Plugins/PCG/Source/PCG/Public/Elements/PCGCopyAttributes.h
new file mode 100644
--- /dev/null
+++ b/Engine/Plugins/PCG/Source/PCG/Public/Elements/PCGCopyAttributes.h
@@ -0,0 +1,89 @@
+// Copyright Epic Games, Inc. All Rights Reserved.
+
+#pragma once
+
+#include "PCGSettings.h"
+#include "Metadata/PCGAttributePropertySelector.h"
+
+#include "PCGCopyAttributes.generated.h"
+
+UENUM()
+enum class EPCGCopyAttributesOperation
+{
+	/** 
+	* N:N operation, N:1 or 1:N Operation. If there are SourceA/SourceB and TargetA/TagretB, SourceA will be copied to TargetA and SourceB to TargetB.
+	* If there are SourceA and TargetA/TargetB, SourceA will be copied to TargetA and TargetB.
+	* If there are SourceA/SourceB and TargetA, SourceA will be copied to TargetA and SourceB to a copy of TargetA.
+	* Produces Max(N,M) data, N being the number of targets and M being the number of sources. Either N == M, or N == 1 or M == 1.
+	*/
+	CopyEachSourceToEachTargetRespectively,
+
+	/**
+	* N:M operation. If there are SourceA/SourceB and TargetA/TagretB, SourceA and SourceB will be copied to TargetA and also to TargetB.
+	* Copy will be sequential, so if attribute names are clashing, they will be overwritten by the last Source.
+	* Produces N data, N being the number of targets.
+	*/
+	MergeSourcesAndCopyToAllTargets,
+
+	/**
+	* N:M operation. If there are SourceA/SourceB and TargetA/TagretB, we will have those copies: SourceA -> TargetA, SourceA -> TargetB, SourceB -> TargetA, SourceB -> TargetB.
+	* Produces N*M data, N being the number of targets and M the number of sources.
+	*/
+	CopyEachSourceOnEveryTarget
+};
+
+UCLASS(MinimalAPI, BlueprintType, ClassGroup = (Procedural))
+class UPCGCopyAttributesSettings : public UPCGSettings
+{
+	GENERATED_BODY()
+
+public:
+	UPCGCopyAttributesSettings();
+
+	//~Begin UPCGSettings interface
+#if WITH_EDITOR
+	virtual FName GetDefaultNodeName() const override { return FName(TEXT("CopyAttributes")); }
+	virtual FText GetDefaultNodeTitle() const override { return NSLOCTEXT("PCGMetadataOperationSettings", "NodeTitle", "Copy Attributes"); }
+	virtual FText GetNodeTooltipText() const;
+	virtual EPCGSettingsType GetType() const override { return EPCGSettingsType::Metadata; }
+	virtual void ApplyDeprecation(UPCGNode* InOutNode) override;
+	virtual void ApplyDeprecationBeforeUpdatePins(UPCGNode* InOutNode, TArray<TObjectPtr<UPCGPin>>& InputPins, TArray<TObjectPtr<UPCGPin>>& OutputPins) override;
+	virtual EPCGDataType GetCurrentPinTypes(const UPCGPin* InPin) const override;
+	virtual bool HasDynamicPins() const override { return true; }
+#endif
+	virtual FString GetAdditionalTitleInformation() const override;
+
+protected:
+#if WITH_EDITOR
+	virtual EPCGChangeType GetChangeTypeForProperty(const FName& InPropertyName) const override;
+#endif
+	virtual TArray<FPCGPinProperties> InputPinProperties() const override;
+	virtual TArray<FPCGPinProperties> OutputPinProperties() const override;
+	virtual FPCGElementPtr CreateElement() const override;
+	//~End UPCGSettings interface
+
+public:
+	UPROPERTY(BlueprintReadWrite, EditAnywhere, Category = Settings, meta = (PCG_Overridable))
+	EPCGCopyAttributesOperation Operation = EPCGCopyAttributesOperation::CopyEachSourceToEachTargetRespectively;
+
+	UPROPERTY(BlueprintReadWrite, EditAnywhere, Category = Settings, meta = (PCG_Overridable))
+	bool bCopyAllAttributes = false;
+
+	UPROPERTY(BlueprintReadWrite, EditAnywhere, Category = Settings, meta = (EditCondition="!bCopyAllAttributes", PCG_Overridable))
+	FPCGAttributePropertyInputSelector InputSource;
+
+	UPROPERTY(BlueprintReadWrite, EditAnywhere, Category = Settings, meta = (EditCondition = "!bCopyAllAttributes", PCG_Overridable))
+	FPCGAttributePropertyOutputSelector OutputTarget;
+
+private:
+#if WITH_EDITOR
+	// To know if we need to reconnect Target to Source.
+	bool bWasDeprecated = false;
+#endif // WITH_EDITOR
+};
+
+class FPCGCopyAttributesElement : public IPCGElement
+{
+protected:
+	virtual bool ExecuteInternal(FPCGContext* Context) const override;
+};
diff --git a/Engine/Plugins/PCG/Source/PCG/Public/Tests/PCGTestsCommon.h b/Engine/Plugins/PCG/Source/PCG/Public/Tests/PCGTestsCommon.h
--- a/Engine/Plugins/PCG/Source/PCG/Public/Tests/PCGTestsCommon.h
+++ b/Engine/Plugins/PCG/Source/PCG/Public/Tests/PCGTestsCommon.h
@@ -3,7 +3,12 @@
 #pragma once
 
 #include "PCGData.h"
+#include "PCGPoint.h"
 #include "PCGSettings.h"
+#include "Data/PCGPointData.h"
+#include "Metadata/PCGMetadata.h"
+#include "Metadata/PCGMetadataAttributeTpl.h"
+#include "Metadata/PCGMetadataAttributeTraits.h"
 
 #include "GameFramework/Actor.h"
 #include "Misc/AutomationTest.h"
@@ -12,14 +17,12 @@ class IPCGElement;
 class UPCGComponent;
 class UPCGNode;
 class UPCGParamData;
-class UPCGPointData;
 class UPCGPolyLineData;
 class UPCGPrimitiveData;
 class UPCGSurfaceData;
 class UPCGVolumeData;
 struct FPCGContext;
 struct FPCGPinProperties;
-struct FPCGPoint;
 
 namespace PCGTestsCommon
 {
@@ -89,6 +92,118 @@ namespace PCGTestsCommon
 
 		return TypedSettings;
 	}
+
+	// Numerical
+	template <typename T, typename std::enable_if_t<std::is_arithmetic_v<T>, bool> = true>
+	T GenerateRandomValue(FRandomStream& RandomStream) requires std::is_arithmetic_v<T>
+	{
+		if constexpr (std::is_same_v<bool, T>)
+		{
+			return RandomStream.FRand() > 0.5;
+		}
+		else
+		{
+			return T(RandomStream.FRandRange(-50.0, 50.0));
+		}
+	}
+
+	 // Vectors
+	template <typename T, typename std::enable_if_t<PCG::Private::IsOfTypes<T, FVector2D, FVector, FVector4>(), bool> = true>
+	T GenerateRandomValue(FRandomStream& RandomStream)
+	{
+		return T(RandomStream.VRand());
+	}
+
+	 // Rotators/Quat
+	template <typename T, typename std::enable_if_t<PCG::Private::IsOfTypes<T, FRotator, FQuat>(), bool> = true>
+	T GenerateRandomValue(FRandomStream& RandomStream)
+	{
+		return T::MakeFromEuler(RandomStream.VRand() * 360.0);
+	}
+
+	 // String/FName
+	template <typename T, typename std::enable_if_t<PCG::Private::IsOfTypes<T, FName, FString>(), bool> = true>
+	T GenerateRandomValue(FRandomStream& RandomStream)
+	{
+		static constexpr const TCHAR* Dictionary[] = { TEXT("Foo"), TEXT("Bar"), TEXT("PCG"), TEXT("YOLO"), TEXT("Bla") };
+		constexpr int32 NumWords = static_cast<int32>(sizeof(Dictionary) / sizeof(const TCHAR*));
+		constexpr int32 MaxNumWords = 5;
+		TStringBuilder<32> Builder;
+
+		for (uint32 i = 0; i < (RandomStream.GetUnsignedInt() % (MaxNumWords - 1) + 1); ++i)
+		{
+			if (i != 0)
+			{
+				Builder += TEXT("_");
+			}
+
+			Builder += Dictionary[RandomStream.GetUnsignedInt() % NumWords];
+		}
+
+		return Builder.ToString();
+	}
+
+	// SoftObjectPath
+	template <typename T, typename std::enable_if_t<std::is_same_v<T, FSoftObjectPath>, bool> = true>
+	T GenerateRandomValue(FRandomStream& RandomStream)
+	{
+		static const FSoftObjectPath Dictionary[] = 
+		{ 
+			FSoftObjectPath(TEXT("/PCG/DebugObjects/PCG_Cube.PCG_Cube")),
+			FSoftObjectPath(TEXT("Material'/PCG/DebugObjects/PCG_DebugMaterial.PCG_DebugMaterial'"))
+		};
+
+		constexpr int32 NumWords = static_cast<int32>(sizeof(Dictionary) / sizeof(FSoftObjectPath));
+		return Dictionary[RandomStream.GetUnsignedInt() % NumWords];
+	}
+
+	// FSoftClassPath
+	template <typename T, typename std::enable_if_t<std::is_same_v<T, FSoftClassPath>, bool> = true>
+	T GenerateRandomValue(FRandomStream& RandomStream)
+	{
+		static const FSoftClassPath Dictionary[] =
+		{
+			UPCGData::StaticClass(),
+			UPCGPointData::StaticClass(),
+			UPCGMetadata::StaticClass(),
+		};
+
+		constexpr int32 NumWords = static_cast<int32>(sizeof(Dictionary) / sizeof(FSoftClassPath));
+		return Dictionary[RandomStream.GetUnsignedInt() % NumWords];
+	}
+
+	template <typename T>
+	void CreateAndFillRandomAttribute(UPCGData* InData, const FName AttributeName, T DefaultValue, const int32 NumValues, int32 Seed = 42, const bool* bForceAllowInterpolation = nullptr) 
+	{
+		check(InData);
+
+		UPCGMetadata* Metadata = InData->MutableMetadata();
+		UPCGPointData* PointData = Cast<UPCGPointData>(InData);
+
+		check(Metadata && (!PointData || NumValues == 0 || NumValues == PointData->GetPoints().Num()));
+
+		const bool bAllowInterpolation = bForceAllowInterpolation ? *bForceAllowInterpolation : PCG::Private::MetadataTraits<T>::CanInterpolate;
+		FPCGMetadataAttribute<T>* NewAttribute = Metadata->CreateAttribute<T>(AttributeName, /*DefaultValue=*/std::move(DefaultValue), /*bAllowInterpolation=*/bAllowInterpolation, /*bOverrideParent=*/false);
+		check(NewAttribute);
+
+		FRandomStream RandomStream(Seed);
+
+		for (int i = 0; i < NumValues; ++i)
+		{
+			PCGMetadataEntryKey EntryKey;
+			if (PointData)
+			{
+				Metadata->InitializeOnSet(PointData->GetMutablePoints()[i].MetadataEntry);
+				EntryKey = PointData->GetMutablePoints()[i].MetadataEntry;
+			}
+			else
+			{
+				EntryKey = Metadata->AddEntry();
+			}
+
+			NewAttribute->SetValue(EntryKey, GenerateRandomValue<T>(RandomStream));
+		}
+	}
 }
 
 class FPCGTestBaseClass : public FAutomationTestBase
