Added native node to export data (Save Data Asset).
Added UI element to save collection/specific data being inspected to an asset.

diff --git a/Engine/Plugins/PCG/Source/PCG/Private/Elements/IO/PCGSaveAssetElement.cpp b/Engine/Plugins/PCG/Source/PCG/Private/Elements/IO/PCGSaveAssetElement.cpp
new file mode 100644
--- /dev/null
+++ b/Engine/Plugins/PCG/Source/PCG/Private/Elements/IO/PCGSaveAssetElement.cpp
@@ -0,0 +1,57 @@
+// Copyright Epic Games, Inc. All Rights Reserved.
+
+#include "Elements/IO/PCGSaveAssetElement.h"
+
+#include "PCGAssetExporterUtils.h"
+#include "PCGModule.h"
+
+#include "AssetRegistry/AssetData.h"
+
+bool UPCGDataCollectionExporter::ExportAsset(const FString& PackageName, UPCGDataAsset* Asset)
+{
+	// Relies on default behavior to duplicate if needed
+	Asset->Data = Data;
+	return true;
+}
+
+UPackage* UPCGDataCollectionExporter::UpdateAsset(const FAssetData& PCGAsset)
+{
+	return nullptr;
+}
+
+UPCGSaveDataAssetSettings::UPCGSaveDataAssetSettings()
+{
+	Pins = Super::InputPinProperties();
+}
+
+FPCGElementPtr UPCGSaveDataAssetSettings::CreateElement() const
+{
+	return MakeShared<FPCGSaveDataAssetElement>();
+}
+
+bool FPCGSaveDataAssetElement::ExecuteInternal(FPCGContext* Context) const
+{
+	TRACE_CPUPROFILER_EVENT_SCOPE(FPCGSaveDataAssetElement::Execute);
+	check(Context);
+	const UPCGSaveDataAssetSettings* Settings = Context->GetInputSettings<UPCGSaveDataAssetSettings>();
+	check(Settings);
+
+	UPCGDataCollectionExporter* Exporter = nullptr;
+
+	if (Settings->CustomDataCollectionExporterClass)
+	{
+		Exporter = NewObject<UPCGDataCollectionExporter>(GetTransientPackage(), Settings->CustomDataCollectionExporterClass);
+	}
+	
+	if (!Exporter)
+	{
+		Exporter = NewObject<UPCGDataCollectionExporter>();
+	}
+
+	check(Exporter);
+	Exporter->Data = Context->InputData;
+
+	UPCGAssetExporterUtils::CreateAsset(Exporter, Settings->Params);
+
+	return true;
+}
\ No newline at end of file
diff --git a/Engine/Plugins/PCG/Source/PCG/Private/PCGAssetExporter.cpp b/Engine/Plugins/PCG/Source/PCG/Private/PCGAssetExporter.cpp
--- a/Engine/Plugins/PCG/Source/PCG/Private/PCGAssetExporter.cpp
+++ b/Engine/Plugins/PCG/Source/PCG/Private/PCGAssetExporter.cpp
@@ -161,6 +161,14 @@ void UPCGAssetExporter::DuplicateAndReOuterData(UPCGDataAsset* Asset)
 			UPCGData* DuplicatedData = TaggedData.Data->DuplicateData(nullptr);
 			DuplicatedData->Rename(nullptr, Asset);
 
+			DuplicatedData->VisitDataNetwork([Asset](const UPCGData* InData)
+			{
+				if (InData && InData->GetOuter() == Asset)
+				{
+					const_cast<UPCGData*>(InData)->Flatten();
+				}
+			});
+
 			TaggedData.Data = DuplicatedData;
 		}
 	}
diff --git a/Engine/Plugins/PCG/Source/PCG/Public/Elements/IO/PCGSaveAssetElement.h b/Engine/Plugins/PCG/Source/PCG/Public/Elements/IO/PCGSaveAssetElement.h
new file mode 100644
--- /dev/null
+++ b/Engine/Plugins/PCG/Source/PCG/Public/Elements/IO/PCGSaveAssetElement.h
@@ -0,0 +1,79 @@
+// Copyright Epic Games, Inc. All Rights Reserved.
+
+#pragma once
+
+#include "PCGAssetExporter.h"
+#include "PCGContext.h"
+#include "PCGDataAsset.h"
+#include "PCGSettings.h"
+
+struct FAssetData;
+
+#include "PCGSaveAssetElement.generated.h"
+
+/** Default exporter to save data collections, with no capacity for update. */
+UCLASS(BlueprintType)
+class PCG_API UPCGDataCollectionExporter : public UPCGAssetExporter
+{
+	GENERATED_BODY()
+
+protected:
+	//~Begin UPCGAssetExporter interface
+	virtual bool ExportAsset(const FString& PackageName, UPCGDataAsset* Asset);
+	virtual UPackage* UpdateAsset(const FAssetData& PCGAsset);
+	//~End UPCGAssetExporter interface
+
+public:
+	UPROPERTY(BlueprintReadWrite, VisibleAnywhere, Category = Data)
+	FPCGDataCollection Data;
+};
+
+/** Node that will save input data to a PCG data asset */
+UCLASS(BlueprintType, CLassGroup=(Procedural))
+class PCG_API UPCGSaveDataAssetSettings : public UPCGSettings
+{
+	GENERATED_BODY()
+
+public:
+	UPCGSaveDataAssetSettings();
+
+	//~UPCGSettings interface implementation
+#if WITH_EDITOR
+	virtual FName GetDefaultNodeName() const override { return FName(TEXT("PCGSaveDataAsset")); }
+	virtual FText GetDefaultNodeTitle() const override { return NSLOCTEXT("PCGSaveDataAssetElement", "NodeTitle", "Save PCG Data Asset"); }
+	virtual FText GetNodeTooltipText() const override { return NSLOCTEXT("PCGSaveDataAssetElement", "NodeTooltip", "Exports the input data to a PCG Data Asset."); }
+	virtual EPCGSettingsType GetType() const override { return EPCGSettingsType::InputOutput; }
+#endif // WITH_EDITOR
+
+protected:
+	virtual FPCGElementPtr CreateElement() const override;
+	
+	virtual TArray<FPCGPinProperties> InputPinProperties() const override { return Pins; }
+	virtual TArray<FPCGPinProperties> OutputPinProperties() const override { return TArray<FPCGPinProperties>(); }
+	// ~End UPCGSettings interface
+
+public:
+	UPROPERTY(BlueprintReadWrite, EditAnywhere, Category = Data)
+	TArray<FPCGPinProperties> Pins;
+
+	UPROPERTY(BlueprintReadWrite, EditAnywhere, Category = Data)
+	TSubclassOf<UPCGDataCollectionExporter> CustomDataCollectionExporterClass = nullptr;
+
+	UPROPERTY(BlueprintReadWrite, EditAnywhere, Category = Data, meta = (ShowOnlyInnerProperties, PCG_Overridable))
+	FPCGAssetExporterParameters Params;
+
+	UPROPERTY(BlueprintReadWrite, EditAnywhere, Category = Data, meta = (PCG_Overridable))
+	FString AssetDescription;
+
+	UPROPERTY(BlueprintReadWrite, EditAnywhere, Category = Data, meta = (PCG_Overridable))
+	FLinearColor AssetColor = FLinearColor::White;
+};
+
+class PCG_API FPCGSaveDataAssetElement : public IPCGElement
+{
+public:
+	virtual bool CanExecuteOnlyOnMainThread(FPCGContext* Context) const override { return true; }
+
+protected:
+	virtual bool ExecuteInternal(FPCGContext* Context) const override;
+};
\ No newline at end of file
diff --git a/Engine/Plugins/PCG/Source/PCG/Public/PCGAssetExporter.h b/Engine/Plugins/PCG/Source/PCG/Public/PCGAssetExporter.h
--- a/Engine/Plugins/PCG/Source/PCG/Public/PCGAssetExporter.h
+++ b/Engine/Plugins/PCG/Source/PCG/Public/PCGAssetExporter.h
@@ -27,6 +27,7 @@ struct PCG_API FPCGAssetExporterParameters
 	FString AssetPath;
 
 	/** Controls whether the assets will be saved at the end of the process or not. */
+	UPROPERTY(BlueprintReadWrite, EditAnywhere, Category = Settings)
 	bool bSaveOnExportEnded = true;
 };
 
diff --git a/Engine/Plugins/PCG/Source/PCGEditor/Private/SPCGEditorGraphAttributeListView.cpp b/Engine/Plugins/PCG/Source/PCGEditor/Private/SPCGEditorGraphAttributeListView.cpp
--- a/Engine/Plugins/PCG/Source/PCGEditor/Private/SPCGEditorGraphAttributeListView.cpp
+++ b/Engine/Plugins/PCG/Source/PCGEditor/Private/SPCGEditorGraphAttributeListView.cpp
@@ -2,6 +2,7 @@
 
 #include "SPCGEditorGraphAttributeListView.h"
 
+#include "PCGAssetExporterUtils.h"
 #include "PCGComponent.h"
 #include "PCGData.h"
 #include "PCGDataVisualization.h"
@@ -11,6 +12,7 @@
 #include "PCGPin.h"
 #include "PCGSubsystem.h"
 #include "Data/PCGSpatialData.h"
+#include "Elements/IO/PCGSaveAssetElement.h"
 #include "Metadata/PCGAttributePropertySelector.h"
 #include "Metadata/Accessors/IPCGAttributeAccessor.h"
 #include "Metadata/Accessors/PCGAttributeAccessor.h"
@@ -326,7 +328,7 @@ void SPCGEditorGraphAttributeListView::Construct(const FArguments& InArgs, TShar
 
 	FilterImage->AddLayer(TAttribute<const FSlateBrush*>(this, &SPCGEditorGraphAttributeListView::GetFilterBadgeIcon));
 
-	SAssignNew(LockButton, SButton)
+	TSharedPtr<SButton> LockButton = SNew(SButton)
 		.ButtonStyle(FAppStyle::Get(), "SimpleButton")
 		.OnClicked(this, &SPCGEditorGraphAttributeListView::OnLockClick)
 		.ContentPadding(FMargin(4, 2))
@@ -339,7 +341,7 @@ void SPCGEditorGraphAttributeListView::Construct(const FArguments& InArgs, TShar
 				.Image(this, &SPCGEditorGraphAttributeListView::OnGetLockButtonImageResource)
 		];
 
-	SAssignNew(FilterButton, SComboButton)
+	TSharedPtr<SComboButton> FilterButton = SNew(SComboButton)
 		.ForegroundColor(FSlateColor::UseStyle())
 		.HasDownArrow(false)
 		.OnGetMenuContent(this, &SPCGEditorGraphAttributeListView::OnGenerateFilterMenu)
@@ -349,6 +351,18 @@ void SPCGEditorGraphAttributeListView::Construct(const FArguments& InArgs, TShar
 			FilterImage.ToSharedRef()
 		];
 
+	TSharedPtr<SComboButton> AdditionalOperationsButton = SNew(SComboButton)
+		.ForegroundColor(FSlateColor::UseStyle())
+		.HasDownArrow(false)
+		.OnGetMenuContent(this, &SPCGEditorGraphAttributeListView::OnGenerateAdditionalOperationsMenu)
+		.ContentPadding(1)
+		.ButtonContent()
+		[
+			SNew(SImage)
+			.Image(FAppStyle::GetBrush("EditorViewportToolBar.OptionsDropdown"))
+			.ColorAndOpacity(FSlateColor::UseForeground())
+		];
+
 	SAssignNew(SearchBoxWidget, SSearchBox)
 		.MinDesiredWidth(300.0f)
 		.InitialText(ActiveFilterText)
@@ -410,6 +424,13 @@ void SPCGEditorGraphAttributeListView::Construct(const FArguments& InArgs, TShar
 				]
 			]
 			+SHorizontalBox::Slot()
+			.AutoWidth()
+			.VAlign(EVerticalAlignment::VAlign_Center)
+			.Padding(1.0f, 0.0f)
+			[
+				AdditionalOperationsButton->AsShared()
+			]
+			+SHorizontalBox::Slot()
 			.FillWidth(1.0f)
 			+SHorizontalBox::Slot()
 			.AutoWidth()
@@ -568,7 +589,7 @@ void SPCGEditorGraphAttributeListView::OnGenerateUpdated(UPCGComponent* /*InPCGC
 	RequestRefresh();
 }
 
-const FPCGDataCollection* SPCGEditorGraphAttributeListView::GetInspectionData() const
+const FPCGDataCollection* SPCGEditorGraphAttributeListView::GetInspectionData(const TSharedPtr<FPinComboBoxItem>& EditorPin) const
 {
 	if (!PCGComponent.IsValid())
 	{
@@ -582,12 +603,12 @@ const FPCGDataCollection* SPCGEditorGraphAttributeListView::GetInspectionData()
 	}
 
 	const UPCGPin* Pin = nullptr;
-	if (const TSharedPtr<FPinComboBoxItem> SelectedPin = PinComboBox->GetSelectedItem())
+	if (EditorPin)
 	{
-		const TArray<TObjectPtr<UPCGPin>>& Pins = SelectedPin->bIsOutputPin ? PCGNode->GetOutputPins() : PCGNode->GetInputPins();
-		if (Pins.IsValidIndex(SelectedPin->PinIndex))
+		const TArray<TObjectPtr<UPCGPin>>& Pins = EditorPin->bIsOutputPin ? PCGNode->GetOutputPins() : PCGNode->GetInputPins();
+		if (Pins.IsValidIndex(EditorPin->PinIndex))
 		{
-			Pin = Pins[SelectedPin->PinIndex];
+			Pin = Pins[EditorPin->PinIndex];
 		}
 	}
 
@@ -620,6 +641,11 @@ const FPCGDataCollection* SPCGEditorGraphAttributeListView::GetInspectionData()
 	return PCGComponent->GetInspectionData(Stack);
 }
 
+const FPCGDataCollection* SPCGEditorGraphAttributeListView::GetInspectionData() const
+{
+	return GetInspectionData(PinComboBox->GetSelectedItem());
+}
+
 void SPCGEditorGraphAttributeListView::RefreshAttributeList()
 {
 	HiddenAttributes = ListViewHeader->GetHiddenColumnIds();
@@ -839,6 +865,84 @@ TSharedRef<SWidget> SPCGEditorGraphAttributeListView::OnGenerateFilterMenu()
 	return MenuBuilder.MakeWidget();
 }
 
+TSharedRef<SWidget> SPCGEditorGraphAttributeListView::OnGenerateAdditionalOperationsMenu()
+{
+	FMenuBuilder MenuBuilder(false, nullptr);
+
+	MenuBuilder.AddMenuEntry(
+		LOCTEXT("SaveThisData", "Save this data"),
+		LOCTEXT("SaveThisDataTooltip", "Saves this data to a PCG Data Asset."),
+		FSlateIcon(FAppStyle::GetAppStyleSetName(), "Icons.Download"),
+		FUIAction(FExecuteAction::CreateSP(this, &SPCGEditorGraphAttributeListView::SaveData, true, true)),
+		NAME_None,
+		EUserInterfaceActionType::Button);
+
+	MenuBuilder.AddMenuEntry(
+		LOCTEXT("SaveThisPinData", "Save pin data"),
+		LOCTEXT("SaveThisPinDataTooltip", "Saves all data from the selected pin to a PCG Data Asset."),
+		FSlateIcon(FAppStyle::GetAppStyleSetName(), "Icons.Download"),
+		FUIAction(FExecuteAction::CreateSP(this, &SPCGEditorGraphAttributeListView::SaveData, true, false)),
+		NAME_None,
+		EUserInterfaceActionType::Button);
+
+	MenuBuilder.AddMenuEntry(
+		LOCTEXT("SaveAllData", "Save all"),
+		LOCTEXT("SaveAllDataTooltip", "Saves all the input or output data to a PCG Data Asset."),
+		FSlateIcon(FAppStyle::GetAppStyleSetName(), "Icons.Download"),
+		FUIAction(FExecuteAction::CreateSP(this, &SPCGEditorGraphAttributeListView::SaveData, false, false)),
+		NAME_None,
+		EUserInterfaceActionType::Button);
+
+	return MenuBuilder.MakeWidget();
+}
+
+void SPCGEditorGraphAttributeListView::SaveData(bool bUsePinComboIndex, bool bUseDataComboIndex)
+{
+	// Can't save N'th data on all pins, as it makes no sense
+	check(!bUseDataComboIndex || bUsePinComboIndex);
+
+	// Build a data collection from the indices we're given.
+	FPCGDataCollection Collection;
+
+	if (!bUsePinComboIndex)
+	{
+		for (const auto& PinComboBoxItem : PinComboBoxItems)
+		{
+			const FPCGDataCollection* PinInspectionData = GetInspectionData(PinComboBoxItem);
+			if (PinInspectionData)
+			{
+				Collection.TaggedData.Append(PinInspectionData->TaggedData);
+			}
+		}
+	}
+	else
+	{
+		if (const FPCGDataCollection* PinInspectionData = GetInspectionData())
+		{
+			if (!bUseDataComboIndex)
+			{
+				Collection = *PinInspectionData;
+			}
+			else
+			{
+
+				Collection.TaggedData.Add(PinInspectionData->TaggedData[GetSelectedDataIndex()]);
+			}
+		}
+	}
+
+	if (Collection.TaggedData.IsEmpty())
+	{
+		return;
+	}
+
+	UPCGDataCollectionExporter* Exporter = NewObject<UPCGDataCollectionExporter>();
+	Exporter->Data = Collection;
+
+	FPCGAssetExporterParameters Parameters;
+	UPCGAssetExporterUtils::CreateAsset(Exporter, Parameters);
+}
+
 FText SPCGEditorGraphAttributeListView::OnGenerateSelectedPinText() const
 {
 	if (const TSharedPtr<FPinComboBoxItem> SelectedPin = PinComboBox->GetSelectedItem())
diff --git a/Engine/Plugins/PCG/Source/PCGEditor/Private/SPCGEditorGraphAttributeListView.h b/Engine/Plugins/PCG/Source/PCGEditor/Private/SPCGEditorGraphAttributeListView.h
--- a/Engine/Plugins/PCG/Source/PCGEditor/Private/SPCGEditorGraphAttributeListView.h
+++ b/Engine/Plugins/PCG/Source/PCGEditor/Private/SPCGEditorGraphAttributeListView.h
@@ -160,8 +160,6 @@ private:
 
 	void OnGenerateUpdated(UPCGComponent* InPCGComponent);
 
-	const FPCGDataCollection* GetInspectionData() const;
-
 	void RefreshAttributeList();
 	void RefreshPinComboBox();
 	void RefreshDataComboBox();
@@ -182,12 +180,16 @@ private:
 		bool bIsOutputPin = true;
 	};
 
+	const FPCGDataCollection* GetInspectionData(const TSharedPtr<FPinComboBoxItem>& EditorPin) const;
+	const FPCGDataCollection* GetInspectionData() const;
+
 	FText OnGenerateSelectedPinText() const;
 	void OnSelectionChangedPin(TSharedPtr<FPinComboBoxItem> InItem, ESelectInfo::Type InSelectInfo);
 	TSharedRef<SWidget> OnGeneratePinWidget(TSharedPtr<FPinComboBoxItem> InItem) const;
 
 	const FSlateBrush* GetFilterBadgeIcon() const;
 	TSharedRef<SWidget> OnGenerateFilterMenu();
+	TSharedRef<SWidget> OnGenerateAdditionalOperationsMenu();
 	TSharedRef<SWidget> OnGenerateDataWidget(TSharedPtr<FName> InItem) const;
 	void OnSelectionChanged(TSharedPtr<FName> Item, ESelectInfo::Type SelectInfo);
 	FText OnGenerateSelectedDataText() const;
@@ -198,6 +200,8 @@ private:
 	ECheckBoxState GetAnyAttributeEnabledState() const;
 	bool IsAttributeEnabled(FName InAttributeName) const;
 
+	void SaveData(bool bUsePinIndex, bool bUseDataIndex);
+
 	TSharedRef<ITableRow> OnGenerateRow(PCGListviewItemPtr Item, const TSharedRef<STableViewBase>& OwnerTable);
 	void OnItemDoubleClicked(PCGListviewItemPtr Item) const;
 
@@ -249,8 +253,6 @@ private:
 
 	TSharedPtr<STextBlock> NodeNameTextBlock;
 	TSharedPtr<STextBlock> InfoTextBlock;
-	TSharedPtr<SComboButton> FilterButton;
-	TSharedPtr<SButton> LockButton;
 
 	TArray<FName> HiddenAttributes;
 
