[PCG] Fixing Render Target related issues

* SampleTexture was not supporting BaseTexture
* Invalid RenderTarget was crashing
* Change color of BaseTexture  to be closer to texture than surface.

#rnx

diff --git a/Engine/Plugins/PCG/Source/PCG/Private/Data/PCGRenderTargetData.cpp b/Engine/Plugins/PCG/Source/PCG/Private/Data/PCGRenderTargetData.cpp
--- a/Engine/Plugins/PCG/Source/PCG/Private/Data/PCGRenderTargetData.cpp
+++ b/Engine/Plugins/PCG/Source/PCG/Private/Data/PCGRenderTargetData.cpp
@@ -28,16 +28,11 @@ void UPCGRenderTargetData::Initialize(UTextureRenderTarget2D* InRenderTarget, co
 				RTResource->ReadLinearColorPixels(ColorData, ReadPixelFlags, Rect);
 			}
 		}
-
-		Width = RenderTarget->SizeX;
-		Height = RenderTarget->SizeY;
-	}
-	else
-	{
-		Width = 0;
-		Height = 0;
 	}
 
+	Width = !ColorData.IsEmpty() ? RenderTarget->SizeX : 0;
+	Height = !ColorData.IsEmpty() ? RenderTarget->SizeY : 0;
+
 	Bounds = FBox(EForceInit::ForceInit);
 	Bounds += FVector(-1.0f, -1.0f, 0.0f);
 	Bounds += FVector(1.0f, 1.0f, 0.0f);
diff --git a/Engine/Plugins/PCG/Source/PCG/Private/Data/PCGTextureData.cpp b/Engine/Plugins/PCG/Source/PCG/Private/Data/PCGTextureData.cpp
--- a/Engine/Plugins/PCG/Source/PCG/Private/Data/PCGTextureData.cpp
+++ b/Engine/Plugins/PCG/Source/PCG/Private/Data/PCGTextureData.cpp
@@ -307,12 +307,12 @@ const UPCGPointData* UPCGBaseTextureData::CreatePointData(FPCGContext* Context)
 
 bool UPCGBaseTextureData::IsValid() const
 {
-	return Height > 0 && Width > 0;
+	return Height > 0 && Width > 0 && !ColorData.IsEmpty();
 }
 
 bool UPCGBaseTextureData::SamplePointLocal(const FVector2D& LocalPosition, FVector4& OutColor, float& OutDensity) const
 {
-	if (!ensure(Width > 0 && Height > 0))
+	if (!IsValid())
 	{
 		return false;
 	}
diff --git a/Engine/Plugins/PCG/Source/PCG/Private/Elements/PCGSampleTexture.cpp b/Engine/Plugins/PCG/Source/PCG/Private/Elements/PCGSampleTexture.cpp
--- a/Engine/Plugins/PCG/Source/PCG/Private/Elements/PCGSampleTexture.cpp
+++ b/Engine/Plugins/PCG/Source/PCG/Private/Elements/PCGSampleTexture.cpp
@@ -22,7 +22,7 @@ TArray<FPCGPinProperties> UPCGSampleTextureSettings::InputPinProperties() const
 		/*bAllowMultipleData=*/true);
 
 	PinProperties.Emplace(PCGSampleTextureConstants::InputTextureLabel,
-		EPCGDataType::Texture,
+		EPCGDataType::BaseTexture,
 		/*bAllowMultipleConnections=*/false,
 		/*bAllowMultipleData=*/false);
 
diff --git a/Engine/Plugins/PCG/Source/PCGEditor/Private/PCGEditorCommon.h b/Engine/Plugins/PCG/Source/PCGEditor/Private/PCGEditorCommon.h
--- a/Engine/Plugins/PCG/Source/PCGEditor/Private/PCGEditorCommon.h
+++ b/Engine/Plugins/PCG/Source/PCGEditor/Private/PCGEditorCommon.h
@@ -40,6 +40,7 @@ namespace FPCGEditorCommon
 	const FName PolyLineDataType = FName(TEXT("Poly Line Data"));
 	const FName SurfaceDataType = FName(TEXT("Surface Data"));
 	const FName LandscapeDataType = FName(TEXT("Landscape Data"));
+	const FName BaseTextureDataType = FName(TEXT("Base Texture Data"));
 	const FName TextureDataType = FName(TEXT("Texture Data"));
 	const FName RenderTargetDataType = FName(TEXT("Render Target Data"));
 	const FName VolumeDataType = FName(TEXT("Volume Data"));
diff --git a/Engine/Plugins/PCG/Source/PCGEditor/Private/PCGEditorGraphNodeBase.cpp b/Engine/Plugins/PCG/Source/PCGEditor/Private/PCGEditorGraphNodeBase.cpp
--- a/Engine/Plugins/PCG/Source/PCGEditor/Private/PCGEditorGraphNodeBase.cpp
+++ b/Engine/Plugins/PCG/Source/PCGEditor/Private/PCGEditorGraphNodeBase.cpp
@@ -919,6 +919,10 @@ FEdGraphPinType UPCGEditorGraphNodeBase::GetPinType(const UPCGPin* InPin)
 		{
 			EdPinType.PinSubCategory = FPCGEditorCommon::LandscapeDataType;
 		}
+		else if (CheckType(EPCGDataType::BaseTexture))
+		{
+			EdPinType.PinSubCategory = FPCGEditorCommon::BaseTextureDataType;
+		}
 		else if (CheckType(EPCGDataType::Texture))
 		{
 			EdPinType.PinSubCategory = FPCGEditorCommon::TextureDataType;
diff --git a/Engine/Plugins/PCG/Source/PCGEditor/Private/PCGEditorSettings.cpp b/Engine/Plugins/PCG/Source/PCGEditor/Private/PCGEditorSettings.cpp
--- a/Engine/Plugins/PCG/Source/PCGEditor/Private/PCGEditorSettings.cpp
+++ b/Engine/Plugins/PCG/Source/PCGEditor/Private/PCGEditorSettings.cpp
@@ -35,6 +35,7 @@ UPCGEditorSettings::UPCGEditorSettings(const FObjectInitializer& ObjectInitializ
 	PolyLineDataPinColor = FLinearColor(0.05f, 0.75f, 0.82f);
 	SurfaceDataPinColor = FLinearColor(0.06f, 0.55f, 0.21f);
 	LandscapeDataPinColor = FLinearColor(0.66f, 0.66f, 0.07f);
+	BaseTextureDataPinColor = FLinearColor(0.81f, 0.28f, 0.22f);
 	TextureDataPinColor = FLinearColor(0.79f, 0.08f, 0.01f);
 	RenderTargetDataPinColor = FLinearColor(0.8f, 0.18f, 0.12f);
 	VolumeDataPinColor = FLinearColor(0.79f, 0.06f, 0.5f);
@@ -116,6 +117,10 @@ FLinearColor UPCGEditorSettings::GetPinColor(const FEdGraphPinType& PinType) con
 		{
 			return LandscapeDataPinColor;
 		}
+		else if (PinType.PinSubCategory == FPCGEditorCommon::BaseTextureDataType)
+		{
+			return BaseTextureDataPinColor;
+		}
 		else if (PinType.PinSubCategory == FPCGEditorCommon::TextureDataType)
 		{
 			return TextureDataPinColor;
diff --git a/Engine/Plugins/PCG/Source/PCGEditor/Public/PCGEditorSettings.h b/Engine/Plugins/PCG/Source/PCGEditor/Public/PCGEditorSettings.h
--- a/Engine/Plugins/PCG/Source/PCGEditor/Public/PCGEditorSettings.h
+++ b/Engine/Plugins/PCG/Source/PCGEditor/Public/PCGEditorSettings.h
@@ -121,6 +121,10 @@ public:
 	UPROPERTY(EditAnywhere, config, Category = Node, meta = (HideAlphaChannel))
 	FLinearColor LandscapeDataPinColor;
 
+	/** Color used for data pins of type Base Texture */
+	UPROPERTY(EditAnywhere, config, Category = Node, meta = (HideAlphaChannel))
+	FLinearColor BaseTextureDataPinColor;
+	
 	/** Color used for data pins of type Texture */
 	UPROPERTY(EditAnywhere, config, Category = Node, meta = (HideAlphaChannel))
 	FLinearColor TextureDataPinColor;
