[PCG] Add Transfer Attributes as alias to Copy Attributes now that it is deprecated

#rnx

diff --git a/Engine/Plugins/PCG/Source/PCG/Private/Elements/PCGCopyAttributes.cpp b/Engine/Plugins/PCG/Source/PCG/Private/Elements/PCGCopyAttributes.cpp
--- a/Engine/Plugins/PCG/Source/PCG/Private/Elements/PCGCopyAttributes.cpp
+++ b/Engine/Plugins/PCG/Source/PCG/Private/Elements/PCGCopyAttributes.cpp
@@ -120,6 +120,11 @@ EPCGDataType UPCGCopyAttributesSettings::GetCurrentPinTypes(const UPCGPin* InPin
 	return OutputType != EPCGDataType::None ? OutputType : EPCGDataType::Any;
 }
 
+TArray<FText> UPCGCopyAttributesSettings::GetNodeTitleAliases() const
+{
+	return { LOCTEXT("CopyAttributesAliasTransfer", "Transfer Attributes") };
+}
+
 #endif // WITH_EDITOR
 
 TArray<FPCGPinProperties> UPCGCopyAttributesSettings::InputPinProperties() const
diff --git a/Engine/Plugins/PCG/Source/PCG/Public/Elements/PCGCopyAttributes.h b/Engine/Plugins/PCG/Source/PCG/Public/Elements/PCGCopyAttributes.h
--- a/Engine/Plugins/PCG/Source/PCG/Public/Elements/PCGCopyAttributes.h
+++ b/Engine/Plugins/PCG/Source/PCG/Public/Elements/PCGCopyAttributes.h
@@ -50,6 +50,7 @@ public:
 	virtual void ApplyDeprecationBeforeUpdatePins(UPCGNode* InOutNode, TArray<TObjectPtr<UPCGPin>>& InputPins, TArray<TObjectPtr<UPCGPin>>& OutputPins) override;
 	virtual EPCGDataType GetCurrentPinTypes(const UPCGPin* InPin) const override;
 	virtual bool HasDynamicPins() const override { return true; }
+	virtual TArray<FText> GetNodeTitleAliases() const override;
 #endif
 	virtual FString GetAdditionalTitleInformation() const override;
 
