[PCG] Missed review feedback - don't bother explicitly marking data binding as garbage.

#rnx

diff --git a/Engine/Plugins/PCG/Source/PCG/Private/Compute/Elements/PCGComputeGraphElement.cpp b/Engine/Plugins/PCG/Source/PCG/Private/Compute/Elements/PCGComputeGraphElement.cpp
--- a/Engine/Plugins/PCG/Source/PCG/Private/Compute/Elements/PCGComputeGraphElement.cpp
+++ b/Engine/Plugins/PCG/Source/PCG/Private/Compute/Elements/PCGComputeGraphElement.cpp
@@ -239,9 +239,6 @@ void FPCGComputeGraphElement::PostExecuteInternal(FPCGContext* InContext) const
 		}
 	}
 #endif
-
-	Context->DataBinding->MarkAsGarbage();
-	Context->DataBinding = nullptr;
 }
 
 void FPCGComputeGraphElement::AbortInternal(FPCGContext* InContext) const
