[PCG] Only connected "islands" of GPU nodes are assembled into compute graphs, rather than greedily batching them together.

#rnx

diff --git a/Engine/Plugins/PCG/Source/PCG/Private/Graph/PCGGraphCompilerGPU.cpp b/Engine/Plugins/PCG/Source/PCG/Private/Graph/PCGGraphCompilerGPU.cpp
--- a/Engine/Plugins/PCG/Source/PCG/Private/Graph/PCGGraphCompilerGPU.cpp
+++ b/Engine/Plugins/PCG/Source/PCG/Private/Graph/PCGGraphCompilerGPU.cpp
@@ -34,8 +34,72 @@ namespace PCGGraphCompilerGPU
 #endif
 }
 
-void FPCGGraphCompilerGPU::CollectGPUNodeSubsets(const TArray<FPCGGraphTask>& InCompiledTasks, const TMap<FPCGTaskId, TArray<FPCGTaskId>>& InTaskSuccessors, const TSet<FPCGTaskId>& InGPUCompatibleTaskIds, TArray<TSet<FPCGTaskId>>& OutNodeSubsetsToConvertToCFGraph)
+void FPCGGraphCompilerGPU::LabelConnectedGPUNodeIslands(
+	const TArray<FPCGGraphTask>& InCompiledTasks,
+	const TSet<FPCGTaskId>& InGPUCompatibleTaskIds,
+	const FTaskToSuccessors& InTaskSuccessors,
+	TArray<uint32>& OutIslandIDs)
+{
+	OutIslandIDs.SetNumZeroed(InCompiledTasks.Num());
+
+	// Traverses task inputs and successors and assigns the given island ID to each one. Memoized via output OutIslandIDs.
+	auto FloodFillIslandID = [&InCompiledTasks, &InTaskSuccessors, &InGPUCompatibleTaskIds, &OutIslandIDs](FPCGTaskId InTaskId, int InIslandID, FPCGTaskId InTraversedFromTaskId, auto&& RecursiveCall) -> void
+	{
+		check(InTaskId != InTraversedFromTaskId);
+
+		OutIslandIDs[InTaskId] = InIslandID;
+
+		for (const FPCGGraphTaskInput& Input : InCompiledTasks[InTaskId].Inputs)
+		{
+			if (Input.TaskId == InTraversedFromTaskId)
+			{
+				continue;
+			}
+
+			if (OutIslandIDs[Input.TaskId] == 0 && InGPUCompatibleTaskIds.Contains(Input.TaskId))
+			{
+				RecursiveCall(Input.TaskId, InIslandID, InTaskId, RecursiveCall);
+			}
+		}
+
+		if (const TArray<FPCGTaskId>* Successors = InTaskSuccessors.Find(InTaskId))
+		{
+			for (FPCGTaskId Successor : *Successors)
+			{
+				if (Successor == InTraversedFromTaskId)
+				{
+					continue;
+				}
+
+				if (OutIslandIDs[Successor] == 0 && InGPUCompatibleTaskIds.Contains(Successor))
+				{
+					RecursiveCall(Successor, InIslandID, InTaskId, RecursiveCall);
+				}
+			}
+		}
+	};
+
+	for (FPCGTaskId GPUTaskId : InGPUCompatibleTaskIds)
+	{
+		if (OutIslandIDs[GPUTaskId] == 0)
+		{
+			// Really doesn't matter what the island IDs are so just use ID of first task encountered in island.
+			const uint32 IslandID = static_cast<uint32>(GPUTaskId);
+			FloodFillIslandID(GPUTaskId, IslandID, InvalidPCGTaskId, FloodFillIslandID);
+		}
+	}
+}
+
+void FPCGGraphCompilerGPU::CollectGPUNodeSubsets(
+	const TArray<FPCGGraphTask>& InCompiledTasks,
+	const FTaskToSuccessors& InTaskSuccessors,
+	const TSet<FPCGTaskId>& InGPUCompatibleTaskIds,
+	TArray<TSet<FPCGTaskId>>& OutNodeSubsetsToConvertToCFGraph)
 {
+	// Identifies connected sets of GPU nodes, giving each a non-zero ID value.
+	TArray<uint32> ConnectedGPUNodeIslandIDs;
+	LabelConnectedGPUNodeIslands(InCompiledTasks, InGPUCompatibleTaskIds, InTaskSuccessors, ConnectedGPUNodeIslandIDs);
+
 	// Populate initial sets of tasks that are ready to consume vs ones that currently blocked.
 	TSet<FPCGTaskId> ReadyTaskIds;
 	TSet<FPCGTaskId> RemainingTaskIds;
@@ -110,7 +174,8 @@ void FPCGGraphCompilerGPU::CollectGPUNodeSubsets(const TArray<FPCGGraphTask>& In
 
 			for (FPCGTaskId ReadyTaskId : ReadyTaskIds)
 			{
-				if (!InGPUCompatibleTaskIds.Contains(ReadyTaskId))
+				const bool bIsCPUNode = ConnectedGPUNodeIslandIDs[ReadyTaskId] == 0;
+				if (bIsCPUNode)
 				{
 					FoundReadyTaskIds.Add(ReadyTaskId);
 				}
@@ -128,6 +193,7 @@ void FPCGGraphCompilerGPU::CollectGPUNodeSubsets(const TArray<FPCGGraphTask>& In
 		GPUSubsetTaskIds.Reset();
 
 		int StackIndex = INDEX_NONE;
+		uint32 IslandID = INDEX_NONE;
 
 		// Now the opposite - consume as many GPU nodes as we can and accumulate them into a set that will be compiled into a compute graph.
 		bQueuedTasks = !ReadyTaskIds.IsEmpty();
@@ -137,14 +203,24 @@ void FPCGGraphCompilerGPU::CollectGPUNodeSubsets(const TArray<FPCGGraphTask>& In
 
 			for (FPCGTaskId ReadyTaskId : ReadyTaskIds)
 			{
-				if (InGPUCompatibleTaskIds.Contains(ReadyTaskId))
+				const uint32 TaskIslandID = ConnectedGPUNodeIslandIDs[ReadyTaskId];
+				if (TaskIslandID == 0)
 				{
-					// For now don't mix tasks from different execution stacks (in and out of subgraphs for instance) into one compute graph.
-					if ((StackIndex == INDEX_NONE) || InCompiledTasks[ReadyTaskId].StackIndex == StackIndex)
-					{
-						StackIndex = InCompiledTasks[ReadyTaskId].StackIndex;
-						FoundReadyTaskIds.Add(ReadyTaskId);
-					}
+					// Non-gpu task - skip
+					continue;
+				}
+
+				const bool bIslandMatches = (IslandID == INDEX_NONE) || (IslandID == TaskIslandID);
+
+				// For now don't mix tasks from different execution stacks (in and out of subgraphs for instance) into one compute graph.
+				const bool bStackMatches = (StackIndex == INDEX_NONE) || (InCompiledTasks[ReadyTaskId].StackIndex == StackIndex);
+
+				if (bIslandMatches && bStackMatches)
+				{
+					IslandID = TaskIslandID;
+					StackIndex = InCompiledTasks[ReadyTaskId].StackIndex;
+
+					FoundReadyTaskIds.Add(ReadyTaskId);
 				}
 			}
 
@@ -274,7 +350,7 @@ void FPCGGraphCompilerGPU::WireGPUGraphNode(
 	const TSet<FPCGTaskId>& InCollapsedTasks,
 	const TSet<FPCGTaskId>& InGPUCompatibleTaskIds,
 	TArray<FPCGGraphTask>& InOutCompiledTasks,
-	const TMap<FPCGTaskId, TArray<FPCGTaskId>>& InTaskSuccessors,
+	const FTaskToSuccessors& InTaskSuccessors,
 	FOriginalToVirtualPin& OutOriginalToVirtualPin,
 	TMap<const UPCGPin*, FName>& OutOutputCPUPinToVirtualPin)
 {
@@ -379,7 +455,7 @@ void FPCGGraphCompilerGPU::BuildGPUGraphTask(
 	UPCGGraph* InGraph,
 	FPCGTaskId InGPUGraphTaskId,
 	const TSet<FPCGTaskId>& InCollapsedTasks,
-	const TMap<FPCGTaskId, TArray<FPCGTaskId>>& InTaskSuccessors,
+	const FTaskToSuccessors& InTaskSuccessors,
 	TArray<FPCGGraphTask>& InOutCompiledTasks,
 	const FOriginalToVirtualPin& InOriginalToVirtualPin,
 	const TMap<const UPCGPin*, FName>& InOutputCPUPinToVirtualPin)
@@ -945,7 +1021,7 @@ void FPCGGraphCompilerGPU::CreateGPUNodes(UPCGGraph* InGraph, TArray<FPCGGraphTa
 		return;
 	}
 
-	TMap<FPCGTaskId, TArray<FPCGTaskId>> TaskSuccessors;
+	FTaskToSuccessors TaskSuccessors;
 	TaskSuccessors.Reserve(InOutCompiledTasks.Num());
 	for (FPCGTaskId TaskId = 0; TaskId < InOutCompiledTasks.Num(); ++TaskId)
 	{
diff --git a/Engine/Plugins/PCG/Source/PCG/Private/Graph/PCGGraphCompilerGPU.h b/Engine/Plugins/PCG/Source/PCG/Private/Graph/PCGGraphCompilerGPU.h
--- a/Engine/Plugins/PCG/Source/PCG/Private/Graph/PCGGraphCompilerGPU.h
+++ b/Engine/Plugins/PCG/Source/PCG/Private/Graph/PCGGraphCompilerGPU.h
@@ -14,11 +14,19 @@ public:
 	/** Used to track new unique virtual pins created on generated compute graph elements. */
 	using FNodePin = TTuple<FPCGTaskId, /*Pin label*/FName, /*Pin is input*/bool>;
 	using FOriginalToVirtualPin = TMap<FNodePin, /*Virtual pin label*/FName>;
-	
+	using FTaskToSuccessors = TMap<FPCGTaskId, TArray<FPCGTaskId>>;
+
+	/** Identifies connected sets of GPU nodes, giving each a non - zero ID value. */
+	static void LabelConnectedGPUNodeIslands(
+		const TArray<FPCGGraphTask>& InCompiledTasks,
+		const TSet<FPCGTaskId>& InGPUCompatibleTaskIds,
+		const FTaskToSuccessors& InTaskSuccessors,
+		TArray<uint32>& OutIslandIDs);
+
 	/** Outputs sets of task IDs, where each set is GPU nodes that can be compiled into a compute graph and dispatched together. */
 	static void CollectGPUNodeSubsets(
 		const TArray<FPCGGraphTask>& InCompiledTasks,
-		const TMap<FPCGTaskId, TArray<FPCGTaskId>>& InTaskSuccessors,
+		const FTaskToSuccessors& InTaskSuccessors,
 		const TSet<FPCGTaskId>& InGPUCompatibleTaskIds,
 		TArray<TSet<FPCGTaskId>>& OutNodeSubsetsToConvertToCFGraph);
 	
@@ -43,7 +51,7 @@ public:
 		UPCGGraph* InGraph,
 		FPCGTaskId InGPUGraphTaskId,
 		const TSet<FPCGTaskId>& InCollapsedTasks,
-		const TMap<FPCGTaskId, TArray<FPCGTaskId>>& InTaskSuccessors,
+		const FTaskToSuccessors& InTaskSuccessors,
 		TArray<FPCGGraphTask>& InOutCompiledTasks,
 		const FOriginalToVirtualPin& InOriginalToVirtualPin,
 		const TMap<const UPCGPin*, FName>& InOutputCPUPinToVirtualPin);
