[PCG] Initialize header with 0 data count from CPU rather than writing from every GPU thread.

#rnx

diff --git a/Engine/Plugins/PCG/Source/PCG/Private/Compute/DataInterfaces/PCGDataCollectionDataInterface.cpp b/Engine/Plugins/PCG/Source/PCG/Private/Compute/DataInterfaces/PCGDataCollectionDataInterface.cpp
--- a/Engine/Plugins/PCG/Source/PCG/Private/Compute/DataInterfaces/PCGDataCollectionDataInterface.cpp
+++ b/Engine/Plugins/PCG/Source/PCG/Private/Compute/DataInterfaces/PCGDataCollectionDataInterface.cpp
@@ -464,8 +464,6 @@ UComputeDataProvider* UPCGDataCollectionDataInterface::CreateDataProvider(TObjec
 	UPCGDataCollectionDataProvider* Provider = NewObject<UPCGDataCollectionDataProvider>();
 	Provider->Binding = Binding;
 	Provider->PinDesc = ProducerSettings->ComputeOutputPinDataDesc(OutputPinLabel, Binding);
-	Provider->SizeBytes = Provider->PinDesc.ComputePackedSize();
-	check(Provider->SizeBytes >= 4);
 
 	return Provider;
 }
@@ -474,9 +472,6 @@ FComputeDataProviderRenderProxy* UPCGDataCollectionDataProvider::GetRenderProxy(
 {
 	FPCGDataCollectionDataProviderProxy* Proxy = new FPCGDataCollectionDataProviderProxy(Binding, PinDesc);
 
-	check(SizeBytes > 0);
-	Proxy->SizeBytes = SizeBytes;
-
 	return Proxy;
 }
 
@@ -510,17 +505,15 @@ void FPCGDataCollectionDataProviderProxy::GatherDispatchData(FDispatchData const
 
 void FPCGDataCollectionDataProviderProxy::AllocateResources(FRDGBuilder& GraphBuilder, FAllocationData const& InAllocationData)
 {
-	check(SizeBytes > 0 && (SizeBytes % sizeof(uint32)) == 0);
-
-	FRDGBufferDesc Desc = FRDGBufferDesc::CreateStructuredDesc(sizeof(uint32), SizeBytes / sizeof(uint32));
-	Buffer = GraphBuilder.CreateBuffer(Desc, TEXT("PCGDataCollectionBuffer"));
-	BufferUAV = GraphBuilder.CreateUAV(Buffer);
-
 	// Initialize with an empty data collection. The kernel may not run, for example if indirect dispatch args end up being 0. Ensure
 	// there is something meaningful to readback.
 	// TODO could have a statically-allocated resource rather than allocating & uploading here.
 	TArray<uint32> PackedDataCollection;
-	PinDesc.PackDataCollection(/*InDataCollection=*/{}, /*InPin=*/NAME_None, PackedDataCollection);
+	PinDesc.PrepareBufferForKernelOutput(PackedDataCollection);
+
+	FRDGBufferDesc Desc = FRDGBufferDesc::CreateStructuredDesc(PackedDataCollection.GetTypeSize(), PackedDataCollection.Num());
+	Buffer = GraphBuilder.CreateBuffer(Desc, TEXT("PCGDataCollectionBuffer"));
+	BufferUAV = GraphBuilder.CreateUAV(Buffer);
 
 	GraphBuilder.QueueBufferUpload(Buffer, PackedDataCollection.GetData(), PackedDataCollection.Num() * PackedDataCollection.GetTypeSize(), ERDGInitialDataFlags::None);
 }
diff --git a/Engine/Plugins/PCG/Source/PCG/Private/Compute/DataInterfaces/PCGDataCollectionDataInterface.h b/Engine/Plugins/PCG/Source/PCG/Private/Compute/DataInterfaces/PCGDataCollectionDataInterface.h
--- a/Engine/Plugins/PCG/Source/PCG/Private/Compute/DataInterfaces/PCGDataCollectionDataInterface.h
+++ b/Engine/Plugins/PCG/Source/PCG/Private/Compute/DataInterfaces/PCGDataCollectionDataInterface.h
@@ -57,7 +57,6 @@ public:
 
 	TWeakObjectPtr<UPCGDataBinding> Binding;
 	FPCGDataCollectionDesc PinDesc;
-	int SizeBytes = -1;
 };
 
 class FPCGDataCollectionDataProviderProxy : public FComputeDataProviderRenderProxy
@@ -74,8 +73,6 @@ public:
 	void AllocateResources(FRDGBuilder& GraphBuilder, FAllocationData const& InAllocationData) override;
 	//~ End FComputeDataProviderRenderProxy Interface
 
-	int SizeBytes = -1;
-
 protected:
 	using FParameters = FPCGDataCollectionDataInterfaceParameters;
 
diff --git a/Engine/Plugins/PCG/Source/PCG/Private/Compute/DataInterfaces/PCGDataCollectionReadbackDataInterface.cpp b/Engine/Plugins/PCG/Source/PCG/Private/Compute/DataInterfaces/PCGDataCollectionReadbackDataInterface.cpp
--- a/Engine/Plugins/PCG/Source/PCG/Private/Compute/DataInterfaces/PCGDataCollectionReadbackDataInterface.cpp
+++ b/Engine/Plugins/PCG/Source/PCG/Private/Compute/DataInterfaces/PCGDataCollectionReadbackDataInterface.cpp
@@ -18,8 +18,6 @@ UComputeDataProvider* UPCGDataCollectionReadbackDataInterface::CreateDataProvide
 	Provider->PinDesc = ProducerSettings->ComputeOutputPinDataDesc(OutputPinLabel, Binding);
 	// Use the aliased label as this is the output from the compute graph.
 	Provider->OutputPinLabelAlias = OutputPinLabelAlias;
-	Provider->SizeBytes = Provider->PinDesc.ComputePackedSize();
-	check(Provider->SizeBytes >= 4);
 
 	return Provider;
 }
@@ -60,7 +58,7 @@ FComputeDataProviderRenderProxy* UPCGDataProviderDataCollectionReadback::GetRend
 		ThisDataProvider->OnReadbackComplete.Broadcast();
 	};
 
-	return new FPCGDataProviderDataCollectionReadbackProxy(Binding, PinDesc, SizeBytes, ProcessReadbackData_RenderThread);
+	return new FPCGDataProviderDataCollectionReadbackProxy(Binding, PinDesc, ProcessReadbackData_RenderThread);
 }
 
 bool UPCGDataProviderDataCollectionReadback::ProcessReadBackData()
@@ -91,11 +89,10 @@ bool UPCGDataProviderDataCollectionReadback::ProcessReadBackData()
 FPCGDataProviderDataCollectionReadbackProxy::FPCGDataProviderDataCollectionReadbackProxy(
 	TWeakObjectPtr<UPCGDataBinding> InBinding,
 	const FPCGDataCollectionDesc& InPinDesc,
-	int InSizeBytes,
 	FReadbackCallback InAsyncReadbackCallback_RenderThread)
 	: FPCGDataCollectionDataProviderProxy(InBinding, InPinDesc)
 {
-	SizeBytes = InSizeBytes;
+	SizeBytes = InPinDesc.ComputePackedSize();
 	AsyncReadbackCallback_RenderThread = InAsyncReadbackCallback_RenderThread;
 }
 
@@ -123,7 +120,7 @@ void FPCGDataProviderDataCollectionReadbackProxy::AllocateResources(FRDGBuilder&
 	// there is something meaningful to readback.
 	// TODO factor out to avoid doing this work each time.
 	TArray<uint32> PackedDataCollection;
-	PinDesc.PackDataCollection(/*InDataCollection=*/{}, /*InPin=*/NAME_None, PackedDataCollection);
+	PinDesc.PrepareBufferForKernelOutput(PackedDataCollection);
 
 	GraphBuilder.QueueBufferUpload(Buffer, PackedDataCollection.GetData(), PackedDataCollection.Num() * PackedDataCollection.GetTypeSize(), ERDGInitialDataFlags::None);
 }
diff --git a/Engine/Plugins/PCG/Source/PCG/Private/Compute/DataInterfaces/PCGDataCollectionReadbackDataInterface.h b/Engine/Plugins/PCG/Source/PCG/Private/Compute/DataInterfaces/PCGDataCollectionReadbackDataInterface.h
--- a/Engine/Plugins/PCG/Source/PCG/Private/Compute/DataInterfaces/PCGDataCollectionReadbackDataInterface.h
+++ b/Engine/Plugins/PCG/Source/PCG/Private/Compute/DataInterfaces/PCGDataCollectionReadbackDataInterface.h
@@ -53,7 +53,7 @@ private:
 class FPCGDataProviderDataCollectionReadbackProxy : public FPCGDataCollectionDataProviderProxy
 {
 public:
-	FPCGDataProviderDataCollectionReadbackProxy(TWeakObjectPtr<UPCGDataBinding> InBinding, const FPCGDataCollectionDesc& InPinDesc, int InSizeBytes, FReadbackCallback InAsyncReadbackCallback_RenderThread);
+	FPCGDataProviderDataCollectionReadbackProxy(TWeakObjectPtr<UPCGDataBinding> InBinding, const FPCGDataCollectionDesc& InPinDesc, FReadbackCallback InAsyncReadbackCallback_RenderThread);
 
 	//~ Begin FComputeDataProviderRenderProxy Interface
 	bool IsValid(FValidationData const& InValidationData) const override;
@@ -61,6 +61,8 @@ public:
 	void GetReadbackData(TArray<FReadbackData>& OutReadbackData) const override;
 	//~ End FComputeDataProviderRenderProxy Interface
 
+	int SizeBytes = -1;
+
 private:
 	FReadbackCallback AsyncReadbackCallback_RenderThread;
 };
diff --git a/Engine/Plugins/PCG/Source/PCG/Private/Compute/DataInterfaces/PCGDataCollectionUploadDataInterface.cpp b/Engine/Plugins/PCG/Source/PCG/Private/Compute/DataInterfaces/PCGDataCollectionUploadDataInterface.cpp
--- a/Engine/Plugins/PCG/Source/PCG/Private/Compute/DataInterfaces/PCGDataCollectionUploadDataInterface.cpp
+++ b/Engine/Plugins/PCG/Source/PCG/Private/Compute/DataInterfaces/PCGDataCollectionUploadDataInterface.cpp
@@ -59,8 +59,6 @@ void FPCGDataProviderDataCollectionUploadProxy::AllocateResources(FRDGBuilder& G
 	check(!DownstreamInputPinLabels.IsEmpty());
 	PinDesc.PackDataCollection(Binding->DataForGPU.InputDataCollection, DownstreamInputPinLabels[0], PackedDataCollection);
 	
-	SizeBytes = PackedDataCollection.Num() * sizeof(uint32);
-
 	const FRDGBufferDesc Desc = FRDGBufferDesc::CreateStructuredDesc(sizeof(uint32), PackedDataCollection.Num());
 	Buffer = GraphBuilder.CreateBuffer(Desc, TEXT("PCGDataCollectionUploadBuffer"));
 	BufferUAV = GraphBuilder.CreateUAV(Buffer);
diff --git a/Engine/Plugins/PCG/Source/PCG/Private/Compute/Elements/PCGCustomHLSL.cpp b/Engine/Plugins/PCG/Source/PCG/Private/Compute/Elements/PCGCustomHLSL.cpp
--- a/Engine/Plugins/PCG/Source/PCG/Private/Compute/Elements/PCGCustomHLSL.cpp
+++ b/Engine/Plugins/PCG/Source/PCG/Private/Compute/Elements/PCGCustomHLSL.cpp
@@ -943,39 +943,9 @@ FString UPCGCustomHLSLSettings::GetCookedKernelSource(const TMap<FPCGKernelAttri
 		HeaderWriters += FString::Format(TEXT(
 			"    if (ThreadIndex >= GetNumThreads()) return;\n"
 			"    \n"
-			"    // Copy data collection header from pin {0} to pin {1}\n"
-			"    {\n"
-			"        const uint NumData = {0}_GetNumData();\n"
-			"        {1}_WriteNumData(NumData);\n"
-			"        \n"
-			"        for (uint DataIndex = 0; DataIndex < NumData; ++DataIndex)\n"
-			"        {\n"
-			"            const uint DataAddress = {0}_ReadDataAddress(DataIndex);\n"
-			"            const uint Id = {0}_ReadDataId(DataAddress);\n"
-			"            const uint NumAttributes = {0}_ReadDataNumAttributes(DataAddress);\n"
-			"            const uint PreambleSizeBytes = {0}_ReadDataPreambleSize(DataAddress);\n"
-			"            const uint NumElements = {0}_ReadDataInfo(DataAddress);\n"
-			"            \n"
-			"            {1}_WriteDataAddress(DataIndex, DataAddress);\n"
-			"            {1}_WriteDataId(DataAddress, Id);\n"
-			"            {1}_WriteDataNumAttributes(DataAddress, NumAttributes);\n"
-			"            {1}_WriteDataPreambleSize(DataAddress, PreambleSizeBytes);\n"
-			"            {1}_WriteDataInfo(DataAddress, NumElements);\n"
-			"            \n"
-			"            const uint AttributeHeadersAddress = {0}_ReadDataAttributeHeadersAddress(DataAddress);\n"
-			"            const uint PCG_MAX_NUM_ATTRIBUTES = {2};\n"
-			"            \n"
-			"            for (uint AttributeIndex = 0; AttributeIndex < PCG_MAX_NUM_ATTRIBUTES; ++AttributeIndex)\n"
-			"            {\n"
-			"                const uint AttributeHeaderAddress = {0}_ReadAttributeHeaderAddress(AttributeHeadersAddress, AttributeIndex);\n"
-			"                const uint AttributeIdAndStride = {0}_ReadAttributeIdAndStride(AttributeHeaderAddress);\n"
-			"                const uint AttributeAddress = {0}_ReadAttributeAddress(AttributeHeaderAddress);\n"
-			"                \n"
-			"                {1}_WriteAttributeIdAndStride(AttributeHeaderAddress, AttributeIdAndStride);\n"
-			"                {1}_WriteAttributeAddress(AttributeHeaderAddress, AttributeAddress);\n"
-			"            }\n"
-			"        }\n"
-			"    }\n"),
+			"    // Signal kernel executed by copying data count from pin {0} to pin {1} from first thread. Rest of header was already set up by the CPU.\n"
+			"    if (GroupIndex == 0) {1}_WriteNumData({0}_GetNumData());\n"
+			"    AllMemoryBarrier();\n"),
 			{ InFromPin->Properties.Label.ToString(), InOutputPinProps.Label.ToString(), PCGComputeConstants::MAX_NUM_ATTRS });
 	};
 
@@ -1019,9 +989,10 @@ FString UPCGCustomHLSLSettings::GetCookedKernelSource(const TMap<FPCGKernelAttri
 				HeaderWriters += FString::Format(TEXT(
 					"    if (ThreadIndex >= GetNumThreads()) return;\n"
 					"    \n"
-					"    // Write data collection header for pin {0}\n"
-					"    {0}_WriteSinglePointDataCollectionHeader({1});\n"),
-					{ InOutputPinProps.Label.ToString(), InPointCount });
+					"    // Signal kernel executed by writing data count (1) for pin {0} from first thread. Rest of header was already set up by the CPU.\n"
+					"    if (GroupIndex == 0) {0}_WriteNumData(1);\n"
+					"    AllMemoryBarrier();\n"),
+					{ InOutputPinProps.Label.ToString() });
 			};
 
 			for (const UPCGPin* OutputPin : Node->GetOutputPins())
@@ -1077,7 +1048,7 @@ FString UPCGCustomHLSLSettings::GetCookedKernelSource(const TMap<FPCGKernelAttri
 				"    const uint {0}_DataAddress = {0}_ThreadInfo[1];\n"),
 				{ OutputPin->Properties.Label.ToString() });
 
-			// Automatically copy all attributes
+			// Automatically copy value of all attributes for this element.
 			KernelSpecificPreamble += FString::Format(TEXT(
 				"\n"
 				"    // Loop over all attribute headers, if the address is non-zero, then copy it from pin {0} to pin {1}.\n"
diff --git a/Engine/Plugins/PCG/Source/PCG/Private/Compute/PCGDataForGPU.cpp b/Engine/Plugins/PCG/Source/PCG/Private/Compute/PCGDataForGPU.cpp
--- a/Engine/Plugins/PCG/Source/PCG/Private/Compute/PCGDataForGPU.cpp
+++ b/Engine/Plugins/PCG/Source/PCG/Private/Compute/PCGDataForGPU.cpp
@@ -702,6 +702,63 @@ void FPCGDataCollectionDesc::PackDataCollection(const FPCGDataCollection& InData
 	}
 }
 
+void FPCGDataCollectionDesc::PrepareBufferForKernelOutput(TArray<uint32>& OutPackedDataCollection)
+{
+	const uint32 NumData = DataDescs.Num();
+
+	TArray<uint32> DataAddresses;
+	const uint32 PackedDataCollectionSizeBytes = ComputePackedSize(&DataAddresses);
+
+	OutPackedDataCollection.SetNumZeroed(PackedDataCollectionSizeBytes / sizeof(uint32));
+	
+	// Num data - set to zero if writing kernel executes. If kernel doesn't execute, 0 means data collection is empty.
+	OutPackedDataCollection[0] = 0;
+
+	for (uint32 DataIndex = 0; DataIndex < NumData; ++DataIndex)
+	{
+		TRACE_CPUPROFILER_EVENT_SCOPE(FPCGDataCollectionDesc::PackDataItem);
+
+		// Write the data addresses
+		const uint32 CurrentDataAddress = DataAddresses[DataIndex];
+		const uint32 CurrentDataIndex = CurrentDataAddress / sizeof(uint32);
+		OutPackedDataCollection[DataIndex + 1] = CurrentDataAddress;
+
+		// DataHeader: (TypeId, NumAttrs, AttrHeaderStartOffset, TypeInfo), Attr0 Header, Attr1 Header, ..., Attr255 Header
+		// Data: Attr0, Attr1, ...
+		if (DataDescs[DataIndex].Type == EPCGDataType::Point)
+		{
+			const uint32 NumElements = DataDescs[DataIndex].ElementCount;
+
+			const TArray<FPCGKernelAttributeDesc>& AttributeDescs = DataDescs[DataIndex].AttributeDescs;
+			const uint32 NumAttributes = AttributeDescs.Num();
+
+			OutPackedDataCollection[CurrentDataIndex + 0] = /*PointDataTypeId=*/POINT_DATA_TYPE_ID;
+			OutPackedDataCollection[CurrentDataIndex + 1] = NumAttributes;
+			OutPackedDataCollection[CurrentDataIndex + 2] = POINT_DATA_HEADER_PREAMBLE_SIZE_BYTES;
+			OutPackedDataCollection[CurrentDataIndex + 3] = NumElements; // TypeInfo for PointData is just NumPoints
+
+			const uint32 BaseAttributeHeaderAddress = CurrentDataAddress + POINT_DATA_HEADER_PREAMBLE_SIZE_BYTES;
+			uint32 CurrentAttributeAddress = CurrentDataAddress + POINT_DATA_HEADER_SIZE_BYTES;
+
+			for (const FPCGKernelAttributeDesc& AttributeDesc : AttributeDescs)
+			{
+				const uint32 AttributeId = AttributeDesc.Index;
+				const uint32 AttributeStrideBytes = PCGDataForGPUHelpers::GetAttributeTypeStrideBytes(AttributeDesc.Type);
+				const uint32 AttributeNumComponents = AttributeStrideBytes / sizeof(uint32); // E.g. float3 has 3 components
+				const uint32 AttributeHeaderIndex = (BaseAttributeHeaderAddress + AttributeId * ATTRIBUTE_HEADER_SIZE_BYTES) / sizeof(uint32);
+
+				// Pack Position (24 bits for AttributeId, 8 bits for Stride)
+				const uint32 PackedIdAndStride = (AttributeId << 8) + AttributeStrideBytes;
+				OutPackedDataCollection[AttributeHeaderIndex + 0] = PackedIdAndStride;
+
+				OutPackedDataCollection[AttributeHeaderIndex + 1] = CurrentAttributeAddress;
+				CurrentAttributeAddress += NumElements * AttributeNumComponents * 4;
+			}
+		}
+		else { /* TODO: Support non-point data. */ }
+	}
+}
+
 void FPCGDataCollectionDesc::UnpackDataCollection(const TArray<uint8>& InPackedData, FName InPin, FPCGDataCollection& OutDataCollection) const
 {
 	const void* PackedData = InPackedData.GetData();
diff --git a/Engine/Plugins/PCG/Source/PCG/Public/Compute/PCGDataForGPU.h b/Engine/Plugins/PCG/Source/PCG/Public/Compute/PCGDataForGPU.h
--- a/Engine/Plugins/PCG/Source/PCG/Public/Compute/PCGDataForGPU.h
+++ b/Engine/Plugins/PCG/Source/PCG/Public/Compute/PCGDataForGPU.h
@@ -82,6 +82,9 @@ struct FPCGDataCollectionDesc
 	/** Pack a data collection into the GPU data format. DataDescs defines which attributes are packed. */
 	void PackDataCollection(const FPCGDataCollection& InDataCollection, FName InPin, TArray<uint32>& OutPackedDataCollection) const;
 
+	/** Allocates the correct size and sets up header. Initializes data count to 0, which kernel will then overwrite if it executes at least one thread. */
+	void PrepareBufferForKernelOutput(TArray<uint32>& OutPackedDataCollection);
+
 	/** Unpack a buffer of 8-bit uints to a data collection. */
 	void UnpackDataCollection(const TArray<uint8>& InPackedData, FName InPin, FPCGDataCollection& OutDataCollection) const;
 
