[PCG] Minor fixes to Custom HLSL node in logic for Custom kernel mode.

#rnx

diff --git a/Engine/Plugins/PCG/Source/PCG/Private/Compute/Elements/PCGCustomHLSL.cpp b/Engine/Plugins/PCG/Source/PCG/Private/Compute/Elements/PCGCustomHLSL.cpp
--- a/Engine/Plugins/PCG/Source/PCG/Private/Compute/Elements/PCGCustomHLSL.cpp
+++ b/Engine/Plugins/PCG/Source/PCG/Private/Compute/Elements/PCGCustomHLSL.cpp
@@ -156,24 +156,24 @@ int UPCGCustomHLSLSettings::ComputeKernelThreadCount(const UPCGDataBinding* Bind
 	}
 	else if (KernelType == EPCGKernelType::Custom)
 	{
-		auto GetFromFirstInput = [this, Binding]()
+		auto GetFromFirstInput = [this, Binding]() -> int
 		{
-			const UPCGPin* InputPin = GetPointProcessingInputPin();
+			const UPCGPin* InputPin = GetFirstInputPin();
 			return GetProcessingElemCountForInputPin(InputPin, Binding);
 		};
 
-		auto GetFromSecondInput = [this, Binding]()
+		auto GetFromSecondInput = [this, Binding]() -> int
 		{
-			const UPCGPin* InputPin = GetSecondPointProcessingInputPin();
-			return GetProcessingElemCountForInputPin(InputPin, Binding);
+			const UPCGPin* InputPin = GetSecondInputPin();
+			return InputPin ? GetProcessingElemCountForInputPin(InputPin, Binding) : 0;
 		};
 
-		auto GetFromFirstInputXSecondInput = [this, Binding]()
+		auto GetFromFirstInputXSecondInput = [this, Binding]() -> int
 		{
-			const UPCGPin* FirstPin = GetPointProcessingInputPin();
-			const int FirstElemCount = GetProcessingElemCountForInputPin(FirstPin, Binding);
-			const UPCGPin* SecondPin = GetSecondPointProcessingInputPin();
-			const int SecondElemCount = GetProcessingElemCountForInputPin(SecondPin, Binding);
+			const UPCGPin* FirstPin = GetFirstInputPin();
+			const UPCGPin* SecondPin = GetSecondInputPin();
+			const int FirstElemCount = FirstPin ? GetProcessingElemCountForInputPin(FirstPin, Binding) : 0;
+			const int SecondElemCount = SecondPin ? GetProcessingElemCountForInputPin(SecondPin, Binding) : 0;
 			return FirstElemCount * SecondElemCount;
 		};
 
@@ -344,19 +344,21 @@ FPCGDataCollectionDesc UPCGCustomHLSLSettings::ComputeOutputPinDataDesc(const UP
 		}
 	}
 
+	const bool bNeedsPointPin = KernelType == EPCGKernelType::PointProcessor || KernelType == EPCGKernelType::PointGenerator;
+
 	// No size set by kernel, fall back to pin settings.
 	if (OutputPin->Properties.BufferSizeMode == EPCGPinBufferSizeMode::FromFirstPin)
 	{
-		if (const UPCGPin* PointProcessingInputPin = GetPointProcessingInputPin())
+		if (const UPCGPin* InputPin = bNeedsPointPin ? GetPointProcessingInputPin() : GetFirstInputPin())
 		{
-			PinDesc = ComputeInputPinDataDesc(PointProcessingInputPin, Binding);
+			PinDesc = ComputeInputPinDataDesc(InputPin, Binding);
 		}
 	}
 	else if (OutputPin->Properties.BufferSizeMode == EPCGPinBufferSizeMode::FromSecondPin)
 	{
-		if (const UPCGPin* PointProcessingInputPin = GetSecondPointProcessingInputPin())
+		if (const UPCGPin* InputPin = bNeedsPointPin ? GetSecondPointProcessingInputPin() : GetSecondInputPin())
 		{
-			PinDesc = ComputeInputPinDataDesc(PointProcessingInputPin, Binding);
+			PinDesc = ComputeInputPinDataDesc(InputPin, Binding);
 		}
 	}
 	else if (OutputPin->Properties.BufferSizeMode == EPCGPinBufferSizeMode::FirstPinXSecondPin)
@@ -412,6 +414,26 @@ EPCGChangeType UPCGCustomHLSLSettings::GetChangeTypeForProperty(const FName& InP
 }
 #endif
 
+const UPCGPin* UPCGCustomHLSLSettings::GetFirstInputPin() const
+{
+	if (const UPCGNode* Node = Cast<UPCGNode>(GetOuter()))
+	{
+		return !Node->GetInputPins().IsEmpty() ? Node->GetInputPins()[0] : nullptr;
+	}
+
+	return nullptr;
+}
+
+const UPCGPin* UPCGCustomHLSLSettings::GetSecondInputPin() const
+{
+	if (const UPCGNode* Node = Cast<UPCGNode>(GetOuter()))
+	{
+		return (Node->GetInputPins().Num() > 1) ? Node->GetInputPins()[1] : nullptr;
+	}
+
+	return nullptr;
+}
+
 const UPCGPin* UPCGCustomHLSLSettings::GetPointProcessingInputPin() const
 {
 	if (const UPCGNode* Node = Cast<UPCGNode>(GetOuter()))
@@ -763,7 +785,9 @@ void UPCGCustomHLSLSettings::UpdatePinSettings()
 		}
 
 		// Primary pin settings driven by kernel (if not custom kernel type).
-		Properties.bDisplayBufferSizeSettings = PinIndex > 0 || KernelType == EPCGKernelType::Custom;
+		const bool bPinCanBeSized = PinIndex > 0 || KernelType == EPCGKernelType::Custom;
+		const bool bDataCanBeSized = Properties.AllowedTypes == EPCGDataType::Point;
+		Properties.bDisplayBufferSizeSettings = bPinCanBeSized && bDataCanBeSized;
 
 		// Output pins should always allow multiple connections.
 		// TODO this could be hoisted up somewhere in the future.
diff --git a/Engine/Plugins/PCG/Source/PCG/Public/Compute/Elements/PCGCustomHLSL.h b/Engine/Plugins/PCG/Source/PCG/Public/Compute/Elements/PCGCustomHLSL.h
--- a/Engine/Plugins/PCG/Source/PCG/Public/Compute/Elements/PCGCustomHLSL.h
+++ b/Engine/Plugins/PCG/Source/PCG/Public/Compute/Elements/PCGCustomHLSL.h
@@ -14,7 +14,7 @@ enum class EPCGKernelType : uint8
 {
 	PointProcessor UMETA(Tooltip = "Kernel executes on each point in first input pin."),
 	PointGenerator UMETA(Tooltip = "Kernel executes for fixed number of points, configurable on node."),
-	Custom UMETA(Tooltip = "Execution thread counts and output buffer sizes configurable on node."),
+	Custom UMETA(Tooltip = "Execution thread counts and output buffer sizes configurable on node. All data read/write indices must be manually bounds checked."),
 };
 
 /** Total number of threads that will be dispatched for this kernel. */
@@ -84,6 +84,8 @@ public:
 	int GetPointCount() const { return PointCount; }
 	int GetFixedThreadCount() const { return FixedThreadCount; }
 
+	const UPCGPin* GetFirstInputPin() const;
+	const UPCGPin* GetSecondInputPin() const;
 	const UPCGPin* GetPointProcessingInputPin() const;
 	const UPCGPin* GetSecondPointProcessingInputPin() const;
 	const UPCGPin* GetFirstOutputPin() const;
diff --git a/Engine/Plugins/PCG/Source/PCG/Public/PCGPin.h b/Engine/Plugins/PCG/Source/PCG/Public/PCGPin.h
--- a/Engine/Plugins/PCG/Source/PCG/Public/PCGPin.h
+++ b/Engine/Plugins/PCG/Source/PCG/Public/PCGPin.h
@@ -79,10 +79,10 @@ struct PCG_API FPCGPinProperties
 #endif
 
 	/** Compute graphs use this to calculate the buffer size of output pins. */
-	UPROPERTY(BlueprintReadWrite, EditAnywhere, Category = Settings, meta = (EditCondition = "bDisplayBufferSizeSettings && AllowedTypes == EPCGDataType::Point", EditConditionHides))
+	UPROPERTY(BlueprintReadWrite, EditAnywhere, Category = Settings, meta = (EditCondition = "bDisplayBufferSizeSettings", EditConditionHides))
 	EPCGPinBufferSizeMode BufferSizeMode = EPCGPinBufferSizeMode::FromFirstPin;
 
-	UPROPERTY(BlueprintReadWrite, EditAnywhere, Category = Settings, meta = (EditCondition = "bDisplayBufferSizeSettings && AllowedTypes == EPCGDataType::Point && BufferSizeMode == EPCGPinBufferSizeMode::FixedElementCount", EditConditionHides))
+	UPROPERTY(BlueprintReadWrite, EditAnywhere, Category = Settings, meta = (EditCondition = "bDisplayBufferSizeSettings && BufferSizeMode == EPCGPinBufferSizeMode::FixedElementCount", EditConditionHides))
 	int FixedBufferElementCount = 4;
 
 	// Multiple connections are only possible if we support multi data.
