[PCG] Change Indirection Settings to a ObjectPtr, since it is necessary to be loaded to get the pins anyway

* Fix issue where the pins disappeared on load for PCG BP Settings (such as the old PCG Data Asset)

#rnx

diff --git a/Engine/Plugins/PCG/Source/PCG/Private/Elements/PCGIndirectionElement.cpp b/Engine/Plugins/PCG/Source/PCG/Private/Elements/PCGIndirectionElement.cpp
--- a/Engine/Plugins/PCG/Source/PCG/Private/Elements/PCGIndirectionElement.cpp
+++ b/Engine/Plugins/PCG/Source/PCG/Private/Elements/PCGIndirectionElement.cpp
@@ -20,13 +20,13 @@ namespace PCGIndirectionSettings
 #if WITH_EDITOR
 void UPCGIndirectionSettings::GetStaticTrackedKeys(FPCGSelectionKeyToSettingsMap& OutKeysToSettings, TArray<TObjectPtr<const UPCGGraph>>& OutVisitedGraphs) const
 {
-	if (IsPropertyOverriddenByPin(GET_MEMBER_NAME_CHECKED(UPCGIndirectionSettings, Settings)) || Settings.IsNull())
+	if (IsPropertyOverriddenByPin(GET_MEMBER_NAME_CHECKED(UPCGIndirectionSettings, Settings)) || !Settings)
 	{
 		// Dynamic tracking or null settings
 		return;
 	}
 
-	FPCGSelectionKey Key = FPCGSelectionKey::CreateFromPath(Settings.ToSoftObjectPath());
+	FPCGSelectionKey Key = FPCGSelectionKey::CreateFromPath(Settings);
 
 	OutKeysToSettings.FindOrAdd(Key).Emplace(this, /*bCulling=*/false);
 }
@@ -73,7 +73,7 @@ FString UPCGIndirectionSettings::GetAdditionalTitleInformation() const
 		}
 		break;
 	case EPCGProxyInterfaceMode::BySettings:
-		if (const UPCGSettings* SettingsPtr = Settings.LoadSynchronous())
+		if (const UPCGSettings* SettingsPtr = Settings.Get())
 		{
 			return SettingsPtr->GetName();
 		}
@@ -110,7 +110,7 @@ TArray<FPCGPinProperties> UPCGIndirectionSettings::InputPinProperties() const
 			}
 			break;
 		case EPCGProxyInterfaceMode::BySettings:
-			if (const UPCGSettings* SettingsPtr = Settings.LoadSynchronous())
+			if (const UPCGSettings* SettingsPtr = Settings.Get())
 			{
 				InputProperties = SettingsPtr->DefaultInputPinProperties();
 				bSetProperties = true;
@@ -159,7 +159,7 @@ TArray<FPCGPinProperties> UPCGIndirectionSettings::OutputPinProperties() const
 			}
 			break;
 		case EPCGProxyInterfaceMode::BySettings:
-			if (const UPCGSettings* SettingsPtr = Settings.LoadSynchronous())
+			if (const UPCGSettings* SettingsPtr = Settings.Get())
 			{
 				OutputProperties = SettingsPtr->DefaultOutputPinProperties();
 			}
@@ -228,7 +228,7 @@ bool FPCGIndirectionElement::PrepareDataInternal(FPCGContext* InContext) const
 	const UPCGIndirectionSettings* Settings = Context->GetInputSettings<UPCGIndirectionSettings>();
 	check(Settings);
 
-	Context->InnerSettings = Settings->Settings.LoadSynchronous();
+	Context->InnerSettings = Settings->Settings.Get();
 
 	if(UPCGSettings* InnerSettings = Context->InnerSettings)
 	{
diff --git a/Engine/Plugins/PCG/Source/PCG/Public/Elements/PCGIndirectionElement.h b/Engine/Plugins/PCG/Source/PCG/Public/Elements/PCGIndirectionElement.h
--- a/Engine/Plugins/PCG/Source/PCG/Public/Elements/PCGIndirectionElement.h
+++ b/Engine/Plugins/PCG/Source/PCG/Public/Elements/PCGIndirectionElement.h
@@ -62,7 +62,7 @@ public:
 
 	/** The element settings, which can be overriden, that will be used during the proxy execution */
 	UPROPERTY(EditAnywhere, BlueprintReadWrite, Category = Settings, meta = (PCG_Overridable))
-	TSoftObjectPtr<UPCGSettings> Settings;
+	TObjectPtr<UPCGSettings> Settings;
 
 	UPROPERTY(BlueprintReadWrite, EditAnywhere, Category = Settings, meta = (PCG_Overridable))
 	bool bTagOutputsBasedOnOutputPins = true;
