[PCG] Support Inspection on GPU nodes to help debugging.

Decision to readback or not is an execution time decision, not a compilation time decision. Remove the specialized readback DI that was selected during compilation, and just move readback functionality into the data collection DI. It's likely we'll also absorb the specialized upload DI into the base data collection DI as well but this is left for future work.

#rnx

diff --git a/Engine/Plugins/PCG/Source/PCG/Private/Compute/DataInterfaces/PCGDataCollectionDataInterface.cpp b/Engine/Plugins/PCG/Source/PCG/Private/Compute/DataInterfaces/PCGDataCollectionDataInterface.cpp
--- a/Engine/Plugins/PCG/Source/PCG/Private/Compute/DataInterfaces/PCGDataCollectionDataInterface.cpp
+++ b/Engine/Plugins/PCG/Source/PCG/Private/Compute/DataInterfaces/PCGDataCollectionDataInterface.cpp
@@ -2,9 +2,11 @@
 
 #include "Compute/DataInterfaces/PCGDataCollectionDataInterface.h"
 
+#include "PCGComponent.h"
 #include "PCGModule.h"
 #include "PCGSettings.h"
 #include "Compute/PCGDataBinding.h"
+#include "Compute/Elements/PCGComputeGraphElement.h"
 
 #include "ComputeFramework/ShaderParamTypeDefinition.h"
 #include "RenderGraphBuilder.h"
@@ -463,18 +465,145 @@ UComputeDataProvider* UPCGDataCollectionDataInterface::CreateDataProvider(TObjec
 
 	UPCGDataCollectionDataProvider* Provider = NewObject<UPCGDataCollectionDataProvider>();
 	Provider->Binding = Binding;
+	Provider->ProducerSettings = ProducerSettings;
 	Provider->PinDesc = ProducerSettings->ComputeOutputPinDataDesc(OutputPinLabel, Binding);
 
+	Provider->ReadbackMode = bRequiresReadback ? EPCGReadbackMode::GraphOutput : EPCGReadbackMode::None;
+#if WITH_EDITOR
+	if (Binding->SourceComponent->IsInspecting())
+	{
+		Provider->ReadbackMode |= EPCGReadbackMode::Inspection;
+	}
+#endif
+
+	// Use the aliased label for normal data output as this is the output from the compute graph.
+	Provider->OutputPinLabelAlias = OutputPinLabelAlias;
+
+	// The original label is needed to store inspection data.
+	Provider->OutputPinLabel = OutputPinLabel;
+
 	return Provider;
 }
 
+bool UPCGDataCollectionDataInterface::GetRequiresReadback() const
+{
+#if WITH_EDITOR
+	// In editor DIs are always flagged as requiring readback, and decision whether to actually readback data or not
+	// is made at execution time (in GetReadbackData) based on inspection state.
+	return true;
+#else
+	return bRequiresReadback;
+#endif
+}
+
 FComputeDataProviderRenderProxy* UPCGDataCollectionDataProvider::GetRenderProxy()
 {
-	FPCGDataCollectionDataProviderProxy* Proxy = new FPCGDataCollectionDataProviderProxy(Binding, PinDesc);
+	FPCGDataCollectionDataProviderProxy* Proxy = new FPCGDataCollectionDataProviderProxy(Binding, PinDesc, ReadbackMode);
+
+	if (ReadbackMode != EPCGReadbackMode::None)
+	{
+		TWeakObjectPtr<UPCGDataCollectionDataProvider> ThisWeakPtr(this);
+
+		Proxy->AsyncReadbackCallback_RenderThread = [ThisWeakPtr](const void* InData, int InNumBytes)
+		{
+			TRACE_CPUPROFILER_EVENT_SCOPE(UPCGDataCollectionDataProvider::ProcessReadbackData_RenderThread);
+
+			if (!ThisWeakPtr.IsValid() || !ThisWeakPtr.Get())
+			{
+				return;
+			}
+
+			UPCGDataCollectionDataProvider* ThisDataProvider = ThisWeakPtr.Get();
+			check(ThisDataProvider);
+
+			// We should never find ourselves stomping existing data.
+			check(ThisDataProvider->RawReadbackData.IsEmpty());
+
+			if (InData && InNumBytes > 0)
+			{
+				// Copy the data to temp storage for the game thread to pick up.
+				ThisDataProvider->RawReadbackData.SetNumUninitialized(InNumBytes);
+				FMemory::Memcpy(ThisDataProvider->RawReadbackData.GetData(), InData, InNumBytes);
+			}
+			else
+			{
+				// Can happen currently if no threads dispatched.
+				ThisDataProvider->RawReadbackData.Reset();
+			}
+
+			ThisDataProvider->bReadbackComplete = true;
+
+			ThisDataProvider->OnReadbackComplete.Broadcast();
+		};
+	}
 
 	return Proxy;
 }
 
+bool UPCGDataCollectionDataProvider::ProcessReadBackData(FPCGComputeGraphContext* InContext)
+{
+	check(InContext);
+	check(ProducerSettings);
+
+	if (!ensure(bReadbackComplete))
+	{
+		// This should not be called until readback has been done.
+		return false;
+	}
+
+	if (RawReadbackData.IsEmpty())
+	{
+		// If there's no data, we leave an empty data collection and we're done.
+		return true;
+	}
+
+	FPCGDataCollection DataFromGPU;
+	PinDesc.UnpackDataCollection(RawReadbackData, OutputPinLabelAlias, DataFromGPU);
+	RawReadbackData.Reset();
+
+	// Store data in output collection.
+	if (!!(ReadbackMode & EPCGReadbackMode::GraphOutput) && ensure(Binding.IsValid()))
+	{
+		Binding->OutputDataCollection.TaggedData.Append(DataFromGPU.TaggedData);
+	}
+
+	// Store data for Inspection.
+#if WITH_EDITOR
+	if (!!(ReadbackMode & EPCGReadbackMode::Inspection))
+	{
+		const UPCGNode* Node = Cast<UPCGNode>(ProducerSettings->GetOuter());
+		UPCGComponent* Component = InContext->SourceComponent.Get();
+		if (Component && InContext->Stack && Node)
+		{
+			// Virtual pin labels confuse inspection. Apply the original output label before storing.
+			for (FPCGTaggedData& Data : DataFromGPU.TaggedData)
+			{
+				Data.Pin = OutputPinLabel;
+			}
+
+			// Required by inspection code.
+			DataFromGPU.ComputeCrcs(/*bFullDataCrc=*/false);
+
+			// Input data not yet supported.
+			Component->StoreInspectionData(InContext->Stack, Node, /*InTimer=*/nullptr, /*InInputData=*/{}, DataFromGPU, /*bUsedCache*/false);
+		}
+	}
+#endif // WITH_EDITOR
+
+	return true;
+}
+
+FPCGDataCollectionDataProviderProxy::FPCGDataCollectionDataProviderProxy(
+	TWeakObjectPtr<UPCGDataBinding> InBinding,
+	const FPCGDataCollectionDesc& InPinDesc,
+	EPCGReadbackMode InReadbackMode)
+	: ReadbackMode(InReadbackMode)
+	, Binding(InBinding)
+	, PinDesc(InPinDesc)
+{
+	SizeBytes = PinDesc.ComputePackedSize();
+}
+
 bool FPCGDataCollectionDataProviderProxy::IsValid(FValidationData const& InValidationData) const
 {
 	if (InValidationData.ParameterStructSize != sizeof(FParameters))
@@ -488,6 +617,12 @@ bool FPCGDataCollectionDataProviderProxy::IsValid(FValidationData const& InValid
 		return false;
 	}
 
+	if (SizeBytes <= 0)
+	{
+		UE_LOG(LogPCG, Error, TEXT("Proxy invalid due to invalid buffer size."));
+		return false;
+	}
+
 	return true;
 }
 
@@ -505,6 +640,8 @@ void FPCGDataCollectionDataProviderProxy::GatherDispatchData(FDispatchData const
 
 void FPCGDataCollectionDataProviderProxy::AllocateResources(FRDGBuilder& GraphBuilder, FAllocationData const& InAllocationData)
 {
+	check(SizeBytes > 0);
+
 	// Initialize with an empty data collection. The kernel may not run, for example if indirect dispatch args end up being 0. Ensure
 	// there is something meaningful to readback.
 	// TODO could have a statically-allocated resource rather than allocating & uploading here.
@@ -517,3 +654,16 @@ void FPCGDataCollectionDataProviderProxy::AllocateResources(FRDGBuilder& GraphBu
 
 	GraphBuilder.QueueBufferUpload(Buffer, PackedDataCollection.GetData(), PackedDataCollection.Num() * PackedDataCollection.GetTypeSize(), ERDGInitialDataFlags::None);
 }
+
+void FPCGDataCollectionDataProviderProxy::GetReadbackData(TArray<FReadbackData>& OutReadbackData) const
+{
+	if (ReadbackMode != EPCGReadbackMode::None)
+	{
+		FReadbackData Data;
+		Data.Buffer = Buffer;
+		Data.NumBytes = SizeBytes;
+		Data.ReadbackCallback_RenderThread = &AsyncReadbackCallback_RenderThread;
+
+		OutReadbackData.Add(MoveTemp(Data));
+	}
+}
diff --git a/Engine/Plugins/PCG/Source/PCG/Private/Compute/DataInterfaces/PCGDataCollectionDataInterface.h b/Engine/Plugins/PCG/Source/PCG/Private/Compute/DataInterfaces/PCGDataCollectionDataInterface.h
--- a/Engine/Plugins/PCG/Source/PCG/Private/Compute/DataInterfaces/PCGDataCollectionDataInterface.h
+++ b/Engine/Plugins/PCG/Source/PCG/Private/Compute/DataInterfaces/PCGDataCollectionDataInterface.h
@@ -9,11 +9,23 @@
 
 #include "PCGDataCollectionDataInterface.generated.h"
 
+struct FPCGComputeGraphContext;
 class FPCGDataCollectionDataInterfaceParameters;
 class FRDGBuffer;
 class FRDGBufferSRV;
 class UPCGDataBinding;
 
+UENUM(meta = (Bitflags))
+enum class EPCGReadbackMode : uint8
+{
+	None = 0,
+	/** Read back and store in compute graph output data collection. */
+	GraphOutput = 1 << 1,
+	/** Read back and store inspection data. */
+	Inspection = 1 << 2,
+};
+ENUM_CLASS_FLAGS(EPCGReadbackMode);
+
 /** Compute Framework Data Interface for reading PCG data. */
 UCLASS(ClassGroup = (Procedural))
 class UPCGDataCollectionDataInterface : public UPCGComputeDataInterface
@@ -36,11 +48,19 @@ public:
 	UComputeDataProvider* CreateDataProvider(TObjectPtr<UObject> InBinding, uint64 InInputMask, uint64 InOutputMask) const override;
 	//~ End UComputeDataInterface Interface
 
+	void SetProducerSettings(const UPCGSettings* InProducerSettings) { ProducerSettings = InProducerSettings; }
+
+	bool GetRequiresReadback() const override;
+	void SetRequiresReadback(bool bInRequiresReadback) { bRequiresReadback = bInRequiresReadback; }
+
+protected:
 	/** Settings of node that produces this data, normally the upstream node. */
 	UPROPERTY()
 	const UPCGSettings* ProducerSettings;
 
-private:
+	/** Whether this data should be copied from GPU to CPU. */
+	bool bRequiresReadback = false;
+
 	static TCHAR const* TemplateFilePath;
 };
 
@@ -55,25 +75,57 @@ public:
 	FComputeDataProviderRenderProxy* GetRenderProxy() override;
 	//~ End UComputeDataProvider Interface
 
+	bool RequiresReadback() const { return ReadbackMode != EPCGReadbackMode::None; }
+
+	/** Processed read back data. Should be called after readback complete. Called from game thread. Returns true if data processed successfully. */
+	bool ProcessReadBackData(FPCGComputeGraphContext* InContext);
+
+	DECLARE_EVENT(UPCGDataProviderDataCollectionReadback, FOnReadbackComplete);
+	FOnReadbackComplete& OnReadbackComplete_RenderThread() { return OnReadbackComplete; }
+
+public:
+	/** Settings of node that produces this data, normally the upstream node. */
+	UPROPERTY()
+	const UPCGSettings* ProducerSettings;
+
 	TWeakObjectPtr<UPCGDataBinding> Binding;
 	FPCGDataCollectionDesc PinDesc;
+
+	EPCGReadbackMode ReadbackMode;
+
+	FName OutputPinLabel;
+	FName OutputPinLabelAlias;
+
+protected:
+	std::atomic<bool> bReadbackComplete = false;
+	TArray<uint8> RawReadbackData;
+
+	FOnReadbackComplete OnReadbackComplete;
 };
 
 class FPCGDataCollectionDataProviderProxy : public FComputeDataProviderRenderProxy
 {
 public:
-	FPCGDataCollectionDataProviderProxy(TWeakObjectPtr<UPCGDataBinding> InBinding, const FPCGDataCollectionDesc& InPinDesc)
-		: Binding(InBinding)
-		, PinDesc(InPinDesc)
-	{}
+	FPCGDataCollectionDataProviderProxy(
+		TWeakObjectPtr<UPCGDataBinding> InBinding,
+		const FPCGDataCollectionDesc& InPinDesc,
+		EPCGReadbackMode InReadbackMode);
 
 	//~ Begin FComputeDataProviderRenderProxy Interface
 	bool IsValid(FValidationData const& InValidationData) const override;
 	void GatherDispatchData(FDispatchData const& InDispatchData) override;
 	void AllocateResources(FRDGBuilder& GraphBuilder, FAllocationData const& InAllocationData) override;
+	void GetReadbackData(TArray<FReadbackData>& OutReadbackData) const override;
 	//~ End FComputeDataProviderRenderProxy Interface
 
+public:
+	FReadbackCallback AsyncReadbackCallback_RenderThread;
+
 protected:
+	int SizeBytes = -1;
+
+	EPCGReadbackMode ReadbackMode;
+
 	using FParameters = FPCGDataCollectionDataInterfaceParameters;
 
 	TWeakObjectPtr<UPCGDataBinding> Binding;
diff --git a/Engine/Plugins/PCG/Source/PCG/Private/Compute/DataInterfaces/PCGDataCollectionReadbackDataInterface.cpp b/Engine/Plugins/PCG/Source/PCG/Private/Compute/DataInterfaces/PCGDataCollectionReadbackDataInterface.cpp
deleted file mode 100644
diff --git a/Engine/Plugins/PCG/Source/PCG/Private/Compute/DataInterfaces/PCGDataCollectionReadbackDataInterface.h b/Engine/Plugins/PCG/Source/PCG/Private/Compute/DataInterfaces/PCGDataCollectionReadbackDataInterface.h
deleted file mode 100644
diff --git a/Engine/Plugins/PCG/Source/PCG/Private/Compute/DataInterfaces/PCGDataCollectionUploadDataInterface.cpp b/Engine/Plugins/PCG/Source/PCG/Private/Compute/DataInterfaces/PCGDataCollectionUploadDataInterface.cpp
--- a/Engine/Plugins/PCG/Source/PCG/Private/Compute/DataInterfaces/PCGDataCollectionUploadDataInterface.cpp
+++ b/Engine/Plugins/PCG/Source/PCG/Private/Compute/DataInterfaces/PCGDataCollectionUploadDataInterface.cpp
@@ -41,7 +41,7 @@ FPCGDataProviderDataCollectionUploadProxy::FPCGDataProviderDataCollectionUploadP
 	TWeakObjectPtr<UPCGDataBinding> InBinding,
 	const FPCGDataCollectionDesc& InPinDesc,
 	const TArray<FName>& InDownstreamInputPinLabels)
-	: FPCGDataCollectionDataProviderProxy(InBinding, InPinDesc)
+	: FPCGDataCollectionDataProviderProxy(InBinding, InPinDesc, EPCGReadbackMode::None)
 {
 	DownstreamInputPinLabels = InDownstreamInputPinLabels;
 }
diff --git a/Engine/Plugins/PCG/Source/PCG/Private/Compute/Elements/PCGComputeGraphElement.cpp b/Engine/Plugins/PCG/Source/PCG/Private/Compute/Elements/PCGComputeGraphElement.cpp
--- a/Engine/Plugins/PCG/Source/PCG/Private/Compute/Elements/PCGComputeGraphElement.cpp
+++ b/Engine/Plugins/PCG/Source/PCG/Private/Compute/Elements/PCGComputeGraphElement.cpp
@@ -5,7 +5,7 @@
 #include "PCGComponent.h"
 #include "PCGModule.h"
 #include "PCGSubsystem.h"
-#include "Compute/DataInterfaces/PCGDataCollectionReadbackDataInterface.h"
+#include "Compute/DataInterfaces/PCGDataCollectionDataInterface.h"
 
 #include "ComputeWorkerInterface.h"
 #include "ComputeFramework/ComputeFramework.h"
@@ -46,17 +46,18 @@ bool FPCGComputeGraphElement::ExecuteInternal(FPCGContext* InContext) const
 	// The sequence appears out of order as calls so that paths that are executed multiple times (like checking for completion) are as
 	// short as possible.
 
-	// 7. Execution is is complete when any async readbacks are complete.
+	// 7. Execution is complete when any async readbacks are complete.
 	if (Context->bAllAsyncOperationsDone)
 	{
 		Context->bExecutionSuccess = true;
 
 		for (UComputeDataProvider* DataProvider : Context->ComputeGraphInstance.GetDataProviders())
 		{
-			if (UPCGDataProviderDataCollectionReadback* Readback = Cast<UPCGDataProviderDataCollectionReadback>(DataProvider))
+			UPCGDataCollectionDataProvider* PCGDataProvider = Cast<UPCGDataCollectionDataProvider>(DataProvider);
+			if (PCGDataProvider && PCGDataProvider->RequiresReadback())
 			{
 				// Process data for all readbacks, and track whether all succeeded.
-				const bool bProcessResult = Readback->ProcessReadBackData();
+				const bool bProcessResult = PCGDataProvider->ProcessReadBackData(Context);
 				Context->bExecutionSuccess &= bProcessResult;
 			}
 		}
@@ -125,18 +126,19 @@ bool FPCGComputeGraphElement::ExecuteInternal(FPCGContext* InContext) const
 		// Register all providers running async operations. TODO review if we should have a general API like "RunsAsyncOperations()"?
 		for (UComputeDataProvider* DataProvider : Context->ComputeGraphInstance.GetDataProviders())
 		{
-			if (UPCGDataProviderDataCollectionReadback* Readback = Cast<UPCGDataProviderDataCollectionReadback>(DataProvider))
+			UPCGDataCollectionDataProvider* PCGDataProvider = Cast<UPCGDataCollectionDataProvider>(DataProvider);
+			if (PCGDataProvider && PCGDataProvider->RequiresReadback())
 			{
-				Context->ProvidersRunningAsyncOperations.Add(Readback);
+				Context->ProvidersRunningAsyncOperations.Add(PCGDataProvider);
 
-				Readback->OnReadbackComplete_RenderThread().AddLambda([Context, Readback]()
+				PCGDataProvider->OnReadbackComplete_RenderThread().AddLambda([Context, PCGDataProvider]()
 				{
 					FWriteScopeLock Lock(Context->ProvidersRunningAsyncOperationsLock);
 
 					const bool bEmptyBefore = Context->ProvidersRunningAsyncOperations.IsEmpty();
 
-					ensure(Context->ProvidersRunningAsyncOperations.Contains(Readback));
-					Context->ProvidersRunningAsyncOperations.Remove(Readback);
+					ensure(Context->ProvidersRunningAsyncOperations.Contains(PCGDataProvider));
+					Context->ProvidersRunningAsyncOperations.Remove(PCGDataProvider);
 
 					if (!bEmptyBefore && Context->ProvidersRunningAsyncOperations.IsEmpty())
 					{
@@ -232,9 +234,7 @@ void FPCGComputeGraphElement::PostExecuteInternal(FPCGContext* InContext) const
 			UPCGComponent* Component = Context->SourceComponent.Get();
 			if (Component && Context->Stack && Node)
 			{
-				// TODO calling this to register that node executed. Regarding inspection data, we need to create
-				// a readback and pipe the data to here.
-				Component->StoreInspectionData(Context->Stack, Node, /*InTimer=*/nullptr, {}, {}, /*bUsedCache*/false);
+				Component->NotifyNodeExecuted(Node, Context->Stack, /*InTimer=*/nullptr, /*bUsedCache*/false);
 			}
 		}
 	}
@@ -254,9 +254,10 @@ void FPCGComputeGraphElement::ResetAsyncOperations(FPCGContext* InContext) const
 
 		for (UComputeDataProvider* DataProvider : Context->ComputeGraphInstance.GetDataProviders())
 		{
-			if (UPCGDataProviderDataCollectionReadback* Readback = Cast<UPCGDataProviderDataCollectionReadback>(DataProvider))
+			UPCGDataCollectionDataProvider* PCGDataProvider = Cast<UPCGDataCollectionDataProvider>(DataProvider);
+			if (PCGDataProvider && PCGDataProvider->RequiresReadback())
 			{
-				Readback->OnReadbackComplete_RenderThread().Clear();
+				PCGDataProvider->OnReadbackComplete_RenderThread().Clear();
 			}
 		}
 
diff --git a/Engine/Plugins/PCG/Source/PCG/Private/Graph/PCGGraphCompilerGPU.cpp b/Engine/Plugins/PCG/Source/PCG/Private/Graph/PCGGraphCompilerGPU.cpp
--- a/Engine/Plugins/PCG/Source/PCG/Private/Graph/PCGGraphCompilerGPU.cpp
+++ b/Engine/Plugins/PCG/Source/PCG/Private/Graph/PCGGraphCompilerGPU.cpp
@@ -10,7 +10,6 @@
 #include "Compute/DataInterfaces/PCGComputeDataInterface.h"
 #include "Compute/DataInterfaces/PCGCustomKernelDataInterface.h"
 #include "Compute/DataInterfaces/PCGDataCollectionDataInterface.h"
-#include "Compute/DataInterfaces/PCGDataCollectionReadbackDataInterface.h"
 #include "Compute/DataInterfaces/PCGDataCollectionUploadDataInterface.h"
 #include "Compute/DataInterfaces/PCGDebugDataInterface.h"
 #include "Compute/DataInterfaces/PCGLandscapeDataInterface.h"
@@ -500,25 +499,18 @@ void FPCGGraphCompilerGPU::BuildGPUGraphTask(
 
 			if (bUpstreamIsGPUTask)
 			{
-				if (bRequiresReadback)
-				{
-					// GPU -> CPU
-					DataInterfacePCGData = NewObject<UPCGDataCollectionReadbackDataInterface>(ComputeGraph);
-				}
-				else
-				{
-					// GPU -> GPU
-					DataInterfacePCGData = NewObject<UPCGDataCollectionDataInterface>(ComputeGraph);
-				}
+				// Provides data for GPU -> GPU and GPU -> CPU edges.
+				DataInterfacePCGData = NewObject<UPCGDataCollectionDataInterface>(ComputeGraph);
 			}
 			else
 			{
-				// CPU -> GPU
+				// Provides data for CPU -> GPU edge.
 				DataInterfacePCGData = NewObject<UPCGDataCollectionUploadDataInterface>(ComputeGraph);
 			}
 
 			check(DataInterfacePCGData);
-			DataInterfacePCGData->ProducerSettings = InOutCompiledTasks[InTaskId].Node ? InOutCompiledTasks[InTaskId].Node->GetSettings() : nullptr;
+			DataInterfacePCGData->SetProducerSettings(InOutCompiledTasks[InTaskId].Node ? InOutCompiledTasks[InTaskId].Node->GetSettings() : nullptr);
+			DataInterfacePCGData->SetRequiresReadback(bRequiresReadback);
 
 			DataInterface = DataInterfacePCGData;
 			break;
