[PCG] Add seed support to Custom HLSL node and add a basic random stream that appears to match FRandomStream.

#rnx

diff --git a/Engine/Plugins/PCG/Shaders/Private/PCGDataCollectionDataInterface.ush b/Engine/Plugins/PCG/Shaders/Private/PCGDataCollectionDataInterface.ush
--- a/Engine/Plugins/PCG/Shaders/Private/PCGDataCollectionDataInterface.ush
+++ b/Engine/Plugins/PCG/Shaders/Private/PCGDataCollectionDataInterface.ush
@@ -208,6 +208,12 @@ int GetInt_{DataInterfaceName}(uint DataIndex, uint AttributeId, uint ElementInd
 	return asint({DataInterfaceName}_DataCollectionBuffer.Load(AttributeAddress + ElementIndex * 1 * 4));
 }
 
+int GetUint_{DataInterfaceName}(uint DataIndex, uint AttributeId, uint ElementIndex)
+{
+	const uint AttributeAddress = ComputeAddressHelper_{DataInterfaceName}(DataIndex, AttributeId, ElementIndex);
+	return {DataInterfaceName}_DataCollectionBuffer.Load(AttributeAddress + ElementIndex * 1 * 4);
+}
+
 float GetFloat_{DataInterfaceName}(uint DataIndex, uint AttributeId, uint ElementIndex)
 {
 	const uint AttributeAddress = ComputeAddressHelper_{DataInterfaceName}(DataIndex, AttributeId, ElementIndex);
@@ -258,6 +264,12 @@ void SetInt_{DataInterfaceName}(uint DataIndex, uint AttributeId, uint ElementIn
 	{DataInterfaceName}_DataCollectionBuffer.Store(AttributeAddress + ElementIndex * 1 * 4, asuint(Value));
 }
 
+void SetUint_{DataInterfaceName}(uint DataIndex, uint AttributeId, uint ElementIndex, uint Value)
+{
+	const uint AttributeAddress = ComputeAddressHelper_{DataInterfaceName}(DataIndex, AttributeId, ElementIndex);
+	{DataInterfaceName}_DataCollectionBuffer.Store(AttributeAddress + ElementIndex * 1 * 4, Value);
+}
+
 void SetFloat_{DataInterfaceName}(uint DataIndex, uint AttributeId, uint ElementIndex, float Value)
 {
 	const uint AttributeAddress = ComputeAddressHelper_{DataInterfaceName}(DataIndex, AttributeId, ElementIndex);
@@ -336,9 +348,9 @@ float GetSteepness_{DataInterfaceName}(uint DataIndex, uint ElementIndex)
 	return GetFloat_{DataInterfaceName}(DataIndex, POINT_STEEPNESS_ID, ElementIndex);
 }
 
-int GetSeed_{DataInterfaceName}(uint DataIndex, uint ElementIndex)
+uint GetSeed_{DataInterfaceName}(uint DataIndex, uint ElementIndex)
 {
-	return GetInt_{DataInterfaceName}(DataIndex, POINT_SEED_ID, ElementIndex);
+	return GetUint_{DataInterfaceName}(DataIndex, POINT_SEED_ID, ElementIndex);
 }
 
 float4x4 GetPointTransform_{DataInterfaceName}(uint InDataIndex, uint InElementIndex)
@@ -401,17 +413,9 @@ void SetSteepness_{DataInterfaceName}(uint DataIndex, uint ElementIndex, float S
 	SetFloat_{DataInterfaceName}(DataIndex, POINT_STEEPNESS_ID, ElementIndex, Steepness);
 }
 
-void SetSeed_{DataInterfaceName}(uint DataIndex, uint ElementIndex, int Seed)
-{
-	SetInt_{DataInterfaceName}(DataIndex, POINT_SEED_ID, ElementIndex, Seed);
-}
-
-void SetSeedFromPosition_{DataInterfaceName}(uint DataIndex, uint ElementIndex, float3 Position)
+void SetSeed_{DataInterfaceName}(uint DataIndex, uint ElementIndex, uint Seed)
 {
-	const int3 PosInt = (int3)Position;
-	const int Seed = ((PosInt.x * 196314165U) + 907633515U) ^ ((PosInt.y * 73148459U) + 453816763U) ^ ((PosInt.z * 34731343U) + 453816743U);
-
-	SetInt_{DataInterfaceName}(DataIndex, POINT_SEED_ID, ElementIndex, Seed);
+	SetUint_{DataInterfaceName}(DataIndex, POINT_SEED_ID, ElementIndex, Seed);
 }
 
 void SetPointTransform_{DataInterfaceName}(uint InDataIndex, uint InElementIndex, float4x4 Transform)
diff --git a/Engine/Plugins/PCG/Shaders/Private/PCGShaderUtils.ush b/Engine/Plugins/PCG/Shaders/Private/PCGShaderUtils.ush
--- a/Engine/Plugins/PCG/Shaders/Private/PCGShaderUtils.ush
+++ b/Engine/Plugins/PCG/Shaders/Private/PCGShaderUtils.ush
@@ -3,6 +3,32 @@
 #include "/Engine/Private/Common.ush"
 #include "/Engine/Public/Platform.ush"
 
+uint ComputeSeed(uint A, uint B)
+{
+	return ((A * 196314165U) + 907633515U) ^ ((B * 73148459U) + 453816763U);
+}
+
+uint ComputeSeed(uint A, uint B, uint C)
+{
+	return ((A * 196314165U) + 907633515U) ^ ((B * 73148459U) + 453816763U) ^ ((C * 34731343U) + 453816743U);
+}
+
+uint ComputeSeedFromPosition(float3 InPosition)
+{
+	const int3 PosInt = (int3)InPosition;
+	const int Seed = ((PosInt.x * 196314165U) + 907633515U) ^ ((PosInt.y * 73148459U) + 453816763U) ^ ((PosInt.z * 34731343U) + 453816743U);
+	return (uint)Seed;
+}
+
+float FRand(inout uint InOutSeed)
+{
+	InOutSeed = (InOutSeed * 196314165U) + 907633515U;
+	
+	const float Result = asfloat(0x3F800000U | (InOutSeed >> 9));
+
+	return Result - 1.0f;
+}
+
 float3 CreateGrid2D(int ElementIndex, int NumPoints, int NumRows, float3 Min, float3 Max)
 {
     if (NumRows == 0)
diff --git a/Engine/Plugins/PCG/Source/PCG/Private/Compute/DataInterfaces/PCGCustomKernelDataInterface.cpp b/Engine/Plugins/PCG/Source/PCG/Private/Compute/DataInterfaces/PCGCustomKernelDataInterface.cpp
--- a/Engine/Plugins/PCG/Source/PCG/Private/Compute/DataInterfaces/PCGCustomKernelDataInterface.cpp
+++ b/Engine/Plugins/PCG/Source/PCG/Private/Compute/DataInterfaces/PCGCustomKernelDataInterface.cpp
@@ -15,14 +15,16 @@
 
 #include UE_INLINE_GENERATED_CPP_BY_NAME(PCGCustomKernelDataInterface)
 
-const TCHAR* UPCGCustomKernelDataInterface::NumThreadsReservedName = TEXT("NumThreads");
-
 void UPCGCustomKernelDataInterface::GetSupportedInputs(TArray<FShaderFunctionDefinition>& OutFunctions) const
 {
 	OutFunctions.AddDefaulted_GetRef()
-		.SetName(FString::Printf(TEXT("Get%s"), NumThreadsReservedName))
+		.SetName(TEXT("GetNumThreads"))
 		.AddReturnType(FShaderValueType::Get(EShaderFundamentalType::Int));
 
+	OutFunctions.AddDefaulted_GetRef()
+		.SetName(TEXT("GetSeed"))
+		.AddReturnType(FShaderValueType::Get(EShaderFundamentalType::Uint));
+
 	// A convenient way to serve component bounds to all kernels. Could be pulled out into a PCG context DI in the future.
 	OutFunctions.AddDefaulted_GetRef()
 		.SetName(TEXT("GetComponentBoundsMin"))
@@ -35,6 +37,7 @@ void UPCGCustomKernelDataInterface::GetSupportedInputs(TArray<FShaderFunctionDef
 
 BEGIN_SHADER_PARAMETER_STRUCT(FPCGKernelDataInterfaceParameters, )
 	SHADER_PARAMETER(FIntVector3, NumThreads)
+	SHADER_PARAMETER(uint32, Seed)
 	SHADER_PARAMETER(FVector3f, ComponentBoundsMin)
 	SHADER_PARAMETER(FVector3f, ComponentBoundsMax)
 END_SHADER_PARAMETER_STRUCT()
@@ -44,52 +47,24 @@ void UPCGCustomKernelDataInterface::GetShaderParameters(TCHAR const* UID, FShade
 	InOutBuilder.AddNestedStruct<FPCGKernelDataInterfaceParameters>(UID);
 }
 
-void UPCGCustomKernelDataInterface::GetShaderHash(FString& InOutKey) const
-{
-	// UComputeGraph::BuildKernelSource hashes the result of GetHLSL()
-	// Only append additional hashes here if the HLSL contains any additional includes	
-}
-
 void UPCGCustomKernelDataInterface::GetHLSL(FString& OutHLSL, FString const& InDataInterfaceName) const
 {
-	const FString TypeName = FShaderValueType::Get(EShaderFundamentalType::Int, 3)->ToString();
-
-	if (ensure(!TypeName.IsEmpty()))
+	TMap<FString, FStringFormatArg> TemplateArgs =
 	{
-		TMap<FString, FStringFormatArg> TemplateArgs =
-		{
-			{ TEXT("DataInterfaceName"), InDataInterfaceName },
-		};
-
-		// Add uniforms.
-		OutHLSL += FString::Printf(TEXT("%s %s_%s;\n"), 
-			*TypeName, 
-			*InDataInterfaceName, 
-			NumThreadsReservedName);
-			
-		// Add function getters.
-		OutHLSL += FString::Printf(TEXT("%s Get%s_%s()\n{\n\treturn %s_%s;\n}\n"), 
-			*TypeName,
-			NumThreadsReservedName,
-			*InDataInterfaceName, 
-			*InDataInterfaceName, 
-			NumThreadsReservedName);
-
-		// Add helpers (e.g. GetComponentBounds())
-		OutHLSL += FString::Format(TEXT(
-			"float3 {DataInterfaceName}_ComponentBoundsMin;\n"
-			"float3 {DataInterfaceName}_ComponentBoundsMax;\n"
-			"\n"
-			"float3 GetComponentBoundsMin_{DataInterfaceName}()\n"
-			"{\n"
-			"	return {DataInterfaceName}_ComponentBoundsMin;\n"
-			"}\n"
-			"\n"
-			"float3 GetComponentBoundsMax_{DataInterfaceName}()\n"
-			"{\n"
-			"	return {DataInterfaceName}_ComponentBoundsMax;\n"
-			"}\n"), TemplateArgs);
-	}
+		{ TEXT("DataInterfaceName"), InDataInterfaceName },
+	};
+
+	OutHLSL += FString::Format(TEXT(
+		"int3 {DataInterfaceName}_NumThreads;\n"
+		"uint {DataInterfaceName}_Seed;\n"
+		"float3 {DataInterfaceName}_ComponentBoundsMin;\n"
+		"float3 {DataInterfaceName}_ComponentBoundsMax;\n"
+		"\n"
+		"int3 GetNumThreads_{DataInterfaceName}()\n{\n\treturn {DataInterfaceName}_NumThreads;\n}\n\n"
+		"uint GetSeed_{DataInterfaceName}()\n{\n\treturn {DataInterfaceName}_Seed;\n}\n\n"
+		"float3 GetComponentBoundsMin_{DataInterfaceName}()\n{\n\treturn {DataInterfaceName}_ComponentBoundsMin;\n}\n\n"
+		"float3 GetComponentBoundsMax_{DataInterfaceName}()\n{\n\treturn {DataInterfaceName}_ComponentBoundsMax;\n}\n\n"),
+		TemplateArgs);
 }
 
 UComputeDataProvider* UPCGCustomKernelDataInterface::CreateDataProvider(TObjectPtr<UObject> InBinding, uint64 InInputMask, uint64 InOutputMask) const
@@ -99,11 +74,10 @@ UComputeDataProvider* UPCGCustomKernelDataInterface::CreateDataProvider(TObjectP
 	UPCGDataBinding* Binding = CastChecked<UPCGDataBinding>(InBinding);
 	check(Binding->SourceComponent.IsValid() && Binding->SourceComponent.Get());
 
-	const FBox ComponentBounds = Binding->SourceComponent.Get()->GetGridBounds();
-
 	UPCGCustomComputeKernelDataProvider* Provider = NewObject<UPCGCustomComputeKernelDataProvider>();
-	Provider->SourceComponentBounds = ComponentBounds;
 	Provider->ThreadCount = Settings->ComputeKernelThreadCount(Binding);
+	Provider->Seed = static_cast<uint32>(Settings->GetSeed(Binding->SourceComponent.Get()));
+	Provider->SourceComponentBounds = Binding->SourceComponent.Get()->GetGridBounds();
 
 	return Provider;
 }
@@ -118,7 +92,7 @@ FComputeDataProviderRenderProxy* UPCGCustomComputeKernelDataProvider::GetRenderP
 		InvocationCounts.Reset();
 	}
 
-	return new FPCGCustomComputeKernelDataProviderProxy(MoveTemp(InvocationCounts), TotalThreadCount, SourceComponentBounds);
+	return new FPCGCustomComputeKernelDataProviderProxy(MoveTemp(InvocationCounts), TotalThreadCount, Seed, SourceComponentBounds);
 }
 
 bool UPCGCustomComputeKernelDataProvider::GetInvocationThreadCounts(TArray<int32>& OutInvocationThreadCount, int32& OutTotalThreadCount) const
@@ -163,18 +137,14 @@ void FPCGCustomComputeKernelDataProviderProxy::GatherDispatchData(FDispatchData
 	const TStridedView<FParameters> ParameterArray = MakeStridedParameterView<FParameters>(InDispatchData);
 	for (int32 InvocationIndex = 0; InvocationIndex < InDispatchData.NumInvocations; ++InvocationIndex)
 	{
-		int32 NumThreads;
-		if (InDispatchData.bUnifiedDispatch)
-		{
-			NumThreads = TotalThreadCount;
-		}
-		else
-		{
-			NumThreads = InvocationThreadCounts[InvocationIndex];
-		}
-
 		FParameters& Parameters = ParameterArray[InvocationIndex];
-		Parameters.NumThreads = FIntVector(NumThreads, 1, 1);
+
+		// Thread count
+		Parameters.NumThreads.X = InDispatchData.bUnifiedDispatch ? TotalThreadCount : InvocationThreadCounts[InvocationIndex];
+		Parameters.NumThreads.Y = Parameters.NumThreads.Z = 1;
+
+		// Seed for the node
+		Parameters.Seed = Seed;
 
 		// Set component bounds
 		Parameters.ComponentBoundsMin = (FVector3f)SourceComponentBounds.Min;
diff --git a/Engine/Plugins/PCG/Source/PCG/Private/Compute/DataInterfaces/PCGCustomKernelDataInterface.h b/Engine/Plugins/PCG/Source/PCG/Private/Compute/DataInterfaces/PCGCustomKernelDataInterface.h
--- a/Engine/Plugins/PCG/Source/PCG/Private/Compute/DataInterfaces/PCGCustomKernelDataInterface.h
+++ b/Engine/Plugins/PCG/Source/PCG/Private/Compute/DataInterfaces/PCGCustomKernelDataInterface.h
@@ -25,7 +25,6 @@ public:
 	bool CanSupportUnifiedDispatch() const override { return true; }
 	void GetSupportedInputs(TArray<FShaderFunctionDefinition>& OutFunctions) const override;
 	void GetShaderParameters(TCHAR const* UID, FShaderParametersMetadataBuilder& InOutBuilder, FShaderParametersMetadataAllocations& InOutAllocations) const override;
-	void GetShaderHash(FString& InOutKey) const override;
 	void GetHLSL(FString& OutHLSL, FString const& InDataInterfaceName) const override;
 	UComputeDataProvider* CreateDataProvider(TObjectPtr<UObject> InBinding, uint64 InInputMask, uint64 InOutputMask) const override;
 
@@ -48,7 +47,10 @@ public:
 	FComputeDataProviderRenderProxy* GetRenderProxy() override;
 	//~ End UComputeDataProvider Interface
 
-	int ThreadCount = -1;
+	int32 ThreadCount = -1;
+
+	uint32 Seed = 42;
+
 	FBox SourceComponentBounds;
 
 protected:
@@ -58,9 +60,10 @@ protected:
 class FPCGCustomComputeKernelDataProviderProxy : public FComputeDataProviderRenderProxy
 {
 public:
-	FPCGCustomComputeKernelDataProviderProxy(TArray<int32>&& InInvocationThreadCounts, int32 InTotalThreadCount, const FBox& InSourceComponentBounds)
+	FPCGCustomComputeKernelDataProviderProxy(TArray<int32>&& InInvocationThreadCounts, int32 InTotalThreadCount, int32 InSeed, const FBox& InSourceComponentBounds)
 		: InvocationThreadCounts(MoveTemp(InInvocationThreadCounts))
 		, TotalThreadCount(InTotalThreadCount)
+		, Seed(InSeed)
 		, SourceComponentBounds(InSourceComponentBounds)
 	{}
 
@@ -75,5 +78,8 @@ protected:
 
 	TArray<int32> InvocationThreadCounts;
 	int32 TotalThreadCount;
+
+	uint32 Seed = 42;
+
 	FBox SourceComponentBounds;
 };
diff --git a/Engine/Plugins/PCG/Source/PCG/Private/Compute/DataInterfaces/PCGDataCollectionDataInterface.cpp b/Engine/Plugins/PCG/Source/PCG/Private/Compute/DataInterfaces/PCGDataCollectionDataInterface.cpp
--- a/Engine/Plugins/PCG/Source/PCG/Private/Compute/DataInterfaces/PCGDataCollectionDataInterface.cpp
+++ b/Engine/Plugins/PCG/Source/PCG/Private/Compute/DataInterfaces/PCGDataCollectionDataInterface.cpp
@@ -394,12 +394,6 @@ void UPCGDataCollectionDataInterface::GetSupportedInputs(TArray<FShaderFunctionD
 			.AddParam(EShaderFundamentalType::Uint)
 			.AddParam(EShaderFundamentalType::Int);
 
-		OutFunctions.AddDefaulted_GetRef()
-			.SetName(TEXT("SetSeedFromPosition"))
-			.AddParam(EShaderFundamentalType::Uint)
-			.AddParam(EShaderFundamentalType::Uint)
-			.AddParam(EShaderFundamentalType::Float, 3);
-
 		OutFunctions.AddDefaulted_GetRef()
 			.SetName(TEXT("SetSteepness"))
 			.AddParam(EShaderFundamentalType::Uint)
diff --git a/Engine/Plugins/PCG/Source/PCG/Private/Compute/Elements/PCGCustomHLSL.cpp b/Engine/Plugins/PCG/Source/PCG/Private/Compute/Elements/PCGCustomHLSL.cpp
--- a/Engine/Plugins/PCG/Source/PCG/Private/Compute/Elements/PCGCustomHLSL.cpp
+++ b/Engine/Plugins/PCG/Source/PCG/Private/Compute/Elements/PCGCustomHLSL.cpp
@@ -51,6 +51,11 @@ namespace PCGHLSLElement
 	}
 }
 
+UPCGCustomHLSLSettings::UPCGCustomHLSLSettings()
+{
+	bUseSeed = true;
+}
+
 #if WITH_EDITOR
 void UPCGCustomHLSLSettings::PostLoad()
 {
@@ -567,7 +572,7 @@ void UPCGCustomHLSLSettings::UpdateDeclarations()
 		}
 
 		InputDeclarations += TEXT("uint ElementIndex;\n");
-		InputDeclarations += TEXT("uint GetNumThreads();");
+		InputDeclarations += TEXT("uint GetNumThreads();\n");
 	}
 
 	// Add debug category
@@ -575,9 +580,8 @@ void UPCGCustomHLSLSettings::UpdateDeclarations()
 		if (bPrintShaderDebugValues)
 		{
 			InputDeclarations += FString::Format(
-				TEXT("\n\n")
-				TEXT("// Debug\n")
-				TEXT("void WriteDebugValue(uint Index, float Value); // Index in [0, {0}] (set from 'Debug Buffer Size' property)"),
+				TEXT("\n// Debug\n")
+				TEXT("void WriteDebugValue(uint Index, float Value); // Index in [0, {0}] (set from 'Debug Buffer Size' property)\n"),
 				{ DebugBufferSize - 1 });
 		}
 	}
@@ -616,33 +620,38 @@ void UPCGCustomHLSLSettings::UpdateDeclarations()
 			}
 		}
 
-		InputDeclarations += TEXT("\n\n### HELPER FUNCTIONS ###\n\n");
-		InputDeclarations += TEXT("float3 GetComponentBoundsMin(); // World-space\n");
-		InputDeclarations += TEXT("float3 GetComponentBoundsMax();");
+		InputDeclarations += TEXT("\n### HELPER FUNCTIONS ###\n");
+		InputDeclarations += TEXT("\nfloat3 GetComponentBoundsMin(); // World-space\n");
+		InputDeclarations += TEXT("float3 GetComponentBoundsMax();\n");
+		InputDeclarations += TEXT("uint GetSeed();\n");
+
+		InputDeclarations += TEXT("\nfloat FRand(inout uint Seed); // Returns random float between 0 and 1.\n");
+		InputDeclarations += TEXT("uint ComputeSeed(uint A, uint B);\n");
+		InputDeclarations += TEXT("uint ComputeSeed(uint A, uint B, uint C);\n");
+		InputDeclarations += TEXT("uint ComputeSeedFromPosition(float3 Position);\n");
 
-		InputDeclarations += TEXT("\n\n");
-		InputDeclarations += TEXT("float3 CreateGrid2D(int ElementIndex, int NumPoints, float3 Min, float3 Max);\n");
+		InputDeclarations += TEXT("\nfloat3 CreateGrid2D(int ElementIndex, int NumPoints, float3 Min, float3 Max);\n");
 		InputDeclarations += TEXT("float3 CreateGrid2D(int ElementIndex, int NumPoints, int NumRows, float3 Min, float3 Max);\n");
 		InputDeclarations += TEXT("float3 CreateGrid3D(int ElementIndex, int NumPoints, float3 Min, float3 Max);\n");
-		InputDeclarations += TEXT("float3 CreateGrid3D(int ElementIndex, int NumPoints, int NumRows, int NumCols, float3 Min, float3 Max);");
+		InputDeclarations += TEXT("float3 CreateGrid3D(int ElementIndex, int NumPoints, int NumRows, int NumCols, float3 Min, float3 Max);\n");
 
 		if (!DataPins.IsEmpty())
 		{
-			InputDeclarations += TEXT("\n\n### DATA FUNCTIONS ###\n\n");
-			InputDeclarations += TEXT("// Valid pins: ") + FString::Join(DataPins, TEXT(", ")) + TEXT("\n");
+			InputDeclarations += TEXT("\n### DATA FUNCTIONS ###\n");
+			InputDeclarations += TEXT("\n// Valid pins: ") + FString::Join(DataPins, TEXT(", ")) + TEXT("\n");
 			InputDeclarations += TEXT("// Valid types: bool, int, float, float2, float3, float4, Rotator (float3), Quat (float4), Transform (float4x4)\n");
-			InputDeclarations += TEXT("\n");
-			InputDeclarations += TEXT("uint <pin>_GetNumData();\n");
+
+			InputDeclarations += TEXT("\nuint <pin>_GetNumData();\n");
 			InputDeclarations += TEXT("uint <pin>_GetNumElements();\n");
-			InputDeclarations += TEXT("<type> <pin>_Get<type>(uint DataIndex, uint AttributeId, uint ElementIndex);");
+			InputDeclarations += TEXT("<type> <pin>_Get<type>(uint DataIndex, uint AttributeId, uint ElementIndex);\n");
 		}
 
 		if (!PointDataPins.IsEmpty())
 		{
-			InputDeclarations += TEXT("\n\n### POINT DATA FUNCTIONS ###\n\n");
-			InputDeclarations += TEXT("// Valid pins: ") + FString::Join(PointDataPins, TEXT(", ")) + TEXT("\n");
-			InputDeclarations += TEXT("\n");
-			InputDeclarations += TEXT("uint <pin>_GetNumPoints(uint DataIndex);\n");
+			InputDeclarations += TEXT("\n### POINT DATA FUNCTIONS ###\n");
+			InputDeclarations += TEXT("\n// Valid pins: ") + FString::Join(PointDataPins, TEXT(", ")) + TEXT("\n");
+
+			InputDeclarations += TEXT("\nuint <pin>_GetNumPoints(uint DataIndex);\n");
 			InputDeclarations += TEXT("float3 <pin>_GetPosition(uint DataIndex, uint ElementIndex);\n");
 			InputDeclarations += TEXT("float4 <pin>_GetRotation(uint DataIndex, uint ElementIndex);\n");
 			InputDeclarations += TEXT("float3 <pin>_GetScale(uint DataIndex, uint ElementIndex);\n");
@@ -652,34 +661,34 @@ void UPCGCustomHLSLSettings::UpdateDeclarations()
 			InputDeclarations += TEXT("float <pin>_GetDensity(uint DataIndex, uint ElementIndex);\n");
 			InputDeclarations += TEXT("int <pin>_GetSeed(uint DataIndex, uint ElementIndex);\n");
 			InputDeclarations += TEXT("float <pin>_GetSteepness(uint DataIndex, uint ElementIndex);\n");
-			InputDeclarations += TEXT("float4x4 <pin>_GetPointTransform(uint DataIndex, uint ElementIndex);");
+			InputDeclarations += TEXT("float4x4 <pin>_GetPointTransform(uint DataIndex, uint ElementIndex);\n");
 		}
 
 		if (!LandscapeDataPins.IsEmpty())
 		{
-			InputDeclarations += TEXT("\n\n### LANDSCAPE DATA FUNCTIONS ###\n\n");
-			InputDeclarations += TEXT("// Valid pins: ") + FString::Join(LandscapeDataPins, TEXT(", ")) + TEXT("\n");
-			InputDeclarations += TEXT("\n");
-			InputDeclarations += TEXT("float <pin>_GetHeight(float3 WorldPos);\n");
-			InputDeclarations += TEXT("float3 <pin>_GetNormal(float3 WorldPos);");
+			InputDeclarations += TEXT("\n### LANDSCAPE DATA FUNCTIONS ###\n");
+			InputDeclarations += TEXT("\n// Valid pins: ") + FString::Join(LandscapeDataPins, TEXT(", ")) + TEXT("\n");
+
+			InputDeclarations += TEXT("\nfloat <pin>_GetHeight(float3 WorldPos);\n");
+			InputDeclarations += TEXT("float3 <pin>_GetNormal(float3 WorldPos);\n");
 		}
 
 		if (!TextureDataPins.IsEmpty())
 		{
-			InputDeclarations += TEXT("\n\n### TEXTURE DATA FUNCTIONS ###\n\n");
-			InputDeclarations += TEXT("// Valid pins: ") + FString::Join(TextureDataPins, TEXT(", ")) + TEXT("\n");
-			InputDeclarations += TEXT("\n");
-			InputDeclarations += TEXT("float2 <pin>_GetTexCoords(float2 WorldPos, float2 Min, float2 Max);\n");
-			InputDeclarations += TEXT("float4 <pin>_Sample(float2 TexCoords);");
+			InputDeclarations += TEXT("\n### TEXTURE DATA FUNCTIONS ###\n");
+			InputDeclarations += TEXT("\n// Valid pins: ") + FString::Join(TextureDataPins, TEXT(", ")) + TEXT("\n");
+
+			InputDeclarations += TEXT("\nfloat2 <pin>_GetTexCoords(float2 WorldPos, float2 Min, float2 Max);\n");
+			InputDeclarations += TEXT("float4 <pin>_Sample(float2 TexCoords);\n");
 		}
 
 		if (!RawBufferDataPins.IsEmpty())
 		{
-			InputDeclarations += TEXT("\n\n### BYTE ADDRESS BUFFER DATA FUNCTIONS ###\n\n");
-			InputDeclarations += TEXT("// Valid pins: ") + FString::Join(RawBufferDataPins, TEXT(", ")) + TEXT("\n");
-			InputDeclarations += TEXT("\n");
-			InputDeclarations += TEXT("uint <pin>_ReadNumValues();\n");
-			InputDeclarations += TEXT("uint <pin>_ReadValue(uint Index);");
+			InputDeclarations += TEXT("\n### BYTE ADDRESS BUFFER DATA FUNCTIONS ###\n");
+			InputDeclarations += TEXT("\n// Valid pins: ") + FString::Join(RawBufferDataPins, TEXT(", ")) + TEXT("\n");
+
+			InputDeclarations += TEXT("\nuint <pin>_ReadNumValues();\n");
+			InputDeclarations += TEXT("uint <pin>_ReadValue(uint Index);\n");
 		}
 	}
 
@@ -707,19 +716,19 @@ void UPCGCustomHLSLSettings::UpdateDeclarations()
 
 		if (!DataPins.IsEmpty())
 		{
-			OutputDeclarations += TEXT("### DATA FUNCTIONS ###\n\n");
-			OutputDeclarations += TEXT("// Valid pins: ") + FString::Join(DataPins, TEXT(", ")) + TEXT("\n");
+			OutputDeclarations += TEXT("### DATA FUNCTIONS ###\n");
+			OutputDeclarations += TEXT("\n// Valid pins: ") + FString::Join(DataPins, TEXT(", ")) + TEXT("\n");
 			OutputDeclarations += TEXT("// Valid types: bool, int, float, float2, float3, float4, Rotator (float3), Quat (float4), Transform (float4x4)\n");
-			OutputDeclarations += TEXT("\n");
-			OutputDeclarations += TEXT("void <pin>_Set<type>(uint DataIndex, uint AttributeId, uint ElementIndex, <type> Value);");
+
+			OutputDeclarations += TEXT("\nvoid <pin>_Set<type>(uint DataIndex, uint AttributeId, uint ElementIndex, <type> Value);\n");
 		}
 
 		if (!PointDataPins.IsEmpty())
 		{
-			OutputDeclarations += TEXT("\n\n### POINT DATA FUNCTIONS ###\n\n");
-			OutputDeclarations += TEXT("// Valid pins: ") + FString::Join(PointDataPins, TEXT(", ")) + TEXT("\n");
-			OutputDeclarations += TEXT("\n");
-			OutputDeclarations += TEXT("void <pin>_SetPosition(uint DataIndex, uint ElementIndex, float3 Position);\n");
+			OutputDeclarations += TEXT("\n### POINT DATA FUNCTIONS ###\n");
+			OutputDeclarations += TEXT("\n// Valid pins: ") + FString::Join(PointDataPins, TEXT(", ")) + TEXT("\n");
+
+			OutputDeclarations += TEXT("\nvoid <pin>_SetPosition(uint DataIndex, uint ElementIndex, float3 Position);\n");
 			OutputDeclarations += TEXT("void <pin>_SetRotation(uint DataIndex, uint ElementIndex, float4 Rotation);\n");
 			OutputDeclarations += TEXT("void <pin>_SetScale(uint DataIndex, uint ElementIndex, float3 Scale);\n");
 			OutputDeclarations += TEXT("void <pin>_SetBoundsMin(uint DataIndex, uint ElementIndex, float3 BoundsMin);\n");
@@ -729,17 +738,23 @@ void UPCGCustomHLSLSettings::UpdateDeclarations()
 			OutputDeclarations += TEXT("void <pin>_SetSeed(uint DataIndex, uint ElementIndex, int Seed);\n");
 			OutputDeclarations += TEXT("void <pin>_SetSeedFromPosition(uint DataIndex, uint ElementIndex, float3 Position);\n");
 			OutputDeclarations += TEXT("void <pin>_SetSteepness(uint DataIndex, uint ElementIndex, float Steepness);\n");
-			OutputDeclarations += TEXT("void <pin>_SetPointTransform(uint DataIndex, uint ElementIndex, float4x4 Transform);");
+			OutputDeclarations += TEXT("void <pin>_SetPointTransform(uint DataIndex, uint ElementIndex, float4x4 Transform);\n");
 		}
 
 		if (!RawBufferDataPins.IsEmpty())
 		{
-			OutputDeclarations += TEXT("\n\n### BYTE ADDRESS BUFFER DATA FUNCTIONS ###\n\n");
+			OutputDeclarations += TEXT("\n### BYTE ADDRESS BUFFER DATA FUNCTIONS ###\n\n");
 			OutputDeclarations += TEXT("// Valid pins: ") + FString::Join(RawBufferDataPins, TEXT(", ")) + TEXT("\n");
-			OutputDeclarations += TEXT("\n");
-			OutputDeclarations += TEXT("uint <pin>_WriteValue(uint Index, uint Value);");
+
+			OutputDeclarations += TEXT("\nuint <pin>_WriteValue(uint Index, uint Value);\n");
 		}
 	}
+
+	if (!OutputDeclarations.IsEmpty())
+	{
+		// Remove final newline as a small UI improvement.
+		OutputDeclarations = OutputDeclarations.LeftChop(1);
+	}
 }
 
 void UPCGCustomHLSLSettings::UpdatePinSettings()
diff --git a/Engine/Plugins/PCG/Source/PCG/Public/Compute/Elements/PCGCustomHLSL.h b/Engine/Plugins/PCG/Source/PCG/Public/Compute/Elements/PCGCustomHLSL.h
--- a/Engine/Plugins/PCG/Source/PCG/Public/Compute/Elements/PCGCustomHLSL.h
+++ b/Engine/Plugins/PCG/Source/PCG/Public/Compute/Elements/PCGCustomHLSL.h
@@ -37,6 +37,8 @@ class UPCGCustomHLSLSettings : public UPCGSettings
 	GENERATED_BODY()
 
 public:
+	UPCGCustomHLSLSettings();
+
 	//~Begin UObject interface
 #if WITH_EDITOR
 	virtual void PostLoad() override;
