[PCG] Missed feedback note - add extra doc notes to the new seed / random stream shader functionality

#rnx

diff --git a/Engine/Plugins/PCG/Shaders/Private/PCGShaderUtils.ush b/Engine/Plugins/PCG/Shaders/Private/PCGShaderUtils.ush
--- a/Engine/Plugins/PCG/Shaders/Private/PCGShaderUtils.ush
+++ b/Engine/Plugins/PCG/Shaders/Private/PCGShaderUtils.ush
@@ -3,6 +3,7 @@
 #include "/Engine/Private/Common.ush"
 #include "/Engine/Public/Platform.ush"
 
+// These follow the corresponding C++ functions in PCGHelpers.
 uint ComputeSeed(uint A, uint B)
 {
 	return ((A * 196314165U) + 907633515U) ^ ((B * 73148459U) + 453816763U);
@@ -20,6 +21,7 @@ uint ComputeSeedFromPosition(float3 InPosition)
 	return (uint)Seed;
 }
 
+// This follows the implementation in FRandomStream.
 float FRand(inout uint InOutSeed)
 {
 	InOutSeed = (InOutSeed * 196314165U) + 907633515U;
