[PCG] Graph Editor Fix Undo/Redo ensure

- Avoid calls to ReconstructNode while transacting (when undo/redo in progress)

#rnx

diff --git a/Engine/Plugins/PCG/Source/PCGEditor/Private/PCGEditorGraphNodeBase.cpp b/Engine/Plugins/PCG/Source/PCGEditor/Private/PCGEditorGraphNodeBase.cpp
--- a/Engine/Plugins/PCG/Source/PCGEditor/Private/PCGEditorGraphNodeBase.cpp
+++ b/Engine/Plugins/PCG/Source/PCGEditor/Private/PCGEditorGraphNodeBase.cpp
@@ -21,6 +21,7 @@
 #include "PCGEditorSettings.h"
 #include "PCGEditorStyle.h"
 
+#include "CoreGlobals.h"
 #include "GraphEditorActions.h"
 #include "ScopedTransaction.h"
 #include "ToolMenu.h"
@@ -796,6 +797,13 @@ void UPCGEditorGraphNodeBase::ReconstructNode()
 		bDeferredReconstruct = true;
 		return;
 	}
+
+	// While in an Undo/Redo a call to ReconstructNode should not be needed as the transaction object records
+	// should be enough to serialize the nodes back into their proper state 
+	if (GIsTransacting)
+	{
+		return;
+	}
 	
 	Modify();
 
