[PCG] Fix some Undo/Redo issues with PCGWorldActor
- Spawn PCGWorldActor as a non-transactionable actor and once spawned make it transactionable again so that property changes can be undone but not the creation
- Avoid creating a transaction inside the UPCGComponent::Generate() call, only call site is the PCGVolumeFactory and this was creating 2 Transactions when dragging a PCG Volume from the Content Browser

#rnx

diff --git a/Engine/Plugins/PCG/Source/PCG/Private/PCGComponent.cpp b/Engine/Plugins/PCG/Source/PCG/Private/PCGComponent.cpp
--- a/Engine/Plugins/PCG/Source/PCG/Private/PCGComponent.cpp
+++ b/Engine/Plugins/PCG/Source/PCG/Private/PCGComponent.cpp
@@ -403,10 +403,6 @@ void UPCGComponent::Generate()
 		return;
 	}
 
-#if WITH_EDITOR
-	FScopedTransaction Transaction(LOCTEXT("PCGGenerate", "Execute generation on PCG component"));
-#endif
-
 	GenerateLocal(/*bForce=*/false);
 }
 
diff --git a/Engine/Plugins/PCG/Source/PCG/Private/PCGWorldActor.cpp b/Engine/Plugins/PCG/Source/PCG/Private/PCGWorldActor.cpp
--- a/Engine/Plugins/PCG/Source/PCG/Private/PCGWorldActor.cpp
+++ b/Engine/Plugins/PCG/Source/PCG/Private/PCGWorldActor.cpp
@@ -167,10 +167,14 @@ APCGWorldActor* APCGWorldActor::CreatePCGWorldActor(UWorld* InWorld)
 	{
 		FActorSpawnParameters SpawnParams;
 		SpawnParams.OverrideLevel = InWorld->PersistentLevel;
+
+		// We don't want the PCGWorldActor creation to be part of a transaction, once it is created we add the flag to the actor
+		SpawnParams.ObjectFlags &= ~RF_Transactional;
 		PCGActor = InWorld->SpawnActor<APCGWorldActor>(SpawnParams);
 
 		if (PCGActor)
 		{
+			PCGActor->SetFlags(RF_Transactional);
 			PCGActor->RegisterToSubsystem();
 		}
 	}
diff --git a/Engine/Plugins/PCG/Source/PCG/Public/PCGComponent.h b/Engine/Plugins/PCG/Source/PCG/Public/PCGComponent.h
--- a/Engine/Plugins/PCG/Source/PCG/Public/PCGComponent.h
+++ b/Engine/Plugins/PCG/Source/PCG/Public/PCGComponent.h
@@ -153,7 +153,6 @@ public:
 public:
 	void ForEachManagedResource(TFunctionRef<void(UPCGManagedResource*)> InFunction);
 
-	/** Transactionable methods to be called from details UI */
 	void Generate();
 	void Cleanup();
 
