[PCG] Tentative fix for compilation issue

#rnx

diff --git a/Engine/Plugins/PCG/Source/PCG/Public/Tests/PCGTestsCommon.h b/Engine/Plugins/PCG/Source/PCG/Public/Tests/PCGTestsCommon.h
--- a/Engine/Plugins/PCG/Source/PCG/Public/Tests/PCGTestsCommon.h
+++ b/Engine/Plugins/PCG/Source/PCG/Public/Tests/PCGTestsCommon.h
@@ -93,9 +93,12 @@ namespace PCGTestsCommon
 		return TypedSettings;
 	}
 
+	template <class T, class... Ts>
+	struct TIsAny : std::disjunction<std::is_same<T, Ts>...> {};
+
 	// Numerical
 	template <typename T, typename std::enable_if_t<std::is_arithmetic_v<T>, bool> = true>
-	T GenerateRandomValue(FRandomStream& RandomStream) requires std::is_arithmetic_v<T>
+	T GenerateRandomValue(FRandomStream& RandomStream)
 	{
 		if constexpr (std::is_same_v<bool, T>)
 		{
@@ -108,21 +111,21 @@ namespace PCGTestsCommon
 	}
 
 	 // Vectors
-	template <typename T, typename std::enable_if_t<PCG::Private::IsOfTypes<T, FVector2D, FVector, FVector4>(), bool> = true>
+	template <typename T, typename std::enable_if_t<TIsAny<T, FVector2D, FVector, FVector4>::value, bool> = true>
 	T GenerateRandomValue(FRandomStream& RandomStream)
 	{
 		return T(RandomStream.VRand());
 	}
 
 	 // Rotators/Quat
-	template <typename T, typename std::enable_if_t<PCG::Private::IsOfTypes<T, FRotator, FQuat>(), bool> = true>
+	template <typename T, typename std::enable_if_t<TIsAny<T, FRotator, FQuat>::value, bool> = true>
 	T GenerateRandomValue(FRandomStream& RandomStream)
 	{
 		return T::MakeFromEuler(RandomStream.VRand() * 360.0);
 	}
 
 	 // String/FName
-	template <typename T, typename std::enable_if_t<PCG::Private::IsOfTypes<T, FName, FString>(), bool> = true>
+	template <typename T, typename std::enable_if_t<TIsAny<T, FName, FString>::value, bool> = true>
 	T GenerateRandomValue(FRandomStream& RandomStream)
 	{
 		static constexpr const TCHAR* Dictionary[] = { TEXT("Foo"), TEXT("Bar"), TEXT("PCG"), TEXT("YOLO"), TEXT("Bla") };
@@ -144,7 +147,7 @@ namespace PCGTestsCommon
 	}
 
 	// SoftObjectPath
-	template <typename T, typename std::enable_if_t<std::is_same_v<T, FSoftObjectPath>, bool> = true>
+	template <typename T, typename std::enable_if_t<TIsAny<T, FSoftObjectPath>::value, bool> = true>
 	T GenerateRandomValue(FRandomStream& RandomStream)
 	{
 		static const FSoftObjectPath Dictionary[] = 
@@ -158,7 +161,7 @@ namespace PCGTestsCommon
 	}
 
 	// FSoftClassPath
-	template <typename T, typename std::enable_if_t<std::is_same_v<T, FSoftClassPath>, bool> = true>
+	template <typename T, typename std::enable_if_t<TIsAny<T, FSoftClassPath>::value, bool> = true>
 	T GenerateRandomValue(FRandomStream& RandomStream)
 	{
 		static const FSoftClassPath Dictionary[] =
