[PCG] Some types of compile issues do not create explicit error diagnostic messages, and soft lock the graph execution.

These messages typically have 'failed' in the message - treat any such message as an error.

#rnx

diff --git a/Engine/Plugins/PCG/Source/PCG/Private/Compute/Elements/PCGComputeGraphElement.cpp b/Engine/Plugins/PCG/Source/PCG/Private/Compute/Elements/PCGComputeGraphElement.cpp
--- a/Engine/Plugins/PCG/Source/PCG/Private/Compute/Elements/PCGComputeGraphElement.cpp
+++ b/Engine/Plugins/PCG/Source/PCG/Private/Compute/Elements/PCGComputeGraphElement.cpp
@@ -174,7 +174,8 @@ bool FPCGComputeGraphElement::ExecuteInternal(FPCGContext* InContext) const
 		{
 			for (const FComputeKernelCompileMessage& Message : NodeAndCompileMessages.Get<1>())
 			{
-				if (Message.Type == FComputeKernelCompileMessage::EMessageType::Error)
+				// Currently failure messages can come through with Log severity, pattern match for now.
+				if (Message.Type == FComputeKernelCompileMessage::EMessageType::Error || Message.Text.Contains(TEXT("failed"), ESearchCase::IgnoreCase))
 				{
 					return true;
 				}
@@ -286,6 +287,11 @@ void FPCGComputeGraphElement::LogCompilationMessages(FPCGComputeGraphContext* In
 				{
 					Verbosity = ELogVerbosity::Error;
 				}
+				else if (Message.Text.Contains(TEXT("failed"), ESearchCase::IgnoreCase))
+				{
+					// Currently failure messages can come through with Log verbosity.
+					Verbosity = ELogVerbosity::Error;
+				}
 
 				if (Verbosity < ELogVerbosity::Log)
 				{
