[PCG] Segment Slicer - Replace Extremity bool Attribute with two distinct attributes: IsFirst/IsFinal

#rnx

diff --git a/Engine/Plugins/PCG/Source/PCG/Private/Elements/Grammar/PCGSegmentSlicer.cpp b/Engine/Plugins/PCG/Source/PCG/Private/Elements/Grammar/PCGSegmentSlicer.cpp
--- a/Engine/Plugins/PCG/Source/PCG/Private/Elements/Grammar/PCGSegmentSlicer.cpp
+++ b/Engine/Plugins/PCG/Source/PCG/Private/Elements/Grammar/PCGSegmentSlicer.cpp
@@ -21,7 +21,8 @@ public:
 		FPCGMetadataAttribute<FName>* SymbolAttribute = nullptr;
 		FPCGMetadataAttribute<FVector4>* DebugColorAttribute = nullptr;
 		FPCGMetadataAttribute<int32>* ModuleIndexAttribute = nullptr;
-		FPCGMetadataAttribute<bool>* ExtremityAttribute = nullptr;
+		FPCGMetadataAttribute<bool>* IsFirstPointAttribute = nullptr;
+		FPCGMetadataAttribute<bool>* IsFinalPointAttribute = nullptr;
 		FPCGMetadataAttribute<int32>* ExtremityNeighborIndexAttribute = nullptr;
 
 		FPCGSlicingBaseElement::FPCGModulesInfoMap ModulesInfo;
@@ -107,7 +108,7 @@ public:
 					const FPCGSlicingModule& SlicingModule = InOutParameters.ModulesInfo[Symbol];
 
 					const bool bIsFirstModule = (ModuleInstanceIndex == 0) && (j == 0) && (SymbolIndex == 0);
-					const bool bIsLastModule = (ModuleInstanceIndex == ModulesInstances.Num() - 1) && (j == ModuleInstance.NumRepeat - 1) && (SymbolIndex == ModuleInstance.Module->Symbols.Num() - 1);
+					const bool bIsFinalModule = (ModuleInstanceIndex == ModulesInstances.Num() - 1) && (j == ModuleInstance.NumRepeat - 1) && (SymbolIndex == ModuleInstance.Module->Symbols.Num() - 1);
 					const double HalfDisplacement = SlicingModule.Size * 0.5;
 					const double HalfScaledDisplacement = Scale.Dot(InOutParameters.SlicingDirection) * HalfDisplacement;
 
@@ -135,12 +136,14 @@ public:
 						InOutParameters.ModuleIndexAttribute->SetValue(OutPoint.MetadataEntry, ModuleIndex++);
 					}
 
-					if (bIsFirstModule || bIsLastModule)
+					if (bIsFirstModule && InOutParameters.IsFirstPointAttribute)
 					{
-						if (InOutParameters.ExtremityAttribute)
-						{
-							InOutParameters.ExtremityAttribute->SetValue(OutPoint.MetadataEntry, true);
-						}
+						InOutParameters.IsFirstPointAttribute->SetValue(OutPoint.MetadataEntry, true);
+					}
+
+					if (bIsFinalModule && InOutParameters.IsFinalPointAttribute)
+					{
+						InOutParameters.IsFinalPointAttribute->SetValue(OutPoint.MetadataEntry, true);
 					}
 				}
 			}
@@ -301,7 +304,8 @@ bool FPCGSegmentSlicerElement::ExecuteInternal(FPCGContext* InContext) const
 		if (!CreateAndValidateAttribute(Settings->SymbolAttributeName, FName(NAME_None), true, Parameters.SymbolAttribute)
 			|| !CreateAndValidateAttribute(Settings->DebugColorAttributeName, FVector4::Zero(), Settings->bOutputDebugColorAttribute, Parameters.DebugColorAttribute)
 			|| !CreateAndValidateAttribute(Settings->ModuleIndexAttributeName, -1, Settings->bOutputModuleIndexAttribute, Parameters.ModuleIndexAttribute)
-			|| !CreateAndValidateAttribute(Settings->ExtremityAttributeName, false, Settings->bOutputExtremityAttribute, Parameters.ExtremityAttribute)
+			|| !CreateAndValidateAttribute(Settings->IsFirstAttributeName, false, Settings->bOutputExtremityAttributes, Parameters.IsFirstPointAttribute)
+			|| !CreateAndValidateAttribute(Settings->IsFinalAttributeName, false, Settings->bOutputExtremityAttributes, Parameters.IsFinalPointAttribute)
 			|| !CreateAndValidateAttribute(Settings->ExtremityNeighborIndexAttributeName, -1, Settings->bOutputExtremityNeighborIndexAttribute, Parameters.ExtremityNeighborIndexAttribute))
 		{
 			continue;
diff --git a/Engine/Plugins/PCG/Source/PCG/Public/Elements/Grammar/PCGSegmentSlicer.h b/Engine/Plugins/PCG/Source/PCG/Public/Elements/Grammar/PCGSegmentSlicer.h
--- a/Engine/Plugins/PCG/Source/PCG/Public/Elements/Grammar/PCGSegmentSlicer.h
+++ b/Engine/Plugins/PCG/Source/PCG/Public/Elements/Grammar/PCGSegmentSlicer.h
@@ -57,12 +57,17 @@ public:
 	UPROPERTY(BlueprintReadWrite, EditAnywhere, Category = "Settings|Extra Output Attributes", meta = (PCG_Overridable, EditCondition = "bOutputModuleIndexAttribute"))
 	FName ModuleIndexAttributeName = TEXT("ModuleIndex");
 
-	UPROPERTY(BlueprintReadWrite, EditAnywhere, Category = "Settings|Extra Output Attributes", meta = (PCG_Overridable, InlineEditConditionToggle))
-	bool bOutputExtremityAttribute = false;
+	/** Output attributes labeling the first and final points after slicing. */
+	UPROPERTY(BlueprintReadWrite, EditAnywhere, Category = "Settings|Extra Output Attributes", meta = (PCG_Overridable))
+	bool bOutputExtremityAttributes = false;
+
+	/** Name of the attribute labeling the first output point from the first module. */
+	UPROPERTY(BlueprintReadWrite, EditAnywhere, Category = "Settings|Extra Output Attributes", meta = (PCG_Overridable, EditCondition = "bOutputExtremityAttributes", EditConditionHides))
+	FName IsFirstAttributeName = TEXT("IsFirst");
 
-	/** Name of the Extremity output attribute name. */
-	UPROPERTY(BlueprintReadWrite, EditAnywhere, Category = "Settings|Extra Output Attributes", meta = (PCG_Overridable, EditCondition = "bOutputExtremityAttribute"))
-	FName ExtremityAttributeName = TEXT("Extremity");
+	/** Name of the attribute labeling the final output point from the final module. */
+	UPROPERTY(BlueprintReadWrite, EditAnywhere, Category = "Settings|Extra Output Attributes", meta = (PCG_Overridable, EditCondition = "bOutputExtremityAttributes", EditConditionHides))
+	FName IsFinalAttributeName = TEXT("IsFinal");
 
 	UPROPERTY(BlueprintReadWrite, EditAnywhere, Category = "Settings|Extra Output Attributes", meta = (PCG_Overridable, InlineEditConditionToggle))
 	bool bOutputExtremityNeighborIndexAttribute = false;
