[PCG] Don't produce compute graphs if any node fails validation, to fix weird behaviour from partial graphs.

Also soften a hard check & crash to an ensure which comes up often during development.

#rnx

diff --git a/Engine/Plugins/PCG/Source/PCG/Private/Compute/Elements/PCGCustomHLSL.cpp b/Engine/Plugins/PCG/Source/PCG/Private/Compute/Elements/PCGCustomHLSL.cpp
--- a/Engine/Plugins/PCG/Source/PCG/Private/Compute/Elements/PCGCustomHLSL.cpp
+++ b/Engine/Plugins/PCG/Source/PCG/Private/Compute/Elements/PCGCustomHLSL.cpp
@@ -53,6 +53,7 @@ namespace PCGHLSLElement
 
 UPCGCustomHLSLSettings::UPCGCustomHLSLSettings()
 {
+	bExecuteOnGPU = true;
 	bUseSeed = true;
 }
 
@@ -74,11 +75,6 @@ void UPCGCustomHLSLSettings::PostInitProperties()
 }
 #endif
 
-bool UPCGCustomHLSLSettings::ShouldExecuteOnGPU() const
-{
-	return IsKernelValid();
-}
-
 #if WITH_EDITOR
 void UPCGCustomHLSLSettings::PostEditChangeProperty(FPropertyChangedEvent& PropertyChangedEvent)
 {
@@ -1188,9 +1184,7 @@ bool FPCGCustomHLSLElement::ExecuteInternal(FPCGContext* Context) const
 	check(Settings);
 
 	// Calling IsKernelValid with the Context will emit warnings/errors on the node.
-	const bool bIsValid = Settings->IsKernelValid(Context);
-
-	ensureMsgf(!bIsValid, TEXT("Custom HLSL element should only execute on CPU if the kernel is invalid."));
+	Settings->IsKernelValid(Context);
 
 	return true;
 }
diff --git a/Engine/Plugins/PCG/Source/PCG/Private/Compute/PCGDataForGPU.cpp b/Engine/Plugins/PCG/Source/PCG/Private/Compute/PCGDataForGPU.cpp
--- a/Engine/Plugins/PCG/Source/PCG/Private/Compute/PCGDataForGPU.cpp
+++ b/Engine/Plugins/PCG/Source/PCG/Private/Compute/PCGDataForGPU.cpp
@@ -769,7 +769,10 @@ void FPCGDataCollectionDesc::UnpackDataCollection(const TArray<uint8>& InPackedD
 	const uint32 NumPackedFloats = InPackedData.Num() / 4;
 	const uint32 NumData = DataAsUint[0];
 
-	check(NumData == DataDescs.Num());
+	if (!ensureAlwaysMsgf(NumData == DataDescs.Num(), TEXT("Mismatch in expected data from static analysis and actual received data.")))
+	{
+		return;
+	}
 
 	TArray<FPCGTaggedData>& OutData = OutDataCollection.TaggedData;
 
diff --git a/Engine/Plugins/PCG/Source/PCG/Private/Graph/PCGGraphCompilerGPU.cpp b/Engine/Plugins/PCG/Source/PCG/Private/Graph/PCGGraphCompilerGPU.cpp
--- a/Engine/Plugins/PCG/Source/PCG/Private/Graph/PCGGraphCompilerGPU.cpp
+++ b/Engine/Plugins/PCG/Source/PCG/Private/Graph/PCGGraphCompilerGPU.cpp
@@ -235,7 +235,21 @@ void FPCGGraphCompilerGPU::CollectGPUNodeSubsets(
 
 		if (!GPUSubsetTaskIds.IsEmpty())
 		{
-			OutNodeSubsetsToConvertToCFGraph.Add(MoveTemp(GPUSubsetTaskIds));
+			bool bAllNodesValid = true;
+			for (FPCGTaskId& TaskId : GPUSubsetTaskIds)
+			{
+				const UPCGCustomHLSLSettings* Settings = Cast<UPCGCustomHLSLSettings>(InCompiledTasks[TaskId].Node ? InCompiledTasks[TaskId].Node->GetSettings() : nullptr);
+				if (Settings && !Settings->IsKernelValid())
+				{
+					bAllNodesValid = false;
+					break;
+				}
+			}
+
+			if (bAllNodesValid)
+			{
+				OutNodeSubsetsToConvertToCFGraph.Add(MoveTemp(GPUSubsetTaskIds));
+			}
 		}
 	}
 }
@@ -995,6 +1009,8 @@ void FPCGGraphCompilerGPU::BuildGPUGraphTask(
 
 void FPCGGraphCompilerGPU::CreateGPUNodes(UPCGGraph* InGraph, TArray<FPCGGraphTask>& InOutCompiledTasks)
 {
+	TRACE_CPUPROFILER_EVENT_SCOPE(FPCGGraphCompilerGPU::CreateGPUNodes);
+
 	TSet<FPCGTaskId> GPUCompatibleTaskIds;
 	GPUCompatibleTaskIds.Reserve(InOutCompiledTasks.Num());
 	for (FPCGTaskId TaskId = 0; TaskId < InOutCompiledTasks.Num(); ++TaskId)
@@ -1082,8 +1098,16 @@ void FPCGGraphCompilerGPU::CreateGPUNodes(UPCGGraph* InGraph, TArray<FPCGGraphTa
 	}
 
 	// Now cull all the GPU compatible nodes. The compute graph task are already wired in so we're fine to just delete.
-	FPCGGraphCompiler::CullTasks(InOutCompiledTasks, /*bAddPassthroughWires=*/false, [&GPUCompatibleTaskIds](const FPCGGraphTask& InTask)
+	FPCGGraphCompiler::CullTasks(InOutCompiledTasks, /*bAddPassthroughWires=*/false, [&NodeSubsetsToConvertToCFGraph](const FPCGGraphTask& InTask)
 	{
-		return GPUCompatibleTaskIds.Contains(InTask.NodeId);
+		for (TSet<FPCGTaskId>& NodeSubsetToConvertToCFGraph : NodeSubsetsToConvertToCFGraph)
+		{
+			if (NodeSubsetToConvertToCFGraph.Contains(InTask.NodeId))
+			{
+				return true;
+			}
+		}
+
+		return false;
 	});
 }
diff --git a/Engine/Plugins/PCG/Source/PCG/Public/Compute/Elements/PCGCustomHLSL.h b/Engine/Plugins/PCG/Source/PCG/Public/Compute/Elements/PCGCustomHLSL.h
--- a/Engine/Plugins/PCG/Source/PCG/Public/Compute/Elements/PCGCustomHLSL.h
+++ b/Engine/Plugins/PCG/Source/PCG/Public/Compute/Elements/PCGCustomHLSL.h
@@ -30,7 +30,7 @@ enum class EPCGDispatchThreadCount : uint8
 	// TODO: Could support inlined expression evaluation for dispatch thread count.
 };
 
-/** Prouduces a HLSL compute shader which will be executed on the GPU. */
+/** Produces a HLSL compute shader which will be executed on the GPU. */
 UCLASS(MinimalAPI, BlueprintType, ClassGroup = (Procedural))
 class UPCGCustomHLSLSettings : public UPCGSettings
 {
@@ -49,13 +49,12 @@ public:
 	//~Begin UPCGSettings interface
 	virtual TArray<FPCGPinProperties> InputPinProperties() const override { return InputPins; }
 	virtual TArray<FPCGPinProperties> OutputPinProperties() const override { return OutputPins; }
-	virtual bool ShouldExecuteOnGPU() const override;
 	virtual bool IsInputPinRequiredByExecution(const UPCGPin* InPin) const { return true; }
 #if WITH_EDITOR
 	virtual bool DisplayExecuteOnGPUSetting() const override { return false; }
 	virtual FName GetDefaultNodeName() const override { return FName(TEXT("CustomHLSL")); }
 	virtual FText GetDefaultNodeTitle() const override { return NSLOCTEXT("PCGCustomHLSLElement", "NodeTitle", "Custom HLSL"); }
-	virtual FText GetNodeTooltipText() const override { return NSLOCTEXT("PCGCustomHLSLElement", "NodeTooltip", "Prouduces a HLSL compute shader which will be executed on the GPU."); }
+	virtual FText GetNodeTooltipText() const override { return NSLOCTEXT("PCGCustomHLSLElement", "NodeTooltip", "Produces a HLSL compute shader which will be executed on the GPU."); }
 	virtual EPCGSettingsType GetType() const override { return EPCGSettingsType::GPU; }
 #endif
 
