TObjectPtr fix from newly introduced field in engine plugin.

diff --git a/Engine/Plugins/PCG/Source/PCG/Private/Compute/DataInterfaces/PCGDataCollectionDataInterface.h b/Engine/Plugins/PCG/Source/PCG/Private/Compute/DataInterfaces/PCGDataCollectionDataInterface.h
--- a/Engine/Plugins/PCG/Source/PCG/Private/Compute/DataInterfaces/PCGDataCollectionDataInterface.h
+++ b/Engine/Plugins/PCG/Source/PCG/Private/Compute/DataInterfaces/PCGDataCollectionDataInterface.h
@@ -86,7 +86,7 @@ public:
 public:
 	/** Settings of node that produces this data, normally the upstream node. */
 	UPROPERTY()
-	const UPCGSettings* ProducerSettings;
+	TObjectPtr<const UPCGSettings> ProducerSettings;
 
 	TWeakObjectPtr<UPCGDataBinding> Binding;
 	FPCGDataCollectionDesc PinDesc;
