[PCG] Add tests flags on Texture sampling tests to make sure rendering is enabled

#rnx

diff --git a/Engine/Plugins/PCG/Source/PCG/Private/Tests/Elements/PCGSampleTextureTest.cpp b/Engine/Plugins/PCG/Source/PCG/Private/Tests/Elements/PCGSampleTextureTest.cpp
--- a/Engine/Plugins/PCG/Source/PCG/Private/Tests/Elements/PCGSampleTextureTest.cpp
+++ b/Engine/Plugins/PCG/Source/PCG/Private/Tests/Elements/PCGSampleTextureTest.cpp
@@ -145,8 +145,8 @@ protected:
 	}
 };
 
-
-IMPLEMENT_CUSTOM_SIMPLE_AUTOMATION_TEST(FPCGSampleTextureTest_LocalCorner, PCGSampleTextureTestBase, "Plugins.PCG.SampleTexture.LocalCorner", PCGTestsCommon::TestFlags)
+// EAutomationTestFlags::NonNullRHI is important here since texture need some kind of GPU support, and without any rendering enabled, it will crash otherwise.
+IMPLEMENT_CUSTOM_SIMPLE_AUTOMATION_TEST(FPCGSampleTextureTest_LocalCorner, PCGSampleTextureTestBase, "Plugins.PCG.SampleTexture.LocalCorner", PCGTestsCommon::TestFlags | EAutomationTestFlags::NonNullRHI)
 
 bool FPCGSampleTextureTest_LocalCorner::RunTest(const FString& Parameters)
 {
@@ -159,7 +159,7 @@ bool FPCGSampleTextureTest_LocalCorner::RunTest(const FString& Parameters)
 	return GenerateTestDataRunAndValidate(LocalCornerParameters);
 }
 
-IMPLEMENT_CUSTOM_SIMPLE_AUTOMATION_TEST(FPCGSampleTextureTest_LocalTexelCenter, PCGSampleTextureTestBase, "Plugins.PCG.SampleTexture.LocalTexelCenter", PCGTestsCommon::TestFlags)
+IMPLEMENT_CUSTOM_SIMPLE_AUTOMATION_TEST(FPCGSampleTextureTest_LocalTexelCenter, PCGSampleTextureTestBase, "Plugins.PCG.SampleTexture.LocalTexelCenter", PCGTestsCommon::TestFlags | EAutomationTestFlags::NonNullRHI)
 
 bool FPCGSampleTextureTest_LocalTexelCenter::RunTest(const FString& Parameters)
 {
@@ -172,7 +172,7 @@ bool FPCGSampleTextureTest_LocalTexelCenter::RunTest(const FString& Parameters)
 	return GenerateTestDataRunAndValidate(LocalTexelCenterParameters);
 }
 
-IMPLEMENT_CUSTOM_SIMPLE_AUTOMATION_TEST(FPCGSampleTextureTest_LocalCenter, PCGSampleTextureTestBase, "Plugins.PCG.SampleTexture.LocalCenter", PCGTestsCommon::TestFlags)
+IMPLEMENT_CUSTOM_SIMPLE_AUTOMATION_TEST(FPCGSampleTextureTest_LocalCenter, PCGSampleTextureTestBase, "Plugins.PCG.SampleTexture.LocalCenter", PCGTestsCommon::TestFlags | EAutomationTestFlags::NonNullRHI)
 
 bool FPCGSampleTextureTest_LocalCenter::RunTest(const FString& Parameters)
 {
