[PCG]
Exposed some surface functions to the external api

diff --git a/Engine/Plugins/PCG/Source/PCG/Public/Elements/PCGSurfaceSampler.h b/Engine/Plugins/PCG/Source/PCG/Public/Elements/PCGSurfaceSampler.h
--- a/Engine/Plugins/PCG/Source/PCG/Public/Elements/PCGSurfaceSampler.h
+++ b/Engine/Plugins/PCG/Source/PCG/Public/Elements/PCGSurfaceSampler.h
@@ -39,7 +39,7 @@ namespace PCGSurfaceSampler
 	struct FSurfaceSamplerData
 	{
 		bool Initialize(const UPCGSurfaceSamplerSettings* Settings, const FPCGContext* Context, const FBox& InEffectiveGridBounds, const FTransform& InSurfaceTransform = FTransform::Identity);
-		bool Initialize(const FPCGContext* Context, const FBox& InEffectiveGridBounds, const FTransform& InSurfaceTransform = FTransform::Identity);
+		PCG_API bool Initialize(const FPCGContext* Context, const FBox& InEffectiveGridBounds, const FTransform& InSurfaceTransform = FTransform::Identity);
 
 		FIntVector2 ComputeCellIndices(int32 Index) const;
 
@@ -82,7 +82,7 @@ namespace PCGSurfaceSampler
 	};
 
 	/** Sample a surface and returns the resulting point data. Can't be timesliced. */
-	UPCGPointData* SampleSurface(FPCGContext* Context, const UPCGSurfaceData* InSurface, const UPCGSpatialData* InBoundingShape, const FBox& EffectiveBounds, const FSurfaceSamplerParams& ExecutionParams);
+	PCG_API UPCGPointData* SampleSurface(FPCGContext* Context, const UPCGSurfaceData* InSurface, const UPCGSpatialData* InBoundingShape, const FBox& EffectiveBounds, const FSurfaceSamplerParams& ExecutionParams);
 
 	/** Sample a surface and write the results in the given point data. Can be timesliced, and will return false if the processing is not done, true otherwise. */
 	bool SampleSurface(FPCGContext* Context, const FSurfaceSamplerData& SamplerData, const UPCGSurfaceData* InSurface, const UPCGSpatialData* InBoundingShape, UPCGPointData* SampledData, const bool bTimeSlicingIsEnabled = false);
