[PCG] Release progress notification when PCG execution is paused

#rnx

diff --git a/Engine/Plugins/PCG/Source/PCG/Private/Graph/PCGGraphExecutor.cpp b/Engine/Plugins/PCG/Source/PCG/Private/Graph/PCGGraphExecutor.cpp
--- a/Engine/Plugins/PCG/Source/PCG/Private/Graph/PCGGraphExecutor.cpp
+++ b/Engine/Plugins/PCG/Source/PCG/Private/Graph/PCGGraphExecutor.cpp
@@ -1316,6 +1316,7 @@ void FPCGGraphExecutor::Execute()
 #if WITH_EDITOR
 	if (PCGSystemSwitches::CVarPausePCGExecution.GetValueOnAnyThread())
 	{
+		UpdateGenerationNotification();
 		return;
 	}
 #endif // WITH_EDITOR
@@ -2620,7 +2621,8 @@ void FPCGGraphExecutor::UpdateGenerationNotification()
 	}
 
 	const int32 RemainingTaskNum = GetNonScheduledRemainingTaskCount();
-	if (RemainingTaskNum == 0)
+	const bool bPaused = PCGSystemSwitches::CVarPausePCGExecution.GetValueOnAnyThread();
+	if (RemainingTaskNum == 0 || bPaused)
 	{
 		// If we had tasks on the last frame, start a timer to avoid releasing the existing notification too soon
 		if (GenerationProgressLastTaskNum != 0)
@@ -2629,7 +2631,7 @@ void FPCGGraphExecutor::UpdateGenerationNotification()
 			GenerationProgressNotificationStartTime = FPlatformTime::Seconds();
 		}
 
-		if ((FPlatformTime::Seconds() - GenerationProgressNotificationStartTime) > PCGGraphExecutor::CVarEditorNotificationDelayInSeconds.GetValueOnAnyThread())
+		if (bPaused || (FPlatformTime::Seconds() - GenerationProgressNotificationStartTime) > PCGGraphExecutor::CVarEditorNotificationDelayInSeconds.GetValueOnAnyThread())
 		{
 			ReleaseGenerationNotification();
 		}
