[PCG] Temporary fix for Copy Attribute issue

* For multi entries param data it doesn't seem to work well, so revert to previous way of doing it so it will be slower but at least be on par
* Also fix missing Super call in deprecation that was removing some pin connections

#rnx

diff --git a/Engine/Plugins/PCG/Source/PCG/Private/Elements/PCGCopyAttributes.cpp b/Engine/Plugins/PCG/Source/PCG/Private/Elements/PCGCopyAttributes.cpp
--- a/Engine/Plugins/PCG/Source/PCG/Private/Elements/PCGCopyAttributes.cpp
+++ b/Engine/Plugins/PCG/Source/PCG/Private/Elements/PCGCopyAttributes.cpp
@@ -98,6 +98,8 @@ void UPCGCopyAttributesSettings::ApplyDeprecation(UPCGNode* InOutNode)
 
 void UPCGCopyAttributesSettings::ApplyDeprecationBeforeUpdatePins(UPCGNode* InOutNode, TArray<TObjectPtr<UPCGPin>>& InputPins, TArray<TObjectPtr<UPCGPin>>& OutputPins)
 {
+	Super::ApplyDeprecationBeforeUpdatePins(InOutNode, InputPins, OutputPins);
+
 	if (!InOutNode)
 	{
 		return;
diff --git a/Engine/Plugins/PCG/Source/PCG/Private/Helpers/PCGMetadataHelpers.cpp b/Engine/Plugins/PCG/Source/PCG/Private/Helpers/PCGMetadataHelpers.cpp
--- a/Engine/Plugins/PCG/Source/PCG/Private/Helpers/PCGMetadataHelpers.cpp
+++ b/Engine/Plugins/PCG/Source/PCG/Private/Helpers/PCGMetadataHelpers.cpp
@@ -139,6 +139,7 @@ namespace PCGMetadataHelpers
 			}
 
 			// We need accessors if we have a multi entry source attribute or we have extractors
+			const bool bIsMultiEntries = SourceData->IsA<UPCGParamData>() && SourceMetadata->GetLocalItemCount() > 1;
 			const bool bInputHasAnyExtra = !InputSource.GetExtraNames().IsEmpty();
 			const bool bOutputHasAnyExtra = !OutputTarget.GetExtraNames().IsEmpty();
 			const bool bSourceIsAttribute = InputSource.GetSelection() == EPCGAttributePropertySelection::Attribute;
@@ -146,7 +147,7 @@ namespace PCGMetadataHelpers
 			// Cast is only required if it is on an output attribute that has no extra (that we will create)
 			const bool bOutputTypeCast = bTargetIsAttribute && !bOutputHasAnyExtra && (RequestedOutputType != EPCGMetadataTypes::Unknown);
 
-			const bool bNeedAccessors = bInputHasAnyExtra || bOutputHasAnyExtra || !bSourceIsAttribute || !bTargetIsAttribute || bOutputTypeCast;
+			const bool bNeedAccessors = bIsMultiEntries || bInputHasAnyExtra || bOutputHasAnyExtra || !bSourceIsAttribute || !bTargetIsAttribute || bOutputTypeCast;
 
 			// If no accessor, copy over the attribute
 			if (!bNeedAccessors)
