[CF] Add support for argument modifiers on data interface functions
[PCG] Add out params to GetThreadData to clean up this function

#rnx

diff --git a/Engine/Plugins/PCG/Shaders/Private/PCGDataCollectionDataInterface.ush b/Engine/Plugins/PCG/Shaders/Private/PCGDataCollectionDataInterface.ush
--- a/Engine/Plugins/PCG/Shaders/Private/PCGDataCollectionDataInterface.ush
+++ b/Engine/Plugins/PCG/Shaders/Private/PCGDataCollectionDataInterface.ush
@@ -49,10 +49,9 @@ uint ReadDataInfo_{DataInterfaceName}(uint DataAddress)
 	return {DataInterfaceName}_DataCollectionBuffer.Load(DataAddress + 12);
 }
 
-// Returns (DataIndex, DataAddress, ElementIndex)
-uint3 GetThreadData_{DataInterfaceName}(uint ThreadIndex)
+bool GetThreadData_{DataInterfaceName}(uint InThreadIndex, out uint OutDataIndex, out uint OutDataAddress, out uint OutElementIndex)
 {
-	int ElementIndex = ThreadIndex;
+	int ElementIndex = InThreadIndex;
 
 	const uint NumData = GetNumData_{DataInterfaceName}();
 	for (uint DataIndex = 0; DataIndex < NumData; ++DataIndex)
@@ -62,13 +61,27 @@ uint3 GetThreadData_{DataInterfaceName}(uint ThreadIndex)
 
 		if (ElementIndex < ElemCount)
 		{
-			return uint3(DataIndex, DataAddress, ElementIndex);
+			OutDataIndex = DataIndex;
+			OutDataAddress = DataAddress;
+			OutElementIndex = ElementIndex;
+
+			return true;
 		}
 
 		ElementIndex -= ElemCount;
 	}
 
-	return uint3(-1, -1, -1);
+	OutDataIndex = -1;
+	OutDataAddress = -1;
+	OutElementIndex = -1;
+
+	return false;
+}
+
+bool GetThreadData_{DataInterfaceName}(uint InThreadIndex, out uint OutDataIndex, out uint OutElementIndex)
+{
+	uint DataAddressDummy = -1;
+	return GetThreadData_{DataInterfaceName}(InThreadIndex, OutDataIndex, DataAddressDummy, OutElementIndex);
 }
 
 uint GetNumElements_{DataInterfaceName}()
diff --git a/Engine/Plugins/PCG/Source/PCG/Private/Compute/DataInterfaces/PCGDataCollectionDataInterface.cpp b/Engine/Plugins/PCG/Source/PCG/Private/Compute/DataInterfaces/PCGDataCollectionDataInterface.cpp
--- a/Engine/Plugins/PCG/Source/PCG/Private/Compute/DataInterfaces/PCGDataCollectionDataInterface.cpp
+++ b/Engine/Plugins/PCG/Source/PCG/Private/Compute/DataInterfaces/PCGDataCollectionDataInterface.cpp
@@ -64,8 +64,18 @@ void UPCGDataCollectionDataInterface::GetSupportedInputs(TArray<FShaderFunctionD
 
 		OutFunctions.AddDefaulted_GetRef()
 			.SetName(TEXT("GetThreadData"))
-			.AddReturnType(EShaderFundamentalType::Uint, 3)
-			.AddParam(EShaderFundamentalType::Uint);
+			.AddParam(EShaderFundamentalType::Uint) // InThreadIndex
+			.AddParam(EShaderFundamentalType::Uint, 0, 0, EShaderParamModifier::Out) // OutDataIndex
+			.AddParam(EShaderFundamentalType::Uint, 0, 0, EShaderParamModifier::Out) // OutDataAddress
+			.AddParam(EShaderFundamentalType::Uint, 0, 0, EShaderParamModifier::Out) // OutElementIndex
+			.AddReturnType(EShaderFundamentalType::Bool);
+
+		OutFunctions.AddDefaulted_GetRef()
+			.SetName(TEXT("GetThreadData"))
+			.AddParam(EShaderFundamentalType::Uint) // InThreadIndex
+			.AddParam(EShaderFundamentalType::Uint, 0, 0, EShaderParamModifier::Out) // OutDataIndex
+			.AddParam(EShaderFundamentalType::Uint, 0, 0, EShaderParamModifier::Out) // OutElementIndex
+			.AddReturnType(EShaderFundamentalType::Bool);
 
 		OutFunctions.AddDefaulted_GetRef()
 			.SetName(TEXT("GetNumElements"))
diff --git a/Engine/Plugins/PCG/Source/PCG/Private/Compute/Elements/PCGCustomHLSL.cpp b/Engine/Plugins/PCG/Source/PCG/Private/Compute/Elements/PCGCustomHLSL.cpp
--- a/Engine/Plugins/PCG/Source/PCG/Private/Compute/Elements/PCGCustomHLSL.cpp
+++ b/Engine/Plugins/PCG/Source/PCG/Private/Compute/Elements/PCGCustomHLSL.cpp
@@ -547,8 +547,9 @@ void UPCGCustomHLSLSettings::UpdateDeclarations()
 			auto EmitGetThreadElement = [&InInputDeclarations = InputDeclarations](const FPCGPinProperties& Properties)
 			{
 				InInputDeclarations += FString::Format(TEXT(
-					"// (DataIndex, DataAddress, ElementIndex). Returns -1 if invalid!\n"
-					"uint3 {0}_GetThreadData(uint ThreadIndex);\n"),
+					"// Returns false if thread has no data to operate on.\n"
+					"bool {0}_GetThreadData(uint ThreadIndex, out uint OutDataIndex, out uint OutElementIndex);\n"
+					"bool {0}_GetThreadData(uint ThreadIndex, out uint OutDataIndex, out uint OutDataAddress, out uint OutElementIndex);\n"),
 					{ Properties.Label.ToString() });
 			};
 
@@ -1063,6 +1064,15 @@ FString UPCGCustomHLSLSettings::GetCookedKernelSource(const TMap<FPCGKernelAttri
 	// Per-kernel-type preamble. Set up shader inputs and initialize output data.
 	FString KernelSpecificPreamble = TEXT("    // Kernel preamble\n");
 	
+	auto AddThreadInfoForPin = [&KernelSpecificPreamble](FName PinLabel)
+	{
+		KernelSpecificPreamble += FString::Format(TEXT(
+			"    uint {0}_DataIndex;\n"
+			"    uint {0}_DataAddress;\n"
+			"    if (!{0}_GetThreadData(ThreadIndex, {0}_DataIndex, {0}_DataAddress, ElementIndex)) return;\n"),
+			{ PinLabel.ToString() });
+	};
+
 	if (KernelType == EPCGKernelType::PointProcessor)
 	{
 		const UPCGPin* InputPin = GetPointProcessingInputPin();
@@ -1070,18 +1080,10 @@ FString UPCGCustomHLSLSettings::GetCookedKernelSource(const TMap<FPCGKernelAttri
 
 		if (InputPin && OutputPin)
 		{
-			KernelSpecificPreamble += FString::Format(TEXT(
-				"    const uint3 {0}_ThreadInfo = {0}_GetThreadData(ThreadIndex);\n"
-				"    const uint {0}_DataIndex = {0}_ThreadInfo[0];\n"
-				"    const uint {0}_DataAddress = {0}_ThreadInfo[1];\n"
-				"    const uint ElementIndex = {0}_ThreadInfo[2]; // Assumption - element index identical in input and output data.\n"),
-				{ InputPin->Properties.Label.ToString() });
+			KernelSpecificPreamble += TEXT("uint ElementIndex; // Assumption - element index identical in input and output data.\n");
 
-			KernelSpecificPreamble += FString::Format(TEXT(
-				"    const uint3 {0}_ThreadInfo = {0}_GetThreadData(ThreadIndex);\n"
-				"    const uint {0}_DataIndex = {0}_ThreadInfo[0];\n"
-				"    const uint {0}_DataAddress = {0}_ThreadInfo[1];\n"),
-				{ OutputPin->Properties.Label.ToString() });
+			AddThreadInfoForPin(InputPin->Properties.Label);
+			AddThreadInfoForPin(OutputPin->Properties.Label);
 
 			// Automatically copy value of all attributes for this element.
 			KernelSpecificPreamble += FString::Format(TEXT(
@@ -1117,18 +1119,13 @@ FString UPCGCustomHLSLSettings::GetCookedKernelSource(const TMap<FPCGKernelAttri
 	}
 	else if (KernelType == EPCGKernelType::PointGenerator)
 	{
-		KernelSpecificPreamble += FString::Format(TEXT(
-			"    const uint PointCount = {0};\n"),
-			{ PointCount });
+		KernelSpecificPreamble += FString::Format(TEXT("    const uint PointCount = {0};\n"), { PointCount });
 
 		if (const UPCGPin* OutputPin = GetFirstPointOutputPin())
 		{
-			KernelSpecificPreamble += FString::Format(TEXT(
-				"    const uint3 {0}_ThreadInfo = {0}_GetThreadData(ThreadIndex);\n"
-				"    const uint {0}_DataIndex = {0}_ThreadInfo[0];\n"
-				"    const uint {0}_DataAddress = {0}_ThreadInfo[1];\n"
-				"    const uint ElementIndex = {0}_ThreadInfo[2];\n"),
-				{ OutputPin->Properties.Label.ToString() });
+			KernelSpecificPreamble += TEXT("uint ElementIndex; // Assumption - element index identical in input and output data.\n");
+
+			AddThreadInfoForPin(OutputPin->Properties.Label);
 
 			KernelSpecificPreamble += FString::Format(TEXT(
 				"    \n"
diff --git a/Engine/Plugins/Runtime/ComputeFramework/Source/ComputeFramework/Private/ComputeGraph.cpp b/Engine/Plugins/Runtime/ComputeFramework/Source/ComputeFramework/Private/ComputeGraph.cpp
--- a/Engine/Plugins/Runtime/ComputeFramework/Source/ComputeFramework/Private/ComputeGraph.cpp
+++ b/Engine/Plugins/Runtime/ComputeFramework/Source/ComputeFramework/Private/ComputeGraph.cpp
@@ -432,6 +432,21 @@ namespace
 		
 		for (int32 ParameterIndex = bHasReturnWrap ? 1 : 0; ParameterIndex < NumWrapParams; ++ParameterIndex)
 		{
+			switch (FnWrap.ParamTypes[ParameterIndex].Modifier)
+			{
+			case EShaderParamModifier::In:
+				StringBuilder.Append(TEXT("in "));
+				break;
+			case EShaderParamModifier::Out:
+				StringBuilder.Append(TEXT("out "));
+				break;
+			case EShaderParamModifier::InOut:
+				StringBuilder.Append(TEXT("inout "));
+				break;
+			default:
+				break;
+			}
+
 			StringBuilder.Append(*FnWrap.ParamTypes[ParameterIndex].TypeDeclaration);
 			StringBuilder.Appendf(TEXT(" P%d"), ParameterIndex);
 			StringBuilder.Append((ParameterIndex < NumWrapParams - 1) ? TEXT(", ") : TEXT(""));
diff --git a/Engine/Plugins/Runtime/ComputeFramework/Source/ComputeFramework/Private/ShaderParamTypeDefinition.cpp b/Engine/Plugins/Runtime/ComputeFramework/Source/ComputeFramework/Private/ShaderParamTypeDefinition.cpp
--- a/Engine/Plugins/Runtime/ComputeFramework/Source/ComputeFramework/Private/ShaderParamTypeDefinition.cpp
+++ b/Engine/Plugins/Runtime/ComputeFramework/Source/ComputeFramework/Private/ShaderParamTypeDefinition.cpp
@@ -754,18 +754,20 @@ FShaderFunctionDefinition& FShaderFunctionDefinition::SetName(FString InName)
 	return *this;
 }
 
-FShaderFunctionDefinition& FShaderFunctionDefinition::AddParam(FShaderValueTypeHandle InValueType)
+FShaderFunctionDefinition& FShaderFunctionDefinition::AddParam(FShaderValueTypeHandle InValueType, EShaderParamModifier InModifier)
 {
 	FShaderParamTypeDefinition Def;
 	Def.ValueType = InValueType;
+	Def.Modifier = InModifier;
 	ParamTypes.Add(Def);
 	return *this;
 }
 
-FShaderFunctionDefinition& FShaderFunctionDefinition::AddParam(EShaderFundamentalType InType, int32 InRowCount, int32 InColumnCount)
+FShaderFunctionDefinition& FShaderFunctionDefinition::AddParam(EShaderFundamentalType InType, int32 InRowCount, int32 InColumnCount, EShaderParamModifier InModifier)
 {
 	FShaderParamTypeDefinition Def;
-	
+	Def.Modifier = InModifier;
+
 	if (InRowCount > 0 && InColumnCount > 0)
 	{
 		Def.ValueType = FShaderValueType::Get(InType, InRowCount, InColumnCount);
@@ -793,9 +795,9 @@ FShaderFunctionDefinition& FShaderFunctionDefinition::AddReturnType(FShaderValue
 	AddParam(InValueType);
 
 	// Return type is expected to be the first param.
-	if (ParamTypes.Num() > 1)
+	for (int I = ParamTypes.Num() - 1; I > 0; --I)
 	{
-		ParamTypes.Swap(0, ParamTypes.Num() - 1);
+		ParamTypes.Swap(I, I - 1);
 	}
 	return *this;
 }
@@ -810,9 +812,9 @@ FShaderFunctionDefinition& FShaderFunctionDefinition::AddReturnType(EShaderFunda
 	AddParam(InType, InRowCount, InColumnCount);
 
 	// Return type is expected to be the first param.
-	if (ParamTypes.Num() > 1)
+	for (int I = ParamTypes.Num() - 1; I > 0; --I)
 	{
-		ParamTypes.Swap(0, ParamTypes.Num() - 1);
+		ParamTypes.Swap(I, I - 1);
 	}
 	return *this;
 }
diff --git a/Engine/Plugins/Runtime/ComputeFramework/Source/ComputeFramework/Public/ComputeFramework/ShaderParamTypeDefinition.h b/Engine/Plugins/Runtime/ComputeFramework/Source/ComputeFramework/Public/ComputeFramework/ShaderParamTypeDefinition.h
--- a/Engine/Plugins/Runtime/ComputeFramework/Source/ComputeFramework/Public/ComputeFramework/ShaderParamTypeDefinition.h
+++ b/Engine/Plugins/Runtime/ComputeFramework/Source/ComputeFramework/Public/ComputeFramework/ShaderParamTypeDefinition.h
@@ -59,6 +59,15 @@ enum class EShaderResourceType : uint8
 	ByteAddressBuffer,
 };
 
+UENUM()
+enum class EShaderParamModifier : uint8
+{
+	None,
+	In,
+	Out,
+	InOut,
+};
+
 /*  */
 USTRUCT()
 struct FShaderValueTypeHandle
@@ -339,6 +348,9 @@ public:
 	UPROPERTY()
 	EShaderResourceType	ResourceType = EShaderResourceType::None;
 
+	UPROPERTY()
+	EShaderParamModifier Modifier;
+	
 	bool IsAnyBufferType() const
 	{
 		return
@@ -417,8 +429,8 @@ public:
 
 public:
 	FShaderFunctionDefinition& SetName(FString InName);
-	FShaderFunctionDefinition& AddParam(FShaderValueTypeHandle InValueType);
-	FShaderFunctionDefinition& AddParam(EShaderFundamentalType InType, int32 InRowCount = 0, int32 InColumnCount = 0);
+	FShaderFunctionDefinition& AddParam(FShaderValueTypeHandle InValueType, EShaderParamModifier InModifier = EShaderParamModifier::None);
+	FShaderFunctionDefinition& AddParam(EShaderFundamentalType InType, int32 InRowCount = 0, int32 InColumnCount = 0, EShaderParamModifier InModifier = EShaderParamModifier::None);
 	FShaderFunctionDefinition& AddReturnType(FShaderValueTypeHandle InValueType);
 	FShaderFunctionDefinition& AddReturnType(EShaderFundamentalType InType, int32 InRowCount = 0, int32 InColumnCount = 0);
 };
