[PCG] Discard points support. Only does actual point removal upon readback from GPU.

Future extension could be to compact active points and set indirect dispatch args for downstream kernels to avoid executing over removed points. For now downstream kernels early out on inactive points.

#rnx

diff --git a/Engine/Plugins/PCG/Shaders/Private/PCGDataCollectionDataInterface.ush b/Engine/Plugins/PCG/Shaders/Private/PCGDataCollectionDataInterface.ush
--- a/Engine/Plugins/PCG/Shaders/Private/PCGDataCollectionDataInterface.ush
+++ b/Engine/Plugins/PCG/Shaders/Private/PCGDataCollectionDataInterface.ush
@@ -16,6 +16,8 @@ RWByteAddressBuffer {DataInterfaceName}_DataCollectionBuffer;
 #define POINT_STEEPNESS_ID 8
 #define POINT_METADATA_ENTRY_ID 9
 
+#define INVALID_DENSITY 1.#INF
+
 // #################### HEADER READERS ##########################
 
 uint GetNumData_{DataInterfaceName}()
@@ -384,6 +386,11 @@ float4x4 GetPointTransform_{DataInterfaceName}(uint InDataIndex, uint InElementI
 		(float3)0.0, 1.0);
 }
 
+bool IsValid_{DataInterfaceName}(uint InDataIndex, uint InElementIndex)
+{
+	return GetDensity_{DataInterfaceName}(InDataIndex, InElementIndex) != INVALID_DENSITY;
+}
+
 // #################### POINT ATTRIBUTE SETTERS ##########################
 
 void SetPosition_{DataInterfaceName}(uint DataIndex, uint ElementIndex, float3 Position)
@@ -449,6 +456,11 @@ void SetPointTransform_{DataInterfaceName}(uint InDataIndex, uint InElementIndex
 	SetPosition_{DataInterfaceName}(InDataIndex, InElementIndex, Position);
 }
 
+void RemovePoint_{DataInterfaceName}(uint InDataIndex, uint InElementIndex)
+{
+	SetDensity_{DataInterfaceName}(InDataIndex, InElementIndex, INVALID_DENSITY);
+}
+
 // Write single point data collection header
 void WriteSinglePointDataCollectionHeader_{DataInterfaceName}(uint InPointCount)
 {
diff --git a/Engine/Plugins/PCG/Source/PCG/Private/Compute/DataInterfaces/PCGDataCollectionDataInterface.cpp b/Engine/Plugins/PCG/Source/PCG/Private/Compute/DataInterfaces/PCGDataCollectionDataInterface.cpp
--- a/Engine/Plugins/PCG/Source/PCG/Private/Compute/DataInterfaces/PCGDataCollectionDataInterface.cpp
+++ b/Engine/Plugins/PCG/Source/PCG/Private/Compute/DataInterfaces/PCGDataCollectionDataInterface.cpp
@@ -347,6 +347,12 @@ void UPCGDataCollectionDataInterface::GetSupportedInputs(TArray<FShaderFunctionD
 			.AddParam(EShaderFundamentalType::Uint)
 			.AddParam(EShaderFundamentalType::Uint);
 
+		OutFunctions.AddDefaulted_GetRef()
+			.SetName(TEXT("IsValid"))
+			.AddReturnType(EShaderFundamentalType::Bool)
+			.AddParam(EShaderFundamentalType::Uint) // DataIndex
+			.AddParam(EShaderFundamentalType::Uint); // ElementIndex
+
 		OutFunctions.AddDefaulted_GetRef()
 			.SetName(TEXT("GetPointTransform"))
 			.AddReturnType(EShaderFundamentalType::Float, 4, 4)
@@ -416,6 +422,11 @@ void UPCGDataCollectionDataInterface::GetSupportedInputs(TArray<FShaderFunctionD
 			.AddParam(EShaderFundamentalType::Uint)
 			.AddParam(EShaderFundamentalType::Float, 4, 4);
 
+		OutFunctions.AddDefaulted_GetRef()
+			.SetName(TEXT("RemovePoint"))
+			.AddParam(EShaderFundamentalType::Uint) // DataIndex
+			.AddParam(EShaderFundamentalType::Uint); // ElementIndex
+
 		OutFunctions.AddDefaulted_GetRef()
 			.SetName(TEXT("WriteSinglePointDataCollectionHeader"))
 			.AddParam(EShaderFundamentalType::Uint);
diff --git a/Engine/Plugins/PCG/Source/PCG/Private/Compute/Elements/PCGCustomHLSL.cpp b/Engine/Plugins/PCG/Source/PCG/Private/Compute/Elements/PCGCustomHLSL.cpp
--- a/Engine/Plugins/PCG/Source/PCG/Private/Compute/Elements/PCGCustomHLSL.cpp
+++ b/Engine/Plugins/PCG/Source/PCG/Private/Compute/Elements/PCGCustomHLSL.cpp
@@ -659,6 +659,7 @@ void UPCGCustomHLSLSettings::UpdateDeclarations()
 			InputDeclarations += TEXT("int <pin>_GetSeed(uint DataIndex, uint ElementIndex);\n");
 			InputDeclarations += TEXT("float <pin>_GetSteepness(uint DataIndex, uint ElementIndex);\n");
 			InputDeclarations += TEXT("float4x4 <pin>_GetPointTransform(uint DataIndex, uint ElementIndex);\n");
+			InputDeclarations += TEXT("bool <pin>_IsValid(uint DataIndex, uint ElementIndex);\n");
 		}
 
 		if (!LandscapeDataPins.IsEmpty())
@@ -736,6 +737,7 @@ void UPCGCustomHLSLSettings::UpdateDeclarations()
 			OutputDeclarations += TEXT("void <pin>_SetSeedFromPosition(uint DataIndex, uint ElementIndex, float3 Position);\n");
 			OutputDeclarations += TEXT("void <pin>_SetSteepness(uint DataIndex, uint ElementIndex, float Steepness);\n");
 			OutputDeclarations += TEXT("void <pin>_SetPointTransform(uint DataIndex, uint ElementIndex, float4x4 Transform);\n");
+			OutputDeclarations += TEXT("bool <pin>_RemovePoint(uint DataIndex, uint ElementIndex);\n");
 		}
 
 		if (!RawBufferDataPins.IsEmpty())
@@ -1080,11 +1082,20 @@ FString UPCGCustomHLSLSettings::GetCookedKernelSource(const TMap<FPCGKernelAttri
 
 		if (InputPin && OutputPin)
 		{
-			KernelSpecificPreamble += TEXT("uint ElementIndex; // Assumption - element index identical in input and output data.\n");
+			KernelSpecificPreamble += TEXT("    uint ElementIndex; // Assumption - element index identical in input and output data.\n");
 
 			AddThreadInfoForPin(InputPin->Properties.Label);
 			AddThreadInfoForPin(OutputPin->Properties.Label);
 
+			// If input point is invalid, mark output point as invalid and abort.
+			KernelSpecificPreamble += FString::Format(TEXT(
+				"    if (!{0}_IsValid({0}_DataIndex, ElementIndex))\n"
+				"    {\n"
+				"        {1}_RemovePoint({1}_DataIndex, ElementIndex);\n"
+				"        return;\n"
+				"    }\n"),
+				{ InputPin->Properties.Label.ToString(), OutputPin->Properties.Label.ToString() });
+
 			// Automatically copy value of all attributes for this element.
 			KernelSpecificPreamble += FString::Format(TEXT(
 				"\n"
diff --git a/Engine/Plugins/PCG/Source/PCG/Private/Compute/PCGDataForGPU.cpp b/Engine/Plugins/PCG/Source/PCG/Private/Compute/PCGDataForGPU.cpp
--- a/Engine/Plugins/PCG/Source/PCG/Private/Compute/PCGDataForGPU.cpp
+++ b/Engine/Plugins/PCG/Source/PCG/Private/Compute/PCGDataForGPU.cpp
@@ -6,6 +6,7 @@
 #include "PCGPoint.h"
 #include "Compute/PCGComputeCommon.h"
 #include "Data/PCGPointData.h"
+#include "Helpers/PCGAsync.h"
 #include "Metadata/PCGMetadata.h"
 
 #include "Async/ParallelFor.h"
@@ -949,6 +950,21 @@ void FPCGDataCollectionDesc::UnpackDataCollection(const TArray<uint8>& InPackedD
 					}
 				});
 			}
+
+			// TODO: It may be more efficient to create a mapping from input point index to final output point index and do everything in one pass.
+			auto DiscardInvalidPoints = [&OutPoints](int32 Index, FPCGPoint& OutPoint) -> bool
+			{
+				if (!FMath::IsFinite(OutPoints[Index].Density))
+				{
+					return false;
+				}
+
+				OutPoint = OutPoints[Index];
+				return true;
+			};
+
+			FPCGAsync::AsyncPointProcessing(/*Context=*/nullptr, OutPoints.Num(), OutPoints, DiscardInvalidPoints);
+
 		}
 		else { /* TODO: Support non-point data. */ }
 	}
