[PCG] New pin selection UI - specify list of pins when computing over a product of input data.

To avoid adding a dynamic array to all pin properties, this is added to an extension struct and used just where needed - currently only for output pins on custom hlsl nodes.

#rnx

diff --git a/Engine/Plugins/PCG/Source/PCG/Private/Compute/Elements/PCGCustomHLSL.cpp b/Engine/Plugins/PCG/Source/PCG/Private/Compute/Elements/PCGCustomHLSL.cpp
--- a/Engine/Plugins/PCG/Source/PCG/Private/Compute/Elements/PCGCustomHLSL.cpp
+++ b/Engine/Plugins/PCG/Source/PCG/Private/Compute/Elements/PCGCustomHLSL.cpp
@@ -75,6 +75,13 @@ void UPCGCustomHLSLSettings::PostInitProperties()
 }
 #endif
 
+TArray<FPCGPinProperties> UPCGCustomHLSLSettings::OutputPinProperties() const
+{
+	TArray<FPCGPinProperties> PinProperties;
+	Algo::Transform(OutputPins, PinProperties, [](const FPCGPinPropertiesGPU& InPropertiesGPU) { return InPropertiesGPU; });
+	return PinProperties;
+}
+
 #if WITH_EDITOR
 void UPCGCustomHLSLSettings::PostEditChangeProperty(FPropertyChangedEvent& PropertyChangedEvent)
 {
@@ -157,59 +164,44 @@ int UPCGCustomHLSLSettings::ComputeKernelThreadCount(const UPCGDataBinding* Bind
 	}
 	else if (KernelType == EPCGKernelType::Custom)
 	{
-		auto GetFromFirstInput = [this, Binding]() -> int
-		{
-			const UPCGPin* InputPin = GetFirstInputPin();
-			return GetProcessingElemCountForInputPin(InputPin, Binding);
-		};
-
-		auto GetFromSecondInput = [this, Binding]() -> int
-		{
-			const UPCGPin* InputPin = GetSecondInputPin();
-			return InputPin ? GetProcessingElemCountForInputPin(InputPin, Binding) : 0;
-		};
-
-		auto GetFromFirstInputXSecondInput = [this, Binding]() -> int
-		{
-			const UPCGPin* FirstPin = GetFirstInputPin();
-			const UPCGPin* SecondPin = GetSecondInputPin();
-			const int FirstElemCount = FirstPin ? GetProcessingElemCountForInputPin(FirstPin, Binding) : 0;
-			const int SecondElemCount = SecondPin ? GetProcessingElemCountForInputPin(SecondPin, Binding) : 0;
-			return FirstElemCount * SecondElemCount;
-		};
-
-		if (DispatchThreadCount == EPCGDispatchThreadCount::FromOutput)
+		if (DispatchThreadCount == EPCGDispatchThreadCount::FromFirstOutputPin)
 		{
 			const UPCGPin* OutputPin = GetFirstOutputPin();
-
-			if (OutputPin->Properties.BufferSizeMode == EPCGPinBufferSizeMode::FixedElementCount)
-			{
-				ThreadCount = OutputPin->Properties.FixedBufferElementCount;
-			}
-			else if (OutputPin->Properties.BufferSizeMode == EPCGPinBufferSizeMode::FromFirstPin)
-			{
-				ThreadCount = GetFromFirstInput();
-			}
-			else if (OutputPin->Properties.BufferSizeMode == EPCGPinBufferSizeMode::FromSecondPin)
-			{
-				ThreadCount = GetFromSecondInput();
-			}
-			else if (OutputPin->Properties.BufferSizeMode == EPCGPinBufferSizeMode::FirstPinXSecondPin)
+			const FPCGPinPropertiesGPU* PropertiesGPU = OutputPin ? GetOutputPinPropertiesGPU(OutputPin->Properties.Label) : nullptr;
+			if (PropertiesGPU)
 			{
-				ThreadCount = GetFromFirstInputXSecondInput();
+				if (PropertiesGPU->BufferSizeMode == EPCGPinBufferSizeMode::FixedElementCount)
+				{
+					ThreadCount = PropertiesGPU->FixedBufferElementCount;
+				}
+				else if (PropertiesGPU->BufferSizeMode == EPCGPinBufferSizeMode::FromFirstPin)
+				{
+					if (const UPCGPin* InputPin = GetFirstInputPin())
+					{
+						ThreadCount = GetProcessingElemCountForInputPin(InputPin, Binding);
+					}
+				}
+				else if (PropertiesGPU->BufferSizeMode == EPCGPinBufferSizeMode::FromProductOfInputPins)
+				{
+					for (const FName& PinLabel : PropertiesGPU->BufferSizeInputPinLabels)
+					{
+						if (const UPCGPin* InputPin = GetInputPin(PinLabel))
+						{
+							ThreadCount = FMath::Max(ThreadCount, 1) * GetProcessingElemCountForInputPin(InputPin, Binding);
+						}
+					}
+				}
 			}
 		}
-		else if (DispatchThreadCount == EPCGDispatchThreadCount::FromFirstInput)
-		{
-			ThreadCount = GetFromFirstInput();
-		}
-		else if (DispatchThreadCount == EPCGDispatchThreadCount::FromSecondInput)
-		{
-			ThreadCount = GetFromSecondInput();
-		}
-		else if (DispatchThreadCount == EPCGDispatchThreadCount::FirstInputXSecondInput)
+		else if (DispatchThreadCount == EPCGDispatchThreadCount::FromProductOfInputPins)
 		{
-			ThreadCount = GetFromFirstInputXSecondInput();
+			for (const FName& PinLabel : ThreadCountInputPinLabels)
+			{
+				if (const UPCGPin* InputPin = GetInputPin(PinLabel))
+				{
+					ThreadCount = FMath::Max(ThreadCount, 1) * GetProcessingElemCountForInputPin(InputPin, Binding);
+				}
+			}
 		}
 		else if (DispatchThreadCount == EPCGDispatchThreadCount::Fixed)
 		{
@@ -345,52 +337,50 @@ FPCGDataCollectionDesc UPCGCustomHLSLSettings::ComputeOutputPinDataDesc(const UP
 		}
 	}
 
+	const FPCGPinPropertiesGPU* PropertiesGPU = GetOutputPinPropertiesGPU(OutputPin->Properties.Label);
+	if (!ensure(PropertiesGPU))
+	{
+		return PinDesc;
+	}
+
 	const bool bNeedsPointPin = KernelType == EPCGKernelType::PointProcessor || KernelType == EPCGKernelType::PointGenerator;
 
 	// No size set by kernel, fall back to pin settings.
-	if (OutputPin->Properties.BufferSizeMode == EPCGPinBufferSizeMode::FromFirstPin)
+	if (PropertiesGPU->BufferSizeMode == EPCGPinBufferSizeMode::FromFirstPin)
 	{
 		if (const UPCGPin* InputPin = bNeedsPointPin ? GetPointProcessingInputPin() : GetFirstInputPin())
 		{
 			PinDesc = ComputeInputPinDataDesc(InputPin, Binding);
 		}
 	}
-	else if (OutputPin->Properties.BufferSizeMode == EPCGPinBufferSizeMode::FromSecondPin)
-	{
-		if (const UPCGPin* InputPin = bNeedsPointPin ? GetSecondPointProcessingInputPin() : GetSecondInputPin())
-		{
-			PinDesc = ComputeInputPinDataDesc(InputPin, Binding);
-		}
-	}
-	else if (OutputPin->Properties.BufferSizeMode == EPCGPinBufferSizeMode::FirstPinXSecondPin)
+	else if (PropertiesGPU->BufferSizeMode == EPCGPinBufferSizeMode::FromProductOfInputPins)
 	{
-		const UPCGPin* FirstPin = GetPointProcessingInputPin();
-		const UPCGPin* SecondPin = GetSecondPointProcessingInputPin();
+		// TODO for now we only support point. Needs extended logic for other data types.
+		const EPCGDataType DataType = EPCGDataType::Point;
 
-		if (FirstPin && SecondPin)
-		{
-			const FPCGDataCollectionDesc FirstPinDesc = ComputeInputPinDataDesc(FirstPin, Binding);
-			const FPCGDataCollectionDesc SecondPinDesc = ComputeInputPinDataDesc(SecondPin, Binding);
+		int TotalElementCount = 0;
 
-			int TotalSourceElementCount = 0;
-			for (const FPCGDataDesc& DataDesc : FirstPinDesc.DataDescs)
-			{
-				if (DataDesc.Type == EPCGDataType::Point)
-				{
-					TotalSourceElementCount += DataDesc.ElementCount;
-				}
-			}
-
-			int TotalOutputElementCount = 0;
-			for (const FPCGDataDesc& DataDesc : SecondPinDesc.DataDescs)
+		for (const FName& PinLabel : PropertiesGPU->BufferSizeInputPinLabels)
+		{
+			if (const UPCGPin* InputPin = GetInputPin(PinLabel))
 			{
-				if (DataDesc.Type == EPCGDataType::Point)
-				{
-					TotalOutputElementCount += DataDesc.ElementCount * TotalSourceElementCount;
-				}
+				const int ElementCount = ComputeInputPinDataDesc(InputPin, Binding).ComputeDataElementCount(DataType);
+				TotalElementCount = FMath::Max(TotalElementCount, 1) * ElementCount;
 			}
+		}
 
-			PinDesc.DataDescs.Emplace(EPCGDataType::Point, TotalOutputElementCount);
+		if (TotalElementCount > 0)
+		{
+			PinDesc.DataDescs.Emplace(DataType, TotalElementCount);
+		}
+	}
+	else if (PropertiesGPU->BufferSizeMode == EPCGPinBufferSizeMode::FixedElementCount)
+	{
+		if (ensure(PropertiesGPU->FixedBufferElementCount > 0))
+		{
+			// TODO for now we only support point. Needs extended logic for other data types.
+			const EPCGDataType DataType = EPCGDataType::Point;
+			PinDesc.DataDescs.Emplace(DataType, PropertiesGPU->FixedBufferElementCount);
 		}
 	}
 
@@ -415,21 +405,23 @@ EPCGChangeType UPCGCustomHLSLSettings::GetChangeTypeForProperty(const FName& InP
 }
 #endif
 
-const UPCGPin* UPCGCustomHLSLSettings::GetFirstInputPin() const
+const UPCGPin* UPCGCustomHLSLSettings::GetInputPin(FName Label) const
 {
-	if (const UPCGNode* Node = Cast<UPCGNode>(GetOuter()))
-	{
-		return !Node->GetInputPins().IsEmpty() ? Node->GetInputPins()[0] : nullptr;
-	}
+	const UPCGNode* Node = Cast<UPCGNode>(GetOuter());
+	return Node ? Node->GetInputPin(Label) : nullptr;
+}
 
-	return nullptr;
+const UPCGPin* UPCGCustomHLSLSettings::GetOutputPin(FName Label) const
+{
+	const UPCGNode* Node = Cast<UPCGNode>(GetOuter());
+	return Node ? Node->GetOutputPin(Label) : nullptr;
 }
 
-const UPCGPin* UPCGCustomHLSLSettings::GetSecondInputPin() const
+const UPCGPin* UPCGCustomHLSLSettings::GetFirstInputPin() const
 {
 	if (const UPCGNode* Node = Cast<UPCGNode>(GetOuter()))
 	{
-		return (Node->GetInputPins().Num() > 1) ? Node->GetInputPins()[1] : nullptr;
+		return !Node->GetInputPins().IsEmpty() ? Node->GetInputPins()[0] : nullptr;
 	}
 
 	return nullptr;
@@ -449,28 +441,6 @@ const UPCGPin* UPCGCustomHLSLSettings::GetPointProcessingInputPin() const
 	return nullptr;
 }
 
-const UPCGPin* UPCGCustomHLSLSettings::GetSecondPointProcessingInputPin() const
-{
-	if (const UPCGNode* Node = Cast<UPCGNode>(GetOuter()))
-	{
-		bool bEncounteredPointPin = false;
-		for (const UPCGPin* Pin : Node->GetInputPins())
-		{
-			if (Pin->Properties.AllowedTypes == EPCGDataType::Point)
-			{
-				if (bEncounteredPointPin)
-				{
-					return Pin;
-				}
-
-				bEncounteredPointPin = true;
-			}
-		}
-	}
-
-	return nullptr;
-}
-
 const UPCGPin* UPCGCustomHLSLSettings::GetFirstOutputPin() const
 {
 	if (const UPCGNode* Node = Cast<UPCGNode>(GetOuter()))
@@ -495,6 +465,14 @@ const UPCGPin* UPCGCustomHLSLSettings::GetFirstPointOutputPin() const
 	return nullptr;
 }
 
+const FPCGPinPropertiesGPU* UPCGCustomHLSLSettings::GetOutputPinPropertiesGPU(const FName& InPinLabel) const
+{
+	return OutputPins.FindByPredicate([InPinLabel](const FPCGPinPropertiesGPU& InProperties)
+	{
+		return InProperties.Label == InPinLabel;
+	});
+}
+
 #if WITH_EDITOR
 void UPCGCustomHLSLSettings::UpdateDeclarations()
 {
@@ -758,10 +736,10 @@ void UPCGCustomHLSLSettings::UpdateDeclarations()
 
 void UPCGCustomHLSLSettings::UpdatePinSettings()
 {
+	// Setup input pins.
 	for (int PinIndex = 0; PinIndex < InputPins.Num(); ++PinIndex)
 	{
 		FPCGPinProperties& Properties = InputPins[PinIndex];
-		Properties.bDisplayBufferSizeSettings = false;
 
 		// Type Any is not allowed, default to Point
 		if (Properties.AllowedTypes == EPCGDataType::Any)
@@ -788,9 +766,10 @@ void UPCGCustomHLSLSettings::UpdatePinSettings()
 		Properties.bAllowEditMultipleConnections = false;
 	}
 
+	// Setup output pins.
 	for (int PinIndex = 0; PinIndex < OutputPins.Num(); ++PinIndex)
 	{
-		FPCGPinProperties& Properties = OutputPins[PinIndex];
+		FPCGPinPropertiesGPU& Properties = OutputPins[PinIndex];
 
 		// Type Any is not allowed, default to Point
 		if (Properties.AllowedTypes == EPCGDataType::Any)
@@ -831,20 +810,21 @@ bool UPCGCustomHLSLSettings::IsKernelValid(FPCGContext* InContext) const
 		{
 			if (InContext)
 			{
-				PCGE_LOG_C(Error, GraphAndLog, InContext, FText::Format(LOCTEXT("InvalidAnyInput", "Custom kernels do not support inputs of type Any, found on pin {0}."), FText::FromName(Properties.Label)));
+				PCGE_LOG_C(Error, GraphAndLog, InContext, FText::Format(LOCTEXT("InvalidAnyInput", "Custom kernels do not support inputs of type Any, found on pin '{0}'."), FText::FromName(Properties.Label)));
 			}
 
 			return false;
 		}
 	}
 
-	for (const FPCGPinProperties& Properties : OutputPinProperties())
+	bool bIsFirstOutputPin = true;
+	for (const FPCGPinPropertiesGPU& Properties : OutputPins)
 	{
 		if (Properties.AllowedTypes == EPCGDataType::Any)
 		{
 			if (InContext)
 			{
-				PCGE_LOG_C(Error, GraphAndLog, InContext, FText::Format(LOCTEXT("InvalidAnyOutput", "Custom kernels do not support outputs of type Any, found on pin {0}."), FText::FromName(Properties.Label)));
+				PCGE_LOG_C(Error, GraphAndLog, InContext, FText::Format(LOCTEXT("InvalidAnyOutput", "Custom kernels do not support outputs of type Any, found on pin '{0}'."), FText::FromName(Properties.Label)));
 			}
 
 			return false;
@@ -854,7 +834,7 @@ bool UPCGCustomHLSLSettings::IsKernelValid(FPCGContext* InContext) const
 		{
 			if (InContext)
 			{
-				PCGE_LOG_C(Error, GraphAndLog, InContext, FText::Format(LOCTEXT("InvalidLSOutput", "Custom kernels do not support outputs of type Landscape, found on pin {0}."), FText::FromName(Properties.Label)));
+				PCGE_LOG_C(Error, GraphAndLog, InContext, FText::Format(LOCTEXT("InvalidLSOutput", "Custom kernels do not support outputs of type Landscape, found on pin '{0}'."), FText::FromName(Properties.Label)));
 			}
 
 			return false;
@@ -864,11 +844,88 @@ bool UPCGCustomHLSLSettings::IsKernelValid(FPCGContext* InContext) const
 		{
 			if (InContext)
 			{
-				PCGE_LOG_C(Error, GraphAndLog, InContext, FText::Format(LOCTEXT("InvalidTextureOutput", "Custom kernels do not support outputs of type Texture, found on pin {0}."), FText::FromName(Properties.Label)));
+				PCGE_LOG_C(Error, GraphAndLog, InContext, FText::Format(LOCTEXT("InvalidTextureOutput", "Custom kernels do not support outputs of type Texture, found on pin '{0}'."), FText::FromName(Properties.Label)));
 			}
 
 			return false;
 		}
+
+		// First output pin of point generator and point processor has fixed settings.
+		if ((KernelType != EPCGKernelType::PointGenerator && KernelType != EPCGKernelType::PointProcessor) || !bIsFirstOutputPin)
+		{
+			if (Properties.BufferSizeMode == EPCGPinBufferSizeMode::FixedElementCount)
+			{
+				if (Properties.FixedBufferElementCount <= 0)
+				{
+					if (InContext)
+					{
+						PCGE_LOG_C(Error, GraphAndLog, InContext, FText::Format(
+							LOCTEXT("InvalidFixedBufferSize", "Fixed GPU buffer size on '{0}' was invalid (%d)."),
+							FText::FromName(Properties.Label),
+							Properties.FixedBufferElementCount));
+					}
+
+				}
+			}
+			else if (Properties.BufferSizeMode == EPCGPinBufferSizeMode::FromFirstPin)
+			{
+				if (InputPinProperties().IsEmpty())
+				{
+					if (InContext)
+					{
+						PCGE_LOG_C(Error, GraphAndLog, InContext, FText::Format(
+							LOCTEXT("InvalidBufferSizeNoInputPin", "GPU buffer size for pin '{0}' could not be computed as there are no input pins."),
+							FText::FromName(Properties.Label)));
+					}
+
+					return false;
+				}
+			}
+			else if (Properties.BufferSizeMode == EPCGPinBufferSizeMode::FromProductOfInputPins)
+			{
+				if (InputPinProperties().IsEmpty())
+				{
+					if (InContext)
+					{
+						PCGE_LOG_C(Error, GraphAndLog, InContext, FText::Format(
+							LOCTEXT("InvalidBufferSizeNoInputPins", "GPU buffer size for pin '{0}' could not be computed as there are no input pins on this node."),
+							FText::FromName(Properties.Label)));
+					}
+
+					return false;
+				}
+
+				if (Properties.BufferSizeInputPinLabels.IsEmpty())
+				{
+					if (InContext)
+					{
+						PCGE_LOG_C(Error, GraphAndLog, InContext, FText::Format(
+							LOCTEXT("InvalidBufferSizeNoInputPins", "GPU buffer size for pin '{0}' could not be computed as input pins are specified in the pin settings."),
+							FText::FromName(Properties.Label)));
+					}
+
+					return false;
+				}
+
+				for (const FName& Label : Properties.BufferSizeInputPinLabels)
+				{
+					if (!GetInputPin(Label))
+					{
+						if (InContext)
+						{
+							PCGE_LOG_C(Error, GraphAndLog, InContext, FText::Format(
+								LOCTEXT("MissingBufferSizePin", "GPU buffer size for pin '{0}' could not be computed. Invalid pin specified in Input Pins array: '{1}'."),
+								FText::FromName(Properties.Label),
+								FText::FromName(Label)));
+						}
+
+						return false;
+					}
+				}
+			}
+		}
+
+		bIsFirstOutputPin = false;
 	}
 
 	if (KernelType == EPCGKernelType::PointProcessor)
@@ -918,6 +975,31 @@ bool UPCGCustomHLSLSettings::IsKernelValid(FPCGContext* InContext) const
 		return false;
 	}
 
+	if (KernelType == EPCGKernelType::Custom && DispatchThreadCount == EPCGDispatchThreadCount::FromProductOfInputPins)
+	{
+		if (ThreadCountInputPinLabels.IsEmpty())
+		{
+			if (InContext)
+			{
+				PCGE_LOG_C(Error, GraphAndLog, InContext, LOCTEXT("MissingThreadCountPins", "Dispatch thread count is based on input pins but no labels have been set in Input Pins array."));
+			}
+
+			return false;
+		}
+
+		for (const FName& Label : ThreadCountInputPinLabels)
+		{
+			if (!GetInputPin(Label))
+			{
+				if (InContext)
+				{
+					PCGE_LOG_C(Error, GraphAndLog, InContext, FText::Format(LOCTEXT("MissingThreadCountPin", "Invalid pin specified in Input Pins array: '{0}'."), FText::FromName(Label)));
+				}
+
+				return false;
+			}
+		}
+	}
 	return true;
 }
 
@@ -993,10 +1075,8 @@ FString UPCGCustomHLSLSettings::GetCookedKernelSource(const TMap<FPCGKernelAttri
 			? GetPointProcessingInputPin()
 			: CastChecked<UPCGNode>(GetOuter())->GetPassThroughInputPin();
 
-		const UPCGPin* SecondPin = GetSecondPointProcessingInputPin();
-
 		// Initialize all output headers.
-		for (const FPCGPinProperties& PinProps : OutputPinProperties())
+		for (const FPCGPinPropertiesGPU& PinProps : OutputPins)
 		{
 			// Only support Point pins for now.
 			if (PinProps.AllowedTypes != EPCGDataType::Point)
@@ -1008,9 +1088,17 @@ FString UPCGCustomHLSLSettings::GetCookedKernelSource(const TMap<FPCGKernelAttri
 			{
 				EmitHeaderWriterFromInputPin(PinProps, FirstPin);
 			}
-			else if (PinProps.BufferSizeMode == EPCGPinBufferSizeMode::FromSecondPin && SecondPin)
+			else if (PinProps.BufferSizeMode == EPCGPinBufferSizeMode::FixedElementCount)
+			{
+				HeaderWriters += FString::Format(TEXT(
+					"    {0}_WriteSinglePointDataCollectionHeader({1});"
+					"    AllMemoryBarrier();"),
+					{ PinProps.Label.ToString(), PinProps.FixedBufferElementCount });
+
+			}
+			else if (PinProps.BufferSizeMode == EPCGPinBufferSizeMode::FromProductOfInputPins)
 			{
-				EmitHeaderWriterFromInputPin(PinProps, SecondPin);
+				HeaderWriters += TEXT("    // The output data header must be written from the kernel code, for example by using WriteSinglePointDataCollectionHeader.");
 			}
 		}
 	}
@@ -1019,7 +1107,6 @@ FString UPCGCustomHLSLSettings::GetCookedKernelSource(const TMap<FPCGKernelAttri
 		if (const UPCGNode* Node = Cast<UPCGNode>(GetOuter()))
 		{
 			const UPCGPin* FirstPin = Node->GetPassThroughInputPin();
-			const UPCGPin* SecondPin = GetSecondPointProcessingInputPin();
 			const UPCGPin* PrimaryOutputPin = GetFirstPointOutputPin();
 			
 			auto EmitPointGenHeader = [&HeaderWriters, InPointCount = PointCount](const FPCGPinProperties& InOutputPinProps)
@@ -1041,20 +1128,22 @@ FString UPCGCustomHLSLSettings::GetCookedKernelSource(const TMap<FPCGKernelAttri
 					continue;
 				}
 
+				const FPCGPinPropertiesGPU* PinPropsGPU = GetOutputPinPropertiesGPU(OutputPin->Properties.Label);
+				if (!ensure(PinPropsGPU))
+				{
+					continue;
+				}
+
 				const FPCGPinProperties& PinProps = OutputPin->Properties;
 
 				if (OutputPin == PrimaryOutputPin)
 				{
 					EmitPointGenHeader(PinProps);
 				}
-				else if (PinProps.BufferSizeMode == EPCGPinBufferSizeMode::FromFirstPin && FirstPin)
+				else if (PinPropsGPU->BufferSizeMode == EPCGPinBufferSizeMode::FromFirstPin && FirstPin)
 				{
 					EmitHeaderWriterFromInputPin(PinProps, FirstPin);
 				}
-				else if (PinProps.BufferSizeMode == EPCGPinBufferSizeMode::FromSecondPin && SecondPin)
-				{
-					EmitHeaderWriterFromInputPin(PinProps, SecondPin);
-				}
 			}
 		}
 	}
diff --git a/Engine/Plugins/PCG/Source/PCG/Private/PCGPin.cpp b/Engine/Plugins/PCG/Source/PCG/Private/PCGPin.cpp
--- a/Engine/Plugins/PCG/Source/PCG/Private/PCGPin.cpp
+++ b/Engine/Plugins/PCG/Source/PCG/Private/PCGPin.cpp
@@ -91,9 +91,7 @@ bool FPCGPinProperties::operator==(const FPCGPinProperties& Other) const
 		bAllowMultipleData == Other.bAllowMultipleData &&
 		Usage == Other.Usage &&
 		PinStatus == Other.PinStatus &&
-		bInvisiblePin == Other.bInvisiblePin &&
-		BufferSizeMode == Other.BufferSizeMode &&
-		FixedBufferElementCount == Other.FixedBufferElementCount;
+		bInvisiblePin == Other.bInvisiblePin;
 }
 
 void FPCGPinProperties::PostSerialize(const FArchive& Ar)
diff --git a/Engine/Plugins/PCG/Source/PCG/Public/Compute/Elements/PCGCustomHLSL.h b/Engine/Plugins/PCG/Source/PCG/Public/Compute/Elements/PCGCustomHLSL.h
--- a/Engine/Plugins/PCG/Source/PCG/Public/Compute/Elements/PCGCustomHLSL.h
+++ b/Engine/Plugins/PCG/Source/PCG/Public/Compute/Elements/PCGCustomHLSL.h
@@ -8,6 +8,46 @@
 
 class UPCGPin;
 
+/** Method for computing the size of a pin on a GPU node. */
+UENUM()
+enum class EPCGPinBufferSizeMode : uint8
+{
+	FromFirstPin UMETA(DisplayName = "Match First Input Pin"),
+	FromProductOfInputPins UMETA(Tooltip = "Dispatches a thread per element in the product of one or more pins. So if there are 4 data elements in pin A and 6 data elements in pin B, 24 threads will be dispatched."),
+	FixedElementCount,
+};
+
+/** An extension of the pin properties that adds hints for GPU thread count / buffer size calculations. */
+USTRUCT(BlueprintType, meta = (HasNativeBreak = "/Script/PCG.PCGBlueprintPinHelpers.BreakPinProperty", HasNativeMake = "/Script/PCG.PCGBlueprintPinHelpers.MakePinProperty"))
+struct PCG_API FPCGPinPropertiesGPU : public FPCGPinProperties
+{
+	GENERATED_BODY()
+
+public:
+	FPCGPinPropertiesGPU() = default;
+
+	explicit FPCGPinPropertiesGPU(const FName& InLabel, EPCGDataType InAllowedTypes)
+		: FPCGPinProperties(InLabel, InAllowedTypes)
+	{
+	}
+
+public:
+	/** Compute graphs use this to calculate the buffer size of output pins. */
+	UPROPERTY(BlueprintReadWrite, EditAnywhere, Category = "GPU Buffer Size", meta = (EditCondition = "bDisplayBufferSizeSettings", EditConditionHides))
+	EPCGPinBufferSizeMode BufferSizeMode = EPCGPinBufferSizeMode::FromFirstPin;
+
+	UPROPERTY(BlueprintReadWrite, EditAnywhere, Category = "GPU Buffer Size", meta = (EditCondition = "bDisplayBufferSizeSettings && BufferSizeMode == EPCGPinBufferSizeMode::FixedElementCount", EditConditionHides))
+	int FixedBufferElementCount = 4;
+
+	UPROPERTY(BlueprintReadWrite, EditAnywhere, DisplayName = "Input Pins", Category = "GPU Buffer Size", meta = (EditCondition = "bDisplayBufferSizeSettings && BufferSizeMode == EPCGPinBufferSizeMode::FromProductOfInputPins", EditConditionHides))
+	TArray<FName> BufferSizeInputPinLabels;
+
+#if WITH_EDITORONLY_DATA
+	UPROPERTY(Transient)
+	bool bDisplayBufferSizeSettings = true;
+#endif // WITH_EDITORONLY_DATA
+};
+
 /** Type of kernel allows us to make decisions about execution automatically, streamlining authoring. */
 UENUM()
 enum class EPCGKernelType : uint8
@@ -21,13 +61,9 @@ enum class EPCGKernelType : uint8
 UENUM()
 enum class EPCGDispatchThreadCount : uint8
 {
-	FromOutput UMETA(DisplayName = "Match First Output Pin", Tooltip = "One thread per pin data element."),
-	FromFirstInput UMETA(DisplayName = "Match First Input Pin", Tooltip = "One thread per pin data element."),
-	FromSecondInput UMETA(DisplayName = "Match Second Input Pin", Tooltip = "One thread per pin data element."),
-	FirstInputXSecondInput UMETA(DisplayName = "First Input X Second Input"),
+	FromFirstOutputPin UMETA(Tooltip = "One thread per pin data element."),
 	Fixed UMETA(DisplayName = "Fixed Thread Count"),
-
-	// TODO: Could support inlined expression evaluation for dispatch thread count.
+	FromProductOfInputPins UMETA(Tooltip = "Dispatches a thread per element in the product of one or more pins. So if there are 4 data elements in pin A and 6 data elements in pin B, 24 threads will be dispatched."),
 };
 
 /** Produces a HLSL compute shader which will be executed on the GPU. */
@@ -48,7 +84,7 @@ public:
 
 	//~Begin UPCGSettings interface
 	virtual TArray<FPCGPinProperties> InputPinProperties() const override { return InputPins; }
-	virtual TArray<FPCGPinProperties> OutputPinProperties() const override { return OutputPins; }
+	virtual TArray<FPCGPinProperties> OutputPinProperties() const override;
 	virtual bool IsInputPinRequiredByExecution(const UPCGPin* InPin) const { return true; }
 #if WITH_EDITOR
 	virtual bool DisplayExecuteOnGPUSetting() const override { return false; }
@@ -70,6 +106,10 @@ protected:
 	virtual FPCGElementPtr CreateElement() const override;
 	//~End UPCGSettings interface
 
+protected:
+	/** Gets the GPU pin properties for the output pin with the given label. */
+	const FPCGPinPropertiesGPU* GetOutputPinPropertiesGPU(const FName& InPinLabel) const;
+
 #if WITH_EDITOR
 	void UpdateDeclarations();
 	void UpdatePinSettings();
@@ -85,8 +125,9 @@ public:
 	int GetPointCount() const { return PointCount; }
 	int GetFixedThreadCount() const { return FixedThreadCount; }
 
+	const UPCGPin* GetInputPin(FName Label) const;
+	const UPCGPin* GetOutputPin(FName Label) const;
 	const UPCGPin* GetFirstInputPin() const;
-	const UPCGPin* GetSecondInputPin() const;
 	const UPCGPin* GetPointProcessingInputPin() const;
 	const UPCGPin* GetSecondPointProcessingInputPin() const;
 	const UPCGPin* GetFirstOutputPin() const;
@@ -102,15 +143,18 @@ protected:
 	UPROPERTY(BlueprintReadWrite, EditAnywhere, Category = "Settings", meta = (EditCondition = "KernelType == EPCGKernelType::PointGenerator", EditConditionHides))
 	int PointCount = 256;
 
-	UPROPERTY(BlueprintReadWrite, EditAnywhere, Category = "Settings", meta = (EditCondition = "KernelType == EPCGKernelType::Custom", EditConditionHides))
-	EPCGDispatchThreadCount DispatchThreadCount = EPCGDispatchThreadCount::FromOutput;
+	UPROPERTY(BlueprintReadWrite, EditAnywhere, Category = "Thread Count", meta = (EditCondition = "KernelType == EPCGKernelType::Custom", EditConditionHides))
+	EPCGDispatchThreadCount DispatchThreadCount = EPCGDispatchThreadCount::FromFirstOutputPin;
 
-	UPROPERTY(BlueprintReadWrite, EditAnywhere, Category = "Settings", meta = (EditCondition = "KernelType == EPCGKernelType::Custom && DispatchThreadCount != EPCGDispatchThreadCount::Fixed", EditConditionHides))
+	UPROPERTY(BlueprintReadWrite, EditAnywhere, Category = "Thread Count", meta = (EditCondition = "KernelType == EPCGKernelType::Custom && DispatchThreadCount != EPCGDispatchThreadCount::Fixed", EditConditionHides))
 	int ThreadCountMultiplier = 1;
 
-	UPROPERTY(BlueprintReadWrite, EditAnywhere, Category = "Settings", meta = (EditCondition = "KernelType == EPCGKernelType::Custom && DispatchThreadCount == EPCGDispatchThreadCount::Fixed", EditConditionHides))
+	UPROPERTY(BlueprintReadWrite, EditAnywhere, Category = "Thread Count", meta = (EditCondition = "KernelType == EPCGKernelType::Custom && DispatchThreadCount == EPCGDispatchThreadCount::Fixed", EditConditionHides))
 	int FixedThreadCount = 1;
 
+	UPROPERTY(BlueprintReadWrite, EditAnywhere, DisplayName = "Input Pins", Category = "Thread Count", meta = (EditCondition = "KernelType == EPCGKernelType::Custom && DispatchThreadCount == EPCGDispatchThreadCount::FromProductOfInputPins", EditConditionHides))
+	TArray<FName> ThreadCountInputPinLabels;
+
 public:
 	/** Dump the cooked HLSL into the log after it is generated. */
 	UPROPERTY(EditAnywhere, Category = "Debug")
@@ -129,7 +173,7 @@ public:
 	TArray<FPCGPinProperties> InputPins = Super::DefaultPointInputPinProperties();
 
 	UPROPERTY(BlueprintReadWrite, EditAnywhere, Category = "Settings")
-	TArray<FPCGPinProperties> OutputPins = Super::DefaultPointOutputPinProperties();
+	TArray<FPCGPinPropertiesGPU> OutputPins = { FPCGPinPropertiesGPU(PCGPinConstants::DefaultOutputLabel, EPCGDataType::Point) };
 
 protected:
 	UPROPERTY(Transient, VisibleAnywhere, Category = "Settings|Declarations", meta = (MultiLine = true))
diff --git a/Engine/Plugins/PCG/Source/PCG/Public/PCGPin.h b/Engine/Plugins/PCG/Source/PCG/Public/PCGPin.h
--- a/Engine/Plugins/PCG/Source/PCG/Public/PCGPin.h
+++ b/Engine/Plugins/PCG/Source/PCG/Public/PCGPin.h
@@ -34,16 +34,6 @@ enum class EPCGPinStatus : uint8
 	Advanced
 };
 
-/** Method for computing the size of a pin on a GPU node. */
-UENUM()
-enum class EPCGPinBufferSizeMode : uint8
-{
-	FromFirstPin UMETA(DisplayName = "Match First Input Pin"),
-	FromSecondPin UMETA(DisplayName = "Match Second Input Pin"),
-	FirstPinXSecondPin UMETA(DisplayName = "First Input Pin X Second Input Pin"),
-	FixedElementCount,
-};
-
 USTRUCT(BlueprintType, meta=(HasNativeBreak="/Script/PCG.PCGBlueprintPinHelpers.BreakPinProperty", HasNativeMake="/Script/PCG.PCGBlueprintPinHelpers.MakePinProperty"))
 struct PCG_API FPCGPinProperties
 {
@@ -78,13 +68,6 @@ struct PCG_API FPCGPinProperties
 	FText Tooltip;
 #endif
 
-	/** Compute graphs use this to calculate the buffer size of output pins. */
-	UPROPERTY(BlueprintReadWrite, EditAnywhere, Category = Settings, meta = (EditCondition = "bDisplayBufferSizeSettings", EditConditionHides))
-	EPCGPinBufferSizeMode BufferSizeMode = EPCGPinBufferSizeMode::FromFirstPin;
-
-	UPROPERTY(BlueprintReadWrite, EditAnywhere, Category = Settings, meta = (EditCondition = "bDisplayBufferSizeSettings && BufferSizeMode == EPCGPinBufferSizeMode::FixedElementCount", EditConditionHides))
-	int FixedBufferElementCount = 4;
-
 	// Multiple connections are only possible if we support multi data.
 	bool AllowsMultipleConnections() const { return bAllowMultipleData && bAllowMultipleConnections; }
 
@@ -108,9 +91,6 @@ struct PCG_API FPCGPinProperties
 #endif
 
 #if WITH_EDITORONLY_DATA
-	UPROPERTY(Transient)
-	bool bDisplayBufferSizeSettings = true;
-
 	UPROPERTY(Transient)
 	bool bAllowEditMultipleData = true;
 	
