[PCG] GraphExecutor V2: Instead of relying on IsPostLoadThreadSafe(), add a new overridable OnOverrideSettingsDuplicatedInternal() on the UPCGSettings class to allow specific fixup to happen if needed

- Also update bIsDone flag after Execute call in case task was cancelled through execution (UTexture::GetPlatformData workaround)

#rnx

diff --git a/Engine/Plugins/PCG/Source/PCG/Private/Graph/PCGGraphExecutor.cpp b/Engine/Plugins/PCG/Source/PCG/Private/Graph/PCGGraphExecutor.cpp
--- a/Engine/Plugins/PCG/Source/PCG/Private/Graph/PCGGraphExecutor.cpp
+++ b/Engine/Plugins/PCG/Source/PCG/Private/Graph/PCGGraphExecutor.cpp
@@ -1288,7 +1288,14 @@ void FPCGGraphExecutor::PrepareForExecute(FPCGGraphTask& Task, FCachedResult*& O
 	
 	{
 		Task.Context = Task.Element->Initialize(Task.TaskInput, Task.SourceComponent, Task.Node);
-		Task.Context->InitializeSettings();
+		
+		// Skip PostLoad only when in GeV2 and not on GameThread
+		// The reason is that some nodes like SpawnActor / CreateTargetActor have TemplateActor subobjects that need PostLoad calls to function properly
+		// If a node supports being prepared outside of the GameThread it needs to make sure that it and its sub-object do not depend on PostLoad being called.
+		// This is true for most nodes as they only run some deprecation code in their PostLoad which isn't needed after duplication
+		const bool bSkipPostLoad = GetExecuteVersion() == EExecuteVersion::V2 && !IsInGameThread();
+		Task.Context->InitializeSettings(bSkipPostLoad);
+		
 		Task.Context->TaskId = Task.NodeId;
 		Task.Context->CompiledTaskId = Task.CompiledTaskId;
 		Task.Context->DependenciesCrc = DependenciesCrc;
@@ -1532,9 +1539,13 @@ bool FPCGGraphExecutor::ExecuteScheduling(double EndTime, TSharedPtr<FPCGGraphAc
 						{
 							TRACE_CPUPROFILER_EVENT_SCOPE(FPCGGraphExecutor::ExecuteAsyncTask);
 
-							const bool bIsDone = ActiveTask->bWasCancelled || ActiveTask->Element->Execute(ActiveTask->Context.Get());
+							bool bIsDone = ActiveTask->bWasCancelled || ActiveTask->Element->Execute(ActiveTask->Context.Get());
 							const bool bIsPaused = ActiveTask->Context->bIsPaused;
 
+							// It is currently possible for a call to Execute to endup cancelling the task through a Windows message pump caused by a ShowDialog (when calling UTexture::GetPlatformData() from a task)
+							// If this happens mark the task as done
+							bIsDone |= ActiveTask->bWasCancelled;
+
 							PostTaskExecute(ActiveTask, bIsDone);
 
 							if (bIsDone || bIsPaused)
@@ -1614,7 +1625,7 @@ void FPCGGraphExecutor::ExecuteV2()
 	bool bContinueExecute = GetNonScheduledRemainingTaskCount() > 0;
 	bool bFirstLoop = true;
 
-	const double StartTime = FPlatformTime::Seconds();
+	const double StartTime = FPlatformTime::Seconds(); 
 	const double EndTime = StartTime + GetTickBudgetInSeconds();
 
 	while ((bFirstLoop || (FPlatformTime::Seconds() < EndTime)) && bContinueExecute)
@@ -1627,10 +1638,14 @@ void FPCGGraphExecutor::ExecuteV2()
 		{
 			TRACE_CPUPROFILER_EVENT_SCOPE(FPCGGraphExecutor::Execute::ExecuteMainThreadTask);
 #if WITH_EDITOR
-			const bool bIsDone = (MainThreadTask->bIsBypassed || MainThreadTask->bWasCancelled || MainThreadTask->Element->Execute(MainThreadTask->Context.Get()));
+			bool bIsDone = (MainThreadTask->bIsBypassed || MainThreadTask->bWasCancelled || MainThreadTask->Element->Execute(MainThreadTask->Context.Get()));
 #else
-			const bool bIsDone = (MainThreadTask->bWasCancelled || MainThreadTask->Element->Execute(MainThreadTask->Context.Get()));
+			bool bIsDone = (MainThreadTask->bWasCancelled || MainThreadTask->Element->Execute(MainThreadTask->Context.Get()));
 #endif
+			// It is currently possible for a call to Execute to endup cancelling the task through a Windows message pump caused by a ShowDialog (when calling UTexture::GetPlatformData() from a task)
+			// If this happens mark the task as done
+			bIsDone |= MainThreadTask->bWasCancelled;
+
 			PostTaskExecute(MainThreadTask, bIsDone);
 		}
 	}
diff --git a/Engine/Plugins/PCG/Source/PCG/Private/PCGContext.cpp b/Engine/Plugins/PCG/Source/PCG/Private/PCGContext.cpp
--- a/Engine/Plugins/PCG/Source/PCG/Private/PCGContext.cpp
+++ b/Engine/Plugins/PCG/Source/PCG/Private/PCGContext.cpp
@@ -96,7 +96,7 @@ const UPCGSettingsInterface* FPCGContext::GetInputSettingsInterface() const
 	}
 }
 
-void FPCGContext::InitializeSettings()
+void FPCGContext::InitializeSettings(bool bSkipPostLoad)
 {
 	TRACE_CPUPROFILER_EVENT_SCOPE(FPCGContext::InitializeSettings);
 
@@ -129,7 +129,7 @@ void FPCGContext::InitializeSettings()
 			{				
 				FObjectDuplicationParameters DuplicateParams(const_cast<UPCGSettings*>(NodeSettings), GetTransientPackage());
 				const bool bIsInGameThread = IsInGameThread();
-				DuplicateParams.bSkipPostLoad = !bIsInGameThread;
+				DuplicateParams.bSkipPostLoad = bSkipPostLoad;
 				DuplicateParams.ApplyFlags = RF_Transient;
 			
 				TMap<UObject*, UObject*> CreatedObjects;
@@ -140,18 +140,21 @@ void FPCGContext::InitializeSettings()
 					SettingsWithOverride = Cast<UPCGSettings>(StaticDuplicateObjectEx(DuplicateParams));
 				}
 			
-				// StaticDuplicateObjectEx will not allow PostLoad to be called outside of GameThread (except AsyncLoadingThread) so we take care of the PostLoad and clear Async flags
 				for (auto& KeyValuePair : CreatedObjects)
 				{
 					if (KeyValuePair.Value)
 					{
 						if (!bIsInGameThread)
 						{
-							ensure(KeyValuePair.Value->IsPostLoadThreadSafe());
-							KeyValuePair.Value->ConditionalPostLoad();
+							// Outside of GameThread we need to clear the Async flags on newly duplicated objects
 							KeyValuePair.Value->ClearInternalFlags(EInternalObjectFlags::Async);
 						}
 
+						if (bSkipPostLoad)
+						{
+							// We are not calling PostLoad so remove the NeedPostLoad flags here
+							KeyValuePair.Value->ClearFlags(EObjectFlags::RF_NeedPostLoad | EObjectFlags::RF_NeedPostLoadSubobjects);
+						}
 #if WITH_EDITOR
 						// @todo_pcg: find a way to avoid the call to SetupCallbacks() all together but for now unregister the callbacks after duplication
 						if (UPCGGraphInstance* GraphInstance = Cast<UPCGGraphInstance>(KeyValuePair.Value))
@@ -165,6 +168,9 @@ void FPCGContext::InitializeSettings()
 				// Force seed copy to prevent issue due to delta serialization vs. Seed being initialized in the constructor only for new nodes
 				SettingsWithOverride->Seed = NodeSettings->Seed;
 				SettingsWithOverride->OriginalSettings = NodeSettings;
+
+				// If anything needs to be done by the Settings object after its been duplicated for override it should be done in here, outside of the gamethread that might include code that would be normally done in PostLoad
+				SettingsWithOverride->OnOverrideSettingsDuplicated(bSkipPostLoad);
 			}
 		}
 	}
diff --git a/Engine/Plugins/PCG/Source/PCG/Private/PCGSettings.cpp b/Engine/Plugins/PCG/Source/PCG/Private/PCGSettings.cpp
--- a/Engine/Plugins/PCG/Source/PCG/Private/PCGSettings.cpp
+++ b/Engine/Plugins/PCG/Source/PCG/Private/PCGSettings.cpp
@@ -576,6 +576,11 @@ TArray<FPCGPinProperties> UPCGSettings::DefaultOutputPinProperties() const
 	return OutputPinProperties();
 }
 
+void UPCGSettings::OnOverrideSettingsDuplicated(bool bSkippedPostLoad)
+{
+	OnOverrideSettingsDuplicatedInternal(bSkippedPostLoad);
+}
+
 FPCGElementPtr UPCGSettings::GetElement() const
 {
 	if (!CachedElement)
diff --git a/Engine/Plugins/PCG/Source/PCG/Public/PCGContext.h b/Engine/Plugins/PCG/Source/PCG/Public/PCGContext.h
--- a/Engine/Plugins/PCG/Source/PCG/Public/PCGContext.h
+++ b/Engine/Plugins/PCG/Source/PCG/Public/PCGContext.h
@@ -80,7 +80,7 @@ struct PCG_API FPCGContext
 	
 	// After initializing the context, we can call this method to prepare for parameter override
 	// It will create a copy of the original settings if there is indeed a possible override.
-	void InitializeSettings();
+	void InitializeSettings(bool bSkipPostLoad = false);
 
 	// If we any any parameter override, it will read from the params and override matching values
 	// in the settings copy.
diff --git a/Engine/Plugins/PCG/Source/PCG/Public/PCGSettings.h b/Engine/Plugins/PCG/Source/PCG/Public/PCGSettings.h
--- a/Engine/Plugins/PCG/Source/PCG/Public/PCGSettings.h
+++ b/Engine/Plugins/PCG/Source/PCG/Public/PCGSettings.h
@@ -234,6 +234,8 @@ public:
 #endif // WITH_EDITOR
 	//~End UObject interface
 
+	void OnOverrideSettingsDuplicated(bool bSkippedPostLoad);
+
 	// TODO: check if we need this to be virtual, we don't really need if we're always caching
 	/*virtual*/ FPCGElementPtr GetElement() const;
 	virtual UPCGNode* CreateNode() const;
@@ -434,6 +436,9 @@ protected:
 	/** Can be overriden by child class if they ever got renamed to avoid changing the default seed for this one. Otherwise default is hash of the class name. */
 	virtual uint32 GetTypeNameHash() const;
 
+	/** Can be overriden by child class if some fixup code needs to run after duplication in the context of FPCGContext::InitializeSettings */
+	virtual void OnOverrideSettingsDuplicatedInternal(bool bSkippedPostLoad) {};
+
 #if WITH_EDITOR
 	virtual void PostEditChangeProperty(FPropertyChangedEvent& PropertyChangedEvent) override;
 
