[PCG] Remove GPU pin properties native make/break for now

#rnx

diff --git a/Engine/Plugins/PCG/Source/PCG/Public/Compute/Elements/PCGCustomHLSL.h b/Engine/Plugins/PCG/Source/PCG/Public/Compute/Elements/PCGCustomHLSL.h
--- a/Engine/Plugins/PCG/Source/PCG/Public/Compute/Elements/PCGCustomHLSL.h
+++ b/Engine/Plugins/PCG/Source/PCG/Public/Compute/Elements/PCGCustomHLSL.h
@@ -18,7 +18,7 @@ enum class EPCGPinBufferSizeMode : uint8
 };
 
 /** An extension of the pin properties that adds hints for GPU thread count / buffer size calculations. */
-USTRUCT(BlueprintType, meta = (HasNativeBreak = "/Script/PCG.PCGBlueprintPinHelpers.BreakPinProperty", HasNativeMake = "/Script/PCG.PCGBlueprintPinHelpers.MakePinProperty"))
+USTRUCT(BlueprintType)
 struct PCG_API FPCGPinPropertiesGPU : public FPCGPinProperties
 {
 	GENERATED_BODY()
