[PCG] Override OnOverrideSettingsDuplicatedInternal to call CachePinLabels

#rnx

diff --git a/Engine/Plugins/PCG/Source/PCG/Private/Elements/ControlFlow/PCGMultiSelect.cpp b/Engine/Plugins/PCG/Source/PCG/Private/Elements/ControlFlow/PCGMultiSelect.cpp
--- a/Engine/Plugins/PCG/Source/PCG/Private/Elements/ControlFlow/PCGMultiSelect.cpp
+++ b/Engine/Plugins/PCG/Source/PCG/Private/Elements/ControlFlow/PCGMultiSelect.cpp
@@ -54,6 +54,15 @@ void UPCGMultiSelectSettings::PostLoad()
 #endif // WITH_EDITOR
 }
 
+void UPCGMultiSelectSettings::OnOverrideSettingsDuplicatedInternal(bool bSkippedPostLoad)
+{
+	Super::OnOverrideSettingsDuplicatedInternal(bSkippedPostLoad);
+	if (bSkippedPostLoad)
+	{
+		CachePinLabels();
+	}
+}
+
 #if WITH_EDITOR
 void UPCGMultiSelectSettings::PostEditChangeProperty(FPropertyChangedEvent& PropertyChangedEvent)
 {
diff --git a/Engine/Plugins/PCG/Source/PCG/Private/Elements/ControlFlow/PCGSwitch.cpp b/Engine/Plugins/PCG/Source/PCG/Private/Elements/ControlFlow/PCGSwitch.cpp
--- a/Engine/Plugins/PCG/Source/PCG/Private/Elements/ControlFlow/PCGSwitch.cpp
+++ b/Engine/Plugins/PCG/Source/PCG/Private/Elements/ControlFlow/PCGSwitch.cpp
@@ -46,6 +46,15 @@ void UPCGSwitchSettings::PostLoad()
 #endif // WITH_EDITOR
 }
 
+void UPCGSwitchSettings::OnOverrideSettingsDuplicatedInternal(bool bSkippedPostLoad)
+{
+	Super::OnOverrideSettingsDuplicatedInternal(bSkippedPostLoad);
+	if (bSkippedPostLoad)
+	{
+		CachePinLabels();
+	}
+}
+
 #if WITH_EDITOR
 void UPCGSwitchSettings::PostEditChangeProperty(FPropertyChangedEvent& PropertyChangedEvent)
 {
diff --git a/Engine/Plugins/PCG/Source/PCG/Public/Elements/ControlFlow/PCGMultiSelect.h b/Engine/Plugins/PCG/Source/PCG/Public/Elements/ControlFlow/PCGMultiSelect.h
--- a/Engine/Plugins/PCG/Source/PCG/Public/Elements/ControlFlow/PCGMultiSelect.h
+++ b/Engine/Plugins/PCG/Source/PCG/Public/Elements/ControlFlow/PCGMultiSelect.h
@@ -39,6 +39,7 @@ protected:
 	virtual TArray<FPCGPinProperties> InputPinProperties() const override;
 	virtual TArray<FPCGPinProperties> OutputPinProperties() const override;
 	virtual FPCGElementPtr CreateElement() const override;
+	virtual void OnOverrideSettingsDuplicatedInternal(bool bSkippedPostLoad) override;
 	//~End UPCGSettings interface
 
 public:
diff --git a/Engine/Plugins/PCG/Source/PCG/Public/Elements/ControlFlow/PCGSwitch.h b/Engine/Plugins/PCG/Source/PCG/Public/Elements/ControlFlow/PCGSwitch.h
--- a/Engine/Plugins/PCG/Source/PCG/Public/Elements/ControlFlow/PCGSwitch.h
+++ b/Engine/Plugins/PCG/Source/PCG/Public/Elements/ControlFlow/PCGSwitch.h
@@ -44,6 +44,7 @@ protected:
 	virtual TArray<FPCGPinProperties> InputPinProperties() const override;
 	virtual TArray<FPCGPinProperties> OutputPinProperties() const override;
 	virtual FPCGElementPtr CreateElement() const override;
+	virtual void OnOverrideSettingsDuplicatedInternal(bool bSkippedPostLoad) override;
 	//~End UPCGSettings interface
 
 	/** Switch is dynamic if selection value is overridden / cannot be determined prior to execution. */
