[PCG] Fix duplicated localisation key

#rnx

diff --git a/Engine/Plugins/PCG/Source/PCG/Private/Compute/Elements/PCGCustomHLSL.cpp b/Engine/Plugins/PCG/Source/PCG/Private/Compute/Elements/PCGCustomHLSL.cpp
--- a/Engine/Plugins/PCG/Source/PCG/Private/Compute/Elements/PCGCustomHLSL.cpp
+++ b/Engine/Plugins/PCG/Source/PCG/Private/Compute/Elements/PCGCustomHLSL.cpp
@@ -900,7 +900,7 @@ bool UPCGCustomHLSLSettings::IsKernelValid(FPCGContext* InContext) const
 					if (InContext)
 					{
 						PCGE_LOG_C(Error, GraphAndLog, InContext, FText::Format(
-							LOCTEXT("InvalidBufferSizeNoInputPins", "GPU buffer size for pin '{0}' could not be computed as input pins are specified in the pin settings."),
+							LOCTEXT("InvalidBufferSizeNoBufferPins", "GPU buffer size for pin '{0}' could not be computed as input pins are specified in the pin settings."),
 							FText::FromName(Properties.Label)));
 					}
 
