[PCG] More options when trying to generate at runtime in BP

* Add Flush Cache option on Generate
* Add Flush Cache function in BP
* Add Refresh Runtime Component in BP

#rnx

diff --git a/Engine/Plugins/PCG/Source/PCG/Private/Helpers/PCGBlueprintHelpers.cpp b/Engine/Plugins/PCG/Source/PCG/Private/Helpers/PCGBlueprintHelpers.cpp
--- a/Engine/Plugins/PCG/Source/PCG/Private/Helpers/PCGBlueprintHelpers.cpp
+++ b/Engine/Plugins/PCG/Source/PCG/Private/Helpers/PCGBlueprintHelpers.cpp
@@ -244,4 +244,30 @@ TArray<FPCGLandscapeLayerWeight> UPCGBlueprintHelpers::GetInterpolatedPCGLandsca
 int64 UPCGBlueprintHelpers::GetTaskId(FPCGContext& Context)
 {
 	return static_cast<int64>(Context.TaskId);
+}
+
+bool UPCGBlueprintHelpers::FlushPCGCache()
+{
+	if (UPCGSubsystem* PCGSubsystem = UPCGSubsystem::GetSubsystemForCurrentWorld())
+	{
+		PCGSubsystem->FlushCache();
+		return true;
+	}
+	else
+	{
+		return false;
+	}
+}
+
+void UPCGBlueprintHelpers::RefreshPCGRuntimeComponent(UPCGComponent* InComponent, const bool bFlushCache)
+{
+	if (UPCGSubsystem* PCGSubsystem = UPCGSubsystem::GetSubsystemForCurrentWorld())
+	{
+		if (bFlushCache)
+		{
+			PCGSubsystem->FlushCache();
+		}
+
+		PCGSubsystem->RefreshRuntimeGenComponent(InComponent, EPCGChangeType::GenerationGrid);
+	}
 }
\ No newline at end of file
diff --git a/Engine/Plugins/PCG/Source/PCG/Public/Helpers/PCGBlueprintHelpers.h b/Engine/Plugins/PCG/Source/PCG/Public/Helpers/PCGBlueprintHelpers.h
--- a/Engine/Plugins/PCG/Source/PCG/Public/Helpers/PCGBlueprintHelpers.h
+++ b/Engine/Plugins/PCG/Source/PCG/Public/Helpers/PCGBlueprintHelpers.h
@@ -84,6 +84,14 @@ public:
 
 	UFUNCTION(BLueprintCallable, Category = "PCG|Helpers", meta = (ScriptMethod))
 	static int64 GetTaskId(UPARAM(ref) FPCGContext& Context);
+
+	/** Flush the cache, to be used if you have changed something PCG depends on at runtime. Same as `pcg.FlushCache` command. Returns true if it succeeded. */
+	UFUNCTION(BlueprintCallable, Category = "PCG", meta=(DisplayName = "Flush PCG Cache"))
+	static bool FlushPCGCache();
+
+	/** Refresh a component set to Generate At Runtime, if some parameters changed. Can also flush the cache. */
+	UFUNCTION(BlueprintCallable, Category = "PCG|Runtime", meta = (ScriptMethod, DisplayName = "Refresh PCG Runtime Component"))
+	static void RefreshPCGRuntimeComponent(UPCGComponent* InComponent, const bool bFlushCache = false);
 };
 
 #if UE_ENABLE_INCLUDE_ORDER_DEPRECATED_IN_5_2
