[PCG] Make Random Choice support Attribute Sets

#rnx

diff --git a/Engine/Plugins/PCG/Config/BasePCG.ini b/Engine/Plugins/PCG/Config/BasePCG.ini
--- a/Engine/Plugins/PCG/Config/BasePCG.ini
+++ b/Engine/Plugins/PCG/Config/BasePCG.ini
@@ -70,6 +70,7 @@
 +PropertyRedirects=(OldName="PCGAttributePropertySelectorBlueprintHelpers.CopyAndFixSource.InSelector",NewName="PCGAttributePropertySelectorBlueprintHelpers.CopyAndFixSource.InputSelector")
 +PropertyRedirects=(OldName="PCGAttributeTransferSettings.SourceAttributeProperty",NewName="PCGAttributeTransferSettings.InputSource")
 +PropertyRedirects=(OldName="PCGAttributeTransferSettings.TargetAttributeProperty",NewName="PCGAttributeTransferSettings.OutputTarget")
++PropertyRedirects=(OldName="PCGRandomChoiceSettings.bOutputDiscardedPoints",NewName="PCGRandomChoiceSettings.bOutputDiscardedEntries")
 
 +ClassRedirects=(OldName="PCGPointSamplerSettings",NewName="/Script/PCG.PCGSelectPointsSettings")
 +ClassRedirects=(OldName="PCGDensityNoiseSettings",NewName="/Script/PCG.PCGAttributeNoiseSettings")
diff --git a/Engine/Plugins/PCG/Source/PCG/Private/Elements/PCGRandomChoice.cpp b/Engine/Plugins/PCG/Source/PCG/Private/Elements/PCGRandomChoice.cpp
--- a/Engine/Plugins/PCG/Source/PCG/Private/Elements/PCGRandomChoice.cpp
+++ b/Engine/Plugins/PCG/Source/PCG/Private/Elements/PCGRandomChoice.cpp
@@ -10,6 +10,49 @@
 
 #define LOCTEXT_NAMESPACE "PCGRandomChoiceElement"
 
+namespace PCGRandomChoice
+{
+	UPCGData* ChoosePointData(const UPCGData* InData, TArrayView<int32> InIndexes, FPCGContext* InContext)
+	{
+		// We know it is a point data
+		const UPCGPointData* InPointData = CastChecked<const UPCGPointData>(InData);
+		UPCGPointData* OutPointData = FPCGContext::NewObject_AnyThread<UPCGPointData>(InContext);
+		OutPointData->InitializeFromData(InPointData);
+
+		const TArray<FPCGPoint>& InPoints = InPointData->GetPoints();
+		TArray<FPCGPoint>& OutPoints = OutPointData->GetMutablePoints();
+		OutPoints.Reserve(InIndexes.Num());
+
+		// Order needs to be stable to sort this part of the array
+		Algo::Sort(InIndexes);
+
+		for (int32 Index : InIndexes)
+		{
+			OutPoints.Add(InPoints[Index]);
+		}
+
+		return OutPointData;
+	}
+
+	UPCGData* ChooseParamData(const UPCGData* InData, TArrayView<int32> InIndexes, FPCGContext* InContext)
+	{
+		// We know it is a param data
+		const UPCGParamData* InParamData = CastChecked<const UPCGParamData>(InData);
+
+		UPCGParamData* OutParamData = FPCGContext::NewObject_AnyThread<UPCGParamData>(InContext);
+
+		// Order needs to be stable to sort this part of the array
+		Algo::Sort(InIndexes);
+
+		TArray<PCGMetadataEntryKey> SelectedEntries;
+		Algo::Transform(InIndexes, SelectedEntries, [](int32 In) -> PCGMetadataEntryKey { return In; });
+
+		OutParamData->Metadata->InitializeAsCopy(InParamData->Metadata, &SelectedEntries);
+
+		return OutParamData;
+	}
+}
+
 UPCGRandomChoiceSettings::UPCGRandomChoiceSettings()
 {
 	bUseSeed = true;
@@ -28,38 +71,27 @@ FText UPCGRandomChoiceSettings::GetDefaultNodeTitle() const
 
 FText UPCGRandomChoiceSettings::GetNodeTooltipText() const
 {
-	return LOCTEXT("NodeTooltip", "Chooses points randomly through ratio or a fixed number of points.\n"
-		"Chosen/Discarded Points will be in the same order than they appear in the input data.");
+	return LOCTEXT("NodeTooltip", "Chooses entries randomly through ratio or a fixed number of entries.\n"
+		"Chosen/Discarded entries will be in the same order than they appear in the input data.");
 }
 #endif
 
-EPCGDataType UPCGRandomChoiceSettings::GetCurrentPinTypes(const UPCGPin* InPin) const
-{
-	// We have dynamic pins to true because we can change the number of output pins.
-	// But the output type is fully defined (point) so just return the pin allowed type.
-	check(InPin);
-	return InPin->Properties.AllowedTypes;
-}
-
 TArray<FPCGPinProperties> UPCGRandomChoiceSettings::InputPinProperties() const
 {
-	return Super::DefaultPointInputPinProperties();
+	TArray<FPCGPinProperties> PinProperties;
+	PinProperties.Emplace_GetRef(PCGPinConstants::DefaultInputLabel, EPCGDataType::PointOrParam).SetRequiredPin();
+
+	return PinProperties;
 }
 
 TArray<FPCGPinProperties> UPCGRandomChoiceSettings::OutputPinProperties() const
 {
 	TArray<FPCGPinProperties> PinProperties;
-	PinProperties.Emplace(PCGRandomChoiceConstants::ChosenPointsLabel,
-		EPCGDataType::Point,
-		/*bAllowMultipleConnections=*/true,
-		/*bAllowMultipleData=*/true);
+	PinProperties.Emplace(PCGRandomChoiceConstants::ChosenEntriesLabel, EPCGDataType::PointOrParam);
 
-	if (bOutputDiscardedPoints)
+	if (bOutputDiscardedEntries)
 	{
-		PinProperties.Emplace(PCGRandomChoiceConstants::DiscardedPointsLabel,
-			EPCGDataType::Point,
-			/*bAllowMultipleConnections=*/true,
-			/*bAllowMultipleData=*/true);
+		PinProperties.Emplace(PCGRandomChoiceConstants::DiscardedEntriesLabel, EPCGDataType::PointOrParam);
 	}
 
 	return PinProperties;
@@ -86,62 +118,72 @@ bool FPCGRandomChoiceElement::ExecuteInternal(FPCGContext* Context) const
 	{
 		const FPCGTaggedData& CurrentInput = Inputs[i];
 
-		const UPCGPointData* InputPointData = Cast<UPCGPointData>(CurrentInput.Data);
-		if (!InputPointData)
+		int NumElements = 0;
+		int32 Seed = Context->GetSeed();
+
+		UPCGData* (*ChooseFunc)(const UPCGData*, TArrayView<int32>, FPCGContext*) = nullptr;
+
+		if (const UPCGPointData* InputPointData = Cast<UPCGPointData>(CurrentInput.Data))
+		{
+			const TArray<FPCGPoint>& InPoints = InputPointData->GetPoints();
+			NumElements = InPoints.Num();
+			ChooseFunc = &PCGRandomChoice::ChoosePointData;
+
+			// Combine the seed with the first point so that multiple data produces different results.
+			if (!InPoints.IsEmpty())
+			{
+				Seed = PCGHelpers::ComputeSeed(Seed, InPoints[0].Seed);
+			}
+		}
+		else if (const UPCGParamData* InputParamData = Cast<UPCGParamData>(CurrentInput.Data))
 		{
-			PCGE_LOG(Verbose, GraphAndLog, FText::Format(LOCTEXT("InvalidPointData", "Input {0} is not point data"), i));
+			NumElements = InputParamData->Metadata ? InputParamData->Metadata->GetLocalItemCount() : 0;
+			ChooseFunc = &PCGRandomChoice::ChooseParamData;
+		}
+		else
+		{
+			PCGE_LOG(Verbose, GraphAndLog, FText::Format(LOCTEXT("InvalidData", "Input {0} is not a supported data"), i));
 			continue;
 		}
 
-		const TArray<FPCGPoint>& InPoints = InputPointData->GetPoints();
+		check(ChooseFunc);
 
 		int NumOfElementsToKeep = 0;
 
 		if (Settings->bFixedMode)
 		{
-			NumOfElementsToKeep = FMath::Clamp(Settings->FixedNumber, 0, InPoints.Num());
+			NumOfElementsToKeep = FMath::Clamp(Settings->FixedNumber, 0, NumElements);
 		}
 		else
 		{
-			NumOfElementsToKeep = FMath::CeilToInt(InPoints.Num() * FMath::Clamp(Settings->Ratio, 0, 1));
+			NumOfElementsToKeep = FMath::CeilToInt(NumElements * FMath::Clamp(Settings->Ratio, 0, 1));
 		}
 
-		auto CreateData = [InputPointData, &Outputs, &CurrentInput, Context](const int NumPoints, const FName PinLabel) -> TArray<FPCGPoint>*
-		{
-			FPCGTaggedData& Output = Outputs.Add_GetRef(CurrentInput);
-			UPCGPointData* OutPointData = FPCGContext::NewObject_AnyThread<UPCGPointData>(Context);
-			OutPointData->InitializeFromData(InputPointData);
-
-			TArray<FPCGPoint>& Points = OutPointData->GetMutablePoints();
-			Points.Reserve(NumPoints);
-
-			Output.Data = OutPointData;
-			Output.Pin = PinLabel;
-
-			return &Points;
-		};
-
 		if (NumOfElementsToKeep == 0)
 		{
-			// We keep no points, forward the input to the Discarded Points, and create an empty point data on chosen for parity.
-			if (Settings->bOutputDiscardedPoints)
+			// We keep no entries, forward the input to the Discarded Entries, and create an empty data on chosen for parity.
+			if (Settings->bOutputDiscardedEntries)
 			{
 				FPCGTaggedData& DiscardedOutput = Outputs.Add_GetRef(CurrentInput);
-				DiscardedOutput.Pin = PCGRandomChoiceConstants::DiscardedPointsLabel;
+				DiscardedOutput.Pin = PCGRandomChoiceConstants::DiscardedEntriesLabel;
 			}
 
-			CreateData(0, PCGRandomChoiceConstants::ChosenPointsLabel);
+			FPCGTaggedData& ChosenOutput = Outputs.Add_GetRef(CurrentInput);
+			ChosenOutput.Data = ChooseFunc(CurrentInput.Data, {}, Context);
+			ChosenOutput.Pin = PCGRandomChoiceConstants::ChosenEntriesLabel;
 			continue;
 		}
-		else if (NumOfElementsToKeep == InPoints.Num())
+		else if (NumOfElementsToKeep == NumElements)
 		{
 			// We keep all the points, forward the input to the Chosen Points, and create an empty point data on discarded for parity.
 			FPCGTaggedData& ChosenOutput = Outputs.Add_GetRef(CurrentInput);
-			ChosenOutput.Pin = PCGRandomChoiceConstants::ChosenPointsLabel;
+			ChosenOutput.Pin = PCGRandomChoiceConstants::ChosenEntriesLabel;
 
-			if (Settings->bOutputDiscardedPoints)
+			if (Settings->bOutputDiscardedEntries)
 			{
-				CreateData(0, PCGRandomChoiceConstants::DiscardedPointsLabel);
+				FPCGTaggedData& DiscardedOutput = Outputs.Add_GetRef(CurrentInput);
+				DiscardedOutput.Data = ChooseFunc(CurrentInput.Data, {}, Context);
+				DiscardedOutput.Pin = PCGRandomChoiceConstants::DiscardedEntriesLabel;
 			}
 
 			continue;
@@ -153,19 +195,13 @@ bool FPCGRandomChoiceElement::ExecuteInternal(FPCGContext* Context) const
 		// * Pick a number in [0, n[
 		// * Pick a number in[0, n - 1[, then for each previously selected number, if it's larger, add +1
 		// It's O(n) cpu + O(n) memory vs O(s^2) cpu + O(s) memory. (n = total number of points, s = number of points to keep)
-		int32 Seed = Context->GetSeed();
-		if (ensure(!InPoints.IsEmpty()))
-		{
-			// Combine the seed with the first point so that multiple data produces different results.
-			Seed = PCGHelpers::ComputeSeed(Seed, InPoints[0].Seed);
-		}
 
 		FRandomStream RandStream(Seed);
 		TArray<int32> ShuffledIndexes;
-		ShuffledIndexes.Reserve(InPoints.Num());
+		ShuffledIndexes.Reserve(NumElements);
 
 		// Instead of swapping directly, we swap the indexes since it cheaper than copying points around
-		for (int j = 0; j < InPoints.Num(); ++j)
+		for (int j = 0; j < NumElements; ++j)
 		{
 			ShuffledIndexes.Add(j);
 		}
@@ -173,31 +209,29 @@ bool FPCGRandomChoiceElement::ExecuteInternal(FPCGContext* Context) const
 		// We only have to shuffle until we reached the number of elements to keep.
 		for (int j = 0; j < NumOfElementsToKeep; ++j)
 		{
-			const int RandomElement = RandStream.RandRange(j, InPoints.Num() - 1);
+			const int RandomElement = RandStream.RandRange(j, NumElements - 1);
 			if (RandomElement != j)
 			{
 				Swap(ShuffledIndexes[j], ShuffledIndexes[RandomElement]);
 			}
 		}
 
-		auto AddToOutput = [&CreateData, &InPoints, &ShuffledIndexes](const int StartIndex, const int Count, const FName Label)
+		auto AddToOutput = [&ShuffledIndexes, &CurrentInput, &Outputs, &ChooseFunc, Context](const int StartIndex, const int Count, const FName Label)
 		{
-			// Order needs to be stable to sort this part of the array
-			Algo::Sort(MakeArrayView(ShuffledIndexes.GetData() + StartIndex, Count));
-
-			TArray<FPCGPoint>* Points = CreateData(Count, Label);
-			check(Points);
-			for (int j = 0; j < Count; ++j)
-			{
-				Points->Add(InPoints[ShuffledIndexes[StartIndex + j]]);
-			}
+			FPCGTaggedData& ChosenOutput = Outputs.Add_GetRef(CurrentInput);
+			ChosenOutput.Data = ChooseFunc(CurrentInput.Data, MakeArrayView(ShuffledIndexes.GetData() + StartIndex, Count), Context);
+			ChosenOutput.Pin = PCGRandomChoiceConstants::ChosenEntriesLabel;
 		};
 
-		AddToOutput(0, NumOfElementsToKeep, PCGRandomChoiceConstants::ChosenPointsLabel);
+		FPCGTaggedData& ChosenOutput = Outputs.Add_GetRef(CurrentInput);
+		ChosenOutput.Data = ChooseFunc(CurrentInput.Data, MakeArrayView(ShuffledIndexes.GetData(), NumOfElementsToKeep), Context);
+		ChosenOutput.Pin = PCGRandomChoiceConstants::ChosenEntriesLabel;
 
-		if (Settings->bOutputDiscardedPoints)
+		if (Settings->bOutputDiscardedEntries)
 		{
-			AddToOutput(NumOfElementsToKeep, InPoints.Num() - NumOfElementsToKeep, PCGRandomChoiceConstants::DiscardedPointsLabel);
+			FPCGTaggedData& DiscardedOutput = Outputs.Add_GetRef(CurrentInput);
+			DiscardedOutput.Data = ChooseFunc(CurrentInput.Data, MakeArrayView(ShuffledIndexes.GetData() + NumOfElementsToKeep, NumElements - NumOfElementsToKeep), Context);
+			DiscardedOutput.Pin = PCGRandomChoiceConstants::DiscardedEntriesLabel;
 		}
 	}
 
diff --git a/Engine/Plugins/PCG/Source/PCG/Private/Tests/Elements/PCGRandomChoiceTest.cpp b/Engine/Plugins/PCG/Source/PCG/Private/Tests/Elements/PCGRandomChoiceTest.cpp
--- a/Engine/Plugins/PCG/Source/PCG/Private/Tests/Elements/PCGRandomChoiceTest.cpp
+++ b/Engine/Plugins/PCG/Source/PCG/Private/Tests/Elements/PCGRandomChoiceTest.cpp
@@ -6,7 +6,10 @@
 
 #include "PCGComponent.h"
 #include "PCGContext.h"
+#include "PCGParamData.h"
 #include "Data/PCGPointData.h"
+#include "Metadata/PCGMetadata.h"
+#include "Metadata/PCGMetadataAttributeTpl.h"
 
 #include "Elements/PCGRandomChoice.h"
 
@@ -18,6 +21,8 @@ IMPLEMENT_CUSTOM_SIMPLE_AUTOMATION_TEST(FPCGRandomChoiceTest_NoDiscard, FPCGTest
 IMPLEMENT_CUSTOM_SIMPLE_AUTOMATION_TEST(FPCGRandomChoiceTest_MultiData_SameSeed, FPCGTestBaseClass, "Plugins.PCG.RandomChoice.MultiData.SameSeed", PCGTestsCommon::TestFlags)
 IMPLEMENT_CUSTOM_SIMPLE_AUTOMATION_TEST(FPCGRandomChoiceTest_MultiData_DifferentSeed, FPCGTestBaseClass, "Plugins.PCG.RandomChoice.MultiData.DifferentSeed", PCGTestsCommon::TestFlags)
 
+IMPLEMENT_CUSTOM_SIMPLE_AUTOMATION_TEST(FPCGRandomChoiceTest_Fixed_ParamData, FPCGTestBaseClass, "Plugins.PCG.RandomChoice.FixedParamData", PCGTestsCommon::TestFlags)
+
 namespace PCGRandomChoiceTest
 {
 	UPCGPointData* CreateInputPointData(FPCGContext* Context, const int NumPoints, const int Seed = 42)
@@ -42,6 +47,30 @@ namespace PCGRandomChoiceTest
 		return NewPointData;
 	}
 
+	const FName DensityAttributeName = TEXT("MyDensity");
+
+	UPCGParamData* CreateInputParamData(FPCGContext* Context, const int NumElements)
+	{
+		check(Context);
+
+		UPCGParamData* NewParamData = NewObject<UPCGParamData>();
+		NewParamData->SetFlags(RF_Transient);
+
+		FPCGMetadataAttribute<double>* DensityAttribute = NewParamData->Metadata->CreateAttribute<double>(DensityAttributeName, 0.0, true, false);
+		check(DensityAttribute);
+
+		for (int i = 0; i < NumElements; ++i)
+		{
+			DensityAttribute->SetValue(NewParamData->Metadata->AddEntry(), static_cast<double>(i));
+		}
+
+		FPCGTaggedData& InputData = Context->InputData.TaggedData.Emplace_GetRef();
+		InputData.Data = NewParamData;
+		InputData.Pin = PCGPinConstants::DefaultInputLabel;
+
+		return NewParamData;
+	}
+
 	bool VerifyAllPointsThere(const int NumPoints, const UPCGPointData* ChosenPointData, const UPCGPointData* DiscardedPointData)
 	{
 		check(ChosenPointData && DiscardedPointData);
@@ -62,7 +91,7 @@ namespace PCGRandomChoiceTest
 				IndexesSeen.Add(Index);
 
 				// It needs to be stable so density should be ascending.
-				if (i > 0 && Points[i].Density < Points[i - 1].Density)
+				if (i > 0 && Points[i].Density <= Points[i - 1].Density)
 				{
 					return false;
 				}
@@ -83,6 +112,55 @@ namespace PCGRandomChoiceTest
 
 		return IndexesSeen.Num() == NumPoints;
 	}
+
+	bool VerifyAllEntriesAreThere(const int NumElements, const UPCGParamData* ChosenParamData, const UPCGParamData* DiscardedParamData)
+	{
+		check(ChosenParamData && DiscardedParamData);
+		TSet<int> IndexesSeen;
+
+		auto Check = [&IndexesSeen](const UPCGParamData* ParamData)
+		{
+			const FPCGMetadataAttribute<double>* DensityAttribute = ParamData->Metadata->GetConstTypedAttribute<double>(DensityAttributeName);
+			if (!DensityAttribute)
+			{
+				return false;
+			}
+
+			int PreviousIndex = -1;
+			for (int i = 0; i < ParamData->Metadata->GetLocalItemCount(); ++i)
+			{
+				const int Index = static_cast<int>(DensityAttribute->GetValueFromItemKey(PCGMetadataEntryKey(i)));
+				if (IndexesSeen.Contains(Index))
+				{
+					return false;
+				}
+
+				IndexesSeen.Add(Index);
+
+				// It needs to be stable so density should be ascending.
+				if (i > 0 && Index <= PreviousIndex)
+				{
+					return false;
+				}
+
+				PreviousIndex = Index;
+			}
+
+			return true;
+		};
+
+		if (!Check(ChosenParamData))
+		{
+			return false;
+		}
+
+		if (!Check(DiscardedParamData))
+		{
+			return false;
+		}
+
+		return IndexesSeen.Num() == NumElements;
+	}
 }
 
 bool FPCGRandomChoiceTest_Fixed::RunTest(const FString& Parameters)
@@ -104,8 +182,8 @@ bool FPCGRandomChoiceTest_Fixed::RunTest(const FString& Parameters)
 	FPCGElementPtr TestElement = TestData.Settings->GetElement();
 	while (!TestElement->Execute(Context.Get())) {}
 
-	TArray<FPCGTaggedData> ChosenOutputTagged = Context->OutputData.GetInputsByPin(PCGRandomChoiceConstants::ChosenPointsLabel);
-	TArray<FPCGTaggedData> DiscardedOutputTagged = Context->OutputData.GetInputsByPin(PCGRandomChoiceConstants::DiscardedPointsLabel);
+	TArray<FPCGTaggedData> ChosenOutputTagged = Context->OutputData.GetInputsByPin(PCGRandomChoiceConstants::ChosenEntriesLabel);
+	TArray<FPCGTaggedData> DiscardedOutputTagged = Context->OutputData.GetInputsByPin(PCGRandomChoiceConstants::DiscardedEntriesLabel);
 
 	const UPCGPointData* ChosenOutputData = ChosenOutputTagged.Num() == 1 ? Cast<const UPCGPointData>(ChosenOutputTagged[0].Data) : nullptr;
 	const UPCGPointData* DiscardedOutputData = DiscardedOutputTagged.Num() == 1 ? Cast<const UPCGPointData>(DiscardedOutputTagged[0].Data) : nullptr;
@@ -140,8 +218,8 @@ bool FPCGRandomChoiceTest_Ratio::RunTest(const FString& Parameters)
 	FPCGElementPtr TestElement = TestData.Settings->GetElement();
 	while (!TestElement->Execute(Context.Get())) {}
 
-	TArray<FPCGTaggedData> ChosenOutputTagged = Context->OutputData.GetInputsByPin(PCGRandomChoiceConstants::ChosenPointsLabel);
-	TArray<FPCGTaggedData> DiscardedOutputTagged = Context->OutputData.GetInputsByPin(PCGRandomChoiceConstants::DiscardedPointsLabel);
+	TArray<FPCGTaggedData> ChosenOutputTagged = Context->OutputData.GetInputsByPin(PCGRandomChoiceConstants::ChosenEntriesLabel);
+	TArray<FPCGTaggedData> DiscardedOutputTagged = Context->OutputData.GetInputsByPin(PCGRandomChoiceConstants::DiscardedEntriesLabel);
 
 	const UPCGPointData* ChosenOutputData = ChosenOutputTagged.Num() == 1 ? Cast<const UPCGPointData>(ChosenOutputTagged[0].Data) : nullptr;
 	const UPCGPointData* DiscardedOutputData = DiscardedOutputTagged.Num() == 1 ? Cast<const UPCGPointData>(DiscardedOutputTagged[0].Data) : nullptr;
@@ -174,8 +252,8 @@ bool FPCGRandomChoiceTest_SelectNone::RunTest(const FString& Parameters)
 	FPCGElementPtr TestElement = TestData.Settings->GetElement();
 	while (!TestElement->Execute(Context.Get())) {}
 
-	TArray<FPCGTaggedData> ChosenOutputTagged = Context->OutputData.GetInputsByPin(PCGRandomChoiceConstants::ChosenPointsLabel);
-	TArray<FPCGTaggedData> DiscardedOutputTagged = Context->OutputData.GetInputsByPin(PCGRandomChoiceConstants::DiscardedPointsLabel);
+	TArray<FPCGTaggedData> ChosenOutputTagged = Context->OutputData.GetInputsByPin(PCGRandomChoiceConstants::ChosenEntriesLabel);
+	TArray<FPCGTaggedData> DiscardedOutputTagged = Context->OutputData.GetInputsByPin(PCGRandomChoiceConstants::DiscardedEntriesLabel);
 
 	const UPCGPointData* ChosenOutputData = ChosenOutputTagged.Num() == 1 ? Cast<const UPCGPointData>(ChosenOutputTagged[0].Data) : nullptr;
 	const UPCGPointData* DiscardedOutputData = DiscardedOutputTagged.Num() == 1 ? Cast<const UPCGPointData>(DiscardedOutputTagged[0].Data) : nullptr;
@@ -206,8 +284,8 @@ bool FPCGRandomChoiceTest_SelectAll::RunTest(const FString& Parameters)
 	FPCGElementPtr TestElement = TestData.Settings->GetElement();
 	while (!TestElement->Execute(Context.Get())) {}
 
-	TArray<FPCGTaggedData> ChosenOutputTagged = Context->OutputData.GetInputsByPin(PCGRandomChoiceConstants::ChosenPointsLabel);
-	TArray<FPCGTaggedData> DiscardedOutputTagged = Context->OutputData.GetInputsByPin(PCGRandomChoiceConstants::DiscardedPointsLabel);
+	TArray<FPCGTaggedData> ChosenOutputTagged = Context->OutputData.GetInputsByPin(PCGRandomChoiceConstants::ChosenEntriesLabel);
+	TArray<FPCGTaggedData> DiscardedOutputTagged = Context->OutputData.GetInputsByPin(PCGRandomChoiceConstants::DiscardedEntriesLabel);
 
 	const UPCGPointData* ChosenOutputData = ChosenOutputTagged.Num() == 1 ? Cast<const UPCGPointData>(ChosenOutputTagged[0].Data) : nullptr;
 	const UPCGPointData* DiscardedOutputData = DiscardedOutputTagged.Num() == 1 ? Cast<const UPCGPointData>(DiscardedOutputTagged[0].Data) : nullptr;
@@ -232,7 +310,7 @@ bool FPCGRandomChoiceTest_NoDiscard::RunTest(const FString& Parameters)
 
 	Settings->bFixedMode = false;
 	Settings->Ratio = 0.1f;
-	Settings->bOutputDiscardedPoints = false;
+	Settings->bOutputDiscardedEntries = false;
 
 	TUniquePtr<FPCGContext> Context = TestData.InitializeTestContext();
 	const UPCGPointData* InputPointData = PCGRandomChoiceTest::CreateInputPointData(Context.Get(), NumOfPoints);
@@ -240,8 +318,8 @@ bool FPCGRandomChoiceTest_NoDiscard::RunTest(const FString& Parameters)
 	FPCGElementPtr TestElement = TestData.Settings->GetElement();
 	while (!TestElement->Execute(Context.Get())) {}
 
-	TArray<FPCGTaggedData> ChosenOutputTagged = Context->OutputData.GetInputsByPin(PCGRandomChoiceConstants::ChosenPointsLabel);
-	TArray<FPCGTaggedData> DiscardedOutputTagged = Context->OutputData.GetInputsByPin(PCGRandomChoiceConstants::DiscardedPointsLabel);
+	TArray<FPCGTaggedData> ChosenOutputTagged = Context->OutputData.GetInputsByPin(PCGRandomChoiceConstants::ChosenEntriesLabel);
+	TArray<FPCGTaggedData> DiscardedOutputTagged = Context->OutputData.GetInputsByPin(PCGRandomChoiceConstants::DiscardedEntriesLabel);
 
 	const UPCGPointData* ChosenOutputData = ChosenOutputTagged.Num() == 1 ? Cast<const UPCGPointData>(ChosenOutputTagged[0].Data) : nullptr;
 	const UPCGPointData* DiscardedOutputData = DiscardedOutputTagged.Num() == 1 ? Cast<const UPCGPointData>(DiscardedOutputTagged[0].Data) : nullptr;
@@ -274,8 +352,8 @@ bool FPCGRandomChoiceTest_MultiData_SameSeed::RunTest(const FString& Parameters)
 	FPCGElementPtr TestElement = TestData.Settings->GetElement();
 	while (!TestElement->Execute(Context.Get())) {}
 
-	TArray<FPCGTaggedData> ChosenOutputTagged = Context->OutputData.GetInputsByPin(PCGRandomChoiceConstants::ChosenPointsLabel);
-	TArray<FPCGTaggedData> DiscardedOutputTagged = Context->OutputData.GetInputsByPin(PCGRandomChoiceConstants::DiscardedPointsLabel);
+	TArray<FPCGTaggedData> ChosenOutputTagged = Context->OutputData.GetInputsByPin(PCGRandomChoiceConstants::ChosenEntriesLabel);
+	TArray<FPCGTaggedData> DiscardedOutputTagged = Context->OutputData.GetInputsByPin(PCGRandomChoiceConstants::DiscardedEntriesLabel);
 
 	const UPCGPointData* FirstChosenOutputData = ChosenOutputTagged.Num() == 2 ? Cast<const UPCGPointData>(ChosenOutputTagged[0].Data) : nullptr;
 	const UPCGPointData* FirstDiscardedOutputData = DiscardedOutputTagged.Num() == 2 ? Cast<const UPCGPointData>(DiscardedOutputTagged[0].Data) : nullptr;
@@ -326,8 +404,8 @@ bool FPCGRandomChoiceTest_MultiData_DifferentSeed::RunTest(const FString& Parame
 	FPCGElementPtr TestElement = TestData.Settings->GetElement();
 	while (!TestElement->Execute(Context.Get())) {}
 
-	TArray<FPCGTaggedData> ChosenOutputTagged = Context->OutputData.GetInputsByPin(PCGRandomChoiceConstants::ChosenPointsLabel);
-	TArray<FPCGTaggedData> DiscardedOutputTagged = Context->OutputData.GetInputsByPin(PCGRandomChoiceConstants::DiscardedPointsLabel);
+	TArray<FPCGTaggedData> ChosenOutputTagged = Context->OutputData.GetInputsByPin(PCGRandomChoiceConstants::ChosenEntriesLabel);
+	TArray<FPCGTaggedData> DiscardedOutputTagged = Context->OutputData.GetInputsByPin(PCGRandomChoiceConstants::DiscardedEntriesLabel);
 
 	const UPCGPointData* FirstChosenOutputData = ChosenOutputTagged.Num() == 2 ? Cast<const UPCGPointData>(ChosenOutputTagged[0].Data) : nullptr;
 	const UPCGPointData* FirstDiscardedOutputData = DiscardedOutputTagged.Num() == 2 ? Cast<const UPCGPointData>(DiscardedOutputTagged[0].Data) : nullptr;
@@ -360,4 +438,41 @@ bool FPCGRandomChoiceTest_MultiData_DifferentSeed::RunTest(const FString& Parame
 
 	return true;
 }
+
+bool FPCGRandomChoiceTest_Fixed_ParamData::RunTest(const FString& Parameters)
+{
+	PCGTestsCommon::FTestData TestData;
+	UPCGRandomChoiceSettings* Settings = PCGTestsCommon::GenerateSettings<UPCGRandomChoiceSettings>(TestData);
+	check(Settings);
+
+	constexpr int NumOfElements = 20;
+	constexpr int ExpectedNumElementsChosen = 7;
+	constexpr int ExpectedNumElementsDiscarded = 13;
+
+	Settings->bFixedMode = true;
+	Settings->FixedNumber = ExpectedNumElementsChosen;
+
+	TUniquePtr<FPCGContext> Context = TestData.InitializeTestContext();
+	const UPCGParamData* InputParamData = PCGRandomChoiceTest::CreateInputParamData(Context.Get(), NumOfElements);
+
+	FPCGElementPtr TestElement = TestData.Settings->GetElement();
+	while (!TestElement->Execute(Context.Get())) {}
+
+	TArray<FPCGTaggedData> ChosenOutputTagged = Context->OutputData.GetInputsByPin(PCGRandomChoiceConstants::ChosenEntriesLabel);
+	TArray<FPCGTaggedData> DiscardedOutputTagged = Context->OutputData.GetInputsByPin(PCGRandomChoiceConstants::DiscardedEntriesLabel);
+
+	const UPCGParamData* ChosenOutputData = ChosenOutputTagged.Num() == 1 ? Cast<const UPCGParamData>(ChosenOutputTagged[0].Data) : nullptr;
+	const UPCGParamData* DiscardedOutputData = DiscardedOutputTagged.Num() == 1 ? Cast<const UPCGParamData>(DiscardedOutputTagged[0].Data) : nullptr;
+
+	UTEST_NOT_NULL("There is a param data in chosen entries", ChosenOutputData);
+	UTEST_NOT_NULL("There is a param data in discarded entries", DiscardedOutputData);
+
+	UTEST_EQUAL("There is the right number of entries in chosen", ChosenOutputData->Metadata->GetLocalItemCount(), ExpectedNumElementsChosen);
+	UTEST_EQUAL("There is the right number of entries in discarded", DiscardedOutputData->Metadata->GetLocalItemCount(), ExpectedNumElementsDiscarded);
+
+	UTEST_TRUE("All entries are there and in the right order", PCGRandomChoiceTest::VerifyAllEntriesAreThere(NumOfElements, ChosenOutputData, DiscardedOutputData));
+
+	return true;
+}
+
 #endif // WITH_EDITOR
diff --git a/Engine/Plugins/PCG/Source/PCG/Public/Elements/PCGRandomChoice.h b/Engine/Plugins/PCG/Source/PCG/Public/Elements/PCGRandomChoice.h
--- a/Engine/Plugins/PCG/Source/PCG/Public/Elements/PCGRandomChoice.h
+++ b/Engine/Plugins/PCG/Source/PCG/Public/Elements/PCGRandomChoice.h
@@ -8,13 +8,13 @@
 
 namespace PCGRandomChoiceConstants
 {
-	const FName ChosenPointsLabel = TEXT("Chosen");
-	const FName DiscardedPointsLabel = TEXT("Discarded");
+	const FName ChosenEntriesLabel = TEXT("Chosen");
+	const FName DiscardedEntriesLabel = TEXT("Discarded");
 }
 
 /**
-* Chooses points randomly through ratio or a fixed number of points.
-* Chosen/Discarded Points will be in the same order than they appear in the input data.
+* Chooses entries randomly through ratio or a fixed number of entries.
+* Chosen/Discarded entries will be in the same order than they appear in the input data.
 */
 UCLASS(BlueprintType)
 class UPCGRandomChoiceSettings : public UPCGSettings
@@ -33,7 +33,6 @@ public:
 #endif
 
 	virtual bool HasDynamicPins() const override { return true; }
-	virtual EPCGDataType GetCurrentPinTypes(const UPCGPin* InPin) const override;
 
 protected:
 	virtual TArray<FPCGPinProperties> InputPinProperties() const override;
@@ -43,21 +42,21 @@ protected:
 
 
 public:
-	// Either choose a fixed number of points, or a ratio of points.
+	// Either choose a fixed number of entries, or a ratio of entries.
 	UPROPERTY(BlueprintReadWrite, EditAnywhere, Category = Settings, meta = (PCG_Overridable))
 	bool bFixedMode = true;
 
-	// Defines the number of points to keep.
+	// Defines the number of entries to keep.
 	UPROPERTY(BlueprintReadWrite, EditAnywhere, Category = Settings, meta = (PCG_Overridable, EditCondition = "bFixedMode", ClampMin = "0"))
 	int FixedNumber = 1;
 
-	// Defines the ratio of points to keep.
+	// Defines the ratio of entries to keep.
 	UPROPERTY(BlueprintReadWrite, EditAnywhere, Category = Settings, meta = (PCG_Overridable, EditCondition = "!bFixedMode", ClampMin = "0", ClampMax = "1"))
 	float Ratio = 0.5;
 
-	// By default, we output discarded points. If you don't need them, disable this option.
+	// By default, we output discarded entries. If you don't need them, disable this option.
 	UPROPERTY(BlueprintReadWrite, EditAnywhere, Category = Settings)
-	bool bOutputDiscardedPoints = true;
+	bool bOutputDiscardedEntries = true;
 };
 
 class FPCGRandomChoiceElement : public IPCGElement
