[PCG] Add more export to library options for Graph and Graph Instances

* Can override color and title for graph and their instances
* Can expose to library instances
* Can override category and description for instances

#rnx

diff --git a/Engine/Plugins/PCG/Source/PCG/Private/PCGGraph.cpp b/Engine/Plugins/PCG/Source/PCG/Private/PCGGraph.cpp
--- a/Engine/Plugins/PCG/Source/PCG/Private/PCGGraph.cpp
+++ b/Engine/Plugins/PCG/Source/PCG/Private/PCGGraph.cpp
@@ -1455,6 +1455,13 @@ void UPCGGraph::PostEditChangeProperty(FPropertyChangedEvent& PropertyChangedEve
 		// The higen settings change the structure of the graph (presence or absence of links between grid levels).
 		NotifyGraphChanged(EPCGChangeType::Structural | EPCGChangeType::GenerationGrid);
 	}
+	else if (PropertyName == GET_MEMBER_NAME_CHECKED(UPCGGraphInterface, TitleOverride)
+		|| PropertyName == GET_MEMBER_NAME_CHECKED(UPCGGraphInterface, bTitleOverride)
+		|| PropertyName == GET_MEMBER_NAME_CHECKED(UPCGGraphInterface, ColorOverride)
+		|| PropertyName == GET_MEMBER_NAME_CHECKED(UPCGGraphInterface, bColorOverride))
+	{
+		NotifyGraphChanged(EPCGChangeType::Cosmetic);
+	}
 
 	PreviousPropertyBag = nullptr;
 }
@@ -1778,6 +1785,11 @@ void UPCGGraphInstance::PostEditChangeProperty(FPropertyChangedEvent& PropertyCh
 	{
 		OnGraphParametersChanged(this, EPCGGraphParameterEvent::ValueModifiedLocally, PropertyChangedEvent.GetMemberPropertyName());
 	}
+	else
+	{
+		// For other changes, push a cosmetic change
+		OnGraphChangedDelegate.Broadcast(this, EPCGChangeType::Cosmetic);
+	}
 }
 
 void UPCGGraphInstance::PreEditUndo()
@@ -1907,6 +1919,30 @@ void UPCGGraphInstance::NotifyGraphParametersChanged(EPCGGraphParameterEvent InC
 	// Also propagates the changes
 	OnGraphChanged(Graph, GetChangeTypeForGraphParameterChange(InChangeType, InChangedPropertyName));
 }
+
+void UPCGGraphInstance::GetTitleOverride(bool& bOutIsOverridden, FText& OutTitleOverride) const
+{
+	if (!bTitleOverride && Graph)
+	{
+		Graph->GetTitleOverride(bOutIsOverridden, OutTitleOverride);
+	}
+	else
+	{
+		UPCGGraphInterface::GetTitleOverride(bOutIsOverridden, OutTitleOverride);
+	}
+}
+
+void UPCGGraphInstance::GetColorOverride(bool& bOutIsOverridden, FLinearColor& OutColorOverride) const
+{
+	if (!bColorOverride && Graph)
+	{
+		Graph->GetColorOverride(bOutIsOverridden, OutColorOverride);
+	}
+	else
+	{
+		UPCGGraphInterface::GetColorOverride(bOutIsOverridden, OutColorOverride);
+	}
+}
 #endif // WITH_EDITOR
 
 void UPCGGraphInstance::OnGraphParametersChanged(UPCGGraphInterface* InGraph, EPCGGraphParameterEvent InChangeType, FName InChangedPropertyName)
diff --git a/Engine/Plugins/PCG/Source/PCG/Private/PCGSubgraph.cpp b/Engine/Plugins/PCG/Source/PCG/Private/PCGSubgraph.cpp
--- a/Engine/Plugins/PCG/Source/PCG/Private/PCGSubgraph.cpp
+++ b/Engine/Plugins/PCG/Source/PCG/Private/PCGSubgraph.cpp
@@ -307,6 +307,15 @@ FString UPCGSubgraphSettings::GetAdditionalTitleInformation() const
 		// Subgraphs with the subgraph override pin connected should not display any asset path.
 		return FString();
 	}
+
+	bool bIsOverridden = false;
+	FText OverrideTitle;
+
+	SubgraphInstance->GetTitleOverride(bIsOverridden, OverrideTitle);
+	if (bIsOverridden)
+	{
+		return OverrideTitle.ToString();
+	}
 #endif
 
 	if (UPCGGraph* TargetSubgraph = GetSubgraph())
@@ -349,6 +358,22 @@ void UPCGSubgraphSettings::PostEditChangeProperty(struct FPropertyChangedEvent&
 
 	Super::PostEditChangeProperty(PropertyChangedEvent);
 }
+
+FLinearColor UPCGSubgraphSettings::GetNodeTitleColor() const
+{
+	if (!IsDynamicGraph())
+	{
+		bool bIsOverridden = false;
+		FLinearColor OverrideColor;
+		SubgraphInstance->GetColorOverride(bIsOverridden, OverrideColor);
+		if (bIsOverridden)
+		{
+			return OverrideColor;
+		}
+	}
+
+	return Super::GetNodeTitleColor();
+}
 #endif
 
 #if WITH_EDITOR
diff --git a/Engine/Plugins/PCG/Source/PCG/Public/PCGGraph.h b/Engine/Plugins/PCG/Source/PCG/Public/PCGGraph.h
--- a/Engine/Plugins/PCG/Source/PCG/Public/PCGGraph.h
+++ b/Engine/Plugins/PCG/Source/PCG/Public/PCGGraph.h
@@ -114,8 +114,30 @@ public:
 #if WITH_EDITOR
 	FOnPCGGraphChanged OnGraphChangedDelegate;
 	FOnPCGGraphParametersChanged OnGraphParametersChangedDelegate;
+
+	virtual void GetTitleOverride(bool& bOutIsOverridden, FText& OutTitleOverride) const { bOutIsOverridden = bTitleOverride; OutTitleOverride = TitleOverride; }
+	virtual void GetColorOverride(bool& bOutIsOverridden, FLinearColor& OutColorOverride) const { bOutIsOverridden = bColorOverride; OutColorOverride = ColorOverride; }
 #endif // WITH_EDITOR
 
+#if WITH_EDITORONLY_DATA
+	UPROPERTY(BlueprintReadWrite, EditAnywhere, Category = AssetInfo, AssetRegistrySearchable)
+	bool bExposeToLibrary = false;
+
+	UPROPERTY(BlueprintReadWrite, EditAnywhere, Category = AssetInfo, AssetRegistrySearchable, meta = (InlineEditConditionToggle))
+	bool bTitleOverride = false;
+
+	/** Override of the title for the subgraph node for this graph. */
+	UPROPERTY(BlueprintReadWrite, EditAnywhere, Category = AssetInfo, AssetRegistrySearchable, meta = (EditCondition = "bTitleOverride"))
+	FText TitleOverride;
+
+	UPROPERTY(BlueprintReadWrite, EditAnywhere, Category = AssetInfo, meta = (InlineEditConditionToggle))
+	bool bColorOverride = false;
+
+	/** Override of the color for the subgraph node for this graph. */
+	UPROPERTY(BlueprintReadWrite, EditAnywhere, Category = AssetInfo, meta = (EditCondition = "bColorOverride"))
+	FLinearColor ColorOverride = FLinearColor::White;
+#endif
+
 	template <typename T>
 	TValueOrError<T, EPropertyBagResult> GetGraphParameter(const FName PropertyName) const
 	{
@@ -209,9 +231,6 @@ public:
 	bool Use2DGrid() const { return bUse2DGrid; }
 
 #if WITH_EDITORONLY_DATA
-	UPROPERTY(BlueprintReadWrite, EditAnywhere, Category = AssetInfo, AssetRegistrySearchable)
-	bool bExposeToLibrary = false;
-
 	UPROPERTY(BlueprintReadWrite, EditAnywhere, Category = AssetInfo, AssetRegistrySearchable)
 	FText Category;
 
@@ -503,7 +522,7 @@ public:
 	bool IsPropertyOverridden(const FProperty* InProperty) const { return ParametersOverrides.IsPropertyOverridden(InProperty); }
 	bool IsPropertyOverriddenAndNotDefault(const FProperty* InProperty) const;
 
-	UPROPERTY(EditAnywhere, BlueprintReadOnly, Category = Instance)
+	UPROPERTY(EditAnywhere, BlueprintReadOnly, Category = Instance, AssetRegistrySearchable)
 	TObjectPtr<UPCGGraphInterface> Graph;
 
 	UPROPERTY(EditAnywhere, Category = Instance, meta = (NoResetToDefault))
@@ -519,10 +538,30 @@ public:
 	*/
 	bool CanGraphInterfaceBeSet(const UPCGGraphInterface* GraphInterface) const;
 
-private:
+#if WITH_EDITOR
+	virtual void GetTitleOverride(bool& bOutIsOverridden, FText& OutTitleOverride) const override;
+	virtual void GetColorOverride(bool& bOutIsOverridden, FLinearColor& OutColorOverride) const override;
+#endif // WITH_EDITOR
+
 #if WITH_EDITORONLY_DATA
+private:
 	// Transient, to keep track of the previous graph when it changed.
 	TWeakObjectPtr<UPCGGraphInterface> PreGraphCache = nullptr;
+
+public:
+	UPROPERTY(BlueprintReadWrite, EditAnywhere, Category = AssetInfo, AssetRegistrySearchable, meta = (InlineEditConditionToggle))
+	bool bDescriptionOverride = false;
+
+	/** Can override the description of this instance. */
+	UPROPERTY(BlueprintReadWrite, EditAnywhere, Category = AssetInfo, AssetRegistrySearchable, meta = (EditCondition = "bDescriptionOverride", DisplayAfter = ColorOverride))
+	FText DescriptionOverride;
+
+	UPROPERTY(BlueprintReadWrite, EditAnywhere, Category = AssetInfo, AssetRegistrySearchable, meta = (InlineEditConditionToggle))
+	bool bCategoryOverride = false;
+
+	/** Can override the category of this instance. */
+	UPROPERTY(BlueprintReadWrite, EditAnywhere, Category = AssetInfo, AssetRegistrySearchable, meta = (EditCondition = "bCategoryOverride", DisplayAfter = ColorOverride))
+	FText CategoryOverride;
 #endif // WITH_EDITORONLY_DATA
 };
 
diff --git a/Engine/Plugins/PCG/Source/PCG/Public/PCGSubgraph.h b/Engine/Plugins/PCG/Source/PCG/Public/PCGSubgraph.h
--- a/Engine/Plugins/PCG/Source/PCG/Public/PCGSubgraph.h
+++ b/Engine/Plugins/PCG/Source/PCG/Public/PCGSubgraph.h
@@ -103,6 +103,7 @@ public:
 #if WITH_EDITOR
 	virtual FName GetDefaultNodeName() const override { return FName(TEXT("Subgraph")); }
 	virtual FText GetDefaultNodeTitle() const override { return NSLOCTEXT("PCGSubgraphSettings", "NodeTitle", "Subgraph"); }
+	virtual FLinearColor GetNodeTitleColor() const override;
 	virtual EPCGSettingsType GetType() const override { return EPCGSettingsType::Subgraph; }
 	virtual UObject* GetJumpTargetForDoubleClick() const override;
 	virtual bool GetPinExtraIcon(const UPCGPin* InPin, FName& OutExtraIcon, FText& OutTooltip) const override { return false; }
diff --git a/Engine/Plugins/PCG/Source/PCGEditor/Private/PCGEditorGraphSchema.cpp b/Engine/Plugins/PCG/Source/PCGEditor/Private/PCGEditorGraphSchema.cpp
--- a/Engine/Plugins/PCG/Source/PCGEditor/Private/PCGEditorGraphSchema.cpp
+++ b/Engine/Plugins/PCG/Source/PCGEditor/Private/PCGEditorGraphSchema.cpp
@@ -25,6 +25,7 @@
 #include "PCGEditorUtils.h"
 
 #include "AssetRegistry/AssetData.h"
+#include "AssetRegistry/AssetRegistryModule.h"
 #include "Blueprint/BlueprintSupport.h"
 #include "Engine/Blueprint.h"
 #include "Framework/Application/SlateApplication.h"
@@ -536,18 +537,74 @@ void UPCGEditorGraphSchema::GetSettingsElementActions(FGraphActionMenuBuilder& A
 
 void UPCGEditorGraphSchema::GetSubgraphElementActions(FGraphActionMenuBuilder& ActionMenuBuilder) const
 {
-	PCGEditorUtils::ForEachPCGGraphAssetData([&ActionMenuBuilder](const FAssetData& AssetData)
+	FAssetRegistryModule& AssetRegistryModule = FModuleManager::Get().LoadModuleChecked<FAssetRegistryModule>(TEXT("AssetRegistry"));
+
+	PCGEditorUtils::ForEachPCGGraphAssetData([&ActionMenuBuilder, &AssetRegistryModule](const FAssetData& AssetData) -> bool
 	{
-		const bool bExposeToLibrary = AssetData.GetTagValueRef<bool>(TEXT("bExposeToLibrary"));
+		const bool bExposeToLibrary = AssetData.GetTagValueRef<bool>(GET_MEMBER_NAME_CHECKED(UPCGGraphInterface, bExposeToLibrary));
 		if (bExposeToLibrary)
 		{
-			const FText MenuDesc = FText::FromString(FName::NameToDisplayString(AssetData.AssetName.ToString(), false));
-			const FText Category = AssetData.GetTagValueRef<FText>(TEXT("Category"));
-			const FText Description = AssetData.GetTagValueRef<FText>(TEXT("Description"));
+			// Only exposing the instance if it's parent graph is defined for instances. Otherwise, it is not interesting.
+			// Also march up the hierarchy to find overrides for category and description.
+			// We don't look for override in titles because if we have no override for an instance but the parent has an override we will have 2 times the same entry in the palette.
+			auto GetRecursiveTextsAndIsValid = [&AssetRegistryModule](const FAssetData& InAssetData, FText& OutCategory, FText& OutDescription, auto Recurse)
+			{
+				if (InAssetData.IsInstanceOf<UPCGGraphInstance>())
+				{
+					const FSoftObjectPath ParentGraph(InAssetData.GetTagValueRef<FString>(GET_MEMBER_NAME_CHECKED(UPCGGraphInstance, Graph)));
+					if (ParentGraph.IsNull())
+					{
+						return false;
+					}
+
+					if (OutCategory.IsEmpty())
+					{
+						OutCategory = InAssetData.GetTagValueRef<bool>(GET_MEMBER_NAME_CHECKED(UPCGGraphInstance, bCategoryOverride)) 
+							? InAssetData.GetTagValueRef<FText>(GET_MEMBER_NAME_CHECKED(UPCGGraphInstance, CategoryOverride)) 
+							: FText();
+					}
+
+					if (OutDescription.IsEmpty())
+					{
+						OutDescription = InAssetData.GetTagValueRef<bool>(GET_MEMBER_NAME_CHECKED(UPCGGraphInstance, bDescriptionOverride)) 
+							? InAssetData.GetTagValueRef<FText>(GET_MEMBER_NAME_CHECKED(UPCGGraphInstance, DescriptionOverride)) 
+							: FText();
+					}
+
+					// Asset data are not big so that should not be that big of a deal, but they are copied all the time.
+					// If we ever have performances issues, might be good to have a cache.
+					const FAssetData ParentAssetData = AssetRegistryModule.Get().GetAssetByObjectPath(ParentGraph);
+					return ParentAssetData.IsValid() ? Recurse(ParentAssetData, OutCategory, OutDescription, Recurse) : false;
+				}
+				else
+				{
+					if (OutCategory.IsEmpty())
+					{
+						OutCategory = InAssetData.GetTagValueRef<FText>(GET_MEMBER_NAME_CHECKED(UPCGGraph, Category));
+					}
+
+					if (OutDescription.IsEmpty())
+					{
+						OutDescription = InAssetData.GetTagValueRef<FText>(GET_MEMBER_NAME_CHECKED(UPCGGraph, Description));
+					}
+
+					return true;
+				}
+			};
 
-			TSharedPtr<FPCGEditorGraphSchemaAction_NewSubgraphElement> NewSubgraphAction(new FPCGEditorGraphSchemaAction_NewSubgraphElement(Category, MenuDesc, Description, 0));
-			NewSubgraphAction->SubgraphObjectPath = AssetData.GetSoftObjectPath();
-			ActionMenuBuilder.AddAction(NewSubgraphAction);
+			FText Category, Description;
+
+			if (GetRecursiveTextsAndIsValid(AssetData, Category, Description, GetRecursiveTextsAndIsValid))
+			{
+				// As stated above, we either have an override and we take it, or we use the asset name, to differentiate all possible instances of the same graph.
+				const FText MenuDesc = AssetData.GetTagValueRef<bool>(GET_MEMBER_NAME_CHECKED(UPCGGraphInterface, bTitleOverride))
+					? AssetData.GetTagValueRef<FText>(GET_MEMBER_NAME_CHECKED(UPCGGraphInterface, TitleOverride))
+					: FText::FromString(FName::NameToDisplayString(AssetData.AssetName.ToString(), false));
+
+				TSharedPtr<FPCGEditorGraphSchemaAction_NewSubgraphElement> NewSubgraphAction(new FPCGEditorGraphSchemaAction_NewSubgraphElement(Category, MenuDesc, Description, 0));
+				NewSubgraphAction->SubgraphObjectPath = AssetData.GetSoftObjectPath();
+				ActionMenuBuilder.AddAction(NewSubgraphAction);
+			}
 		}
 
 		return true;
diff --git a/Engine/Plugins/PCG/Source/PCGEditor/Private/PCGEditorUtils.cpp b/Engine/Plugins/PCG/Source/PCGEditor/Private/PCGEditorUtils.cpp
--- a/Engine/Plugins/PCG/Source/PCGEditor/Private/PCGEditorUtils.cpp
+++ b/Engine/Plugins/PCG/Source/PCGEditor/Private/PCGEditorUtils.cpp
@@ -84,7 +84,7 @@ void PCGEditorUtils::ForEachPCGSettingsAssetData(TFunctionRef<bool(const FAssetD
 void PCGEditorUtils::ForEachPCGGraphAssetData(TFunctionRef<bool(const FAssetData&)> InFunc)
 {
 	FARFilter Filter;
-	Filter.ClassPaths.Add(UPCGGraph::StaticClass()->GetClassPathName());
+	Filter.ClassPaths.Add(UPCGGraphInterface::StaticClass()->GetClassPathName());
 	Filter.bRecursiveClasses = true;
 
 	ForEachAssetData(Filter, InFunc);
