[PCG] Select Grammar - Preliminary First Pass

+ ExtractAttributeSetAsArrayOfStructs (Adrien)
+ New read-only Point Property - LocalSize/ScaledLocalSize
+ New format for commonly used logging errors added - PCGLogErrors.h

#rnx

diff --git a/Engine/Plugins/PCG/Source/PCG/Private/PCGPoint.cpp b/Engine/Plugins/PCG/Source/PCG/Private/PCGPoint.cpp
--- a/Engine/Plugins/PCG/Source/PCG/Private/PCGPoint.cpp
+++ b/Engine/Plugins/PCG/Source/PCG/Private/PCGPoint.cpp
@@ -16,6 +16,8 @@ namespace PCGPointCustomPropertyNames
 	const FName PositionName = TEXT("Position");
 	const FName RotationName = TEXT("Rotation");
 	const FName ScaleName = TEXT("Scale");
+	const FName LocalSizeName = TEXT("LocalSize");
+	const FName ScaledLocalSizeName = TEXT("ScaledLocalSize");
 }
 
 /** Serialized fields of a FPCGPoint, the values here can't change as they are being used to mask out serialization */
@@ -182,7 +184,9 @@ bool FPCGPoint::HasCustomPropertyGetterSetter(FName Name)
 		Name == PCGPointCustomPropertyNames::LocalCenterName ||
 		Name == PCGPointCustomPropertyNames::PositionName ||
 		Name == PCGPointCustomPropertyNames::RotationName ||
-		Name == PCGPointCustomPropertyNames::ScaleName;
+		Name == PCGPointCustomPropertyNames::ScaleName ||
+		Name == PCGPointCustomPropertyNames::LocalSizeName ||
+		Name == PCGPointCustomPropertyNames::ScaledLocalSizeName;
 }
 
 TUniquePtr<IPCGAttributeAccessor> FPCGPoint::CreateCustomPropertyAccessor(FName Name)
@@ -222,6 +226,18 @@ TUniquePtr<IPCGAttributeAccessor> FPCGPoint::CreateCustomPropertyAccessor(FName
 			[](FPCGPoint& Point, const void* InValue) { Point.Transform.SetScale3D(*reinterpret_cast<const FVector*>(InValue)); return true; }
 		);
 	}
+	else if (Name == PCGPointCustomPropertyNames::LocalSizeName)
+	{
+		return MakeUnique<FPCGCustomPointAccessor<FVector>>(
+			[](const FPCGPoint& Point, void* OutValue) { *reinterpret_cast<FVector*>(OutValue) = Point.GetLocalSize(); return true; }
+		);
+	}
+	else if (Name == PCGPointCustomPropertyNames::ScaledLocalSizeName)
+	{
+		return MakeUnique<FPCGCustomPointAccessor<FVector>>(
+			[](const FPCGPoint& Point, void* OutValue) { *reinterpret_cast<FVector*>(OutValue) = Point.GetScaledLocalSize(); return true; }
+		);
+	}
 
 	return TUniquePtr<IPCGAttributeAccessor>();
 }
diff --git a/Engine/Plugins/PCG/Source/PCG/Private/Utils/PCGLogErrors.cpp b/Engine/Plugins/PCG/Source/PCG/Private/Utils/PCGLogErrors.cpp
new file mode 100644
--- /dev/null
+++ b/Engine/Plugins/PCG/Source/PCG/Private/Utils/PCGLogErrors.cpp
@@ -0,0 +1,59 @@
+// Copyright Epic Games, Inc. All Rights Reserved.
+
+#include "Utils/PCGLogErrors.h"
+
+#define LOCTEXT_NAMESPACE "PCGLog"
+
+namespace PCGLog::InputOutput
+{
+	namespace ErrorFormat
+	{
+		const FTextFormat TypedInputNotFoundWarning = LOCTEXT("DataInputNotFound", "Data of type {0} not found on pin '{1}'.");
+		const FTextFormat FirstInputOnlyWarning = LOCTEXT("FirstInputOnly", "Multiple inputs found on single-input pin '{0}'. Only the first will be selected.");
+	}
+
+	void LogTypedDataNotFoundWarning(EPCGDataType DataType, const FName PinLabel, const FPCGContext* InContext)
+	{
+		const UEnum* PCGDataTypeEnum = StaticEnum<EPCGDataType>();
+		const FText TypeText = PCGDataTypeEnum ? PCGDataTypeEnum->GetDisplayNameTextByValue(static_cast<int64>(DataType)) : LOCTEXT("UnknownDataType", "Unknown");
+		LogWarningOnGraph(FText::Format(ErrorFormat::TypedInputNotFoundWarning, FText::FromName(PinLabel), TypeText), InContext);
+	}
+
+	void LogFirstInputOnlyWarning(const FName PinLabel, const FPCGContext* InContext)
+	{
+		LogWarningOnGraph(FText::Format(ErrorFormat::FirstInputOnlyWarning, FText::FromName(PinLabel)), InContext);
+	}
+}
+
+namespace PCGLog::Accessor
+{
+	namespace ErrorFormat
+	{
+		const FTextFormat CreateAccessorFailure = LOCTEXT("CreateAccessorFailure", "Attribute {0} was not found.");
+		const FTextFormat GetAttributeFailure = LOCTEXT("GetAttributeFailure", "Couldn't retrieve attribute '{0}' value. Expected type: {1}, Actual Type: {2}.");
+		const FTextFormat GetTypedAttributeFailure = LOCTEXT("GetAttributeFailure", "Couldn't retrieve attribute '{0}' value. Expected type: {1}, Actual Type: {2}.");
+		const FTextFormat GetTypedAttributeFailureNoAccessor = LOCTEXT("GetAttributeFailureNoAccessor", "Couldn't retrieve attribute '{0}' value of type: '{1}.");
+	}
+
+	void LogFailToCreate(const FPCGAttributePropertySelector& Selector, const FPCGContext* InContext)
+	{
+		LogErrorOnGraph(FText::Format(ErrorFormat::CreateAccessorFailure, Selector.GetDisplayText()), InContext);
+	}
+
+	void LogFailToGet(FName AttributeName, const FPCGContext* InContext)
+	{
+		LogFailToGet(FText::FromName(AttributeName), InContext);
+	}
+
+	void LogFailToGet(FText AttributeName, const FPCGContext* InContext)
+	{
+		LogErrorOnGraph(FText::Format(ErrorFormat::GetAttributeFailure, std::move(AttributeName)), InContext);
+	}
+
+	void LogFailToGet(const FPCGAttributePropertySelector& Selector, const FPCGContext* InContext)
+	{
+		LogErrorOnGraph(FText::Format(ErrorFormat::GetAttributeFailure, Selector.GetDisplayText()), InContext);
+	}
+}
+
+#undef LOCTEXT_NAMESPACE
diff --git a/Engine/Plugins/PCG/Source/PCG/Public/Elements/Metadata/PCGMetadataElementCommon.h b/Engine/Plugins/PCG/Source/PCG/Public/Elements/Metadata/PCGMetadataElementCommon.h
--- a/Engine/Plugins/PCG/Source/PCG/Public/Elements/Metadata/PCGMetadataElementCommon.h
+++ b/Engine/Plugins/PCG/Source/PCG/Public/Elements/Metadata/PCGMetadataElementCommon.h
@@ -120,10 +120,10 @@ namespace PCGMetadataElementCommon
 	* Iterate over the full range of the keys (if Count is negative, otherwise, as many times as Count), calling the callback with values ranges get from the multiple accessors.
 	* Keys needs to be the same for all accessors.
 	* ChunkSize influence the max number of values to get in one go with GetRange.
-	* Callback should have a signature: void(const TArrayView<T>&... Args, int32 Index)
-	* Return false if it process nothing.
+	* Callback should have a signature: void(const TArrayView<T>&... Args, int32 Start, int32 Range)
+	* Return false if it processes nothing.
 	*/
-	template <typename... T, typename Func = TFunction<void(const TArrayView<T>&..., int32 Index)>>
+	template <typename... T, typename Func = TFunction<void(const TArrayView<T>&..., int32 Start, int32 Range)>>
 	bool ApplyOnMultiAccessorsRange(const IPCGAttributeAccessorKeys& Keys, const TArray<const IPCGAttributeAccessor*>& Accessors, Func&& InCallback, EPCGAttributeAccessorFlags Flags = EPCGAttributeAccessorFlags::StrictType, const int32 ChunkSize = DefaultChunkSize, const int32 Count = -1)
 	{
 		// We use C++20 functionality here, so make sure we are at this version at least
diff --git a/Engine/Plugins/PCG/Source/PCG/Public/Helpers/PCGPropertyHelpers.h b/Engine/Plugins/PCG/Source/PCG/Public/Helpers/PCGPropertyHelpers.h
--- a/Engine/Plugins/PCG/Source/PCG/Public/Helpers/PCGPropertyHelpers.h
+++ b/Engine/Plugins/PCG/Source/PCG/Public/Helpers/PCGPropertyHelpers.h
@@ -91,10 +91,19 @@ namespace PCGPropertyHelpers
 	* @param OptionalObjectTraversed - Optional set to store all objects that we traversed, to be able to react to those objects changes.
 	*/
 	PCG_API UPCGParamData* ExtractPropertyAsAttributeSet(const FExtractorParameters& Parameters, FPCGContext* OptionalContext = nullptr, TSet<FSoftObjectPath>* OptionalObjectTraversed = nullptr);
+
+	/**
+	* Extract an attribute set in a array of structures. T MUST be a UStruct. Also, it must only contain supported types (so no arrays nor other structures)
+	* @param InParamData - Attribute set that contains the data.
+	* @param OptionalNameMapping - Optional mapping for the name in the structure and the name in the attribute set. Can also say if this property is required, or not and should be defaulted. By default all are defaulted if not found.
+	* @param OptionalContext - Optional context if the extraction is done in a PCG Node, so errors are using the context to log.
+	*/
+	template <typename T>
+	TArray<T> ExtractAttributeSetAsArrayOfStructs(const UPCGParamData* InParamData, const TMap<FName, TTuple<FName, bool>>* OptionalNameMapping = nullptr, FPCGContext* OptionalContext = nullptr);
 }
 
 //////
-/// PCGPropertiesHelpers Implementation
+/// PCGPropertyHelpers Implementation
 //////
 
 // Func signature : auto(auto&&)
@@ -176,3 +185,91 @@ inline bool PCGPropertyHelpers::SetPropertyValueFromCallback(ObjectType* InObjec
 
 	return PCGMetadataAttribute::CallbackWithRightType(PropertyAccessor->GetUnderlyingType(), Setter);
 }
+
+template <typename T>
+TArray<T> PCGPropertyHelpers::ExtractAttributeSetAsArrayOfStructs(const UPCGParamData* InParamData, const TMap<FName, TTuple<FName, bool>>* OptionalNameMapping, FPCGContext* OptionalContext)
+{
+	if (!ensure(InParamData))
+	{
+		return {};
+	}
+
+	TArray<T> OutResult;
+	for (TFieldIterator<const FProperty> PropIt(T::StaticStruct(), EFieldIterationFlags::IncludeSuper); PropIt; ++PropIt)
+	{
+		if (!PCGAttributeAccessorHelpers::IsPropertyAccessorSupported(*PropIt))
+		{
+			continue;
+		}
+
+		FName PropertyName = *T::StaticStruct()->GetAuthoredNameForField(*PropIt);
+		bool bCanBeDefaulted = true;
+		const TTuple<FName, bool>* It = OptionalNameMapping ? OptionalNameMapping->Find(PropertyName) : nullptr;
+		if (It)
+		{
+			PropertyName = It->Get<0>();
+			bCanBeDefaulted = It->Get<1>();
+		}
+
+		FPCGAttributePropertySelector Selector = FPCGAttributePropertySelector::CreateAttributeSelector(PropertyName);
+		const TUniquePtr<const IPCGAttributeAccessor> Accessor = PCGAttributeAccessorHelpers::CreateConstAccessor(CastChecked<UPCGData>(InParamData), Selector);
+		const TUniquePtr<const IPCGAttributeAccessorKeys> Keys = PCGAttributeAccessorHelpers::CreateConstKeys(CastChecked<UPCGData>(InParamData), Selector);
+
+		if (!Accessor || !Keys)
+		{
+			if (!bCanBeDefaulted)
+			{
+				PCGLog::Accessor::LogFailToGet(Selector, OptionalContext);
+				return {};
+			}
+			else
+			{
+				continue;
+			}
+		}
+
+		const int32 NumElements = Keys->GetNum();
+
+		if (OutResult.IsEmpty())
+		{
+			OutResult.SetNum(NumElements);
+		}
+
+		if (NumElements == 0 || NumElements != OutResult.Num())
+		{
+			return {};
+		}
+
+		TUniquePtr<IPCGAttributeAccessor> StructAccessor = PCGAttributeAccessorHelpers::CreatePropertyAccessor(*PropIt);
+		TArrayView<T> StructData(OutResult.GetData(), OutResult.Num());
+		FPCGAttributeAccessorKeysGeneric<T> StructKeys(StructData);
+
+		auto GetSetRange = [OptionalContext, &Keys, &Accessor, &StructAccessor, &StructKeys, NumElements, PropertyName]<typename PropertyType>(PropertyType Dummy) -> bool
+		{
+			TArray<PropertyType> Temp;
+			if constexpr (std::is_trivially_copyable_v<PropertyType>)
+			{
+				Temp.SetNumUninitialized(NumElements);
+			}
+			else
+			{
+				Temp.SetNum(NumElements);
+			}
+
+			if (!Accessor->GetRange<PropertyType>(Temp, 0, *Keys, EPCGAttributeAccessorFlags::AllowBroadcastAndConstructible))
+			{
+				PCGLog::Accessor::LogFailToGet<PropertyType>(FText::FromName(PropertyName), Accessor.Get(), OptionalContext);
+				return false;
+			}
+
+			return StructAccessor->SetRange<PropertyType>(Temp, 0, StructKeys);
+		};
+
+		if (!PCGMetadataAttribute::CallbackWithRightType(StructAccessor->GetUnderlyingType(), GetSetRange))
+		{
+			return {};
+		}
+	}
+
+	return OutResult;
+}
diff --git a/Engine/Plugins/PCG/Source/PCG/Public/Metadata/Accessors/PCGCustomAccessor.h b/Engine/Plugins/PCG/Source/PCG/Public/Metadata/Accessors/PCGCustomAccessor.h
--- a/Engine/Plugins/PCG/Source/PCG/Public/Metadata/Accessors/PCGCustomAccessor.h
+++ b/Engine/Plugins/PCG/Source/PCG/Public/Metadata/Accessors/PCGCustomAccessor.h
@@ -24,6 +24,11 @@ public:
 		, Setter(InSetter)
 	{}
 
+	FPCGCustomPointAccessor(const FPCGPoint::PointCustomPropertyGetter& InGetter)
+	: Super(/*bInReadOnly=*/ true)
+	, Getter(InGetter)
+	{}
+
 	bool GetRangeImpl(TArrayView<T> OutValues, int32 Index, const IPCGAttributeAccessorKeys& Keys) const
 	{
 		TArray<const FPCGPoint*> PointKeys;
@@ -208,4 +213,4 @@ public:
 	{
 		return false;
 	}
-};
\ No newline at end of file
+};
diff --git a/Engine/Plugins/PCG/Source/PCG/Public/PCGModule.h b/Engine/Plugins/PCG/Source/PCG/Public/PCGModule.h
--- a/Engine/Plugins/PCG/Source/PCG/Public/PCGModule.h
+++ b/Engine/Plugins/PCG/Source/PCG/Public/PCGModule.h
@@ -4,9 +4,9 @@
 
 #include "PCGDataVisualizationRegistry.h"
 #include "Data/PCGGetDataFunctionRegistry.h"
+#include "Utils/PCGLogErrors.h"
 
 #include "Modules/ModuleInterface.h"
-#include "Stats/Stats.h"
 
 // Logs
 PCG_API DECLARE_LOG_CATEGORY_EXTERN(LogPCG, Log, All);
@@ -14,14 +14,6 @@ PCG_API DECLARE_LOG_CATEGORY_EXTERN(LogPCG, Log, All);
 struct FPCGContext;
 class IPCGDataVisualization;
 
-namespace PCGLog
-{
-	/** Convenience function that would either log error on the graph if there is a context, or in the console if not. */
-	PCG_API void LogErrorOnGraph(const FText& InMsg, const FPCGContext* InContext = nullptr);
-	/** Convenience function that would either log warning on the graph if there is a context, or in the console if not. */
-	PCG_API void LogWarningOnGraph(const FText& InMsg, const FPCGContext* InContext = nullptr);
-}
-
 namespace PCGEngineShowFlags
 {
 	static constexpr TCHAR Debug[] = TEXT("PCGDebug");
diff --git a/Engine/Plugins/PCG/Source/PCG/Public/PCGPoint.h b/Engine/Plugins/PCG/Source/PCG/Public/PCGPoint.h
--- a/Engine/Plugins/PCG/Source/PCG/Public/PCGPoint.h
+++ b/Engine/Plugins/PCG/Source/PCG/Public/PCGPoint.h
@@ -11,18 +11,20 @@ class IPCGAttributeAccessor;
 UENUM()
 enum class EPCGPointProperties : uint8
 {
-	Density,
-	BoundsMin,
-	BoundsMax,
-	Extents,
-	Color,
-	Position,
-	Rotation,
-	Scale,
-	Transform,
-	Steepness,
-	LocalCenter,
-	Seed
+	Density UMETA(Tooltip = "When points are sampled, this density value represents the highest value of the density function within that point's volume. It is also used as a weighted value, for example, when testing points against a threshold in filtering operations."),
+	BoundsMin UMETA(Tooltip = "Minimum corner of the point's bounds in local space."),
+	BoundsMax UMETA(Tooltip = "Maximum corner of the point's bounds in local space."),
+	Extents UMETA(Tooltip = "Half the local space difference between the maximum and minimum bounds of the point's volume. Can be used with the point's position to represent the volume."),
+	Color UMETA(Tooltip = "An RGBA (four channel) color value."),
+	Position UMETA(Tooltip = "Location component of the point's transform."),
+	Rotation UMETA(Tooltip = "Rotation component of the point's transform."),
+	Scale UMETA(Tooltip = "Scale component of the point's transform."),
+	Transform UMETA(Tooltip = "The point's transform."),
+	Steepness UMETA(Tooltip = "A normalized value that establishes how 'hard' or 'soft' that volume will be represented. From 0, it will ramp up linearly increasing its influence over the density from the point's center to up to two times the bounds. At 1, it will represent a binary box function with the size of the point's bounds."),
+	LocalCenter UMETA(Tooltip = "The local center location of the point's volume, halfway between the minimum and maximum bounds."),
+	Seed UMETA(Tooltip = "Used to seed random processes during various operations."),
+	LocalSize UMETA(Tooltip = "The difference between the maximum and minimum bounds of the point."),
+	ScaledLocalSize UMETA(Tooltip = "The difference between the maximum and minimum bounds of the point, after only the scale has been applied.")
 };
 
 USTRUCT(BlueprintType)
@@ -80,6 +82,9 @@ public:
 		BoundsMax += Delta;
 	}
 
+	FVector GetLocalSize() const { return BoundsMax - BoundsMin; }
+	FVector GetScaledLocalSize() const { return GetLocalSize() * Transform.GetScale3D(); }
+
 	void ApplyScaleToBounds();
 
 	void ResetPointCenter(const FVector& BoundsRatio);
@@ -101,8 +106,3 @@ struct TStructOpsTypeTraits<FPCGPoint> : public TStructOpsTypeTraitsBase2<FPCGPo
 		WithStructuredSerializer = true,
 	};
 };
-
-#if UE_ENABLE_INCLUDE_ORDER_DEPRECATED_IN_5_2
-#include "CoreMinimal.h"
-#include "Metadata/PCGMetadataAttributeTraits.h"
-#endif
diff --git a/Engine/Plugins/PCG/Source/PCG/Public/Utils/PCGLogErrors.h b/Engine/Plugins/PCG/Source/PCG/Public/Utils/PCGLogErrors.h
new file mode 100644
--- /dev/null
+++ b/Engine/Plugins/PCG/Source/PCG/Public/Utils/PCGLogErrors.h
@@ -0,0 +1,69 @@
+// Copyright Epic Games, Inc. All Rights Reserved.
+
+#pragma once
+
+#include "Internationalization/Internationalization.h"
+#include "Metadata/PCGAttributePropertySelector.h"
+#include "Metadata/Accessors/IPCGAttributeAccessor.h"
+
+enum class EPCGDataType : uint32;
+struct FPCGContext;
+
+namespace PCGLog
+{
+	/** Convenience function that would either log error on the graph if there is a context, or in the console if not. */
+	PCG_API void LogErrorOnGraph(const FText& InMsg, const FPCGContext* InContext = nullptr);
+	/** Convenience function that would either log warning on the graph if there is a context, or in the console if not. */
+	PCG_API void LogWarningOnGraph(const FText& InMsg, const FPCGContext* InContext = nullptr);
+
+	namespace InputOutput
+	{
+		namespace ErrorFormat
+		{
+			extern const FTextFormat TypedInputNotFoundWarning;
+			extern const FTextFormat FirstInputOnlyWarning;
+		}
+
+		void LogTypedDataNotFoundWarning(EPCGDataType DataType, const FName PinLabel, const FPCGContext* InContext = nullptr);
+		void LogFirstInputOnlyWarning(const FName PinLabel, const FPCGContext* InContext = nullptr);
+	}
+
+	namespace Accessor
+	{
+		namespace ErrorFormat
+		{
+			extern const FTextFormat GetTypedAttributeFailure;
+			extern const FTextFormat GetTypedAttributeFailureNoAccessor;
+		}
+
+		void LogFailToCreate(const FPCGAttributePropertySelector& Selector, const FPCGContext* InContext = nullptr);
+		void LogFailToGet(FName AttributeName, const FPCGContext* InContext = nullptr);
+		void LogFailToGet(FText AttributeName, const FPCGContext* InContext = nullptr);
+		void LogFailToGet(const FPCGAttributePropertySelector& Selector, const FPCGContext* InContext = nullptr);
+
+		template <typename T>
+		void LogFailToGet(FText AttributeName, const IPCGAttributeAccessor* Accessor, const FPCGContext* InContext = nullptr)
+		{
+			if (Accessor)
+			{
+				PCGLog::LogErrorOnGraph(FText::Format(ErrorFormat::GetTypedAttributeFailure, AttributeName, PCG::Private::GetTypeNameText<T>(), PCG::Private::GetTypeNameText(Accessor->GetUnderlyingType())), InContext);
+			}
+			else
+			{
+				PCGLog::LogErrorOnGraph(FText::Format(ErrorFormat::GetTypedAttributeFailureNoAccessor, AttributeName, PCG::Private::GetTypeNameText<T>()), InContext);
+			}
+		}
+
+		template <typename T>
+		void LogFailToGet(FName AttributeName, const IPCGAttributeAccessor* Accessor, const FPCGContext* InContext = nullptr)
+		{
+			LogFailToGet<T>(FText::FromName(AttributeName), Accessor, InContext);
+		}
+
+		template <typename T>
+		void LogFailToGet(const FPCGAttributePropertySelector& Selector, const IPCGAttributeAccessor* Accessor, const FPCGContext* InContext = nullptr)
+		{
+			return LogFailToGet<T>(Selector.GetDisplayText(), Accessor, InContext);
+		}
+	}
+}
diff --git a/Engine/Plugins/PCG/Source/PCGEditor/Private/Details/PCGAttributePropertySelectorDetails.cpp b/Engine/Plugins/PCG/Source/PCGEditor/Private/Details/PCGAttributePropertySelectorDetails.cpp
--- a/Engine/Plugins/PCG/Source/PCGEditor/Private/Details/PCGAttributePropertySelectorDetails.cpp
+++ b/Engine/Plugins/PCG/Source/PCGEditor/Private/Details/PCGAttributePropertySelectorDetails.cpp
@@ -32,10 +32,11 @@ namespace PCGAttributePropertySelectorDetails
 			{
 				FString EnumName = EnumPtr->GetDisplayNameTextByIndex(i).ToString();
 				EnumType EnumValue = (EnumType)EnumPtr->GetValueByIndex(i);
+				FText Tooltip = EnumPtr->GetToolTipTextByIndex(i);
 
 				InMenuBuilder.AddMenuEntry(
 					FText::FromString(EnumName),
-					FText(),
+					std::move(Tooltip),
 					FSlateIcon(),
 					FExecuteAction::CreateSP(InDetailsObject, InCallback, EnumValue),
 					NAME_None,
