[PCG] Fix vector truncation warning from custom hlsl nodes.

#rnx

diff --git a/Engine/Plugins/PCG/Source/PCG/Private/Compute/DataInterfaces/PCGCustomKernelDataInterface.cpp b/Engine/Plugins/PCG/Source/PCG/Private/Compute/DataInterfaces/PCGCustomKernelDataInterface.cpp
--- a/Engine/Plugins/PCG/Source/PCG/Private/Compute/DataInterfaces/PCGCustomKernelDataInterface.cpp
+++ b/Engine/Plugins/PCG/Source/PCG/Private/Compute/DataInterfaces/PCGCustomKernelDataInterface.cpp
@@ -19,7 +19,7 @@ void UPCGCustomKernelDataInterface::GetSupportedInputs(TArray<FShaderFunctionDef
 {
 	OutFunctions.AddDefaulted_GetRef()
 		.SetName(TEXT("GetNumThreads"))
-		.AddReturnType(FShaderValueType::Get(EShaderFundamentalType::Int));
+		.AddReturnType(FShaderValueType::Get(EShaderFundamentalType::Int, 3));
 
 	OutFunctions.AddDefaulted_GetRef()
 		.SetName(TEXT("GetSeed"))
diff --git a/Engine/Plugins/PCG/Source/PCG/Private/Compute/Elements/PCGCustomHLSL.cpp b/Engine/Plugins/PCG/Source/PCG/Private/Compute/Elements/PCGCustomHLSL.cpp
--- a/Engine/Plugins/PCG/Source/PCG/Private/Compute/Elements/PCGCustomHLSL.cpp
+++ b/Engine/Plugins/PCG/Source/PCG/Private/Compute/Elements/PCGCustomHLSL.cpp
@@ -547,7 +547,7 @@ void UPCGCustomHLSLSettings::UpdateDeclarations()
 		}
 
 		InputDeclarations += TEXT("uint ElementIndex;\n");
-		InputDeclarations += TEXT("uint GetNumThreads();\n");
+		InputDeclarations += TEXT("int3 GetNumThreads();\n");
 	}
 
 	// Add debug category
@@ -1061,7 +1061,7 @@ FString UPCGCustomHLSLSettings::GetCookedKernelSource(const TMap<FPCGKernelAttri
 	auto EmitHeaderWriterFromInputPin = [&HeaderWriters](const FPCGPinProperties& InOutputPinProps, const UPCGPin* InFromPin)
 	{
 		HeaderWriters += FString::Format(TEXT(
-			"    if (ThreadIndex >= GetNumThreads()) return;\n"
+			"    if (ThreadIndex >= GetNumThreads().x) return;\n"
 			"    \n"
 			"    // Signal kernel executed by copying data count from pin {0} to pin {1} from first thread. Rest of header was already set up by the CPU.\n"
 			"    if (GroupIndex == 0) {1}_WriteNumData({0}_GetNumData());\n"
@@ -1112,7 +1112,7 @@ FString UPCGCustomHLSLSettings::GetCookedKernelSource(const TMap<FPCGKernelAttri
 			auto EmitPointGenHeader = [&HeaderWriters, InPointCount = PointCount](const FPCGPinProperties& InOutputPinProps)
 			{
 				HeaderWriters += FString::Format(TEXT(
-					"    if (ThreadIndex >= GetNumThreads()) return;\n"
+					"    if (ThreadIndex >= GetNumThreads().x) return;\n"
 					"    \n"
 					"    // Signal kernel executed by writing data count (1) for pin {0} from first thread. Rest of header was already set up by the CPU.\n"
 					"    if (GroupIndex == 0) {0}_WriteNumData(1);\n"
