Small UX improvement around the save data button.

#rnx

diff --git a/Engine/Plugins/PCG/Source/PCGEditor/Private/SPCGEditorGraphAttributeListView.cpp b/Engine/Plugins/PCG/Source/PCGEditor/Private/SPCGEditorGraphAttributeListView.cpp
--- a/Engine/Plugins/PCG/Source/PCGEditor/Private/SPCGEditorGraphAttributeListView.cpp
+++ b/Engine/Plugins/PCG/Source/PCGEditor/Private/SPCGEditorGraphAttributeListView.cpp
@@ -873,7 +873,8 @@ TSharedRef<SWidget> SPCGEditorGraphAttributeListView::OnGenerateAdditionalOperat
 		LOCTEXT("SaveThisData", "Save this data"),
 		LOCTEXT("SaveThisDataTooltip", "Saves this data to a PCG Data Asset."),
 		FSlateIcon(FAppStyle::GetAppStyleSetName(), "Icons.Download"),
-		FUIAction(FExecuteAction::CreateSP(this, &SPCGEditorGraphAttributeListView::SaveData, true, true)),
+		FUIAction(FExecuteAction::CreateSP(this, &SPCGEditorGraphAttributeListView::SaveData, true, true),
+			FCanExecuteAction::CreateSP(this, &SPCGEditorGraphAttributeListView::CanSaveData, true, true)),
 		NAME_None,
 		EUserInterfaceActionType::Button);
 
@@ -881,7 +882,8 @@ TSharedRef<SWidget> SPCGEditorGraphAttributeListView::OnGenerateAdditionalOperat
 		LOCTEXT("SaveThisPinData", "Save pin data"),
 		LOCTEXT("SaveThisPinDataTooltip", "Saves all data from the selected pin to a PCG Data Asset."),
 		FSlateIcon(FAppStyle::GetAppStyleSetName(), "Icons.Download"),
-		FUIAction(FExecuteAction::CreateSP(this, &SPCGEditorGraphAttributeListView::SaveData, true, false)),
+		FUIAction(FExecuteAction::CreateSP(this, &SPCGEditorGraphAttributeListView::SaveData, true, false),
+			FCanExecuteAction::CreateSP(this, &SPCGEditorGraphAttributeListView::CanSaveData, true, false)),
 		NAME_None,
 		EUserInterfaceActionType::Button);
 
@@ -889,14 +891,15 @@ TSharedRef<SWidget> SPCGEditorGraphAttributeListView::OnGenerateAdditionalOperat
 		LOCTEXT("SaveAllData", "Save all"),
 		LOCTEXT("SaveAllDataTooltip", "Saves all the input or output data to a PCG Data Asset."),
 		FSlateIcon(FAppStyle::GetAppStyleSetName(), "Icons.Download"),
-		FUIAction(FExecuteAction::CreateSP(this, &SPCGEditorGraphAttributeListView::SaveData, false, false)),
+		FUIAction(FExecuteAction::CreateSP(this, &SPCGEditorGraphAttributeListView::SaveData, false, false),
+			FCanExecuteAction::CreateSP(this, &SPCGEditorGraphAttributeListView::CanSaveData, false, false)),
 		NAME_None,
 		EUserInterfaceActionType::Button);
 
 	return MenuBuilder.MakeWidget();
 }
 
-void SPCGEditorGraphAttributeListView::SaveData(bool bUsePinComboIndex, bool bUseDataComboIndex)
+FPCGDataCollection SPCGEditorGraphAttributeListView::BuildDataCollectionForSave(bool bUsePinComboIndex, bool bUseDataComboIndex) const
 {
 	// Can't save N'th data on all pins, as it makes no sense
 	check(!bUseDataComboIndex || bUsePinComboIndex);
@@ -931,6 +934,13 @@ void SPCGEditorGraphAttributeListView::SaveData(bool bUsePinComboIndex, bool bUs
 		}
 	}
 
+	return Collection;
+}
+
+void SPCGEditorGraphAttributeListView::SaveData(bool bUsePinComboIndex, bool bUseDataComboIndex)
+{
+	FPCGDataCollection Collection = BuildDataCollectionForSave(bUsePinComboIndex, bUseDataComboIndex);
+
 	if (Collection.TaggedData.IsEmpty())
 	{
 		return;
@@ -943,6 +953,12 @@ void SPCGEditorGraphAttributeListView::SaveData(bool bUsePinComboIndex, bool bUs
 	UPCGAssetExporterUtils::CreateAsset(Exporter, Parameters);
 }
 
+bool SPCGEditorGraphAttributeListView::CanSaveData(bool bUsePinComboIndex, bool bUseDataComboIndex) const
+{
+	FPCGDataCollection Collection = BuildDataCollectionForSave(bUsePinComboIndex, bUseDataComboIndex);
+	return !Collection.TaggedData.IsEmpty();
+}
+
 FText SPCGEditorGraphAttributeListView::OnGenerateSelectedPinText() const
 {
 	if (const TSharedPtr<FPinComboBoxItem> SelectedPin = PinComboBox->GetSelectedItem())
diff --git a/Engine/Plugins/PCG/Source/PCGEditor/Private/SPCGEditorGraphAttributeListView.h b/Engine/Plugins/PCG/Source/PCGEditor/Private/SPCGEditorGraphAttributeListView.h
--- a/Engine/Plugins/PCG/Source/PCGEditor/Private/SPCGEditorGraphAttributeListView.h
+++ b/Engine/Plugins/PCG/Source/PCGEditor/Private/SPCGEditorGraphAttributeListView.h
@@ -200,7 +200,9 @@ private:
 	ECheckBoxState GetAnyAttributeEnabledState() const;
 	bool IsAttributeEnabled(FName InAttributeName) const;
 
+	FPCGDataCollection BuildDataCollectionForSave(bool bUsePinComboIndex, bool bUseDataComboIndex) const;
 	void SaveData(bool bUsePinIndex, bool bUseDataIndex);
+	bool CanSaveData(bool bUsePinIndex, bool bUseDataIndex) const;
 
 	TSharedRef<ITableRow> OnGenerateRow(PCGListviewItemPtr Item, const TSharedRef<STableViewBase>& OwnerTable);
 	void OnItemDoubleClicked(PCGListviewItemPtr Item) const;
