Fixed inline edit condition toggle that shouldn't appear in the component selector struct

#rnx

diff --git a/Engine/Plugins/PCG/Source/PCG/Public/Elements/PCGActorSelector.h b/Engine/Plugins/PCG/Source/PCG/Public/Elements/PCGActorSelector.h
--- a/Engine/Plugins/PCG/Source/PCG/Public/Elements/PCGActorSelector.h
+++ b/Engine/Plugins/PCG/Source/PCG/Public/Elements/PCGActorSelector.h
@@ -186,7 +186,7 @@ struct PCG_API FPCGComponentSelectorSettings
 	GENERATED_BODY()
 
 	/** How to select when filtering actors. */
-	UPROPERTY(BlueprintReadWrite, EditAnywhere, Category = "Component Selector Settings", meta = (EditCondition = "bShowComponentSelection", EditConditionHides))
+	UPROPERTY(BlueprintReadWrite, EditAnywhere, Category = "Component Selector Settings", meta = (EditCondition = "bShowComponentSelection", EditConditionHides, HideEditConditionToggle))
 	EPCGComponentSelection ComponentSelection = EPCGComponentSelection::ByTag;
 
 	/** Tag to match against when filtering actors. */
