More folder options for generated actors

diff --git a/Engine/Plugins/PCG/Source/PCG/Private/Elements/PCGCreateTargetActor.cpp b/Engine/Plugins/PCG/Source/PCG/Private/Elements/PCGCreateTargetActor.cpp
--- a/Engine/Plugins/PCG/Source/PCG/Private/Elements/PCGCreateTargetActor.cpp
+++ b/Engine/Plugins/PCG/Source/PCG/Private/Elements/PCGCreateTargetActor.cpp
@@ -326,7 +326,7 @@ bool FPCGCreateTargetActorElement::ExecuteInternal(FPCGContext* Context) const
 	}
 
 	// Always attach if root actor is provided
-	PCGHelpers::AttachToParent(GeneratedActor, TargetActor, Settings->RootActor.Get() ? EPCGAttachOptions::Attached : Settings->AttachOptions);
+	PCGHelpers::AttachToParent(GeneratedActor, TargetActor, Settings->RootActor.Get() ? EPCGAttachOptions::Attached : Settings->AttachOptions, Context);
 
 #if WITH_EDITOR
 	if (Settings->ActorLabel != FString())
diff --git a/Engine/Plugins/PCG/Source/PCG/Private/Elements/PCGSpawnActor.cpp b/Engine/Plugins/PCG/Source/PCG/Private/Elements/PCGSpawnActor.cpp
--- a/Engine/Plugins/PCG/Source/PCG/Private/Elements/PCGSpawnActor.cpp
+++ b/Engine/Plugins/PCG/Source/PCG/Private/Elements/PCGSpawnActor.cpp
@@ -892,7 +892,7 @@ void FPCGSpawnActorElement::SpawnActors(FPCGSubgraphContext* Context, AActor* Ta
 		// If generated actors are not directly attached, place them in a subfolder for tidiness.
 		FString GeneratedActorsFolderPath;
 #if WITH_EDITOR
-		PCGHelpers::GetGeneratedActorsFolderPath(TargetActor, GeneratedActorsFolderPath);
+		PCGHelpers::GetGeneratedActorsFolderPath(TargetActor, Context, Settings->AttachOptions, GeneratedActorsFolderPath);
 #endif
 
 		const UFunction* FunctionPrototypeWithNoParams = UPCGFunctionPrototypes::GetPrototypeWithNoParams();
@@ -924,7 +924,7 @@ void FPCGSpawnActorElement::SpawnActors(FPCGSubgraphContext* Context, AActor* Ta
 			// HACK: until UE-62747 is fixed, we have to force set the scale after spawning the actor
 			GeneratedActor->SetActorRelativeScale3D(Point.Transform.GetScale3D());
 			GeneratedActor->Tags.Append(NewActorTags);
-			PCGHelpers::AttachToParent(GeneratedActor, TargetActor, Settings->AttachOptions, GeneratedActorsFolderPath);
+			PCGHelpers::AttachToParent(GeneratedActor, TargetActor, Settings->AttachOptions, Context, GeneratedActorsFolderPath);
 
 			for (UFunction* PostSpawnFunction : PostSpawnFunctions)
 			{
diff --git a/Engine/Plugins/PCG/Source/PCG/Private/Helpers/PCGHelpers.cpp b/Engine/Plugins/PCG/Source/PCG/Private/Helpers/PCGHelpers.cpp
--- a/Engine/Plugins/PCG/Source/PCG/Private/Helpers/PCGHelpers.cpp
+++ b/Engine/Plugins/PCG/Source/PCG/Private/Helpers/PCGHelpers.cpp
@@ -501,9 +501,35 @@ namespace PCGHelpers
 		GeneratedActorsFolder << InTargetActor->GetActorLabel() << "_Generated";
 		OutFolderPath = GeneratedActorsFolder;
 	}
+
+	void GetGeneratedActorsFolderPath(const AActor* InTargetActor, const FPCGContext* InContext, EPCGAttachOptions AttachOptions, FString& OutFolderPath)
+	{
+		if (AttachOptions == EPCGAttachOptions::Attached || AttachOptions == EPCGAttachOptions::NotAttached)
+		{
+			OutFolderPath = FString();
+		}
+		else if (AttachOptions == EPCGAttachOptions::InFolder)
+		{
+			GetGeneratedActorsFolderPath(InTargetActor, OutFolderPath);
+		}
+		else if (AttachOptions == EPCGAttachOptions::InGraphFolder && InContext && InContext->Stack && InContext->Stack->GetRootGraph())
+		{
+			OutFolderPath = InContext->Stack->GetRootGraph()->GetName() + "_Generated";
+		}
+		else // Generated folder
+		{
+			OutFolderPath = TEXT("PCG_Generated_Actors");
+		}
+	}
 #endif
 
+	// Note: deprecated
 	void AttachToParent(AActor* InActorToAttach, AActor* InParent, EPCGAttachOptions AttachOptions, const FString& InGeneratedPath)
+	{
+		AttachToParent(InActorToAttach, InParent, AttachOptions, nullptr, InGeneratedPath);
+	}
+
+	void AttachToParent(AActor* InActorToAttach, AActor* InParent, EPCGAttachOptions AttachOptions, const FPCGContext* InContext, const FString& InGeneratedPath)
 	{
 		if (!InParent)
 		{
@@ -515,13 +541,12 @@ namespace PCGHelpers
 			InActorToAttach->AttachToActor(InParent, FAttachmentTransformRules::KeepWorldTransform);
 		}
 #if WITH_EDITOR
-		else if (AttachOptions == EPCGAttachOptions::InFolder)
+		else if(AttachOptions != EPCGAttachOptions::NotAttached)
 		{
 			FString DefaultFolderPath;
-
 			if (InGeneratedPath.IsEmpty())
 			{
-				GetGeneratedActorsFolderPath(InParent, DefaultFolderPath);
+				GetGeneratedActorsFolderPath(InParent, InContext, AttachOptions, DefaultFolderPath);
 			}
 
 			const FString& FolderPath = (InGeneratedPath.IsEmpty() ? DefaultFolderPath : InGeneratedPath);
diff --git a/Engine/Plugins/PCG/Source/PCG/Private/PCGComponent.cpp b/Engine/Plugins/PCG/Source/PCG/Private/PCGComponent.cpp
--- a/Engine/Plugins/PCG/Source/PCG/Private/PCGComponent.cpp
+++ b/Engine/Plugins/PCG/Source/PCG/Private/PCGComponent.cpp
@@ -1359,9 +1359,27 @@ FPCGTaskId UPCGComponent::CreateCleanupTask(bool bRemoveComponents, const TArray
 			}
 		}
 
+		TSet<FName> DeletedActorFolders;
+
 		if (UWorld* World = WorldPtr.Get())
 		{
-			UPCGActorHelpers::DeleteActors(World, Context->ActorsToDelete.Array());
+			const TArray<TSoftObjectPtr<AActor>> ActorsToDelete = Context->ActorsToDelete.Array();
+
+#if WITH_EDITOR
+			for (const TSoftObjectPtr<AActor>& Actor : ActorsToDelete)
+			{
+				if (Actor.IsValid())
+				{
+					FName ActorFolderPath = Actor->GetFolderPath();
+					if (ActorFolderPath != NAME_None)
+					{
+						DeletedActorFolders.Add(ActorFolderPath);
+					}
+				}
+			}
+#endif
+
+			UPCGActorHelpers::DeleteActors(World, ActorsToDelete);
 		}
 
 		if (UPCGComponent* ThisComponent = ThisComponentWeakPtr.Get())
@@ -1369,21 +1387,16 @@ FPCGTaskId UPCGComponent::CreateCleanupTask(bool bRemoveComponents, const TArray
 #if WITH_EDITOR
 			if (UWorld* ThisWorld = ThisComponent->GetWorld(); ThisWorld && GEditor) // FActorFolders require the editor
 			{
-				// Look for a nested generated results subfolder and remove it if it exists
-				FString FolderPath;
-				PCGHelpers::GetGeneratedActorsFolderPath(ThisComponent->GetOwner(), FolderPath);
-
-				if (!FolderPath.IsEmpty())
+				for (FName FolderPath : DeletedActorFolders)
 				{
-					FFolder GeneratedFolder(FFolder::GetWorldRootFolder(ThisWorld).GetRootObject(), *FolderPath);
-
+					FFolder GeneratedFolder(FFolder::GetWorldRootFolder(ThisWorld).GetRootObject(), FolderPath);
 					const bool bFolderExists = GeneratedFolder.IsValid() && FActorFolders::Get().ContainsFolder(*ThisWorld, GeneratedFolder);
 					bool bFoundActors = false;
 
 					if (bFolderExists)
 					{
-						TArray<FName> Paths = { *FolderPath };
-						FActorFolders::ForEachActorInFolders(*ThisWorld, Paths, [&bFoundActors](AActor* InActor)
+						TSet<FName> Folders{ GeneratedFolder.GetPath() };
+						FActorFolders::ForEachActorInFolders(*ThisWorld, Folders, [&bFoundActors](AActor* InActor)
 						{
 							if (InActor)
 							{
diff --git a/Engine/Plugins/PCG/Source/PCG/Public/Helpers/PCGHelpers.h b/Engine/Plugins/PCG/Source/PCG/Public/Helpers/PCGHelpers.h
--- a/Engine/Plugins/PCG/Source/PCG/Public/Helpers/PCGHelpers.h
+++ b/Engine/Plugins/PCG/Source/PCG/Public/Helpers/PCGHelpers.h
@@ -76,10 +76,14 @@ namespace PCGHelpers
 
 #if WITH_EDITOR
 	PCG_API void GetGeneratedActorsFolderPath(const AActor* InTargetActor, FString& OutFolderPath);
+	PCG_API void GetGeneratedActorsFolderPath(const AActor* InTargetActor, const FPCGContext* InContext, EPCGAttachOptions AttachOptions, FString& OutFolderPath);
 #endif
 
+	UE_DEPRECATED(5.5, "This function has been deprecated in favor of the version with a context.")
 	PCG_API void AttachToParent(AActor* InActorToAttach, AActor* InParent, EPCGAttachOptions AttachOptions, const FString& GeneratedPath = FString());
 
+	PCG_API void AttachToParent(AActor* InActorToAttach, AActor* InParent, EPCGAttachOptions AttachOptions, const FPCGContext* InContext, const FString& GeneratedPath = FString());
+
 	/**
 	 * Finds functions on the actor matching the provided function names. Functions must be marked as CallInEditor
 	 * and have parameters matching one of the provided prototypes. Some prototypes are provided in UPCGFunctionPrototypes.
diff --git a/Engine/Plugins/PCG/Source/PCG/Public/PCGCommon.h b/Engine/Plugins/PCG/Source/PCG/Public/PCGCommon.h
--- a/Engine/Plugins/PCG/Source/PCG/Public/PCGCommon.h
+++ b/Engine/Plugins/PCG/Source/PCG/Public/PCGCommon.h
@@ -265,7 +265,9 @@ enum class EPCGAttachOptions : uint32
 {
 	NotAttached UMETA(Tooltip="Actor will not be attached to the target actor nor placed in an actor folder"),
 	Attached UMETA(Tooltip="Actor will be attached to the target actor in the given node"),
-	InFolder UMETA(Tooltip="Actor will be placed in an actor folder containing the name of the target actor.")
+	InFolder UMETA(Tooltip="Actor will be placed in an actor folder containing the name of the target actor."),
+	InGraphFolder UMETA(Tooltip="Actor will be placed in a folder named after the top graph it was generated from."),
+	InGeneratedFolder UMETA(Tooltip="Actor will be placed in the PCG_Generated folder.")
 };
 
 UENUM()
