[PCG] Select Grammar - Second Pass

+ Grammar Selection node
* Factorization of some grammar related things in Segment Slicer

diff --git a/Engine/Plugins/PCG/Source/PCG/Private/Elements/Grammar/PCGSegmentSlicer.cpp b/Engine/Plugins/PCG/Source/PCG/Private/Elements/Grammar/PCGSegmentSlicer.cpp
--- a/Engine/Plugins/PCG/Source/PCG/Private/Elements/Grammar/PCGSegmentSlicer.cpp
+++ b/Engine/Plugins/PCG/Source/PCG/Private/Elements/Grammar/PCGSegmentSlicer.cpp
@@ -25,8 +25,8 @@ public:
 		FPCGMetadataAttribute<bool>* IsFinalPointAttribute = nullptr;
 		FPCGMetadataAttribute<int32>* ExtremityNeighborIndexAttribute = nullptr;
 
-		FPCGSlicingBaseElement::FPCGModulesInfoMap ModulesInfo;
-		TMap<FString, TArray<FPCGTokenizedGrammar>> CachedGrammars;
+		PCGSlicingBase::FPCGModulesInfoMap ModulesInfo;
+		TMap<FString, PCGGrammar::FTokenizedGrammar> CachedModules;
 		TArray<int32> CornerIndexes;
 
 		FVector SlicingDirection;
@@ -39,17 +39,17 @@ public:
 		UPCGMetadata* OutputMetadata = nullptr;
 	};
 
-	static void Process(const FPCGSegmentSlicerElement* InElement, FParameters& InOutParameters, const FString& InGrammar, const bool bFlipAxis, int32 Index)
+	static void Process(FParameters& InOutParameters, const FString& InGrammar, const bool bFlipAxis, int32 Index)
 	{
-		if (!InOutParameters.CachedGrammars.Contains(InGrammar))
+		if (!InOutParameters.CachedModules.Contains(InGrammar))
 		{
 			double MinSize;
-			InOutParameters.CachedGrammars.Emplace(InGrammar, InElement->GetTokenizeGrammar(InOutParameters.Context, InGrammar, InOutParameters.ModulesInfo, MinSize));
+			InOutParameters.CachedModules.Emplace(InGrammar, PCGSlicingBase::GetTokenizedGrammar(InOutParameters.Context, InGrammar, InOutParameters.ModulesInfo, MinSize));
 		}
 
-		const TArray<FPCGTokenizedGrammar>& CurrentTokenizeGrammar = InOutParameters.CachedGrammars[InGrammar];
+		const PCGGrammar::FTokenizedGrammar& CurrentTokenizedGrammar = InOutParameters.CachedModules[InGrammar];
 
-		if (CurrentTokenizeGrammar.IsEmpty())
+		if (CurrentTokenizedGrammar.IsEmpty())
 		{
 			return;
 		}
@@ -61,24 +61,24 @@ public:
 		const int32 FirstModuleIndex = InOutParameters.OutPoints->Num();
 
 		FBox Segment = Point.GetLocalBounds();
-		FVector FullExtents = Point.GetScaledExtents() * 2.0;
+		FVector PointScaledSize = Point.GetScaledLocalSize();
 		if (bFlipAxis)
 		{
 			// Swap coordinates on the slicing direction
 			const FVector PreviousMin = Segment.Min;
 			Segment.Min = Segment.Min * InOutParameters.PerpendicularSlicingDirection + Segment.Max * InOutParameters.SlicingDirection;
 			Segment.Max = Segment.Max * InOutParameters.PerpendicularSlicingDirection + PreviousMin * InOutParameters.SlicingDirection;
-			FullExtents *= (InOutParameters.PerpendicularSlicingDirection - InOutParameters.SlicingDirection);
+			PointScaledSize *= (InOutParameters.PerpendicularSlicingDirection - InOutParameters.SlicingDirection);
 		}
 
 		const FVector Direction = TransformNoTranslation.TransformVectorNoScale(InOutParameters.SlicingDirection).GetSafeNormal();
-		const FVector OtherDirection = TransformNoTranslation.TransformVectorNoScale(FullExtents * InOutParameters.PerpendicularSlicingDirection) * 0.5;
-		const FVector HalfExtents2D = FullExtents * InOutParameters.PerpendicularSlicingDirection * 0.5;
-		const double Size = FullExtents.Dot(InOutParameters.SlicingDirection);
+		const FVector OtherDirection = TransformNoTranslation.TransformVectorNoScale(PointScaledSize * InOutParameters.PerpendicularSlicingDirection) * 0.5;
+		const FVector HalfExtents2D = PointScaledSize * InOutParameters.PerpendicularSlicingDirection * 0.5;
+		const double Size = PointScaledSize.Dot(InOutParameters.SlicingDirection);
 
-		TArray<PCGSlicingBase::PCGSubDivModuleInstance<FPCGTokenizedGrammar>> ModulesInstances;
+		TArray<PCGSlicingBase::TPCGSubDivModuleInstance<PCGGrammar::FTokenizedModule>> ModulesInstances;
 		double RemainingSubdivide;
-		const bool bSubdivideSuccess = PCGSlicingBase::Subdivide(CurrentTokenizeGrammar, Size, ModulesInstances, RemainingSubdivide, InOutParameters.Context);
+		const bool bSubdivideSuccess = PCGSlicingBase::Subdivide(CurrentTokenizedGrammar, Size, ModulesInstances, RemainingSubdivide, InOutParameters.Context);
 
 		if (!bSubdivideSuccess)
 		{
@@ -98,18 +98,18 @@ public:
 
 		for (int32 ModuleInstanceIndex = 0; ModuleInstanceIndex < ModulesInstances.Num(); ModuleInstanceIndex++)
 		{
-			const PCGSlicingBase::PCGSubDivModuleInstance<FPCGTokenizedGrammar>& ModuleInstance = ModulesInstances[ModuleInstanceIndex];
+			const PCGSlicingBase::TPCGSubDivModuleInstance<PCGGrammar::FTokenizedModule>& ModuleInstance = ModulesInstances[ModuleInstanceIndex];
 			for (int32 j = 0; j < ModuleInstance.NumRepeat; ++j)
 			{
 				for (int32 SymbolIndex = 0; SymbolIndex < ModuleInstance.Module->Symbols.Num(); ++SymbolIndex)
 				{
 					const FName Symbol = ModuleInstance.Module->Symbols[SymbolIndex];
 					const FVector Scale = FVector::OneVector + (InOutParameters.SlicingDirection * ModuleInstance.ExtraScales[SymbolIndex]);
-					const FPCGSlicingModule& SlicingModule = InOutParameters.ModulesInfo[Symbol];
+					const FPCGSlicingSubmodule& SlicingSubmodule = InOutParameters.ModulesInfo[Symbol];
 
 					const bool bIsFirstModule = (ModuleInstanceIndex == 0) && (j == 0) && (SymbolIndex == 0);
 					const bool bIsFinalModule = (ModuleInstanceIndex == ModulesInstances.Num() - 1) && (j == ModuleInstance.NumRepeat - 1) && (SymbolIndex == ModuleInstance.Module->Symbols.Num() - 1);
-					const double HalfDisplacement = SlicingModule.Size * 0.5;
+					const double HalfDisplacement = SlicingSubmodule.Size * 0.5;
 					const double HalfScaledDisplacement = Scale.Dot(InOutParameters.SlicingDirection) * HalfDisplacement;
 
 					const FVector LocalBoundsExtents = InOutParameters.SlicingDirection * HalfDisplacement + HalfExtents2D;
@@ -128,7 +128,7 @@ public:
 
 					if (InOutParameters.DebugColorAttribute)
 					{
-						InOutParameters.DebugColorAttribute->SetValue(OutPoint.MetadataEntry, FVector4(SlicingModule.DebugColor, 1.0));
+						InOutParameters.DebugColorAttribute->SetValue(OutPoint.MetadataEntry, FVector4(SlicingSubmodule.DebugColor, 1.0));
 					}
 
 					if (InOutParameters.ModuleIndexAttribute)
@@ -244,14 +244,14 @@ bool FPCGSegmentSlicerElement::ExecuteInternal(FPCGContext* InContext) const
 		{
 			continue;
 		}
-		
+
 		TUniquePtr<const IPCGAttributeAccessor> GrammarAccessor;
 		TUniquePtr<const IPCGAttributeAccessor> FlipAxisAccessor;
 		TUniquePtr<const IPCGAttributeAccessorKeys> Keys;
 
-		if (Settings->bGrammarAsAttribute)
+		if (Settings->GrammarSelection.bGrammarAsAttribute)
 		{
-			const FPCGAttributePropertyInputSelector Selector = Settings->GrammarAttribute.CopyAndFixLast(InputPointData);
+			const FPCGAttributePropertyInputSelector Selector = Settings->GrammarSelection.GrammarAttribute.CopyAndFixLast(InputPointData);
 			GrammarAccessor = PCGAttributeAccessorHelpers::CreateConstAccessor(InputPointData, Selector);
 			Keys = PCGAttributeAccessorHelpers::CreateConstKeys(InputPointData, Selector);
 			if (!GrammarAccessor || !Keys)
@@ -312,29 +312,29 @@ bool FPCGSegmentSlicerElement::ExecuteInternal(FPCGContext* InContext) const
 		}
 
 
-		if (Settings->bGrammarAsAttribute && Settings->bFlipAxisAsAttribute)
+		if (Settings->GrammarSelection.bGrammarAsAttribute && Settings->bFlipAxisAsAttribute)
 		{
-			auto Process = [this, &Parameters](const FString& InGrammar, const bool bFlipAxis, int32 Index) -> void
+			auto Process = [&Parameters](const FString& InGrammar, const bool bFlipAxis, int32 Index) -> void
 			{
-				PCGSegmentSlicerHelpers::Process(this, Parameters, InGrammar, bFlipAxis, Index);
+				PCGSegmentSlicerHelpers::Process(Parameters, InGrammar, bFlipAxis, Index);
 			};
 
 			PCGMetadataElementCommon::ApplyOnMultiAccessors<FString, bool>(*Keys, { GrammarAccessor.Get(), FlipAxisAccessor.Get() }, Process);
 		}
-		else if (Settings->bGrammarAsAttribute)
+		else if (Settings->GrammarSelection.bGrammarAsAttribute)
 		{
-			auto Process = [this, &Parameters, bShouldFlipAxis = Settings->bShouldFlipAxis](const FString& InGrammar, int32 Index) -> void
+			auto Process = [&Parameters, bShouldFlipAxis = Settings->bShouldFlipAxis](const FString& InGrammar, int32 Index) -> void
 			{
-				PCGSegmentSlicerHelpers::Process(this, Parameters, InGrammar, bShouldFlipAxis, Index);
+				PCGSegmentSlicerHelpers::Process(Parameters, InGrammar, bShouldFlipAxis, Index);
 			};
 
 			PCGMetadataElementCommon::ApplyOnAccessor<FString>(*Keys, *GrammarAccessor, Process);
 		}
 		else if (Settings->bFlipAxisAsAttribute)
 		{
-			auto Process = [this, &Parameters, Grammar = Settings->Grammar](const bool bFlipAxis, int32 Index) -> void
+			auto Process = [&Parameters, Grammar = Settings->GrammarSelection.GrammarString](const bool bFlipAxis, int32 Index) -> void
 			{
-				PCGSegmentSlicerHelpers::Process(this, Parameters, Grammar, bFlipAxis, Index);
+				PCGSegmentSlicerHelpers::Process(Parameters, Grammar, bFlipAxis, Index);
 			};
 
 			PCGMetadataElementCommon::ApplyOnAccessor<bool>(*Keys, *FlipAxisAccessor, Process);
@@ -343,7 +343,7 @@ bool FPCGSegmentSlicerElement::ExecuteInternal(FPCGContext* InContext) const
 		{
 			for (int32 SegmentIndex = 0; SegmentIndex < Parameters.InPoints->Num(); ++SegmentIndex)
 			{
-				PCGSegmentSlicerHelpers::Process(this, Parameters, Settings->Grammar, Settings->bShouldFlipAxis, SegmentIndex);
+				PCGSegmentSlicerHelpers::Process(Parameters, Settings->GrammarSelection.GrammarString, Settings->bShouldFlipAxis, SegmentIndex);
 			}
 		}
 
diff --git a/Engine/Plugins/PCG/Source/PCG/Private/Elements/Grammar/PCGSelectGrammar.cpp b/Engine/Plugins/PCG/Source/PCG/Private/Elements/Grammar/PCGSelectGrammar.cpp
new file mode 100644
--- /dev/null
+++ b/Engine/Plugins/PCG/Source/PCG/Private/Elements/Grammar/PCGSelectGrammar.cpp
@@ -0,0 +1,320 @@
+// Copyright Epic Games, Inc. All Rights Reserved.
+
+#include "Elements/Grammar/PCGSelectGrammar.h"
+
+#include "PCGContext.h"
+#include "PCGParamData.h"
+#include "Data/PCGPointData.h"
+#include "Elements/Metadata/PCGMetadataElementCommon.h"
+#include "Helpers/PCGPropertyHelpers.h"
+
+#define LOCTEXT_NAMESPACE "PCGSelectGrammarElement"
+
+namespace PCGSelectGrammarConstants
+{
+	const FName CriteriaDataPinLabel = TEXT("Selection Criteria Data");
+	const FName OutputGrammarAttributeName = TEXT("Grammar");
+	// Most likely used value for comparison with grammars
+	const FString ComparedValueAttributeString = TEXT("$ScaledLocalSize.X");
+}
+
+namespace PCGSelectGrammar
+{
+	struct FGrammarSizePair
+	{
+		FGrammarSizePair(const FString& InGrammar, double InSize)
+			: Grammar(InGrammar)
+			, Size(InSize)
+		{
+		}
+
+		FString Grammar;
+		double Size;
+	};
+
+	constexpr uint16 BinaryEnumIndex = static_cast<uint16>(EPCGSelectGrammarComparator::BinaryOps);
+	constexpr uint16 TernaryEnumIndex = static_cast<uint16>(EPCGSelectGrammarComparator::TernaryOps);
+
+	// Function indices must match EPCGSelectGrammarComparator indices directly
+	template <typename T>
+	using TBinarySignature = bool(*)(const T&, const T&);
+
+	template <typename T>
+	static constexpr TBinarySignature<T> BinaryCompare[] =
+	{
+		[](const T&, const T&) { return true; }, // Select
+		PCG::Private::MetadataTraits<T>::Less,
+		PCG::Private::MetadataTraits<T>::LessOrEqual,
+		PCG::Private::MetadataTraits<T>::Equal,
+		PCG::Private::MetadataTraits<T>::GreaterOrEqual,
+		PCG::Private::MetadataTraits<T>::Greater,
+	};
+
+	template <typename T>
+	using TTernarySignature = bool(*)(const T&, const T&, const T&);
+
+	template <typename T>
+	static constexpr TTernarySignature<T> TernaryCompare[] =
+	{
+		// Exclusive range
+		[](const T& Value, const T& Start, const T& End) { return PCG::Private::MetadataTraits<T>::Greater(Value, Start) && PCG::Private::MetadataTraits<T>::Less(Value, End); },
+		// Inclusive range
+		[](const T& Value, const T& Start, const T& End) { return PCG::Private::MetadataTraits<T>::GreaterOrEqual(Value, Start) && PCG::Private::MetadataTraits<T>::LessOrEqual(Value, End); }
+	};
+}
+
+UPCGSelectGrammarSettings::UPCGSelectGrammarSettings()
+{
+	ComparedValueAttribute.Update(PCGSelectGrammarConstants::ComparedValueAttributeString);
+	OutputGrammarAttribute.SetAttributeName(PCGSelectGrammarConstants::OutputGrammarAttributeName);
+}
+
+#if WITH_EDITOR
+FName UPCGSelectGrammarSettings::GetDefaultNodeName() const
+{
+	return FName(TEXT("SelectGrammar"));
+}
+
+FText UPCGSelectGrammarSettings::GetDefaultNodeTitle() const
+{
+	return LOCTEXT("NodeTitle", "Select Grammar");
+}
+
+FText UPCGSelectGrammarSettings::GetNodeTooltipText() const
+{
+	return LOCTEXT("NodeTooltip", "Select a grammar by comparing an input attribute against a provided criteria.");
+}
+#endif // WITH_EDITOR
+
+FPCGElementPtr UPCGSelectGrammarSettings::CreateElement() const
+{
+	return MakeShared<FPCGSelectGrammarElement>();
+}
+
+TArray<FPCGPinProperties> UPCGSelectGrammarSettings::InputPinProperties() const
+{
+	TArray<FPCGPinProperties> PinProperties;
+
+	FPCGPinProperties& InputPin = PinProperties.Emplace_GetRef(PCGPinConstants::DefaultInputLabel, EPCGDataType::Point);
+	InputPin.SetRequiredPin();
+#if WITH_EDITOR
+	InputPin.Tooltip = LOCTEXT("InputPinTooltip", "A grammar will be selected based on criteria for all points.");
+#endif // WITH_EDITOR
+
+	if (bCriteriaAsInput)
+	{
+		FPCGPinProperties& CriteriaPin = PinProperties.Emplace_GetRef(PCGSelectGrammarConstants::CriteriaDataPinLabel, EPCGDataType::Param, false, false);
+		CriteriaPin.SetRequiredPin();
+#if WITH_EDITOR
+		CriteriaPin.Tooltip = LOCTEXT("CriteriaPinTooltip", "Criteria to compare against an input value for conditionally selecting a grammar. Will be evaluated in sequential order.");
+#endif // WITH_EDITOR
+	}
+
+	return PinProperties;
+}
+
+TArray<FPCGPinProperties> UPCGSelectGrammarSettings::OutputPinProperties() const
+{
+	TArray<FPCGPinProperties> PinProperties;
+	PinProperties.Emplace_GetRef(PCGPinConstants::DefaultOutputLabel, EPCGDataType::Point);
+
+	return PinProperties;
+}
+
+bool FPCGSelectGrammarElement::ExecuteInternal(FPCGContext* InContext) const
+{
+	TRACE_CPUPROFILER_EVENT_SCOPE(FPCGSelectGrammarElement::Execute);
+
+	const UPCGSelectGrammarSettings* Settings = InContext->GetInputSettings<UPCGSelectGrammarSettings>();
+	check(Settings);
+
+	const TArray<FPCGTaggedData> PointInputs = InContext->InputData.GetInputsByPin(PCGPinConstants::DefaultInputLabel);
+	const TArray<FPCGTaggedData> CriteriaDataInputs = InContext->InputData.GetInputsByPin(PCGSelectGrammarConstants::CriteriaDataPinLabel);
+
+	if (PointInputs.IsEmpty() || (Settings->bCriteriaAsInput && CriteriaDataInputs.IsEmpty()))
+	{
+		return true;
+	}
+
+	TArray<FPCGSelectGrammarCriterion> Criteria = Settings->bCriteriaAsInput ? TArray<FPCGSelectGrammarCriterion>{} : Settings->Criteria;
+	if (Settings->bCriteriaAsInput)
+	{
+		if (CriteriaDataInputs.Num() != 1)
+		{
+			PCGLog::InputOutput::LogFirstInputOnlyWarning(PCGSelectGrammarConstants::CriteriaDataPinLabel, InContext);
+		}
+
+		// Should only be one, so take the first
+		if (const UPCGParamData* CriteriaParamData = Cast<UPCGParamData>(CriteriaDataInputs[0].Data))
+		{
+			check(CriteriaParamData->Metadata);
+
+			const FPCGSelectGrammarCriteriaAttributeNames& CriteriaAttributeNames =
+					Settings->bRemapCriteriaAttributeNames
+						? Settings->CriteriaAttributeNames
+						: FPCGSelectGrammarCriteriaAttributeNames();
+
+			// Attempt an extraction of attribute set matching 'FPCGSelectGrammarCriterion'
+			TMap<FName, TTuple<FName, bool>> PropertyNameMapping;
+			PropertyNameMapping.Emplace(GET_MEMBER_NAME_CHECKED(FPCGSelectGrammarCriterion, Key), {CriteriaAttributeNames.KeyAttributeName, /*bCanBeDefaulted=*/true});
+			PropertyNameMapping.Emplace(GET_MEMBER_NAME_CHECKED(FPCGSelectGrammarCriterion, FirstValue), {CriteriaAttributeNames.FirstValueAttributeName, /*bCanBeDefaulted=*/true});
+			PropertyNameMapping.Emplace(GET_MEMBER_NAME_CHECKED(FPCGSelectGrammarCriterion, Comparator), {CriteriaAttributeNames.ComparatorAttributeName, /*bCanBeDefaulted=*/true});
+			PropertyNameMapping.Emplace(GET_MEMBER_NAME_CHECKED(FPCGSelectGrammarCriterion, SecondValue), {CriteriaAttributeNames.SecondValueAttributeName, /*bCanBeDefaulted=*/true});
+			PropertyNameMapping.Emplace(GET_MEMBER_NAME_CHECKED(FPCGSelectGrammarCriterion, Grammar), {CriteriaAttributeNames.GrammarAttributeName, /*bCanBeDefaulted=*/true});
+
+			Criteria = PCGPropertyHelpers::ExtractAttributeSetAsArrayOfStructs<FPCGSelectGrammarCriterion>(CriteriaParamData, &PropertyNameMapping, InContext);
+		}
+		else
+		{
+			PCGLog::InputOutput::LogTypedDataNotFoundWarning(EPCGDataType::Param, PCGSelectGrammarConstants::CriteriaDataPinLabel, InContext);
+			return true;
+		}
+	}
+
+	// Early out with nothing to compare to
+	if (Criteria.IsEmpty())
+	{
+		return true;
+	}
+
+	// Validate the criteria, to ensure the ranges are in order, etc
+	for (const FPCGSelectGrammarCriterion& Criterion : Criteria)
+	{
+		switch (Criterion.Comparator)
+		{
+			case EPCGSelectGrammarComparator::RangeExclusive: // fall-through
+			case EPCGSelectGrammarComparator::RangeInclusive:
+				{
+					if (Criterion.FirstValue > Criterion.SecondValue)
+					{
+						PCGLog::LogErrorOnGraph(FText::Format(LOCTEXT("InvalidCriteriaRange", "Criteria range is inverted for criterion with Key->Grammar: '{0}'->'{1}'"), FText::FromName(Criterion.Key), FText::FromString(Criterion.Grammar)));
+						return true;
+					}
+				}
+				break;
+			default:
+				break;
+		}
+	}
+
+	// Pre-process the criteria into a map for O(1) access later
+	TMap<FName, TArray<FPCGSelectGrammarCriterion, TInlineAllocator<16>>> CriteriaMap;
+	for (FPCGSelectGrammarCriterion& Criterion : Criteria)
+	{
+		if (!CriteriaMap.Contains(Criterion.Key))
+		{
+			CriteriaMap.Emplace(Criterion.Key, {});
+		}
+
+		CriteriaMap[Criterion.Key].Emplace(std::move(Criterion));
+	}
+
+	TArray<FPCGTaggedData>& Outputs = InContext->OutputData.TaggedData;
+	for (const FPCGTaggedData& PointInput : PointInputs)
+	{
+		const UPCGPointData* InputPointData = Cast<UPCGPointData>(PointInput.Data);
+		const int32 NumPoints = InputPointData->GetPoints().Num();
+		if (!InputPointData || NumPoints < 1)
+		{
+			continue;
+		}
+
+		UPCGData* OutputData = InputPointData->DuplicateData(InContext);
+		UPCGPointData* OutputPointData = CastChecked<UPCGPointData>(OutputData);
+		check(OutputPointData->Metadata)
+
+		Outputs.Emplace_GetRef(PointInput).Data = OutputData;
+
+		FPCGAttributePropertyInputSelector KeySelector = Settings->KeyAttribute.CopyAndFixLast(InputPointData);
+		FPCGAttributePropertyInputSelector ComparedValueSelector = Settings->ComparedValueAttribute.CopyAndFixLast(InputPointData);
+
+		TUniquePtr<const IPCGAttributeAccessor> InputKeyAccessor = PCGAttributeAccessorHelpers::CreateConstAccessor(InputPointData, KeySelector);
+		TUniquePtr<const IPCGAttributeAccessorKeys> InputKeys = PCGAttributeAccessorHelpers::CreateConstKeys(InputPointData, KeySelector);
+		if (!InputKeyAccessor || !InputKeys)
+		{
+			PCGLog::Accessor::LogFailToCreate(KeySelector);
+			continue;
+		}
+
+		TUniquePtr<const IPCGAttributeAccessor> ComparedValuesAccessor = PCGAttributeAccessorHelpers::CreateConstAccessor(InputPointData, ComparedValueSelector);
+		if (!ComparedValuesAccessor)
+		{
+			PCGLog::Accessor::LogFailToCreate(ComparedValueSelector);
+			continue;
+		}
+
+		FPCGAttributePropertyOutputSelector OutputGrammarSelector = Settings->OutputGrammarAttribute.CopyAndFixSource(&KeySelector);
+		FPCGMetadataAttribute<FString>* OutputAttribute = OutputPointData->Metadata->FindOrCreateAttribute(OutputGrammarSelector.GetName(), FString{}, /*bAllowsInterpolation=*/false);
+
+		TUniquePtr<IPCGAttributeAccessor> WriteGrammarAccessor = PCGAttributeAccessorHelpers::CreateAccessor(OutputAttribute, OutputPointData->Metadata);
+		TUniquePtr<IPCGAttributeAccessorKeys> OutputKeys = PCGAttributeAccessorHelpers::CreateKeys(OutputPointData, OutputGrammarSelector);
+		if (!WriteGrammarAccessor || !OutputKeys)
+		{
+			PCGLog::Accessor::LogFailToCreate(OutputGrammarSelector);
+			continue;
+		}
+
+		TArray Accessors = {InputKeyAccessor.Get(), ComparedValuesAccessor.Get()};
+		// TODO: For now, cast everything to double for numerical comparison, but this could easily expand in the future
+		PCGMetadataElementCommon::ApplyOnMultiAccessorsRange<FName, double>(*InputKeys, Accessors, [&WriteGrammarAccessor, &OutputKeys, &CriteriaMap, bPassThroughKey = Settings->bCopyKeyForUnselectedGrammar](const TArrayView<FName>& KeyValues, const TArrayView<double>& CompareValues, int32 Start, int32 Range)
+		{
+			TArray<FString, TInlineAllocator<PCGMetadataElementCommon::DefaultChunkSize>> Grammars;
+			Grammars.SetNum(Range);
+
+			// Iterate through this chunk of points and evaluate the selection criteria for each point
+			for (int32 i = 0; i < Range; ++i)
+			{
+				bool bFoundGrammar = false;
+
+				const FName& Key = KeyValues[i];
+				if (!CriteriaMap.Contains(Key))
+				{
+					continue;
+				}
+
+				const TArray<FPCGSelectGrammarCriterion, TInlineAllocator<16>>& CriteriaForThisKey = CriteriaMap[Key];
+				for (const FPCGSelectGrammarCriterion& Criterion : CriteriaForThisKey)
+				{
+					check(Key == Criterion.Key);
+
+					// Find the comparison function based on the enum selection
+					const uint16 EnumIndex = static_cast<uint16>(Criterion.Comparator);
+					if (EnumIndex < PCGSelectGrammar::TernaryEnumIndex) // Binary comparison
+					{
+						// Convert enum index to static function array index
+						const uint16 CompIndex = EnumIndex - PCGSelectGrammar::BinaryEnumIndex;
+						check(CompIndex < std::size(PCGSelectGrammar::BinaryCompare<double>))
+						if (PCGSelectGrammar::BinaryCompare<double>[CompIndex](CompareValues[i], Criterion.FirstValue))
+						{
+							bFoundGrammar = true;
+							Grammars[i] = Criterion.Grammar;
+							break;
+						}
+					}
+					else // Ternary comparison
+					{
+						const uint16 CompIndex = EnumIndex - PCGSelectGrammar::TernaryEnumIndex;
+						check(CompIndex < std::size(PCGSelectGrammar::TernaryCompare<double>))
+						if (PCGSelectGrammar::TernaryCompare<double>[CompIndex](CompareValues[i], Criterion.FirstValue, Criterion.SecondValue))
+						{
+							bFoundGrammar = true;
+							Grammars[i] = Criterion.Grammar;
+							break;
+						}
+					}
+				}
+
+				if (bPassThroughKey && !bFoundGrammar)
+				{
+					Grammars[i] = KeyValues[i].ToString();
+				}
+			}
+
+			WriteGrammarAccessor->SetRange<FString>(Grammars, Start, *OutputKeys);
+		});
+	}
+
+	return true;
+}
+
+#undef LOCTEXT_NAMESPACE
diff --git a/Engine/Plugins/PCG/Source/PCG/Private/Elements/Grammar/PCGSlicingBase.cpp b/Engine/Plugins/PCG/Source/PCG/Private/Elements/Grammar/PCGSlicingBase.cpp
--- a/Engine/Plugins/PCG/Source/PCG/Private/Elements/Grammar/PCGSlicingBase.cpp
+++ b/Engine/Plugins/PCG/Source/PCG/Private/Elements/Grammar/PCGSlicingBase.cpp
@@ -7,6 +7,7 @@
 #include "Data/PCGPointData.h"
 #include "Elements/Metadata/PCGMetadataElementCommon.h"
 #include "Grammar/PCGGrammarParser.h"
+#include "Helpers/PCGPropertyHelpers.h"
 #include "Metadata/Accessors/IPCGAttributeAccessor.h"
 #include "Metadata/Accessors/PCGAttributeAccessorHelpers.h"
 #include "Metadata/Accessors/PCGAttributeAccessorKeys.h"
@@ -15,137 +16,133 @@
 
 namespace PCGSlicingBase
 {
-	static const FText AccessorFailGetText = LOCTEXT("AccessorFailGet", "Couldn't retrieve attribute {0} value, could it be the wrong type? Expected type: {1}, Attribute Type: {2}.");
-	static const FText AccessorGrammarFailCreate = LOCTEXT("GrammarAccessor", "Attribute {0} was not found for the grammar.");
 	static const FText DuplicatedSymbolText = LOCTEXT("SymbolDuplicate", "Symbol {0} is duplicated, ignored.");
 }
 
-FPCGSlicingBaseElement::FPCGModulesInfoMap FPCGSlicingBaseElement::GetModulesInfoMap(FPCGContext* InContext, const UPCGSlicingBaseSettings* InSettings, const UPCGParamData*& OutModuleInfoParamData) const
+void UPCGSlicingBaseSettings::PostLoad()
 {
-	FPCGModulesInfoMap ModulesInfo;
-	OutModuleInfoParamData = nullptr;
+	Super::PostLoad();
 
-	if (InSettings->bModuleInfoAsInput)
+#if WITH_EDITOR
+	PRAGMA_DISABLE_DEPRECATION_WARNINGS
+	if (bGrammarAsAttribute_DEPRECATED != false)
 	{
-		const TArray<FPCGTaggedData> ModulesInfoInputs = InContext->InputData.GetInputsByPin(PCGSlicingBaseConstants::ModulesInfoPinLabel);
+		GrammarSelection.bGrammarAsAttribute = bGrammarAsAttribute_DEPRECATED;
+		bGrammarAsAttribute_DEPRECATED = false;
+	}
 
-		if (ModulesInfoInputs.IsEmpty())
-		{
-			PCGLog::LogWarningOnGraph(LOCTEXT("NoModuleInfo", "No data was found on the module info pin."), InContext);
-			return ModulesInfo;
-		}
+	if (!Grammar_DEPRECATED.IsEmpty())
+	{
+		GrammarSelection.GrammarString = Grammar_DEPRECATED;
+		Grammar_DEPRECATED.Empty();
+	}
+	PRAGMA_ENABLE_DEPRECATION_WARNINGS
+#endif // WITH_EDITOR
+}
 
-		const UPCGParamData* ParamData = Cast<UPCGParamData>(ModulesInfoInputs[0].Data);
-		if (!ParamData)
+FPCGSlicingBaseElement::FPCGModulesInfoMap FPCGSlicingBaseElement::GetModulesInfoMap(FPCGContext* InContext, const TArray<FPCGSlicingSubmodule>& SubmodulesInfo, const UPCGParamData*& OutModuleInfoParamData) const
+{
+	PCGSlicingBase::FPCGModulesInfoMap ModulesInfo;
+	OutModuleInfoParamData = nullptr;
+
+	ModulesInfo.Reserve(SubmodulesInfo.Num());
+	for (const FPCGSlicingSubmodule& SlicingModule : SubmodulesInfo)
+	{
+		if (ModulesInfo.Contains(SlicingModule.Symbol))
 		{
-			PCGLog::LogWarningOnGraph(LOCTEXT("ModuleInfoWrongType", "Module info input is not of type attribute set."), InContext);
-			return ModulesInfo;
+			PCGLog::LogWarningOnGraph(FText::Format(PCGSlicingBase::DuplicatedSymbolText, FText::FromName(SlicingModule.Symbol)), InContext);
+			continue;
 		}
 
-		const TUniquePtr<const IPCGAttributeAccessorKeys> Keys = PCGAttributeAccessorHelpers::CreateConstKeys(ParamData, FPCGAttributePropertySelector::CreateAttributeSelector(InSettings->ModulesInfoAttributeNames.SymbolAttributeName));
+		ModulesInfo.Emplace(SlicingModule.Symbol, SlicingModule);
+	}
 
-		auto GetRange = [InContext, &Keys, ParamData]<typename T>(TArray<T>& OutValues, const FName AttributeName, const bool bProvided = true)
-		{
-			if (!bProvided)
-			{
-				return true;
-			}
+	return ModulesInfo;
+}
 
-			const TUniquePtr<const IPCGAttributeAccessor> Accessor = PCGAttributeAccessorHelpers::CreateConstAccessor(ParamData, FPCGAttributePropertySelector::CreateAttributeSelector(AttributeName));
-			if (!Accessor || !Keys)
-			{
-				PCGLog::LogErrorOnGraph(FText::Format(LOCTEXT("ModuleInfoMissingAttribute", "Module info input is missing attribute {0}."), FText::FromName(AttributeName)), InContext);
-				return false;
-			}
+FPCGSlicingBaseElement::FPCGModulesInfoMap FPCGSlicingBaseElement::GetModulesInfoMap(FPCGContext* InContext, const FPCGSlicingModuleAttributeNames& InSlicingModuleAttributeNames, const UPCGParamData*& OutModuleInfoParamData) const
+{
+	PCGSlicingBase::FPCGModulesInfoMap ModulesInfo;
+	OutModuleInfoParamData = nullptr;
 
-			OutValues.SetNum(Keys->GetNum());
+	const TArray<FPCGTaggedData> ModulesInfoInputs = InContext->InputData.GetInputsByPin(PCGSlicingBaseConstants::ModulesInfoPinLabel);
 
-			if (!Accessor->GetRange<T>(OutValues, 0, *Keys, EPCGAttributeAccessorFlags::AllowBroadcastAndConstructible))
-			{
-				PCGLog::LogErrorOnGraph(FText::Format(PCGSlicingBase::AccessorFailGetText, FText::FromName(AttributeName), PCG::Private::GetTypeNameText<T>(), PCG::Private::GetTypeNameText(Accessor->GetUnderlyingType())), InContext);
-				return false;
-			}
+	if (ModulesInfoInputs.IsEmpty())
+	{
+		PCGLog::LogWarningOnGraph(LOCTEXT("NoModuleInfo", "No data was found on the module info pin."), InContext);
+		return ModulesInfo;
+	}
 
-			return true;
-		};
+	const UPCGParamData* ParamData = Cast<UPCGParamData>(ModulesInfoInputs[0].Data);
+	if (!ParamData)
+	{
+		PCGLog::LogWarningOnGraph(LOCTEXT("ModuleInfoWrongType", "Module info input is not of type attribute set."), InContext);
+		return ModulesInfo;
+	}
 
-		TArray<FName> Symbols;
-		TArray<double> Sizes;
-		TArray<bool> Scalables;
-		TArray<FVector4> DebugColors;
+	TMap<FName, TTuple<FName, bool>> PropertyNameMapping;
+	PropertyNameMapping.Emplace(GET_MEMBER_NAME_CHECKED(FPCGSlicingSubmodule, Symbol), {InSlicingModuleAttributeNames.SymbolAttributeName, /*bCanBeDefaulted=*/false});
+	PropertyNameMapping.Emplace(GET_MEMBER_NAME_CHECKED(FPCGSlicingSubmodule, Size), {InSlicingModuleAttributeNames.SizeAttributeName, /*bCanBeDefaulted=*/false});
+	PropertyNameMapping.Emplace(GET_MEMBER_NAME_CHECKED(FPCGSlicingSubmodule, bScalable), {InSlicingModuleAttributeNames.ScalableAttributeName, /*bCanBeDefaulted=*/!InSlicingModuleAttributeNames.bProvideScalable});
+	PropertyNameMapping.Emplace(GET_MEMBER_NAME_CHECKED(FPCGSlicingSubmodule, DebugColor), {InSlicingModuleAttributeNames.DebugColorAttributeName, /*bCanBeDefaulted=*/!InSlicingModuleAttributeNames.bProvideDebugColor});
 
-		if (!GetRange(Symbols, InSettings->ModulesInfoAttributeNames.SymbolAttributeName)
-			|| !GetRange(Sizes, InSettings->ModulesInfoAttributeNames.SizeAttributeName)
-			|| !GetRange(Scalables, InSettings->ModulesInfoAttributeNames.ScalableAttributeName, InSettings->ModulesInfoAttributeNames.bProvideScalable)
-			|| !GetRange(DebugColors, InSettings->ModulesInfoAttributeNames.DebugColorAttributeName, InSettings->ModulesInfoAttributeNames.bProvideDebugColor))
-		{
-			return ModulesInfo;
-		}
+	const TArray<FPCGSlicingSubmodule> AllModules = PCGPropertyHelpers::ExtractAttributeSetAsArrayOfStructs<FPCGSlicingSubmodule>(ParamData, &PropertyNameMapping, InContext);
 
-		check(Keys);
-		ModulesInfo.Reserve(Keys->GetNum());
+	ModulesInfo.Reserve(AllModules.Num());
 
-		for (int32 i = 0; i < Keys->GetNum(); ++i)
+	for (int32 i = 0; i < AllModules.Num(); ++i)
+	{
+		if (ModulesInfo.Contains(AllModules[i].Symbol))
 		{
-			if (ModulesInfo.Contains(Symbols[i]))
-			{
-				PCGLog::LogWarningOnGraph(FText::Format(PCGSlicingBase::DuplicatedSymbolText, FText::FromName(Symbols[i])), InContext);
-				continue;
-			}
+			PCGLog::LogWarningOnGraph(FText::Format(PCGSlicingBase::DuplicatedSymbolText, FText::FromName(AllModules[i].Symbol)), InContext);
+			continue;
+		}
 
-			FPCGSlicingModule& Module = ModulesInfo.Emplace(Symbols[i]);
-			Module.Symbol = Symbols[i];
-			Module.Size = Sizes[i];
+		ModulesInfo.Emplace(AllModules[i].Symbol, std::move(AllModules[i]));
+	}
 
-			Module.bScalable = InSettings->ModulesInfoAttributeNames.bProvideScalable ? Scalables[i] : false;
-			Module.DebugColor = InSettings->ModulesInfoAttributeNames.bProvideDebugColor ? DebugColors[i] : FVector4::One();
-		}
+	OutModuleInfoParamData = ParamData;
 
-		OutModuleInfoParamData = ParamData;
+	return ModulesInfo;
+}
+
+PCGSlicingBase::FPCGModulesInfoMap FPCGSlicingBaseElement::GetModulesInfoMap(FPCGContext* InContext, const UPCGSlicingBaseSettings* InSettings, const UPCGParamData*& OutModuleInfoParamData) const
+{
+	if (InSettings->bModuleInfoAsInput)
+	{
+		return GetModulesInfoMap(InContext, InSettings->ModulesInfoAttributeNames, OutModuleInfoParamData);
 	}
 	else
 	{
-		ModulesInfo.Reserve(InSettings->ModulesInfo.Num());
-		for (const FPCGSlicingModule& SlicingModule : InSettings->ModulesInfo)
-		{
-			if (ModulesInfo.Contains(SlicingModule.Symbol))
-			{
-				PCGLog::LogWarningOnGraph(FText::Format(PCGSlicingBase::DuplicatedSymbolText, FText::FromName(SlicingModule.Symbol)), InContext);
-				continue;
-			}
-
-			ModulesInfo.Emplace(SlicingModule.Symbol, SlicingModule);
-		}
+		return GetModulesInfoMap(InContext, InSettings->ModulesInfo, OutModuleInfoParamData);
 	}
-
-	return ModulesInfo;
 }
 
-TArray<FPCGTokenizedGrammar> FPCGSlicingBaseElement::GetTokenizeGrammar(FPCGContext* InContext, const UPCGData* InputData, const UPCGSlicingBaseSettings* InSettings, const FPCGModulesInfoMap& InModulesInfo, double& OutMinSize) const
+PCGGrammar::FTokenizedGrammar FPCGSlicingBaseElement::GetTokenizedGrammar(FPCGContext* InContext, const UPCGData* InputData, const UPCGSlicingBaseSettings* InSettings, const FPCGModulesInfoMap& InModulesInfo, double& OutMinSize) const
 {
-	FString Grammar = InSettings->Grammar;
+	FString Grammar = InSettings->GrammarSelection.GrammarString;
 
-	if (InSettings->bGrammarAsAttribute)
+	if (InSettings->GrammarSelection.bGrammarAsAttribute)
 	{
-		const FPCGAttributePropertyInputSelector Selector = InSettings->GrammarAttribute.CopyAndFixLast(InputData);
+		const FPCGAttributePropertyInputSelector Selector = InSettings->GrammarSelection.GrammarAttribute.CopyAndFixLast(InputData);
 		const TUniquePtr<const IPCGAttributeAccessor> Accessor = PCGAttributeAccessorHelpers::CreateConstAccessor(InputData, Selector);
 		if (!Accessor)
 		{
-			PCGLog::LogErrorOnGraph(FText::Format(PCGSlicingBase::AccessorGrammarFailCreate, Selector.GetDisplayText()), InContext);
+			PCGLog::Accessor::LogFailToCreate(Selector, InContext);
 			return {};
 		}
 
 		if (!Accessor->Get(Grammar, FPCGAttributeAccessorKeysEntries(PCGInvalidEntryKey), EPCGAttributeAccessorFlags::AllowBroadcastAndConstructible))
 		{
-			PCGLog::LogErrorOnGraph(FText::Format(PCGSlicingBase::AccessorFailGetText, Selector.GetDisplayText(), PCG::Private::GetTypeNameText<FString>(), PCG::Private::GetTypeNameText(Accessor->GetUnderlyingType())), InContext);
+			PCGLog::Accessor::LogFailToGet<FString>(Selector, Accessor.Get(), InContext);
 			return {};
 		}
 	}
 
-	return GetTokenizeGrammar(InContext, Grammar, InModulesInfo, OutMinSize);
+	return PCGSlicingBase::GetTokenizedGrammar(InContext, Grammar, InModulesInfo, OutMinSize);
 }
 
-TArray<FPCGTokenizedGrammar> FPCGSlicingBaseElement::GetTokenizeGrammar(FPCGContext* InContext, const FString& InGrammar, const FPCGModulesInfoMap& InModulesInfo, double& OutMinSize) const
+PCGGrammar::FTokenizedGrammar PCGSlicingBase::GetTokenizedGrammar(FPCGContext* InContext, const FString& InGrammar, const FPCGModulesInfoMap& InModulesInfo, double& OutMinSize)
 {
 	const FPCGGrammarResult Result = PCGGrammar::Parse(InGrammar);
 
@@ -172,17 +169,17 @@ TArray<FPCGTokenizedGrammar> FPCGSlicingBaseElement::GetTokenizeGrammar(FPCGCont
 		return {};
 	}
 
-	TArray<FPCGTokenizedGrammar> TokenizeGrammar;
+	PCGGrammar::FTokenizedGrammar TokenizedGrammar;
 	OutMinSize = 0.0;
 
 	for (const PCGGrammar::FModuleDescriptor& ModuleDescriptor : Result.Modules)
 	{
-		FPCGTokenizedGrammar& CurrentModule = TokenizeGrammar.Emplace_GetRef();
+		PCGGrammar::FTokenizedModule& CurrentModule = TokenizedGrammar.Emplace_GetRef();
 		CurrentModule.NumRepeat = ModuleDescriptor.Repetitions;
 
 		for (const PCGGrammar::FModuleDescriptor::FSubmodule& SubmoduleDescriptor : ModuleDescriptor.Submodules)
 		{
-			if (const FPCGSlicingModule* It = InModulesInfo.Find(SubmoduleDescriptor.ID))
+			if (const FPCGSlicingSubmodule* It = InModulesInfo.Find(SubmoduleDescriptor.ID))
 			{
 				CurrentModule.Symbols.Add(SubmoduleDescriptor.ID);
 				CurrentModule.Size += It->Size;
@@ -191,33 +188,34 @@ TArray<FPCGTokenizedGrammar> FPCGSlicingBaseElement::GetTokenizeGrammar(FPCGCont
 				CurrentModule.SymbolSizes.Add(It->Size);
 			}
 		}
-		if (CurrentModule.NumRepeat > 0)
-		{
-			OutMinSize += CurrentModule.Size * CurrentModule.NumRepeat;
-		}
 
 		if (CurrentModule.Symbols.IsEmpty())
 		{
 			// If we have no symbol, we skip.
 			continue;
 		}
+
+		if (CurrentModule.NumRepeat > 0)
+		{
+			OutMinSize += CurrentModule.Size * CurrentModule.NumRepeat;
+		}
 	}
 
-	return TokenizeGrammar;
+	return TokenizedGrammar;
 }
 
-TMap<FString, TArray<FPCGTokenizedGrammar>> FPCGSlicingBaseElement::GetTokenizeGrammarForPoints(FPCGContext* InContext, const UPCGPointData* InputData, const UPCGSlicingBaseSettings* InSettings, const FPCGModulesInfoMap& InModulesInfo, double& OutMinSize) const
+TMap<FString, PCGGrammar::FTokenizedGrammar> FPCGSlicingBaseElement::GetTokenizedGrammarForPoints(FPCGContext* InContext, const UPCGPointData* InputData, const UPCGSlicingBaseSettings* InSettings, const FPCGModulesInfoMap& InModulesInfo, double& OutMinSize) const
 {
-	TMap<FString, TArray<FPCGTokenizedGrammar>> Result;
+	TMap<FString, PCGGrammar::FTokenizedGrammar> Result;
 
-	if (InSettings->bGrammarAsAttribute)
+	if (InSettings->GrammarSelection.bGrammarAsAttribute)
 	{
-		const FPCGAttributePropertyInputSelector Selector = InSettings->GrammarAttribute.CopyAndFixLast(InputData);
+		const FPCGAttributePropertyInputSelector Selector = InSettings->GrammarSelection.GrammarAttribute.CopyAndFixLast(InputData);
 		const TUniquePtr<const IPCGAttributeAccessor> Accessor = PCGAttributeAccessorHelpers::CreateConstAccessor(InputData, Selector);
 		const TUniquePtr<const IPCGAttributeAccessorKeys> Keys = PCGAttributeAccessorHelpers::CreateConstKeys(InputData, Selector);
 		if (!Accessor || !Keys)
 		{
-			PCGLog::LogErrorOnGraph(FText::Format(PCGSlicingBase::AccessorGrammarFailCreate, Selector.GetDisplayText()), InContext);
+			PCGLog::Accessor::LogFailToCreate(Selector, InContext);
 			return Result;
 		}
 
@@ -231,18 +229,18 @@ TMap<FString, TArray<FPCGTokenizedGrammar>> FPCGSlicingBaseElement::GetTokenizeG
 
 		if (!bSuccess)
 		{
-			PCGLog::LogErrorOnGraph(FText::Format(PCGSlicingBase::AccessorFailGetText, Selector.GetDisplayText(), PCG::Private::GetTypeNameText<FString>(), PCG::Private::GetTypeNameText(Accessor->GetUnderlyingType())), InContext);
+			PCGLog::Accessor::LogFailToGet<FString>(Selector, Accessor.Get(), InContext);
 			return Result;
 		}
 	}
 	else
 	{
-		Result.Emplace(InSettings->Grammar);
+		Result.Emplace(InSettings->GrammarSelection.GrammarString);
 	}
 
 	for (auto& [Grammar, TokenizeGrammar] : Result)
 	{
-		TokenizeGrammar = GetTokenizeGrammar(InContext, Grammar, InModulesInfo, OutMinSize);
+		TokenizeGrammar = PCGSlicingBase::GetTokenizedGrammar(InContext, Grammar, InModulesInfo, OutMinSize);
 	}
 
 	return Result;
@@ -254,7 +252,7 @@ bool FPCGSlicingBaseElement::MatchAndSetAttributes(const TArray<FPCGTaggedData>&
 
 	const UPCGMetadata* InputMetadata = InModuleInfoParamData->Metadata;
 
-	// We prepare everything to process all the output data afterwrads.
+	// We prepare everything to process all the output data afterward.
 	// Build the Symbol -> EntryKey mapping
 	// Since we don't know if it is a Name or a String, we need to get both.
 	TMap<FName, PCGMetadataEntryKey> SymbolToEntryKeyMapping;
@@ -398,4 +396,4 @@ bool FPCGSlicingBaseElement::MatchAndSetAttributes(const TArray<FPCGTaggedData>&
 	return true;
 }
 
-#undef LOCTEXT_NAMESPACE
\ No newline at end of file
+#undef LOCTEXT_NAMESPACE
diff --git a/Engine/Plugins/PCG/Source/PCG/Private/Elements/Grammar/PCGVolumeSlicer.cpp b/Engine/Plugins/PCG/Source/PCG/Private/Elements/Grammar/PCGVolumeSlicer.cpp
--- a/Engine/Plugins/PCG/Source/PCG/Private/Elements/Grammar/PCGVolumeSlicer.cpp
+++ b/Engine/Plugins/PCG/Source/PCG/Private/Elements/Grammar/PCGVolumeSlicer.cpp
@@ -66,7 +66,7 @@ bool FPCGVolumeSlicerElement::ExecuteInternal(FPCGContext* InContext) const
 	TArray<FPCGTaggedData>& Outputs = InContext->OutputData.TaggedData;
 
 	const UPCGParamData* ModuleInfoParamData = nullptr;
-	FPCGModulesInfoMap ModulesInfo = GetModulesInfoMap(InContext, Settings, ModuleInfoParamData);
+	PCGSlicingBase::FPCGModulesInfoMap ModulesInfo = GetModulesInfoMap(InContext, Settings, ModuleInfoParamData);
 
 	const bool bMatchAndSetAttributes = Settings->bForwardAttributesFromModulesInfo && ModuleInfoParamData;
 
@@ -93,13 +93,13 @@ bool FPCGVolumeSlicerElement::ExecuteInternal(FPCGContext* InContext) const
 			const TUniquePtr<const IPCGAttributeAccessor> Accessor = PCGAttributeAccessorHelpers::CreateConstAccessor(InputSplineData, Selector);
 			if (!Accessor)
 			{
-				PCGLog::LogErrorOnGraph(FText::Format(LOCTEXT("AccessorFailCreate", "Attribute {0} was not found in the input data."), Selector.GetDisplayText()), InContext);
+				PCGLog::Accessor::LogFailToCreate(Selector, InContext);
 				return false;
 			}
 
 			if (!Accessor->Get(OutValue, FPCGAttributeAccessorKeysEntries(PCGInvalidEntryKey), EPCGAttributeAccessorFlags::AllowBroadcastAndConstructible))
 			{
-				PCGLog::LogErrorOnGraph(FText::Format(LOCTEXT("AccessorFailGet", "Couldn't retrieve attribute {0} value, could it be the wrong type? Expected type: {1}, Attribute Type: {2}."), Selector.GetDisplayText(), PCG::Private::GetTypeNameText<T>(), PCG::Private::GetTypeNameText(Accessor->GetUnderlyingType())), InContext);
+				PCGLog::Accessor::LogFailToGet<T>(Selector, Accessor.Get(), InContext);
 				return false;
 			}
 
@@ -120,16 +120,16 @@ bool FPCGVolumeSlicerElement::ExecuteInternal(FPCGContext* InContext) const
 		const FVector ExtrudeDirection = ExtrudeVector / ExtrudeLength;
 
 		double MinSize = 0.0;
-		TArray<FPCGTokenizedGrammar> TokenizeGrammar = GetTokenizeGrammar(InContext, InputSplineData, Settings, ModulesInfo, MinSize);
+		PCGGrammar::FTokenizedGrammar TokenizedGrammar = GetTokenizedGrammar(InContext, InputSplineData, Settings, ModulesInfo, MinSize);
 
-		if (TokenizeGrammar.IsEmpty())
+		if (TokenizedGrammar.IsEmpty())
 		{
 			continue;
 		}
 
-		TArray<PCGSlicingBase::PCGSubDivModuleInstance<FPCGTokenizedGrammar>> Instances;
+		TArray<PCGSlicingBase::TPCGSubDivModuleInstance<PCGGrammar::FTokenizedModule>> Instances;
 		double RemainingLength = 0.0;
-		const bool bHeightSubdivideSuccess = PCGSlicingBase::Subdivide(TokenizeGrammar, ExtrudeLength, Instances, RemainingLength);
+		const bool bHeightSubdivideSuccess = PCGSlicingBase::Subdivide(TokenizedGrammar, ExtrudeLength, Instances, RemainingLength);
 
 		if (!bHeightSubdivideSuccess)
 		{
@@ -140,7 +140,7 @@ bool FPCGVolumeSlicerElement::ExecuteInternal(FPCGContext* InContext) const
 		FVector CurrentDisplacement = FVector::ZeroVector;
 		int32 SplineIndex = 0;
 
-		for (const PCGSlicingBase::PCGSubDivModuleInstance<FPCGTokenizedGrammar>& Instance : Instances)
+		for (const PCGSlicingBase::TPCGSubDivModuleInstance<PCGGrammar::FTokenizedModule>& Instance : Instances)
 		{
 			for (int i = 0; i < Instance.NumRepeat; ++i)
 			{
@@ -148,7 +148,7 @@ bool FPCGVolumeSlicerElement::ExecuteInternal(FPCGContext* InContext) const
 				{
 					const FName Symbol = Instance.Module->Symbols[SymbolIndex];
 					const FVector Size = ExtrudeDirection * Instance.Module->SymbolSizes[SymbolIndex] * (FVector::OneVector + Instance.ExtraScales[SymbolIndex]);
-					const FPCGSlicingModule& CurrentBlock = ModulesInfo[Symbol];
+					const FPCGSlicingSubmodule& CurrentBlock = ModulesInfo[Symbol];
 					FPCGSplineStruct NewSpline = InputSplineData->SplineStruct;
 					for (FInterpCurvePointVector& ControlPoint : NewSpline.SplineCurves.Position.Points)
 					{
diff --git a/Engine/Plugins/PCG/Source/PCG/Public/Elements/Grammar/PCGSelectGrammar.h b/Engine/Plugins/PCG/Source/PCG/Public/Elements/Grammar/PCGSelectGrammar.h
new file mode 100644
--- /dev/null
+++ b/Engine/Plugins/PCG/Source/PCG/Public/Elements/Grammar/PCGSelectGrammar.h
@@ -0,0 +1,148 @@
+// Copyright Epic Games, Inc. All Rights Reserved.
+
+#pragma once
+
+#include "PCGSettings.h"
+#include "PCGSlicingBase.h"
+
+#include "PCGSelectGrammar.generated.h"
+
+/** An enum representation of basic numerical comparisons for the Select Grammar operation. */
+UENUM()
+enum class EPCGSelectGrammarComparator : uint16
+{
+	// 32 options
+	BinaryOps = 1 << 5 UMETA(Hidden),
+	Select = BinaryOps UMETA(DisplayName = "Select", ToolTip = "Select this grammar. Always evaluates to true."),
+	LessThan UMETA(DisplayName = "<", ToolTip = "Select if 'less than' (A < B)"),
+	LessThanEqualTo UMETA(DisplayName = "<=", ToolTip = "Select if 'less than or equal to' (A <= B)"),
+	EqualTo UMETA(DisplayName = "==", ToolTip = "Select if 'equal to' (A == B)"),
+	GreaterThanEqualTo UMETA(DisplayName = ">=", ToolTip = "Select if 'greater than or equal to' (A >= B)"),
+	GreaterThan UMETA(DisplayName = ">", ToolTip = "Select if 'greater than' (A > B)"),
+
+	// 64 options
+	TernaryOps = 1 << 6 UMETA(Hidden),
+	RangeExclusive = TernaryOps UMETA(DisplayName = "Range (exclusive)", ToolTip = "Select if 'within exclusive range' (B < A < C)"),
+	RangeInclusive UMETA(DisplayName = "Range (inclusive)", ToolTip = "Select if 'within inclusive range' (B <= A <= C)"),
+};
+
+/** Criteria to compare against an input value for conditionally selecting a grammar. */
+USTRUCT(BlueprintType)
+struct FPCGSelectGrammarCriterion
+{
+	GENERATED_BODY()
+
+	/** Key representation of a potential grammar selection. */
+	UPROPERTY(BlueprintReadWrite, EditAnywhere, Category = Settings)
+	FName Key;
+
+	/** Comparison operator of a selected input value against a criterion's value. */
+	UPROPERTY(BlueprintReadWrite, EditAnywhere, Category = Settings)
+	EPCGSelectGrammarComparator Comparator = EPCGSelectGrammarComparator::LessThan;
+
+	/** Compared against the input value. */
+	UPROPERTY(BlueprintReadWrite, EditAnywhere, Category = Settings, meta = (EditCondition = "Comparator != EPCGSelectGrammarComparator::Select", EditConditionHides))
+	double FirstValue = 0.0;
+
+	/** Used for comparison in range-based or ternary conditions. */
+	UPROPERTY(BlueprintReadWrite, EditAnywhere, Category = Settings, meta = (EditCondition = "Comparator == EPCGSelectGrammarComparator::RangeExclusive || Comparator == EPCGSelectGrammarComparator::RangeInclusive", EditConditionHides))
+	double SecondValue = 0.0;
+
+	/** The selected grammar result if the condition evaluates to true. */
+	UPROPERTY(BlueprintReadWrite, EditAnywhere, Category = Settings)
+	FString Grammar;
+};
+
+/** Used for directly mapping custom attribute names to a criteria attribute set. */
+USTRUCT(BlueprintType)
+struct FPCGSelectGrammarCriteriaAttributeNames
+{
+	GENERATED_BODY()
+
+	UPROPERTY(BlueprintReadWrite, EditAnywhere, Category = Settings)
+	FName KeyAttributeName = TEXT("Key");
+
+	UPROPERTY(BlueprintReadWrite, EditAnywhere, Category = Settings)
+	FName ComparatorAttributeName = TEXT("Comparator");
+
+	UPROPERTY(BlueprintReadWrite, EditAnywhere, Category = Settings)
+	FName FirstValueAttributeName = TEXT("FirstValue");
+
+	UPROPERTY(BlueprintReadWrite, EditAnywhere, Category = Settings)
+	FName SecondValueAttributeName = TEXT("SecondValue");
+
+	UPROPERTY(BlueprintReadWrite, EditAnywhere, Category = Settings)
+	FName GrammarAttributeName = TEXT("Grammar");
+};
+
+/** Select a grammar by comparing an input attribute against a provided set criteria one-by-one in a sequential order. */
+UCLASS(MinimalAPI, BlueprintType, ClassGroup = (Procedural))
+class UPCGSelectGrammarSettings : public UPCGSettings
+{
+	GENERATED_BODY()
+
+public:
+	UPCGSelectGrammarSettings();
+
+	//~Begin UPCGSettings interface
+#if WITH_EDITOR
+	virtual FName GetDefaultNodeName() const override;
+	virtual FText GetDefaultNodeTitle() const override;
+	virtual FText GetNodeTooltipText() const override;
+	virtual EPCGSettingsType GetType() const override { return EPCGSettingsType::Generic; }
+#endif
+
+protected:
+	virtual TArray<FPCGPinProperties> InputPinProperties() const override;
+	virtual TArray<FPCGPinProperties> OutputPinProperties() const override;
+	virtual FPCGElementPtr CreateElement() const override;
+	//~End UPCGSettings interface
+
+public:
+	/** Select the key with an attribute. */
+	UPROPERTY(BlueprintReadWrite, EditAnywhere, Category = Settings, meta = (PCG_Overridable))
+	bool bKeyAsAttribute = true;
+
+	/** An attribute key that represents the desired set of grammars. */
+	UPROPERTY(BlueprintReadWrite, EditAnywhere, Category = Settings, meta = (PCG_Overridable, EditCondition = "!bKeyAsAttribute", EditConditionHides))
+	FName Key;
+
+	/** An attribute selector for a key that represents the desired set of grammars. */
+	UPROPERTY(BlueprintReadWrite, EditAnywhere, Category = Settings, meta = (PCG_Overridable, EditCondition = "bKeyAsAttribute", EditConditionHides))
+	FPCGAttributePropertyInputSelector KeyAttribute;
+
+	/** The attribute on the input data to be compared against. Will be numerically evaluated. */
+	UPROPERTY(BlueprintReadWrite, EditAnywhere, Category = Settings, meta = (PCG_Overridable))
+	FPCGAttributePropertyInputSelector ComparedValueAttribute;
+
+	/** Pass criteria info as an attribute set matching the structure type 'FPCGSelectGrammarCriteria'. */
+	UPROPERTY(BlueprintReadWrite, EditAnywhere, Category = Settings, meta = (PCG_Overridable))
+	bool bCriteriaAsInput = false;
+
+	/** Selection criteria that will be evaluated in order. */
+	UPROPERTY(BlueprintReadWrite, EditAnywhere, Category = Settings, meta = (PCG_Overridable, EditCondition = "!bCriteriaAsInput", EditConditionHides))
+	TArray<FPCGSelectGrammarCriterion> Criteria;
+
+	/** If no grammar is selected for a given point, pass through the key value. */
+	UPROPERTY(BlueprintReadWrite, EditAnywhere, Category = Settings, meta = (PCG_Overridable))
+	bool bCopyKeyForUnselectedGrammar = false;
+
+	/** Remap expected attribute names for the comparison criteria. */
+	UPROPERTY(BlueprintReadWrite, EditAnywhere, Category = Settings, meta = (PCG_Overridable, InlineEditConditionToggle, EditCondition = "bCriteriaAsInput", EditConditionHides))
+	bool bRemapCriteriaAttributeNames = false;
+
+	/** The attribute names expected for the comparison criteria. */
+	UPROPERTY(BlueprintReadWrite, EditAnywhere, Category = Settings, meta = (PCG_Overridable, DisplayName = "Remap Criteria Attribute Names", EditCondition = "bRemapCriteriaAttributeNames", DisplayAfter = "bRemapCriteriaAttributeNames"))
+	FPCGSelectGrammarCriteriaAttributeNames CriteriaAttributeNames;
+
+	/** The attribute to output the selected grammar. */
+	UPROPERTY(BlueprintReadWrite, EditAnywhere, Category = Settings, meta = (PCG_Overridable))
+	FPCGAttributePropertyOutputSelector OutputGrammarAttribute;
+};
+
+// TODO: For V2, this element could be async and/or timesliced
+class FPCGSelectGrammarElement : public FPCGSlicingBaseElement
+{
+protected:
+	virtual bool ExecuteInternal(FPCGContext* InContext) const override;
+};
diff --git a/Engine/Plugins/PCG/Source/PCG/Public/Elements/Grammar/PCGSlicingBase.h b/Engine/Plugins/PCG/Source/PCG/Public/Elements/Grammar/PCGSlicingBase.h
--- a/Engine/Plugins/PCG/Source/PCG/Public/Elements/Grammar/PCGSlicingBase.h
+++ b/Engine/Plugins/PCG/Source/PCG/Public/Elements/Grammar/PCGSlicingBase.h
@@ -2,9 +2,9 @@
 
 #pragma once
 
-#include "PCGModule.h"
 #include "PCGSettings.h"
-#include "Metadata/PCGAttributePropertySelector.h"
+#include "Grammar/PCGGrammar.h"
+#include "Utils/PCGLogErrors.h"
 
 #include "PCGSlicingBase.generated.h"
 
@@ -12,7 +12,7 @@ class UPCGData;
 class UPCGPointData;
 
 USTRUCT(BlueprintType)
-struct FPCGSlicingModule
+struct FPCGSlicingSubmodule
 {
 	GENERATED_BODY()
 
@@ -83,38 +83,30 @@ public:
 #endif
 	//~End UPCGSettings interface
 
+	virtual void PostLoad() override;
+
 public:
 	/** Set it to true to pass the info as attribute set. */
 	UPROPERTY(BlueprintReadWrite, EditAnywhere, Category = Settings)
 	bool bModuleInfoAsInput = false;
 
-	/** Read the grammar from the input spline as an attribute rather than directly from the settings.
-	* Grammar syntax:
-	* - Each symbol can have multiple characters
-	* - Modules are defined in '[]', multiple symbols in a module are seperated with ','
-	* - Modules can be repeated a fixed number of times, by adding a number after it (like [A,B]3 will produce ABABAB)
-	* - Modules can be marked repeated an indefinite number of times, with '*'. (like [A,B]* will produce ABABABAB... while it fits the allowed size).
-	*/
-	UPROPERTY(BlueprintReadWrite, EditAnywhere, Category = Settings)
-	bool bGrammarAsAttribute = false;
-
 	/** Fixed array of modules used for the slicing. */
 	UPROPERTY(BlueprintReadWrite, EditAnywhere, Category = Settings, meta = (EditCondition = "!bModuleInfoAsInput", EditConditionHides, DisplayAfter = bModuleInfoAsInput))
-	TArray<FPCGSlicingModule> ModulesInfo;
+	TArray<FPCGSlicingSubmodule> ModulesInfo;
 
 	/** Fixed array of modules used for the slicing. */
 	UPROPERTY(BlueprintReadWrite, EditAnywhere, Category = Settings, meta = (EditCondition = "bModuleInfoAsInput", EditConditionHides, DisplayAfter = bModuleInfoAsInput, DisplayName = "Attribute Names for Module Info"))
 	FPCGSlicingModuleAttributeNames ModulesInfoAttributeNames;
 
 	/** An encoded string that represents how to apply a set of rules to a series of defined modules. */
-	UPROPERTY(BlueprintReadWrite, EditAnywhere, Category = Settings, meta = (EditCondition = "!bGrammarAsAttribute", EditConditionHides, DisplayAfter = bGrammarAsAttribute, PCG_Overridable))
-	FString Grammar;
+	UPROPERTY(BlueprintReadWrite, EditAnywhere, Category = Settings, meta = (ShowOnlyInnerProperties))
+	FPCGGrammarSelection GrammarSelection;
 
 	/** Attribute to be taken from the input spline containing the grammar to use for the slicing. */
 	UPROPERTY(BlueprintReadWrite, EditAnywhere, Category = Settings, meta = (EditCondition = "bGrammarAsAttribute", EditConditionHides, DisplayAfter = bGrammarAsAttribute, PCG_Overridable))
 	FPCGAttributePropertyInputSelector GrammarAttribute;
 
-	/** Do a match and set with the incoming module info, only if the module info is passed as input. */
+	/** Do a match and set with the incoming modules info, only if the modules info is passed as input. */
 	UPROPERTY(BlueprintReadWrite, EditAnywhere, Category = "Settings|Extra Output Attributes", meta = (EditCondition = "bModuleInfoAsInput", EditConditionHides, PCG_Overridable))
 	bool bForwardAttributesFromModulesInfo = false;
 
@@ -142,46 +134,28 @@ public:
 	/** Name of the Debug Color output attribute name, ignored if match and set from module info is true. */
 	UPROPERTY(BlueprintReadWrite, EditAnywhere, Category = "Settings|Extra Output Attributes", meta = (PCG_Overridable, EditCondition = "bOutputDebugColorAttribute"))
 	FName DebugColorAttributeName = PCGSlicingBaseConstants::DebugColorAttributeName;
-};
 
-struct FPCGTokenizedGrammar
-{
-	TArray<FName> Symbols;
-	TArray<bool> AreSymbolsScalable;
-	TArray<double> SymbolSizes;
-	int NumRepeat = 1;
-	bool bScalable = false;
-	double Size = 0.0;
-
-	bool IsValid() const { return Size > 0 && !Symbols.IsEmpty() && Symbols.Num() == AreSymbolsScalable.Num() && Symbols.Num() == SymbolSizes.Num(); }
-	double GetSize() const { return Size; }
-	int GetNumRepeat() const { return NumRepeat; }
-	bool IsScalable() const { return bScalable; }
-	int32 GetSubmodulesCount() const { return Symbols.Num(); }
-	TArrayView<const bool> AreSubmodulesScalable() const { return AreSymbolsScalable; }
-	TArrayView<const double> SubmoduleSizes() const { return SymbolSizes; }
-};
+private:
+#if WITH_EDITORONLY_DATA
+	UE_DEPRECATED(5.5, "Use 'GrammarSelection' instead.")
+	UPROPERTY()
+	bool bGrammarAsAttribute_DEPRECATED = false;
 
-class FPCGSlicingBaseElement : public IPCGElement
-{
-public:
-	using FPCGModulesInfoMap = TMap<FName, FPCGSlicingModule>;
-
-protected:
-	FPCGModulesInfoMap GetModulesInfoMap(FPCGContext* InContext, const UPCGSlicingBaseSettings* InSettings, const UPCGParamData*& OutModuleInfoParamData) const;
-	TArray<FPCGTokenizedGrammar> GetTokenizeGrammar(FPCGContext* InContext, const UPCGData* InputData, const UPCGSlicingBaseSettings* InSettings, const FPCGModulesInfoMap& InModulesInfo, double& OutMinSize) const;
-	TArray<FPCGTokenizedGrammar> GetTokenizeGrammar(FPCGContext* InContext, const FString& InGrammar, const FPCGModulesInfoMap& InModulesInfo, double& OutMinSize) const;
-	TMap<FString, TArray<FPCGTokenizedGrammar>> GetTokenizeGrammarForPoints(FPCGContext* InContext, const UPCGPointData* InputData, const UPCGSlicingBaseSettings* InSettings, const FPCGModulesInfoMap& InModulesInfo, double& OutMinSize) const;
-	bool MatchAndSetAttributes(const TArray<FPCGTaggedData>& InputData, TArray<FPCGTaggedData>& OutputData, const UPCGParamData* InModuleInfoParamData, const UPCGSlicingBaseSettings* InSettings) const;
+	UE_DEPRECATED(5.5, "Use 'GrammarSelection' instead.")
+	UPROPERTY()
+	FString Grammar_DEPRECATED;
+#endif // WITH_EDITORONLY_DATA
 };
 
 namespace PCGSlicingBase
 {
+	using FPCGModulesInfoMap = TMap<FName, FPCGSlicingSubmodule>;
+
 #define PCG_SLICING_BASE_USES_CONCEPTS 0
 
 #if PCG_SLICING_BASE_USES_CONCEPTS
 	// To validate that the type has the right interface
-	// To be enable when this is supported on all our platforms.
+	// To be enabled when this is supported on all our platforms.
 	template <typename T>
 	concept IsValidPCGSubDivModuleInstance =
 		requires(T t) {
@@ -196,7 +170,7 @@ namespace PCGSlicingBase
 #endif // PCG_SLICING_BASE_USES_CONCEPTS
 
 	template <typename T>
-	struct PCGSubDivModuleInstance
+	struct TPCGSubDivModuleInstance
 	{
 		const T* Module = nullptr;
 		int32 NumRepeat = 0;
@@ -205,8 +179,10 @@ namespace PCGSlicingBase
 		TArray<double> ExtraScales;
 	};
 
+	PCGGrammar::FTokenizedGrammar GetTokenizedGrammar(FPCGContext* InContext, const FString& InGrammar, const FPCGModulesInfoMap& InModulesInfo, double& OutMinSize);
+
 	template <typename T>
-	bool Subdivide(const TArray<T>& Modules, double Length, TArray<PCGSubDivModuleInstance<T>>& OutModuleInstances, double& RemainingLength, FPCGContext* InOptionalContext = nullptr)
+	bool Subdivide(const TArray<T>& Modules, double Length, TArray<TPCGSubDivModuleInstance<T>>& OutModuleInstances, double& RemainingLength, FPCGContext* InOptionalContext = nullptr)
 	{
 		OutModuleInstances.Empty(Modules.Num());
 		RemainingLength = Length;
@@ -223,7 +199,7 @@ namespace PCGSlicingBase
 				continue;
 			}
 
-			PCGSubDivModuleInstance<T>& ModuleInstance = OutModuleInstances.Emplace_GetRef();
+			TPCGSubDivModuleInstance<T>& ModuleInstance = OutModuleInstances.Emplace_GetRef();
 			ModuleInstance.Module = &Module;
 			ModuleInstance.ExtraScales.SetNumZeroed(Module.GetSubmodulesCount());
 			if (Module.GetNumRepeat() > 0)
@@ -248,7 +224,7 @@ namespace PCGSlicingBase
 		bool bHasModifiedSomething = false;
 		while (RemainingLength >= 0)
 		{
-			PCGSubDivModuleInstance<T>& ModuleInstance = OutModuleInstances[CurrentModuleIndex];
+			TPCGSubDivModuleInstance<T>& ModuleInstance = OutModuleInstances[CurrentModuleIndex];
 			if (ModuleInstance.Module->GetNumRepeat() <= 0)
 			{
 				if (RemainingLength >= ModuleInstance.Module->GetSize())
@@ -277,8 +253,8 @@ namespace PCGSlicingBase
 		{
 			int32 NumScalableSubmodules = 0;
 
-			TArray<PCGSubDivModuleInstance<T>*> ScalableInstances;
-			for (PCGSubDivModuleInstance<T>& ModuleInstance : OutModuleInstances)
+			TArray<TPCGSubDivModuleInstance<T>*> ScalableInstances;
+			for (TPCGSubDivModuleInstance<T>& ModuleInstance : OutModuleInstances)
 			{
 				if (ModuleInstance.Module->IsScalable() && ModuleInstance.NumRepeat > 0)
 				{
@@ -298,7 +274,7 @@ namespace PCGSlicingBase
 			{
 				const double ExtraLengthPerSubmodule = RemainingLength / NumScalableSubmodules;
 
-				for (PCGSubDivModuleInstance<T>* ModuleInstance : ScalableInstances)
+				for (TPCGSubDivModuleInstance<T>* ModuleInstance : ScalableInstances)
 				{
 					const TArrayView<const bool> AreSubmodulesScalable = ModuleInstance->Module->AreSubmodulesScalable();
 					const TArrayView<const double> SubmoduleSizes = ModuleInstance->Module->SubmoduleSizes();
@@ -322,4 +298,17 @@ namespace PCGSlicingBase
 
 		return true;
 	}
-}
\ No newline at end of file
+}
+
+class FPCGSlicingBaseElement : public IPCGElement
+{
+protected:
+	using FPCGModulesInfoMap = PCGSlicingBase::FPCGModulesInfoMap;
+
+	FPCGModulesInfoMap GetModulesInfoMap(FPCGContext* InContext, const TArray<FPCGSlicingSubmodule>& SubmodulesInfo, const UPCGParamData*& OutModuleInfoParamData) const;
+	FPCGModulesInfoMap GetModulesInfoMap(FPCGContext* InContext, const FPCGSlicingModuleAttributeNames& InSlicingModuleAttributeNames, const UPCGParamData*& OutModuleInfoParamData) const;
+	FPCGModulesInfoMap GetModulesInfoMap(FPCGContext* InContext, const UPCGSlicingBaseSettings* InSettings, const UPCGParamData*& OutModuleInfoParamData) const;
+	PCGGrammar::FTokenizedGrammar GetTokenizedGrammar(FPCGContext* InContext, const UPCGData* InputData, const UPCGSlicingBaseSettings* InSettings, const FPCGModulesInfoMap& InModulesInfo, double& OutMinSize) const;
+	TMap<FString, PCGGrammar::FTokenizedGrammar> GetTokenizedGrammarForPoints(FPCGContext* InContext, const UPCGPointData* InputData, const UPCGSlicingBaseSettings* InSettings, const FPCGModulesInfoMap& InModulesInfo, double& OutMinSize) const;
+	bool MatchAndSetAttributes(const TArray<FPCGTaggedData>& InputData, TArray<FPCGTaggedData>& OutputData, const UPCGParamData* InModuleInfoParamData, const UPCGSlicingBaseSettings* InSettings) const;
+};
diff --git a/Engine/Plugins/PCG/Source/PCG/Public/Grammar/PCGGrammar.h b/Engine/Plugins/PCG/Source/PCG/Public/Grammar/PCGGrammar.h
new file mode 100644
--- /dev/null
+++ b/Engine/Plugins/PCG/Source/PCG/Public/Grammar/PCGGrammar.h
@@ -0,0 +1,54 @@
+// Copyright Epic Games, Inc. All Rights Reserved.
+
+#pragma once
+
+#include "Metadata/PCGAttributePropertySelector.h"
+
+#include "PCGGrammar.generated.h"
+
+namespace PCGGrammar
+{
+	struct FTokenizedModule
+	{
+		TArray<FName> Symbols;
+		TArray<bool> AreSymbolsScalable;
+		TArray<double> SymbolSizes;
+		int NumRepeat = 1;
+		bool bScalable = false;
+		double Size = 0.0;
+
+		bool IsValid() const { return Size > 0 && !Symbols.IsEmpty() && Symbols.Num() == AreSymbolsScalable.Num() && Symbols.Num() == SymbolSizes.Num(); }
+		double GetSize() const { return Size; }
+		int GetNumRepeat() const { return NumRepeat; }
+		bool IsScalable() const { return bScalable; }
+		int32 GetSubmodulesCount() const { return Symbols.Num(); }
+		TArrayView<const bool> AreSubmodulesScalable() const { return AreSymbolsScalable; }
+		TArrayView<const double> SubmoduleSizes() const { return SymbolSizes; }
+	};
+
+	using FTokenizedGrammar = TArray<FTokenizedModule>;
+}
+
+USTRUCT(Blueprintable, BlueprintType)
+struct FPCGGrammarSelection
+{
+	GENERATED_BODY()
+
+	/** Read the grammar as an attribute rather than directly from the settings.
+	* Grammar syntax:
+	* - Each symbol can have multiple characters
+	* - Modules are defined in '[]', multiple symbols in a module are separated with ','
+	* - Modules can be repeated a fixed number of times, by adding a number after it (like [A,B]3 will produce ABABAB)
+	* - Modules can be marked repeated an indefinite number of times, with '*'. (like [A,B]* will produce ABABABAB... while it fits the allowed size).
+	*/
+	UPROPERTY(BlueprintReadWrite, EditAnywhere, Category = Settings)
+	bool bGrammarAsAttribute = false;
+
+	/** An encoded string that represents how to apply a set of rules to a series of defined modules. */
+	UPROPERTY(BlueprintReadWrite, EditAnywhere, Category = Settings, meta = (EditCondition = "!bGrammarAsAttribute", EditConditionHides, PCG_Overridable))
+	FString GrammarString;
+
+	/** Attribute to be taken from the input spline containing the grammar to use. */
+	UPROPERTY(BlueprintReadWrite, EditAnywhere, Category = Settings, meta = (EditCondition = "bGrammarAsAttribute", EditConditionHides, PCG_Overridable))
+	FPCGAttributePropertyInputSelector GrammarAttribute;
+};
