Removed const or made explicit const_casts where TSoftObjectPtr was removing const implicitly.

diff --git a/Engine/Plugins/MovieScene/MovieRenderPipeline/Source/MovieRenderPipelineCore/Private/Graph/MovieGraphRenderLayerSubsystem.cpp b/Engine/Plugins/MovieScene/MovieRenderPipeline/Source/MovieRenderPipelineCore/Private/Graph/MovieGraphRenderLayerSubsystem.cpp
--- a/Engine/Plugins/MovieScene/MovieRenderPipeline/Source/MovieRenderPipelineCore/Private/Graph/MovieGraphRenderLayerSubsystem.cpp
+++ b/Engine/Plugins/MovieScene/MovieRenderPipeline/Source/MovieRenderPipelineCore/Private/Graph/MovieGraphRenderLayerSubsystem.cpp
@@ -772,7 +772,7 @@ void UMovieGraphConditionGroupQuery_Actor::AddActors(const TArray<AActor*>& InAc
 	const FScopedTransaction Transaction(LOCTEXT("AddActorsToCollection", "Add Actors to Collection"));
 	Modify();
 	
-	for (const AActor* Actor : InActors)
+	for (AActor* Actor : InActors)
 	{
 		if (Actor && !ActorsToMatch.Contains(Actor))
 		{
@@ -2283,7 +2283,7 @@ void UMovieGraphConditionGroupQuery_DataLayer::AddDataLayers(const TArray<const
 	{
 		if (DataLayer && !DataLayers.Contains(DataLayer))
 		{
-			DataLayers.Add(DataLayer);
+			DataLayers.Add(const_cast<UDataLayerAsset*>(DataLayer));
 			ListDataSource.Add(MakeShared<TSoftObjectPtr<UDataLayerAsset>>(DataLayer));
 		}
 	}
diff --git a/Engine/Plugins/MovieScene/MovieRenderPipeline/Source/MovieRenderPipelineCore/Private/MoviePipelineQueue.cpp b/Engine/Plugins/MovieScene/MovieRenderPipeline/Source/MovieRenderPipelineCore/Private/MoviePipelineQueue.cpp
--- a/Engine/Plugins/MovieScene/MovieRenderPipeline/Source/MovieRenderPipelineCore/Private/MoviePipelineQueue.cpp
+++ b/Engine/Plugins/MovieScene/MovieRenderPipeline/Source/MovieRenderPipelineCore/Private/MoviePipelineQueue.cpp
@@ -36,7 +36,7 @@ namespace UE::MovieGraph::Private
 		// Create the variable assignments container if it wasn't found
 		TObjectPtr<UMovieJobVariableAssignmentContainer> NewVariableAssignments = NewObject<UMovieJobVariableAssignmentContainer>(InAssignmentsOwner);
 		InVariableAssignments.Add(NewVariableAssignments);
-		NewVariableAssignments->SetGraphConfig(InGraph);
+		NewVariableAssignments->SetGraphConfig(const_cast<UMovieGraphConfig*>(InGraph));
 
 #if WITH_EDITOR
 		NewVariableAssignments->UpdateGraphVariableOverrides();
diff --git a/Engine/Plugins/PCG/Source/PCG/Private/PCGComponent.cpp b/Engine/Plugins/PCG/Source/PCG/Private/PCGComponent.cpp
--- a/Engine/Plugins/PCG/Source/PCG/Private/PCGComponent.cpp
+++ b/Engine/Plugins/PCG/Source/PCG/Private/PCGComponent.cpp
@@ -803,7 +803,7 @@ void UPCGComponent::PurgeUnlinkedResources(const AActor* InActor)
 
 	ThisActor->GetAttachedActors(AttachedActors);
 
-	for (const AActor* Actor : AttachedActors)
+	for (AActor* Actor : AttachedActors)
 	{
 		if (Actor && Actor->ActorHasTag(PCGHelpers::DefaultPCGActorTag))
 		{
diff --git a/Engine/Plugins/PCG/Source/PCG/Private/Tests/Elements/PCGInstancePackerTest.cpp b/Engine/Plugins/PCG/Source/PCG/Private/Tests/Elements/PCGInstancePackerTest.cpp
--- a/Engine/Plugins/PCG/Source/PCG/Private/Tests/Elements/PCGInstancePackerTest.cpp
+++ b/Engine/Plugins/PCG/Source/PCG/Private/Tests/Elements/PCGInstancePackerTest.cpp
@@ -182,7 +182,7 @@ bool FPCGStaticMeshSpawnerInstancePackerByAttributeTest::RunTest(const FString&
 	UPCGMeshSelectorWeighted* MeshSelector = CastChecked<UPCGMeshSelectorWeighted>(Settings->MeshSelectorParameters);
 
 	const FString CubePath = TEXT("StaticMesh'/Engine/BasicShapes/Cube.Cube'");
-	const UStaticMesh* CubeMesh = Cast<UStaticMesh>(StaticLoadObject(UStaticMesh::StaticClass(), NULL, *CubePath));
+	UStaticMesh* CubeMesh = Cast<UStaticMesh>(StaticLoadObject(UStaticMesh::StaticClass(), NULL, *CubePath));
 	MeshSelector->MeshEntries.Add(FPCGMeshSelectorWeightedEntry(CubeMesh, 1));
 
 	bool bTestPassed = true;
@@ -339,7 +339,7 @@ bool FPCGStaticMeshSpawnerInstancePackerByRegexTest::RunTest(const FString& Para
 	UPCGMeshSelectorWeighted* MeshSelector = CastChecked<UPCGMeshSelectorWeighted>(Settings->MeshSelectorParameters);
 
 	const FString CubePath = TEXT("StaticMesh'/Engine/BasicShapes/Cube.Cube'");
-	const UStaticMesh* CubeMesh = Cast<UStaticMesh>(StaticLoadObject(UStaticMesh::StaticClass(), NULL, *CubePath));
+	UStaticMesh* CubeMesh = Cast<UStaticMesh>(StaticLoadObject(UStaticMesh::StaticClass(), NULL, *CubePath));
 	MeshSelector->MeshEntries.Add(FPCGMeshSelectorWeightedEntry(CubeMesh, 1));
 
 	bool bTestPassed = true;
diff --git a/Engine/Source/Runtime/Engine/Private/ActorEditor.cpp b/Engine/Source/Runtime/Engine/Private/ActorEditor.cpp
--- a/Engine/Source/Runtime/Engine/Private/ActorEditor.cpp
+++ b/Engine/Source/Runtime/Engine/Private/ActorEditor.cpp
@@ -2037,7 +2037,7 @@ bool FAssignActorDataLayer::AddDataLayerAsset(AActor* InActor, const UDataLayerA
 	else if (!InActor->DataLayerAssets.Contains(InDataLayerAsset))
 	{
 		InActor->Modify();
-		InActor->DataLayerAssets.Add(InDataLayerAsset);
+		InActor->DataLayerAssets.Add(const_cast<UDataLayerAsset*>(InDataLayerAsset));
 		return true;
 	}
 
@@ -2057,7 +2057,7 @@ bool FAssignActorDataLayer::RemoveDataLayerAsset(AActor* InActor, const UDataLay
 	else if (InActor->DataLayerAssets.Contains(InDataLayerAsset))
 	{
 		InActor->Modify();
-		InActor->DataLayerAssets.Remove(InDataLayerAsset);
+		InActor->DataLayerAssets.Remove(const_cast<UDataLayerAsset*>(InDataLayerAsset));
 		return true;
 	}
 
diff --git a/Engine/Source/Runtime/Engine/Private/Animation/Skeleton.cpp b/Engine/Source/Runtime/Engine/Private/Animation/Skeleton.cpp
--- a/Engine/Source/Runtime/Engine/Private/Animation/Skeleton.cpp
+++ b/Engine/Source/Runtime/Engine/Private/Animation/Skeleton.cpp
@@ -303,7 +303,7 @@ bool USkeleton::IsCompatibleSkeletonByAssetString(const FString& SkeletonAssetSt
 
 void USkeleton::AddCompatibleSkeleton(const USkeleton* SourceSkeleton)
 {
-	CompatibleSkeletons.AddUnique(SourceSkeleton);
+	CompatibleSkeletons.AddUnique(const_cast<USkeleton*>(SourceSkeleton));
 }
 
 void USkeleton::AddCompatibleSkeletonSoft(const TSoftObjectPtr<USkeleton>& SourceSkeleton)
@@ -313,7 +313,7 @@ void USkeleton::AddCompatibleSkeletonSoft(const TSoftObjectPtr<USkeleton>& Sourc
 
 void USkeleton::RemoveCompatibleSkeleton(const USkeleton* SourceSkeleton)
 {
-	CompatibleSkeletons.Remove(SourceSkeleton);
+	CompatibleSkeletons.Remove(const_cast<USkeleton*>(SourceSkeleton));
 }
 
 void USkeleton::RemoveCompatibleSkeleton(const TSoftObjectPtr<USkeleton>& SourceSkeleton)
