[PCG] PCGHiGenGridSize PrepareForExecute needs to run on GT until we have proper pre-execution cache of UPCGComponent::GetPCGActorData()

#rnx

diff --git a/Engine/Plugins/PCG/Source/PCG/Public/Elements/PCGHiGenGridSize.h b/Engine/Plugins/PCG/Source/PCG/Public/Elements/PCGHiGenGridSize.h
--- a/Engine/Plugins/PCG/Source/PCG/Public/Elements/PCGHiGenGridSize.h
+++ b/Engine/Plugins/PCG/Source/PCG/Public/Elements/PCGHiGenGridSize.h
@@ -44,6 +44,10 @@ protected:
 class FPCGHiGenGridSizeElement : public IPCGElement
 {
 public:
+	// Returns true when Context is nullptr which happens in FPCGGraphExecutor::OnTaskInputsReady
+	// We can't know for sure at that point if we need to call UPCGComponent::GetActorPCGData() in our GetDependenciesCrc so we need to make sure it happens on MainThread
+	virtual bool CanExecuteOnlyOnMainThread(FPCGContext* Context) const override { return Context == nullptr; }
+
 	virtual void GetDependenciesCrc(const FPCGDataCollection& InInput, const UPCGSettings* InSettings, UPCGComponent* InComponent, FPCGCrc& OutCrc) const override;
 
 protected:
