[PCG] Fix Graph Executor V2 Cancel issues

#rnx

diff --git a/Engine/Plugins/PCG/Source/PCG/Private/Graph/PCGGraphExecutor.cpp b/Engine/Plugins/PCG/Source/PCG/Private/Graph/PCGGraphExecutor.cpp
--- a/Engine/Plugins/PCG/Source/PCG/Private/Graph/PCGGraphExecutor.cpp
+++ b/Engine/Plugins/PCG/Source/PCG/Private/Graph/PCGGraphExecutor.cpp
@@ -495,23 +495,20 @@ TSet<UPCGComponent*> FPCGGraphExecutor::Cancel(TFunctionRef<bool(TWeakObjectPtr<
 			}
 		}
 
-		// Visit active tasks
-		for (const TSharedPtr<FPCGGraphActiveTask>& Task : ActiveTasks)
+		auto GatherCancelledComponentsForActiveTasks = [&CancelFilter, &CancelledComponents](const TArray<TSharedPtr<FPCGGraphActiveTask>>& InActiveTasks)
 		{
-			if (Task->Context && CancelFilter(Task->Context->SourceComponent))
+			for (const TSharedPtr<FPCGGraphActiveTask>& Task : InActiveTasks)
 			{
-				CancelledComponents.Add(Task->Context->SourceComponent.Get());
-			}
-		}
-	
-		// Visit sleeping tasks
-		for (const TSharedPtr<FPCGGraphActiveTask>& Task : SleepingTasks)
-		{
-			if (Task->Context && CancelFilter(Task->Context->SourceComponent))
-			{
-				CancelledComponents.Add(Task->Context->SourceComponent.Get());
+				if (Task->Context && CancelFilter(Task->Context->SourceComponent))
+				{
+					CancelledComponents.Add(Task->Context->SourceComponent.Get());
+				}
 			}
-		}
+		};
+
+		GatherCancelledComponentsForActiveTasks(ActiveTasks);
+		GatherCancelledComponentsForActiveTasks(ActiveTasksGameThreadOnly);
+		GatherCancelledComponentsForActiveTasks(SleepingTasks);
 	}
 
 	// In one instance this function was observed to return nullptr in the CancelledComponents set.
@@ -601,7 +598,7 @@ TSet<UPCGComponent*> FPCGGraphExecutor::Cancel(TFunctionRef<bool(TWeakObjectPtr<
 				{
 					for (int32 ActiveTaskIndex = InActiveTasks.Num() - 1; ActiveTaskIndex >= 0; --ActiveTaskIndex)
 					{
-						TSharedPtr<FPCGGraphActiveTask>& ActiveTask = InActiveTasks[ActiveTaskIndex];
+						TSharedPtr<FPCGGraphActiveTask> ActiveTask = InActiveTasks[ActiveTaskIndex];
 						if (ActiveTask->Context && CancelledComponents.Contains(ActiveTask->Context->SourceComponent.Get()))
 						{
 							// While we have lock Task can't complete, but we can't wait on this task with the lock neither so we capture it here 
@@ -614,7 +611,10 @@ TSet<UPCGComponent*> FPCGGraphExecutor::Cancel(TFunctionRef<bool(TWeakObjectPtr<
 							// Avoid removing if using old execution path as it doesn't keep a sharedptr to the executing task (which would have been a way of allowing removal here, but since that path is going away, lets keep it this way)
 							if (CurrentExecuteVersion == EExecuteVersion::V2)
 							{
-								InActiveTasks.RemoveAtSwap(ActiveTaskIndex);
+								check(ActiveTask->TaskIndex == ActiveTaskIndex);
+								PCGGraphExecutor::RemoveAtFromActiveTaskArrayNoLock(InActiveTasks, ActiveTaskIndex);
+								ActiveTask->TaskIndex = INDEX_NONE;
+								ActiveTask->bWasCancelled = true;
 							}
 						}
 					}
@@ -678,11 +678,13 @@ TSet<UPCGComponent*> FPCGGraphExecutor::Cancel(TFunctionRef<bool(TWeakObjectPtr<
 				PCGGraphExecutor::TScopeLock ScopeLock(LiveTasksLock);
 				for (int32 SleepingTaskIndex = SleepingTasks.Num() - 1; SleepingTaskIndex >= 0; --SleepingTaskIndex)
 				{
-					TSharedPtr<FPCGGraphActiveTask>& Task = SleepingTasks[SleepingTaskIndex];
+					TSharedPtr<FPCGGraphActiveTask> Task = SleepingTasks[SleepingTaskIndex];
 					if (Task->Context && CancelledComponents.Contains(Task->Context->SourceComponent.Get()))
 					{
-						CancelledSleepingTasks.Add(MoveTemp(Task));
-						SleepingTasks.RemoveAtSwap(SleepingTaskIndex);
+						check(Task->TaskIndex == SleepingTaskIndex);
+						CancelledSleepingTasks.Add(Task);
+						PCGGraphExecutor::RemoveAtFromActiveTaskArrayNoLock(SleepingTasks, SleepingTaskIndex);
+						Task->TaskIndex = INDEX_NONE;
 					}
 				}
 			}
@@ -986,29 +988,34 @@ void FPCGGraphExecutor::PostTaskExecute(TSharedPtr<FPCGGraphActiveTask> ActiveTa
 	{
 		TRACE_CPUPROFILER_EVENT_SCOPE(FPCGGraphExecutor::PostTaskExecute::NotDone);
 		PCGGraphExecutor::TScopeLock ScopeLock(LiveTasksLock);
-		ActiveTask.StopExecuting();
+		
+		// Once we get the lock check for cancellation again in case task was cancelled while executing
+		if (!ActiveTaskPtr->bWasCancelled)
+		{
+			ActiveTask.StopExecuting();
 
-		TArray<TSharedPtr<FPCGGraphActiveTask>>& ActiveTaskArrayToRemove = ActiveTask.bIsGameThreadOnly ? ActiveTasksGameThreadOnly : ActiveTasks;
-		const int32 TaskRemoveIndex = ActiveTask.TaskIndex;
-		check(TaskRemoveIndex != INDEX_NONE && ActiveTaskArrayToRemove[TaskRemoveIndex] == ActiveTaskPtr);
+			TArray<TSharedPtr<FPCGGraphActiveTask>>& ActiveTaskArrayToRemove = ActiveTask.bIsGameThreadOnly ? ActiveTasksGameThreadOnly : ActiveTasks;
+			const int32 TaskRemoveIndex = ActiveTask.TaskIndex;
+			check(TaskRemoveIndex != INDEX_NONE && ActiveTaskArrayToRemove[TaskRemoveIndex] == ActiveTaskPtr);
 
-		if (ActiveTask.Context->bIsPaused)
-		{
-			PCGGraphExecutor::AddToActiveTaskArrayNoLock(SleepingTasks, ActiveTaskPtr);
-			PCGGraphExecutor::RemoveAtFromActiveTaskArrayNoLock(ActiveTaskArrayToRemove, TaskRemoveIndex);
-		}
-		else if (ActiveTask.bIsGameThreadOnly != ActiveTask.Element->CanExecuteOnlyOnMainThread(ActiveTask.Context.Get()))
-		{
-			ActiveTask.bIsGameThreadOnly = !ActiveTask.bIsGameThreadOnly;
-			TArray<TSharedPtr<FPCGGraphActiveTask>>& ActiveTaskArrayToAdd = ActiveTask.bIsGameThreadOnly ? ActiveTasksGameThreadOnly : ActiveTasks;
+			if (ActiveTask.Context->bIsPaused)
+			{
+				PCGGraphExecutor::AddToActiveTaskArrayNoLock(SleepingTasks, ActiveTaskPtr);
+				PCGGraphExecutor::RemoveAtFromActiveTaskArrayNoLock(ActiveTaskArrayToRemove, TaskRemoveIndex);
+			}
+			else if (ActiveTask.bIsGameThreadOnly != ActiveTask.Element->CanExecuteOnlyOnMainThread(ActiveTask.Context.Get()))
+			{
+				ActiveTask.bIsGameThreadOnly = !ActiveTask.bIsGameThreadOnly;
+				TArray<TSharedPtr<FPCGGraphActiveTask>>& ActiveTaskArrayToAdd = ActiveTask.bIsGameThreadOnly ? ActiveTasksGameThreadOnly : ActiveTasks;
 
-			PCGGraphExecutor::AddToActiveTaskArrayNoLock(ActiveTaskArrayToAdd, ActiveTaskPtr);
-			PCGGraphExecutor::RemoveAtFromActiveTaskArrayNoLock(ActiveTaskArrayToRemove, TaskRemoveIndex);
-		}
+				PCGGraphExecutor::AddToActiveTaskArrayNoLock(ActiveTaskArrayToAdd, ActiveTaskPtr);
+				PCGGraphExecutor::RemoveAtFromActiveTaskArrayNoLock(ActiveTaskArrayToRemove, TaskRemoveIndex);
+			}
 
-		// Task might have been moved to SleepingTasks or changed from ActiveTasks to ActiveTasksGameThreadOnly (or vice-versa) 
-		// but it isn't done so we return
-		return;
+			// Task might have been moved to SleepingTasks or changed from ActiveTasks to ActiveTasksGameThreadOnly (or vice-versa) 
+			// but it isn't done so we return
+			return;
+		}
 	}
 
 	const bool bTaskWasCancelled = ActiveTask.bWasCancelled;
@@ -1018,12 +1025,6 @@ void FPCGGraphExecutor::PostTaskExecute(TSharedPtr<FPCGGraphActiveTask> ActiveTa
 	const bool bTaskWasBypassed = false;
 #endif
 
-	// Needs to be done before QueueNextTasks
-	if (bTaskWasCancelled)
-	{
-		RemoveTaskFromInputSuccessors(ActiveTask.NodeId, ActiveTask.Inputs);
-	}
-
 	bNeedToExecuteTasksEnded = true;
 
 	const UPCGSettingsInterface* ActiveTaskSettingsInterface = ActiveTask.Context->GetInputSettingsInterface();
@@ -1539,13 +1540,9 @@ bool FPCGGraphExecutor::ExecuteScheduling(double EndTime, TSharedPtr<FPCGGraphAc
 						{
 							TRACE_CPUPROFILER_EVENT_SCOPE(FPCGGraphExecutor::ExecuteAsyncTask);
 
-							bool bIsDone = ActiveTask->bWasCancelled || ActiveTask->Element->Execute(ActiveTask->Context.Get());
+							const bool bIsDone = ActiveTask->bWasCancelled || ActiveTask->Element->Execute(ActiveTask->Context.Get());
 							const bool bIsPaused = ActiveTask->Context->bIsPaused;
 
-							// It is currently possible for a call to Execute to endup cancelling the task through a Windows message pump caused by a ShowDialog (when calling UTexture::GetPlatformData() from a task)
-							// If this happens mark the task as done
-							bIsDone |= ActiveTask->bWasCancelled;
-
 							PostTaskExecute(ActiveTask, bIsDone);
 
 							if (bIsDone || bIsPaused)
@@ -1638,14 +1635,10 @@ void FPCGGraphExecutor::ExecuteV2()
 		{
 			TRACE_CPUPROFILER_EVENT_SCOPE(FPCGGraphExecutor::Execute::ExecuteMainThreadTask);
 #if WITH_EDITOR
-			bool bIsDone = (MainThreadTask->bIsBypassed || MainThreadTask->bWasCancelled || MainThreadTask->Element->Execute(MainThreadTask->Context.Get()));
+			const bool bIsDone = (MainThreadTask->bIsBypassed || MainThreadTask->bWasCancelled || MainThreadTask->Element->Execute(MainThreadTask->Context.Get()));
 #else
-			bool bIsDone = (MainThreadTask->bWasCancelled || MainThreadTask->Element->Execute(MainThreadTask->Context.Get()));
+			const bool bIsDone = (MainThreadTask->bWasCancelled || MainThreadTask->Element->Execute(MainThreadTask->Context.Get()));
 #endif
-			// It is currently possible for a call to Execute to endup cancelling the task through a Windows message pump caused by a ShowDialog (when calling UTexture::GetPlatformData() from a task)
-			// If this happens mark the task as done
-			bIsDone |= MainThreadTask->bWasCancelled;
-
 			PostTaskExecute(MainThreadTask, bIsDone);
 		}
 	}
