[PCG] Fix static analysis warning in Select Grammar node

#rnx

diff --git a/Engine/Plugins/PCG/Source/PCG/Private/Elements/Grammar/PCGSelectGrammar.cpp b/Engine/Plugins/PCG/Source/PCG/Private/Elements/Grammar/PCGSelectGrammar.cpp
--- a/Engine/Plugins/PCG/Source/PCG/Private/Elements/Grammar/PCGSelectGrammar.cpp
+++ b/Engine/Plugins/PCG/Source/PCG/Private/Elements/Grammar/PCGSelectGrammar.cpp
@@ -213,8 +213,7 @@ bool FPCGSelectGrammarElement::ExecuteInternal(FPCGContext* InContext) const
 	for (const FPCGTaggedData& PointInput : PointInputs)
 	{
 		const UPCGPointData* InputPointData = Cast<UPCGPointData>(PointInput.Data);
-		const int32 NumPoints = InputPointData->GetPoints().Num();
-		if (!InputPointData || NumPoints < 1)
+		if (!InputPointData || InputPointData->GetPoints().IsEmpty())
 		{
 			continue;
 		}
