[PCG] PCGContext::InitializeSettings delay clearing of Async flag to the GT by adding the duplicated objects to the Context AsyncObjects

#rnx

diff --git a/Engine/Plugins/PCG/Source/PCG/Private/PCGContext.cpp b/Engine/Plugins/PCG/Source/PCG/Private/PCGContext.cpp
--- a/Engine/Plugins/PCG/Source/PCG/Private/PCGContext.cpp
+++ b/Engine/Plugins/PCG/Source/PCG/Private/PCGContext.cpp
@@ -147,7 +147,7 @@ void FPCGContext::InitializeSettings(bool bSkipPostLoad)
 						if (!bIsInGameThread)
 						{
 							// Outside of GameThread we need to clear the Async flags on newly duplicated objects
-							KeyValuePair.Value->ClearInternalFlags(EInternalObjectFlags::Async);
+							AsyncObjects.Add(KeyValuePair.Value);
 						}
 
 						if (bSkipPostLoad)
