Changed the default spawn option in the Spawn Actor node to 'No merging'.

diff --git a/Engine/Plugins/PCG/Source/PCG/Private/Elements/PCGSpawnActor.cpp b/Engine/Plugins/PCG/Source/PCG/Private/Elements/PCGSpawnActor.cpp
--- a/Engine/Plugins/PCG/Source/PCG/Private/Elements/PCGSpawnActor.cpp
+++ b/Engine/Plugins/PCG/Source/PCG/Private/Elements/PCGSpawnActor.cpp
@@ -132,6 +132,7 @@ UPCGSpawnActorSettings::UPCGSpawnActorSettings(const FObjectInitializer& ObjectI
 {
 	if (PCGHelpers::IsNewObjectAndNotDefault(this))
 	{
+		Option = EPCGSpawnActorOption::NoMerging;
 		AttachOptions = EPCGAttachOptions::InFolder;
 	}
 }
diff --git a/Engine/Plugins/PCG/Source/PCG/Public/Elements/PCGSpawnActor.h b/Engine/Plugins/PCG/Source/PCG/Public/Elements/PCGSpawnActor.h
--- a/Engine/Plugins/PCG/Source/PCG/Public/Elements/PCGSpawnActor.h
+++ b/Engine/Plugins/PCG/Source/PCG/Public/Elements/PCGSpawnActor.h
@@ -51,7 +51,7 @@ public:
 	TArray<FName> PostSpawnFunctionNames;
 
 	UPROPERTY(BlueprintReadWrite, EditAnywhere, Category = Settings)
-	EPCGSpawnActorOption Option = EPCGSpawnActorOption::CollapseActors;
+	EPCGSpawnActorOption Option = EPCGSpawnActorOption::CollapseActors; // Note that this is no longer the default value for new nodes, it is now EPCGSpawnActorOption::NoMerging.
 
 	UPROPERTY(BlueprintReadWrite, EditAnywhere, Category = Settings, meta = (EditCondition = "Option==EPCGSpawnActorOption::NoMerging", EditConditionHides))
 	bool bForceDisableActorParsing = true;
