[PCG] Add overrides to Create Spline Mesh node

#rnx

diff --git a/Engine/Plugins/PCG/Source/PCG/Private/Elements/PCGCreateSplineMesh.cpp b/Engine/Plugins/PCG/Source/PCG/Private/Elements/PCGCreateSplineMesh.cpp
--- a/Engine/Plugins/PCG/Source/PCG/Private/Elements/PCGCreateSplineMesh.cpp
+++ b/Engine/Plugins/PCG/Source/PCG/Private/Elements/PCGCreateSplineMesh.cpp
@@ -107,18 +107,7 @@ bool FPCGCreateSplineMeshElement::ExecuteInternal(FPCGContext* Context) const
 
 	FSplineMeshComponentDescriptor Descriptor(Settings->SplineMeshDescriptor);
 
-	// TODO: Optionally read static mesh, materials, etc. from attributes.
-	UStaticMesh* StaticMesh = Descriptor.StaticMesh.Get();
-
-	if (!StaticMesh)
-	{
-		return true;
-	}
-
-	const FVector MeshExtents = StaticMesh->GetBounds().BoxExtent;
-
 	FPCGSplineMeshComponentBuilderParameters SMCBuilderParams;
-	SMCBuilderParams.Descriptor = Descriptor;
 	SMCBuilderParams.SplineMeshParams = Settings->SplineMeshParams;
 
 	TArray<FPCGTaggedData> Inputs = Context->InputData.GetInputs();
@@ -146,6 +135,27 @@ bool FPCGCreateSplineMeshElement::ExecuteInternal(FPCGContext* Context) const
 		const int NumSegments = SplineData->GetNumSegments();
 		const bool bIsClosed = SplineData->IsClosed();
 
+		// Copy the descriptor, for the overrides (and reset the descriptor if it was overriden by previous inputs)
+		SMCBuilderParams.Descriptor = Descriptor;
+
+		// TODO: Support overrides on control points when we have better support for those.
+		FPCGObjectOverrides DescriptorOverride(&SMCBuilderParams.Descriptor);
+		DescriptorOverride.Initialize(Settings->SplineMeshOverrideDescriptions, &SMCBuilderParams.Descriptor, SplineData, Context);
+		if (DescriptorOverride.IsValid() && !DescriptorOverride.Apply(0))
+		{
+			const int32 Index = static_cast<int32>(&Input - Inputs.GetData());
+			PCGLog::LogWarningOnGraph(FText::Format(LOCTEXT("FailOverride", "Failed to override descriptor for input {0}"), Index));
+		}
+
+		UStaticMesh* StaticMesh = SMCBuilderParams.Descriptor.StaticMesh.Get();
+
+		if (!StaticMesh)
+		{
+			continue;
+		}
+
+		const FVector MeshExtents = StaticMesh->GetBounds().BoxExtent;
+
 		for (int Index = 0; Index < NumSegments; ++Index)
 		{
 			const int NextIndex = (bIsClosed && Index == NumSegments - 1) ? 0 : (Index + 1);
diff --git a/Engine/Plugins/PCG/Source/PCG/Private/Metadata/Accessors/PCGAttributeAccessorHelpers.cpp b/Engine/Plugins/PCG/Source/PCG/Private/Metadata/Accessors/PCGAttributeAccessorHelpers.cpp
--- a/Engine/Plugins/PCG/Source/PCG/Private/Metadata/Accessors/PCGAttributeAccessorHelpers.cpp
+++ b/Engine/Plugins/PCG/Source/PCG/Private/Metadata/Accessors/PCGAttributeAccessorHelpers.cpp
@@ -205,7 +205,10 @@ namespace PCGAttributeAccessorHelpers
 			{
 				if (const UPCGMetadata* Metadata = InData->ConstMetadata())
 				{
-					return MakeUnique<FPCGAttributeAccessorKeysEntries>(Metadata);
+					// For const metadata, that is on SpatialData (not points), we allow for Default Value to support "data-wide" metadata
+					// It's only for spatial data and not Param data because for ParamData, GetNum on the keys should return the number of entries.
+					// If we have a default key, GetNum will return 1, while there is actually 0 values.
+					return MakeUnique<FPCGAttributeAccessorKeysEntries>(Metadata, /*bAddDefaultValueIfEmpty=*/!!Cast<UPCGSpatialData>(InData));
 				}
 			}
 			else
diff --git a/Engine/Plugins/PCG/Source/PCG/Private/Metadata/Accessors/PCGAttributeAccessorKeys.cpp b/Engine/Plugins/PCG/Source/PCG/Private/Metadata/Accessors/PCGAttributeAccessorKeys.cpp
--- a/Engine/Plugins/PCG/Source/PCG/Private/Metadata/Accessors/PCGAttributeAccessorKeys.cpp
+++ b/Engine/Plugins/PCG/Source/PCG/Private/Metadata/Accessors/PCGAttributeAccessorKeys.cpp
@@ -20,23 +20,17 @@ FPCGAttributeAccessorKeysEntries::FPCGAttributeAccessorKeysEntries(PCGMetadataEn
 	Entries = TArrayView<PCGMetadataEntryKey>(ExtractedEntries);
 }
 
-FPCGAttributeAccessorKeysEntries::FPCGAttributeAccessorKeysEntries(const UPCGMetadata* Metadata)
+FPCGAttributeAccessorKeysEntries::FPCGAttributeAccessorKeysEntries(const UPCGMetadata* Metadata, bool bAddDefaultValueIfEmpty)
 	: IPCGAttributeAccessorKeys(/*bInReadOnly=*/ true)
 {
-	InitializeFromMetadata(Metadata);
+	InitializeFromMetadata(Metadata, bAddDefaultValueIfEmpty);
 	Entries = TArrayView<PCGMetadataEntryKey>(ExtractedEntries);
 }
 
-FPCGAttributeAccessorKeysEntries::FPCGAttributeAccessorKeysEntries(UPCGMetadata* Metadata)
+FPCGAttributeAccessorKeysEntries::FPCGAttributeAccessorKeysEntries(UPCGMetadata* Metadata, bool bAddDefaultValueIfEmpty)
 	: IPCGAttributeAccessorKeys(/*bInReadOnly=*/ false)
 {
-	InitializeFromMetadata(Metadata);
-
-	// If the entries are still empty, we will always take the default value.
-	if (ExtractedEntries.IsEmpty())
-	{
-		ExtractedEntries.Add(PCGInvalidEntryKey);
-	}
+	InitializeFromMetadata(Metadata, bAddDefaultValueIfEmpty);
 
 	Entries = TArrayView<PCGMetadataEntryKey>(ExtractedEntries);
 }
@@ -55,7 +49,7 @@ FPCGAttributeAccessorKeysEntries::FPCGAttributeAccessorKeysEntries(const TArrayV
 
 }
 
-void FPCGAttributeAccessorKeysEntries::InitializeFromMetadata(const UPCGMetadata* Metadata)
+void FPCGAttributeAccessorKeysEntries::InitializeFromMetadata(const UPCGMetadata* Metadata, bool bAddDefaultValueIfEmpty)
 {
 	if (!Metadata)
 	{
@@ -77,6 +71,11 @@ void FPCGAttributeAccessorKeysEntries::InitializeFromMetadata(const UPCGMetadata
 			ExtractedEntries.Add(Entry);
 		}
 	}
+
+	if (ExtractedEntries.IsEmpty() && bAddDefaultValueIfEmpty)
+	{
+		ExtractedEntries.Add(PCGInvalidEntryKey);
+	}
 }
 
 bool FPCGAttributeAccessorKeysEntries::GetMetadataEntryKeys(int32 InStart, TArrayView<PCGMetadataEntryKey*> OutEntryKeys)
diff --git a/Engine/Plugins/PCG/Source/PCG/Public/Elements/PCGCreateSplineMesh.h b/Engine/Plugins/PCG/Source/PCG/Public/Elements/PCGCreateSplineMesh.h
--- a/Engine/Plugins/PCG/Source/PCG/Public/Elements/PCGCreateSplineMesh.h
+++ b/Engine/Plugins/PCG/Source/PCG/Public/Elements/PCGCreateSplineMesh.h
@@ -8,6 +8,7 @@
 #include "PCGSettings.h"
 #include "PCGSplineMeshParams.h"
 #include "Async/PCGAsyncLoadingContext.h"
+#include "Metadata/PCGObjectPropertyOverride.h"
 
 #include "Engine/SplineMeshComponentDescriptor.h"
 
@@ -53,6 +54,10 @@ public:
 	/** Force meshes/materials to load synchronously. */
 	UPROPERTY(BlueprintReadWrite, EditAnywhere, Category = "Settings|Debug")
 	bool bSynchronousLoad = false;
+
+	/** Overrides for spline mesh descriptor. For now it is only support data wide attributes. Does not work for per-control point overrides for now. */
+	UPROPERTY(EditAnywhere, BlueprintReadWrite, Category = Settings)
+	TArray<FPCGObjectPropertyOverrideDescription> SplineMeshOverrideDescriptions;
 };
 
 struct FPCGCreateSplineMeshContext : public FPCGContext, public IPCGAsyncLoadingContext {};
diff --git a/Engine/Plugins/PCG/Source/PCG/Public/Metadata/Accessors/PCGAttributeAccessorKeys.h b/Engine/Plugins/PCG/Source/PCG/Public/Metadata/Accessors/PCGAttributeAccessorKeys.h
--- a/Engine/Plugins/PCG/Source/PCG/Public/Metadata/Accessors/PCGAttributeAccessorKeys.h
+++ b/Engine/Plugins/PCG/Source/PCG/Public/Metadata/Accessors/PCGAttributeAccessorKeys.h
@@ -145,9 +145,9 @@ public:
 	explicit FPCGAttributeAccessorKeysEntries(const TArrayView<PCGMetadataEntryKey>& InEntries);
 	explicit FPCGAttributeAccessorKeysEntries(const TArrayView<const PCGMetadataEntryKey>& InEntries);
 
-	// Iterates on all the entries in the metadata.
-	explicit FPCGAttributeAccessorKeysEntries(const UPCGMetadata* Metadata);
-	explicit FPCGAttributeAccessorKeysEntries(UPCGMetadata* Metadata);
+	// Iterates on all the entries in the metadata. By default, const keys don't have the default value if empty, non-const have it if empty.
+	explicit FPCGAttributeAccessorKeysEntries(const UPCGMetadata* Metadata, bool bAddDefaultValueIfEmpty = false);
+	explicit FPCGAttributeAccessorKeysEntries(UPCGMetadata* Metadata, bool bAddDefaultValueIfEmpty = true);
 
 	virtual int32 GetNum() const override { return Entries.Num(); }
 
@@ -155,7 +155,7 @@ protected:
 	virtual bool GetMetadataEntryKeys(int32 InStart, TArrayView<PCGMetadataEntryKey*> OutEntryKeys) override;
 	virtual bool GetMetadataEntryKeys(int32 InStart, TArrayView<const PCGMetadataEntryKey*> OutEntryKeys) const override;
 
-	void InitializeFromMetadata(const UPCGMetadata* Metadata);
+	void InitializeFromMetadata(const UPCGMetadata* Metadata, bool bAddDefaultValueIfEmpty);
 
 	TArrayView<PCGMetadataEntryKey> Entries;
 	TArray<PCGMetadataEntryKey> ExtractedEntries;
diff --git a/Engine/Plugins/PCG/Source/PCG/Public/Metadata/PCGObjectPropertyOverride.h b/Engine/Plugins/PCG/Source/PCG/Public/Metadata/PCGObjectPropertyOverride.h
--- a/Engine/Plugins/PCG/Source/PCG/Public/Metadata/PCGObjectPropertyOverride.h
+++ b/Engine/Plugins/PCG/Source/PCG/Public/Metadata/PCGObjectPropertyOverride.h
@@ -117,6 +117,8 @@ struct FPCGObjectOverrides
 	/** Initialize the object overrides. Call before using Apply(InputKeyIndex). */
 	void Initialize(const TArray<FPCGObjectPropertyOverrideDescription>& OverrideDescriptions, T* TemplateObject, const UPCGData* SourceData, FPCGContext* Context)
 	{
+		bInitialized = false;
+
 		if (!TemplateObject)
 		{
 			PCGLog::LogErrorOnGraph(NSLOCTEXT("PCGObjectPropertyOverride", "InitializeOverrideFailedNoObject", "Failed to initialize property overrides. No template object was provided."), Context);
@@ -133,7 +135,7 @@ struct FPCGObjectOverrides
 			ClassObject = TemplateObject->StaticStruct();
 		}
 
-		ObjectSingleOverrides.Reserve(OverrideDescriptions.Num());
+		ObjectSingleOverrides.Empty(OverrideDescriptions.Num());
 
 		for (int32 i = 0; i < OverrideDescriptions.Num(); ++i)
 		{
@@ -165,6 +167,8 @@ struct FPCGObjectOverrides
 				PCGLog::LogErrorOnGraph(FText::Format(NSLOCTEXT("PCGObjectPropertyOverride", "InitializeOverrideFailed", "Failed to initialize override '{0}' for property {1} on object '{2}'."), InputSelector.GetDisplayText(), FText::FromString(OutputProperty), FText::FromName(ClassObject->GetFName())), Context);
 			}
 		}
+
+		bInitialized = true;
 	}
 
 	/** Applies each property override to the object by reading from the InputAccessor at the given KeyIndex, and writing to the OutputKey which represents the object property. */
@@ -180,9 +184,13 @@ struct FPCGObjectOverrides
 		return bAllSucceeded;
 	}
 
+	/** Returns true if we have any override to apply */
+	bool IsValid() const { return bInitialized && !ObjectSingleOverrides.IsEmpty(); }
+
 private:
 	FPCGAttributeAccessorKeysSingleObjectPtr<T> OutputKey;
 	TArray<FPCGObjectSingleOverride> ObjectSingleOverrides;
+	bool bInitialized = false;
 };
 
 USTRUCT(BlueprintType, meta=(Deprecated = "5.4", DeprecationMessage="Use FPCGObjectPropertyOverrideDescription instead."))
