[PCG] Refactor texture data initialize to solve blocking texture builds.

If texture async processing is in progress, return false rather than force it to complete, which creates UI popups and triggers deep callstacks of events.

[RN] [PCG] Fix blocking texture builds that could trigger modal editor popups when using Get Texture Data

diff --git a/Engine/Plugins/PCG/Source/PCG/Private/Data/PCGTextureData.cpp b/Engine/Plugins/PCG/Source/PCG/Private/Data/PCGTextureData.cpp
--- a/Engine/Plugins/PCG/Source/PCG/Private/Data/PCGTextureData.cpp
+++ b/Engine/Plugins/PCG/Source/PCG/Private/Data/PCGTextureData.cpp
@@ -17,13 +17,26 @@
 
 namespace PCGTextureSamplingHelpers
 {
-	bool IsTextureCPUAccessible(UTexture2D* Texture)
+	TOptional<bool> IsTextureCPUAccessible(UTexture2D* Texture)
 	{
+		if (!Texture)
+		{
+			return false;
+		}
+
+#if WITH_EDITOR
+		if (!Texture->IsAsyncCacheComplete())
+		{
+			return {};
+		}
+#endif
+
 		FTexturePlatformData* PlatformData = Texture ? Texture->GetPlatformData() : nullptr;
+
 		return PlatformData && PlatformData->GetHasCpuCopy();
 	}
 
-	bool CanGPUTextureBeCPUAccessed(UTexture2D* Texture)
+	TOptional<bool> CanGPUTextureBeCPUAccessed(UTexture2D* Texture)
 	{
 		// SRGB textures need to be GPU sampled.
 		if (!Texture || Texture->SRGB)
@@ -31,6 +44,13 @@ namespace PCGTextureSamplingHelpers
 			return false;
 		}
 
+#if WITH_EDITOR
+		if (!Texture->IsAsyncCacheComplete())
+		{
+			return {};
+		}
+#endif
+
 		FTexturePlatformData* PlatformData = Texture->GetPlatformData();
 
 		// If a CPU copy is available, this is a CPU texture and not a GPU texture, so we return false.
@@ -348,10 +368,36 @@ void UPCGBaseTextureData::CopyBaseTextureData(UPCGBaseTextureData* NewTextureDat
 	NewTextureData->Width = Width;
 }
 
-void UPCGTextureData::Initialize(UTexture* InTexture, uint32 InTextureIndex, const FTransform& InTransform, const TFunction<void()>& PostInitializeCallback, bool bCreateCPUDuplicateEditorOnly)
+void UPCGTextureData::InitializeInternal(UTexture* InTexture, uint32 InTextureIndex, const FTransform& InTransform, bool* bOutInitializeCompleted, bool bCreateCPUDuplicateEditorOnly)
 {
 	TRACE_CPUPROFILER_EVENT_SCOPE(UPCGTextureData::Initialize);
 
+	auto SetInitCompleted = [bOutInitializeCompleted](bool bIsDone)
+	{
+		if (bOutInitializeCompleted)
+		{
+			*bOutInitializeCompleted = bIsDone;
+		}
+	};
+
+	if (bSuccessfullyInitialized)
+	{
+		SetInitCompleted(true);
+		return;
+	}
+
+	if (!InTexture)
+	{
+		SetInitCompleted(true);
+		return;
+	}
+
+	if (bReadbackFromGPUInitiated)
+	{
+		SetInitCompleted(false);
+		return;
+	}
+
 	Texture = InTexture;
 	TextureIndex = InTextureIndex;
 	Transform = InTransform;
@@ -363,9 +409,19 @@ void UPCGTextureData::Initialize(UTexture* InTexture, uint32 InTextureIndex, con
 	Bounds += FVector(1.0f, 1.0f, 0.0f);
 	Bounds = Bounds.TransformBy(Transform);
 
-	if (!InTexture)
+	// Prioritize initializing from a CPU texture when the provided texture is marked as CPU accessible
+	TOptional<bool> InitializedFromCPUTexture = InitializeFromCPUTexture();
+	if (!InitializedFromCPUTexture.IsSet())
 	{
-		PostInitializeCallback();
+		// Wait until we can determine this.
+		SetInitCompleted(false);
+		return;
+	}
+
+	if (*InitializedFromCPUTexture)
+	{
+		bSuccessfullyInitialized = true;
+		SetInitCompleted(true);
 		return;
 	}
 
@@ -380,8 +436,25 @@ void UPCGTextureData::Initialize(UTexture* InTexture, uint32 InTextureIndex, con
 			Texture2D = Texture2DArray->SourceTextures.IsValidIndex(TextureIndex) ? Texture2DArray->SourceTextures[TextureIndex] : nullptr;
 		}
 
-		if (Texture2D && !PCGTextureSamplingHelpers::IsTextureCPUAccessible(Texture2D) && !PCGTextureSamplingHelpers::CanGPUTextureBeCPUAccessed(Texture2D))
+		TOptional<bool> CanGPUTextureBeCPUAccessed = PCGTextureSamplingHelpers::CanGPUTextureBeCPUAccessed(Texture2D);
+		if (!CanGPUTextureBeCPUAccessed.IsSet())
 		{
+			// Wait until we can ascertain this.
+			SetInitCompleted(false);
+			return;
+		}
+
+		TOptional<bool> IsCPUAccessible = PCGTextureSamplingHelpers::IsTextureCPUAccessible(Texture2D);
+		if (!IsCPUAccessible.IsSet())
+		{
+			// Wait until we can ascertain this.
+			SetInitCompleted(false);
+			return;
+		}
+
+		if (Texture2D && !DuplicateTexture && !*CanGPUTextureBeCPUAccessed && !*IsCPUAccessible)
+		{
+			// Duplicate texture and change access flags (editor only). This duplicate texture will be used by the normal logic below.
 			FObjectDuplicationParameters DuplicationParams(Texture2D, /*Outer=*/this);
 			DuplicateTexture = CastChecked<UTexture2D>(StaticDuplicateObjectEx(DuplicationParams));
 			DuplicateTexture->MipGenSettings = TextureMipGenSettings::TMGS_NoMipmaps;
@@ -392,28 +465,49 @@ void UPCGTextureData::Initialize(UTexture* InTexture, uint32 InTextureIndex, con
 	}
 #endif
 
-	// Prioritize initializing from a CPU texture when the provided texture is marked as CPU accessible
-	if (InitializeFromCPUTexture())
+#if WITH_EDITOR
+	// Try reading the texture back from CPU-accessible memory if possible.
+	TOptional<bool> InitGPUTextureFromCPU = InitializeGPUTextureFromCPU();
+	if (!InitGPUTextureFromCPU)
 	{
-		PostInitializeCallback();
+		// Not ready.
+		SetInitCompleted(false);
+		return;
 	}
-	else
+
+	if (*InitGPUTextureFromCPU)
 	{
-#if WITH_EDITOR
-		// Try reading the texture back from CPU-accessible memory if possible.
-		if (InitializeGPUTextureFromCPU())
-		{
-			PostInitializeCallback();
-			return;
-		}
+		bSuccessfullyInitialized = true;
+		SetInitCompleted(true);
+		return;
+	}
 #endif
 
-		if (!InitializeFromGPUTexture(PostInitializeCallback))
-		{
-			UE_LOG(LogPCG, Error, TEXT("PCGTextureData failed to initialize texture '%s'"), *Texture->GetFName().ToString());
+	// Finally try the GPU -> CPU readback path. We don't flag success yet though - this will be done when the readback data arrives.
+	SetInitCompleted(InitializeFromGPUTexture());
 
-			PostInitializeCallback();
-		}
+	return;
+}
+
+bool UPCGTextureData::Initialize(UTexture* InTexture, uint32 InTextureIndex, const FTransform& InTransform, bool bCreateCPUDuplicateEditorOnly)
+{
+	bool bInitializeDone = false;
+	InitializeInternal(InTexture, InTextureIndex, InTransform, &bInitializeDone, bCreateCPUDuplicateEditorOnly);
+	return bInitializeDone;
+}
+
+void UPCGTextureData::Initialize(UTexture* InTexture, uint32 InTextureIndex, const FTransform& InTransform, const TFunction<void()>& InPostInitializeCallback, bool bCreateCPUDuplicateEditorOnly)
+{
+	bool bInitializeDone = false;
+	InitializeInternal(InTexture, InTextureIndex, InTransform, &bInitializeDone, bCreateCPUDuplicateEditorOnly);
+
+	if (bInitializeDone)
+	{
+		InPostInitializeCallback();
+	}
+	else
+	{
+		PostInitializeCallback = InPostInitializeCallback;
 	}
 }
 
@@ -436,7 +530,7 @@ UPCGSpatialData* UPCGTextureData::CopyInternal(FPCGContext* Context) const
 	return NewTextureData;
 }
 
-bool UPCGTextureData::InitializeFromCPUTexture()
+TOptional<bool> UPCGTextureData::InitializeFromCPUTexture()
 {
 	if (!Texture.IsValid())
 	{
@@ -450,6 +544,14 @@ bool UPCGTextureData::InitializeFromCPUTexture()
 		return false;
 	}
 
+#if WITH_EDITOR
+	if (!Texture2D->IsAsyncCacheComplete())
+	{
+		// Wait until texture ready before interrogating it for access options.
+		return {};
+	}
+#endif
+
 	FSharedImageConstRef CPUTextureRef = Texture2D->GetCPUCopy();
 	if (!CPUTextureRef.IsValid())
 	{
@@ -561,12 +663,18 @@ bool UPCGTextureData::InitializeFromCPUTexture()
 	return true;
 }
 
-bool UPCGTextureData::InitializeFromGPUTexture(const TFunction<void()>& PostInitializeCallback)
+bool UPCGTextureData::InitializeFromGPUTexture()
 {
 	if (!Texture.IsValid())
+	{
+		return true;
+	}
+
+	if (bReadbackFromGPUInitiated)
 	{
 		return false;
 	}
+	bReadbackFromGPUInitiated = true;
 
 #if WITH_EDITOR
 	// Force a wait on any incomplete async texture compilation and caching operations
@@ -600,51 +708,66 @@ bool UPCGTextureData::InitializeFromGPUTexture(const TFunction<void()>& PostInit
 		Params.SourceDimensions = FIntPoint(PlatformData->SizeX, PlatformData->SizeY);
 		Params.SourceTextureIndex = TextureIndex;
 
-		FPCGTextureReadbackInterface::Dispatch(Params, [this, PlatformData, PostInitializeCallback](void* OutBuffer, int32 ReadbackWidth, int32 ReadbackHeight)
+		TWeakObjectPtr<UPCGTextureData> ThisWeakPtr(this);
+
+		FPCGTextureReadbackInterface::Dispatch(Params, [ThisWeakPtr, PlatformData](void* OutBuffer, int32 ReadbackWidth, int32 ReadbackHeight)
 		{
 			TRACE_CPUPROFILER_EVENT_SCOPE(UPCGTextureData::Initialize::DispatchCallback);
 
+			UPCGTextureData* This = ThisWeakPtr.IsValid() ? ThisWeakPtr.Get() : nullptr;
+			if (!This)
+			{
+				return;
+			}
+
 			const int32 PixelCount = ReadbackWidth * ReadbackHeight;
 
 			if (PlatformData->SizeX * PlatformData->SizeY != PixelCount)
 			{
 				UE_LOG(LogPCG, Error,
 					TEXT("PCGTextureData readback has different dimensions than the source texture '%s'. Expected (%d, %d), received (%d, %d)."),
-					*Texture->GetFName().ToString(),
+					*This->Texture->GetFName().ToString(),
 					PlatformData->SizeX, PlatformData->SizeY,
 					ReadbackWidth, ReadbackHeight);
 			}
 
 			if (const FColor* FormattedImageData = reinterpret_cast<const FColor*>(OutBuffer))
 			{
-				Width = ReadbackWidth;
-				Height = ReadbackHeight;
-				ColorData.SetNum(PixelCount);
+				This->Width = ReadbackWidth;
+				This->Height = ReadbackHeight;
+				This->ColorData.SetNum(PixelCount);
 
 				for (int32 D = 0; D < PixelCount; ++D)
 				{
-					ColorData[D] = FormattedImageData[D].ReinterpretAsLinear();
+					This->ColorData[D] = FormattedImageData[D].ReinterpretAsLinear();
 				}
 			}
 			else
 			{
-				UE_LOG(LogPCG, Error, TEXT("PCGTextureData unable to get readback results from '%s'"), *Texture->GetFName().ToString());
+				UE_LOG(LogPCG, Error, TEXT("PCGTextureData unable to get readback results from '%s'"), *This->Texture->GetFName().ToString());
 			}
 
-			PostInitializeCallback();
+			This->bSuccessfullyInitialized = true;
+
+			// Deprecated in 5.5, should be removed when the deprecated Initialize() function is removed.
+			if (This->PostInitializeCallback)
+			{
+				This->PostInitializeCallback();
+			}
 		});
 	}
 	else
 	{
 		UE_LOG(LogPCG, Error, TEXT("PCGTextureData failed to acquire texture resource for '%s'"), *Texture->GetFName().ToString());
-		return false;
+		return true;
 	}
 
-	return true;
+	// Not complete - wait for readback result.
+	return false;
 }
 
 #if WITH_EDITOR
-bool UPCGTextureData::InitializeGPUTextureFromCPU()
+TOptional<bool> UPCGTextureData::InitializeGPUTextureFromCPU()
 {
 	// There's a bit of a mix of texture types in this class currently, due to some functionality for readback being 2D-only.
 	UTexture2D* TextureAs2D = Cast<UTexture2D>(Texture.Get());
@@ -657,18 +780,46 @@ bool UPCGTextureData::InitializeGPUTextureFromCPU()
 	}
 
 	UTexture2D* TextureForReadback = nullptr;
-	if (TextureAs2D && PCGTextureSamplingHelpers::CanGPUTextureBeCPUAccessed(TextureAs2D))
+	TOptional<bool> bCanGPUTextureBeCPUAccessed = PCGTextureSamplingHelpers::CanGPUTextureBeCPUAccessed(TextureAs2D);
+	if (!bCanGPUTextureBeCPUAccessed)
+	{
+		return {};
+	}
+
+	if (TextureAs2D && *bCanGPUTextureBeCPUAccessed)
 	{
 		TextureForReadback = TextureAs2D;
 	}
-	else if (PCGTextureSamplingHelpers::CanGPUTextureBeCPUAccessed(DuplicateTexture))
+	else
+	{
+		TOptional<bool> bCanGPUTextureBeCPUAccessedDupe = PCGTextureSamplingHelpers::CanGPUTextureBeCPUAccessed(DuplicateTexture);
+		if (!bCanGPUTextureBeCPUAccessedDupe)
+		{
+			return {};
+		}
+
+		if (*bCanGPUTextureBeCPUAccessedDupe)
+		{
+			TextureForReadback = DuplicateTexture;
+		}
+	}
+
+	if (!TextureForReadback)
 	{
-		TextureForReadback = DuplicateTexture;
+		return false;
+	}
+
+#if WITH_EDITOR
+	if (!TextureForReadback->IsAsyncCacheComplete())
+	{
+		return {};
 	}
+#endif
 
 	FTexturePlatformData* PlatformData = TextureForReadback ? TextureForReadback->GetPlatformData() : nullptr;
 	if (!PlatformData)
 	{
+		UE_LOG(LogPCG, Warning, TEXT("GetPlatformData failed"));
 		return false;
 	}
 
diff --git a/Engine/Plugins/PCG/Source/PCG/Private/Elements/PCGTextureSampler.cpp b/Engine/Plugins/PCG/Source/PCG/Private/Elements/PCGTextureSampler.cpp
--- a/Engine/Plugins/PCG/Source/PCG/Private/Elements/PCGTextureSampler.cpp
+++ b/Engine/Plugins/PCG/Source/PCG/Private/Elements/PCGTextureSampler.cpp
@@ -5,6 +5,7 @@
 #include "PCGComponent.h"
 #include "PCGContext.h"
 #include "PCGCrc.h"
+#include "PCGSubsystem.h"
 #include "Helpers/PCGBlueprintHelpers.h"
 #include "Helpers/PCGDynamicTrackingHelpers.h"
 #include "Helpers/PCGHelpers.h"
@@ -107,6 +108,14 @@ void UPCGTextureSamplerSettings::UpdateDisplayTextureArrayIndex()
 }
 #endif
 
+void FPCGTextureSamplerContext::AddExtraStructReferencedObjects(FReferenceCollector& Collector)
+{
+	if (TextureData)
+	{
+		Collector.AddReferencedObject(TextureData);
+	}
+}
+
 bool FPCGTextureSamplerElement::PrepareDataInternal(FPCGContext* InContext) const
 {
 	TRACE_CPUPROFILER_EVENT_SCOPE(FPCGTextureSamplerElement::Execute);
@@ -190,67 +199,81 @@ bool FPCGTextureSamplerElement::ExecuteInternal(FPCGContext* InContext) const
 	}
 
 	const FTransform& Transform = Settings->Transform;
-	const bool bUseAbsoluteTransform = Settings->bUseAbsoluteTransform;
-	const bool bUseDensitySourceChannel = Settings->bUseDensitySourceChannel;
-	const EPCGTextureColorChannel ColorChannel = Settings->ColorChannel;
-	const EPCGTextureFilter Filter = Settings->Filter;
-	const float TexelSize = Settings->TexelSize;
-	const bool bUseAdvancedTiling = Settings->bUseAdvancedTiling;
-	const FVector2D& Tiling = Settings->Tiling;
-	const FVector2D& CenterOffset = Settings->CenterOffset;
-	const float Rotation = Settings->Rotation;
-	const bool bUseTileBounds = Settings->bUseTileBounds;
-	const FVector2D& TileBoundsMin = Settings->TileBoundsMin;
-	const FVector2D& TileBoundsMax = Settings->TileBoundsMax;
 #if WITH_EDITOR
 	const bool bForceEditorOnlyCPUSampling = Settings->bForceEditorOnlyCPUSampling;
 #else
 	const bool bForceEditorOnlyCPUSampling = false;
 #endif
 
-	TArray<FPCGTaggedData>& Outputs = Context->OutputData.TaggedData;
-	FPCGTaggedData& Output = Outputs.Emplace_GetRef();
+	// Texture data can take some frames to prepare, so we poll it once per frame until it is done.
+	// TODO - review other similar cases and consider adding helpers/abstractions to support this.
+	UPCGTextureData* TextureData = nullptr;
+	if (Context->TextureData)
+	{
+		TextureData = Context->TextureData.Get();
+	}
+	else
+	{
+		TextureData = FPCGContext::NewObject_AnyThread<UPCGTextureData>(Context);
+		Context->TextureData = TextureData;
 
-	UPCGTextureData* TextureData = FPCGContext::NewObject_AnyThread<UPCGTextureData>(Context);
-	Output.Data = TextureData;
+		AActor* OriginalActor = UPCGBlueprintHelpers::GetOriginalComponent(*Context)->GetOwner();
 
-	AActor* OriginalActor = UPCGBlueprintHelpers::GetOriginalComponent(*Context)->GetOwner();
-	FTransform FinalTransform = Transform;
-	if (!bUseAbsoluteTransform)
-	{
-		FTransform OriginalActorTransform = OriginalActor->GetTransform();
-		FinalTransform = Transform * OriginalActorTransform;
+		if (Settings->bUseAbsoluteTransform)
+		{
+			Context->Transform = Transform;
+		}
+		else
+		{
+			FTransform OriginalActorTransform = OriginalActor->GetTransform();
+			Context->Transform = Transform * OriginalActorTransform;
 
-		FBox OriginalActorLocalBounds = PCGHelpers::GetActorLocalBounds(OriginalActor);
-		FinalTransform.SetScale3D(FinalTransform.GetScale3D() * 0.5 * (OriginalActorLocalBounds.Max - OriginalActorLocalBounds.Min));
+			FBox OriginalActorLocalBounds = PCGHelpers::GetActorLocalBounds(OriginalActor);
+			Context->Transform.SetScale3D(Context->Transform.GetScale3D() * 0.5 * (OriginalActorLocalBounds.Max - OriginalActorLocalBounds.Min));
+		}
 	}
 
-	// Initialize & set properties
-	Context->bIsPaused = true;
+	if (!ensure(TextureData))
+	{
+		PCGE_LOG_C(Error, LogOnly, Context, LOCTEXT("TextureDataInitFailed", "Failed to initialize texture data."));
+		return true;
+	}
 
-	auto PostInitializeCallback = [Context, TextureData]()
+	if (!TextureData->IsSuccessfullyInitialized())
 	{
-		Context->bIsPaused = false;
-		Context->bTextureReadbackDone = true;
+		if (!TextureData->Initialize(Texture, TextureArrayIndex, Context->Transform, bForceEditorOnlyCPUSampling))
+		{
+			// Initialization not complete. Could be waiting on async texture processing or for GPU readback. Sleep until next frame.
+			Context->bIsPaused = true;
+			Context->SourceComponent->GetSubsystem()->RegisterBeginTickAction([Context]()
+			{
+				Context->bIsPaused = false;
+			});
+
+			return false;
+		}
 
-		if (!TextureData->IsValid())
+		if (!TextureData->IsSuccessfullyInitialized())
 		{
-			PCGE_LOG_C(Error, GraphAndLog, Context, LOCTEXT("TextureDataInitFailed", "Texture data failed to initialize, check log for more information"));
+			PCGE_LOG(Warning, LogOnly, LOCTEXT("TextureInitFailed", "Data could not be retrieved for this texture, initialization failed."));
+			Context->OutputData.TaggedData.Empty();
+			return true;
 		}
-	};
+	}
 
-	TextureData->Initialize(Texture, TextureArrayIndex, FinalTransform, PostInitializeCallback, bForceEditorOnlyCPUSampling);
+	// Commit to adding texture data.
+	Context->OutputData.TaggedData.Emplace_GetRef().Data = TextureData;
 
-	TextureData->bUseDensitySourceChannel = bUseDensitySourceChannel;
-	TextureData->ColorChannel = ColorChannel;
-	TextureData->Filter = Filter;
-	TextureData->TexelSize = TexelSize;
-	TextureData->bUseAdvancedTiling = bUseAdvancedTiling;
-	TextureData->Tiling = Tiling;
-	TextureData->CenterOffset = CenterOffset;
-	TextureData->Rotation = Rotation;
-	TextureData->bUseTileBounds = bUseTileBounds;
-	TextureData->TileBounds = FBox2D(TileBoundsMin, TileBoundsMax);
+	TextureData->bUseDensitySourceChannel = Settings->bUseDensitySourceChannel;
+	TextureData->ColorChannel = Settings->ColorChannel;
+	TextureData->Filter = Settings->Filter;
+	TextureData->TexelSize = Settings->TexelSize;
+	TextureData->bUseAdvancedTiling = Settings->bUseAdvancedTiling;
+	TextureData->Tiling = Settings->Tiling;
+	TextureData->CenterOffset = Settings->CenterOffset;
+	TextureData->Rotation = Settings->Rotation;
+	TextureData->bUseTileBounds = Settings->bUseTileBounds;
+	TextureData->TileBounds = FBox2D(Settings->TileBoundsMin, Settings->TileBoundsMax);
 
 #if WITH_EDITOR
 	// If we have an override, register for dynamic tracking.
@@ -260,7 +283,7 @@ bool FPCGTextureSamplerElement::ExecuteInternal(FPCGContext* InContext) const
 	}
 #endif // WITH_EDITOR
 
-	return false;
+	return true;
 }
 
 FPCGContext* FPCGTextureSamplerElement::CreateContext()
diff --git a/Engine/Plugins/PCG/Source/PCG/Private/Tests/Data/PCGTextureDataTest.cpp b/Engine/Plugins/PCG/Source/PCG/Private/Tests/Data/PCGTextureDataTest.cpp
--- a/Engine/Plugins/PCG/Source/PCG/Private/Tests/Data/PCGTextureDataTest.cpp
+++ b/Engine/Plugins/PCG/Source/PCG/Private/Tests/Data/PCGTextureDataTest.cpp
@@ -39,71 +39,76 @@ bool FPCGTextureDataOffsetTilingRotation::RunTest(const FString& Parameters)
 
 	bool bTestPassed = true;
 
-	TextureData->Initialize(Texture2D, /*TextureIndex=*/0, FTransform(), [this, TextureData, &bTestPassed]() -> void
-		{
-			FRandomStream RandomStream;
-			FPCGPoint OutPoint;
+	FTransform Transform;
+	while (!TextureData->Initialize(Texture2D, /*TextureIndex=*/0, Transform)) {}
+
+	bTestPassed &= TestTrue("Texture data successfully initialized", TextureData->IsSuccessfullyInitialized());
+
+	if (TextureData->IsSuccessfullyInitialized())
+	{
+		FRandomStream RandomStream;
+		FPCGPoint OutPoint;
 
-			const float PixelCenterX = WhitePixelX + 0.5f;
-			const float PixelCenterY = WhitePixelY + 0.5f;
+		const float PixelCenterX = WhitePixelX + 0.5f;
+		const float PixelCenterY = WhitePixelY + 0.5f;
 
-			const float TextureSpacePixelX = PixelCenterX / TextureSize;
-			const float ScaledPixelX = (2.0f * TextureSpacePixelX) - 1.0f;
-			const float TextureSpacePixelY = PixelCenterY / TextureSize;
-			const float ScaledPixelY = (2.0f * TextureSpacePixelY) - 1.0f;
+		const float TextureSpacePixelX = PixelCenterX / TextureSize;
+		const float ScaledPixelX = (2.0f * TextureSpacePixelX) - 1.0f;
+		const float TextureSpacePixelY = PixelCenterY / TextureSize;
+		const float ScaledPixelY = (2.0f * TextureSpacePixelY) - 1.0f;
 
+		{
+			TextureData->Rotation = 0.f;
+
+			// sampling with no offset
+			TextureData->CenterOffset = FVector2D::ZeroVector;
+			TextureData->SamplePoint(FTransform(), FBox(), OutPoint, nullptr);
+			bTestPassed &= TestEqual("Valid color sampled for no offset at (0, 0)", OutPoint.Color, static_cast<FVector>(FColor::Black));
+
+			// sampling at position, with no offset
+			TextureData->CenterOffset = FVector2D::ZeroVector;
+			TextureData->SamplePoint(FTransform(FVector(0.5, 0.5, 0.0)), FBox(), OutPoint, nullptr);
+			bTestPassed &= TestEqual("Valid color sampled for no offset at position", OutPoint.Color, static_cast<FVector>(FColor::Black));
+
+			// sampling at expected position
+			TextureData->CenterOffset = FVector2D::ZeroVector;
+			TextureData->SamplePoint(FTransform(FVector(ScaledPixelX, ScaledPixelY, 0)), FBox(), OutPoint, nullptr);
+			bTestPassed &= TestEqual("Valid color sampled at expected position", OutPoint.Color, static_cast<FVector>(FColor::White));
+
+			// sampling with offset
+			TextureData->CenterOffset = FVector2D(0.5 - TextureSpacePixelX, 0.5 - TextureSpacePixelY);
+			TextureData->SamplePoint(FTransform(), FBox(), OutPoint, nullptr);
+			bTestPassed &= TestEqual("Valid color sampled for offset from (0, 0)", OutPoint.Color, static_cast<FVector>(FColor::White));
+
+			// sampling at position, with offset
+			TextureData->CenterOffset = FVector2D(1.0 - TextureSpacePixelX, 1.0 - TextureSpacePixelY);
+			TextureData->SamplePoint(FTransform(FVector(1.0, 1.0, 0.0)), FBox(), OutPoint, nullptr);
+			bTestPassed &= TestEqual("Valid color sampled for offset from position", OutPoint.Color, static_cast<FVector>(FColor::White));
+
+			for (float Rotation = -360.f; Rotation < 360.f; Rotation += 10.f)
 			{
-				TextureData->Rotation = 0.f;
+				const float Theta = FMath::DegreesToRadians(Rotation);
+				const float CosTheta = FMath::Cos(Theta);
+				const float SinTheta = FMath::Sin(Theta);
+				const FVector::FReal X = (ScaledPixelX * CosTheta) - (ScaledPixelY * SinTheta);
+				const FVector::FReal Y = (ScaledPixelY * CosTheta) + (ScaledPixelX * SinTheta);
 
-				// sampling with no offset
-				TextureData->CenterOffset = FVector2D::ZeroVector;
-				TextureData->SamplePoint(FTransform(), FBox(), OutPoint, nullptr);
-				bTestPassed &= TestEqual("Valid color sampled for no offset at (0, 0)", OutPoint.Color, static_cast<FVector>(FColor::Black));
+				FTransform RotatedTransform(FVector(X, Y, 0.f));
 
-				// sampling at position, with no offset
+				// sampling with rotation at black position
+				TextureData->Rotation = -Rotation;
 				TextureData->CenterOffset = FVector2D::ZeroVector;
 				TextureData->SamplePoint(FTransform(FVector(0.5, 0.5, 0.0)), FBox(), OutPoint, nullptr);
-				bTestPassed &= TestEqual("Valid color sampled for no offset at position", OutPoint.Color, static_cast<FVector>(FColor::Black));
+				bTestPassed &= TestEqual("Valid color sampled for rotated off-position", OutPoint.Color, static_cast<FVector>(FColor::Black));
 
-				// sampling at expected position
+				// sampling with rotation at white position
+				TextureData->Rotation = -Rotation;
 				TextureData->CenterOffset = FVector2D::ZeroVector;
-				TextureData->SamplePoint(FTransform(FVector(ScaledPixelX, ScaledPixelY, 0)), FBox(), OutPoint, nullptr);
-				bTestPassed &= TestEqual("Valid color sampled at expected position", OutPoint.Color, static_cast<FVector>(FColor::White));
-
-				// sampling with offset
-				TextureData->CenterOffset = FVector2D(0.5 - TextureSpacePixelX, 0.5 - TextureSpacePixelY);
-				TextureData->SamplePoint(FTransform(), FBox(), OutPoint, nullptr);
-				bTestPassed &= TestEqual("Valid color sampled for offset from (0, 0)", OutPoint.Color, static_cast<FVector>(FColor::White));
-
-				// sampling at position, with offset
-				TextureData->CenterOffset = FVector2D(1.0 - TextureSpacePixelX, 1.0 - TextureSpacePixelY);
-				TextureData->SamplePoint(FTransform(FVector(1.0, 1.0, 0.0)), FBox(), OutPoint, nullptr);
-				bTestPassed &= TestEqual("Valid color sampled for offset from position", OutPoint.Color, static_cast<FVector>(FColor::White));
-
-				for (float Rotation = -360.f; Rotation < 360.f; Rotation += 10.f)
-				{
-					const float Theta = FMath::DegreesToRadians(Rotation);
-					const float CosTheta = FMath::Cos(Theta);
-					const float SinTheta = FMath::Sin(Theta);
-					const FVector::FReal X = (ScaledPixelX * CosTheta) - (ScaledPixelY * SinTheta);
-					const FVector::FReal Y = (ScaledPixelY * CosTheta) + (ScaledPixelX * SinTheta);
-
-					FTransform RotatedTransform(FVector(X, Y, 0.f));
-
-					// sampling with rotation at black position
-					TextureData->Rotation = -Rotation;
-					TextureData->CenterOffset = FVector2D::ZeroVector;
-					TextureData->SamplePoint(FTransform(FVector(0.5, 0.5, 0.0)), FBox(), OutPoint, nullptr);
-					bTestPassed &= TestEqual("Valid color sampled for rotated off-position", OutPoint.Color, static_cast<FVector>(FColor::Black));
-
-					// sampling with rotation at white position
-					TextureData->Rotation = -Rotation;
-					TextureData->CenterOffset = FVector2D::ZeroVector;
-					TextureData->SamplePoint(RotatedTransform, FBox(), OutPoint, nullptr);
-					bTestPassed &= TestEqual("Valid color sampled for rotated position", OutPoint.Color, static_cast<FVector>(FColor::White));
-				}
+				TextureData->SamplePoint(RotatedTransform, FBox(), OutPoint, nullptr);
+				bTestPassed &= TestEqual("Valid color sampled for rotated position", OutPoint.Color, static_cast<FVector>(FColor::White));
 			}
-		});
+		}
+	}
 
 	return bTestPassed;
 }
diff --git a/Engine/Plugins/PCG/Source/PCG/Private/Tests/Elements/PCGSampleTextureTest.cpp b/Engine/Plugins/PCG/Source/PCG/Private/Tests/Elements/PCGSampleTextureTest.cpp
--- a/Engine/Plugins/PCG/Source/PCG/Private/Tests/Elements/PCGSampleTextureTest.cpp
+++ b/Engine/Plugins/PCG/Source/PCG/Private/Tests/Elements/PCGSampleTextureTest.cpp
@@ -67,9 +67,16 @@ protected:
 		Texture->UpdateResource();
 
 		UPCGTextureData* TextureData = NewObject<UPCGTextureData>();
-		TextureData->Initialize(Texture, /*InTextureIndex*/0, FTransform(), /*PostInitializeCallback*/[]() {});
+		while (!TextureData->Initialize(Texture, /*InTextureIndex*/0, FTransform())) {}
 
-		return TextureData;
+		if (ensure(TextureData->IsSuccessfullyInitialized()))
+		{
+			return TextureData;
+		}
+		else
+		{
+			return nullptr;
+		}
 	}
 	
 	bool GenerateTestDataRunAndValidate(const FTestParameters& Parameters)
diff --git a/Engine/Plugins/PCG/Source/PCG/Public/Data/PCGTextureData.h b/Engine/Plugins/PCG/Source/PCG/Public/Data/PCGTextureData.h
--- a/Engine/Plugins/PCG/Source/PCG/Public/Data/PCGTextureData.h
+++ b/Engine/Plugins/PCG/Source/PCG/Public/Data/PCGTextureData.h
@@ -44,10 +44,10 @@ enum class EPCGTextureAddressMode : uint8
 namespace PCGTextureSamplingHelpers
 {
 	/** Returns true if a texture is CPU-accessible. */
-	bool IsTextureCPUAccessible(UTexture2D* Texture);
+	TOptional<bool> IsTextureCPUAccessible(UTexture2D* Texture);
 
 	/** Returns true if a texture is both GPU-accessible and reachable from CPU memory. */
-	bool CanGPUTextureBeCPUAccessed(UTexture2D* Texture);
+	TOptional<bool> CanGPUTextureBeCPUAccessed(UTexture2D* Texture);
 }
 
 UCLASS(Abstract)
@@ -153,7 +153,17 @@ class UPCGTextureData : public UPCGBaseTextureData
 	GENERATED_BODY()
 
 public:
-	PCG_API void Initialize(UTexture* InTexture, uint32 InTextureIndex, const FTransform& InTransform, const TFunction<void()>& PostInitializeCallback, bool bCreateCPUDuplicateEditorOnly = false);
+	/**
+	 * Initialize this data. Can depend on async texture operations / async GPU readbacks. Should be polled until it returns true signaling completion,
+	 * and then IsInitialized() is used to verify the initialization was successful and data is ready to use.
+	 */
+	PCG_API bool Initialize(UTexture* InTexture, uint32 InTextureIndex, const FTransform& InTransform, bool bCreateCPUDuplicateEditorOnly = false);
+
+	UE_DEPRECATED(5.5, "Will be removed. Poll the alternate Initialize API until it returns true instead of passing in a callback.")
+	PCG_API void Initialize(UTexture* InTexture, uint32 InTextureIndex, const FTransform& InTransform, const TFunction<void()>& InPostInitializeCallback, bool bCreateCPUDuplicateEditorOnly = false);
+
+	/** Data is successfully initialized and is ready to use. */
+	PCG_API bool IsSuccessfullyInitialized() const { return bSuccessfullyInitialized; }
 
 	// ~Begin UPCGData interface
 	virtual EPCGDataType GetDataType() const override { return EPCGDataType::Texture; }
@@ -165,19 +175,21 @@ protected:
 	virtual UPCGSpatialData* CopyInternal(FPCGContext* Context) const override;
 	//~End UPCGSpatialData interface
 
+	void InitializeInternal(UTexture* InTexture, uint32 InTextureIndex, const FTransform& InTransform, bool* bOutInitializeDone, bool bCreateCPUDuplicateEditorOnly = false);
+
 private:
 	/** Attempts to initialize the UPCGTextureData from a CPU-accessible texture. Returns true if CPU initialization succeeds. */
-	bool InitializeFromCPUTexture();
+	TOptional<bool> InitializeFromCPUTexture();
 
 	/** 
 	* Attempts to initialize the UPCGTextureData from a GPU-accessible texture. Returns true if GPU texture readback can be dispatched. 
 	* The PostInitializeCallback is only executed if this function succeeds.
 	*/
-	bool InitializeFromGPUTexture(const TFunction<void()>& PostInitializeCallback);
+	bool InitializeFromGPUTexture();
 
 #if WITH_EDITOR
 	/** Attempts to initialize the UPCGTextureData from a GPU-accessible texture, but with CPU-accessible memory. Returns true if initialization succeeds. */
-	bool InitializeGPUTextureFromCPU();
+	TOptional<bool> InitializeGPUTextureFromCPU();
 #endif
 
 public:
@@ -192,6 +204,15 @@ public:
 
 	UPROPERTY(BlueprintReadOnly, VisibleAnywhere, Category = Properties)
 	int TextureIndex = 0;
+
+	UPROPERTY(BlueprintReadOnly, VisibleAnywhere, Category = Properties)
+	bool bSuccessfullyInitialized = false;
+
+	UPROPERTY(BlueprintReadOnly, VisibleAnywhere, Category = Properties)
+	bool bReadbackFromGPUInitiated = false;
+
+	// Added to help deprecation in 5.5. To be removed when the deprecated Initialized function is removed.
+	TFunction<void()> PostInitializeCallback;
 };
 
 #if UE_ENABLE_INCLUDE_ORDER_DEPRECATED_IN_5_2
diff --git a/Engine/Plugins/PCG/Source/PCG/Public/Elements/PCGTextureSampler.h b/Engine/Plugins/PCG/Source/PCG/Public/Elements/PCGTextureSampler.h
--- a/Engine/Plugins/PCG/Source/PCG/Public/Elements/PCGTextureSampler.h
+++ b/Engine/Plugins/PCG/Source/PCG/Public/Elements/PCGTextureSampler.h
@@ -141,7 +141,15 @@ protected:
 
 struct FPCGTextureSamplerContext : public FPCGContext, public IPCGAsyncLoadingContext
 {
+public:
 	bool bTextureReadbackDone = false;
+
+	TObjectPtr<UPCGTextureData> TextureData;
+
+	FTransform Transform;
+
+protected:
+	virtual void AddExtraStructReferencedObjects(FReferenceCollector& Collector) override;
 };
 
 class FPCGTextureSamplerElement : public IPCGElement
