Added option to zoom on the data selected in the Attribute List View, and also added a contextual menu to zoom on selected rows (similar to the double-click).

diff --git a/Engine/Plugins/PCG/Source/PCG/Public/PCGDataVisualization.h b/Engine/Plugins/PCG/Source/PCG/Public/PCGDataVisualization.h
--- a/Engine/Plugins/PCG/Source/PCG/Public/PCGDataVisualization.h
+++ b/Engine/Plugins/PCG/Source/PCG/Public/PCGDataVisualization.h
@@ -40,7 +40,7 @@ struct FPCGTableVisualizerInfo
 	TArray<FPCGTableVisualizerColumnInfo> ColumnInfos;
 	EPCGTableVisualizerColumnSortingMode SortingMode = EPCGTableVisualizerColumnSortingMode::Ascending;
 	FName SortingColumn = NAME_None;
-	TFunction<void(const UPCGData*, int)> DoubleClickCallback = nullptr;
+	TFunction<void(const UPCGData*, TArrayView<const int>)> FocusOnDataCallback = nullptr;
 	TSharedPtr<const IPCGAttributeAccessorKeys> AccessorKeys;
 };
 
diff --git a/Engine/Plugins/PCG/Source/PCGEditor/Private/DataVisualizations/PCGSpatialDataVisualization.cpp b/Engine/Plugins/PCG/Source/PCGEditor/Private/DataVisualizations/PCGSpatialDataVisualization.cpp
--- a/Engine/Plugins/PCG/Source/PCGEditor/Private/DataVisualizations/PCGSpatialDataVisualization.cpp
+++ b/Engine/Plugins/PCG/Source/PCGEditor/Private/DataVisualizations/PCGSpatialDataVisualization.cpp
@@ -303,17 +303,42 @@ FPCGTableVisualizerInfo IPCGSpatialDataVisualization::GetTableVisualizerInfo(con
 	// Column Sorting
 	Info.SortingColumn = NAME_Index;
 
-	// Double Click Behavior
-	Info.DoubleClickCallback = [](const UPCGData* Data, int Index)
+	// Focus on data behavior
+	Info.FocusOnDataCallback = [](const UPCGData* Data, TArrayView<const int> Indices)
 	{
-		if (const UPCGPointData* PointData = Cast<UPCGPointData>(Data))
+		if (const UPCGSpatialData* SpatialData = Cast<UPCGSpatialData>(Data))
 		{
+			const UPCGPointData* PointData = SpatialData->ToPointData(nullptr);
+
+			if (!PointData)
+			{
+				return;
+			}
+
 			const TArray<FPCGPoint>& Points = PointData->GetPoints();
-			check(Points.IsValidIndex(Index));
 
-			const FPCGPoint& Point = Points[Index];
-			const FBox BoundingBox = Point.GetLocalBounds().TransformBy(Point.Transform.ToMatrixWithScale());
-			GEditor->MoveViewportCamerasToBox(BoundingBox, /*bActiveViewportOnly=*/true, /*DrawDebugBoxTimeInSeconds=*/2.5f);
+			FBox BoundingBox(EForceInit::ForceInit);
+			if (Indices.IsEmpty())
+			{
+				BoundingBox = PointData->GetBounds();
+			}
+			else
+			{
+				for (const int& Index : Indices)
+				{
+					check(Points.IsValidIndex(Index));
+
+					const FPCGPoint& Point = Points[Index];
+					const FBox PointBoundingBox = Point.GetLocalBounds().TransformBy(Point.Transform.ToMatrixWithScale());
+
+					BoundingBox += PointBoundingBox;
+				}
+			}
+
+			if (GEditor && BoundingBox.IsValid)
+			{
+				GEditor->MoveViewportCamerasToBox(BoundingBox, /*bActiveViewportOnly=*/true, /*DrawDebugBoxTimeInSeconds=*/2.5f);
+			}
 		}
 	};
 
diff --git a/Engine/Plugins/PCG/Source/PCGEditor/Private/DataVisualizations/PCGSplineDataVisualization.cpp b/Engine/Plugins/PCG/Source/PCGEditor/Private/DataVisualizations/PCGSplineDataVisualization.cpp
--- a/Engine/Plugins/PCG/Source/PCGEditor/Private/DataVisualizations/PCGSplineDataVisualization.cpp
+++ b/Engine/Plugins/PCG/Source/PCGEditor/Private/DataVisualizations/PCGSplineDataVisualization.cpp
@@ -15,20 +15,37 @@ FPCGTableVisualizerInfo IPCGSplineDataVisualization::GetTableVisualizerInfo(cons
 {
 	FPCGTableVisualizerInfo Info = IPCGSpatialDataVisualization::GetTableVisualizerInfo(Data);
 
-	Info.DoubleClickCallback = [](const UPCGData* Data, int Index)
+	Info.FocusOnDataCallback = [](const UPCGData* Data, TArrayView<const int> Indices)
 	{
 		if (const UPCGSplineData* SplineData = Cast<UPCGSplineData>(Data))
 		{
-			const TArray<FInterpCurvePointVector>& Positions = SplineData->SplineStruct.SplineCurves.Position.Points;
-			const TArray<FInterpCurvePointVector>& Scales = SplineData->SplineStruct.SplineCurves.Scale.Points;
+			FBox BoundingBox(EForceInit::ForceInit);
 
-			check(Positions.IsValidIndex(Index) && Scales.IsValidIndex(Index));
+			if (Indices.IsEmpty())
+			{
+				BoundingBox = SplineData->GetBounds();
+			}
+			else
+			{
+				const TArray<FInterpCurvePointVector>& Positions = SplineData->SplineStruct.SplineCurves.Position.Points;
+				const TArray<FInterpCurvePointVector>& Scales = SplineData->SplineStruct.SplineCurves.Scale.Points;
 
-			const FVector& Position = SplineData->GetTransform().TransformPosition(Positions[Index].OutVal);
-			const FVector HalfExtent = Scales[Index].OutVal * PCGSplineDataVisualizationConstants::HalfExtents;
+				for (const int& Index : Indices)
+				{
+					check(Positions.IsValidIndex(Index) && Scales.IsValidIndex(Index));
 
-			FBox BoundingBox(Position + HalfExtent, Position - HalfExtent);
-			GEditor->MoveViewportCamerasToBox(BoundingBox, /*bActiveViewportOnly=*/true, /*DrawDebugBoxTimeInSeconds=*/2.5f);
+					const FVector& Position = SplineData->GetTransform().TransformPosition(Positions[Index].OutVal);
+					const FVector HalfExtent = Scales[Index].OutVal * PCGSplineDataVisualizationConstants::HalfExtents;
+
+					FBox PointBoundingBox(Position + HalfExtent, Position - HalfExtent);
+					BoundingBox += PointBoundingBox;
+				}
+			}
+
+			if (GEditor && BoundingBox.IsValid)
+			{
+				GEditor->MoveViewportCamerasToBox(BoundingBox, /*bActiveViewportOnly=*/true, /*DrawDebugBoxTimeInSeconds=*/2.5f);
+			}
 		}
 	};
 
diff --git a/Engine/Plugins/PCG/Source/PCGEditor/Private/PCGEditorStyle.cpp b/Engine/Plugins/PCG/Source/PCGEditor/Private/PCGEditorStyle.cpp
--- a/Engine/Plugins/PCG/Source/PCGEditor/Private/PCGEditorStyle.cpp
+++ b/Engine/Plugins/PCG/Source/PCGEditor/Private/PCGEditorStyle.cpp
@@ -26,6 +26,7 @@ FPCGEditorStyle::FPCGEditorStyle() : FSlateStyleSet("PCGEditorStyle")
 {
 	static const FVector2D Icon16x16(16.0f, 16.0f);
 	static const FVector2D Icon20x20(20.0f, 20.0f);
+	static const FVector2D Icon32x32(32.0f, 32.0f);
 	static const FVector2D Icon64x64(64.0f, 64.0f);
 	static const FVector2D Icon128x128(128.0f, 128.0f);
 	
@@ -123,7 +124,9 @@ FPCGEditorStyle::FPCGEditorStyle() : FSlateStyleSet("PCGEditorStyle")
 	Set("PCG.Command.GraphSettings", new IMAGE_BRUSH_SVG("Style/PCG_Command_GraphSettings", Icon20x20));
 	Set("PCG.Command.OpenDebugTreeTab", new IMAGE_BRUSH_SVG("Style/PCG_Command_OpenDebugTreeTab", Icon20x20));
 	Set("PCG.Command.RunDeterminismTest", new IMAGE_BRUSH_SVG("Style/PCG_Command_RunDeterminismTest", Icon20x20));
-	
+
+	Set("PCG.Editor.CopyToClipboard", new IMAGE_BRUSH("Style/PCG_Editor_CopyToClipboard", Icon32x32));
+	Set("PCG.Editor.ZoomToSelection", new IMAGE_BRUSH_SVG("Style/PCG_Editor_ZoomToSelection", Icon16x16));
 }
 
 const FPCGEditorStyle& FPCGEditorStyle::Get()
diff --git a/Engine/Plugins/PCG/Source/PCGEditor/Private/SPCGEditorGraphAttributeListView.cpp b/Engine/Plugins/PCG/Source/PCGEditor/Private/SPCGEditorGraphAttributeListView.cpp
--- a/Engine/Plugins/PCG/Source/PCGEditor/Private/SPCGEditorGraphAttributeListView.cpp
+++ b/Engine/Plugins/PCG/Source/PCGEditor/Private/SPCGEditorGraphAttributeListView.cpp
@@ -23,6 +23,7 @@
 #include "PCGEditor.h"
 #include "PCGEditorGraph.h"
 #include "PCGEditorGraphNodeBase.h"
+#include "PCGEditorStyle.h"
 
 #include "Fonts/FontMeasure.h"
 #include "Framework/Commands/GenericCommands.h"
@@ -96,10 +97,10 @@ void FPCGListViewUpdater::AsyncSort()
 
 void FPCGListViewUpdater::AsyncFilter()
 {
-	TArray<PCGListviewItemPtr> FilteredListViewItems;
+	TArray<PCGListViewItemPtr> FilteredListViewItems;
 	FilteredListViewItems.Reserve(ListViewItems.Num());
 
-	for (const PCGListviewItemPtr& ListViewItem : ListViewItems)
+	for (const PCGListViewItemPtr& ListViewItem : ListViewItems)
 	{
 		const FPCGPointFilterExpressionContext PointFilterContext(ListViewItem.Get(), &ColumnData);
 		if (TextFilter->TestTextFilter(PointFilterContext))
@@ -117,7 +118,7 @@ void SPCGListViewItemRow::Construct(const FArguments& InArgs, const TSharedRef<S
 	InternalItem = InArgs._ListViewItem;
 	AttributeListView = InArgs._AttributeListView;
 
-	SMultiColumnTableRow<PCGListviewItemPtr>::Construct(
+	SMultiColumnTableRow<PCGListViewItemPtr>::Construct(
 		SMultiColumnTableRow::FArguments()
 		.Style(FAppStyle::Get(), "DataTableEditor.CellListViewRow"),
 		InOwnerTableView);
@@ -293,11 +294,12 @@ void SPCGEditorGraphAttributeListView::Construct(const FArguments& InArgs, TShar
 		return (ListViewItems.IsEmpty() && ListViewHeader->GetColumns().IsEmpty()) ? EVisibility::Hidden : EVisibility::Visible;
 	};
 
-	SAssignNew(ListView, SListView<PCGListviewItemPtr>)
+	SAssignNew(ListView, SListView<PCGListViewItemPtr>)
 		.ListItemsSource(&ListViewItems)
 		.HeaderRow(ListViewHeader)
 		.OnGenerateRow(this, &SPCGEditorGraphAttributeListView::OnGenerateRow)
 		.OnMouseButtonDoubleClick(this, &SPCGEditorGraphAttributeListView::OnItemDoubleClicked)
+		.OnContextMenuOpening(this, &SPCGEditorGraphAttributeListView::OnItemsContextMenu)
 		.AllowOverscroll(EAllowOverscroll::No)
 		.ExternalScrollbar(VerticalScrollBar)
 		.Visibility_Lambda(VisibilityTest)
@@ -334,13 +336,27 @@ void SPCGEditorGraphAttributeListView::Construct(const FArguments& InArgs, TShar
 		.ContentPadding(FMargin(4, 2))
 		.HAlign(HAlign_Center)
 		.VAlign(VAlign_Center)
-		.ToolTipText(LOCTEXT("LockSelectionButton_ToolTip", "Locks the current attribute list view to this selection"))
+		.ToolTipText(LOCTEXT("LockSelectionButton_ToolTip", "Locks the current attribute list view to this selection."))
 		[
 			SNew(SImage)
 				.ColorAndOpacity(FSlateColor::UseForeground())
 				.Image(this, &SPCGEditorGraphAttributeListView::OnGetLockButtonImageResource)
 		];
 
+	TSharedPtr<SButton> FrameDataButton = SNew(SButton)
+		.ButtonStyle(FAppStyle::Get(), "SimpleButton")
+		.OnClicked(this, &SPCGEditorGraphAttributeListView::OnFocusOnDataClicked)
+		.IsEnabled(this, &SPCGEditorGraphAttributeListView::IsFocusOnDataEnabled)
+		.ContentPadding(FMargin(4, 2))
+		.HAlign(HAlign_Center)
+		.VAlign(VAlign_Center)
+		.ToolTipText(LOCTEXT("FocusOnDataButton_ToolTip", "Zoom to selected data."))
+		[
+			SNew(SImage)
+				.ColorAndOpacity(FSlateColor::UseForeground())
+				.Image(FPCGEditorStyle::Get().GetBrush("PCG.Editor.ZoomToSelection"))
+		];
+		
 	TSharedPtr<SComboButton> FilterButton = SNew(SComboButton)
 		.ForegroundColor(FSlateColor::UseStyle())
 		.HasDownArrow(false)
@@ -388,6 +404,12 @@ void SPCGEditorGraphAttributeListView::Construct(const FArguments& InArgs, TShar
 			+SHorizontalBox::Slot()
 			.AutoWidth()
 			.Padding(1.0f, 0.0f)
+			[
+				FrameDataButton->AsShared()
+			]
+			+SHorizontalBox::Slot()
+			.AutoWidth()
+			.Padding(1.0f, 0.0f)
 			[
 				FilterButton->AsShared()
 			]
@@ -651,7 +673,7 @@ void SPCGEditorGraphAttributeListView::RefreshAttributeList()
 	HiddenAttributes = ListViewHeader->GetHiddenColumnIds();
 
 	// Swapping to an empty item list to force a widget clear, otherwise the widgets will try to update during add column and access invalid data
-	static const TArray<PCGListviewItemPtr> EmptyList;
+	static const TArray<PCGListViewItemPtr> EmptyList;
 	ListView->SetItemsSource(&EmptyList);
 
 	PCGColumnData.Empty();
@@ -698,21 +720,24 @@ void SPCGEditorGraphAttributeListView::RefreshAttributeList()
 
 		SortingColumn = TableVisualizerInfo.SortingColumn;
 		SortMode = (EColumnSortMode::Type)TableVisualizerInfo.SortingMode;
+		FocusOnDataCallback = TableVisualizerInfo.FocusOnDataCallback;
 
 		const uint32 NumEntries = TableVisualizerInfo.AccessorKeys.IsValid() ? TableVisualizerInfo.AccessorKeys->GetNum() : 0;
 		ListViewItems.Reserve(NumEntries);
 
-		const TArray<FPCGPoint>& Points = Cast<UPCGPointData>(DataStrongPtr.Get()) ? Cast<UPCGPointData>(DataStrongPtr.Get())->GetPoints() : TArray<FPCGPoint>();
-
 		for (uint32 Index = 0; Index < NumEntries; ++Index)
 		{
-			PCGListviewItemPtr ListViewItem = MakeShared<FPCGListViewItem>();
+			PCGListViewItemPtr ListViewItem = MakeShared<FPCGListViewItem>();
 			ListViewItem->Index = Index;
-			ListViewItem->DoubleClickCallback = TableVisualizerInfo.DoubleClickCallback;
-
 			ListViewItems.Add(ListViewItem);
 		}
 	}
+	else // No visualization, just default back the values
+	{
+		SortingColumn = NAME_None;
+		SortMode = EColumnSortMode::Type::Ascending;
+		FocusOnDataCallback.Reset();
+	}
 
 	ListView->SetItemsSource(&ListViewItems);
 	ListView->RequestListRefresh();
@@ -1077,23 +1102,46 @@ bool SPCGEditorGraphAttributeListView::IsAttributeEnabled(FName InAttributeName)
 	return ListViewHeader->IsColumnVisible(InAttributeName);
 }
 
-TSharedRef<ITableRow> SPCGEditorGraphAttributeListView::OnGenerateRow(PCGListviewItemPtr Item, const TSharedRef<STableViewBase>& OwnerTable)
+TSharedRef<ITableRow> SPCGEditorGraphAttributeListView::OnGenerateRow(PCGListViewItemPtr Item, const TSharedRef<STableViewBase>& OwnerTable)
 {
 	return SNew(SPCGListViewItemRow, OwnerTable)
 		.AttributeListView(SharedThis(this))
 		.ListViewItem(Item);
 }
 
-void SPCGEditorGraphAttributeListView::OnItemDoubleClicked(PCGListviewItemPtr Item) const
+void SPCGEditorGraphAttributeListView::OnItemDoubleClicked(PCGListViewItemPtr Item) const
 {
 	check(Item);
-
-	if (Item->DoubleClickCallback)
+	if (FocusOnDataCallback)
 	{
-		Item->DoubleClickCallback(DataStrongPtr.Get(), Item->Index);
+		FocusOnDataCallback(DataStrongPtr.Get(), { Item->Index });
 	}
 }
 
+TSharedPtr<SWidget> SPCGEditorGraphAttributeListView::OnItemsContextMenu()
+{
+	FMenuBuilder MenuBuilder(/*bInShouldCloseWindowAfterMenuSelection=*/true, /*InCommandList=*/nullptr);
+	MenuBuilder.AddMenuEntry(
+		LOCTEXT("ShowInViewport", "Zoom to selection"),
+		LOCTEXT("ShowInViewport_Tooltip", "Frames the viewport so that all selected entries are visible."),
+		FSlateIcon(FPCGEditorStyle::Get().GetStyleSetName(), "PCG.Editor.ZoomToSelection"),
+		FUIAction(FExecuteAction::CreateSP(this, &SPCGEditorGraphAttributeListView::FocusOnSelection),
+			FCanExecuteAction::CreateSP(this, &SPCGEditorGraphAttributeListView::CanFocusOnSelection)),
+		NAME_None,
+		EUserInterfaceActionType::Button);
+
+	MenuBuilder.AddMenuEntry(
+		LOCTEXT("CopyToClipboard", "Copy to clipboard"),
+		LOCTEXT("CopyToClipboard_Tooltip", "Copies the contents of the entries to the clipboard."),
+		FSlateIcon(FPCGEditorStyle::Get().GetStyleSetName(), "PCG.Editor.CopyToClipboard"),
+		FUIAction(FExecuteAction::CreateSP(this, &SPCGEditorGraphAttributeListView::CopySelectionToClipboard),
+			FCanExecuteAction::CreateSP(this, &SPCGEditorGraphAttributeListView::CanCopySelectionToClipboard)),
+		NAME_None,
+		EUserInterfaceActionType::Button);
+
+	return MenuBuilder.MakeWidget();
+}
+
 void SPCGEditorGraphAttributeListView::OnColumnSortModeChanged(const EColumnSortPriority::Type InSortPriority, const FName& InColumnId, const EColumnSortMode::Type InNewSortMode)
 {
 	if (SortingColumn == InColumnId)
@@ -1206,7 +1254,7 @@ void SPCGEditorGraphAttributeListView::CopySelectionToClipboard() const
 	}
 
 	// Gather selected rows and sort them to match the displayed order instead of selection order
-	TArray<PCGListviewItemPtr> SelectedListViewItems = ListView->GetSelectedItems();
+	TArray<PCGListViewItemPtr> SelectedListViewItems = ListView->GetSelectedItems();
 	if (const FPCGColumnData* ColumnData = PCGColumnData.Find(SortingColumn))
 	{
 		if (ColumnData->DataAccessor.IsValid() && ColumnData->DataKeys.IsValid())
@@ -1217,7 +1265,7 @@ void SPCGEditorGraphAttributeListView::CopySelectionToClipboard() const
 	}
 
 	// Write each row
-	for (const PCGListviewItemPtr&  ListViewItem : SelectedListViewItems)
+	for (const PCGListViewItemPtr&  ListViewItem : SelectedListViewItems)
 	{
 		CSVExport += LineEnd;
 
@@ -1270,4 +1318,47 @@ FReply SPCGEditorGraphAttributeListView::OnNodeNameClicked()
 	return FReply::Handled();
 }
 
+FReply SPCGEditorGraphAttributeListView::OnFocusOnDataClicked() const
+{
+	if (IsFocusOnDataEnabled())
+	{
+		FocusOnDataCallback(DataStrongPtr.Get(), {});
+	}
+
+	return FReply::Handled();
+}
+
+bool SPCGEditorGraphAttributeListView::IsFocusOnDataEnabled() const
+{
+	return (!ListViewItems.IsEmpty() && FocusOnDataCallback);
+}
+
+void SPCGEditorGraphAttributeListView::FocusOnSelection() const
+{
+	if (!IsFocusOnDataEnabled())
+	{
+		return;
+	}
+
+	// Note: this implementation assumes it's the same callback for all entries, which is currently true
+	const TArray<PCGListViewItemPtr> SelectedItems = ListView->GetSelectedItems();
+	TArray<int> Indices;
+	Indices.Reserve(SelectedItems.Num());
+
+	for (PCGListViewItemPtr SelectedItem : SelectedItems)
+	{
+		if (SelectedItem)
+		{
+			Indices.Add(SelectedItem->Index);
+		}
+	}
+
+	FocusOnDataCallback(DataStrongPtr.Get(), Indices);
+}
+
+bool SPCGEditorGraphAttributeListView::CanFocusOnSelection() const
+{
+	return ListView->GetNumItemsSelected() > 0;
+}
+
 #undef LOCTEXT_NAMESPACE
diff --git a/Engine/Plugins/PCG/Source/PCGEditor/Private/SPCGEditorGraphAttributeListView.h b/Engine/Plugins/PCG/Source/PCGEditor/Private/SPCGEditorGraphAttributeListView.h
--- a/Engine/Plugins/PCG/Source/PCGEditor/Private/SPCGEditorGraphAttributeListView.h
+++ b/Engine/Plugins/PCG/Source/PCGEditor/Private/SPCGEditorGraphAttributeListView.h
@@ -45,7 +45,6 @@ namespace PCGEditorGraphAttributeListView
 struct FPCGListViewItem
 {
 	int32 Index = INDEX_NONE;
-	TFunction<void(const UPCGData*, int)> DoubleClickCallback = nullptr;
 };
 
 struct FPCGColumnData
@@ -54,7 +53,7 @@ struct FPCGColumnData
 	TSharedPtr<const IPCGAttributeAccessorKeys> DataKeys;
 };
 
-typedef TSharedPtr<FPCGListViewItem> PCGListviewItemPtr;
+typedef TSharedPtr<FPCGListViewItem> PCGListViewItemPtr;
 
 template <typename T, typename = void>
 struct FTextAsNumberIsValid : std::false_type {};
@@ -68,7 +67,7 @@ class FPCGListViewUpdater : public TSharedFromThis<FPCGListViewUpdater>
 {
 public:
 	FPCGListViewUpdater(
-		const TArray<PCGListviewItemPtr>& InListViewItems,
+		const TArray<PCGListViewItemPtr>& InListViewItems,
 		const TMap<FName, FPCGColumnData>& InColumnData,
 		const EColumnSortMode::Type InSortMode,
 		const FName InSortingColumn,
@@ -83,7 +82,7 @@ public:
 	bool IsCompleted() const;
 	void Launch();
 
-	TArray<PCGListviewItemPtr> ListViewItems;
+	TArray<PCGListViewItemPtr> ListViewItems;
 
 private:
 	void AsyncSort();
@@ -99,12 +98,12 @@ private:
 	UE::Tasks::FTask UpdateTask;
 };
 
-class SPCGListViewItemRow : public SMultiColumnTableRow<PCGListviewItemPtr>
+class SPCGListViewItemRow : public SMultiColumnTableRow<PCGListViewItemPtr>
 {
 public:
 	SLATE_BEGIN_ARGS(SPCGListViewItemRow) {}
 	SLATE_ARGUMENT(TSharedPtr<SPCGEditorGraphAttributeListView>, AttributeListView)
-	SLATE_ARGUMENT(PCGListviewItemPtr, ListViewItem)
+	SLATE_ARGUMENT(PCGListViewItemPtr, ListViewItem)
 	SLATE_END_ARGS()
 
 	void Construct(const FArguments& InArgs, const TSharedRef<STableViewBase>& InOwnerTableView);
@@ -113,7 +112,7 @@ public:
 
 private:
 	TWeakPtr<SPCGEditorGraphAttributeListView> AttributeListView;
-	PCGListviewItemPtr InternalItem;
+	PCGListViewItemPtr InternalItem;
 };
 
 class FPCGPointFilterExpressionContext : public ITextFilterExpressionContext
@@ -204,8 +203,9 @@ private:
 	void SaveData(bool bUsePinIndex, bool bUseDataIndex);
 	bool CanSaveData(bool bUsePinIndex, bool bUseDataIndex) const;
 
-	TSharedRef<ITableRow> OnGenerateRow(PCGListviewItemPtr Item, const TSharedRef<STableViewBase>& OwnerTable);
-	void OnItemDoubleClicked(PCGListviewItemPtr Item) const;
+	TSharedRef<ITableRow> OnGenerateRow(PCGListViewItemPtr Item, const TSharedRef<STableViewBase>& OwnerTable);
+	void OnItemDoubleClicked(PCGListViewItemPtr Item) const;
+	TSharedPtr<SWidget> OnItemsContextMenu();
 
 	void OnColumnSortModeChanged(const EColumnSortPriority::Type SortPriority, const FName& ColumnId, const EColumnSortMode::Type InSortMode);
 	EColumnSortMode::Type GetColumnSortMode(const FName InColumnId) const;
@@ -228,6 +228,12 @@ private:
 	FReply OnLockClick();
 	FReply OnNodeNameClicked();
 
+	FReply OnFocusOnDataClicked() const;
+	bool IsFocusOnDataEnabled() const;
+
+	void FocusOnSelection() const;
+	bool CanFocusOnSelection() const;
+
 	/** Pointer back to the PCG editor that owns us */
 	TWeakPtr<FPCGEditor> PCGEditorPtr;
 
@@ -243,9 +249,9 @@ private:
 
 	TSharedPtr<SSearchBox> SearchBoxWidget;
 	TSharedPtr<SHeaderRow> ListViewHeader;
-	TSharedPtr<SListView<PCGListviewItemPtr>> ListView;
-	TArray<PCGListviewItemPtr> ListViewItems;
-	TArray<PCGListviewItemPtr> FilteredListViewItems;
+	TSharedPtr<SListView<PCGListViewItemPtr>> ListView;
+	TArray<PCGListViewItemPtr> ListViewItems;
+	TArray<PCGListViewItemPtr> FilteredListViewItems;
 
 	TSharedPtr<SComboBox<TSharedPtr<FPinComboBoxItem>>> PinComboBox;
 	TArray<TSharedPtr<FPinComboBoxItem>> PinComboBoxItems;
@@ -264,6 +270,7 @@ private:
 
 	FName SortingColumn = NAME_None;
 	EColumnSortMode::Type SortMode = EColumnSortMode::Type::Ascending;
+	TFunction<void(const UPCGData*, TArrayView<const int>)> FocusOnDataCallback;
 
 	bool bNeedsRefresh : 1 = false;
 
