[PCG] Optimize tree refresh and costly button enable lambda

#rnx

diff --git a/Engine/Plugins/PCG/Source/PCG/Private/PCGSubsystem.cpp b/Engine/Plugins/PCG/Source/PCG/Private/PCGSubsystem.cpp
--- a/Engine/Plugins/PCG/Source/PCG/Private/PCGSubsystem.cpp
+++ b/Engine/Plugins/PCG/Source/PCG/Private/PCGSubsystem.cpp
@@ -1739,7 +1739,7 @@ TArray<FPCGStack> UPCGSubsystem::GetExecutedStacks(const UPCGComponent* InCompon
 
 	for (const FPCGStack& Stack : ExecutedStacks)
 	{
-		if (Stack.GetRootComponent() == InComponent && Stack.GetGraphForCurrentFrame() == InSubgraph)
+		if ((!InComponent || Stack.GetRootComponent() == InComponent) && Stack.GetGraphForCurrentFrame() == InSubgraph)
 		{
 			MatchingStacks.Add(Stack);
 		}
diff --git a/Engine/Plugins/PCG/Source/PCGEditor/Private/SPCGEditorGraphDebugObjectTree.cpp b/Engine/Plugins/PCG/Source/PCGEditor/Private/SPCGEditorGraphDebugObjectTree.cpp
--- a/Engine/Plugins/PCG/Source/PCGEditor/Private/SPCGEditorGraphDebugObjectTree.cpp
+++ b/Engine/Plugins/PCG/Source/PCGEditor/Private/SPCGEditorGraphDebugObjectTree.cpp
@@ -230,6 +230,7 @@ SPCGEditorGraphDebugObjectTree::~SPCGEditorGraphDebugObjectTree()
 	FCoreUObjectDelegates::OnPreObjectPropertyChanged.RemoveAll(this);
 	FCoreUObjectDelegates::OnObjectPropertyChanged.RemoveAll(this);
 	FCoreUObjectDelegates::OnObjectConstructed.RemoveAll(this);
+	USelection::SelectionChangedEvent.RemoveAll(this);
 }
 
 void SPCGEditorGraphDebugObjectTree::Construct(const FArguments& InArgs, TSharedPtr<FPCGEditor> InPCGEditor)
@@ -243,6 +244,7 @@ void SPCGEditorGraphDebugObjectTree::Construct(const FArguments& InArgs, TShared
 	FCoreUObjectDelegates::OnPreObjectPropertyChanged.AddSP(this, &SPCGEditorGraphDebugObjectTree::OnPreObjectPropertyChanged);
 	FCoreUObjectDelegates::OnObjectPropertyChanged.AddSP(this, &SPCGEditorGraphDebugObjectTree::OnObjectPropertyChanged);
 	FCoreUObjectDelegates::OnObjectConstructed.AddSP(this, &SPCGEditorGraphDebugObjectTree::OnObjectConstructed);
+	USelection::SelectionChangedEvent.AddSP(this, &SPCGEditorGraphDebugObjectTree::OnEditorSelectionChanged);
 
 	const TSharedRef<SScrollBar> HorizontalScrollBar = SNew(SScrollBar)
 		.Orientation(Orient_Horizontal)
@@ -349,6 +351,11 @@ void SPCGEditorGraphDebugObjectTree::Tick(const FGeometry& AllottedGeometry, con
 			RefreshTree();
 		}
 	}
+
+	if (!IsSetDebugObjectFromSelectionEnabled.IsSet())
+	{
+		UpdateIsSetDebugObjectFromSelectionEnabled();
+	}
 }
 
 void SPCGEditorGraphDebugObjectTree::SetDebugObjectSelection(const FPCGStack& FullStack)
@@ -469,18 +476,20 @@ void SPCGEditorGraphDebugObjectTree::SetDebugObjectFromSelection_OnClicked()
 	}
 }
 
-bool SPCGEditorGraphDebugObjectTree::IsSetDebugObjectFromSelectionButtonEnabled() const
+void SPCGEditorGraphDebugObjectTree::UpdateIsSetDebugObjectFromSelectionEnabled()
 {
+	IsSetDebugObjectFromSelectionEnabled = false;
+
 	const UPCGGraph* PCGGraph = GetPCGGraph();
 	if (!PCGGraph)
 	{
-		return false;
+		return;
 	}
 
 	USelection* SelectedActors = GEditor->GetSelectedActors();
 	if (!IsValid(SelectedActors))
 	{
-		return false;
+		return;
 	}
 
 	UPCGSubsystem* Subsystem = PCGEditor.Pin() ? PCGEditor.Pin()->GetSubsystem() : nullptr;
@@ -509,7 +518,8 @@ bool SPCGEditorGraphDebugObjectTree::IsSetDebugObjectFromSelectionButtonEnabled(
 			{
 				if (Algo::AnyOf(StackContext.GetStacks(), [&PCGGraph](const FPCGStack& InStack) { return InStack.GetGraphForCurrentFrame() == PCGGraph; }))
 				{
-					return true;
+					IsSetDebugObjectFromSelectionEnabled = true;
+					return;
 				}
 			}
 
@@ -518,13 +528,12 @@ bool SPCGEditorGraphDebugObjectTree::IsSetDebugObjectFromSelectionButtonEnabled(
 			{
 				if (Algo::AnyOf(Subsystem->GetExecutedStacks(PCGComponent, PCGGraph), [&PCGGraph](const FPCGStack& InStack) { return InStack.GetGraphForCurrentFrame() == PCGGraph; }))
 				{
-					return true;
+					IsSetDebugObjectFromSelectionEnabled = true;
+					return;
 				}
 			}
 		}
 	}
-
-	return false;
 }
 
 void SPCGEditorGraphDebugObjectTree::AddStacksToTree(const TArray<FPCGStack>& Stacks,
@@ -750,38 +759,51 @@ void SPCGEditorGraphDebugObjectTree::RefreshTree()
 	TArray<UObject*> PCGComponents;
 	GetObjectsOfClass(UPCGComponent::StaticClass(), PCGComponents, /*bIncludeDerivedClasses=*/ true);
 
-	TMap<AActor*, TSharedPtr<FPCGEditorGraphDebugObjectItem_Actor>> ActorItems;
-	TMap<const FPCGStack, FPCGEditorGraphDebugObjectItemPtr> StackToItem;
-
-	for (UObject* PCGComponentObject : PCGComponents)
+	if (!PCGComponents.IsEmpty())
 	{
-		if (!IsValid(PCGComponentObject))
+		TMap<AActor*, TSharedPtr<FPCGEditorGraphDebugObjectItem_Actor>> ActorItems;
+		TMap<const FPCGStack, FPCGEditorGraphDebugObjectItemPtr> StackToItem;
+
+		TArray<FPCGStack> GraphStacks = Subsystem->GetExecutedStacks(nullptr, PCGGraph);
+		TMap<const UPCGComponent*, TArray<FPCGStack>> ComponentsGraphStacks;
+		for (FPCGStack& GraphStack : GraphStacks)
 		{
-			continue;
+			ComponentsGraphStacks.FindOrAdd(GraphStack.GetRootComponent()).Add(MoveTemp(GraphStack));
 		}
 
-		UPCGComponent* PCGComponent = Cast<UPCGComponent>(PCGComponentObject);
-		if (!PCGComponent || !PCGComponent->IsRegistered())
+		for (UObject* PCGComponentObject : PCGComponents)
 		{
-			continue;
+			if (!IsValid(PCGComponentObject))
+			{
+				continue;
+			}
+
+			UPCGComponent* PCGComponent = Cast<UPCGComponent>(PCGComponentObject);
+			if (!PCGComponent || !PCGComponent->IsRegistered())
+			{
+				continue;
+			}
+
+			// Process static stacks that can be read from the compiled graph.
+			FPCGStackContext StackContext;
+			if (!PCGComponent->GetStackContext(StackContext))
+			{
+				continue;
+			}
+			AddStacksToTree(StackContext.GetStacks(), ActorItems, StackToItem);
+
+			// Process stacks encountered during execution so far, which will include dynamic subgraphs & loop subgraphs.
+			// There will be overlaps with the static stacks but only unique entries will be added to the tree.
+			if (TArray<FPCGStack>* ComponentGraphStacks = ComponentsGraphStacks.Find(PCGComponent))
+			{
+				AddStacksToTree(*ComponentGraphStacks, ActorItems, StackToItem);
+			}
 		}
 
-		// Process static stacks that can be read from the compiled graph.
-		FPCGStackContext StackContext;
-		if (!PCGComponent->GetStackContext(StackContext))
+		for (TPair<AActor*, TSharedPtr<FPCGEditorGraphDebugObjectItem_Actor>>& ActorItem : ActorItems)
 		{
-			continue;
+			RootItems.Add(MoveTemp(ActorItem.Value));
 		}
-		AddStacksToTree(StackContext.GetStacks(), ActorItems, StackToItem);
-
-		// Process stacks encountered during execution so far, which will include dynamic subgraphs & loop subgraphs.
-		// There will be overlaps with the static stacks but only unique entries will be added to the tree.
-		AddStacksToTree(Subsystem->GetExecutedStacks(PCGComponent, PCGGraph), ActorItems, StackToItem);
-	}
-
-	for (TPair<AActor*, TSharedPtr<FPCGEditorGraphDebugObjectItem_Actor>>& ActorItem : ActorItems)
-	{
-		RootItems.Add(MoveTemp(ActorItem.Value));
 	}
 
 	SortTreeItems();
@@ -939,6 +961,11 @@ void SPCGEditorGraphDebugObjectTree::OnObjectPropertyChanged(UObject* InObject,
 	}
 }
 
+void SPCGEditorGraphDebugObjectTree::OnEditorSelectionChanged(UObject* InObject)
+{
+	IsSetDebugObjectFromSelectionEnabled.Reset();
+}
+
 void SPCGEditorGraphDebugObjectTree::OnObjectConstructed(UObject* InObject)
 {
 	if (const UPCGComponent* PCGComponent = Cast<UPCGComponent>(InObject))
diff --git a/Engine/Plugins/PCG/Source/PCGEditor/Private/SPCGEditorGraphDebugObjectTree.h b/Engine/Plugins/PCG/Source/PCGEditor/Private/SPCGEditorGraphDebugObjectTree.h
--- a/Engine/Plugins/PCG/Source/PCGEditor/Private/SPCGEditorGraphDebugObjectTree.h
+++ b/Engine/Plugins/PCG/Source/PCGEditor/Private/SPCGEditorGraphDebugObjectTree.h
@@ -9,6 +9,7 @@
 
 #include "ToolMenus.h"
 #include "GameFramework/Actor.h"
+#include "Misc/Optional.h"
 #include "Widgets/SCompoundWidget.h"
 #include "Widgets/Views/STreeView.h"
 
@@ -224,7 +225,8 @@ private:
 	bool IsSelectDebugObjectButtonEnabled() const;
 
 	void SetDebugObjectFromSelection_OnClicked();
-	bool IsSetDebugObjectFromSelectionButtonEnabled() const;
+	bool IsSetDebugObjectFromSelectionButtonEnabled() const { return IsSetDebugObjectFromSelectionEnabled.IsSet() && IsSetDebugObjectFromSelectionEnabled.GetValue(); }
+	void UpdateIsSetDebugObjectFromSelectionEnabled();
 
 	void RefreshTree();
 	void SortTreeItems(bool bIsAscending = true, bool bIsRecursive = true);
@@ -238,6 +240,7 @@ private:
 	void OnPreObjectPropertyChanged(UObject* InObject, const FEditPropertyChain& InPropertyChain);
 	void OnObjectPropertyChanged(UObject* InObject, FPropertyChangedEvent& InPropertyChangedEvent);
 	void OnObjectConstructed(UObject* InObject);
+	void OnEditorSelectionChanged(UObject* InObject);
 
 	UPCGGraph* GetPCGGraph() const;
 
@@ -267,6 +270,9 @@ private:
 	/** Set true to avoid broadcasting debug object change notifications when setting the object from code. */
 	bool bDisableDebugObjectChangeNotification = false;
 
+	/** Latest value for IsSetDebugObjectFromSelectionButtonEnabled */
+	TOptional<bool> IsSetDebugObjectFromSelectionEnabled;
+
 	/** Used to retain item expansion state across tree refreshes. */
 	TSet<FPCGStack> ExpandedStacks;
 
