[PCG] More GetActorPCGData calls that need to happen on main thread for now.

- This should change later on so we can make sure the GetActorPCGData is primed before executing nodes (through a PreGraph task) so that all the nodes that need it will be able to run outside of the main thread

#rnx

diff --git a/Engine/Plugins/PCG/Source/PCG/Private/Elements/PCGCreatePoints.cpp b/Engine/Plugins/PCG/Source/PCG/Private/Elements/PCGCreatePoints.cpp
--- a/Engine/Plugins/PCG/Source/PCG/Private/Elements/PCGCreatePoints.cpp
+++ b/Engine/Plugins/PCG/Source/PCG/Private/Elements/PCGCreatePoints.cpp
@@ -139,6 +139,18 @@ bool FPCGCreatePointsElement::ExecuteInternal(FPCGContext* Context) const
 	return true;
 }
 
+bool FPCGCreatePointsElement::CanExecuteOnlyOnMainThread(FPCGContext* Context) const 
+{
+	if (Context && Context->CurrentPhase == EPCGExecutionPhase::Execute)
+	{
+		const UPCGCreatePointsSettings* Settings = Context->GetInputSettings<UPCGCreatePointsSettings>();
+		check(Settings);
+		return Settings->bCullPointsOutsideVolume;
+	}
+
+	return false;
+}
+
 bool FPCGCreatePointsElement::IsCacheable(const UPCGSettings* InSettings) const
 {
 	const UPCGCreatePointsSettings* Settings = Cast<const UPCGCreatePointsSettings>(InSettings);
diff --git a/Engine/Plugins/PCG/Source/PCG/Private/Elements/PCGElevationIsolines.cpp b/Engine/Plugins/PCG/Source/PCG/Private/Elements/PCGElevationIsolines.cpp
--- a/Engine/Plugins/PCG/Source/PCG/Private/Elements/PCGElevationIsolines.cpp
+++ b/Engine/Plugins/PCG/Source/PCG/Private/Elements/PCGElevationIsolines.cpp
@@ -49,6 +49,17 @@ TArray<FPCGPinProperties> UPCGElevationIsolinesSettings::OutputPinProperties() c
 	return Properties;
 }
 
+bool FPCGElevationIsolinesElement::CanExecuteOnlyOnMainThread(FPCGContext* Context) const
+{
+	// Dependency on GetActorPCGData requires to run on Main Thread
+	if (Context && Context->CurrentPhase == EPCGExecutionPhase::Execute)
+	{
+		return true;
+	}
+
+	return false;
+}
+
 bool FPCGElevationIsolinesElement::ExecuteInternal(FPCGContext* InContext) const
 {
 	TRACE_CPUPROFILER_EVENT_SCOPE(FPCGElevationIsolinesElement::Execute);
diff --git a/Engine/Plugins/PCG/Source/PCG/Public/Elements/PCGCreatePoints.h b/Engine/Plugins/PCG/Source/PCG/Public/Elements/PCGCreatePoints.h
--- a/Engine/Plugins/PCG/Source/PCG/Public/Elements/PCGCreatePoints.h
+++ b/Engine/Plugins/PCG/Source/PCG/Public/Elements/PCGCreatePoints.h
@@ -68,8 +68,11 @@ PRAGMA_ENABLE_DEPRECATION_WARNINGS
 
 class FPCGCreatePointsElement : public IPCGElement
 {
-protected:
-	virtual bool ExecuteInternal(FPCGContext* Context) const override;
+public:
+	virtual bool CanExecuteOnlyOnMainThread(FPCGContext* Context) const override;
 	virtual bool IsCacheable(const UPCGSettings* InSettings) const override;
 	virtual bool ShouldComputeFullOutputDataCrc(FPCGContext* Context) const override { return true; }
+
+protected:
+	virtual bool ExecuteInternal(FPCGContext* Context) const override;
 };
\ No newline at end of file
diff --git a/Engine/Plugins/PCG/Source/PCG/Public/Elements/PCGCreatePointsGrid.h b/Engine/Plugins/PCG/Source/PCG/Public/Elements/PCGCreatePointsGrid.h
--- a/Engine/Plugins/PCG/Source/PCG/Public/Elements/PCGCreatePointsGrid.h
+++ b/Engine/Plugins/PCG/Source/PCG/Public/Elements/PCGCreatePointsGrid.h
@@ -89,7 +89,12 @@ PRAGMA_ENABLE_DEPRECATION_WARNINGS
 
 class FPCGCreatePointsGridElement : public IPCGElement
 {
+public:
+	// Returns true when Context is nullptr which happens in FPCGGraphExecutor::OnTaskInputsReady
+	// We can't know for sure at that point if we need to call UPCGComponent::GetActorPCGData() in our GetDependenciesCrc so we need to make sure it happens on MainThread
+	virtual bool CanExecuteOnlyOnMainThread(FPCGContext* Context) const override { return Context == nullptr; }
+	virtual void GetDependenciesCrc(const FPCGDataCollection& InInput, const UPCGSettings* InSettings, UPCGComponent* InComponent, FPCGCrc& OutCrc) const override;
+
 protected:
 	virtual bool ExecuteInternal(FPCGContext* Context) const override;
-	virtual void GetDependenciesCrc(const FPCGDataCollection& InInput, const UPCGSettings* InSettings, UPCGComponent* InComponent, FPCGCrc& OutCrc) const override;
 };
\ No newline at end of file
diff --git a/Engine/Plugins/PCG/Source/PCG/Public/Elements/PCGCreatePointsSphere.h b/Engine/Plugins/PCG/Source/PCG/Public/Elements/PCGCreatePointsSphere.h
--- a/Engine/Plugins/PCG/Source/PCG/Public/Elements/PCGCreatePointsSphere.h
+++ b/Engine/Plugins/PCG/Source/PCG/Public/Elements/PCGCreatePointsSphere.h
@@ -141,7 +141,12 @@ public:
 
 class FPCGCreatePointsSphereElement : public IPCGElement
 {
+public:
+	// Returns true when Context is nullptr which happens in FPCGGraphExecutor::OnTaskInputsReady
+	// We can't know for sure at that point if we need to call UPCGComponent::GetActorPCGData() in our GetDependenciesCrc so we need to make sure it happens on MainThread
+	virtual bool CanExecuteOnlyOnMainThread(FPCGContext* Context) const override { return Context == nullptr; }
+	virtual void GetDependenciesCrc(const FPCGDataCollection& InInput, const UPCGSettings* InSettings, UPCGComponent* InComponent, FPCGCrc& OutCrc) const override;
+
 protected:
 	virtual bool ExecuteInternal(FPCGContext* InContext) const override;
-	virtual void GetDependenciesCrc(const FPCGDataCollection& InInput, const UPCGSettings* InSettings, UPCGComponent* InComponent, FPCGCrc& OutCrc) const override;
 };
diff --git a/Engine/Plugins/PCG/Source/PCG/Public/Elements/PCGCullPointsOutsideActorBounds.h b/Engine/Plugins/PCG/Source/PCG/Public/Elements/PCGCullPointsOutsideActorBounds.h
--- a/Engine/Plugins/PCG/Source/PCG/Public/Elements/PCGCullPointsOutsideActorBounds.h
+++ b/Engine/Plugins/PCG/Source/PCG/Public/Elements/PCGCullPointsOutsideActorBounds.h
@@ -38,6 +38,9 @@ public:
 class FPCGCullPointsOutsideActorBoundsElement : public IPCGElement
 {
 public:
+	// Returns true when Context is nullptr which happens in FPCGGraphExecutor::OnTaskInputsReady
+	// We can't know for sure at that point if we need to call UPCGComponent::GetActorPCGData() in our GetDependenciesCrc so we need to make sure it happens on MainThread
+	virtual bool CanExecuteOnlyOnMainThread(FPCGContext* Context) const override { return Context == nullptr; }
 	virtual void GetDependenciesCrc(const FPCGDataCollection& InInput, const UPCGSettings* InSettings, UPCGComponent* InComponent, FPCGCrc& OutCrc) const override;
 
 protected:
diff --git a/Engine/Plugins/PCG/Source/PCG/Public/Elements/PCGElevationIsolines.h b/Engine/Plugins/PCG/Source/PCG/Public/Elements/PCGElevationIsolines.h
--- a/Engine/Plugins/PCG/Source/PCG/Public/Elements/PCGElevationIsolines.h
+++ b/Engine/Plugins/PCG/Source/PCG/Public/Elements/PCGElevationIsolines.h
@@ -65,6 +65,9 @@ public:
 
 class FPCGElevationIsolinesElement : public IPCGElement
 {
+public:
+	virtual bool CanExecuteOnlyOnMainThread(FPCGContext* Context) const override;
+
 protected:
 	virtual bool ExecuteInternal(FPCGContext* InContext) const override;
 };
