[PCG] Fix palette update issue with graph instances

#rnx

diff --git a/Engine/Plugins/PCG/Source/PCGEditor/Private/SPCGEditorGraphNodePalette.cpp b/Engine/Plugins/PCG/Source/PCGEditor/Private/SPCGEditorGraphNodePalette.cpp
--- a/Engine/Plugins/PCG/Source/PCGEditor/Private/SPCGEditorGraphNodePalette.cpp
+++ b/Engine/Plugins/PCG/Source/PCGEditor/Private/SPCGEditorGraphNodePalette.cpp
@@ -130,8 +130,8 @@ void SPCGEditorGraphNodePalette::CollectAllActions(FGraphActionListBuilderBase&
 
 void SPCGEditorGraphNodePalette::OnAssetChanged(const FAssetData& InAssetData)
 {
-	if (InAssetData.AssetClassPath == UPCGGraph::StaticClass()->GetClassPathName() ||
-		InAssetData.AssetClassPath == UPCGSettings::StaticClass()->GetClassPathName() ||
+	if (InAssetData.IsInstanceOf<UPCGGraphInterface>() ||
+		InAssetData.IsInstanceOf<UPCGSettings>() ||
 		PCGEditorUtils::IsAssetPCGBlueprint(InAssetData))
 	{
 		RefreshActionsList(true);
