[PCG] Add parse string node

* Allow to convert any string into a predefined type.

#rnx

diff --git a/Engine/Plugins/PCG/Source/PCG/Private/Elements/PCGParseString.cpp b/Engine/Plugins/PCG/Source/PCG/Private/Elements/PCGParseString.cpp
new file mode 100644
--- /dev/null
+++ b/Engine/Plugins/PCG/Source/PCG/Private/Elements/PCGParseString.cpp
@@ -0,0 +1,94 @@
+// Copyright Epic Games, Inc. All Rights Reserved.
+
+#include "Elements/PCGParseString.h"
+
+#include "PCGContext.h"
+
+#define LOCTEXT_NAMESPACE "PCGParseStringElement"
+
+#if WITH_EDITOR
+FName UPCGParseStringSettings::GetDefaultNodeName() const
+{
+	return FName(TEXT("ParseString"));
+}
+
+FText UPCGParseStringSettings::GetDefaultNodeTitle() const
+{
+	return LOCTEXT("NodeTitle", "Parse String");
+}
+
+FText UPCGParseStringSettings::GetNodeTooltipText() const
+{
+	return LOCTEXT("NodeTooltip", "Parse string passed as attribute into a compatible PCG type.\n"
+		"Failed to parse will result to the value be set to Identity(for quat and transform) or 0 (for the rest).\n"
+		"Int/Int64: Support base 10 or base 16 number (ie. 0xff). If int overflows, result is platform dependant. \n"
+		"Float/Double: Support normal representation or exposant representation (ie. 13.5e-2f)\n"
+		"Bool: [True/Yes/On] [False/No/Off], or 0 -> false and the rest -> true"
+		"Vectors/Quat: X=<value> Y=<value> Z=<value> W=<value> OR <value>,<value>,<value>,<value>\n"
+		"Rotator: P=<value> Y=<value> R=<value> or  Pitch=<value> Yaw=<value> Roll=<value> or <value,<value>,<value>\n"
+		"Transform: <value>,<value>,<value>|<value>,<value>,<value>|<value>,<value>,<value> (Translation|Rotation|Scale)");
+}
+#endif // WITH_EDITOR
+
+FString UPCGParseStringSettings::GetAdditionalTitleInformation() const
+{
+#if WITH_EDITOR
+	const bool bIsTypeOverridden = IsPropertyOverriddenByPin(GET_MEMBER_NAME_CHECKED(UPCGParseStringSettings, TargetType));
+	const bool bSelectorOverridden = IsPropertyOverriddenByPin(GET_MEMBER_NAME_CHECKED(UPCGParseStringSettings, InputSource)) || IsPropertyOverriddenByPin(GET_MEMBER_NAME_CHECKED(UPCGParseStringSettings, OutputTarget));
+
+	if (bIsTypeOverridden)
+	{
+		return FString();
+	}
+	else
+	{
+		if (bSelectorOverridden)
+		{
+			return FText::Format(LOCTEXT("CastType", "To {0}"), PCG::Private::GetTypeNameText((uint16)TargetType)).ToString();
+		}
+	}
+#endif
+
+	return FText::Format(LOCTEXT("FullInfo", "{0} -> {1} as {2}"), InputSource.GetDisplayText(), OutputTarget.GetDisplayText(), PCG::Private::GetTypeNameText((uint16)TargetType)).ToString();
+}
+
+FPCGElementPtr UPCGParseStringSettings::CreateElement() const
+{
+	return MakeShared<FPCGParseStringElement>();
+}
+
+FPCGAttributePropertyInputSelector UPCGParseStringSettings::GetInputSource(uint32 Index) const
+{
+	ensure(Index == 0);
+	return InputSource;
+}
+
+bool UPCGParseStringSettings::IsSupportedInputType(uint16 TypeId, uint32 InputIndex, bool& bHasSpecialRequirement) const
+{
+	ensure(InputIndex == 0);
+	return PCG::Private::IsOfTypes<FName, FString>(TypeId);
+}
+
+TArray<uint16> UPCGParseStringSettings::GetAllOutputTypes() const
+{
+	return { static_cast<uint16>(TargetType) };
+}
+
+bool FPCGParseStringElement::DoOperation(PCGMetadataOps::FOperationData& OperationData) const
+{
+	TRACE_CPUPROFILER_EVENT_SCOPE(UPCGParseStringElement::Execute);
+
+	const UPCGParseStringSettings* Settings = CastChecked<UPCGParseStringSettings>(OperationData.Settings);
+
+	return PCGMetadataAttribute::CallbackWithRightType(static_cast<uint16>(Settings->TargetType), [this, &OperationData]<typename T>(T Dummy) -> bool
+	{
+		return DoUnaryOp<FString>(OperationData, [](const FString& InValue) -> T
+			{
+				bool bSuccess;
+				return PCG::Private::MetadataTraits<T>::ParseFromString(InValue, bSuccess);
+			});
+	});
+}
+
+
+#undef LOCTEXT_NAMESPACE
diff --git a/Engine/Plugins/PCG/Source/PCG/Private/Tests/Elements/PCGParseStringTests.cpp b/Engine/Plugins/PCG/Source/PCG/Private/Tests/Elements/PCGParseStringTests.cpp
new file mode 100644
--- /dev/null
+++ b/Engine/Plugins/PCG/Source/PCG/Private/Tests/Elements/PCGParseStringTests.cpp
@@ -0,0 +1,132 @@
+// Copyright Epic Games, Inc. All Rights Reserved.
+
+#include "Tests/PCGTestsCommon.h"
+
+#include "PCGContext.h"
+#include "PCGParamData.h"
+#include "Elements/PCGParseString.h"
+
+class FPCGParseStringTestsBase : public FPCGTestBaseClass
+{
+public:
+	using FPCGTestBaseClass::FPCGTestBaseClass;
+
+protected:
+	template <typename T>
+	bool RunTestInternal(TArrayView<const FString> InData, TArrayView<const T> ExpectedOutData)
+	{
+		PCGTestsCommon::FTestData TestData;
+		UPCGParseStringSettings* Settings = PCGTestsCommon::GenerateSettings<UPCGParseStringSettings>(TestData);
+		check(Settings);
+
+		check(InData.Num() == ExpectedOutData.Num() && InData.Num() > 1);
+
+		const FName InputAttributeName = TEXT("InAttr");
+		const FName OutputAttributeName = TEXT("OutAttr");
+		Settings->InputSource.SetAttributeName(InputAttributeName);
+		Settings->OutputTarget.SetAttributeName(OutputAttributeName);
+		Settings->TargetType = static_cast<EPCGMetadataTypes>(PCG::Private::MetadataTypes<T>::Id);
+
+		UPCGParamData* InputData = NewObject<UPCGParamData>();
+		FPCGMetadataAttribute<FString>* Attribute = InputData->Metadata->CreateAttribute<FString>(InputAttributeName, InData[0], true, false);
+		for (int32 i = 1; i < InData.Num(); ++i)
+		{
+			Attribute->SetValue(InputData->Metadata->AddEntry(), InData[i]);
+		}
+
+		FPCGTaggedData& InputTaggedData = TestData.InputData.TaggedData.Emplace_GetRef();
+		InputTaggedData.Data = InputData;
+		InputTaggedData.Pin = PCGPinConstants::DefaultInputLabel;
+
+		FPCGElementPtr TestElement = TestData.Settings->GetElement();
+		TUniquePtr<FPCGContext> Context = TestData.InitializeTestContext();
+
+		while (!TestElement->Execute(Context.Get())) {}
+
+		UTEST_EQUAL("There is one output", Context->OutputData.TaggedData.Num(), 1);
+
+		const UPCGParamData* OutputData = Cast<const UPCGParamData>(Context->OutputData.TaggedData[0].Data);
+
+		UTEST_NOT_NULL("Output is an of the same type as input", OutputData);
+
+		const FPCGMetadataAttribute<T>* OutputAttribute = OutputData->Metadata->template GetConstTypedAttribute<T>(OutputAttributeName);
+
+		UTEST_NOT_NULL("Output metadata has the expected attribute", OutputAttribute);
+
+		UTEST_EQUAL(*FString::Printf(TEXT("Default value %s, was parsed correctly"), *InData[0]), OutputAttribute->GetValue(PCGDefaultValueKey), ExpectedOutData[0]);
+		for (int32 i = 1; i < ExpectedOutData.Num(); ++i)
+		{
+			PCGMetadataEntryKey EntryKey(i - 1);
+			UTEST_EQUAL(*FString::Printf(TEXT("Entry %d (%s) was parsed expectedly"), i - 1, *InData[i]), OutputAttribute->GetValueFromItemKey(EntryKey), ExpectedOutData[i]);
+		}
+
+		return true;
+	}
+};
+
+IMPLEMENT_CUSTOM_SIMPLE_AUTOMATION_TEST(FPCGParseStringTests_Bool, FPCGParseStringTestsBase, "Plugins.PCG.ParseString.Bool", PCGTestsCommon::TestFlags)
+IMPLEMENT_CUSTOM_SIMPLE_AUTOMATION_TEST(FPCGParseStringTests_Int32, FPCGParseStringTestsBase, "Plugins.PCG.ParseString.Int32", PCGTestsCommon::TestFlags)
+IMPLEMENT_CUSTOM_SIMPLE_AUTOMATION_TEST(FPCGParseStringTests_Int64, FPCGParseStringTestsBase, "Plugins.PCG.ParseString.Int64", PCGTestsCommon::TestFlags)
+IMPLEMENT_CUSTOM_SIMPLE_AUTOMATION_TEST(FPCGParseStringTests_Float, FPCGParseStringTestsBase, "Plugins.PCG.ParseString.Float", PCGTestsCommon::TestFlags)
+IMPLEMENT_CUSTOM_SIMPLE_AUTOMATION_TEST(FPCGParseStringTests_Double, FPCGParseStringTestsBase, "Plugins.PCG.ParseString.Double", PCGTestsCommon::TestFlags)
+IMPLEMENT_CUSTOM_SIMPLE_AUTOMATION_TEST(FPCGParseStringTests_Vector2D, FPCGParseStringTestsBase, "Plugins.PCG.ParseString.Vector2D", PCGTestsCommon::TestFlags)
+IMPLEMENT_CUSTOM_SIMPLE_AUTOMATION_TEST(FPCGParseStringTests_Vector, FPCGParseStringTestsBase, "Plugins.PCG.ParseString.Vector", PCGTestsCommon::TestFlags)
+IMPLEMENT_CUSTOM_SIMPLE_AUTOMATION_TEST(FPCGParseStringTests_Vector4, FPCGParseStringTestsBase, "Plugins.PCG.ParseString.Vector4", PCGTestsCommon::TestFlags)
+IMPLEMENT_CUSTOM_SIMPLE_AUTOMATION_TEST(FPCGParseStringTests_Rotator, FPCGParseStringTestsBase, "Plugins.PCG.ParseString.Rotator", PCGTestsCommon::TestFlags)
+IMPLEMENT_CUSTOM_SIMPLE_AUTOMATION_TEST(FPCGParseStringTests_Quat, FPCGParseStringTestsBase, "Plugins.PCG.ParseString.Quat", PCGTestsCommon::TestFlags)
+IMPLEMENT_CUSTOM_SIMPLE_AUTOMATION_TEST(FPCGParseStringTests_Transform, FPCGParseStringTestsBase, "Plugins.PCG.ParseString.Transform", PCGTestsCommon::TestFlags)
+
+bool FPCGParseStringTests_Bool::RunTest(const FString& Parameters)
+{
+	return RunTestInternal<bool>({ "On", "Off", "True", "False", "Yes", "No", "0", "1", "54", "Blabla"}, {true, false, true, false, true, false, false, true, true, false });
+}
+
+bool FPCGParseStringTests_Int32::RunTest(const FString& Parameters)
+{
+	return RunTestInternal<int32>({ "0", "1", "-1", "1254", "9223372036854775807", "-9223372036854775808", "Bla", "0xff", "0547"}, {0, 1, -1, 1254, -1, -1, 0, 0xff, 547 });
+}
+
+bool FPCGParseStringTests_Int64::RunTest(const FString& Parameters)
+{
+	return RunTestInternal<int64>({ "0", "1", "-1", "1254", "9223372036854775807", "-9223372036854775808", "99999999999999999999999999999999999999999", "Bla", "0xff", "0547" }, {0, 1, -1, 1254, std::numeric_limits<int64>::max(), std::numeric_limits<int64>::min(), std::numeric_limits<int64>::max(), 0, 0xff, 547 });
+}
+
+bool FPCGParseStringTests_Float::RunTest(const FString& Parameters)
+{
+	return RunTestInternal<float>({ "0.0", "1.0", "-1.0", "12.54", "92233720368.54775807", "-92233720368.54775808", "Bla", "0xff", "1e-3"}, {0.0f, 1.0f, -1.0f, 12.54f, 92233720368.54775807f, -92233720368.54775808f, 0.0f, 0.0f, 1e-3f});
+}
+
+bool FPCGParseStringTests_Double::RunTest(const FString& Parameters)
+{
+	return RunTestInternal<double>({ "0.0", "1.0", "-1.0", "12.54", "92233720368.54775807", "-92233720368.54775808", "Bla", "0xff", "1e-3" }, { 0.0, 1.0, -1.0, 12.54, 92233720368.54775807, -92233720368.54775808, 0.0, 0.0, 1e-3 });
+}
+
+bool FPCGParseStringTests_Vector2D::RunTest(const FString& Parameters)
+{
+	return RunTestInternal<FVector2D>({ "0.1,0.2", "1.0,    2.0", "X=0.5 Y=0.7", "X=8.0,            Y=4.1", "Bla", "0.1;0.2" }, { FVector2D{0.1, 0.2}, FVector2D{1.0, 2.0}, FVector2D{0.5, 0.7}, FVector2D{8.0, 4.1}, FVector2D::ZeroVector, FVector2D::ZeroVector });
+}
+
+bool FPCGParseStringTests_Vector::RunTest(const FString& Parameters)
+{
+	return RunTestInternal<FVector>({ "0.1,0.2,0.3", "1.0,   2.0,    3.0", "X=0.5 Y=0.7 Z=4.5", "X=8.0,            Y=4.1,         Z=9.4", "Bla", "0.1;0.2;0.3" }, { FVector{0.1, 0.2, 0.3}, FVector{1.0, 2.0, 3.0}, FVector{0.5, 0.7, 4.5}, FVector{8.0, 4.1, 9.4}, FVector::ZeroVector, FVector::ZeroVector });
+}
+
+bool FPCGParseStringTests_Vector4::RunTest(const FString& Parameters)
+{
+	return RunTestInternal<FVector4>({ "0.1,0.2,0.3,0.4", "1.0,   2.0,    3.0,    4.0", "X=0.5 Y=0.7 Z=4.5 W=6.1", "X=8.0,            Y=4.1,         Z=9.4,     W=7.3", "Bla", "0.1;0.2;0.3;0.4" }, { FVector4{0.1, 0.2, 0.3, 0.4}, FVector4{1.0, 2.0, 3.0, 4.0}, FVector4{0.5, 0.7, 4.5, 6.1}, FVector4{8.0, 4.1, 9.4, 7.3}, FVector4::Zero(), FVector4::Zero()});
+}
+
+bool FPCGParseStringTests_Rotator::RunTest(const FString& Parameters)
+{
+	return RunTestInternal<FRotator>({ "20.5,12.4,3.14", "10.2,     51.3,    65.4", "P=54.7 Y=87 R=89",  "Pitch=54.7 Yaw=87 Roll=89", "P=54.7     Y=87      R=89", "Bla", "25.1;21.4;98.3"}, {FRotator(20.5, 12.4, 3.14),FRotator(10.2, 51.3, 65.4), FRotator(54.7, 87, 89), FRotator(54.7, 87, 89), FRotator(54.7, 87, 89), FRotator::ZeroRotator, FRotator::ZeroRotator });
+}
+
+bool FPCGParseStringTests_Quat::RunTest(const FString& Parameters)
+{
+	return RunTestInternal<FQuat>({ "0.1,0.2,0.3,0.4", "1.0,   2.0,    3.0,    4.0", "X=0.5 Y=0.7 Z=4.5 W=6.1", "X=8.0,            Y=4.1,         Z=9.4,     W=7.3", "Bla", "0.1;0.2;0.3;0.4" }, { FQuat{0.1, 0.2, 0.3, 0.4}, FQuat{1.0, 2.0, 3.0, 4.0}, FQuat{0.5, 0.7, 4.5, 6.1}, FQuat{8.0, 4.1, 9.4, 7.3}, FQuat::Identity, FQuat::Identity});
+}
+
+bool FPCGParseStringTests_Transform::RunTest(const FString& Parameters)
+{
+	return RunTestInternal<FTransform>({ "1.0,2.0,3.0|45.0,45.0,46.0|1.5,2.5,3.5", "Bla", "1.0;2.0;3.0|45.0;45.0;46.0|1.5;2.5;3.5"}, {FTransform(FRotator(45, 45, 46).Quaternion(), FVector(1.0,2.0,3.0), FVector(1.5, 2.5, 3.5)), FTransform::Identity, FTransform::Identity});
+}
\ No newline at end of file
diff --git a/Engine/Plugins/PCG/Source/PCG/Public/Elements/PCGParseString.h b/Engine/Plugins/PCG/Source/PCG/Public/Elements/PCGParseString.h
new file mode 100644
--- /dev/null
+++ b/Engine/Plugins/PCG/Source/PCG/Public/Elements/PCGParseString.h
@@ -0,0 +1,54 @@
+// Copyright Epic Games, Inc. All Rights Reserved.
+
+#pragma once
+
+#include "Elements/Metadata/PCGMetadataOpElementBase.h"
+#include "Metadata/PCGMetadataAttributeTraits.h"
+
+#include "PCGParseString.generated.h"
+
+/**
+* Convert a string to a numerical value.
+*/
+UCLASS(MinimalAPI, BlueprintType, ClassGroup = (Procedural))
+class UPCGParseStringSettings : public UPCGMetadataSettingsBase
+{
+	GENERATED_BODY()
+
+public:
+	//~Begin UPCGSettings interface
+#if WITH_EDITOR
+	virtual FName GetDefaultNodeName() const override;
+	virtual FText GetDefaultNodeTitle() const override;
+	virtual FText GetNodeTooltipText() const override;
+#endif
+
+	virtual bool HasFlippedTitleLines() const override { return false; }
+	virtual FString GetAdditionalTitleInformation() const override;
+
+protected:
+	virtual FPCGElementPtr CreateElement() const override;
+	//~End UPCGSettings interface
+
+public:
+	//~Begin UPCGMetadataSettingsBase interface
+	virtual FPCGAttributePropertyInputSelector GetInputSource(uint32 Index) const override;
+	virtual bool IsSupportedInputType(uint16 TypeId, uint32 InputIndex, bool& bHasSpecialRequirement) const override;
+	virtual bool HasDifferentOutputTypes() const { return true; }
+	virtual TArray<uint16> GetAllOutputTypes() const;
+	//~End UPCGMetadataSettingsBase interface
+
+public:
+	UPROPERTY(BlueprintReadWrite, EditAnywhere, Category = "Input", meta = (PCG_Overridable))
+	FPCGAttributePropertyInputSelector InputSource;
+
+	UPROPERTY(BlueprintReadWrite, EditAnywhere, Category = "Settings", meta = (PCG_Overridable))
+	EPCGMetadataTypes TargetType = EPCGMetadataTypes::Integer32;
+};
+
+class FPCGParseStringElement : public FPCGMetadataElementBase
+{
+protected:
+	virtual bool DoOperation(PCGMetadataOps::FOperationData& OperationData) const override;
+};
+
diff --git a/Engine/Plugins/PCG/Source/PCG/Public/Metadata/PCGMetadataAttributeTraits.h b/Engine/Plugins/PCG/Source/PCG/Public/Metadata/PCGMetadataAttributeTraits.h
--- a/Engine/Plugins/PCG/Source/PCG/Public/Metadata/PCGMetadataAttributeTraits.h
+++ b/Engine/Plugins/PCG/Source/PCG/Public/Metadata/PCGMetadataAttributeTraits.h
@@ -3,6 +3,7 @@
 #pragma once
 
 #include "Math/Transform.h" // IWYU pragma: keep
+#include "Misc/DefaultValueHelper.h"
 #include "UObject/Class.h"
 #include "UObject/SoftObjectPath.h"
 
@@ -339,6 +340,12 @@ namespace PCG
 			{
 				return A.ToString();
 			}
+
+			static T ParseFromString(const FString& InString, bool& bOutSuccess)
+			{
+				bOutSuccess = true;
+				return T(InString);
+			}
 		};
 
 		template<typename T>
@@ -485,7 +492,6 @@ namespace PCG
 			}
 		};
 
-		// Common traits for int32, int64, float, double
 		template<typename T>
 		struct MetadataTraits : DefaultOperationTraits<T>, DefaultWeightedSumTraits<T>, DefaultMinMaxTraits<T>, DefaultCompareTraits<T>, LexToStringTraits<T>, DefaultScalarDistanceTraits<T>
 		{
@@ -494,6 +500,37 @@ namespace PCG
 			enum { NeedsConstruction = false };
 		};
 
+		template<>
+		struct MetadataTraits<int32> : DefaultOperationTraits<int32>, DefaultWeightedSumTraits<int32>, DefaultMinMaxTraits<int32>, DefaultCompareTraits<int32>, LexToStringTraits<int32>, DefaultScalarDistanceTraits<int32>
+		{
+			enum { CompressData = false };
+			enum { CanSearchString = false };
+			enum { NeedsConstruction = false };
+
+			static int32 ParseFromString(const FString& InValue, bool& bOutSuccess)
+			{
+				bOutSuccess = FDefaultValueHelper::IsStringValidInteger(InValue);
+				// FDefaultValueHelper::ParseInt/ParseInt64 have this weird behavior where numbers starting with 0, and not followed by 'x' are treated as base 8.
+				// So do our own parsing.
+				return bOutSuccess ? FCString::Strtoi(*InValue, nullptr, InValue.StartsWith(TEXT("0x")) ? 16 : 10) : 0;
+			}
+		};
+
+		template<>
+		struct MetadataTraits<int64> : DefaultOperationTraits<int64>, DefaultWeightedSumTraits<int64>, DefaultMinMaxTraits<int64>, DefaultCompareTraits<int64>, LexToStringTraits<int64>, DefaultScalarDistanceTraits<int64>
+		{
+			enum { CompressData = false };
+			enum { CanSearchString = false };
+			enum { NeedsConstruction = false };
+
+			static int64 ParseFromString(const FString& InValue, bool& bOutSuccess)
+			{
+				bOutSuccess = FDefaultValueHelper::IsStringValidInteger(InValue);
+				// cf int32
+				return bOutSuccess ? FCString::Strtoi64(*InValue, nullptr, InValue.StartsWith(TEXT("0x")) ? 16 : 10) : 0;
+			}
+		};
+
 		template<>
 		struct MetadataTraits<float> : DefaultOperationTraits<float>, DefaultWeightedSumTraits<float>, DefaultMinMaxTraits<float>, DefaultCompareTraits<float>, LexToStringTraits<float>, DefaultScalarDistanceTraits<float>
 		{
@@ -505,6 +542,13 @@ namespace PCG
 			{
 				return FMath::IsNearlyEqual(A, B);
 			}
+
+			static float ParseFromString(const FString& InValue, bool& bOutSuccess)
+			{
+				float OutValue;
+				bOutSuccess = FDefaultValueHelper::ParseFloat(InValue, OutValue);
+				return bOutSuccess ? OutValue : 0.0f;
+			}
 		};
 
 		template<>
@@ -518,6 +562,13 @@ namespace PCG
 			{
 				return FMath::IsNearlyEqual(A, B);
 			}
+
+			static double ParseFromString(const FString& InValue, bool& bOutSuccess)
+			{
+				double OutValue;
+				bOutSuccess = FDefaultValueHelper::ParseDouble(InValue, OutValue);
+				return bOutSuccess ? OutValue : 0.0;
+			}
 		};
 
 		template<>
@@ -583,6 +634,12 @@ namespace PCG
 			{
 				return !Less(A, B);
 			}
+
+			static bool ParseFromString(const FString& InValue, bool& bOutSuccess)
+			{
+				bOutSuccess = true;
+				return FCString::ToBool(*InValue);
+			}
 		};
 
 		// Vector types
@@ -664,6 +721,14 @@ namespace PCG
 			{
 				return (A.X >= B.X) && (A.Y >= B.Y);
 			}
+
+			static FVector2D ParseFromString(const FString& InValue, bool& bOutSuccess)
+			{
+				FVector2D OutValue;
+				// FDefaultValueHelper::ParseVector expect commas, even with X=... Y=..., which is not the format of ToString. So try to init from string if it fails.
+				bOutSuccess = FDefaultValueHelper::ParseVector2D(InValue, OutValue) || OutValue.InitFromString(InValue);
+				return bOutSuccess ? OutValue : FVector2D::ZeroVector;
+			}
 		};
 
 		template<>
@@ -703,6 +768,14 @@ namespace PCG
 			{
 				return (A.X >= B.X) && (A.Y >= B.Y) && (A.Z >= B.Z);
 			}
+
+			static FVector ParseFromString(const FString& InValue, bool& bOutSuccess)
+			{
+				FVector OutValue;
+				// cf Vector2D
+				bOutSuccess = FDefaultValueHelper::ParseVector(InValue, OutValue) || OutValue.InitFromString(InValue);
+				return bOutSuccess ? OutValue : FVector::ZeroVector;
+			}
 		};
 		
 		template<>
@@ -742,6 +815,14 @@ namespace PCG
 			{
 				return (A.X >= B.X) && (A.Y >= B.Y) && (A.Z >= B.Z) && (A.W >= B.W);
 			}
+
+			static FVector4 ParseFromString(const FString& InValue, bool& bOutSuccess)
+			{
+				FVector4 OutValue;
+				// cf Vector2D
+				bOutSuccess = FDefaultValueHelper::ParseVector4(InValue, OutValue) || OutValue.InitFromString(InValue);
+				return bOutSuccess ? OutValue : FVector4::Zero();
+			}
 		};
 
 		// Quaternion
@@ -809,6 +890,14 @@ namespace PCG
 			{
 				return A.AngularDistance(B);
 			}
+
+			static FQuat ParseFromString(const FString& InValue, bool& bOutSuccess)
+			{
+				FVector4 Temp;
+				// cf Vector2D
+				bOutSuccess = FDefaultValueHelper::ParseVector4(InValue, Temp) || Temp.InitFromString(InValue);
+				return bOutSuccess ? FQuat(Temp.X, Temp.Y, Temp.Z, Temp.W) : FQuat::Identity;
+			}
 		};
 
 		// Rotator
@@ -872,6 +961,13 @@ namespace PCG
 				// Use quaternions for distance calculation, but return degrees instead of radians here
 				return 180.0 * FQuat(A).AngularDistance(FQuat(B)) / UE_DOUBLE_PI;
 			}
+
+			static FRotator ParseFromString(const FString& InValue, bool& bOutSuccess)
+			{
+				FRotator OutValue;
+				bOutSuccess = FDefaultValueHelper::ParseRotator(InValue, OutValue);
+				return bOutSuccess ? OutValue : FRotator::ZeroRotator;
+			}
 		};
 
 		// Transform
@@ -932,6 +1028,13 @@ namespace PCG
 			{
 				return FTransform(FQuat(0.0, 0.0, 0.0, 0.0), FVector::ZeroVector, FVector::ZeroVector);
 			}
+
+			static FTransform ParseFromString(const FString& InValue, bool& bOutSuccess)
+			{
+				FTransform OutValue;
+				bOutSuccess = OutValue.InitFromString(InValue);
+				return bOutSuccess ? OutValue : FTransform::Identity;
+			}
 		};
 
 		// Strings
