[PCG] GPU Create Attribute syntax

diff --git a/Engine/Plugins/PCG/Shaders/Private/PCGDataCollectionDataInterface.ush b/Engine/Plugins/PCG/Shaders/Private/PCGDataCollectionDataInterface.ush
--- a/Engine/Plugins/PCG/Shaders/Private/PCGDataCollectionDataInterface.ush
+++ b/Engine/Plugins/PCG/Shaders/Private/PCGDataCollectionDataInterface.ush
@@ -461,48 +461,6 @@ void RemovePoint_{DataInterfaceName}(uint InDataIndex, uint InElementIndex)
 	SetDensity_{DataInterfaceName}(InDataIndex, InElementIndex, INVALID_DENSITY);
 }
 
-// Write single point data collection header
-void WriteSinglePointDataCollectionHeader_{DataInterfaceName}(uint InPointCount)
-{
-    const uint NumData = 1;
-    WriteNumData_{DataInterfaceName}(NumData);
-    
-	// TODO code gen the various constants here and below rather than hardcoding them.
-    const static uint AttributeStrides[9] = { 3, 4, 3, 3, 3, 4, 1, 1, 1 };
-    
-    for (uint DataIndex = 0; DataIndex < NumData; ++DataIndex)
-    {
-        const uint DataAddress = 8;
-        const uint Id = 0; // PCG_POINT_DATA_TYPE_ID
-        const uint NumAttributes = 9; // Number of point attributes
-        const uint PreambleSizeBytes = 16; // PCG_POINT_DATA_HEADER_PREAMBLE_SIZE_BYTES
-        const uint NumElements = InPointCount;
-        
-        WriteDataAddress_{DataInterfaceName}(DataIndex, DataAddress);
-        WriteDataId_{DataInterfaceName}(DataAddress, Id);
-        WriteDataNumAttributes_{DataInterfaceName}(DataAddress, NumAttributes);
-        WriteDataPreambleSize_{DataInterfaceName}(DataAddress, PreambleSizeBytes);
-        WriteDataInfo_{DataInterfaceName}(DataAddress, NumElements);
-        
-        const uint AttributeHeadersAddress = DataAddress + PreambleSizeBytes;
-        // TODO: These magic numbers should be code gen'd using the constants in PCGComputeConstants
-        uint CurrentAttributeAddress = AttributeHeadersAddress + 128 * 8; // MAX_NUM_ATTRS * ATTR_HEADER_SIZE_BYTES
-        
-        for (uint AttributeIndex = 0; AttributeIndex < NumAttributes; ++AttributeIndex)
-        {
-            const uint AttributeHeaderAddress = AttributeHeadersAddress + 8 * AttributeIndex;
-            const uint AttributeId = AttributeIndex; // Point attributes are [0, 8)
-            const uint AttributeStrideBytes = AttributeStrides[AttributeIndex] * 4;
-            const uint AttributeIdAndStride = (AttributeId << 8) + AttributeStrideBytes;
-            
-            WriteAttributeIdAndStride_{DataInterfaceName}(AttributeHeaderAddress, AttributeIdAndStride);
-            WriteAttributeAddress_{DataInterfaceName}(AttributeHeaderAddress, CurrentAttributeAddress);
-            
-            CurrentAttributeAddress += NumElements * AttributeStrideBytes;
-        }
-    }
-}
-
 // Initialize a single point with default values.
 void InitializePoint_{DataInterfaceName}(uint InDataIndex, uint InElementIndex)
 {
@@ -515,4 +473,6 @@ void InitializePoint_{DataInterfaceName}(uint InDataIndex, uint InElementIndex)
 	SetDensity_{DataInterfaceName}(InDataIndex, InElementIndex, 1.0f);
 	SetSeed_{DataInterfaceName}(InDataIndex, InElementIndex, 42);
 	SetSteepness_{DataInterfaceName}(InDataIndex, InElementIndex, 0.5f);
+	
+	// TODO - code gen initialization of all attributes, or remove this function.
 }
diff --git a/Engine/Plugins/PCG/Source/PCG/Private/Compute/DataInterfaces/PCGDataCollectionDataInterface.cpp b/Engine/Plugins/PCG/Source/PCG/Private/Compute/DataInterfaces/PCGDataCollectionDataInterface.cpp
--- a/Engine/Plugins/PCG/Source/PCG/Private/Compute/DataInterfaces/PCGDataCollectionDataInterface.cpp
+++ b/Engine/Plugins/PCG/Source/PCG/Private/Compute/DataInterfaces/PCGDataCollectionDataInterface.cpp
@@ -3,8 +3,10 @@
 #include "Compute/DataInterfaces/PCGDataCollectionDataInterface.h"
 
 #include "PCGComponent.h"
+#include "PCGContext.h"
 #include "PCGModule.h"
 #include "PCGSettings.h"
+#include "PCGSubsystem.h"
 #include "Compute/PCGDataBinding.h"
 #include "Compute/Elements/PCGComputeGraphElement.h"
 
@@ -14,6 +16,8 @@
 #include "ShaderCompilerCore.h"
 #include "ShaderParameterMetadataBuilder.h"
 
+#define LOCTEXT_NAMESPACE "PCGDataCollectionDataInterface"
+
 void UPCGDataCollectionDataInterface::GetSupportedInputs(TArray<FShaderFunctionDefinition>& OutFunctions) const
 {
 	OutFunctions.AddDefaulted_GetRef()
@@ -427,10 +431,6 @@ void UPCGDataCollectionDataInterface::GetSupportedInputs(TArray<FShaderFunctionD
 			.AddParam(EShaderFundamentalType::Uint) // DataIndex
 			.AddParam(EShaderFundamentalType::Uint); // ElementIndex
 
-		OutFunctions.AddDefaulted_GetRef()
-			.SetName(TEXT("WriteSinglePointDataCollectionHeader"))
-			.AddParam(EShaderFundamentalType::Uint);
-
 		OutFunctions.AddDefaulted_GetRef()
 			.SetName(TEXT("InitializePoint"))
 			.AddParam(EShaderFundamentalType::Uint)
@@ -573,7 +573,35 @@ bool UPCGDataCollectionDataProvider::ProcessReadBackData(FPCGComputeGraphContext
 	}
 
 	FPCGDataCollection DataFromGPU;
-	PinDesc.UnpackDataCollection(RawReadbackData, OutputPinLabelAlias, DataFromGPU);
+	const EPCGUnpackDataCollectionResult Result = PinDesc.UnpackDataCollection(RawReadbackData, OutputPinLabelAlias, DataFromGPU);
+
+#if !(UE_BUILD_SHIPPING || UE_BUILD_TEST) || USE_LOGGING_IN_SHIPPING
+	if (Result == EPCGUnpackDataCollectionResult::DataMismatch)
+	{
+		const FText DataMismatchWarningText = FText::Format(
+			LOCTEXT("UnpackDataCollectionDataMismatch", "Mismatch in expected data while unpacking GPU data collection on pin '{0}'. Static analysis does not match the received data. Data collection will be ignored."),
+			FText::FromName(OutputPinLabel));
+
+#if WITH_EDITOR
+		// InContext will be for the compute graph element (injected during compilation). The code below logs against the original
+		// node so that the log will be visible on the graph.
+		if (ensure(InContext->SourceComponent.IsValid() && InContext->SourceComponent.Get()))
+		{
+			if (UPCGSubsystem* Subsystem = InContext->SourceComponent->GetSubsystem())
+			{
+				FPCGStack StackWithNode = InContext->Stack ? *InContext->Stack : FPCGStack();
+				StackWithNode.PushFrame(ProducerSettings->GetOuter());
+				Subsystem->GetNodeVisualLogsMutable().Log(StackWithNode, ELogVerbosity::Warning, DataMismatchWarningText);
+			}
+		}
+		else
+#endif
+		{
+			PCGE_LOG_C(Warning, LogOnly, InContext, DataMismatchWarningText);
+		}
+	}
+#endif
+
 	RawReadbackData.Reset();
 
 	// Store data in output collection.
@@ -682,3 +710,5 @@ void FPCGDataCollectionDataProviderProxy::GetReadbackData(TArray<FReadbackData>&
 		OutReadbackData.Add(MoveTemp(Data));
 	}
 }
+
+#undef LOCTEXT_NAMESPACE
diff --git a/Engine/Plugins/PCG/Source/PCG/Private/Compute/Elements/PCGComputeGraphElement.cpp b/Engine/Plugins/PCG/Source/PCG/Private/Compute/Elements/PCGComputeGraphElement.cpp
--- a/Engine/Plugins/PCG/Source/PCG/Private/Compute/Elements/PCGComputeGraphElement.cpp
+++ b/Engine/Plugins/PCG/Source/PCG/Private/Compute/Elements/PCGComputeGraphElement.cpp
@@ -149,6 +149,21 @@ bool FPCGComputeGraphElement::ExecuteInternal(FPCGContext* InContext) const
 				});
 			}
 		}
+
+		for (TWeakObjectPtr<const UPCGNode> Node : Graph->KernelToNode)
+		{
+			const UPCGSettings* Settings = Node.Get() ? Node->GetSettings() : nullptr;
+
+			if (Settings)
+			{
+				const UPCGCustomHLSLSettings* KernelSettings = CastChecked<UPCGCustomHLSLSettings>(Settings);
+
+				if (!KernelSettings->IsKernelValid(Context))
+				{
+					return true;
+				}
+			}
+		}
 	}
 
 	check(Context->DataBinding && InContext->SourceComponent.Get());
diff --git a/Engine/Plugins/PCG/Source/PCG/Private/Compute/Elements/PCGComputeGraphElement.h b/Engine/Plugins/PCG/Source/PCG/Private/Compute/Elements/PCGComputeGraphElement.h
--- a/Engine/Plugins/PCG/Source/PCG/Private/Compute/Elements/PCGComputeGraphElement.h
+++ b/Engine/Plugins/PCG/Source/PCG/Private/Compute/Elements/PCGComputeGraphElement.h
@@ -15,6 +15,9 @@
 
 struct FPCGComputeGraphContext : public FPCGContext
 {
+public:
+	virtual bool IsComputeContext() const override { return true; }
+
 protected:
 	virtual void AddExtraStructReferencedObjects(FReferenceCollector& Collector) override;
 
diff --git a/Engine/Plugins/PCG/Source/PCG/Private/Compute/Elements/PCGCustomHLSL.cpp b/Engine/Plugins/PCG/Source/PCG/Private/Compute/Elements/PCGCustomHLSL.cpp
--- a/Engine/Plugins/PCG/Source/PCG/Private/Compute/Elements/PCGCustomHLSL.cpp
+++ b/Engine/Plugins/PCG/Source/PCG/Private/Compute/Elements/PCGCustomHLSL.cpp
@@ -2,22 +2,44 @@
 
 #include "Compute/Elements/PCGCustomHLSL.h"
 
+#include "PCGComputeGraphElement.h"
 #include "PCGContext.h"
+#include "PCGComponent.h"
 #include "PCGEdge.h"
 #include "PCGModule.h"
 #include "PCGPoint.h"
+#include "PCGSubsystem.h"
 #include "Compute/PCGComputeCommon.h"
 #include "Compute/PCGComputeGraph.h"
 #include "Compute/PCGDataBinding.h"
 #include "Data/PCGPointData.h"
 
+#include "Internationalization/Regex.h"
+
 #include UE_INLINE_GENERATED_CPP_BY_NAME(PCGCustomHLSL)
 
 #define LOCTEXT_NAMESPACE "PCGCustomHLSLElement"
+#define PCG_LOGGING_ENABLED (!(UE_BUILD_SHIPPING || UE_BUILD_TEST) || USE_LOGGING_IN_SHIPPING)
 
 namespace PCGHLSLElement
 {
-	const FString KernelAttributeToken = TEXT("@");
+	/** First capture: Pin name(supports a - z, A - Z, and 0 - 9) */
+	constexpr int AttributePinCaptureGroup = 1;
+
+	/** Second capture: Function name(Get, Set, or Create) */
+	constexpr int AttributeFunctionCaptureGroup = 2;
+
+	/** Third capture: Attribute type(e.g.Int, Float, Rotator, etc.) */
+	constexpr int AttributeTypeCaptureGroup = 3;
+
+	/** Fourth capture: Attribute name (supports a-z, A-Z, 0-9, ' ', '-', '_', and '/') */
+	constexpr int AttributeNameCaptureGroup = 4;
+
+	/** Regex pattern used to detect and parse attribute function usage in kernels. */
+	constexpr TCHAR AttributeFunctionPattern[] = { TEXT("([a-zA-Z0-9]+)_(Get|Set)(.*)\\(.*'([a-zA-Z0-9 -_\\/]+)'.*") };
+
+	constexpr TCHAR AttributeFunctionGet[] = { TEXT("Get") };
+	constexpr TCHAR AttributeFunctionSet[] = { TEXT("Set") };
 
 	void ConvertObjectPathToShaderFilePath(FString& InOutPath)
 	{
@@ -47,7 +69,7 @@ namespace PCGHLSLElement
 
 	FString GetKernelAttributeKeyAsString(const FPCGKernelAttributeKey& Key)
 	{
-		return FString::Format(TEXT("{0}{1}"), { KernelAttributeToken, Key.Name.ToString()});
+		return FString::Format(TEXT("'{0}'"), { Key.Name.ToString() });
 	}
 }
 
@@ -63,6 +85,7 @@ void UPCGCustomHLSLSettings::PostLoad()
 	Super::PostLoad();
 
 	UpdatePinSettings();
+	UpdateAttributeKeys();
 }
 
 void UPCGCustomHLSLSettings::PostInitProperties()
@@ -70,7 +93,6 @@ void UPCGCustomHLSLSettings::PostInitProperties()
 	Super::PostInitProperties();
 
 	UpdatePinSettings();
-
 	UpdateDeclarations();
 }
 #endif
@@ -91,6 +113,7 @@ void UPCGCustomHLSLSettings::PostEditChangeProperty(FPropertyChangedEvent& Prope
 	Super::PostEditChangeProperty(PropertyChangedEvent);
 
 	UpdateDeclarations();
+	UpdateAttributeKeys();
 }
 #endif
 
@@ -216,106 +239,16 @@ int UPCGCustomHLSLSettings::ComputeKernelThreadCount(const UPCGDataBinding* Bind
 	return ThreadCount;
 }
 
-TArray<FPCGKernelAttributeKey> UPCGCustomHLSLSettings::GetKernelAttributeKeys() const
-{
-	TArray<FPCGKernelAttributeKey> Keys;
-
-	const FString SourceToParse = ShaderFunctions + ShaderSource;
-
-	int NumExistingAttributes = PCGComputeConstants::NUM_RESERVED_ATTRS;
-	int AttributeSourceIndex = SourceToParse.Find(PCGHLSLElement::KernelAttributeToken);
-
-	while (AttributeSourceIndex != INDEX_NONE)
-	{
-		const int GetStringIndex = SourceToParse.Find(TEXT("Get"), ESearchCase::CaseSensitive, ESearchDir::FromEnd, AttributeSourceIndex);
-		const int SetStringIndex = SourceToParse.Find(TEXT("Set"), ESearchCase::CaseSensitive, ESearchDir::FromEnd, AttributeSourceIndex);
-		const int StartTypeStringIndex = FMath::Max(GetStringIndex, SetStringIndex) + 3;
-		const int EndTypeStringIndex = SourceToParse.Find(TEXT("("), ESearchCase::CaseSensitive, ESearchDir::FromStart, StartTypeStringIndex);
-
-		const FString TypeString = (StartTypeStringIndex != INDEX_NONE || EndTypeStringIndex != INDEX_NONE)
-			? SourceToParse.Mid(StartTypeStringIndex, EndTypeStringIndex - StartTypeStringIndex)
-			: TEXT("");
-
-		if (TypeString.IsEmpty())
-		{
-			UE_LOG(LogPCG, Error, TEXT("Read invalid attribute type in shader source.")); 
-		}
-
-		EPCGKernelAttributeType Type = EPCGKernelAttributeType::Bool;
-
-		if (TypeString == TEXT("Bool"))
-		{
-			Type = EPCGKernelAttributeType::Bool;
-		}
-		else if (TypeString == TEXT("Int"))
-		{
-			Type = EPCGKernelAttributeType::Int;
-		}
-		else if (TypeString == TEXT("Float"))
-		{
-			Type = EPCGKernelAttributeType::Float;
-		}
-		else if (TypeString == TEXT("Float2"))
-		{
-			Type = EPCGKernelAttributeType::Float2;
-		}
-		else if (TypeString == TEXT("Float3"))
-		{
-			Type = EPCGKernelAttributeType::Float3;
-		}
-		else if (TypeString == TEXT("Float4"))
-		{
-			Type = EPCGKernelAttributeType::Float4;
-		}
-		else if (TypeString == TEXT("Rotator"))
-		{
-			Type = EPCGKernelAttributeType::Rotator;
-		}
-		else if (TypeString == TEXT("Quat"))
-		{
-			Type = EPCGKernelAttributeType::Quat;
-		}
-		else if (TypeString == TEXT("Transform"))
-		{
-			Type = EPCGKernelAttributeType::Transform;
-		}
-		else
-		{
-			UE_LOG(LogPCG, Error, TEXT("Read invalid attribute type in shader source."));
-		}
-
-		FString AttrNameString = "";
-
-		// TODO: Copy the full attribute name at once with AttrNameString.Mid()
-		for (int I = AttributeSourceIndex + 1; I < SourceToParse.Len(); ++I)
-		{
-			const char Char = SourceToParse[I];
-
-			if (isalnum(Char))
-			{
-				AttrNameString += Char;
-			}
-			else
-			{
-				break;
-			}
-		}
-
-		Keys.AddUnique({ Type, FName(*AttrNameString) });
-
-		AttributeSourceIndex = SourceToParse.Find(PCGHLSLElement::KernelAttributeToken, ESearchCase::IgnoreCase, ESearchDir::FromStart, AttributeSourceIndex + 1);
-	}
-
-	return Keys;
-}
-
 FPCGDataCollectionDesc UPCGCustomHLSLSettings::ComputeOutputPinDataDesc(const UPCGPin* OutputPin, const UPCGDataBinding* Binding) const
 {
 	check(OutputPin);
 	FPCGDataCollectionDesc PinDesc;
 
+	const FPCGPinPropertiesGPU* PropertiesGPU = GetOutputPinPropertiesGPU(OutputPin->Properties.Label);
+	const bool bNeedsPointPin = KernelType == EPCGKernelType::PointProcessor || KernelType == EPCGKernelType::PointGenerator;
+
 	// The primary output pin follows any rules prescribed by kernel type.
-	if (OutputPin == GetFirstOutputPin())
+	if (OutputPin == GetFirstOutputPin() && bNeedsPointPin)
 	{
 		if (KernelType == EPCGKernelType::PointProcessor)
 		{
@@ -324,63 +257,76 @@ FPCGDataCollectionDesc UPCGCustomHLSLSettings::ComputeOutputPinDataDesc(const UP
 			{
 				PinDesc = ComputeInputPinDataDesc(PointProcessingInputPin, Binding);
 			}
-
-			return PinDesc;
 		}
 		else if (KernelType == EPCGKernelType::PointGenerator)
 		{
 			// Generators always produce a single point data with known point count.
 			PinDesc.DataDescs.Emplace(EPCGDataType::Point, PointCount);
-
-			// TODO: Until support for 'Create Attributes' exists, PointGenerators cannot manipulate attributes.
-			return PinDesc;
 		}
 	}
-
-	const FPCGPinPropertiesGPU* PropertiesGPU = GetOutputPinPropertiesGPU(OutputPin->Properties.Label);
-	if (!ensure(PropertiesGPU))
-	{
-		return PinDesc;
-	}
-
-	const bool bNeedsPointPin = KernelType == EPCGKernelType::PointProcessor || KernelType == EPCGKernelType::PointGenerator;
-
-	// No size set by kernel, fall back to pin settings.
-	if (PropertiesGPU->BufferSizeMode == EPCGPinBufferSizeMode::FromFirstPin)
+	else if (ensure(PropertiesGPU))
 	{
-		if (const UPCGPin* InputPin = bNeedsPointPin ? GetPointProcessingInputPin() : GetFirstInputPin())
+		// No size set by kernel, fall back to pin settings.
+		if (PropertiesGPU->BufferSizeMode == EPCGPinBufferSizeMode::FromFirstPin)
 		{
-			PinDesc = ComputeInputPinDataDesc(InputPin, Binding);
+			if (const UPCGPin* InputPin = bNeedsPointPin ? GetPointProcessingInputPin() : GetFirstInputPin())
+			{
+				PinDesc = ComputeInputPinDataDesc(InputPin, Binding);
+			}
 		}
-	}
-	else if (PropertiesGPU->BufferSizeMode == EPCGPinBufferSizeMode::FromProductOfInputPins)
-	{
-		// TODO for now we only support point. Needs extended logic for other data types.
-		const EPCGDataType DataType = EPCGDataType::Point;
+		else if (PropertiesGPU->BufferSizeMode == EPCGPinBufferSizeMode::FromProductOfInputPins)
+		{
+			// TODO for now we only support point. Needs extended logic for other data types.
+			const EPCGDataType DataType = EPCGDataType::Point;
 
-		int TotalElementCount = 0;
+			int TotalElementCount = 0;
 
-		for (const FName& PinLabel : PropertiesGPU->BufferSizeInputPinLabels)
-		{
-			if (const UPCGPin* InputPin = GetInputPin(PinLabel))
+			for (const FName& PinLabel : PropertiesGPU->BufferSizeInputPinLabels)
 			{
-				const int ElementCount = ComputeInputPinDataDesc(InputPin, Binding).ComputeDataElementCount(DataType);
-				TotalElementCount = FMath::Max(TotalElementCount, 1) * ElementCount;
+				if (const UPCGPin* InputPin = GetInputPin(PinLabel))
+				{
+					const int ElementCount = ComputeInputPinDataDesc(InputPin, Binding).ComputeDataElementCount(DataType);
+					TotalElementCount = FMath::Max(TotalElementCount, 1) * ElementCount;
+				}
 			}
-		}
 
-		if (TotalElementCount > 0)
+			if (TotalElementCount > 0)
+			{
+				PinDesc.DataDescs.Emplace(DataType, TotalElementCount);
+			}
+		}
+		else if (PropertiesGPU->BufferSizeMode == EPCGPinBufferSizeMode::FixedElementCount)
 		{
-			PinDesc.DataDescs.Emplace(DataType, TotalElementCount);
+			if (ensure(PropertiesGPU->FixedBufferElementCount > 0))
+			{
+				// TODO for now we only support point. Needs extended logic for other data types.
+				const EPCGDataType DataType = EPCGDataType::Point;
+				PinDesc.DataDescs.Emplace(DataType, PropertiesGPU->FixedBufferElementCount);
+			}
 		}
 	}
-	else if (PropertiesGPU->BufferSizeMode == EPCGPinBufferSizeMode::FixedElementCount)
+
+	const TMap<FPCGKernelAttributeKey, int32>& GlobalAttributeLookupTable = ensure(Binding && Binding->Graph) ? Binding->Graph->GetAttributeLookupTable() : TMap<FPCGKernelAttributeKey, int32>();
+
+	for (const FPCGKernelAttributeKey& AttributeKey : KernelAttributeKeys)
 	{
-		if (ensure(PropertiesGPU->FixedBufferElementCount > 0))
+		// Add attributes that will be created for this pin on the GPU.
+		if (const TArray<TTuple<FPCGKernelAttributeKey, bool>>* Keys = PinToAttributeKeys.Find(OutputPin->Properties.Label))
 		{
-			// TODO for now we only support point. Needs extended logic for other data types.
-			const EPCGDataType DataType = EPCGDataType::Point;
-			PinDesc.DataDescs.Emplace(DataType, PropertiesGPU->FixedBufferElementCount);
+			const TTuple<FPCGKernelAttributeKey, bool>* Pair = Keys->FindByPredicate([AttributeKey](const TTuple<FPCGKernelAttributeKey, bool>& Pair) { return Pair.Key == AttributeKey; });
+			const bool bCreatedOnGPU = Pair && Pair->Value;
+
+			if (bCreatedOnGPU)
+			{
+				for (FPCGDataDesc& DataDesc : PinDesc.DataDescs)
+				{
+					if (const int* Index = GlobalAttributeLookupTable.Find(AttributeKey))
+					{
+						const FPCGKernelAttributeDesc AttributeDesc(*Index, AttributeKey.Type, AttributeKey.Name);
+						DataDesc.AttributeDescs.AddUnique(AttributeDesc);
+					}
+				}
+			}
 		}
 	}
 
@@ -619,6 +565,7 @@ void UPCGCustomHLSLSettings::UpdateDeclarations()
 			InputDeclarations += TEXT("\nuint <pin>_GetNumData();\n");
 			InputDeclarations += TEXT("uint <pin>_GetNumElements();\n");
 			InputDeclarations += TEXT("<type> <pin>_Get<type>(uint DataIndex, uint AttributeId, uint ElementIndex);\n");
+			InputDeclarations += TEXT("<type> <pin>_Get<type>(uint DataIndex, 'AttributeName', uint ElementIndex);\n");
 		}
 
 		if (!PointDataPins.IsEmpty())
@@ -697,6 +644,7 @@ void UPCGCustomHLSLSettings::UpdateDeclarations()
 			OutputDeclarations += TEXT("// Valid types: bool, int, float, float2, float3, float4, Rotator (float3), Quat (float4), Transform (float4x4)\n");
 
 			OutputDeclarations += TEXT("\nvoid <pin>_Set<type>(uint DataIndex, uint AttributeId, uint ElementIndex, <type> Value);\n");
+			OutputDeclarations += TEXT("void <pin>_Set<type>(uint DataIndex, 'AttributeName', uint ElementIndex, <type> Value);\n");
 		}
 
 		if (!PointDataPins.IsEmpty())
@@ -790,28 +738,126 @@ void UPCGCustomHLSLSettings::UpdatePinSettings()
 		Properties.bAllowEditMultipleData = false;
 	}
 }
+
+void UPCGCustomHLSLSettings::UpdateAttributeKeys()
+{
+	KernelAttributeKeys.Reset();
+	PinToAttributeKeys.Reset();
+
+	const FString Source = ShaderFunctions + ShaderSource;
+	FRegexMatcher ModuleMatcher(FRegexPattern(PCGHLSLElement::AttributeFunctionPattern), Source);
+
+	while (ModuleMatcher.FindNext())
+	{
+		const FString PinStr = ModuleMatcher.GetCaptureGroup(PCGHLSLElement::AttributePinCaptureGroup);
+		const FString FuncStr = ModuleMatcher.GetCaptureGroup(PCGHLSLElement::AttributeFunctionCaptureGroup);
+		const FString TypeStr = ModuleMatcher.GetCaptureGroup(PCGHLSLElement::AttributeTypeCaptureGroup);
+		const FString NameStr = ModuleMatcher.GetCaptureGroup(PCGHLSLElement::AttributeNameCaptureGroup);
+
+		const int LineStartIndex = ModuleMatcher.GetMatchBeginning();
+		int CurrentSourceIndex = Source.Find("\n", ESearchCase::IgnoreCase, ESearchDir::FromStart, 0);
+		int LineNumber = 0;
+
+		while (CurrentSourceIndex < LineStartIndex && CurrentSourceIndex != INDEX_NONE)
+		{
+			++LineNumber;
+			CurrentSourceIndex = Source.Find("\n", ESearchCase::IgnoreCase, ESearchDir::FromStart, CurrentSourceIndex + 1);
+		}
+
+		if (PinStr.IsEmpty() || FuncStr.IsEmpty() || TypeStr.IsEmpty() || NameStr.IsEmpty())
+		{
+			UE_LOG(LogPCG, Error, TEXT("Invalid attribute usage in shader source, line %d."), LineNumber);
+			continue;
+		}
+
+		const UEnum* AttributeTypeEnum = StaticEnum<EPCGKernelAttributeType>();
+		check(AttributeTypeEnum);
+
+		const int64 AttributeType = AttributeTypeEnum->GetValueByName(FName(*TypeStr));
+
+		if (AttributeType == INDEX_NONE)
+		{
+			UE_LOG(LogPCG, Error, TEXT("Invalid attribute type in shader source, line %d."), LineNumber);
+			continue;
+		}
+
+		// Add the attribute if it hasn't already been referenced.
+		FPCGKernelAttributeKey Key;
+		Key.Type = static_cast<EPCGKernelAttributeType>(AttributeType);
+		Key.Name = FName(*NameStr);
+		KernelAttributeKeys.AddUnique(Key);
+
+		// Add an entry mapping this pin to the referenced attribute, if the entry doesn't already exist.
+		TArray<TTuple<FPCGKernelAttributeKey, bool>>& Keys = PinToAttributeKeys.FindOrAdd(FName(*PinStr));
+		Keys.AddUnique(MakeTuple(Key, false));
+	}
+
+	// Process each output pin for any new attributes they want to create.
+	for (const FPCGPinPropertiesGPU& OutputPin : OutputPins)
+	{
+		for (const FPCGKernelAttributeKey& Key : OutputPin.CreatedKernelAttributeKeys)
+		{
+			KernelAttributeKeys.AddUnique(Key);
+
+			TArray<TTuple<FPCGKernelAttributeKey, bool>>& Keys = PinToAttributeKeys.FindOrAdd(OutputPin.Label);
+			TTuple<FPCGKernelAttributeKey, bool>* Pair = Keys.FindByPredicate([Key](const TTuple<FPCGKernelAttributeKey, bool>& Pair) { return Pair.Key == Key; });
+
+			// Mark as created on GPU
+			if (Pair)
+			{
+				Pair->Value = true;
+			}
+			else
+			{
+				Keys.Add(MakeTuple(Key, /*bCreatedOnGPU=*/true));
+			}
+		}
+	}
+}
 #endif
 
 bool UPCGCustomHLSLSettings::IsKernelValid(FPCGContext* InContext) const
 {
-	if (OutputPinProperties().IsEmpty())
+	auto LogGraphError = [InContext, This=this](const FText& InText)
+	{
+#if WITH_EDITOR
+		if (InContext && ensure(InContext->SourceComponent.IsValid() && InContext->SourceComponent.Get()))
+		{
+			if (UPCGSubsystem* Subsystem = InContext->SourceComponent->GetSubsystem())
+			{
+				FPCGStack StackWithNode = InContext->Stack ? *InContext->Stack : FPCGStack();
+				StackWithNode.PushFrame(This->GetOuter());
+
+				Subsystem->GetNodeVisualLogsMutable().Log(StackWithNode, ELogVerbosity::Error, InText);
+			}
+		}
+#endif
+
+		PCGE_LOG_C(Error, LogOnly, InContext, InText);
+	};
+
+	if (OutputPins.IsEmpty())
 	{
+#if PCG_LOGGING_ENABLED
 		if (InContext)
 		{
-			PCGE_LOG_C(Error, GraphAndLog, InContext, LOCTEXT("NoOutputs", "Kernels must have at least one output."));
+			LogGraphError(LOCTEXT("NoOutputs", "Custom HLSL nodes must have at least one output."));
 		}
+#endif
 
 		return false;
 	}
 
-	for (const FPCGPinProperties& Properties : InputPinProperties())
+	for (const FPCGPinProperties& Properties : InputPins)
 	{
 		if (Properties.AllowedTypes == EPCGDataType::Any)
 		{
+#if PCG_LOGGING_ENABLED
 			if (InContext)
 			{
-				PCGE_LOG_C(Error, GraphAndLog, InContext, FText::Format(LOCTEXT("InvalidAnyInput", "Custom kernels do not support inputs of type Any, found on pin '{0}'."), FText::FromName(Properties.Label)));
+				LogGraphError(FText::Format(LOCTEXT("InvalidAnyInput", "Custom HLSL nodes do not support inputs of type Any, found on pin '{0}'."), FText::FromName(Properties.Label)));
 			}
+#endif
 
 			return false;
 		}
@@ -822,30 +868,36 @@ bool UPCGCustomHLSLSettings::IsKernelValid(FPCGContext* InContext) const
 	{
 		if (Properties.AllowedTypes == EPCGDataType::Any)
 		{
+#if PCG_LOGGING_ENABLED
 			if (InContext)
 			{
-				PCGE_LOG_C(Error, GraphAndLog, InContext, FText::Format(LOCTEXT("InvalidAnyOutput", "Custom kernels do not support outputs of type Any, found on pin '{0}'."), FText::FromName(Properties.Label)));
+				LogGraphError(FText::Format(LOCTEXT("InvalidAnyOutput", "Custom HLSL nodes do not support outputs of type Any, found on pin '{0}'."), FText::FromName(Properties.Label)));
 			}
+#endif
 
 			return false;
 		}
 
 		if (!!(Properties.AllowedTypes & EPCGDataType::Landscape))
 		{
+#if PCG_LOGGING_ENABLED
 			if (InContext)
 			{
-				PCGE_LOG_C(Error, GraphAndLog, InContext, FText::Format(LOCTEXT("InvalidLSOutput", "Custom kernels do not support outputs of type Landscape, found on pin '{0}'."), FText::FromName(Properties.Label)));
+				LogGraphError(FText::Format(LOCTEXT("InvalidLSOutput", "Custom HLSL nodes do not support outputs of type Landscape, found on pin '{0}'."), FText::FromName(Properties.Label)));
 			}
+#endif
 
 			return false;
 		}
 
 		if (!!(Properties.AllowedTypes & EPCGDataType::Texture))
 		{
+#if PCG_LOGGING_ENABLED
 			if (InContext)
 			{
-				PCGE_LOG_C(Error, GraphAndLog, InContext, FText::Format(LOCTEXT("InvalidTextureOutput", "Custom kernels do not support outputs of type Texture, found on pin '{0}'."), FText::FromName(Properties.Label)));
+				LogGraphError(FText::Format(LOCTEXT("InvalidTextureOutput", "Custom HLSL nodes do not support outputs of type Texture, found on pin '{0}'."), FText::FromName(Properties.Label)));
 			}
+#endif
 
 			return false;
 		}
@@ -857,52 +909,59 @@ bool UPCGCustomHLSLSettings::IsKernelValid(FPCGContext* InContext) const
 			{
 				if (Properties.FixedBufferElementCount <= 0)
 				{
+#if PCG_LOGGING_ENABLED
 					if (InContext)
 					{
-						PCGE_LOG_C(Error, GraphAndLog, InContext, FText::Format(
+						LogGraphError(FText::Format(
 							LOCTEXT("InvalidFixedBufferSize", "Fixed GPU buffer size on '{0}' was invalid (%d)."),
 							FText::FromName(Properties.Label),
 							Properties.FixedBufferElementCount));
 					}
-
+#endif
 				}
 			}
 			else if (Properties.BufferSizeMode == EPCGPinBufferSizeMode::FromFirstPin)
 			{
-				if (InputPinProperties().IsEmpty())
+				if (InputPins.IsEmpty())
 				{
+#if PCG_LOGGING_ENABLED
 					if (InContext)
 					{
-						PCGE_LOG_C(Error, GraphAndLog, InContext, FText::Format(
+						LogGraphError(FText::Format(
 							LOCTEXT("InvalidBufferSizeNoInputPin", "GPU buffer size for pin '{0}' could not be computed as there are no input pins."),
 							FText::FromName(Properties.Label)));
 					}
+#endif
 
 					return false;
 				}
 			}
 			else if (Properties.BufferSizeMode == EPCGPinBufferSizeMode::FromProductOfInputPins)
 			{
-				if (InputPinProperties().IsEmpty())
+				if (InputPins.IsEmpty())
 				{
+#if PCG_LOGGING_ENABLED
 					if (InContext)
 					{
-						PCGE_LOG_C(Error, GraphAndLog, InContext, FText::Format(
+						LogGraphError(FText::Format(
 							LOCTEXT("InvalidBufferSizeNoInputPins", "GPU buffer size for pin '{0}' could not be computed as there are no input pins on this node."),
 							FText::FromName(Properties.Label)));
 					}
+#endif
 
 					return false;
 				}
 
 				if (Properties.BufferSizeInputPinLabels.IsEmpty())
 				{
+#if PCG_LOGGING_ENABLED
 					if (InContext)
 					{
-						PCGE_LOG_C(Error, GraphAndLog, InContext, FText::Format(
+						LogGraphError(FText::Format(
 							LOCTEXT("InvalidBufferSizeNoBufferPins", "GPU buffer size for pin '{0}' could not be computed as input pins are specified in the pin settings."),
 							FText::FromName(Properties.Label)));
 					}
+#endif
 
 					return false;
 				}
@@ -911,13 +970,15 @@ bool UPCGCustomHLSLSettings::IsKernelValid(FPCGContext* InContext) const
 				{
 					if (!GetInputPin(Label))
 					{
+#if PCG_LOGGING_ENABLED
 						if (InContext)
 						{
-							PCGE_LOG_C(Error, GraphAndLog, InContext, FText::Format(
+							LogGraphError(FText::Format(
 								LOCTEXT("MissingBufferSizePin", "GPU buffer size for pin '{0}' could not be computed. Invalid pin specified in Input Pins array: '{1}'."),
 								FText::FromName(Properties.Label),
 								FText::FromName(Label)));
 						}
+#endif
 
 						return false;
 					}
@@ -932,36 +993,53 @@ bool UPCGCustomHLSLSettings::IsKernelValid(FPCGContext* InContext) const
 	{
 		if (!GetPointProcessingInputPin())
 		{
+#if PCG_LOGGING_ENABLED
 			if (InContext)
 			{
-				PCGE_LOG_C(Error, GraphAndLog, InContext, LOCTEXT("InvalidPPInput", "Point processing kernel requires a first input pin of type point."));
+				LogGraphError(LOCTEXT("InvalidPPInput", "Point processing kernel requires a first input pin of type point."));
 			}
+#endif
 
 			return false;
 		}
 
 		if (!GetFirstPointOutputPin())
 		{
+#if PCG_LOGGING_ENABLED
 			if (InContext)
 			{
-				PCGE_LOG_C(Error, GraphAndLog, InContext, LOCTEXT("InvalidPPOutput", "Point processing kernel requires a first output pin of type point."));
+				LogGraphError(LOCTEXT("InvalidPPOutput", "Point processing kernel requires a first output pin of type point."));
 			}
+#endif
 
 			return false;
 		}
 	}
 
-	auto CheckPinLabels = [InContext](const TArray<FPCGPinProperties>& Pins)
+	auto CheckPinLabels = [InContext, &LogGraphError](const TArray<FPCGPinProperties>& Pins)
 	{
 		TSet<FName> EncounteredLabels;
 		for (const FPCGPinProperties& Pin : Pins)
 		{
 			if (EncounteredLabels.Contains(Pin.Label))
 			{
+#if PCG_LOGGING_ENABLED
+				if (InContext)
+				{
+					LogGraphError(FText::Format(LOCTEXT("DuplicatedPinLabels", "Duplicate pin label '{0}', all labels must be unique."), FText::FromName(Pin.Label)));
+				}
+#endif
+
+				return false;
+			}
+			else if (Pin.Label == NAME_None)
+			{
+#if PCG_LOGGING_ENABLED
 				if (InContext)
 				{
-					PCGE_LOG_C(Error, GraphAndLog, InContext, FText::Format(LOCTEXT("DuplicatedPinLabels", "Duplicate pin label '{0}', all labels must be unique."), FText::FromName(Pin.Label)));
+					LogGraphError(LOCTEXT("InvalidPinLabelNone", "Pin label 'None' is not a valid pin label."));
 				}
+#endif
 
 				return false;
 			}
@@ -970,7 +1048,7 @@ bool UPCGCustomHLSLSettings::IsKernelValid(FPCGContext* InContext) const
 		return true;
 	};
 
-	if (!CheckPinLabels(InputPinProperties()) || !CheckPinLabels(OutputPinProperties()))
+	if (!CheckPinLabels(InputPins) || !CheckPinLabels(OutputPinProperties()))
 	{
 		return false;
 	}
@@ -979,10 +1057,12 @@ bool UPCGCustomHLSLSettings::IsKernelValid(FPCGContext* InContext) const
 	{
 		if (ThreadCountInputPinLabels.IsEmpty())
 		{
+#if PCG_LOGGING_ENABLED
 			if (InContext)
 			{
-				PCGE_LOG_C(Error, GraphAndLog, InContext, LOCTEXT("MissingThreadCountPins", "Dispatch thread count is based on input pins but no labels have been set in Input Pins array."));
+				LogGraphError(LOCTEXT("MissingThreadCountPins", "Dispatch thread count is based on input pins but no labels have been set in Input Pins array."));
 			}
+#endif
 
 			return false;
 		}
@@ -991,15 +1071,210 @@ bool UPCGCustomHLSLSettings::IsKernelValid(FPCGContext* InContext) const
 		{
 			if (!GetInputPin(Label))
 			{
+#if PCG_LOGGING_ENABLED
 				if (InContext)
 				{
-					PCGE_LOG_C(Error, GraphAndLog, InContext, FText::Format(LOCTEXT("MissingThreadCountPin", "Invalid pin specified in Input Pins array: '{0}'."), FText::FromName(Label)));
+					LogGraphError(FText::Format(LOCTEXT("MissingThreadCountPin", "Invalid pin specified in Input Pins array: '{0}'."), FText::FromName(Label)));
+				}
+#endif
+
+				return false;
+			}
+		}
+	}
+
+	for (const FPCGKernelAttributeKey& AttributeKey : KernelAttributeKeys)
+	{
+		if (AttributeKey.Name == NAME_None)
+		{
+#if PCG_LOGGING_ENABLED
+			if (InContext)
+			{
+				LogGraphError(LOCTEXT("InvalidAttributeNameNone", "'None' is not a valid GPU attribute name, check the 'Attributes to Create' array on your pins."));
+			}
+#endif
+
+			return false;
+		}
+	}
+
+	if (InContext)
+	{
+		FText* ErrorTextPtr = nullptr;
+#if PCG_LOGGING_ENABLED
+		FText ErrorText;
+		ErrorTextPtr = &ErrorText;
+#endif
+
+		if (!AreKernelAttributesValid(InContext, ErrorTextPtr))
+		{
+			if (ErrorTextPtr)
+			{
+				LogGraphError(*ErrorTextPtr);
+			}
+			return false;
+		}
+	}
+
+	return true;
+}
+
+bool UPCGCustomHLSLSettings::AreKernelAttributesValid(FPCGContext* InContext, FText* OutErrorText) const
+{
+	// The context can either be a compute graph element context (if the compute graph was successfully created), otherwise
+	// it will be the original CPU node context. We need the former to run the following validation.
+	if (!InContext || !InContext->IsComputeContext())
+	{
+		return true;
+	}
+
+	const FPCGComputeGraphContext* Context = static_cast<FPCGComputeGraphContext*>(InContext);
+	const UPCGDataBinding* DataBinding = Context->DataBinding.Get();
+	const UPCGNode* Node = Cast<UPCGNode>(GetOuter());
+
+	if (DataBinding && Node)
+	{
+		const TArray<TObjectPtr<UPCGPin>>& InPins = Node->GetInputPins();
+		const TArray<TObjectPtr<UPCGPin>>& OutPins = Node->GetOutputPins();
+
+		TMap<FName, FPCGDataCollectionDesc> InputPinDescs;
+		TMap<FName, FPCGDataCollectionDesc> OutputPinDescs;
+
+		for (const UPCGPin* InputPin : InPins)
+		{
+			check(InputPin);
+			InputPinDescs.Add(InputPin->Properties.Label, ComputeInputPinDataDesc(InputPin, DataBinding));
+		}
+
+		for (const UPCGPin* OutputPin : OutPins)
+		{
+			check(OutputPin);
+			OutputPinDescs.Add(OutputPin->Properties.Label, ComputeOutputPinDataDesc(OutputPin, DataBinding));
+		}
+
+		const FString Source = ShaderFunctions + ShaderSource;
+		FRegexMatcher ModuleMatcher(FRegexPattern(PCGHLSLElement::AttributeFunctionPattern), Source);
+
+		while (ModuleMatcher.FindNext())
+		{
+			const FString PinStr = ModuleMatcher.GetCaptureGroup(PCGHLSLElement::AttributePinCaptureGroup);
+			const FString FuncStr = ModuleMatcher.GetCaptureGroup(PCGHLSLElement::AttributeFunctionCaptureGroup);
+			const FString TypeStr = ModuleMatcher.GetCaptureGroup(PCGHLSLElement::AttributeTypeCaptureGroup);
+			const FString NameStr = ModuleMatcher.GetCaptureGroup(PCGHLSLElement::AttributeNameCaptureGroup);
+
+			const FName PinName = FName(*PinStr);
+			const FPCGDataCollectionDesc* PinDesc = nullptr;
+
+			auto ConstructFunctionText = [&PinStr, &FuncStr, &TypeStr]()
+				{
+					return FText::FromString(PinStr + TEXT("_") + FuncStr + TypeStr);
+				};
+
+			if (FuncStr == PCGHLSLElement::AttributeFunctionSet)
+			{
+				PinDesc = OutputPinDescs.Find(FName(*PinStr));
+
+				if (!PinDesc && InputPinDescs.Find(PinName))
+				{
+#if PCG_LOGGING_ENABLED
+					if (OutErrorText)
+					{
+						*OutErrorText = FText::Format(
+							LOCTEXT("InvalidSetAttributeUsage", "Tried to call attribute function '{0}' on read-only input pin '{1}'."),
+							ConstructFunctionText(),
+							FText::FromName(PinName));
+					}
+#endif
+
+					return false;
+				}
+			}
+			else if (ensure(FuncStr == PCGHLSLElement::AttributeFunctionGet))
+			{
+				PinDesc = InputPinDescs.Find(FName(*PinStr));
+			}
+
+			if (!PinDesc)
+			{
+#if PCG_LOGGING_ENABLED
+				if (OutErrorText)
+				{
+					*OutErrorText = FText::Format(
+						LOCTEXT("InvalidAttributePinName", "Tried to call attribute function '{0}' on non-existent pin '{1}'."),
+						ConstructFunctionText(),
+						FText::FromName(PinName));
+				}
+#endif
+
+				return false;
+			}
+
+			const UEnum* AttributeTypeEnum = StaticEnum<EPCGKernelAttributeType>();
+			check(AttributeTypeEnum);
+
+			const int64 AttributeType = AttributeTypeEnum->GetValueByName(FName(*TypeStr));
+
+			if (AttributeType == INDEX_NONE)
+			{
+#if PCG_LOGGING_ENABLED
+				if (OutErrorText)
+				{
+					*OutErrorText = FText::Format(
+						LOCTEXT("InvalidAttributePinType", "Tried to call attribute function '{0}' on non-existent type '{1}'."),
+						ConstructFunctionText(),
+						FText::FromString(TypeStr));
+				}
+#endif
+
+				return false;
+			}
+
+			const FPCGKernelAttributeDesc* AttrDesc = nullptr;
+			const FName AttrName = FName(*NameStr);
+
+			if (!PinDesc->DataDescs.IsEmpty())
+			{
+				// Note: This assumes attributes are the same on all data on a pin, which is true for now
+				const FPCGDataDesc& DataDesc = PinDesc->DataDescs[0];
+				AttrDesc = DataDesc.AttributeDescs.FindByPredicate([AttrName](const FPCGKernelAttributeDesc& Desc) { return Desc.Name == AttrName; });
+			}
+
+			if (!AttrDesc)
+			{
+#if PCG_LOGGING_ENABLED
+				if (OutErrorText)
+				{
+					*OutErrorText = FText::Format(
+						LOCTEXT("InvalidAttributeName", "Tried to call attribute function '{0}' on attribute '{1}' which does not exist."),
+						ConstructFunctionText(),
+						FText::FromName(AttrName));
+				}
+#endif
+
+				return false;
+			}
+
+			if (AttrDesc->Type != static_cast<EPCGKernelAttributeType>(AttributeType))
+			{
+#if PCG_LOGGING_ENABLED
+				if (OutErrorText)
+				{
+					const FString ActualTypeStr = AttributeTypeEnum->GetNameStringByIndex(static_cast<int64>(AttrDesc->Type));
+
+					*OutErrorText = FText::Format(
+						LOCTEXT("AttributeTypeMismatch", "Type mismatch for call to attribute function '{0}' on attribute '{1}'. Expected '{2}' but received '{3}'."),
+						ConstructFunctionText(),
+						FText::FromName(AttrName),
+						FText::FromString(TypeStr),
+						FText::FromString(ActualTypeStr));
 				}
+#endif
 
 				return false;
 			}
 		}
 	}
+
 	return true;
 }
 
@@ -1055,18 +1330,24 @@ FString UPCGCustomHLSLSettings::GetCookedKernelSource(const TMap<FPCGKernelAttri
 	);
 
 	// Header writers initialize PCG data collection format headers in output buffers.
-	// TODO: Currently done from all threads. Only write header from global index 0? Or write sections from data local index 0?
 	FString HeaderWriters;
 
+	auto EmitHeaderWriterSingleData = [&HeaderWriters](const FPCGPinProperties& InOutputPinProps)
+	{
+		HeaderWriters += FString::Format(TEXT(
+			"    // Signal kernel executed by setting data count from first thread. Rest of header was already set up by the CPU.\n"
+			"    if (GroupIndex == 0) {0}_WriteNumData(1);\n"
+			"    AllMemoryBarrier();\n"),
+			{ InOutputPinProps.Label.ToString() });
+	};
+
 	auto EmitHeaderWriterFromInputPin = [&HeaderWriters](const FPCGPinProperties& InOutputPinProps, const UPCGPin* InFromPin)
 	{
 		HeaderWriters += FString::Format(TEXT(
-			"    if (ThreadIndex >= GetNumThreads().x) return;\n"
-			"    \n"
 			"    // Signal kernel executed by copying data count from pin {0} to pin {1} from first thread. Rest of header was already set up by the CPU.\n"
 			"    if (GroupIndex == 0) {1}_WriteNumData({0}_GetNumData());\n"
 			"    AllMemoryBarrier();\n"),
-			{ InFromPin->Properties.Label.ToString(), InOutputPinProps.Label.ToString(), PCGComputeConstants::MAX_NUM_ATTRS });
+			{ InFromPin->Properties.Label.ToString(), InOutputPinProps.Label.ToString() });
 	};
 
 	if (KernelType == EPCGKernelType::PointProcessor || KernelType == EPCGKernelType::Custom)
@@ -1090,15 +1371,13 @@ FString UPCGCustomHLSLSettings::GetCookedKernelSource(const TMap<FPCGKernelAttri
 			}
 			else if (PinProps.BufferSizeMode == EPCGPinBufferSizeMode::FixedElementCount)
 			{
-				HeaderWriters += FString::Format(TEXT(
-					"    {0}_WriteSinglePointDataCollectionHeader({1});"
-					"    AllMemoryBarrier();"),
-					{ PinProps.Label.ToString(), PinProps.FixedBufferElementCount });
-
+				// FixedElementCount always produces a single point data.
+				EmitHeaderWriterSingleData(PinProps);
 			}
 			else if (PinProps.BufferSizeMode == EPCGPinBufferSizeMode::FromProductOfInputPins)
 			{
-				HeaderWriters += TEXT("    // The output data header must be written from the kernel code, for example by using WriteSinglePointDataCollectionHeader.");
+				// TODO: FromProductOfInputPins always produces a single point data for now, make it more flexible?
+				EmitHeaderWriterSingleData(PinProps);
 			}
 		}
 	}
@@ -1109,17 +1388,6 @@ FString UPCGCustomHLSLSettings::GetCookedKernelSource(const TMap<FPCGKernelAttri
 			const UPCGPin* FirstPin = Node->GetPassThroughInputPin();
 			const UPCGPin* PrimaryOutputPin = GetFirstPointOutputPin();
 			
-			auto EmitPointGenHeader = [&HeaderWriters, InPointCount = PointCount](const FPCGPinProperties& InOutputPinProps)
-			{
-				HeaderWriters += FString::Format(TEXT(
-					"    if (ThreadIndex >= GetNumThreads().x) return;\n"
-					"    \n"
-					"    // Signal kernel executed by writing data count (1) for pin {0} from first thread. Rest of header was already set up by the CPU.\n"
-					"    if (GroupIndex == 0) {0}_WriteNumData(1);\n"
-					"    AllMemoryBarrier();\n"),
-					{ InOutputPinProps.Label.ToString() });
-			};
-
 			for (const UPCGPin* OutputPin : Node->GetOutputPins())
 			{
 				if (!OutputPin || OutputPin->Properties.AllowedTypes != EPCGDataType::Point)
@@ -1138,7 +1406,7 @@ FString UPCGCustomHLSLSettings::GetCookedKernelSource(const TMap<FPCGKernelAttri
 
 				if (OutputPin == PrimaryOutputPin)
 				{
-					EmitPointGenHeader(PinProps);
+					EmitHeaderWriterSingleData(PinProps);
 				}
 				else if (PinPropsGPU->BufferSizeMode == EPCGPinBufferSizeMode::FromFirstPin && FirstPin)
 				{
@@ -1256,7 +1524,8 @@ FString UPCGCustomHLSLSettings::GetCookedKernelSource(const TMap<FPCGKernelAttri
 			"%s\n\n" // Functions
 			"%s\n" // KernelFunc
 			"{\n"
-			"    const uint ThreadIndex = %s;\n" // UnWrappedDispatchThreadId
+			"	const uint ThreadIndex = %s;\n" // UnWrappedDispatchThreadId
+			"	if (ThreadIndex >= GetNumThreads().x) return;\n"
 			"%s\n" // HeaderWriters
 			"%s\n" // KernelSpecificPreamble
 			"#line 0 \"%s\"\n\n" // ShaderPathName
@@ -1286,4 +1555,5 @@ bool FPCGCustomHLSLElement::ExecuteInternal(FPCGContext* Context) const
 	return true;
 }
 
+#undef PCG_LOGGING_ENABLED
 #undef LOCTEXT_NAMESPACE
diff --git a/Engine/Plugins/PCG/Source/PCG/Private/Compute/PCGDataForGPU.cpp b/Engine/Plugins/PCG/Source/PCG/Private/Compute/PCGDataForGPU.cpp
--- a/Engine/Plugins/PCG/Source/PCG/Private/Compute/PCGDataForGPU.cpp
+++ b/Engine/Plugins/PCG/Source/PCG/Private/Compute/PCGDataForGPU.cpp
@@ -385,11 +385,21 @@ namespace PCGDataForGPUHelpers
 	}
 }
 
+bool FPCGKernelAttributeKey::operator==(const FPCGKernelAttributeKey& Other) const
+{
+	return Type == Other.Type && Name == Other.Name;
+}
+
 uint32 GetTypeHash(const FPCGKernelAttributeKey& In)
 {
 	return HashCombine(GetTypeHash(In.Type), GetTypeHash(In.Name));
 }
 
+bool FPCGKernelAttributeDesc::operator==(const FPCGKernelAttributeDesc& Other) const
+{
+	return Index == Other.Index && Type == Other.Type && Name == Other.Name;
+}
+
 FPCGDataDesc::FPCGDataDesc(EPCGDataType InType, int InElementCount)
 	: Type(InType)
 	, ElementCount(InElementCount)
@@ -760,7 +770,7 @@ void FPCGDataCollectionDesc::PrepareBufferForKernelOutput(TArray<uint32>& OutPac
 	}
 }
 
-void FPCGDataCollectionDesc::UnpackDataCollection(const TArray<uint8>& InPackedData, FName InPin, FPCGDataCollection& OutDataCollection) const
+EPCGUnpackDataCollectionResult FPCGDataCollectionDesc::UnpackDataCollection(const TArray<uint8>& InPackedData, FName InPin, FPCGDataCollection& OutDataCollection) const
 {
 	const void* PackedData = InPackedData.GetData();
 	const float* DataAsFloat = static_cast<const float*>(PackedData);
@@ -770,9 +780,9 @@ void FPCGDataCollectionDesc::UnpackDataCollection(const TArray<uint8>& InPackedD
 	const uint32 NumPackedFloats = InPackedData.Num() / 4;
 	const uint32 NumData = DataAsUint[0];
 
-	if (!ensureAlwaysMsgf(NumData == DataDescs.Num(), TEXT("Mismatch in expected data from static analysis and actual received data.")))
+	if (NumData != DataDescs.Num())
 	{
-		return;
+		return EPCGUnpackDataCollectionResult::DataMismatch;
 	}
 
 	TArray<FPCGTaggedData>& OutData = OutDataCollection.TaggedData;
@@ -968,6 +978,8 @@ void FPCGDataCollectionDesc::UnpackDataCollection(const TArray<uint8>& InPackedD
 		}
 		else { /* TODO: Support non-point data. */ }
 	}
+
+	return EPCGUnpackDataCollectionResult::Success;
 }
 
 uint32 FPCGDataCollectionDesc::ComputeDataElementCount(EPCGDataType InDataType) const
diff --git a/Engine/Plugins/PCG/Source/PCG/Public/Compute/Elements/PCGCustomHLSL.h b/Engine/Plugins/PCG/Source/PCG/Public/Compute/Elements/PCGCustomHLSL.h
--- a/Engine/Plugins/PCG/Source/PCG/Public/Compute/Elements/PCGCustomHLSL.h
+++ b/Engine/Plugins/PCG/Source/PCG/Public/Compute/Elements/PCGCustomHLSL.h
@@ -46,6 +46,10 @@ public:
 	UPROPERTY(Transient)
 	bool bDisplayBufferSizeSettings = true;
 #endif // WITH_EDITORONLY_DATA
+
+	/** Add entries to create new attributes on data emitted by this pin. */
+	UPROPERTY(EditAnywhere, DisplayName = "Attributes to Create", Category = Settings)
+	TArray<FPCGKernelAttributeKey> CreatedKernelAttributeKeys;
 };
 
 /** Type of kernel allows us to make decisions about execution automatically, streamlining authoring. */
@@ -94,7 +98,6 @@ public:
 	virtual EPCGSettingsType GetType() const override { return EPCGSettingsType::GPU; }
 #endif
 
-	virtual TArray<FPCGKernelAttributeKey> GetKernelAttributeKeys() const override;
 	virtual FPCGDataCollectionDesc ComputeOutputPinDataDesc(const UPCGPin* OutputPin, const UPCGDataBinding* Binding) const override;
 	virtual int ComputeKernelThreadCount(const UPCGDataBinding* Binding) const override;
 
@@ -113,6 +116,7 @@ protected:
 #if WITH_EDITOR
 	void UpdateDeclarations();
 	void UpdatePinSettings();
+	void UpdateAttributeKeys();
 #endif
 
 public:
@@ -136,6 +140,9 @@ public:
 	int GetProcessingElemCountForInputPin(const UPCGPin* InputPin, const UPCGDataBinding* Binding) const;
 	virtual const UPCGPin* GetExecutionPin() const { return GetPointProcessingInputPin(); }
 
+protected:
+	bool AreKernelAttributesValid(FPCGContext* InContext, FText* OutErrorText) const;
+
 protected:
 	UPROPERTY(BlueprintReadWrite, EditAnywhere, Category = "Settings")
 	EPCGKernelType KernelType = EPCGKernelType::PointProcessor;
diff --git a/Engine/Plugins/PCG/Source/PCG/Public/Compute/PCGDataForGPU.h b/Engine/Plugins/PCG/Source/PCG/Public/Compute/PCGDataForGPU.h
--- a/Engine/Plugins/PCG/Source/PCG/Public/Compute/PCGDataForGPU.h
+++ b/Engine/Plugins/PCG/Source/PCG/Public/Compute/PCGDataForGPU.h
@@ -11,6 +11,12 @@ class UPCGMetadata;
 class UPCGPin;
 class UPCGSettings;
 
+enum class EPCGUnpackDataCollectionResult
+{
+	Success,
+	DataMismatch
+};
+
 UENUM()
 enum class EPCGKernelAttributeType : uint8
 {
@@ -30,13 +36,13 @@ struct FPCGKernelAttributeKey
 {
 	GENERATED_BODY()
 
-	UPROPERTY(VisibleAnywhere, Category = "Settings")
+	UPROPERTY(EditAnywhere, Category = "Settings")
 	EPCGKernelAttributeType Type = EPCGKernelAttributeType::Float;
 
-	UPROPERTY(VisibleAnywhere, Category = "Settings")
+	UPROPERTY(EditAnywhere, Category = "Settings")
 	FName Name = NAME_None;
 
-	bool operator==(const FPCGKernelAttributeKey& Other) const { return Type == Other.Type && Name == Other.Name; }
+	bool operator==(const FPCGKernelAttributeKey& Other) const;
 	friend uint32 GetTypeHash(const FPCGKernelAttributeKey& In);
 };
 
@@ -52,6 +58,8 @@ struct FPCGKernelAttributeDesc
 	int32 Index = INDEX_NONE;
 	EPCGKernelAttributeType Type = EPCGKernelAttributeType::Float;
 	FName Name = NAME_None;
+
+	bool operator==(const FPCGKernelAttributeDesc& Other) const;
 };
 
 struct FPCGDataDesc
@@ -86,7 +94,7 @@ struct FPCGDataCollectionDesc
 	void PrepareBufferForKernelOutput(TArray<uint32>& OutPackedDataCollection);
 
 	/** Unpack a buffer of 8-bit uints to a data collection. */
-	void UnpackDataCollection(const TArray<uint8>& InPackedData, FName InPin, FPCGDataCollection& OutDataCollection) const;
+	EPCGUnpackDataCollectionResult UnpackDataCollection(const TArray<uint8>& InPackedData, FName InPin, FPCGDataCollection& OutDataCollection) const;
 
 	/** Compute total number of processing elements of the given type. */
 	uint32 ComputeDataElementCount(EPCGDataType InDataType) const;
diff --git a/Engine/Plugins/PCG/Source/PCG/Public/PCGContext.h b/Engine/Plugins/PCG/Source/PCG/Public/PCGContext.h
--- a/Engine/Plugins/PCG/Source/PCG/Public/PCGContext.h
+++ b/Engine/Plugins/PCG/Source/PCG/Public/PCGContext.h
@@ -119,6 +119,9 @@ struct PCG_API FPCGContext
 	/** Time slicing is not enabled by default. */
 	virtual bool TimeSliceIsEnabled() const { return false; }
 
+	/** Is this a context for the compute graph element. */
+	virtual bool IsComputeContext() const { return false; }
+
 #if WITH_EDITOR
 	/** Log warnings and errors to be displayed on node in graph editor. */
 	void LogVisual(ELogVerbosity::Type InVerbosity, const FText& InMessage) const;
diff --git a/Engine/Plugins/PCG/Source/PCG/Public/PCGSettings.h b/Engine/Plugins/PCG/Source/PCG/Public/PCGSettings.h
--- a/Engine/Plugins/PCG/Source/PCG/Public/PCGSettings.h
+++ b/Engine/Plugins/PCG/Source/PCG/Public/PCGSettings.h
@@ -406,7 +406,7 @@ public:
 	const UPCGSettings* OriginalSettings = nullptr;
 
 	/** Get a list of the attributes read or written by this node. */
-	virtual TArray<FPCGKernelAttributeKey> GetKernelAttributeKeys() const { return {}; }
+	virtual const TArray<FPCGKernelAttributeKey>& GetKernelAttributeKeys() const { return KernelAttributeKeys; }
 
 	/** Compute a description of all data arriving on InputPin. */
 	FPCGDataCollectionDesc ComputeInputPinDataDesc(const UPCGPin* InputPin, const UPCGDataBinding* Binding) const;
@@ -477,6 +477,12 @@ protected:
 	TArray<FPCGPinProperties> DefaultPointInputPinProperties() const;
 	TArray<FPCGPinProperties> DefaultPointOutputPinProperties() const;
 
+	/** Attributes statically detected as being read, written, or created by this node. */
+	TArray<FPCGKernelAttributeKey> KernelAttributeKeys;
+
+	/** Maps pins to their attribute keys and whether or not they were created on the GPU. */
+	TMap<FName, TArray<TTuple<FPCGKernelAttributeKey, bool /*bCreatedOnGPU*/>>> PinToAttributeKeys;
+
 public:
 #if WITH_EDITORONLY_DATA
 	/** The version number of the data after load and after any data migration. */
