[PCG] Allow for N:1 filtering with point data and make spatial query false by default for new objects

#rnx

diff --git a/Engine/Plugins/PCG/Source/PCG/Private/Elements/PCGAttributeFilter.cpp b/Engine/Plugins/PCG/Source/PCG/Private/Elements/PCGAttributeFilter.cpp
--- a/Engine/Plugins/PCG/Source/PCG/Private/Elements/PCGAttributeFilter.cpp
+++ b/Engine/Plugins/PCG/Source/PCG/Private/Elements/PCGAttributeFilter.cpp
@@ -8,6 +8,7 @@
 #include "PCGPin.h"
 #include "Data/PCGPointData.h"
 #include "Data/PCGSpatialData.h"
+#include "Helpers/PCGHelpers.h"
 #include "Metadata/Accessors/IPCGAttributeAccessorTpl.h"
 #include "Metadata/Accessors/PCGAttributeAccessorHelpers.h"
 #include "Metadata/Accessors/PCGCustomAccessor.h"
@@ -228,10 +229,10 @@ namespace PCGAttributeFilterHelpers
 			return false;
 		}
 
-		// Check that if we have points as threshold, that the point data has the same number of point that the input data
+		// Check that if we have points as threshold, that the point data has the same number of point that the input data, or there is just a single point
 		if (InOutThresholdInfo.ThresholdSpatialData != nullptr && !InOutThresholdInfo.bUseSpatialQuery)
 		{
-			if (InOutThresholdInfo.ThresholdKeys->GetNum() != NumInput)
+			if (InOutThresholdInfo.ThresholdKeys->GetNum() != NumInput && InOutThresholdInfo.ThresholdKeys->GetNum() != 1)
 			{
 				PCGE_LOG_C(Warning, GraphAndLog, InContext, FText::Format(LOCTEXT("InvalidNumberOfThresholdPoints", "Threshold point data doesn't have the same number of elements ({0}) than the input data ({1})."), InOutThresholdInfo.ThresholdKeys->GetNum(), NumInput));
 				return false;
@@ -332,6 +333,12 @@ UPCGAttributeFilteringSettings::UPCGAttributeFilteringSettings()
 	// Recreate the same default
 	TargetAttribute.SetPointProperty(EPCGPointProperties::Density);
 	ThresholdAttribute.SetPointProperty(EPCGPointProperties::Density);
+
+	// Change the default for spatial query to be false
+	if (PCGHelpers::IsNewObjectAndNotDefault(this))
+	{
+		bUseSpatialQuery = false;
+	}
 }
 
 #if WITH_EDITOR
@@ -439,6 +446,13 @@ UPCGAttributeFilteringRangeSettings::UPCGAttributeFilteringRangeSettings()
 	TargetAttribute.SetPointProperty(EPCGPointProperties::Density);
 	MinThreshold.ThresholdAttribute.SetPointProperty(EPCGPointProperties::Density);
 	MaxThreshold.ThresholdAttribute.SetPointProperty(EPCGPointProperties::Density);
+
+	// Change the default for spatial query to be false
+	if (PCGHelpers::IsNewObjectAndNotDefault(this))
+	{
+		MinThreshold.bUseSpatialQuery = false;
+		MaxThreshold.bUseSpatialQuery = false;
+	}
 }
 
 
diff --git a/Engine/Plugins/PCG/Source/PCG/Public/Elements/PCGAttributeFilter.h b/Engine/Plugins/PCG/Source/PCG/Public/Elements/PCGAttributeFilter.h
--- a/Engine/Plugins/PCG/Source/PCG/Public/Elements/PCGAttributeFilter.h
+++ b/Engine/Plugins/PCG/Source/PCG/Public/Elements/PCGAttributeFilter.h
@@ -41,6 +41,7 @@ struct FPCGAttributeFilterThresholdSettings
 	UPROPERTY(BlueprintReadWrite, EditAnywhere, Category = Settings, meta = (EditCondition = "!bUseConstantThreshold", EditConditionHides, PCG_Overridable))
 	FPCGAttributePropertyInputSelector ThresholdAttribute;
 
+	// This value is now false by default (changed in the ctor of the filtering settings)
 	/** For Point Data, enabling this option will use sampling rather than comparing points 1 to 1 directly. For other spatial data, this is always true, and for attribute sets, always false. */
 	UPROPERTY(BlueprintReadWrite, EditAnywhere, Category = Settings, meta = (EditCondition = "!bUseConstantThreshold", EditConditionHides, PCG_Overridable))
 	bool bUseSpatialQuery = true;
@@ -110,6 +111,7 @@ public:
 	UPROPERTY(BlueprintReadWrite, EditAnywhere, Category = Settings, meta = (EditCondition = "!bUseConstantThreshold", EditConditionHides, PCG_Overridable))
 	FPCGAttributePropertyInputSelector ThresholdAttribute;
 
+	// This value is now false by default (changed in the ctor)
 	/** If the threshold data is Point data, it will sample input points in threshold data. Always true with Spatial data.*/
 	UPROPERTY(BlueprintReadWrite, EditAnywhere, Category = Settings, meta = (EditCondition = "!bUseConstantThreshold", EditConditionHides, PCG_Overridable))
 	bool bUseSpatialQuery = true;
