[PCG] Fixed a bug in the Reverse Spline node where it was neglecting to reverse interp curve modes

#rnx

diff --git a/Engine/Plugins/PCG/Source/PCG/Private/Elements/PCGReverseSpline.cpp b/Engine/Plugins/PCG/Source/PCG/Private/Elements/PCGReverseSpline.cpp
--- a/Engine/Plugins/PCG/Source/PCG/Private/Elements/PCGReverseSpline.cpp
+++ b/Engine/Plugins/PCG/Source/PCG/Private/Elements/PCGReverseSpline.cpp
@@ -48,14 +48,20 @@ namespace PCGReverseSpline
 
 		for (int i = ControlPointsPosition.Points.Num() - 1; i >= 0; --i)
 		{
-			// Tangents are inverted and swapped.
+			/* Implementation Note: Segment interpolation is determined by the interpolation mode of the preceding
+			 * control point. When inverting order of control points, we can decay the interpolation mode by setting all
+			 * modes to Custom Tangent to use the pre-calculated tangents as-is, so long as they were actually
+			 * calculated. This is a slightly destructive process and some information will be lost. Also, its worth
+			 * noting that since each spline segment is calculated from [0..1) there is a slight inconsistency when
+			 * evaluated in reverse order, as effectively the reverse is [1..0) per segment.
+			 */
 			NewControlPoints.Emplace(static_cast<float>(NewControlPoints.Num()),
 				ControlPointsPosition.Points[i].OutVal,
-				-ControlPointsPosition.Points[i].LeaveTangent,
+				-ControlPointsPosition.Points[i].LeaveTangent, // Tangents are inverted and swapped
 				-ControlPointsPosition.Points[i].ArriveTangent,
 				ControlPointsRotation.Points[i].OutVal.Rotator(),
 				ControlPointsScale.Points[i].OutVal,
-				ConvertInterpCurveModeToSplinePointType(ControlPointsPosition.Points[i].InterpMode));
+				ESplinePointType::CurveCustomTangent);
 		}
 
 		UPCGSplineData* NewSplineData = FPCGContext::NewObject_AnyThread<UPCGSplineData>(Context);
@@ -91,7 +97,7 @@ TArray<FPCGPinProperties> UPCGReverseSplineSettings::InputPinProperties() const
 }
 
 TArray<FPCGPinProperties> UPCGReverseSplineSettings::OutputPinProperties() const
-{ 
+{
 	TArray<FPCGPinProperties> Properties;
 	Properties.Emplace(PCGPinConstants::DefaultOutputLabel, EPCGDataType::Spline);
 	return Properties;
@@ -111,7 +117,7 @@ bool FPCGReverseSplineElement::ExecuteInternal(FPCGContext* InContext) const
 		FPCGTaggedData& Output = InContext->OutputData.TaggedData.Emplace_GetRef(InputData);
 
 		const UPCGSplineData* InputSplineData = Cast<const UPCGSplineData>(InputData.Data);
-		if (!InputSplineData)
+		if (!InputSplineData || InputSplineData->SplineStruct.GetNumberOfSplineSegments() < 1)
 		{
 			continue;
 		}
