[PCG] Adjusted Spline to Segment node to account for open splines as well

#rnx

diff --git a/Engine/Plugins/PCG/Source/PCG/Private/Elements/Grammar/PCGSplineToSegment.cpp b/Engine/Plugins/PCG/Source/PCG/Private/Elements/Grammar/PCGSplineToSegment.cpp
--- a/Engine/Plugins/PCG/Source/PCG/Private/Elements/Grammar/PCGSplineToSegment.cpp
+++ b/Engine/Plugins/PCG/Source/PCG/Private/Elements/Grammar/PCGSplineToSegment.cpp
@@ -91,14 +91,11 @@ bool FPCGSplineToSegmentElement::ExecuteInternal(FPCGContext* InContext) const
 			continue;
 		}
 
-		const bool bIsClosed = InputSplineData->SplineStruct.IsClosedLoop();
-
 		UPCGPointData* OutputPointData = FPCGContext::NewObject_AnyThread<UPCGPointData>(InContext);
 		OutputPointData->InitializeFromData(InputSplineData);
 
-		const TArray<FInterpCurvePointVector>& ControlPoints = InputSplineData->SplineStruct.SplineCurves.Position.Points;
-		const int32 ControlPointsNum = ControlPoints.Num();
-		if (ControlPointsNum < 2)
+		const int32 NumSegments = InputSplineData->GetNumSegments();
+		if (NumSegments < 1)
 		{
 			// If we have a malformed spline, just make an empty point data
 			Outputs.Emplace_GetRef(Input).Data = OutputPointData;
@@ -143,7 +140,7 @@ bool FPCGSplineToSegmentElement::ExecuteInternal(FPCGContext* InContext) const
 
 		if (Settings->bExtractConnectivityInfo)
 		{
-			PreviousIndexAttribute = PCGSplineToSegment::CreateAttributeWithLog<int32>(PCGSplineToSegment::PreviousIndexAttributeName, false, OutputPointData->Metadata, InContext);
+			PreviousIndexAttribute = PCGSplineToSegment::CreateAttributeWithLog<int32>(PCGSplineToSegment::PreviousIndexAttributeName, false, OutputPointData->Metadata, InContext, INDEX_NONE);
 			NextIndexAttribute = PCGSplineToSegment::CreateAttributeWithLog<int32>(PCGSplineToSegment::NextIndexAttributeName, false, OutputPointData->Metadata, InContext);
 			SegmentIndexAttribute = PCGSplineToSegment::CreateAttributeWithLog<int32>(PCGSplineToSegment::SegmentIndexAttributeName, false, OutputPointData->Metadata, InContext);
 
@@ -156,18 +153,19 @@ bool FPCGSplineToSegmentElement::ExecuteInternal(FPCGContext* InContext) const
 		}
 
 		TArray<FPCGPoint>& OutputPoints = OutputPointData->GetMutablePoints();
-		OutputPoints.Reserve(bIsClosed ? ControlPointsNum : ControlPointsNum - 1);
+		OutputPoints.Reserve(NumSegments);
 
 		TArray<FBox> Segments;
-		Segments.Reserve(ControlPointsNum);
+		Segments.Reserve(NumSegments);
 		double CumulativeAngle = 0.0;
 		double FirstAngle = 0.0;
+		const bool bIsClosed = InputSplineData->SplineStruct.IsClosedLoop();
 
-		for (int32 Index = 0; Index < ControlPointsNum; ++Index)
+		for (int32 Index = 0; Index < NumSegments; ++Index)
 		{
-			const int32 NextIndex = (Index + 1) % ControlPointsNum;
 			const FVector ControlPoint = InputSplineData->SplineStruct.GetLocationAtSplineInputKey(Index, ESplineCoordinateSpace::World);
-			const FVector NextControlPoint = InputSplineData->SplineStruct.GetLocationAtSplineInputKey(NextIndex, ESplineCoordinateSpace::World);
+			// Note: We're iterating on segments and it is safe to access the final control point (ie. Index + 1)
+			const FVector NextControlPoint = InputSplineData->SplineStruct.GetLocationAtSplineInputKey(Index + 1, ESplineCoordinateSpace::World);
 			const FVector ControlPointZUp = InputSplineData->SplineStruct.GetQuaternionAtSplineInputKey(Index, ESplineCoordinateSpace::World).GetAxisZ();
 
 			FVector ArriveTangent, LeaveTangent;
@@ -196,7 +194,7 @@ bool FPCGSplineToSegmentElement::ExecuteInternal(FPCGContext* InContext) const
 
 			auto SetMetadata = [&](const int32 Index, const FVector& Tangent, const double Angle)
 			{
-				const int32 PreviousIndex = (Index + ControlPointsNum - 1) % ControlPointsNum;
+				const int32 PreviousIndex = (Index + NumSegments - 1) % NumSegments;
 				const PCGMetadataEntryKey PreviousPointKey = OutputPoints[PreviousIndex].MetadataEntry;
 				const PCGMetadataEntryKey CurrentPointKey = OutputPoints[Index].MetadataEntry;
 
@@ -235,13 +233,15 @@ bool FPCGSplineToSegmentElement::ExecuteInternal(FPCGContext* InContext) const
 				FirstAngle = Angle;
 			}
 
-			if (NextIndex == 0 && bHasMetadata)
+
+			// For the last point, if it's a closed spline, set the arrive tangent and angle for the first point
+			if (bIsClosed && bHasMetadata && Index == (NumSegments - 1))
 			{
-				SetMetadata(NextIndex, LeaveTangent, FirstAngle);
+				SetMetadata(0, LeaveTangent, FirstAngle);
 			}
 		}
 
-		if (Settings->bExtractClockwiseInfo && bIsClosed)
+		if (Settings->bExtractClockwiseInfo)
 		{
 			PCGSplineToSegment::CreateAttributeWithLog<bool>(PCGSplineToSegment::ClockwiseAttributeName, false, OutputPointData->Metadata, InContext, CumulativeAngle <= 0.0);
 		}
