[PCG] Added a new Spline Slicer (V1) node for use with Grammars

diff --git a/Engine/Plugins/PCG/Source/PCG/Private/Data/PCGLandscapeSplineData.cpp b/Engine/Plugins/PCG/Source/PCG/Private/Data/PCGLandscapeSplineData.cpp
--- a/Engine/Plugins/PCG/Source/PCG/Private/Data/PCGLandscapeSplineData.cpp
+++ b/Engine/Plugins/PCG/Source/PCG/Private/Data/PCGLandscapeSplineData.cpp
@@ -246,6 +246,20 @@ FVector::FReal UPCGLandscapeSplineData::GetDistanceAtSegmentStart(int SegmentInd
 	return Distance;
 }
 
+FVector UPCGLandscapeSplineData::GetLocationAtAlpha(float Alpha) const
+{
+	const float InputKey = GetInputKeyAtAlpha(Alpha);
+	const int SegmentIndex = FMath::TruncToInt(InputKey);
+	return GetLocationAtDistance(SegmentIndex, FMath::Frac(Alpha) * GetSegmentLength(SegmentIndex));
+}
+
+FTransform UPCGLandscapeSplineData::GetTransformAtAlpha(float Alpha) const
+{
+	const float InputKey = GetInputKeyAtAlpha(Alpha);
+	const int SegmentIndex = FMath::TruncToInt(InputKey);
+	return GetTransformAtDistance(SegmentIndex, FMath::Frac(Alpha) * GetSegmentLength(SegmentIndex));
+}
+
 const UPCGPointData* UPCGLandscapeSplineData::CreatePointData(FPCGContext* Context) const
 {
 	check(Spline.IsValid());
diff --git a/Engine/Plugins/PCG/Source/PCG/Private/Data/PCGPolyLineData.cpp b/Engine/Plugins/PCG/Source/PCG/Private/Data/PCGPolyLineData.cpp
--- a/Engine/Plugins/PCG/Source/PCG/Private/Data/PCGPolyLineData.cpp
+++ b/Engine/Plugins/PCG/Source/PCG/Private/Data/PCGPolyLineData.cpp
@@ -32,6 +32,12 @@ FVector::FReal UPCGPolyLineData::GetLength() const
 	return Length;
 }
 
+float UPCGPolyLineData::GetInputKeyAtAlpha(float Alpha) const
+{
+	Alpha = FMath::Clamp(Alpha, 0, 1);
+	return static_cast<float>(GetNumSegments()) * Alpha;
+}
+
 void UPCGPolyLineData::GetTangentsAtSegmentStart(int SegmentIndex, FVector& OutArriveTangent, FVector& OutLeaveTangent) const
 {
 	OutArriveTangent = FVector::Zero();
diff --git a/Engine/Plugins/PCG/Source/PCG/Private/Data/PCGSplineData.cpp b/Engine/Plugins/PCG/Source/PCG/Private/Data/PCGSplineData.cpp
--- a/Engine/Plugins/PCG/Source/PCG/Private/Data/PCGSplineData.cpp
+++ b/Engine/Plugins/PCG/Source/PCG/Private/Data/PCGSplineData.cpp
@@ -153,6 +153,16 @@ FVector::FReal UPCGSplineData::GetDistanceAtSegmentStart(int SegmentIndex) const
 	return SplineStruct.GetDistanceAlongSplineAtSplinePoint(SegmentIndex);
 }
 
+FVector UPCGSplineData::GetLocationAtAlpha(float Alpha) const
+{
+	return SplineStruct.GetLocationAtSplineInputKey(GetInputKeyAtAlpha(Alpha), ESplineCoordinateSpace::World);
+}
+
+FTransform UPCGSplineData::GetTransformAtAlpha(float Alpha) const
+{
+	return SplineStruct.GetTransformAtSplineInputKey(GetInputKeyAtAlpha(Alpha), ESplineCoordinateSpace::World);
+}
+
 const UPCGPointData* UPCGSplineData::CreatePointData(FPCGContext* Context) const
 {
 	TRACE_CPUPROFILER_EVENT_SCOPE(UPCGSplineData::CreatePointData);
diff --git a/Engine/Plugins/PCG/Source/PCG/Private/Elements/Grammar/PCGSegmentSlicer.cpp b/Engine/Plugins/PCG/Source/PCG/Private/Elements/Grammar/PCGSegmentSlicer.cpp
--- a/Engine/Plugins/PCG/Source/PCG/Private/Elements/Grammar/PCGSegmentSlicer.cpp
+++ b/Engine/Plugins/PCG/Source/PCG/Private/Elements/Grammar/PCGSegmentSlicer.cpp
@@ -293,7 +293,7 @@ bool FPCGSegmentSlicerElement::ExecuteInternal(FPCGContext* InContext) const
 				OutAttribute = OutputMetadata->FindOrCreateAttribute<T>(AttributeName, DefaultValue, false, true);
 				if (!OutAttribute)
 				{
-					PCGLog::LogWarningOnGraph(LOCTEXT("FailCreateAttributes", "Fail to create output attributes."), InContext);
+					PCGLog::Metadata::LogFailToCreateAttribute<T>(AttributeName, InContext);
 					return false;
 				}
 			}
diff --git a/Engine/Plugins/PCG/Source/PCG/Private/Elements/Grammar/PCGSelectGrammar.cpp b/Engine/Plugins/PCG/Source/PCG/Private/Elements/Grammar/PCGSelectGrammar.cpp
--- a/Engine/Plugins/PCG/Source/PCG/Private/Elements/Grammar/PCGSelectGrammar.cpp
+++ b/Engine/Plugins/PCG/Source/PCG/Private/Elements/Grammar/PCGSelectGrammar.cpp
@@ -231,14 +231,14 @@ bool FPCGSelectGrammarElement::ExecuteInternal(FPCGContext* InContext) const
 		TUniquePtr<const IPCGAttributeAccessorKeys> InputKeys = PCGAttributeAccessorHelpers::CreateConstKeys(InputPointData, KeySelector);
 		if (!InputKeyAccessor || !InputKeys)
 		{
-			PCGLog::Accessor::LogFailToCreate(KeySelector);
+			PCGLog::Metadata::LogFailToCreateAccessor(KeySelector);
 			continue;
 		}
 
 		TUniquePtr<const IPCGAttributeAccessor> ComparedValuesAccessor = PCGAttributeAccessorHelpers::CreateConstAccessor(InputPointData, ComparedValueSelector);
 		if (!ComparedValuesAccessor)
 		{
-			PCGLog::Accessor::LogFailToCreate(ComparedValueSelector);
+			PCGLog::Metadata::LogFailToCreateAccessor(ComparedValueSelector);
 			continue;
 		}
 
@@ -249,7 +249,7 @@ bool FPCGSelectGrammarElement::ExecuteInternal(FPCGContext* InContext) const
 		TUniquePtr<IPCGAttributeAccessorKeys> OutputKeys = PCGAttributeAccessorHelpers::CreateKeys(OutputPointData, OutputGrammarSelector);
 		if (!WriteGrammarAccessor || !OutputKeys)
 		{
-			PCGLog::Accessor::LogFailToCreate(OutputGrammarSelector);
+			PCGLog::Metadata::LogFailToCreateAccessor(OutputGrammarSelector);
 			continue;
 		}
 
diff --git a/Engine/Plugins/PCG/Source/PCG/Private/Elements/Grammar/PCGSlicingBase.cpp b/Engine/Plugins/PCG/Source/PCG/Private/Elements/Grammar/PCGSlicingBase.cpp
--- a/Engine/Plugins/PCG/Source/PCG/Private/Elements/Grammar/PCGSlicingBase.cpp
+++ b/Engine/Plugins/PCG/Source/PCG/Private/Elements/Grammar/PCGSlicingBase.cpp
@@ -128,13 +128,13 @@ PCGGrammar::FTokenizedGrammar FPCGSlicingBaseElement::GetTokenizedGrammar(FPCGCo
 		const TUniquePtr<const IPCGAttributeAccessor> Accessor = PCGAttributeAccessorHelpers::CreateConstAccessor(InputData, Selector);
 		if (!Accessor)
 		{
-			PCGLog::Accessor::LogFailToCreate(Selector, InContext);
+			PCGLog::Metadata::LogFailToCreateAccessor(Selector, InContext);
 			return {};
 		}
 
 		if (!Accessor->Get(Grammar, FPCGAttributeAccessorKeysEntries(PCGInvalidEntryKey), EPCGAttributeAccessorFlags::AllowBroadcastAndConstructible))
 		{
-			PCGLog::Accessor::LogFailToGet<FString>(Selector, Accessor.Get(), InContext);
+			PCGLog::Metadata::LogFailToGetAttribute<FString>(Selector, Accessor.Get(), InContext);
 			return {};
 		}
 	}
@@ -215,7 +215,7 @@ TMap<FString, PCGGrammar::FTokenizedGrammar> FPCGSlicingBaseElement::GetTokenize
 		const TUniquePtr<const IPCGAttributeAccessorKeys> Keys = PCGAttributeAccessorHelpers::CreateConstKeys(InputData, Selector);
 		if (!Accessor || !Keys)
 		{
-			PCGLog::Accessor::LogFailToCreate(Selector, InContext);
+			PCGLog::Metadata::LogFailToCreateAccessor(Selector, InContext);
 			return Result;
 		}
 
@@ -229,7 +229,7 @@ TMap<FString, PCGGrammar::FTokenizedGrammar> FPCGSlicingBaseElement::GetTokenize
 
 		if (!bSuccess)
 		{
-			PCGLog::Accessor::LogFailToGet<FString>(Selector, Accessor.Get(), InContext);
+			PCGLog::Metadata::LogFailToGetAttribute<FString>(Selector, Accessor.Get(), InContext);
 			return Result;
 		}
 	}
@@ -339,6 +339,12 @@ bool FPCGSlicingBaseElement::MatchAndSetAttributes(const TArray<FPCGTaggedData>&
 			for (const FPCGPoint& OutPoint : OutPointData->GetPoints())
 			{
 				const FName Symbol = OutSymbolAttribute->GetValueFromItemKey(OutPoint.MetadataEntry);
+
+				if (!SymbolToEntryKeyMapping.Contains(Symbol))
+				{
+					continue;
+				}
+
 				const PCGMetadataEntryKey InputEntryKey = SymbolToEntryKeyMapping[Symbol];
 
 				if (InputEntryKey == PCGInvalidEntryKey)
diff --git a/Engine/Plugins/PCG/Source/PCG/Private/Elements/Grammar/PCGSplineSlicer.cpp b/Engine/Plugins/PCG/Source/PCG/Private/Elements/Grammar/PCGSplineSlicer.cpp
new file mode 100644
--- /dev/null
+++ b/Engine/Plugins/PCG/Source/PCG/Private/Elements/Grammar/PCGSplineSlicer.cpp
@@ -0,0 +1,351 @@
+// Copyright Epic Games, Inc. All Rights Reserved.
+
+#include "Elements/Grammar/PCGSplineSlicer.h"
+
+#include "PCGContext.h"
+#include "Data/PCGPointData.h"
+#include "Data/PCGSplineData.h"
+#include "Elements/Metadata/PCGMetadataElementCommon.h"
+#include "Helpers/PCGHelpers.h"
+#include "Kismet/KismetMathLibrary.h"
+#include "Metadata/Accessors/PCGAttributeAccessorHelpers.h"
+
+#define LOCTEXT_NAMESPACE "PCGSplineSlicerElement"
+
+namespace PCGSplineSlicerHelpers
+{
+	struct FParameters
+	{
+		FPCGContext* Context = nullptr;
+
+		PCGSlicingBase::FPCGModulesInfoMap ModulesInfo;
+		TMap<FString, PCGGrammar::FTokenizedGrammar> CachedModules;
+		bool bAcceptIncompleteSlicing = false;
+		double ModuleHeight = 0.0;
+
+		const UPCGPolyLineData* PolyLineData = nullptr;
+		TArray<FPCGPoint>* OutPoints = nullptr;
+		UPCGMetadata* OutputMetadata = nullptr;
+
+		FPCGMetadataAttribute<FName>* SymbolAttribute = nullptr;
+		FPCGMetadataAttribute<FVector4>* DebugColorAttribute = nullptr;
+		FPCGMetadataAttribute<int32>* ModuleIndexAttribute = nullptr;
+		FPCGMetadataAttribute<bool>* IsFirstPointAttribute = nullptr;
+		FPCGMetadataAttribute<bool>* IsFinalPointAttribute = nullptr;
+	};
+
+	double GetLinearDistanceBetweenPolyLineAlphas(const UPCGPolyLineData* PolyLineData, double FirstAlpha, double SecondAlpha)
+	{
+		check(PolyLineData);
+
+		FirstAlpha = FMath::Clamp(FirstAlpha, 0, 1);
+		SecondAlpha = FMath::Clamp(SecondAlpha, 0, 1);
+		const FVector FirstPoint = PolyLineData->GetLocationAtAlpha(FirstAlpha);
+		const FVector SecondPoint = PolyLineData->GetLocationAtAlpha(SecondAlpha);
+		return FVector::Distance(FirstPoint, SecondPoint);
+	}
+
+	// Using a numerical method (s.a. bisection), determine the spline alpha from a segment length. Bisection can be slower, but guaranteed to converge.
+	static double FindRootAtLinearDistance_Bisection(const UPCGPolyLineData* PolyLineData, const double SegmentLength, const double StartingAlpha, const double Tolerance = 1)
+	{
+		check(PolyLineData);
+
+		// Bisect a number of times, before taking an estimate
+		static constexpr uint16 BisectionCountLimit = 64u;
+
+		double Low = StartingAlpha;
+		double High = 1.0;
+		double Estimate = 0.0;
+
+		for (uint16 BisectionCount = 0u; BisectionCount < BisectionCountLimit; ++BisectionCount)
+		{
+			Estimate = (Low + High) * 0.5;
+			// TODO: If performance is an issue, this can deal with Dist^2 instead to save the sqrt. But, it will skew the tolerance.
+			const double LinearDistance = GetLinearDistanceBetweenPolyLineAlphas(PolyLineData, StartingAlpha, Estimate);
+			check(LinearDistance >= 0);
+
+			if (FMath::IsNearlyEqual(LinearDistance, SegmentLength, Tolerance))
+			{
+				return Estimate;
+			}
+
+			if (LinearDistance < SegmentLength)
+			{
+				Low = Estimate;
+			}
+			else
+			{
+				High = Estimate;
+			}
+		}
+
+		// Couldn't find the root yet, so estimate
+		return (Low + High) * 0.5;
+	}
+
+	static void Process(FParameters& InOutParameters, const FString& InGrammar)
+	{
+		if (!InOutParameters.CachedModules.Contains(InGrammar))
+		{
+			double MinSize;
+			InOutParameters.CachedModules.Emplace(InGrammar, PCGSlicingBase::GetTokenizedGrammar(InOutParameters.Context, InGrammar, InOutParameters.ModulesInfo, MinSize));
+		}
+
+		const PCGGrammar::FTokenizedGrammar& CurrentTokenizedGrammar = InOutParameters.CachedModules[InGrammar];
+
+		if (CurrentTokenizedGrammar.IsEmpty())
+		{
+			return;
+		}
+
+		const double SplineLength = InOutParameters.PolyLineData->GetLength();
+
+		/* Implementation Note: Subdivided spline length will always be equal or greater than discretized linear length, depending on the curvature of the spline.
+		 * For extremely long or curvy splines, this can result in the tokenized grammar being cut short. Alternative subdivision solutions may need to be explored.
+		 */
+		TArray<PCGSlicingBase::TPCGSubDivModuleInstance<PCGGrammar::FTokenizedModule>> ModulesInstances;
+		double RemainingSubdivide;
+		if (!Subdivide(CurrentTokenizedGrammar, SplineLength, ModulesInstances, RemainingSubdivide, InOutParameters.Context))
+		{
+			return;
+		}
+
+		if (!InOutParameters.bAcceptIncompleteSlicing && !FMath::IsNearlyZero(RemainingSubdivide))
+		{
+			PCGLog::LogWarningOnGraph(LOCTEXT("FailSliceFullLength", "The spline has an incomplete slicing (grammar doesn't fit the whole segment)."), InOutParameters.Context);
+			return;
+		}
+
+		const bool bHasMetadata = InOutParameters.SymbolAttribute || InOutParameters.DebugColorAttribute || InOutParameters.ModuleIndexAttribute || InOutParameters.IsFirstPointAttribute || InOutParameters.IsFinalPointAttribute;
+
+		int32 ModuleIndexCounter = 0; // Keep track for the output attribute
+		double SplineAlpha = 0.0; // Will be incremented as we progress through the spline
+		for (int32 ModuleInstanceIndex = 0; ModuleInstanceIndex < ModulesInstances.Num(); ModuleInstanceIndex++)
+		{
+			const PCGSlicingBase::TPCGSubDivModuleInstance<PCGGrammar::FTokenizedModule>& ModuleInstance = ModulesInstances[ModuleInstanceIndex];
+			for (int32 SubmoduleInstanceIndex = 0; SubmoduleInstanceIndex < ModuleInstance.NumRepeat; ++SubmoduleInstanceIndex)
+			{
+				for (int32 SymbolIndex = 0; SymbolIndex < ModuleInstance.Module->Symbols.Num(); ++SymbolIndex)
+				{
+					const FName Symbol = ModuleInstance.Module->Symbols[SymbolIndex];
+					const FPCGSlicingSubmodule& SlicingSubmodule = InOutParameters.ModulesInfo[Symbol];
+					const double SubmoduleSize = SlicingSubmodule.Size;
+
+					// Move to the next segment of the spline
+					const FVector SegmentStartPoint = InOutParameters.PolyLineData->GetLocationAtAlpha(SplineAlpha);
+					const double PreviousAlpha = SplineAlpha;
+
+					// Use a numerical method to find the root of where the module lands on the spline
+					SplineAlpha = FindRootAtLinearDistance_Bisection(InOutParameters.PolyLineData, SubmoduleSize, SplineAlpha);
+
+					const FVector SegmentEndPoint = InOutParameters.PolyLineData->GetLocationAtAlpha(SplineAlpha);
+					const FVector SliceVector = SegmentEndPoint - SegmentStartPoint;
+					const FVector SliceDirection = SliceVector.GetSafeNormal();
+
+					// At the end of the spline, truncate an unfinished submodule and end the process
+					if (FMath::IsNearlyEqual(SplineAlpha, 1))
+					{
+						if (SliceVector.Length() < SubmoduleSize)
+						{
+							if (InOutParameters.IsFinalPointAttribute && !InOutParameters.OutPoints->IsEmpty())
+							{
+								InOutParameters.IsFinalPointAttribute->SetValue(InOutParameters.OutPoints->Last().MetadataEntry, true);
+							}
+
+							return;
+						}
+					}
+
+					// Since its discretized, we won't take the transform's position, but we'll use the up vector--to create the module rotation--and the scale
+					FTransform CenterPointTransform = InOutParameters.PolyLineData->GetTransformAtAlpha((SplineAlpha + PreviousAlpha) * 0.5);
+
+					const FVector Position = SegmentStartPoint + (SliceVector * 0.5) + FVector(0, 0, InOutParameters.ModuleHeight * 0.5);
+					const FRotator Rotation = FRotationMatrix::MakeFromXZ(SliceDirection, CenterPointTransform.GetRotation().GetUpVector()).Rotator();
+					const FVector Scale = FVector::OneVector + (SliceDirection * ModuleInstance.ExtraScales[SymbolIndex]);
+
+					FPCGPoint& OutPoint = InOutParameters.OutPoints->Emplace_GetRef(FTransform(Rotation, Position, Scale), /*InDensity=*/1, PCGHelpers::ComputeSeedFromPosition(Position));
+
+					const double HalfSubmoduleSize = SlicingSubmodule.Size * 0.5;
+					OutPoint.SetLocalBounds(FBox(FVector(-HalfSubmoduleSize, 0, 0), FVector(HalfSubmoduleSize, 1, InOutParameters.ModuleHeight)));
+
+					// Now, handle the metadata attributes
+					if (!bHasMetadata)
+					{
+						continue;
+					}
+
+					InOutParameters.OutputMetadata->InitializeOnSet(OutPoint.MetadataEntry);
+					if (InOutParameters.SymbolAttribute)
+					{
+						InOutParameters.SymbolAttribute->SetValue(OutPoint.MetadataEntry, Symbol);
+					}
+
+					if (InOutParameters.DebugColorAttribute)
+					{
+						InOutParameters.DebugColorAttribute->SetValue(OutPoint.MetadataEntry, FVector4(SlicingSubmodule.DebugColor, 1.0));
+					}
+
+					if (InOutParameters.ModuleIndexAttribute)
+					{
+						InOutParameters.ModuleIndexAttribute->SetValue(OutPoint.MetadataEntry, ModuleIndexCounter++);
+					}
+
+					const bool bIsFirstModule = (ModuleInstanceIndex == 0) && (SubmoduleInstanceIndex == 0) && (SymbolIndex == 0);
+					if (bIsFirstModule && InOutParameters.IsFirstPointAttribute)
+					{
+						InOutParameters.IsFirstPointAttribute->SetValue(OutPoint.MetadataEntry, true);
+					}
+
+					const bool bIsFinalModule = (ModuleInstanceIndex == ModulesInstances.Num() - 1) && (SubmoduleInstanceIndex == ModuleInstance.NumRepeat - 1) && (SymbolIndex == ModuleInstance.Module->Symbols.Num() - 1);
+					if (bIsFinalModule && InOutParameters.IsFinalPointAttribute)
+					{
+						InOutParameters.IsFinalPointAttribute->SetValue(OutPoint.MetadataEntry, true);
+					}
+				}
+			}
+		}
+	}
+}
+
+FPCGElementPtr UPCGSplineSlicerSettings::CreateElement() const
+{
+	return MakeShared<FPCGSplineSlicerElement>();
+}
+
+TArray<FPCGPinProperties> UPCGSplineSlicerSettings::InputPinProperties() const
+{
+	TArray<FPCGPinProperties> PinProperties;
+	FPCGPinProperties& InputPin = PinProperties.Emplace_GetRef(PCGPinConstants::DefaultInputLabel, EPCGDataType::PolyLine);
+	InputPin.SetRequiredPin();
+
+	if (bModuleInfoAsInput)
+	{
+		FPCGPinProperties& ModuleInfoPin = PinProperties.Emplace_GetRef(PCGSlicingBaseConstants::ModulesInfoPinLabel, EPCGDataType::Param);
+		ModuleInfoPin.SetRequiredPin();
+	}
+
+	return PinProperties;
+}
+
+TArray<FPCGPinProperties> UPCGSplineSlicerSettings::OutputPinProperties() const
+{
+	return Super::DefaultPointOutputPinProperties();
+}
+
+bool FPCGSplineSlicerElement::ExecuteInternal(FPCGContext* InContext) const
+{
+	TRACE_CPUPROFILER_EVENT_SCOPE(FPCGSplineSlicerElement::Execute);
+
+	const UPCGSplineSlicerSettings* Settings = InContext->GetInputSettings<UPCGSplineSlicerSettings>();
+	check(Settings);
+
+	const TArray<FPCGTaggedData> Inputs = InContext->InputData.GetInputsByPin(PCGPinConstants::DefaultInputLabel);
+	TArray<FPCGTaggedData>& Outputs = InContext->OutputData.TaggedData;
+
+	const UPCGParamData* ModuleInfoParamData = nullptr;
+	const FPCGModulesInfoMap ModulesInfo = GetModulesInfoMap(InContext, Settings, ModuleInfoParamData);
+
+	double ModuleHeight = Settings->bModuleHeightAsAttribute ? 0.0 : Settings->ModuleHeight;
+
+	PCGSplineSlicerHelpers::FParameters Parameters
+	{
+		.Context = InContext,
+		.ModulesInfo = ModulesInfo,
+		.bAcceptIncompleteSlicing = Settings->bAcceptIncompleteSlicing
+	};
+
+	for (const FPCGTaggedData& Input : Inputs)
+	{
+		const UPCGPolyLineData* InputPolyLineData = Cast<const UPCGPolyLineData>(Input.Data);
+		if (!InputPolyLineData)
+		{
+			continue;
+		}
+
+		if (Settings->bModuleHeightAsAttribute)
+		{
+			check(InputPolyLineData->Metadata);
+			const FPCGAttributePropertyInputSelector Selector = Settings->ModuleHeightAttribute.CopyAndFixLast(InputPolyLineData);
+			TUniquePtr<const IPCGAttributeAccessor> HeightAccessor = PCGAttributeAccessorHelpers::CreateConstAccessor(InputPolyLineData, Selector);
+
+			if (!HeightAccessor)
+			{
+				PCGLog::Metadata::LogFailToCreateAccessor(Selector, InContext);
+				continue;
+			}
+
+			if (!HeightAccessor->Get<double>(ModuleHeight, FPCGAttributeAccessorKeysEntries(PCGInvalidEntryKey), EPCGAttributeAccessorFlags::AllowBroadcastAndConstructible))
+			{
+				PCGLog::Metadata::LogFailToGetAttribute(Selector, InContext);
+				continue;
+			}
+		}
+
+		UPCGPointData* OutputPointData = FPCGContext::NewObject_AnyThread<UPCGPointData>(InContext);
+		OutputPointData->InitializeFromData(InputPolyLineData);
+		FPCGTaggedData& Output = Outputs.Emplace_GetRef(Input);
+		Output.Data = OutputPointData;
+
+		// Update 'per-input' parameters
+		Parameters.PolyLineData = InputPolyLineData;
+		Parameters.OutPoints = &OutputPointData->GetMutablePoints();
+		Parameters.OutputMetadata = OutputPointData->Metadata;
+		Parameters.ModuleHeight = ModuleHeight;
+
+		auto CreateAndValidateAttribute = [InContext, OutputMetadata = OutputPointData->Metadata]<typename T>(const FName AttributeName, const T DefaultValue, const bool bShouldCreate, FPCGMetadataAttribute<T>*& OutAttribute) -> bool
+		{
+			if (bShouldCreate)
+			{
+				OutAttribute = OutputMetadata->FindOrCreateAttribute<T>(AttributeName, DefaultValue, false, true);
+				if (!OutAttribute)
+				{
+					PCGLog::Metadata::LogFailToCreateAttribute<T>(AttributeName, InContext);
+					return false;
+				}
+			}
+
+			return true;
+		};
+
+		if (!CreateAndValidateAttribute(Settings->SymbolAttributeName, FName(NAME_None), true, Parameters.SymbolAttribute)
+			|| !CreateAndValidateAttribute(Settings->DebugColorAttributeName, FVector4::Zero(), Settings->bOutputDebugColorAttribute, Parameters.DebugColorAttribute)
+			|| !CreateAndValidateAttribute(Settings->ModuleIndexAttributeName, -1, Settings->bOutputModuleIndexAttribute, Parameters.ModuleIndexAttribute)
+			|| !CreateAndValidateAttribute(Settings->IsFirstAttributeName, false, Settings->bOutputExtremityAttributes, Parameters.IsFirstPointAttribute)
+			|| !CreateAndValidateAttribute(Settings->IsFinalAttributeName, false, Settings->bOutputExtremityAttributes, Parameters.IsFinalPointAttribute))
+		{
+			continue;
+		}
+
+		if (Settings->GrammarSelection.bGrammarAsAttribute)
+		{
+			const FPCGAttributePropertyInputSelector Selector = Settings->GrammarSelection.GrammarAttribute.CopyAndFixLast(InputPolyLineData);
+			const TUniquePtr<const IPCGAttributeAccessor> GrammarAccessor = PCGAttributeAccessorHelpers::CreateConstAccessor(InputPolyLineData, Selector);
+			const FPCGAttributeAccessorKeysEntries Keys(PCGInvalidEntryKey);
+			if (!GrammarAccessor)
+			{
+				PCGLog::Metadata::LogFailToCreateAccessor(Selector, InContext);
+				continue;
+			}
+
+			auto Process = [&Parameters](const FString& InGrammar, int32 Index) -> void
+			{
+				PCGSplineSlicerHelpers::Process(Parameters, InGrammar);
+			};
+
+			PCGMetadataElementCommon::ApplyOnAccessor<FString>(Keys, *GrammarAccessor, Process);
+		}
+		else
+		{
+			PCGSplineSlicerHelpers::Process(Parameters, Settings->GrammarSelection.GrammarString);
+		}
+	}
+
+	if (Settings->bForwardAttributesFromModulesInfo && ModuleInfoParamData)
+	{
+		MatchAndSetAttributes(Inputs, Outputs, ModuleInfoParamData, Settings);
+	}
+
+	return true;
+}
+
+#undef LOCTEXT_NAMESPACE
diff --git a/Engine/Plugins/PCG/Source/PCG/Private/Elements/Grammar/PCGVolumeSlicer.cpp b/Engine/Plugins/PCG/Source/PCG/Private/Elements/Grammar/PCGVolumeSlicer.cpp
--- a/Engine/Plugins/PCG/Source/PCG/Private/Elements/Grammar/PCGVolumeSlicer.cpp
+++ b/Engine/Plugins/PCG/Source/PCG/Private/Elements/Grammar/PCGVolumeSlicer.cpp
@@ -93,13 +93,13 @@ bool FPCGVolumeSlicerElement::ExecuteInternal(FPCGContext* InContext) const
 			const TUniquePtr<const IPCGAttributeAccessor> Accessor = PCGAttributeAccessorHelpers::CreateConstAccessor(InputSplineData, Selector);
 			if (!Accessor)
 			{
-				PCGLog::Accessor::LogFailToCreate(Selector, InContext);
+				PCGLog::Metadata::LogFailToCreateAccessor(Selector, InContext);
 				return false;
 			}
 
 			if (!Accessor->Get(OutValue, FPCGAttributeAccessorKeysEntries(PCGInvalidEntryKey), EPCGAttributeAccessorFlags::AllowBroadcastAndConstructible))
 			{
-				PCGLog::Accessor::LogFailToGet<T>(Selector, Accessor.Get(), InContext);
+				PCGLog::Metadata::LogFailToGetAttribute<T>(Selector, Accessor.Get(), InContext);
 				return false;
 			}
 
diff --git a/Engine/Plugins/PCG/Source/PCG/Private/Utils/PCGLogErrors.cpp b/Engine/Plugins/PCG/Source/PCG/Private/Utils/PCGLogErrors.cpp
--- a/Engine/Plugins/PCG/Source/PCG/Private/Utils/PCGLogErrors.cpp
+++ b/Engine/Plugins/PCG/Source/PCG/Private/Utils/PCGLogErrors.cpp
@@ -25,32 +25,33 @@ namespace PCGLog::InputOutput
 	}
 }
 
-namespace PCGLog::Accessor
+namespace PCGLog::Metadata
 {
 	namespace ErrorFormat
 	{
-		const FTextFormat CreateAccessorFailure = LOCTEXT("CreateAccessorFailure", "Attribute {0} was not found.");
+		const FTextFormat CreateAccessorFailure = LOCTEXT("CreateAccessorFailure", "Couldn't create accessor. Attribute '{0}' was not found.");
+		const FTextFormat CreateAttributeFailure = LOCTEXT("CreateTypedAttributeFailure", "Couldn't create attribute '{0}' of type '{1}'.");
 		const FTextFormat GetAttributeFailure = LOCTEXT("GetAttributeFailure", "Couldn't retrieve attribute '{0}' value. Expected type: {1}, Actual Type: {2}.");
 		const FTextFormat GetTypedAttributeFailure = LOCTEXT("GetAttributeFailure", "Couldn't retrieve attribute '{0}' value. Expected type: {1}, Actual Type: {2}.");
 		const FTextFormat GetTypedAttributeFailureNoAccessor = LOCTEXT("GetAttributeFailureNoAccessor", "Couldn't retrieve attribute '{0}' value of type: '{1}.");
 	}
 
-	void LogFailToCreate(const FPCGAttributePropertySelector& Selector, const FPCGContext* InContext)
+	void LogFailToCreateAccessor(const FPCGAttributePropertySelector& Selector, const FPCGContext* InContext)
 	{
 		LogErrorOnGraph(FText::Format(ErrorFormat::CreateAccessorFailure, Selector.GetDisplayText()), InContext);
 	}
 
-	void LogFailToGet(FName AttributeName, const FPCGContext* InContext)
+	void LogFailToGetAttribute(FText AttributeName, const FPCGContext* InContext)
 	{
-		LogFailToGet(FText::FromName(AttributeName), InContext);
+		LogErrorOnGraph(FText::Format(ErrorFormat::GetAttributeFailure, std::move(AttributeName)), InContext);
 	}
 
-	void LogFailToGet(FText AttributeName, const FPCGContext* InContext)
+	void LogFailToGetAttribute(FName AttributeName, const FPCGContext* InContext)
 	{
-		LogErrorOnGraph(FText::Format(ErrorFormat::GetAttributeFailure, std::move(AttributeName)), InContext);
+		LogFailToGetAttribute(FText::FromName(AttributeName), InContext);
 	}
 
-	void LogFailToGet(const FPCGAttributePropertySelector& Selector, const FPCGContext* InContext)
+	void LogFailToGetAttribute(const FPCGAttributePropertySelector& Selector, const FPCGContext* InContext)
 	{
 		LogErrorOnGraph(FText::Format(ErrorFormat::GetAttributeFailure, Selector.GetDisplayText()), InContext);
 	}
diff --git a/Engine/Plugins/PCG/Source/PCG/Public/Data/PCGLandscapeSplineData.h b/Engine/Plugins/PCG/Source/PCG/Public/Data/PCGLandscapeSplineData.h
--- a/Engine/Plugins/PCG/Source/PCG/Public/Data/PCGLandscapeSplineData.h
+++ b/Engine/Plugins/PCG/Source/PCG/Public/Data/PCGLandscapeSplineData.h
@@ -35,6 +35,8 @@ public:
 	virtual float GetInputKeyAtDistance(int SegmentIndex, FVector::FReal Distance) const override;
 	virtual void GetTangentsAtSegmentStart(int SegmentIndex, FVector& OutArriveTangent, FVector& OutLeaveTangent) const override;
 	virtual FVector::FReal GetDistanceAtSegmentStart(int SegmentIndex) const override;
+	virtual FVector GetLocationAtAlpha(float Alpha) const override;
+	virtual FTransform GetTransformAtAlpha(float Alpha) const override;
 	//~End UPCGPolyLineData interface
 
 	//~Begin UPCGSpatialDataWithPointCache interface
diff --git a/Engine/Plugins/PCG/Source/PCG/Public/Data/PCGPolyLineData.h b/Engine/Plugins/PCG/Source/PCG/Public/Data/PCGPolyLineData.h
--- a/Engine/Plugins/PCG/Source/PCG/Public/Data/PCGPolyLineData.h
+++ b/Engine/Plugins/PCG/Source/PCG/Public/Data/PCGPolyLineData.h
@@ -30,6 +30,12 @@ public:
 	/** Get the length of a specific segment of the line. */
 	virtual FVector::FReal GetSegmentLength(int SegmentIndex) const PURE_VIRTUAL(UPCGPolyLineData::GetSegmentLength, return 0;);
 
+	/** Get the location of the point at the normalized [0, 1] parameter across the entire the poly line. */
+	virtual FVector GetLocationAtAlpha(float Alpha) const PURE_VIRTUAL(UPCGPolyLine::GetLocationAtAlpha, return FVector::ZeroVector;)
+
+	/** Get the full transform at the normalized [0, 1] parameter across the entire the poly line. */
+	virtual FTransform GetTransformAtAlpha(float Alpha) const PURE_VIRTUAL(UPCGPolyLine::GetTransformAtAlpha, return FTransform::Identity;)
+
 	/** Get the total length of the line. */
 	virtual FVector::FReal GetLength() const;
 
@@ -51,6 +57,9 @@ public:
 	/** Get the input key at a distance along the line. InputKey is a float value in [0, N], where N is the number of control points. Each range [i, i+1] represents an interpolation from 0 to 1 across spline segment i. */
 	virtual float GetInputKeyAtDistance(int SegmentIndex, FVector::FReal Distance) const { return 0; }
 
+	/** Get the input key from the normalized distance of [0, 1] across the entire the poly line. */
+	virtual float GetInputKeyAtAlpha(float Alpha) const;
+
 	/** Get the arrive and leave tangents for a control point via its segment index. */
 	virtual void GetTangentsAtSegmentStart(int SegmentIndex, FVector& OutArriveTangent, FVector& OutLeaveTangent) const;
 
diff --git a/Engine/Plugins/PCG/Source/PCG/Public/Data/PCGSplineData.h b/Engine/Plugins/PCG/Source/PCG/Public/Data/PCGSplineData.h
--- a/Engine/Plugins/PCG/Source/PCG/Public/Data/PCGSplineData.h
+++ b/Engine/Plugins/PCG/Source/PCG/Public/Data/PCGSplineData.h
@@ -41,6 +41,8 @@ public:
 	virtual float GetInputKeyAtDistance(int SegmentIndex, FVector::FReal Distance) const override;
 	virtual void GetTangentsAtSegmentStart(int SegmentIndex, FVector& OutArriveTangent, FVector& OutLeaveTangent) const override;
 	virtual FVector::FReal GetDistanceAtSegmentStart(int SegmentIndex) const override;
+	virtual FVector GetLocationAtAlpha(float Alpha) const override;
+	virtual FTransform GetTransformAtAlpha(float Alpha) const override;
 	virtual bool IsClosed() const override { return SplineStruct.bClosedLoop; }
 	//~End UPCGPolyLineData interface
 
diff --git a/Engine/Plugins/PCG/Source/PCG/Public/Elements/Grammar/PCGSplineSlicer.h b/Engine/Plugins/PCG/Source/PCG/Public/Elements/Grammar/PCGSplineSlicer.h
new file mode 100644
--- /dev/null
+++ b/Engine/Plugins/PCG/Source/PCG/Public/Elements/Grammar/PCGSplineSlicer.h
@@ -0,0 +1,70 @@
+// Copyright Epic Games, Inc. All Rights Reserved.
+
+#pragma once
+
+#include "PCGSettings.h"
+
+#include "Elements/Grammar/PCGSlicingBase.h"
+
+#include "PCGSplineSlicer.generated.h"
+
+UCLASS(MinimalAPI, BlueprintType, ClassGroup = (Procedural))
+class UPCGSplineSlicerSettings : public UPCGSlicingBaseSettings
+{
+	GENERATED_BODY()
+
+public:
+	//~Begin UPCGSettings interface
+#if WITH_EDITOR
+	virtual FName GetDefaultNodeName() const override { return FName(TEXT("SplineSlicer")); }
+	virtual FText GetDefaultNodeTitle() const override { return NSLOCTEXT("PCGSplineSlicerElement", "NodeTitle", "Spline Slicer"); }
+#endif // WITH_EDITOR
+
+protected:
+	virtual TArray<FPCGPinProperties> InputPinProperties() const override;
+	virtual TArray<FPCGPinProperties> OutputPinProperties() const override;
+	virtual FPCGElementPtr CreateElement() const override;
+	//~End UPCGSettings interface
+
+public:
+	/** If the slicing with a given grammar doesn't fill the entire spline, setting it to true makes it a valid case. */
+	UPROPERTY(BlueprintReadWrite, EditAnywhere, Category = Settings)
+	bool bAcceptIncompleteSlicing = false;
+
+	/** Select the module height from an attribute. */
+	UPROPERTY(BlueprintReadWrite, EditAnywhere, Category = Settings)
+	bool bModuleHeightAsAttribute = false;
+
+	/** The height of each placed module. */
+	UPROPERTY(BlueprintReadWrite, EditAnywhere, Category = Settings, meta = (EditCondition = "!bModuleHeightAsAttribute", EditConditionHides))
+	double ModuleHeight = 100.0;
+
+	/** Selection will be used as the module height for placed modules. */
+	UPROPERTY(BlueprintReadWrite, EditAnywhere, Category = Settings, meta = (EditCondition = "bModuleHeightAsAttribute", EditConditionHides))
+	FPCGAttributePropertyInputSelector ModuleHeightAttribute;
+
+	UPROPERTY(BlueprintReadWrite, EditAnywhere, Category = "Settings|Extra Output Attributes", meta = (PCG_Overridable, InlineEditConditionToggle))
+	bool bOutputModuleIndexAttribute = false;
+
+	/** Name of the module index output attribute name. */
+	UPROPERTY(BlueprintReadWrite, EditAnywhere, Category = "Settings|Extra Output Attributes", meta = (PCG_Overridable, EditCondition = "bOutputModuleIndexAttribute"))
+	FName ModuleIndexAttributeName = TEXT("ModuleIndex");
+
+	/** Output attributes labeling the first and final points after slicing. */
+	UPROPERTY(BlueprintReadWrite, EditAnywhere, Category = "Settings|Extra Output Attributes", meta = (PCG_Overridable))
+	bool bOutputExtremityAttributes = false;
+
+	/** Name of the attribute labeling the first output point from the first module. */
+	UPROPERTY(BlueprintReadWrite, EditAnywhere, Category = "Settings|Extra Output Attributes", meta = (PCG_Overridable, EditCondition = "bOutputExtremityAttributes", EditConditionHides))
+	FName IsFirstAttributeName = TEXT("IsFirst");
+
+	/** Name of the attribute labeling the final output point from the final module. */
+	UPROPERTY(BlueprintReadWrite, EditAnywhere, Category = "Settings|Extra Output Attributes", meta = (PCG_Overridable, EditCondition = "bOutputExtremityAttributes", EditConditionHides))
+	FName IsFinalAttributeName = TEXT("IsFinal");
+};
+
+class FPCGSplineSlicerElement : public FPCGSlicingBaseElement
+{
+protected:
+	virtual bool ExecuteInternal(FPCGContext* InContext) const override;
+};
diff --git a/Engine/Plugins/PCG/Source/PCG/Public/Helpers/PCGPropertyHelpers.h b/Engine/Plugins/PCG/Source/PCG/Public/Helpers/PCGPropertyHelpers.h
--- a/Engine/Plugins/PCG/Source/PCG/Public/Helpers/PCGPropertyHelpers.h
+++ b/Engine/Plugins/PCG/Source/PCG/Public/Helpers/PCGPropertyHelpers.h
@@ -219,7 +219,7 @@ TArray<T> PCGPropertyHelpers::ExtractAttributeSetAsArrayOfStructs(const UPCGPara
 		{
 			if (!bCanBeDefaulted)
 			{
-				PCGLog::Accessor::LogFailToGet(Selector, OptionalContext);
+				PCGLog::Metadata::LogFailToCreateAccessor(Selector, OptionalContext);
 				return {};
 			}
 			else
@@ -258,7 +258,7 @@ TArray<T> PCGPropertyHelpers::ExtractAttributeSetAsArrayOfStructs(const UPCGPara
 
 			if (!Accessor->GetRange<PropertyType>(Temp, 0, *Keys, EPCGAttributeAccessorFlags::AllowBroadcastAndConstructible))
 			{
-				PCGLog::Accessor::LogFailToGet<PropertyType>(FText::FromName(PropertyName), Accessor.Get(), OptionalContext);
+				PCGLog::Metadata::LogFailToGetAttribute<PropertyType>(FText::FromName(PropertyName), Accessor.Get(), OptionalContext);
 				return false;
 			}
 
diff --git a/Engine/Plugins/PCG/Source/PCG/Public/Utils/PCGLogErrors.h b/Engine/Plugins/PCG/Source/PCG/Public/Utils/PCGLogErrors.h
--- a/Engine/Plugins/PCG/Source/PCG/Public/Utils/PCGLogErrors.h
+++ b/Engine/Plugins/PCG/Source/PCG/Public/Utils/PCGLogErrors.h
@@ -28,21 +28,35 @@ namespace PCGLog
 		void LogFirstInputOnlyWarning(const FName PinLabel, const FPCGContext* InContext = nullptr);
 	}
 
-	namespace Accessor
+	namespace Metadata
 	{
 		namespace ErrorFormat
 		{
+			extern const FTextFormat CreateAttributeFailure;
 			extern const FTextFormat GetTypedAttributeFailure;
 			extern const FTextFormat GetTypedAttributeFailureNoAccessor;
 		}
 
-		void LogFailToCreate(const FPCGAttributePropertySelector& Selector, const FPCGContext* InContext = nullptr);
-		void LogFailToGet(FName AttributeName, const FPCGContext* InContext = nullptr);
-		void LogFailToGet(FText AttributeName, const FPCGContext* InContext = nullptr);
-		void LogFailToGet(const FPCGAttributePropertySelector& Selector, const FPCGContext* InContext = nullptr);
+		void LogFailToCreateAccessor(const FPCGAttributePropertySelector& Selector, const FPCGContext* InContext = nullptr);
 
 		template <typename T>
-		void LogFailToGet(FText AttributeName, const IPCGAttributeAccessor* Accessor, const FPCGContext* InContext = nullptr)
+		void LogFailToCreateAttribute(FText AttributeName, const FPCGContext* InContext = nullptr)
+		{
+			PCGLog::LogErrorOnGraph(FText::Format(ErrorFormat::CreateAttributeFailure, AttributeName, PCG::Private::GetTypeNameText<T>()), InContext);
+		}
+
+		template <typename T>
+		void LogFailToCreateAttribute(FName AttributeName, const FPCGContext* InContext = nullptr)
+		{
+			LogFailToCreateAttribute<T>(FText::FromName(AttributeName), InContext);
+		}
+
+		void LogFailToGetAttribute(FText AttributeName, const FPCGContext* InContext = nullptr);
+		void LogFailToGetAttribute(FName AttributeName, const FPCGContext* InContext = nullptr);
+		void LogFailToGetAttribute(const FPCGAttributePropertySelector& Selector, const FPCGContext* InContext = nullptr);
+
+		template <typename T>
+		void LogFailToGetAttribute(FText AttributeName, const IPCGAttributeAccessor* Accessor, const FPCGContext* InContext = nullptr)
 		{
 			if (Accessor)
 			{
@@ -55,15 +69,15 @@ namespace PCGLog
 		}
 
 		template <typename T>
-		void LogFailToGet(FName AttributeName, const IPCGAttributeAccessor* Accessor, const FPCGContext* InContext = nullptr)
+		void LogFailToGetAttribute(FName AttributeName, const IPCGAttributeAccessor* Accessor, const FPCGContext* InContext = nullptr)
 		{
-			LogFailToGet<T>(FText::FromName(AttributeName), Accessor, InContext);
+			LogFailToGetAttribute<T>(FText::FromName(AttributeName), Accessor, InContext);
 		}
 
 		template <typename T>
-		void LogFailToGet(const FPCGAttributePropertySelector& Selector, const IPCGAttributeAccessor* Accessor, const FPCGContext* InContext = nullptr)
+		void LogFailToGetAttribute(const FPCGAttributePropertySelector& Selector, const IPCGAttributeAccessor* Accessor, const FPCGContext* InContext = nullptr)
 		{
-			return LogFailToGet<T>(Selector.GetDisplayText(), Accessor, InContext);
+			return LogFailToGetAttribute<T>(Selector.GetDisplayText(), Accessor, InContext);
 		}
 	}
 }
