Small change to include the data type enum instead of forward declaring it.

#rnx

diff --git a/Engine/Plugins/PCG/Source/PCG/Public/Utils/PCGLogErrors.h b/Engine/Plugins/PCG/Source/PCG/Public/Utils/PCGLogErrors.h
--- a/Engine/Plugins/PCG/Source/PCG/Public/Utils/PCGLogErrors.h
+++ b/Engine/Plugins/PCG/Source/PCG/Public/Utils/PCGLogErrors.h
@@ -2,11 +2,12 @@
 
 #pragma once
 
-#include "Internationalization/Internationalization.h"
+#include "PCGCommon.h"
 #include "Metadata/PCGAttributePropertySelector.h"
 #include "Metadata/Accessors/IPCGAttributeAccessor.h"
 
-enum class EPCGDataType : uint32;
+#include "Internationalization/Internationalization.h"
+
 struct FPCGContext;
 
 namespace PCGLog
