[PCG] Fix Automation Test : FPCGGraphExecutor can be deleted before a GameThread task gets ticked

#rnx

diff --git a/Engine/Plugins/PCG/Source/PCG/Private/Graph/PCGGraphExecutor.cpp b/Engine/Plugins/PCG/Source/PCG/Private/Graph/PCGGraphExecutor.cpp
--- a/Engine/Plugins/PCG/Source/PCG/Private/Graph/PCGGraphExecutor.cpp
+++ b/Engine/Plugins/PCG/Source/PCG/Private/Graph/PCGGraphExecutor.cpp
@@ -309,13 +309,21 @@ void FPCGGraphActiveTask::StopExecuting()
 	}
 }
 
+FPCGGraphExecutor::FPCGGraphExecutor()
+	: FPCGGraphExecutor(nullptr)
+{
+}
+
 FPCGGraphExecutor::FPCGGraphExecutor(UWorld* InWorld)
 	: World(InWorld)
 {
+	GameThreadHandler = MakeShared<FGameThreadHandler>(this);
 }
 
 FPCGGraphExecutor::~FPCGGraphExecutor()
 {
+	GameThreadHandler = nullptr;
+	
 	// We don't really need to do this here (it would be done in the destructor of these both)
 	// but this is to clarify/ensure the order in which this happens
 	GraphCache.ClearCache();
@@ -614,7 +622,6 @@ TSet<UPCGComponent*> FPCGGraphExecutor::Cancel(TFunctionRef<bool(TWeakObjectPtr<
 								check(ActiveTask->TaskIndex == ActiveTaskIndex);
 								PCGGraphExecutor::RemoveAtFromActiveTaskArrayNoLock(InActiveTasks, ActiveTaskIndex);
 								ActiveTask->TaskIndex = INDEX_NONE;
-								ActiveTask->bWasCancelled = true;
 							}
 						}
 					}
@@ -1077,12 +1084,15 @@ void FPCGGraphExecutor::PostTaskExecute(TSharedPtr<FPCGGraphActiveTask> ActiveTa
 			CollectGCReferenceTasks.Add(ActiveTaskPtr);
 		}
 
-		PCGGraphExecutor::ExecuteOnGameThread(UE_SOURCE_LOCATION, [this, ActiveTaskPtr]()
+		PCGGraphExecutor::ExecuteOnGameThread(UE_SOURCE_LOCATION, [WeakHandler = GameThreadHandler->AsWeak(), ActiveTaskPtr]()
 		{
-			ActiveTaskPtr->Element->DebugDisplay(ActiveTaskPtr->Context.Get());
+			if (TSharedPtr<FPCGGraphExecutor::FGameThreadHandler> Handler = WeakHandler.Pin())
+			{
+				ActiveTaskPtr->Element->DebugDisplay(ActiveTaskPtr->Context.Get());
 
-			PCGGraphExecutor::TScopeLock ScopeLock(CollectGCReferenceTasksLock);
-			CollectGCReferenceTasks.Remove(ActiveTaskPtr);
+				PCGGraphExecutor::TScopeLock ScopeLock(Handler->GetExecutor()->CollectGCReferenceTasksLock);
+				Handler->GetExecutor()->CollectGCReferenceTasks.Remove(ActiveTaskPtr);
+			}
 		});
 		
 		if (UPCGComponent* SourceComponent = ActiveTask.Context->SourceComponent.Get())
diff --git a/Engine/Plugins/PCG/Source/PCG/Private/Graph/PCGGraphExecutor.h b/Engine/Plugins/PCG/Source/PCG/Private/Graph/PCGGraphExecutor.h
--- a/Engine/Plugins/PCG/Source/PCG/Private/Graph/PCGGraphExecutor.h
+++ b/Engine/Plugins/PCG/Source/PCG/Private/Graph/PCGGraphExecutor.h
@@ -163,7 +163,7 @@ class FPCGGraphExecutor : public FGCObject
 {
 public:
 	// Default constructor used by unittests
-	FPCGGraphExecutor() = default;
+	FPCGGraphExecutor();
 	FPCGGraphExecutor(UWorld* InWorld);
 	~FPCGGraphExecutor();
 
@@ -405,6 +405,21 @@ private:
 	EExecuteVersion ExecuteVersion = EExecuteVersion::None;
 
 	EExecuteVersion GetExecuteVersion() const;
+
+	// Handler that we can use as a Weak ptr to determine if the Executor is still valid
+	class FGameThreadHandler : public TSharedFromThis<FGameThreadHandler>
+	{
+	public:
+		FGameThreadHandler(FPCGGraphExecutor* InExecutor)
+			: Executor(InExecutor) { }
+
+		FPCGGraphExecutor* GetExecutor() { return Executor; }
+
+	private:
+		FPCGGraphExecutor* Executor = nullptr;
+	};
+
+	TSharedPtr<FGameThreadHandler> GameThreadHandler;
 };
 
 class FPCGFetchInputElement : public IPCGElement
