[PCG] Graph Executor V2: Fix two concurrency issues
- bHasDonePrepareForExecute to true needs to be set after Context creation so that FPCGGraphExecutor::AddReferencedObjects always has a reference on inputs
- Cancelled task in a PostExecute state should not StoreResults/QueueNextTasks

diff --git a/Engine/Plugins/PCG/Source/PCG/Private/Graph/PCGGraphExecutor.cpp b/Engine/Plugins/PCG/Source/PCG/Private/Graph/PCGGraphExecutor.cpp
--- a/Engine/Plugins/PCG/Source/PCG/Private/Graph/PCGGraphExecutor.cpp
+++ b/Engine/Plugins/PCG/Source/PCG/Private/Graph/PCGGraphExecutor.cpp
@@ -23,6 +23,7 @@
 #include "Algo/ForEach.h"
 #include "Containers/Ticker.h"
 #include "GameFramework/Actor.h"
+#include "Misc/ScopeExit.h"
 #include "Tasks/Task.h"
 
 #if WITH_EDITOR
@@ -1025,7 +1026,8 @@ void FPCGGraphExecutor::PostTaskExecute(TSharedPtr<FPCGGraphActiveTask> ActiveTa
 		}
 	}
 
-	const bool bTaskWasCancelled = ActiveTask.bWasCancelled;
+	const bool bTaskFullyExecuted = !ActiveTask.bWasCancelled && ActiveTask.Context->CurrentPhase == EPCGExecutionPhase::Done;
+
 #if WITH_EDITOR
 	const bool bTaskWasBypassed = ActiveTask.bIsBypassed;
 #else
@@ -1052,7 +1054,7 @@ void FPCGGraphExecutor::PostTaskExecute(TSharedPtr<FPCGGraphActiveTask> ActiveTa
 		}
 	}
 
-	if (!bTaskWasCancelled && !bTaskWasBypassed)
+	if (bTaskFullyExecuted && !bTaskWasBypassed)
 	{
 		// Store result in cache as needed - done here because it needs to be done on the main thread
 
@@ -1074,9 +1076,9 @@ void FPCGGraphExecutor::PostTaskExecute(TSharedPtr<FPCGGraphActiveTask> ActiveTa
 		CurrentlyUsedThreads -= ActiveTask.Context->AsyncState.NumAvailableTasks;
 	}
 
-#if WITH_EDITOR
-	if (!bTaskWasCancelled)
+	if (bTaskWasBypassed || bTaskFullyExecuted)
 	{
+#if WITH_EDITOR
 		// No need to add if we are post executing off of main thread
 		if(!IsInGameThread())
 		{
@@ -1103,18 +1105,17 @@ void FPCGGraphExecutor::PostTaskExecute(TSharedPtr<FPCGGraphActiveTask> ActiveTa
 				SourceComponent->StoreInspectionData(Stack, ActiveTask.Context->Node, &ActiveTask.Context->Timer, ActiveTask.Context->InputData, ActiveTask.Context->OutputData, /*bUsedCache=*/false);
 			}
 		}
-	}
-#endif
+#endif // WITH_EDITOR
 
-	// If the task is a post execute, then we can safely clear the data after getting it from the results.
-	const bool bTaskIsPostExecute = (ActiveTask.Element == GraphCompiler.GetSharedTrivialPostGraphElement());
-		
-	// Store output in data map.
-	// TODO - investigate if we should avoid doing this if the task was cancelled.
-	StoreResults(ActiveTask.NodeId, ActiveTask.Context->OutputData, bTaskIsPostExecute);
+		// If the task is a post execute, then we can safely clear the data after getting it from the results.
+		const bool bTaskIsPostExecute = (ActiveTask.Element == GraphCompiler.GetSharedTrivialPostGraphElement());
+
+		// Store output in data map.
+		StoreResults(ActiveTask.NodeId, ActiveTask.Context->OutputData, bTaskIsPostExecute);
 
-	// Book-keeping
-	QueueNextTasks(ActiveTask.NodeId);
+		// Book-keeping
+		QueueNextTasks(ActiveTask.NodeId);
+	}
 
 	// Erase from ActiveTasks
 	{
@@ -1208,7 +1209,12 @@ void FPCGGraphExecutor::PrepareForExecute(FPCGGraphTask& Task, FCachedResult*& O
 {
 	TRACE_CPUPROFILER_EVENT_SCOPE(FPCGGraphExecutor::PrepareForExecute);
 	check(Task.bHasDoneSetup && !Task.bHasDonePrepareForExecute);
-	Task.bHasDonePrepareForExecute = true;
+	
+	// Set this only when exiting this method so that Context is created and is referencing the inputs
+	ON_SCOPE_EXIT
+	{
+		Task.bHasDonePrepareForExecute = true;
+	};
 
 	PCGGraphExecutionLogging::LogTaskExecute(Task);
 
