[PCG] Add support for Debug for GPU-executed nodes

#rnx

diff --git a/Engine/Plugins/PCG/Source/PCG/Private/Compute/DataInterfaces/PCGDataCollectionDataInterface.cpp b/Engine/Plugins/PCG/Source/PCG/Private/Compute/DataInterfaces/PCGDataCollectionDataInterface.cpp
--- a/Engine/Plugins/PCG/Source/PCG/Private/Compute/DataInterfaces/PCGDataCollectionDataInterface.cpp
+++ b/Engine/Plugins/PCG/Source/PCG/Private/Compute/DataInterfaces/PCGDataCollectionDataInterface.cpp
@@ -484,11 +484,17 @@ UComputeDataProvider* UPCGDataCollectionDataInterface::CreateDataProvider(TObjec
 	Provider->PinDesc = ProducerSettings->ComputeOutputPinDataDesc(OutputPinLabel, Binding);
 
 	Provider->ReadbackMode = bRequiresReadback ? EPCGReadbackMode::GraphOutput : EPCGReadbackMode::None;
+
 #if WITH_EDITOR
 	if (Binding->SourceComponent->IsInspecting())
 	{
 		Provider->ReadbackMode |= EPCGReadbackMode::Inspection;
 	}
+
+	if (ProducerSettings->bDebug)
+	{
+		Provider->ReadbackMode |= EPCGReadbackMode::DebugVisualization;
+	}
 #endif
 
 	// Use the aliased label for normal data output as this is the output from the compute graph.
@@ -631,6 +637,24 @@ bool UPCGDataCollectionDataProvider::ProcessReadBackData(FPCGComputeGraphContext
 			Component->StoreInspectionData(InContext->Stack, Node, /*InTimer=*/nullptr, /*InInputData=*/{}, DataFromGPU, /*bUsedCache*/false);
 		}
 	}
+
+	if (!!(ReadbackMode & EPCGReadbackMode::DebugVisualization))
+	{
+		const FPCGDataVisualizationRegistry& DataVisRegistry = FPCGModule::GetConstPCGDataVisualizationRegistry();
+
+		for (const FPCGTaggedData& Output : DataFromGPU.TaggedData)
+		{
+			if (!Output.Data)
+			{
+				continue;
+			}
+
+			if (const IPCGDataVisualization* DataVis = DataVisRegistry.GetDataVisualization(Output.Data->GetClass()))
+			{
+				DataVis->ExecuteDebugDisplay(InContext, ProducerSettings, Output.Data, InContext->GetTargetActor(nullptr));
+			}
+		}
+	}
 #endif // WITH_EDITOR
 
 	return true;
diff --git a/Engine/Plugins/PCG/Source/PCG/Private/Compute/DataInterfaces/PCGDataCollectionDataInterface.h b/Engine/Plugins/PCG/Source/PCG/Private/Compute/DataInterfaces/PCGDataCollectionDataInterface.h
--- a/Engine/Plugins/PCG/Source/PCG/Private/Compute/DataInterfaces/PCGDataCollectionDataInterface.h
+++ b/Engine/Plugins/PCG/Source/PCG/Private/Compute/DataInterfaces/PCGDataCollectionDataInterface.h
@@ -23,6 +23,8 @@ enum class EPCGReadbackMode : uint8
 	GraphOutput = 1 << 1,
 	/** Read back and store inspection data. */
 	Inspection = 1 << 2,
+	/** Read back and execute debug visualization. */
+	DebugVisualization = 1 << 3,
 };
 ENUM_CLASS_FLAGS(EPCGReadbackMode);
 
diff --git a/Engine/Plugins/PCG/Source/PCG/Private/Elements/PCGDebugElement.cpp b/Engine/Plugins/PCG/Source/PCG/Private/Elements/PCGDebugElement.cpp
--- a/Engine/Plugins/PCG/Source/PCG/Private/Elements/PCGDebugElement.cpp
+++ b/Engine/Plugins/PCG/Source/PCG/Private/Elements/PCGDebugElement.cpp
@@ -52,7 +52,10 @@ bool FPCGDebugElement::ExecuteInternal(FPCGContext* Context) const
 			continue;
 		}
 
-		AActor* TargetActor = Context->GetInputSettings<UPCGDebugSettings>()->TargetActor.Get();
+		const UPCGDebugSettings* Settings = Context->GetInputSettings<UPCGDebugSettings>();
+		check(Settings);
+
+		AActor* TargetActor = Settings->TargetActor.Get();
 
 		if (!TargetActor)
 		{
@@ -61,7 +64,7 @@ bool FPCGDebugElement::ExecuteInternal(FPCGContext* Context) const
 
 		if (const IPCGDataVisualization* DataVis = DataVisRegistry.GetDataVisualization(Input.Data->GetClass()))
 		{
-			DataVis->ExecuteDebugDisplay(Context, Input.Data, TargetActor);
+			DataVis->ExecuteDebugDisplay(Context, Settings, Input.Data, TargetActor);
 		}
 	}
 #endif
diff --git a/Engine/Plugins/PCG/Source/PCG/Private/PCGElement.cpp b/Engine/Plugins/PCG/Source/PCG/Private/PCGElement.cpp
--- a/Engine/Plugins/PCG/Source/PCG/Private/PCGElement.cpp
+++ b/Engine/Plugins/PCG/Source/PCG/Private/PCGElement.cpp
@@ -364,7 +364,7 @@ void IPCGElement::DebugDisplay(FPCGContext* Context) const
 
 		if (const IPCGDataVisualization* DataVis = DataVisRegistry.GetDataVisualization(Input.Data->GetClass()))
 		{
-			DataVis->ExecuteDebugDisplay(Context, Input.Data, Context->GetTargetActor(nullptr));
+			DataVis->ExecuteDebugDisplay(Context, SettingsInterface, Input.Data, Context->GetTargetActor(nullptr));
 		}
 	}
 
diff --git a/Engine/Plugins/PCG/Source/PCG/Public/PCGDataVisualization.h b/Engine/Plugins/PCG/Source/PCG/Public/PCGDataVisualization.h
--- a/Engine/Plugins/PCG/Source/PCG/Public/PCGDataVisualization.h
+++ b/Engine/Plugins/PCG/Source/PCG/Public/PCGDataVisualization.h
@@ -7,6 +7,7 @@
 #if WITH_EDITOR
 class AActor;
 class UPCGData;
+class UPCGSettingsInterface;
 struct FPCGContext;
 
 enum class EPCGTableVisualizerColumnSortingMode
@@ -49,7 +50,7 @@ class IPCGDataVisualization
 {
 public:
 	virtual ~IPCGDataVisualization() = default;
-	virtual void ExecuteDebugDisplay(FPCGContext* Context, const UPCGData* Data, AActor* TargetActor) const = 0;
+	virtual void ExecuteDebugDisplay(FPCGContext* Context, const UPCGSettingsInterface* SettingsInterface, const UPCGData* Data, AActor* TargetActor) const = 0;
 	virtual FPCGTableVisualizerInfo GetTableVisualizerInfo(const UPCGData* Data) const = 0;
 };
 #endif
diff --git a/Engine/Plugins/PCG/Source/PCGEditor/Private/DataVisualizations/PCGParamDataVisualization.h b/Engine/Plugins/PCG/Source/PCGEditor/Private/DataVisualizations/PCGParamDataVisualization.h
--- a/Engine/Plugins/PCG/Source/PCGEditor/Private/DataVisualizations/PCGParamDataVisualization.h
+++ b/Engine/Plugins/PCG/Source/PCGEditor/Private/DataVisualizations/PCGParamDataVisualization.h
@@ -9,6 +9,6 @@ class UPCGData;
 class PCGEDITOR_API IPCGParamDataVisualization : public IPCGDataVisualization
 {
 public:
-	virtual void ExecuteDebugDisplay(FPCGContext* Context, const UPCGData* Data, AActor* TargetActor) const { /* Do nothing. */ }
+	virtual void ExecuteDebugDisplay(FPCGContext* Context, const UPCGSettingsInterface* SettingsInterface, const UPCGData* Data, AActor* TargetActor) const { /* Do nothing. */ }
 	virtual FPCGTableVisualizerInfo GetTableVisualizerInfo(const UPCGData* Data) const override;
 };
diff --git a/Engine/Plugins/PCG/Source/PCGEditor/Private/DataVisualizations/PCGSpatialDataVisualization.cpp b/Engine/Plugins/PCG/Source/PCGEditor/Private/DataVisualizations/PCGSpatialDataVisualization.cpp
--- a/Engine/Plugins/PCG/Source/PCGEditor/Private/DataVisualizations/PCGSpatialDataVisualization.cpp
+++ b/Engine/Plugins/PCG/Source/PCGEditor/Private/DataVisualizations/PCGSpatialDataVisualization.cpp
@@ -77,7 +77,7 @@ namespace PCGPointDataVisualizationConstants
 	const FText TEXT_MetadataEntryParent = LOCTEXT("MetadataEntryParent", "Parent Key");
 }
 
-void IPCGSpatialDataVisualization::ExecuteDebugDisplay(FPCGContext* Context, const UPCGData* Data, AActor* TargetActor) const
+void IPCGSpatialDataVisualization::ExecuteDebugDisplay(FPCGContext* Context, const UPCGSettingsInterface* SettingsInterface, const UPCGData* Data, AActor* TargetActor) const
 {
 	if (!TargetActor)
 	{
@@ -85,8 +85,6 @@ void IPCGSpatialDataVisualization::ExecuteDebugDisplay(FPCGContext* Context, con
 		return;
 	}
 
-	const UPCGSettingsInterface* SettingsInterface = Context ? Context->GetInputSettingsInterface() : nullptr;
-
 	if (!SettingsInterface)
 	{
 		return;
diff --git a/Engine/Plugins/PCG/Source/PCGEditor/Private/DataVisualizations/PCGSpatialDataVisualization.h b/Engine/Plugins/PCG/Source/PCGEditor/Private/DataVisualizations/PCGSpatialDataVisualization.h
--- a/Engine/Plugins/PCG/Source/PCGEditor/Private/DataVisualizations/PCGSpatialDataVisualization.h
+++ b/Engine/Plugins/PCG/Source/PCGEditor/Private/DataVisualizations/PCGSpatialDataVisualization.h
@@ -15,7 +15,7 @@ class PCGEDITOR_API IPCGSpatialDataVisualization : public IPCGDataVisualization
 {
 public:
 	// ~Begin IPCGDataVisualization interface
-	virtual void ExecuteDebugDisplay(FPCGContext* Context, const UPCGData* Data, AActor* TargetActor) const override;
+	virtual void ExecuteDebugDisplay(FPCGContext* Context, const UPCGSettingsInterface* SettingsInterface, const UPCGData* Data, AActor* TargetActor) const override;
 	virtual FPCGTableVisualizerInfo GetTableVisualizerInfo(const UPCGData* Data) const override;
 	// ~End IPCGDataVisualization interface
 
