[PCG] Fix deprecated comparison of two different pointer types

#rnx

diff --git a/Engine/Plugins/PCG/Source/PCG/Private/Tests/Elements/PCGStaticMeshSpawnerTest.cpp b/Engine/Plugins/PCG/Source/PCG/Private/Tests/Elements/PCGStaticMeshSpawnerTest.cpp
--- a/Engine/Plugins/PCG/Source/PCG/Private/Tests/Elements/PCGStaticMeshSpawnerTest.cpp
+++ b/Engine/Plugins/PCG/Source/PCG/Private/Tests/Elements/PCGStaticMeshSpawnerTest.cpp
@@ -196,7 +196,7 @@ void TestMeshSelectorByAttribute(
 		const int InstanceCount = ISMC->GetInstanceCount();
 		TotalInstanceCount += InstanceCount;
 
-		TObjectPtr<UStaticMesh> StaticMesh = ISMC->GetStaticMesh();
+		UStaticMesh* StaticMesh = ISMC->GetStaticMesh();
 
 		const FPCGByAttributeValidationData* Entry = ValidationDataEntries.FindByPredicate([StaticMesh](const FPCGByAttributeValidationData& Entry) {
 			return Entry.Mesh == StaticMesh;
@@ -301,7 +301,7 @@ void TestMeshSelectorWeighted(
 		const int InstanceCount = ISMC->GetInstanceCount();
 		TotalInstanceCount += InstanceCount;
 
-		TObjectPtr<UStaticMesh> StaticMesh = ISMC->GetStaticMesh();
+		UStaticMesh* StaticMesh = ISMC->GetStaticMesh();
 
 		FPCGMeshSelectorWeightedEntry* Entry = MeshSelector->MeshEntries.FindByPredicate([StaticMesh](const FPCGMeshSelectorWeightedEntry& Entry) {
 			return Entry.Descriptor.StaticMesh == StaticMesh;
@@ -383,7 +383,7 @@ void TestMeshSelectorWeightedByCategory(
 		const int InstanceCount = ISMC->GetInstanceCount();
 		TotalInstanceCount += InstanceCount;
 
-		TObjectPtr<UStaticMesh> StaticMesh = ISMC->GetStaticMesh();
+		UStaticMesh* StaticMesh = ISMC->GetStaticMesh();
 
 		const FPCGWeightedByCategoryValidationData* Entry = ValidationDataEntries.FindByPredicate([StaticMesh](const FPCGWeightedByCategoryValidationData& Entry) {
 			return Entry.Mesh == StaticMesh;
