Replaced strong pointer by weak pointer to prevent leak when inspecting data on the output node

#rnx

diff --git a/Engine/Plugins/PCG/Source/PCGEditor/Private/SPCGEditorGraphAttributeListView.cpp b/Engine/Plugins/PCG/Source/PCGEditor/Private/SPCGEditorGraphAttributeListView.cpp
--- a/Engine/Plugins/PCG/Source/PCGEditor/Private/SPCGEditorGraphAttributeListView.cpp
+++ b/Engine/Plugins/PCG/Source/PCGEditor/Private/SPCGEditorGraphAttributeListView.cpp
@@ -264,7 +264,7 @@ SPCGEditorGraphAttributeListView::~SPCGEditorGraphAttributeListView()
 		PCGEditorPtr.Pin()->OnInspectedStackChangedDelegate.RemoveAll(this);
 	}
 
-	DataStrongPtr.Reset();
+	DataPtr.Reset();
 }
 
 void SPCGEditorGraphAttributeListView::Construct(const FArguments& InArgs, TSharedPtr<FPCGEditor> InPCGEditor)
@@ -540,9 +540,9 @@ void SPCGEditorGraphAttributeListView::Tick(const FGeometry& AllottedGeometry, c
 				InfoTextBlock->SetText(FText::Format(PCGEditorGraphAttributeListView::CrcLabelFormat, InfoTextBlock->GetText(), Crc.GetValue()));
 			}
 
-			if (DataStrongPtr.IsValid() && DataStrongPtr->HasCachedLastSelector())
+			if (DataPtr.IsValid() && DataPtr->HasCachedLastSelector())
 			{
-				const FText LastSelector = DataStrongPtr->GetCachedLastSelector().GetDisplayText();
+				const FText LastSelector = DataPtr->GetCachedLastSelector().GetDisplayText();
 				InfoTextBlock->SetText(FText::Format(PCGEditorGraphAttributeListView::LastLabelFormat, InfoTextBlock->GetText(), LastSelector));
 			}
 		}
@@ -680,7 +680,7 @@ void SPCGEditorGraphAttributeListView::RefreshAttributeList()
 	ListViewItems.Empty();
 	ListViewHeader->ClearColumns();
 	InfoTextBlock->SetText(FText::GetEmpty());
-	DataStrongPtr.Reset();
+	DataPtr.Reset();
 
 	const FPCGDataCollection* InspectionData = GetInspectionData();
 	if (!InspectionData)
@@ -706,7 +706,7 @@ void SPCGEditorGraphAttributeListView::RefreshAttributeList()
 	const FPCGDataVisualizationRegistry& DataVisRegistry = FPCGModule::GetConstPCGDataVisualizationRegistry();
 	const UPCGData* DataToVisualize = PCGData;
 
-	DataStrongPtr.Reset(DataToVisualize);
+	DataPtr = DataToVisualize;
 
 	if (const IPCGDataVisualization* DataVisualization = DataVisRegistry.GetDataVisualization(PCGData->GetClass()))
 	{
@@ -1112,9 +1112,9 @@ TSharedRef<ITableRow> SPCGEditorGraphAttributeListView::OnGenerateRow(PCGListVie
 void SPCGEditorGraphAttributeListView::OnItemDoubleClicked(PCGListViewItemPtr Item) const
 {
 	check(Item);
-	if (FocusOnDataCallback)
+	if (FocusOnDataCallback && DataPtr.Get())
 	{
-		FocusOnDataCallback(DataStrongPtr.Get(), { Item->Index });
+		FocusOnDataCallback(DataPtr.Get(), { Item->Index });
 	}
 }
 
@@ -1320,9 +1320,9 @@ FReply SPCGEditorGraphAttributeListView::OnNodeNameClicked()
 
 FReply SPCGEditorGraphAttributeListView::OnFocusOnDataClicked() const
 {
-	if (IsFocusOnDataEnabled())
+	if (IsFocusOnDataEnabled() && DataPtr.Get())
 	{
-		FocusOnDataCallback(DataStrongPtr.Get(), {});
+		FocusOnDataCallback(DataPtr.Get(), {});
 	}
 
 	return FReply::Handled();
@@ -1335,7 +1335,7 @@ bool SPCGEditorGraphAttributeListView::IsFocusOnDataEnabled() const
 
 void SPCGEditorGraphAttributeListView::FocusOnSelection() const
 {
-	if (!IsFocusOnDataEnabled())
+	if (!IsFocusOnDataEnabled() || !DataPtr.Get())
 	{
 		return;
 	}
@@ -1353,7 +1353,7 @@ void SPCGEditorGraphAttributeListView::FocusOnSelection() const
 		}
 	}
 
-	FocusOnDataCallback(DataStrongPtr.Get(), Indices);
+	FocusOnDataCallback(DataPtr.Get(), Indices);
 }
 
 bool SPCGEditorGraphAttributeListView::CanFocusOnSelection() const
diff --git a/Engine/Plugins/PCG/Source/PCGEditor/Private/SPCGEditorGraphAttributeListView.h b/Engine/Plugins/PCG/Source/PCGEditor/Private/SPCGEditorGraphAttributeListView.h
--- a/Engine/Plugins/PCG/Source/PCGEditor/Private/SPCGEditorGraphAttributeListView.h
+++ b/Engine/Plugins/PCG/Source/PCGEditor/Private/SPCGEditorGraphAttributeListView.h
@@ -280,5 +280,5 @@ private:
 	TSharedPtr<FPCGListViewUpdater> CurrentUpdateTask = nullptr;
 
 	/** Used to ensure data collapsed for inspection is kept alive. */
-	TStrongObjectPtr<const UPCGData> DataStrongPtr = nullptr;
+	TWeakObjectPtr<const UPCGData> DataPtr = nullptr;
 };
