CIS fix.

diff --git a/Engine/Plugins/PCG/Source/PCG/Private/Tests/Elements/PCGStaticMeshSpawnerTest.cpp b/Engine/Plugins/PCG/Source/PCG/Private/Tests/Elements/PCGStaticMeshSpawnerTest.cpp
--- a/Engine/Plugins/PCG/Source/PCG/Private/Tests/Elements/PCGStaticMeshSpawnerTest.cpp
+++ b/Engine/Plugins/PCG/Source/PCG/Private/Tests/Elements/PCGStaticMeshSpawnerTest.cpp
@@ -199,7 +199,7 @@ void TestMeshSelectorByAttribute(
 		UStaticMesh* StaticMesh = ISMC->GetStaticMesh();
 
 		const FPCGByAttributeValidationData* Entry = ValidationDataEntries.FindByPredicate([StaticMesh](const FPCGByAttributeValidationData& Entry) {
-			return Entry.Mesh == StaticMesh.Get();
+			return Entry.Mesh == StaticMesh;
 			});
 
 		if (Test->TestNotNull("Validate instanced mesh exists in MeshEntries", Entry))
