[PCG] Switch order of arguments to get/set DI functions be like (DataIndex, ElementIndex, Attribute, Value) (so flipped ElementIndex and Attribute).

#rnx

diff --git a/Engine/Plugins/PCG/Shaders/Private/PCGDataCollectionDataInterface.ush b/Engine/Plugins/PCG/Shaders/Private/PCGDataCollectionDataInterface.ush
--- a/Engine/Plugins/PCG/Shaders/Private/PCGDataCollectionDataInterface.ush
+++ b/Engine/Plugins/PCG/Shaders/Private/PCGDataCollectionDataInterface.ush
@@ -203,7 +203,7 @@ void StoreBuffer_{DataInterfaceName}(uint Address, uint Value)
 
 // #################### ATTRIBUTE GETTERS ##########################
 
-uint ComputeAddressHelper_{DataInterfaceName}(uint DataIndex, uint AttributeId, uint ElementIndex)
+uint ComputeAddressHelper_{DataInterfaceName}(uint DataIndex, uint ElementIndex, uint AttributeId)
 {
 	const uint DataAddress = ReadDataAddress_{DataInterfaceName}(DataIndex);
 	const uint AttributeHeadersAddress = ReadDataAttributeHeadersAddress_{DataInterfaceName}(DataAddress);
@@ -211,113 +211,113 @@ uint ComputeAddressHelper_{DataInterfaceName}(uint DataIndex, uint AttributeId,
 	return ReadAttributeAddress_{DataInterfaceName}(AttributeHeaderAddress);
 }
 
-bool GetBool_{DataInterfaceName}(uint DataIndex, uint AttributeId, uint ElementIndex)
+bool GetBool_{DataInterfaceName}(uint DataIndex, uint ElementIndex, uint AttributeId)
 {
-	const uint AttributeAddress = ComputeAddressHelper_{DataInterfaceName}(DataIndex, AttributeId, ElementIndex);
+	const uint AttributeAddress = ComputeAddressHelper_{DataInterfaceName}(DataIndex, ElementIndex, AttributeId);
 	return asbool({DataInterfaceName}_DataCollectionBuffer.Load(AttributeAddress + ElementIndex * 1 * 4));
 }
 
-int GetInt_{DataInterfaceName}(uint DataIndex, uint AttributeId, uint ElementIndex)
+int GetInt_{DataInterfaceName}(uint DataIndex, uint ElementIndex, uint AttributeId)
 {
-	const uint AttributeAddress = ComputeAddressHelper_{DataInterfaceName}(DataIndex, AttributeId, ElementIndex);
+	const uint AttributeAddress = ComputeAddressHelper_{DataInterfaceName}(DataIndex, ElementIndex, AttributeId);
 	return asint({DataInterfaceName}_DataCollectionBuffer.Load(AttributeAddress + ElementIndex * 1 * 4));
 }
 
-int GetUint_{DataInterfaceName}(uint DataIndex, uint AttributeId, uint ElementIndex)
+int GetUint_{DataInterfaceName}(uint DataIndex, uint ElementIndex, uint AttributeId)
 {
-	const uint AttributeAddress = ComputeAddressHelper_{DataInterfaceName}(DataIndex, AttributeId, ElementIndex);
+	const uint AttributeAddress = ComputeAddressHelper_{DataInterfaceName}(DataIndex, ElementIndex, AttributeId);
 	return {DataInterfaceName}_DataCollectionBuffer.Load(AttributeAddress + ElementIndex * 1 * 4);
 }
 
-float GetFloat_{DataInterfaceName}(uint DataIndex, uint AttributeId, uint ElementIndex)
+float GetFloat_{DataInterfaceName}(uint DataIndex, uint ElementIndex, uint AttributeId)
 {
-	const uint AttributeAddress = ComputeAddressHelper_{DataInterfaceName}(DataIndex, AttributeId, ElementIndex);
+	const uint AttributeAddress = ComputeAddressHelper_{DataInterfaceName}(DataIndex, ElementIndex, AttributeId);
 	return asfloat({DataInterfaceName}_DataCollectionBuffer.Load(AttributeAddress + ElementIndex * 1 * 4));
 }
 
-float2 GetFloat2_{DataInterfaceName}(uint DataIndex, uint AttributeId, uint ElementIndex)
+float2 GetFloat2_{DataInterfaceName}(uint DataIndex, uint ElementIndex, uint AttributeId)
 {
-	const uint AttributeAddress = ComputeAddressHelper_{DataInterfaceName}(DataIndex, AttributeId, ElementIndex);
+	const uint AttributeAddress = ComputeAddressHelper_{DataInterfaceName}(DataIndex, ElementIndex, AttributeId);
 	return asfloat({DataInterfaceName}_DataCollectionBuffer.Load2(AttributeAddress + ElementIndex * 2 * 4));
 }
 
-float3 GetFloat3_{DataInterfaceName}(uint DataIndex, uint AttributeId, uint ElementIndex)
+float3 GetFloat3_{DataInterfaceName}(uint DataIndex, uint ElementIndex, uint AttributeId)
 {
-	const uint AttributeAddress = ComputeAddressHelper_{DataInterfaceName}(DataIndex, AttributeId, ElementIndex);
+	const uint AttributeAddress = ComputeAddressHelper_{DataInterfaceName}(DataIndex, ElementIndex, AttributeId);
 	return asfloat({DataInterfaceName}_DataCollectionBuffer.Load3(AttributeAddress + ElementIndex * 3 * 4));
 }
 
-float4 GetFloat4_{DataInterfaceName}(uint DataIndex, uint AttributeId, uint ElementIndex)
+float4 GetFloat4_{DataInterfaceName}(uint DataIndex, uint ElementIndex, uint AttributeId)
 {
-	const uint AttributeAddress = ComputeAddressHelper_{DataInterfaceName}(DataIndex, AttributeId, ElementIndex);
+	const uint AttributeAddress = ComputeAddressHelper_{DataInterfaceName}(DataIndex, ElementIndex, AttributeId);
 	return asfloat({DataInterfaceName}_DataCollectionBuffer.Load4(AttributeAddress + ElementIndex * 4 * 4));
 }
 
-float3 GetRotator_{DataInterfaceName}(uint DataIndex, uint AttributeId, uint ElementIndex)
+float3 GetRotator_{DataInterfaceName}(uint DataIndex, uint ElementIndex, uint AttributeId)
 {
-	const uint AttributeAddress = ComputeAddressHelper_{DataInterfaceName}(DataIndex, AttributeId, ElementIndex);
+	const uint AttributeAddress = ComputeAddressHelper_{DataInterfaceName}(DataIndex, ElementIndex, AttributeId);
 	return asfloat({DataInterfaceName}_DataCollectionBuffer.Load3(AttributeAddress + ElementIndex * 3 * 4));
 }
 
-float4 GetQuat_{DataInterfaceName}(uint DataIndex, uint AttributeId, uint ElementIndex)
+float4 GetQuat_{DataInterfaceName}(uint DataIndex, uint ElementIndex, uint AttributeId)
 {
-	const uint AttributeAddress = ComputeAddressHelper_{DataInterfaceName}(DataIndex, AttributeId, ElementIndex);
+	const uint AttributeAddress = ComputeAddressHelper_{DataInterfaceName}(DataIndex, ElementIndex, AttributeId);
 	return asfloat({DataInterfaceName}_DataCollectionBuffer.Load4(AttributeAddress + ElementIndex * 4 * 4));
 }
 
 // #################### ATTRIBUTE SETTERS ##########################
 
-void SetBool_{DataInterfaceName}(uint DataIndex, uint AttributeId, uint ElementIndex, bool Value)
+void SetBool_{DataInterfaceName}(uint DataIndex, uint ElementIndex, uint AttributeId, bool Value)
 {
-	const uint AttributeAddress = ComputeAddressHelper_{DataInterfaceName}(DataIndex, AttributeId, ElementIndex);
+	const uint AttributeAddress = ComputeAddressHelper_{DataInterfaceName}(DataIndex, ElementIndex, AttributeId);
 	{DataInterfaceName}_DataCollectionBuffer.Store(AttributeAddress + ElementIndex * 1 * 4, asuint(Value));
 }
 
-void SetInt_{DataInterfaceName}(uint DataIndex, uint AttributeId, uint ElementIndex, int Value)
+void SetInt_{DataInterfaceName}(uint DataIndex, uint ElementIndex, uint AttributeId, int Value)
 {
-	const uint AttributeAddress = ComputeAddressHelper_{DataInterfaceName}(DataIndex, AttributeId, ElementIndex);
+	const uint AttributeAddress = ComputeAddressHelper_{DataInterfaceName}(DataIndex, ElementIndex, AttributeId);
 	{DataInterfaceName}_DataCollectionBuffer.Store(AttributeAddress + ElementIndex * 1 * 4, asuint(Value));
 }
 
-void SetUint_{DataInterfaceName}(uint DataIndex, uint AttributeId, uint ElementIndex, uint Value)
+void SetUint_{DataInterfaceName}(uint DataIndex, uint ElementIndex, uint AttributeId, uint Value)
 {
-	const uint AttributeAddress = ComputeAddressHelper_{DataInterfaceName}(DataIndex, AttributeId, ElementIndex);
+	const uint AttributeAddress = ComputeAddressHelper_{DataInterfaceName}(DataIndex, ElementIndex, AttributeId);
 	{DataInterfaceName}_DataCollectionBuffer.Store(AttributeAddress + ElementIndex * 1 * 4, Value);
 }
 
-void SetFloat_{DataInterfaceName}(uint DataIndex, uint AttributeId, uint ElementIndex, float Value)
+void SetFloat_{DataInterfaceName}(uint DataIndex, uint ElementIndex, uint AttributeId, float Value)
 {
-	const uint AttributeAddress = ComputeAddressHelper_{DataInterfaceName}(DataIndex, AttributeId, ElementIndex);
+	const uint AttributeAddress = ComputeAddressHelper_{DataInterfaceName}(DataIndex, ElementIndex, AttributeId);
 	{DataInterfaceName}_DataCollectionBuffer.Store(AttributeAddress + ElementIndex * 1 * 4, asuint(Value));
 }
 
-void SetFloat2_{DataInterfaceName}(uint DataIndex, uint AttributeId, uint ElementIndex, float2 Value)
+void SetFloat2_{DataInterfaceName}(uint DataIndex, uint ElementIndex, uint AttributeId, float2 Value)
 {
-	const uint AttributeAddress = ComputeAddressHelper_{DataInterfaceName}(DataIndex, AttributeId, ElementIndex);
+	const uint AttributeAddress = ComputeAddressHelper_{DataInterfaceName}(DataIndex, ElementIndex, AttributeId);
 	{DataInterfaceName}_DataCollectionBuffer.Store2(AttributeAddress + ElementIndex * 2 * 4, asuint(Value));
 }
 
-void SetFloat3_{DataInterfaceName}(uint DataIndex, uint AttributeId, uint ElementIndex, float3 Value)
+void SetFloat3_{DataInterfaceName}(uint DataIndex, uint ElementIndex, uint AttributeId, float3 Value)
 {
-	const uint AttributeAddress = ComputeAddressHelper_{DataInterfaceName}(DataIndex, AttributeId, ElementIndex);
+	const uint AttributeAddress = ComputeAddressHelper_{DataInterfaceName}(DataIndex, ElementIndex, AttributeId);
 	{DataInterfaceName}_DataCollectionBuffer.Store3(AttributeAddress + ElementIndex * 3 * 4, asuint(Value));
 }
 
-void SetFloat4_{DataInterfaceName}(uint DataIndex, uint AttributeId, uint ElementIndex, float4 Value)
+void SetFloat4_{DataInterfaceName}(uint DataIndex, uint ElementIndex, uint AttributeId, float4 Value)
 {
-	const uint AttributeAddress = ComputeAddressHelper_{DataInterfaceName}(DataIndex, AttributeId, ElementIndex);
+	const uint AttributeAddress = ComputeAddressHelper_{DataInterfaceName}(DataIndex, ElementIndex, AttributeId);
 	{DataInterfaceName}_DataCollectionBuffer.Store4(AttributeAddress + ElementIndex * 4 * 4, asuint(Value));
 }
 
-void SetRotator_{DataInterfaceName}(uint DataIndex, uint AttributeId, uint ElementIndex, float3 Value)
+void SetRotator_{DataInterfaceName}(uint DataIndex, uint ElementIndex, uint AttributeId, float3 Value)
 {
-	const uint AttributeAddress = ComputeAddressHelper_{DataInterfaceName}(DataIndex, AttributeId, ElementIndex);
+	const uint AttributeAddress = ComputeAddressHelper_{DataInterfaceName}(DataIndex, ElementIndex, AttributeId);
 	{DataInterfaceName}_DataCollectionBuffer.Store3(AttributeAddress + ElementIndex * 3 * 4, asuint(Value));
 }
 
-void SetQuat_{DataInterfaceName}(uint DataIndex, uint AttributeId, uint ElementIndex, float4 Value)
+void SetQuat_{DataInterfaceName}(uint DataIndex, uint ElementIndex, uint AttributeId, float4 Value)
 {
-	const uint AttributeAddress = ComputeAddressHelper_{DataInterfaceName}(DataIndex, AttributeId, ElementIndex);
+	const uint AttributeAddress = ComputeAddressHelper_{DataInterfaceName}(DataIndex, ElementIndex, AttributeId);
 	{DataInterfaceName}_DataCollectionBuffer.Store4(AttributeAddress + ElementIndex * 4 * 4, asuint(Value));
 }
 
@@ -325,47 +325,47 @@ void SetQuat_{DataInterfaceName}(uint DataIndex, uint AttributeId, uint ElementI
 
 float3 GetPosition_{DataInterfaceName}(uint DataIndex, uint ElementIndex)
 {
-	return GetFloat3_{DataInterfaceName}(DataIndex, POINT_POSITION_ID, ElementIndex);
+	return GetFloat3_{DataInterfaceName}(DataIndex, ElementIndex, POINT_POSITION_ID);
 }
 
 float4 GetRotation_{DataInterfaceName}(uint DataIndex, uint ElementIndex)
 {
-	return GetFloat4_{DataInterfaceName}(DataIndex, POINT_ROTATION_ID, ElementIndex);
+	return GetFloat4_{DataInterfaceName}(DataIndex, ElementIndex, POINT_ROTATION_ID);
 }
 
 float3 GetScale_{DataInterfaceName}(uint DataIndex, uint ElementIndex)
 {
-	return GetFloat3_{DataInterfaceName}(DataIndex, POINT_SCALE_ID, ElementIndex);
+	return GetFloat3_{DataInterfaceName}(DataIndex, ElementIndex, POINT_SCALE_ID);
 }
 
 float3 GetBoundsMin_{DataInterfaceName}(uint DataIndex, uint ElementIndex)
 {
-	return GetFloat3_{DataInterfaceName}(DataIndex, POINT_BOUNDS_MIN_ID, ElementIndex);
+	return GetFloat3_{DataInterfaceName}(DataIndex, ElementIndex, POINT_BOUNDS_MIN_ID);
 }
 
 float3 GetBoundsMax_{DataInterfaceName}(uint DataIndex, uint ElementIndex)
 {
-	return GetFloat3_{DataInterfaceName}(DataIndex, POINT_BOUNDS_MAX_ID, ElementIndex);
+	return GetFloat3_{DataInterfaceName}(DataIndex, ElementIndex, POINT_BOUNDS_MAX_ID);
 }
 
 float4 GetColor_{DataInterfaceName}(uint DataIndex, uint ElementIndex)
 {
-	return GetFloat4_{DataInterfaceName}(DataIndex, POINT_COLOR_ID, ElementIndex);
+	return GetFloat4_{DataInterfaceName}(DataIndex, ElementIndex, POINT_COLOR_ID);
 }
 
 float GetDensity_{DataInterfaceName}(uint DataIndex, uint ElementIndex)
 {
-	return GetFloat_{DataInterfaceName}(DataIndex, POINT_DENSITY_ID, ElementIndex);
+	return GetFloat_{DataInterfaceName}(DataIndex, ElementIndex, POINT_DENSITY_ID);
 }
 
 float GetSteepness_{DataInterfaceName}(uint DataIndex, uint ElementIndex)
 {
-	return GetFloat_{DataInterfaceName}(DataIndex, POINT_STEEPNESS_ID, ElementIndex);
+	return GetFloat_{DataInterfaceName}(DataIndex, ElementIndex, POINT_STEEPNESS_ID);
 }
 
 uint GetSeed_{DataInterfaceName}(uint DataIndex, uint ElementIndex)
 {
-	return GetUint_{DataInterfaceName}(DataIndex, POINT_SEED_ID, ElementIndex);
+	return GetUint_{DataInterfaceName}(DataIndex, ElementIndex, POINT_SEED_ID);
 }
 
 float4x4 GetPointTransform_{DataInterfaceName}(uint InDataIndex, uint InElementIndex)
@@ -395,47 +395,47 @@ bool IsValid_{DataInterfaceName}(uint InDataIndex, uint InElementIndex)
 
 void SetPosition_{DataInterfaceName}(uint DataIndex, uint ElementIndex, float3 Position)
 {
-	SetFloat3_{DataInterfaceName}(DataIndex, POINT_POSITION_ID, ElementIndex, Position);
+	SetFloat3_{DataInterfaceName}(DataIndex, ElementIndex, POINT_POSITION_ID, Position);
 }
 
 void SetRotation_{DataInterfaceName}(uint DataIndex, uint ElementIndex, float4 Rotation)
 {
-	SetFloat4_{DataInterfaceName}(DataIndex, POINT_ROTATION_ID, ElementIndex, Rotation);
+	SetFloat4_{DataInterfaceName}(DataIndex, ElementIndex, POINT_ROTATION_ID, Rotation);
 }
 
 void SetScale_{DataInterfaceName}(uint DataIndex, uint ElementIndex, float3 Scale)
 {
-	SetFloat3_{DataInterfaceName}(DataIndex, POINT_SCALE_ID, ElementIndex, Scale);
+	SetFloat3_{DataInterfaceName}(DataIndex, ElementIndex, POINT_SCALE_ID, Scale);
 }
 
 void SetBoundsMin_{DataInterfaceName}(uint DataIndex, uint ElementIndex, float3 BoundsMin)
 {
-	SetFloat3_{DataInterfaceName}(DataIndex, POINT_BOUNDS_MIN_ID, ElementIndex, BoundsMin);
+	SetFloat3_{DataInterfaceName}(DataIndex, ElementIndex, POINT_BOUNDS_MIN_ID, BoundsMin);
 }
 
 void SetBoundsMax_{DataInterfaceName}(uint DataIndex, uint ElementIndex, float3 BoundsMax)
 {
-	SetFloat3_{DataInterfaceName}(DataIndex, POINT_BOUNDS_MAX_ID, ElementIndex, BoundsMax);
+	SetFloat3_{DataInterfaceName}(DataIndex, ElementIndex, POINT_BOUNDS_MAX_ID, BoundsMax);
 }
 
 void SetColor_{DataInterfaceName}(uint DataIndex, uint ElementIndex, float4 Color)
 {
-	SetFloat4_{DataInterfaceName}(DataIndex, POINT_COLOR_ID, ElementIndex, Color);
+	SetFloat4_{DataInterfaceName}(DataIndex, ElementIndex, POINT_COLOR_ID, Color);
 }
 
 void SetDensity_{DataInterfaceName}(uint DataIndex, uint ElementIndex, float Density)
 {
-	SetFloat_{DataInterfaceName}(DataIndex, POINT_DENSITY_ID, ElementIndex, Density);
+	SetFloat_{DataInterfaceName}(DataIndex, ElementIndex, POINT_DENSITY_ID, Density);
 }
 
 void SetSteepness_{DataInterfaceName}(uint DataIndex, uint ElementIndex, float Steepness)
 {
-	SetFloat_{DataInterfaceName}(DataIndex, POINT_STEEPNESS_ID, ElementIndex, Steepness);
+	SetFloat_{DataInterfaceName}(DataIndex, ElementIndex, POINT_STEEPNESS_ID, Steepness);
 }
 
 void SetSeed_{DataInterfaceName}(uint DataIndex, uint ElementIndex, uint Seed)
 {
-	SetUint_{DataInterfaceName}(DataIndex, POINT_SEED_ID, ElementIndex, Seed);
+	SetUint_{DataInterfaceName}(DataIndex, ElementIndex, POINT_SEED_ID, Seed);
 }
 
 void SetPointTransform_{DataInterfaceName}(uint InDataIndex, uint InElementIndex, float4x4 Transform)
diff --git a/Engine/Plugins/PCG/Source/PCG/Private/Compute/DataInterfaces/PCGDataCollectionDataInterface.cpp b/Engine/Plugins/PCG/Source/PCG/Private/Compute/DataInterfaces/PCGDataCollectionDataInterface.cpp
--- a/Engine/Plugins/PCG/Source/PCG/Private/Compute/DataInterfaces/PCGDataCollectionDataInterface.cpp
+++ b/Engine/Plugins/PCG/Source/PCG/Private/Compute/DataInterfaces/PCGDataCollectionDataInterface.cpp
@@ -169,8 +169,8 @@ void UPCGDataCollectionDataInterface::GetSupportedInputs(TArray<FShaderFunctionD
 			.SetName(TEXT("GetBool"))
 			.AddReturnType(EShaderFundamentalType::Bool)
 			.AddParam(EShaderFundamentalType::Uint) // DataIndex
-			.AddParam(EShaderFundamentalType::Uint) // AttributeId
-			.AddParam(EShaderFundamentalType::Uint); // ElementIndex
+			.AddParam(EShaderFundamentalType::Uint) // ElementIndex
+			.AddParam(EShaderFundamentalType::Uint); // AttributeId
 
 		OutFunctions.AddDefaulted_GetRef()
 			.SetName(TEXT("GetInt"))
@@ -234,8 +234,8 @@ void UPCGDataCollectionDataInterface::GetSupportedInputs(TArray<FShaderFunctionD
 		OutFunctions.AddDefaulted_GetRef()
 			.SetName(TEXT("SetBool"))
 			.AddParam(EShaderFundamentalType::Uint) // DataIndex
-			.AddParam(EShaderFundamentalType::Uint) // AttributeId
 			.AddParam(EShaderFundamentalType::Uint) // ElementIndex
+			.AddParam(EShaderFundamentalType::Uint) // AttributeId
 			.AddParam(EShaderFundamentalType::Bool); // Value
 
 		OutFunctions.AddDefaulted_GetRef()
diff --git a/Engine/Plugins/PCG/Source/PCG/Private/Compute/Elements/PCGCustomHLSL.cpp b/Engine/Plugins/PCG/Source/PCG/Private/Compute/Elements/PCGCustomHLSL.cpp
--- a/Engine/Plugins/PCG/Source/PCG/Private/Compute/Elements/PCGCustomHLSL.cpp
+++ b/Engine/Plugins/PCG/Source/PCG/Private/Compute/Elements/PCGCustomHLSL.cpp
@@ -23,13 +23,13 @@
 
 namespace PCGHLSLElement
 {
-	/** First capture: Pin name(supports a - z, A - Z, and 0 - 9) */
+	/** First capture: Pin name (supports a - z, A - Z, and 0 - 9) */
 	constexpr int AttributePinCaptureGroup = 1;
 
-	/** Second capture: Function name(Get, Set, or Create) */
+	/** Second capture: Function name (Get or Set) */
 	constexpr int AttributeFunctionCaptureGroup = 2;
 
-	/** Third capture: Attribute type(e.g.Int, Float, Rotator, etc.) */
+	/** Third capture: Attribute type (e.g.Int, Float, Rotator, etc.) */
 	constexpr int AttributeTypeCaptureGroup = 3;
 
 	/** Fourth capture: Attribute name (supports a-z, A-Z, 0-9, ' ', '-', '_', and '/') */
@@ -564,8 +564,8 @@ void UPCGCustomHLSLSettings::UpdateDeclarations()
 
 			InputDeclarations += TEXT("\nuint <pin>_GetNumData();\n");
 			InputDeclarations += TEXT("uint <pin>_GetNumElements();\n");
-			InputDeclarations += TEXT("<type> <pin>_Get<type>(uint DataIndex, uint AttributeId, uint ElementIndex);\n");
-			InputDeclarations += TEXT("<type> <pin>_Get<type>(uint DataIndex, 'AttributeName', uint ElementIndex);\n");
+			InputDeclarations += TEXT("<type> <pin>_Get<type>(uint DataIndex, uint ElementIndex, uint AttributeId);\n");
+			InputDeclarations += TEXT("<type> <pin>_Get<type>(uint DataIndex, uint ElementIndex, 'AttributeName');\n");
 		}
 
 		if (!PointDataPins.IsEmpty())
@@ -643,8 +643,8 @@ void UPCGCustomHLSLSettings::UpdateDeclarations()
 			OutputDeclarations += TEXT("\n// Valid pins: ") + FString::Join(DataPins, TEXT(", ")) + TEXT("\n");
 			OutputDeclarations += TEXT("// Valid types: bool, int, float, float2, float3, float4, Rotator (float3), Quat (float4), Transform (float4x4)\n");
 
-			OutputDeclarations += TEXT("\nvoid <pin>_Set<type>(uint DataIndex, uint AttributeId, uint ElementIndex, <type> Value);\n");
-			OutputDeclarations += TEXT("void <pin>_Set<type>(uint DataIndex, 'AttributeName', uint ElementIndex, <type> Value);\n");
+			OutputDeclarations += TEXT("\nvoid <pin>_Set<type>(uint DataIndex, uint ElementIndex, uint AttributeId, <type> Value);\n");
+			OutputDeclarations += TEXT("void <pin>_Set<type>(uint DataIndex, uint ElementIndex, 'AttributeName', <type> Value);\n");
 		}
 
 		if (!PointDataPins.IsEmpty())
@@ -1166,9 +1166,9 @@ bool UPCGCustomHLSLSettings::AreKernelAttributesValid(FPCGContext* InContext, FT
 			const FPCGDataCollectionDesc* PinDesc = nullptr;
 
 			auto ConstructFunctionText = [&PinStr, &FuncStr, &TypeStr]()
-				{
-					return FText::FromString(PinStr + TEXT("_") + FuncStr + TypeStr);
-				};
+			{
+				return FText::FromString(PinStr + TEXT("_") + FuncStr + TypeStr);
+			};
 
 			if (FuncStr == PCGHLSLElement::AttributeFunctionSet)
 			{
