[PCG] Support Attribute Sets on GPU
* Supports upload, readback, and passing between GPU kernels
* Supports multi-element attribute sets
* Attribute Set pins get their buffer size from the FixedElementCount buffer size mode, and can reference an input pin to copy its attributes (in name only), and/or create new attributes of their own
* Cleaned up PCGDataCollectionDataInterface.ush, more consistent with Get/Set vs Read/Write, and marked some functions as internal that should not be called by typical users

diff --git a/Engine/Plugins/PCG/Shaders/Private/PCGDataCollectionDataInterface.ush b/Engine/Plugins/PCG/Shaders/Private/PCGDataCollectionDataInterface.ush
--- a/Engine/Plugins/PCG/Shaders/Private/PCGDataCollectionDataInterface.ush
+++ b/Engine/Plugins/PCG/Shaders/Private/PCGDataCollectionDataInterface.ush
@@ -4,7 +4,7 @@
 
 RWByteAddressBuffer {DataInterfaceName}_DataCollectionBuffer;
 
-// TODO: These defines could be code-gen'd to stay in sync with our PCGComputeConstants
+// TODO: These defines could be code-gen'd to stay in sync with our PCGComputeConstants, see UComputeDataInterface::GetDefines()
 #define POINT_POSITION_ID 0
 #define POINT_ROTATION_ID 1
 #define POINT_SCALE_ID 2
@@ -18,6 +18,18 @@ RWByteAddressBuffer {DataInterfaceName}_DataCollectionBuffer;
 
 #define INVALID_DENSITY 1.#INF
 
+// #################### DIRECT BUFFER OPERATIONS ##########################
+
+uint LoadBufferInternal_{DataInterfaceName}(uint Address)
+{
+	return {DataInterfaceName}_DataCollectionBuffer.Load(Address);
+}
+
+void StoreBufferInternal_{DataInterfaceName}(uint Address, uint Value)
+{
+	return {DataInterfaceName}_DataCollectionBuffer.Store(Address, Value);
+}
+
 // #################### HEADER READERS ##########################
 
 uint GetNumData_{DataInterfaceName}()
@@ -25,28 +37,28 @@ uint GetNumData_{DataInterfaceName}()
 	return {DataInterfaceName}_DataCollectionBuffer.Load(0);
 }
 
-uint ReadDataAddress_{DataInterfaceName}(uint DataIndex)
+uint GetDataAddress_{DataInterfaceName}(uint DataIndex)
 {
 	return {DataInterfaceName}_DataCollectionBuffer.Load((DataIndex + 1) * 4);
 }
 
-uint ReadDataId_{DataInterfaceName}(uint DataAddress)
+uint GetDataId_{DataInterfaceName}(uint DataAddress)
 {
 	return {DataInterfaceName}_DataCollectionBuffer.Load(DataAddress + 0);
 }
 
-uint ReadDataNumAttributes_{DataInterfaceName}(uint DataAddress)
+uint GetDataNumAttributes_{DataInterfaceName}(uint DataAddress)
 {
 	return {DataInterfaceName}_DataCollectionBuffer.Load(DataAddress + 4);
 }
 
-uint ReadDataPreambleSize_{DataInterfaceName}(uint DataAddress)
+uint GetDataPreambleSize_{DataInterfaceName}(uint DataAddress)
 {
 	return {DataInterfaceName}_DataCollectionBuffer.Load(DataAddress + 8);
 }
 
 // For PointData, this is just NumPoints
-uint ReadDataInfo_{DataInterfaceName}(uint DataAddress)
+uint GetDataInfo_{DataInterfaceName}(uint DataAddress)
 {
 	return {DataInterfaceName}_DataCollectionBuffer.Load(DataAddress + 12);
 }
@@ -58,8 +70,8 @@ bool GetThreadData_{DataInterfaceName}(uint InThreadIndex, out uint OutDataIndex
 	const uint NumData = GetNumData_{DataInterfaceName}();
 	for (uint DataIndex = 0; DataIndex < NumData; ++DataIndex)
 	{
-		const uint DataAddress = ReadDataAddress_{DataInterfaceName}(DataIndex);
-		const uint ElemCount = ReadDataInfo_{DataInterfaceName}(DataAddress);
+		const uint DataAddress = GetDataAddress_{DataInterfaceName}(DataIndex);
+		const uint ElemCount = GetDataInfo_{DataInterfaceName}(DataAddress);
 
 		if (ElementIndex < ElemCount)
 		{
@@ -93,55 +105,55 @@ uint GetNumElements_{DataInterfaceName}()
 	const uint NumData = GetNumData_{DataInterfaceName}();
 	for (uint DataIndex = 0; DataIndex < NumData; ++DataIndex)
 	{
-		const uint DataAddress = ReadDataAddress_{DataInterfaceName}(DataIndex);
-		NumElements += ReadDataInfo_{DataInterfaceName}(DataAddress);
+		const uint DataAddress = GetDataAddress_{DataInterfaceName}(DataIndex);
+		NumElements += GetDataInfo_{DataInterfaceName}(DataAddress);
 	}
 
 	return NumElements;
 }
 
-uint ReadDataAttributeHeadersAddress_{DataInterfaceName}(uint DataAddress)
+uint GetDataAttributeHeadersAddress_{DataInterfaceName}(uint DataAddress)
 {
-	return DataAddress + ReadDataPreambleSize_{DataInterfaceName}(DataAddress);
+	return DataAddress + GetDataPreambleSize_{DataInterfaceName}(DataAddress);
 }
 
-uint ReadAttributeHeaderAddress_{DataInterfaceName}(uint AttributeHeadersAddress, uint AttributeId)
+uint GetAttributeHeaderAddress_{DataInterfaceName}(uint AttributeHeadersAddress, uint AttributeId)
 {
 	return AttributeHeadersAddress + AttributeId * 8;
 }
 
-uint ReadAttributeIdAndStride_{DataInterfaceName}(uint AttributeHeaderAddress)
+uint GetAttributeIdAndStride_{DataInterfaceName}(uint AttributeHeaderAddress)
 {
 	return {DataInterfaceName}_DataCollectionBuffer.Load(AttributeHeaderAddress);
 }
 
-uint ReadAttributeId_{DataInterfaceName}(uint AttributeHeaderAddress)
+uint GetAttributeId_{DataInterfaceName}(uint AttributeHeaderAddress)
 {
-	return ReadAttributeIdAndStride_{DataInterfaceName}(AttributeHeaderAddress) >> 8;
+	return GetAttributeIdAndStride_{DataInterfaceName}(AttributeHeaderAddress) >> 8;
 }
 
-uint ReadAttributeStride_{DataInterfaceName}(uint AttributeHeaderAddress)
+uint GetAttributeStride_{DataInterfaceName}(uint AttributeHeaderAddress)
 {
-	return ReadAttributeIdAndStride_{DataInterfaceName}(AttributeHeaderAddress) & (0xFFFFFFFF >> 24);
+	return GetAttributeIdAndStride_{DataInterfaceName}(AttributeHeaderAddress) & (0xFFFFFFFF >> 24);
 }
 
-uint ReadAttributeAddress_{DataInterfaceName}(uint AttributeHeaderAddress)
+uint GetAttributeAddress_{DataInterfaceName}(uint AttributeHeaderAddress)
 {
 	return {DataInterfaceName}_DataCollectionBuffer.Load(AttributeHeaderAddress + 4);
 }
 
-uint ReadAttributeAddress_{DataInterfaceName}(uint InDataAddress, uint InAttributeId)
+uint GetAttributeAddress_{DataInterfaceName}(uint InDataAddress, uint InAttributeId)
 {
-	const uint AttributeHeadersAddress = ReadDataAttributeHeadersAddress_{DataInterfaceName}(InDataAddress);
-	const uint NumAttributes = ReadDataNumAttributes_{DataInterfaceName}(InDataAddress);
+	const uint AttributeHeadersAddress = GetDataAttributeHeadersAddress_{DataInterfaceName}(InDataAddress);
+	const uint NumAttributes = GetDataNumAttributes_{DataInterfaceName}(InDataAddress);
 
 	for (uint AttributeIndex = 0; AttributeIndex < NumAttributes; ++AttributeIndex)
 	{
-		const uint AttributeHeaderAddress = ReadAttributeHeaderAddress_{DataInterfaceName}(AttributeHeadersAddress, AttributeIndex);
+		const uint AttributeHeaderAddress = GetAttributeHeaderAddress_{DataInterfaceName}(AttributeHeadersAddress, AttributeIndex);
 
-		if (ReadAttributeId_{DataInterfaceName}(AttributeHeaderAddress) == InAttributeId)
+		if (GetAttributeId_{DataInterfaceName}(AttributeHeaderAddress) == InAttributeId)
 		{
-			return ReadAttributeAddress_{DataInterfaceName}(AttributeHeaderAddress);
+			return GetAttributeAddress_{DataInterfaceName}(AttributeHeaderAddress);
 		}
 	}
 
@@ -150,65 +162,55 @@ uint ReadAttributeAddress_{DataInterfaceName}(uint InDataAddress, uint InAttribu
 
 // #################### HEADER WRITERS ##########################
 
-void WriteNumData_{DataInterfaceName}(uint NumData)
+void SetNumDataInternal_{DataInterfaceName}(uint NumData)
 {
 	{DataInterfaceName}_DataCollectionBuffer.Store(0, NumData);
 }
 
-void WriteDataAddress_{DataInterfaceName}(uint DataIndex, uint DataAddress)
+void SetDataAddressInternal_{DataInterfaceName}(uint DataIndex, uint DataAddress)
 {
 	{DataInterfaceName}_DataCollectionBuffer.Store((DataIndex + 1) * 4, DataAddress);
 }
 
-void WriteDataId_{DataInterfaceName}(uint DataAddress, uint Id)
+void SetDataIdInternal_{DataInterfaceName}(uint DataAddress, uint Id)
 {
 	{DataInterfaceName}_DataCollectionBuffer.Store(DataAddress + 0, Id);
 }
 
-void WriteDataNumAttributes_{DataInterfaceName}(uint DataAddress, uint NumAttributes)
+void SetDataNumAttributesInternal_{DataInterfaceName}(uint DataAddress, uint NumAttributes)
 {
 	{DataInterfaceName}_DataCollectionBuffer.Store(DataAddress + 4, NumAttributes);
 }
 
-void WriteDataPreambleSize_{DataInterfaceName}(uint DataAddress, uint PreambleSize)
+void SetDataPreambleSizeInternal_{DataInterfaceName}(uint DataAddress, uint PreambleSize)
 {
 	{DataInterfaceName}_DataCollectionBuffer.Store(DataAddress + 8, PreambleSize);
 }
 
 // For PointData, this is just NumPoints
-void WriteDataInfo_{DataInterfaceName}(uint DataAddress, uint Info)
+void SetDataInfoInternal_{DataInterfaceName}(uint DataAddress, uint Info)
 {
 	{DataInterfaceName}_DataCollectionBuffer.Store(DataAddress + 12, Info);
 }
 
-void WriteAttributeIdAndStride_{DataInterfaceName}(uint AttributeHeaderAddress, uint IdAndStride)
+void SetAttributeIdAndStrideInternal_{DataInterfaceName}(uint AttributeHeaderAddress, uint IdAndStride)
 {
 	{DataInterfaceName}_DataCollectionBuffer.Store(AttributeHeaderAddress, IdAndStride);
 }
 
-void WriteAttributeAddress_{DataInterfaceName}(uint AttributeHeaderAddress, uint AttributeAddress)
+void SetAttributeAddressInternal_{DataInterfaceName}(uint AttributeHeaderAddress, uint AttributeAddress)
 {
 	{DataInterfaceName}_DataCollectionBuffer.Store(AttributeHeaderAddress + 4, AttributeAddress);
 }
 
-uint LoadBuffer_{DataInterfaceName}(uint Address)
-{
-	return {DataInterfaceName}_DataCollectionBuffer.Load(Address);
-}
-
-void StoreBuffer_{DataInterfaceName}(uint Address, uint Value)
-{
-	return {DataInterfaceName}_DataCollectionBuffer.Store(Address, Value);
-}
-
 // #################### ATTRIBUTE GETTERS ##########################
 
 uint ComputeAddressHelper_{DataInterfaceName}(uint DataIndex, uint ElementIndex, uint AttributeId)
 {
-	const uint DataAddress = ReadDataAddress_{DataInterfaceName}(DataIndex);
-	const uint AttributeHeadersAddress = ReadDataAttributeHeadersAddress_{DataInterfaceName}(DataAddress);
-	const uint AttributeHeaderAddress = ReadAttributeHeaderAddress_{DataInterfaceName}(AttributeHeadersAddress, AttributeId);
-	return ReadAttributeAddress_{DataInterfaceName}(AttributeHeaderAddress);
+	const uint DataAddress = GetDataAddress_{DataInterfaceName}(DataIndex);
+	const uint AttributeHeadersAddress = GetDataAttributeHeadersAddress_{DataInterfaceName}(DataAddress);
+	const uint AttributeHeaderAddress = GetAttributeHeaderAddress_{DataInterfaceName}(AttributeHeadersAddress, AttributeId);
+	return GetAttributeAddress_{DataInterfaceName}(AttributeHeaderAddress);
 }
 
 bool GetBool_{DataInterfaceName}(uint DataIndex, uint ElementIndex, uint AttributeId)
diff --git a/Engine/Plugins/PCG/Source/PCG/Private/Compute/DataInterfaces/PCGDataCollectionDataInterface.cpp b/Engine/Plugins/PCG/Source/PCG/Private/Compute/DataInterfaces/PCGDataCollectionDataInterface.cpp
--- a/Engine/Plugins/PCG/Source/PCG/Private/Compute/DataInterfaces/PCGDataCollectionDataInterface.cpp
+++ b/Engine/Plugins/PCG/Source/PCG/Private/Compute/DataInterfaces/PCGDataCollectionDataInterface.cpp
@@ -21,12 +21,12 @@
 void UPCGDataCollectionDataInterface::GetSupportedInputs(TArray<FShaderFunctionDefinition>& OutFunctions) const
 {
 	OutFunctions.AddDefaulted_GetRef()
-		.SetName(TEXT("LoadBuffer"))
+		.SetName(TEXT("LoadBufferInternal"))
 		.AddReturnType(EShaderFundamentalType::Uint)
 		.AddParam(EShaderFundamentalType::Uint);
 
 	OutFunctions.AddDefaulted_GetRef()
-		.SetName(TEXT("StoreBuffer"))
+		.SetName(TEXT("StoreBufferInternal"))
 		.AddParam(EShaderFundamentalType::Uint)
 		.AddParam(EShaderFundamentalType::Uint);
 
@@ -37,32 +37,32 @@ void UPCGDataCollectionDataInterface::GetSupportedInputs(TArray<FShaderFunctionD
 			.AddReturnType(EShaderFundamentalType::Uint);
 
 		OutFunctions.AddDefaulted_GetRef()
-			.SetName(TEXT("ReadDataAddress"))
+			.SetName(TEXT("GetDataAddress"))
 			.AddReturnType(EShaderFundamentalType::Uint)
 			.AddParam(EShaderFundamentalType::Uint);
 
 		OutFunctions.AddDefaulted_GetRef()
-			.SetName(TEXT("ReadDataId"))
+			.SetName(TEXT("GetDataId"))
 			.AddReturnType(EShaderFundamentalType::Uint)
 			.AddParam(EShaderFundamentalType::Uint);
 
 		OutFunctions.AddDefaulted_GetRef()
-			.SetName(TEXT("ReadDataNumAttributes"))
+			.SetName(TEXT("GetDataNumAttributes"))
 			.AddReturnType(EShaderFundamentalType::Uint)
 			.AddParam(EShaderFundamentalType::Uint);
 
 		OutFunctions.AddDefaulted_GetRef()
-			.SetName(TEXT("ReadDataPreambleSize"))
+			.SetName(TEXT("GetDataPreambleSize"))
 			.AddReturnType(EShaderFundamentalType::Uint)
 			.AddParam(EShaderFundamentalType::Uint);
 
 		OutFunctions.AddDefaulted_GetRef()
-			.SetName(TEXT("ReadDataInfo"))
+			.SetName(TEXT("GetDataInfo"))
 			.AddReturnType(EShaderFundamentalType::Uint)
 			.AddParam(EShaderFundamentalType::Uint);
 
 		OutFunctions.AddDefaulted_GetRef()
-			.SetName(TEXT("ReadDataAttributeHeadersAddress"))
+			.SetName(TEXT("GetDataAttributeHeadersAddress"))
 			.AddReturnType(EShaderFundamentalType::Uint)
 			.AddParam(EShaderFundamentalType::Uint);
 
@@ -86,33 +86,33 @@ void UPCGDataCollectionDataInterface::GetSupportedInputs(TArray<FShaderFunctionD
 			.AddReturnType(EShaderFundamentalType::Uint);
 
 		OutFunctions.AddDefaulted_GetRef()
-			.SetName(TEXT("ReadAttributeHeaderAddress"))
+			.SetName(TEXT("GetAttributeHeaderAddress"))
 			.AddReturnType(EShaderFundamentalType::Uint)
 			.AddParam(EShaderFundamentalType::Uint)
 			.AddParam(EShaderFundamentalType::Uint);
 
 		OutFunctions.AddDefaulted_GetRef()
-			.SetName(TEXT("ReadAttributeIdAndStride"))
+			.SetName(TEXT("GetAttributeIdAndStride"))
 			.AddReturnType(EShaderFundamentalType::Uint)
 			.AddParam(EShaderFundamentalType::Uint);
 
 		OutFunctions.AddDefaulted_GetRef()
-			.SetName(TEXT("ReadAttributeId"))
+			.SetName(TEXT("GetAttributeId"))
 			.AddReturnType(EShaderFundamentalType::Uint)
 			.AddParam(EShaderFundamentalType::Uint);
 
 		OutFunctions.AddDefaulted_GetRef()
-			.SetName(TEXT("ReadAttributeStride"))
+			.SetName(TEXT("GetAttributeStride"))
 			.AddReturnType(EShaderFundamentalType::Uint)
 			.AddParam(EShaderFundamentalType::Uint);
 
 		OutFunctions.AddDefaulted_GetRef()
-			.SetName(TEXT("ReadAttributeAddress"))
+			.SetName(TEXT("GetAttributeAddress"))
 			.AddReturnType(EShaderFundamentalType::Uint)
 			.AddParam(EShaderFundamentalType::Uint);
 
 		OutFunctions.AddDefaulted_GetRef()
-			.SetName(TEXT("ReadAttributeAddress"))
+			.SetName(TEXT("GetAttributeAddress"))
 			.AddReturnType(EShaderFundamentalType::Uint)
 			.AddParam(EShaderFundamentalType::Uint)
 			.AddParam(EShaderFundamentalType::Uint);
@@ -121,41 +121,41 @@ void UPCGDataCollectionDataInterface::GetSupportedInputs(TArray<FShaderFunctionD
 	// Header Writers
 	{
 		OutFunctions.AddDefaulted_GetRef()
-			.SetName(TEXT("WriteNumData"))
+			.SetName(TEXT("SetNumDataInternal"))
 			.AddParam(EShaderFundamentalType::Uint);
 
 		OutFunctions.AddDefaulted_GetRef()
-			.SetName(TEXT("WriteDataAddress"))
+			.SetName(TEXT("SetDataAddressInternal"))
 			.AddParam(EShaderFundamentalType::Uint)
 			.AddParam(EShaderFundamentalType::Uint);
 
 		OutFunctions.AddDefaulted_GetRef()
-			.SetName(TEXT("WriteDataId"))
+			.SetName(TEXT("SetDataIdInternal"))
 			.AddParam(EShaderFundamentalType::Uint)
 			.AddParam(EShaderFundamentalType::Uint);
 
 		OutFunctions.AddDefaulted_GetRef()
-			.SetName(TEXT("WriteDataNumAttributes"))
+			.SetName(TEXT("SetDataNumAttributesInternal"))
 			.AddParam(EShaderFundamentalType::Uint)
 			.AddParam(EShaderFundamentalType::Uint);
 
 		OutFunctions.AddDefaulted_GetRef()
-			.SetName(TEXT("WriteDataPreambleSize"))
+			.SetName(TEXT("SetDataPreambleSizeInternal"))
 			.AddParam(EShaderFundamentalType::Uint)
 			.AddParam(EShaderFundamentalType::Uint);
 
 		OutFunctions.AddDefaulted_GetRef()
-			.SetName(TEXT("WriteDataInfo"))
+			.SetName(TEXT("SetDataInfoInternal"))
 			.AddParam(EShaderFundamentalType::Uint)
 			.AddParam(EShaderFundamentalType::Uint);
 
 		OutFunctions.AddDefaulted_GetRef()
-			.SetName(TEXT("WriteAttributeIdAndStride"))
+			.SetName(TEXT("SetAttributeIdAndStrideInternal"))
 			.AddParam(EShaderFundamentalType::Uint)
 			.AddParam(EShaderFundamentalType::Uint);
 
 		OutFunctions.AddDefaulted_GetRef()
-			.SetName(TEXT("WriteAttributeAddress"))
+			.SetName(TEXT("SetAttributeAddressInternal"))
 			.AddParam(EShaderFundamentalType::Uint)
 			.AddParam(EShaderFundamentalType::Uint);
 
diff --git a/Engine/Plugins/PCG/Source/PCG/Private/Compute/Elements/PCGComputeGraphElement.cpp b/Engine/Plugins/PCG/Source/PCG/Private/Compute/Elements/PCGComputeGraphElement.cpp
--- a/Engine/Plugins/PCG/Source/PCG/Private/Compute/Elements/PCGComputeGraphElement.cpp
+++ b/Engine/Plugins/PCG/Source/PCG/Private/Compute/Elements/PCGComputeGraphElement.cpp
@@ -158,7 +158,7 @@ bool FPCGComputeGraphElement::ExecuteInternal(FPCGContext* InContext) const
 			{
 				const UPCGCustomHLSLSettings* KernelSettings = CastChecked<UPCGCustomHLSLSettings>(Settings);
 
-				if (!KernelSettings->IsKernelValid(Context))
+				if (!KernelSettings->IsKernelValid(Context, /*bQuiet=*/false))
 				{
 					return true;
 				}
diff --git a/Engine/Plugins/PCG/Source/PCG/Private/Compute/Elements/PCGCustomHLSL.cpp b/Engine/Plugins/PCG/Source/PCG/Private/Compute/Elements/PCGCustomHLSL.cpp
--- a/Engine/Plugins/PCG/Source/PCG/Private/Compute/Elements/PCGCustomHLSL.cpp
+++ b/Engine/Plugins/PCG/Source/PCG/Private/Compute/Elements/PCGCustomHLSL.cpp
@@ -21,6 +21,12 @@
 #define LOCTEXT_NAMESPACE "PCGCustomHLSLElement"
 #define PCG_LOGGING_ENABLED (!(UE_BUILD_SHIPPING || UE_BUILD_TEST) || USE_LOGGING_IN_SHIPPING)
 
+#if PCG_LOGGING_ENABLED
+#define PCG_LOG_VALIDATION(ValidationMessage) LogGraphError(ValidationMessage);
+#else
+#define PCG_LOG_VALIDATION(ValidationMessage) // Log removed
+#endif
+
 namespace PCGHLSLElement
 {
 	/** First capture: Pin name (supports a - z, A - Z, and 0 - 9) */
@@ -50,29 +56,35 @@ namespace PCGHLSLElement
 		InOutPath.ReplaceCharInline(TEXT(':'), TEXT('@'));
 	}
 
-	int GetElementCount(const UPCGData* InData)
+	FString GetKernelAttributeKeyAsString(const FPCGKernelAttributeKey& Key)
 	{
-		if (const UPCGPointData* PointData = Cast<UPCGPointData>(InData))
-		{
-			return PointData->GetPoints().Num();
-		}
-		else if (InData)
-		{
-			// TODO - spline points, other things with multiple elements?
-			return 1;
-		}
-		else
-		{
-			return 0;
-		}
+		return FString::Format(TEXT("'{0}'"), { Key.Name.ToString() });
 	}
 
-	FString GetKernelAttributeKeyAsString(const FPCGKernelAttributeKey& Key)
+	FString GetDataTypeString(EPCGDataType Type)
 	{
-		return FString::Format(TEXT("'{0}'"), { Key.Name.ToString() });
+		const UEnum* DataTypeEnum = StaticEnum<EPCGDataType>();
+		check(DataTypeEnum);
+
+		return DataTypeEnum->GetValueOrBitfieldAsString(static_cast<int64>(Type));
 	}
 }
 
+#if WITH_EDITOR
+bool FPCGPinPropertiesGPU::CanEditChange(const FEditPropertyChain& PropertyChain) const
+{
+	if (FProperty* Property = PropertyChain.GetActiveNode()->GetValue())
+	{
+		if (Property->GetFName() == GET_MEMBER_NAME_CHECKED(FPCGPinProperties, bAllowMultipleData))
+		{
+			return bAllowEditMultipleData;
+		}
+	}
+
+	return true;
+}
+#endif // WITH_EDITOR
+
 UPCGCustomHLSLSettings::UPCGCustomHLSLSettings()
 {
 	bExecuteOnGPU = true;
@@ -142,7 +154,7 @@ int UPCGCustomHLSLSettings::GetProcessingElemCountForInputPin(const UPCGPin* Inp
 		TArray<FPCGTaggedData> DataForPin = DataForGPU.InputDataCollection.GetInputsByPin(PinLabel);
 		for (const FPCGTaggedData& Data : DataForPin)
 		{
-			ProcessingElemCount += PCGHLSLElement::GetElementCount(Data.Data);
+			ProcessingElemCount += PCGComputeHelpers::GetElementCount(Data.Data);
 		}
 	}
 	else
@@ -161,9 +173,7 @@ int UPCGCustomHLSLSettings::GetProcessingElemCountForInputPin(const UPCGPin* Inp
 			check(UpstreamSettings);
 
 			const FPCGDataCollectionDesc PinDesc = UpstreamSettings->ComputeOutputPinDataDesc(UpstreamOutputPin, Binding);
-
-			// Only support Point data right now.
-			ProcessingElemCount += PinDesc.ComputeDataElementCount(EPCGDataType::Point);
+			ProcessingElemCount += PinDesc.ComputeDataElementCount(InputPin->Properties.AllowedTypes);
 		}
 	}
 
@@ -245,65 +255,70 @@ FPCGDataCollectionDesc UPCGCustomHLSLSettings::ComputeOutputPinDataDesc(const UP
 	FPCGDataCollectionDesc PinDesc;
 
 	const FPCGPinPropertiesGPU* PropertiesGPU = GetOutputPinPropertiesGPU(OutputPin->Properties.Label);
-	const bool bNeedsPointPin = KernelType == EPCGKernelType::PointProcessor || KernelType == EPCGKernelType::PointGenerator;
+	const UPCGPin* FirstOutputPin = GetFirstOutputPin();
 
 	// The primary output pin follows any rules prescribed by kernel type.
-	if (OutputPin == GetFirstOutputPin() && bNeedsPointPin)
+	if (OutputPin == FirstOutputPin && KernelType == EPCGKernelType::PointProcessor)
 	{
-		if (KernelType == EPCGKernelType::PointProcessor)
+		// First output pin passes through first input pin.
+		if (const UPCGPin* PointProcessingInputPin = GetPointProcessingInputPin())
 		{
-			// First output pin passes through first input pin.
-			if (const UPCGPin* PointProcessingInputPin = GetPointProcessingInputPin())
-			{
-				PinDesc = ComputeInputPinDataDesc(PointProcessingInputPin, Binding);
-			}
-		}
-		else if (KernelType == EPCGKernelType::PointGenerator)
-		{
-			// Generators always produce a single point data with known point count.
-			PinDesc.DataDescs.Emplace(EPCGDataType::Point, PointCount);
+			PinDesc = ComputeInputPinDataDesc(PointProcessingInputPin, Binding);
 		}
 	}
+	else if (OutputPin == FirstOutputPin && KernelType == EPCGKernelType::PointGenerator)
+	{
+		// Generators always produce a single point data with known point count.
+		PinDesc.DataDescs.Emplace(EPCGDataType::Point, PointCount);
+	}
 	else if (ensure(PropertiesGPU))
 	{
 		// No size set by kernel, fall back to pin settings.
 		if (PropertiesGPU->BufferSizeMode == EPCGPinBufferSizeMode::FromFirstPin)
 		{
-			if (const UPCGPin* InputPin = bNeedsPointPin ? GetPointProcessingInputPin() : GetFirstInputPin())
+			if (const UPCGPin* InputPin = GetFirstInputPin())
 			{
 				PinDesc = ComputeInputPinDataDesc(InputPin, Binding);
 			}
 		}
 		else if (PropertiesGPU->BufferSizeMode == EPCGPinBufferSizeMode::FromProductOfInputPins)
 		{
-			// TODO for now we only support point. Needs extended logic for other data types.
-			const EPCGDataType DataType = EPCGDataType::Point;
-
 			int TotalElementCount = 0;
 
 			for (const FName& PinLabel : PropertiesGPU->BufferSizeInputPinLabels)
 			{
 				if (const UPCGPin* InputPin = GetInputPin(PinLabel))
 				{
-					const int ElementCount = ComputeInputPinDataDesc(InputPin, Binding).ComputeDataElementCount(DataType);
+					const int ElementCount = ComputeInputPinDataDesc(InputPin, Binding).ComputeDataElementCount(PropertiesGPU->AllowedTypes);
 					TotalElementCount = FMath::Max(TotalElementCount, 1) * ElementCount;
 				}
 			}
 
 			if (TotalElementCount > 0)
 			{
-				PinDesc.DataDescs.Emplace(DataType, TotalElementCount);
+				PinDesc.DataDescs.Emplace(PropertiesGPU->AllowedTypes, TotalElementCount);
 			}
 		}
 		else if (PropertiesGPU->BufferSizeMode == EPCGPinBufferSizeMode::FixedElementCount)
 		{
 			if (ensure(PropertiesGPU->FixedBufferElementCount > 0))
 			{
-				// TODO for now we only support point. Needs extended logic for other data types.
-				const EPCGDataType DataType = EPCGDataType::Point;
-				PinDesc.DataDescs.Emplace(DataType, PropertiesGPU->FixedBufferElementCount);
+				const UPCGPin* InitializeFromPin = (PropertiesGPU->AllowedTypes == EPCGDataType::Param) ? GetInputPin(PropertiesGPU->InitializeFromPin) : nullptr;
+
+				if (InitializeFromPin)
+				{
+					PinDesc = ComputeInputPinDataDesc(InitializeFromPin, Binding);
+				}
+				else
+				{
+					PinDesc.DataDescs.Emplace(PropertiesGPU->AllowedTypes, PropertiesGPU->FixedBufferElementCount);
+				}
 			}
 		}
+		else
+		{
+			checkNoEntry();
+		}
 	}
 
 	const TMap<FPCGKernelAttributeKey, int32>& GlobalAttributeLookupTable = ensure(Binding && Binding->Graph) ? Binding->Graph->GetAttributeLookupTable() : TMap<FPCGKernelAttributeKey, int32>();
@@ -696,10 +711,13 @@ void UPCGCustomHLSLSettings::UpdatePinSettings()
 		}
 
 		if (!!(Properties.AllowedTypes & EPCGDataType::Landscape)
-			|| !!(Properties.AllowedTypes & EPCGDataType::Texture))
+			|| !!(Properties.AllowedTypes & EPCGDataType::Texture)
+			|| !!(Properties.AllowedTypes & EPCGDataType::Param))
 		{
 			// Don't allow multiple data on this pin because we do not support a dynamic number of textures/landscapes bound to a
 			// compute kernel.
+			// Also disallow multi-data for Attribute Sets, since we require attributes to be uniform on a pin, thus having different
+			// attribute sets wouldn't work as expected.
 			Properties.bAllowMultipleData = false;
 			Properties.bAllowEditMultipleData = false;
 		}
@@ -734,8 +752,18 @@ void UPCGCustomHLSLSettings::UpdatePinSettings()
 		// TODO this could be hoisted up somewhere in the future.
 		Properties.bAllowEditMultipleConnections = false;
 
-		// Node is free to output any number of data.
-		Properties.bAllowEditMultipleData = false;
+		if (!!(Properties.AllowedTypes & EPCGDataType::Param))
+		{
+			Properties.bAllowMultipleData = false;
+			Properties.bAllowEditMultipleData = false;
+			Properties.bAllowEditInitializationPin = true;
+			Properties.BufferSizeMode = EPCGPinBufferSizeMode::FixedElementCount;
+		}
+		else
+		{
+			Properties.bAllowEditMultipleData = true;
+			Properties.bAllowEditInitializationPin = false;
+		}
 	}
 }
 
@@ -816,122 +844,187 @@ void UPCGCustomHLSLSettings::UpdateAttributeKeys()
 }
 #endif
 
-bool UPCGCustomHLSLSettings::IsKernelValid(FPCGContext* InContext) const
+bool UPCGCustomHLSLSettings::IsKernelValid(FPCGContext* InContext, bool bQuiet) const
 {
-	auto LogGraphError = [InContext, This=this](const FText& InText)
+	auto LogGraphError = [InContext, bQuiet, This=this](const FText& InText)
 	{
-#if WITH_EDITOR
-		if (InContext && ensure(InContext->SourceComponent.IsValid() && InContext->SourceComponent.Get()))
+		if (!bQuiet)
 		{
-			if (UPCGSubsystem* Subsystem = InContext->SourceComponent->GetSubsystem())
+#if WITH_EDITOR
+			if (InContext && ensure(InContext->SourceComponent.IsValid() && InContext->SourceComponent.Get()))
 			{
-				FPCGStack StackWithNode = InContext->Stack ? *InContext->Stack : FPCGStack();
-				StackWithNode.PushFrame(This->GetOuter());
+				if (UPCGSubsystem* Subsystem = InContext->SourceComponent->GetSubsystem())
+				{
+					FPCGStack StackWithNode = InContext->Stack ? *InContext->Stack : FPCGStack();
+					StackWithNode.PushFrame(This->GetOuter());
 
-				Subsystem->GetNodeVisualLogsMutable().Log(StackWithNode, ELogVerbosity::Error, InText);
+					Subsystem->GetNodeVisualLogsMutable().Log(StackWithNode, ELogVerbosity::Error, InText);
+				}
 			}
-		}
 #endif
 
-		PCGE_LOG_C(Error, LogOnly, InContext, InText);
+			PCGE_LOG_C(Error, LogOnly, InContext, InText);
+		}
 	};
 
 	if (OutputPins.IsEmpty())
 	{
-#if PCG_LOGGING_ENABLED
-		if (InContext)
-		{
-			LogGraphError(LOCTEXT("NoOutputs", "Custom HLSL nodes must have at least one output."));
-		}
-#endif
-
+		PCG_LOG_VALIDATION(LOCTEXT("NoOutputs", "Custom HLSL nodes must have at least one output."));
 		return false;
 	}
 
-	for (const FPCGPinProperties& Properties : InputPins)
+	auto CheckPinLabel = [InContext, &InPins = InputPins, OutPins = OutputPinProperties(), &LogGraphError](FName PinLabel)
 	{
-		if (Properties.AllowedTypes == EPCGDataType::Any)
+		if (PinLabel == NAME_None)
 		{
-#if PCG_LOGGING_ENABLED
-			if (InContext)
+			PCG_LOG_VALIDATION(LOCTEXT("InvalidPinLabelNone", "Pin label 'None' is not a valid pin label."));
+			return false;
+		}
+
+		bool bFoundPinLabel = false;
+
+		auto IsAlreadyFound = [InContext, &LogGraphError, PinLabel, &bFoundPinLabel](const FPCGPinProperties PinProps)
+		{
+			if (PinProps.Label == PinLabel)
 			{
-				LogGraphError(FText::Format(LOCTEXT("InvalidAnyInput", "Custom HLSL nodes do not support inputs of type Any, found on pin '{0}'."), FText::FromName(Properties.Label)));
+				if (bFoundPinLabel)
+				{
+					PCG_LOG_VALIDATION(FText::Format(LOCTEXT("DuplicatedPinLabels", "Duplicate pin label '{0}', all labels must be unique."), FText::FromName(PinLabel)));
+					return true;
+				}
+
+				bFoundPinLabel = true;
 			}
-#endif
 
 			return false;
+		};
+
+		for (const FPCGPinProperties& PinProps : InPins)
+		{
+			if (IsAlreadyFound(PinProps))
+			{
+				return false;
+			}
 		}
-	}
 
-	bool bIsFirstOutputPin = true;
-	for (const FPCGPinPropertiesGPU& Properties : OutputPins)
-	{
-		if (Properties.AllowedTypes == EPCGDataType::Any)
+		for (const FPCGPinProperties& PinProps : OutPins)
 		{
-#if PCG_LOGGING_ENABLED
-			if (InContext)
+			if (IsAlreadyFound(PinProps))
 			{
-				LogGraphError(FText::Format(LOCTEXT("InvalidAnyOutput", "Custom HLSL nodes do not support outputs of type Any, found on pin '{0}'."), FText::FromName(Properties.Label)));
+				return false;
 			}
-#endif
+		}
+
+		return true;
+	};
 
+	// Validate input pins
+	bool bIsFirstInputPin = true;
+	for (const FPCGPinProperties& Properties : InputPins)
+	{
+		if (!CheckPinLabel(Properties.Label))
+		{
 			return false;
 		}
 
-		if (!!(Properties.AllowedTypes & EPCGDataType::Landscape))
+		if (bIsFirstInputPin && KernelType == EPCGKernelType::PointProcessor)
 		{
-#if PCG_LOGGING_ENABLED
-			if (InContext)
+			if (Properties.AllowedTypes != EPCGDataType::Point)
 			{
-				LogGraphError(FText::Format(LOCTEXT("InvalidLSOutput", "Custom HLSL nodes do not support outputs of type Landscape, found on pin '{0}'."), FText::FromName(Properties.Label)));
+				PCG_LOG_VALIDATION(FText::Format(
+					LOCTEXT("InvalidNonPointPrimaryInput", "'Point Processor' nodes require primary input pin to be of type 'Point', but found '{0}'."),
+					FText::FromString(PCGHLSLElement::GetDataTypeString(Properties.AllowedTypes))));
+
+				return false;
 			}
-#endif
+		}
+
+		if (!PCGComputeHelpers::IsTypeAllowedAsInput(Properties.AllowedTypes))
+		{
+			PCG_LOG_VALIDATION(FText::Format(
+				LOCTEXT("InvalidInputType", "Unsupported input type '{0}', found on pin '{1}'."),
+				FText::FromString(PCGHLSLElement::GetDataTypeString(Properties.AllowedTypes)),
+				FText::FromName(Properties.Label)));
 
 			return false;
 		}
 
-		if (!!(Properties.AllowedTypes & EPCGDataType::Texture))
+		bIsFirstInputPin = false;
+	}
+
+	// Validate output pins
+	bool bIsFirstOutputPin = true;
+	for (const FPCGPinPropertiesGPU& Properties : OutputPins)
+	{
+		if (!CheckPinLabel(Properties.Label))
 		{
-#if PCG_LOGGING_ENABLED
-			if (InContext)
+			return false;
+		}
+
+		const bool bPinIsDefinedByKernel = bIsFirstOutputPin && (KernelType == EPCGKernelType::PointGenerator || KernelType == EPCGKernelType::PointProcessor);
+
+		if (bPinIsDefinedByKernel)
+		{
+			if (Properties.AllowedTypes != EPCGDataType::Point)
 			{
-				LogGraphError(FText::Format(LOCTEXT("InvalidTextureOutput", "Custom HLSL nodes do not support outputs of type Texture, found on pin '{0}'."), FText::FromName(Properties.Label)));
+				PCG_LOG_VALIDATION(FText::Format(
+					LOCTEXT("InvalidNonPointPrimaryOutput", "'Point Processor' and 'Point Generator' nodes require primary output pin to be of type 'Point', but found '{0}'."),
+					FText::FromString(PCGHLSLElement::GetDataTypeString(Properties.AllowedTypes))));
+
+				return false;
 			}
-#endif
+		}
+
+		if (!PCGComputeHelpers::IsTypeAllowedAsOutput(Properties.AllowedTypes))
+		{
+			PCG_LOG_VALIDATION(FText::Format(
+				LOCTEXT("InvalidOutputType", "Unsupported output type '{0}', found on pin '{1}'."),
+				FText::FromString(PCGHLSLElement::GetDataTypeString(Properties.AllowedTypes)),
+				FText::FromName(Properties.Label)));
 
 			return false;
 		}
 
-		// First output pin of point generator and point processor has fixed settings.
-		if ((KernelType != EPCGKernelType::PointGenerator && KernelType != EPCGKernelType::PointProcessor) || !bIsFirstOutputPin)
+		if (!bPinIsDefinedByKernel)
 		{
 			if (Properties.BufferSizeMode == EPCGPinBufferSizeMode::FixedElementCount)
 			{
 				if (Properties.FixedBufferElementCount <= 0)
 				{
-#if PCG_LOGGING_ENABLED
-					if (InContext)
-					{
-						LogGraphError(FText::Format(
-							LOCTEXT("InvalidFixedBufferSize", "Fixed GPU buffer size on '{0}' was invalid (%d)."),
-							FText::FromName(Properties.Label),
-							Properties.FixedBufferElementCount));
-					}
-#endif
+					PCG_LOG_VALIDATION(FText::Format(
+						LOCTEXT("InvalidFixedBufferSize", "Fixed GPU buffer size on '{0}' was invalid (%d)."),
+						FText::FromName(Properties.Label),
+						Properties.FixedBufferElementCount));
+
+					return false;
+				}
+
+				if (Properties.AllowedTypes == EPCGDataType::Param && Properties.InitializeFromPin != NAME_None && !GetInputPin(Properties.InitializeFromPin))
+				{
+					PCG_LOG_VALIDATION(FText::Format(
+						LOCTEXT("InvalidInitFromPin", "Tried to initialize attribute set pin '{0}' from non-existent pin '{1}'. Must reference a valid input pin or be 'None'."),
+						FText::FromName(Properties.Label),
+						FText::FromName(Properties.InitializeFromPin)));
+
+					return false;
 				}
 			}
 			else if (Properties.BufferSizeMode == EPCGPinBufferSizeMode::FromFirstPin)
 			{
 				if (InputPins.IsEmpty())
 				{
-#if PCG_LOGGING_ENABLED
-					if (InContext)
-					{
-						LogGraphError(FText::Format(
-							LOCTEXT("InvalidBufferSizeNoInputPin", "GPU buffer size for pin '{0}' could not be computed as there are no input pins."),
-							FText::FromName(Properties.Label)));
-					}
-#endif
+					PCG_LOG_VALIDATION(FText::Format(
+						LOCTEXT("InvalidBufferSizeNoInputPin", "GPU buffer size for pin '{0}' could not be computed as there are no input pins."),
+						FText::FromName(Properties.Label)));
+
+					return false;
+				}
+
+				if (!GetFirstInputPin())
+				{
+					PCG_LOG_VALIDATION(FText::Format(
+						LOCTEXT("MissingPrimaryInputPin", "GPU buffer size for pin '{0}' could not be computed, because it refers to the primary input pin, which does not exist."),
+						FText::FromName(Properties.Label)));
 
 					return false;
 				}
@@ -940,28 +1033,18 @@ bool UPCGCustomHLSLSettings::IsKernelValid(FPCGContext* InContext) const
 			{
 				if (InputPins.IsEmpty())
 				{
-#if PCG_LOGGING_ENABLED
-					if (InContext)
-					{
-						LogGraphError(FText::Format(
-							LOCTEXT("InvalidBufferSizeNoInputPins", "GPU buffer size for pin '{0}' could not be computed as there are no input pins on this node."),
-							FText::FromName(Properties.Label)));
-					}
-#endif
+					PCG_LOG_VALIDATION(FText::Format(
+						LOCTEXT("InvalidBufferSizeNoInputPins", "GPU buffer size for pin '{0}' could not be computed as there are no input pins on this node."),
+						FText::FromName(Properties.Label)));
 
 					return false;
 				}
 
 				if (Properties.BufferSizeInputPinLabels.IsEmpty())
 				{
-#if PCG_LOGGING_ENABLED
-					if (InContext)
-					{
-						LogGraphError(FText::Format(
-							LOCTEXT("InvalidBufferSizeNoBufferPins", "GPU buffer size for pin '{0}' could not be computed as input pins are specified in the pin settings."),
-							FText::FromName(Properties.Label)));
-					}
-#endif
+					PCG_LOG_VALIDATION(FText::Format(
+						LOCTEXT("InvalidBufferSizeNoBufferPins", "GPU buffer size for pin '{0}' could not be computed as input pins are specified in the pin settings."),
+						FText::FromName(Properties.Label)));
 
 					return false;
 				}
@@ -970,15 +1053,10 @@ bool UPCGCustomHLSLSettings::IsKernelValid(FPCGContext* InContext) const
 				{
 					if (!GetInputPin(Label))
 					{
-#if PCG_LOGGING_ENABLED
-						if (InContext)
-						{
-							LogGraphError(FText::Format(
-								LOCTEXT("MissingBufferSizePin", "GPU buffer size for pin '{0}' could not be computed. Invalid pin specified in Input Pins array: '{1}'."),
-								FText::FromName(Properties.Label),
-								FText::FromName(Label)));
-						}
-#endif
+						PCG_LOG_VALIDATION(FText::Format(
+							LOCTEXT("MissingBufferSizePin", "GPU buffer size for pin '{0}' could not be computed. Invalid pin specified in Input Pins array: '{1}'."),
+							FText::FromName(Properties.Label),
+							FText::FromName(Label)));
 
 						return false;
 					}
@@ -989,81 +1067,11 @@ bool UPCGCustomHLSLSettings::IsKernelValid(FPCGContext* InContext) const
 		bIsFirstOutputPin = false;
 	}
 
-	if (KernelType == EPCGKernelType::PointProcessor)
-	{
-		if (!GetPointProcessingInputPin())
-		{
-#if PCG_LOGGING_ENABLED
-			if (InContext)
-			{
-				LogGraphError(LOCTEXT("InvalidPPInput", "Point processing kernel requires a first input pin of type point."));
-			}
-#endif
-
-			return false;
-		}
-
-		if (!GetFirstPointOutputPin())
-		{
-#if PCG_LOGGING_ENABLED
-			if (InContext)
-			{
-				LogGraphError(LOCTEXT("InvalidPPOutput", "Point processing kernel requires a first output pin of type point."));
-			}
-#endif
-
-			return false;
-		}
-	}
-
-	auto CheckPinLabels = [InContext, &LogGraphError](const TArray<FPCGPinProperties>& Pins)
-	{
-		TSet<FName> EncounteredLabels;
-		for (const FPCGPinProperties& Pin : Pins)
-		{
-			if (EncounteredLabels.Contains(Pin.Label))
-			{
-#if PCG_LOGGING_ENABLED
-				if (InContext)
-				{
-					LogGraphError(FText::Format(LOCTEXT("DuplicatedPinLabels", "Duplicate pin label '{0}', all labels must be unique."), FText::FromName(Pin.Label)));
-				}
-#endif
-
-				return false;
-			}
-			else if (Pin.Label == NAME_None)
-			{
-#if PCG_LOGGING_ENABLED
-				if (InContext)
-				{
-					LogGraphError(LOCTEXT("InvalidPinLabelNone", "Pin label 'None' is not a valid pin label."));
-				}
-#endif
-
-				return false;
-			}
-		}
-
-		return true;
-	};
-
-	if (!CheckPinLabels(InputPins) || !CheckPinLabels(OutputPinProperties()))
-	{
-		return false;
-	}
-
 	if (KernelType == EPCGKernelType::Custom && DispatchThreadCount == EPCGDispatchThreadCount::FromProductOfInputPins)
 	{
 		if (ThreadCountInputPinLabels.IsEmpty())
 		{
-#if PCG_LOGGING_ENABLED
-			if (InContext)
-			{
-				LogGraphError(LOCTEXT("MissingThreadCountPins", "Dispatch thread count is based on input pins but no labels have been set in Input Pins array."));
-			}
-#endif
-
+			PCG_LOG_VALIDATION(LOCTEXT("MissingThreadCountPins", "Dispatch thread count is based on input pins but no labels have been set in Input Pins array."));
 			return false;
 		}
 
@@ -1071,29 +1079,18 @@ bool UPCGCustomHLSLSettings::IsKernelValid(FPCGContext* InContext) const
 		{
 			if (!GetInputPin(Label))
 			{
-#if PCG_LOGGING_ENABLED
-				if (InContext)
-				{
-					LogGraphError(FText::Format(LOCTEXT("MissingThreadCountPin", "Invalid pin specified in Input Pins array: '{0}'."), FText::FromName(Label)));
-				}
-#endif
-
+				PCG_LOG_VALIDATION(FText::Format(LOCTEXT("MissingThreadCountPin", "Invalid pin specified in Input Pins array: '{0}'."), FText::FromName(Label)));
 				return false;
 			}
 		}
 	}
 
+	// Validate attributes
 	for (const FPCGKernelAttributeKey& AttributeKey : KernelAttributeKeys)
 	{
 		if (AttributeKey.Name == NAME_None)
 		{
-#if PCG_LOGGING_ENABLED
-			if (InContext)
-			{
-				LogGraphError(LOCTEXT("InvalidAttributeNameNone", "'None' is not a valid GPU attribute name, check the 'Attributes to Create' array on your pins."));
-			}
-#endif
-
+			PCG_LOG_VALIDATION(LOCTEXT("InvalidAttributeNameNone", "'None' is not a valid GPU attribute name, check the 'Attributes to Create' array on your pins."));
 			return false;
 		}
 	}
@@ -1110,7 +1107,7 @@ bool UPCGCustomHLSLSettings::IsKernelValid(FPCGContext* InContext) const
 		{
 			if (ErrorTextPtr)
 			{
-				LogGraphError(*ErrorTextPtr);
+				PCG_LOG_VALIDATION(*ErrorTextPtr);
 			}
 			return false;
 		}
@@ -1336,16 +1333,17 @@ FString UPCGCustomHLSLSettings::GetCookedKernelSource(const TMap<FPCGKernelAttri
 	{
 		HeaderWriters += FString::Format(TEXT(
 			"    // Signal kernel executed by setting data count from first thread. Rest of header was already set up by the CPU.\n"
-			"    if (GroupIndex == 0) {0}_WriteNumData(1);\n"
+			"    if (GroupIndex == 0) {0}_SetNumDataInternal(1);\n"
 			"    AllMemoryBarrier();\n"),
 			{ InOutputPinProps.Label.ToString() });
 	};
 
 	auto EmitHeaderWriterFromInputPin = [&HeaderWriters](const FPCGPinProperties& InOutputPinProps, const UPCGPin* InFromPin)
 	{
+		check(InFromPin);
 		HeaderWriters += FString::Format(TEXT(
 			"    // Signal kernel executed by copying data count from pin {0} to pin {1} from first thread. Rest of header was already set up by the CPU.\n"
-			"    if (GroupIndex == 0) {1}_WriteNumData({0}_GetNumData());\n"
+			"    if (GroupIndex == 0) {1}_SetNumDataInternal({0}_GetNumData());\n"
 			"    AllMemoryBarrier();\n"),
 			{ InFromPin->Properties.Label.ToString(), InOutputPinProps.Label.ToString() });
 	};
@@ -1359,20 +1357,22 @@ FString UPCGCustomHLSLSettings::GetCookedKernelSource(const TMap<FPCGKernelAttri
 		// Initialize all output headers.
 		for (const FPCGPinPropertiesGPU& PinProps : OutputPins)
 		{
-			// Only support Point pins for now.
-			if (PinProps.AllowedTypes != EPCGDataType::Point)
-			{
-				continue;
-			}
-
 			if (PinProps.BufferSizeMode == EPCGPinBufferSizeMode::FromFirstPin && FirstPin)
 			{
 				EmitHeaderWriterFromInputPin(PinProps, FirstPin);
 			}
 			else if (PinProps.BufferSizeMode == EPCGPinBufferSizeMode::FixedElementCount)
 			{
-				// FixedElementCount always produces a single point data.
-				EmitHeaderWriterSingleData(PinProps);
+				const UPCGPin* InitFromPin = (PinProps.AllowedTypes == EPCGDataType::Param) ? GetInputPin(PinProps.InitializeFromPin) : nullptr;
+
+				if (InitFromPin)
+				{
+					EmitHeaderWriterFromInputPin(PinProps, InitFromPin);
+				}
+				else
+				{
+					EmitHeaderWriterSingleData(PinProps);
+				}
 			}
 			else if (PinProps.BufferSizeMode == EPCGPinBufferSizeMode::FromProductOfInputPins)
 			{
@@ -1387,12 +1387,11 @@ FString UPCGCustomHLSLSettings::GetCookedKernelSource(const TMap<FPCGKernelAttri
 		{
 			const UPCGPin* FirstPin = Node->GetPassThroughInputPin();
 			const UPCGPin* PrimaryOutputPin = GetFirstPointOutputPin();
-			
+
 			for (const UPCGPin* OutputPin : Node->GetOutputPins())
 			{
-				if (!OutputPin || OutputPin->Properties.AllowedTypes != EPCGDataType::Point)
+				if (!OutputPin || !PCGComputeHelpers::IsTypeAllowedInDataCollection(OutputPin->Properties.AllowedTypes))
 				{
-					// Point only for now.
 					continue;
 				}
 
@@ -1412,6 +1411,17 @@ FString UPCGCustomHLSLSettings::GetCookedKernelSource(const TMap<FPCGKernelAttri
 				{
 					EmitHeaderWriterFromInputPin(PinProps, FirstPin);
 				}
+				else if (PinPropsGPU->AllowedTypes == EPCGDataType::Param)
+				{
+					if (const UPCGPin* InitFromPin = GetInputPin(PinPropsGPU->InitializeFromPin))
+					{
+						EmitHeaderWriterFromInputPin(PinProps, InitFromPin);
+					}
+					else
+					{
+						EmitHeaderWriterSingleData(PinProps);
+					}
+				}
 			}
 		}
 	}
@@ -1458,15 +1468,15 @@ FString UPCGCustomHLSLSettings::GetCookedKernelSource(const TMap<FPCGKernelAttri
 				"\n"
 				"    // Loop over all attribute headers, if the address is non-zero, then copy it from pin {0} to pin {1}.\n"
 				"    {\n"
-				"        const uint HeadersAddress = {0}_ReadDataAttributeHeadersAddress({0}_DataAddress);\n"
-				"        const uint NumAttributes = {0}_ReadDataNumAttributes({0}_DataAddress);\n"
+				"        const uint HeadersAddress = {0}_GetDataAttributeHeadersAddress({0}_DataAddress);\n"
+				"        const uint NumAttributes = {0}_GetDataNumAttributes({0}_DataAddress);\n"
 				"        uint NumAttributesProcessed = 0;\n"
 				"\n"
 				"        for (int AttributeIndex = 0; AttributeIndex < 128; ++AttributeIndex)\n"
 				"        {\n"
-				"            const uint HeaderAddress = {0}_ReadAttributeHeaderAddress(HeadersAddress, AttributeIndex);\n"
-				"            const uint Stride = {0}_ReadAttributeStride(HeaderAddress);\n"
-				"            const uint Address = {0}_ReadAttributeAddress(HeaderAddress);\n"
+				"            const uint HeaderAddress = {0}_GetAttributeHeaderAddress(HeadersAddress, AttributeIndex);\n"
+				"            const uint Stride = {0}_GetAttributeStride(HeaderAddress);\n"
+				"            const uint Address = {0}_GetAttributeAddress(HeaderAddress);\n"
 				"\n"
 				"            if (Address != 0)\n"
 				"            {\n"
@@ -1475,7 +1485,7 @@ FString UPCGCustomHLSLSettings::GetCookedKernelSource(const TMap<FPCGKernelAttri
 				"                for (int I = 0; I < Stride; I += 4)\n"
 				"                {\n"
 				"                    const uint ElementAddress = BaseElementAddress + I;\n"
-				"                    {1}_StoreBuffer(ElementAddress, {0}_LoadBuffer(ElementAddress));\n"
+				"                    {1}_StoreBufferInternal(ElementAddress, {0}_LoadBufferInternal(ElementAddress));\n"
 				"                }\n"
 				"\n"
 				"                if (++NumAttributesProcessed >= NumAttributes) break; // We can early-out when we've looked at all the possible attributes\n"
@@ -1549,11 +1559,11 @@ bool FPCGCustomHLSLElement::ExecuteInternal(FPCGContext* Context) const
 	const UPCGCustomHLSLSettings* Settings = Context->GetInputSettings<UPCGCustomHLSLSettings>();
 	check(Settings);
 
-	// Calling IsKernelValid with the Context will emit warnings/errors on the node.
-	Settings->IsKernelValid(Context);
+	Settings->IsKernelValid(Context, /*bQuiet=*/false);
 
 	return true;
 }
 
+#undef PCG_LOG_VALIDATION
 #undef PCG_LOGGING_ENABLED
 #undef LOCTEXT_NAMESPACE
diff --git a/Engine/Plugins/PCG/Source/PCG/Private/Compute/PCGComputeCommon.cpp b/Engine/Plugins/PCG/Source/PCG/Private/Compute/PCGComputeCommon.cpp
new file mode 100644
--- /dev/null
+++ b/Engine/Plugins/PCG/Source/PCG/Private/Compute/PCGComputeCommon.cpp
@@ -0,0 +1,41 @@
+// Copyright Epic Games, Inc. All Rights Reserved.
+
+#include "PCGComputeCommon.h"
+
+#include "PCGParamData.h"
+#include "Data/PCGPointData.h"
+
+namespace PCGComputeHelpers
+{
+	int GetElementCount(const UPCGData* InData)
+	{
+		if (const UPCGPointData* PointData = Cast<UPCGPointData>(InData))
+		{
+			return PointData->GetPoints().Num();
+		}
+		else if (const UPCGParamData* ParamData = Cast<UPCGParamData>(InData))
+		{
+			if (const UPCGMetadata* Metadata = ParamData->ConstMetadata())
+			{
+				return Metadata->GetItemCountForChild();
+			}
+		}
+
+		return 0;
+	}
+
+	bool IsTypeAllowedAsInput(EPCGDataType Type)
+	{
+		return (Type | PCGComputeConstants::AllowedInputTypes) == PCGComputeConstants::AllowedInputTypes;
+	}
+
+	bool IsTypeAllowedAsOutput(EPCGDataType Type)
+	{
+		return (Type | PCGComputeConstants::AllowedOutputTypes) == PCGComputeConstants::AllowedOutputTypes;
+	}
+
+	bool IsTypeAllowedInDataCollection(EPCGDataType Type)
+	{
+		return (Type | PCGComputeConstants::AllowedDataCollectionTypes) == PCGComputeConstants::AllowedDataCollectionTypes;
+	}
+}
diff --git a/Engine/Plugins/PCG/Source/PCG/Private/Compute/PCGComputeCommon.h b/Engine/Plugins/PCG/Source/PCG/Private/Compute/PCGComputeCommon.h
--- a/Engine/Plugins/PCG/Source/PCG/Private/Compute/PCGComputeCommon.h
+++ b/Engine/Plugins/PCG/Source/PCG/Private/Compute/PCGComputeCommon.h
@@ -2,6 +2,10 @@
 
 #pragma once
 
+#include "PCGCommon.h"
+
+class UPCGData;
+
 namespace PCGComputeConstants
 {
 	constexpr int MAX_NUM_ATTRS = 128;
@@ -11,9 +15,13 @@ namespace PCGComputeConstants
 	constexpr int ATTRIBUTE_HEADER_SIZE_BYTES = 8;
 
 	constexpr int POINT_DATA_TYPE_ID = 0;
-	constexpr int POINT_DATA_HEADER_PREAMBLE_SIZE_BYTES = 16; // 4 bytes for Type, 4 bytes for NumAttrs, 4 bytes for the offset, 4 bytes for TypeInfo
+	constexpr int POINT_DATA_HEADER_PREAMBLE_SIZE_BYTES = 16; // 4 bytes for Type, 4 bytes for NumAttrs, 4 bytes for the address, 4 bytes for TypeInfo
 	constexpr int POINT_DATA_HEADER_SIZE_BYTES = POINT_DATA_HEADER_PREAMBLE_SIZE_BYTES + ATTRIBUTE_HEADER_SIZE_BYTES * MAX_NUM_ATTRS;
 
+	constexpr int PARAM_DATA_TYPE_ID = 1;
+	constexpr int PARAM_DATA_HEADER_PREAMBLE_SIZE_BYTES = 16; // 4 bytes for Type, 4 bytes for NumAttrs, 4 bytes for the address, 4 bytes for TypeInfo
+	constexpr int PARAM_DATA_HEADER_SIZE_BYTES = PARAM_DATA_HEADER_PREAMBLE_SIZE_BYTES + ATTRIBUTE_HEADER_SIZE_BYTES * MAX_NUM_ATTRS;
+
 	constexpr int NUM_POINT_PROPERTIES = 9;
 	constexpr int POINT_POSITION_ATTRIBUTE_ID = 0;
 	constexpr int POINT_ROTATION_ATTRIBUTE_ID = 1;
@@ -24,4 +32,28 @@ namespace PCGComputeConstants
 	constexpr int POINT_DENSITY_ATTRIBUTE_ID = 6;
 	constexpr int POINT_SEED_ATTRIBUTE_ID = 7;
 	constexpr int POINT_STEEPNESS_ATTRIBUTE_ID = 8;
+
+	/** PCG data types supported in GPU node inputs. */
+	constexpr EPCGDataType AllowedInputTypes = EPCGDataType::Point | EPCGDataType::Param | EPCGDataType::Landscape | EPCGDataType::Texture;
+
+	/** PCG data types supported in GPU node outputs. */
+	constexpr EPCGDataType AllowedOutputTypes = EPCGDataType::Point | EPCGDataType::Param;
+
+	/** PCG data types supported in GPU data collections. */
+	constexpr EPCGDataType AllowedDataCollectionTypes = EPCGDataType::Point | EPCGDataType::Param;
+}
+
+namespace PCGComputeHelpers
+{
+	/** Gets the element count for a given data. E.g. number of points in a PointData, number of metadata entries in a ParamData, etc. */
+	int GetElementCount(const UPCGData* InData);
+
+	/** True if 'Type' is valid on a GPU input pin. */
+	bool IsTypeAllowedAsInput(EPCGDataType Type);
+
+	/** True if 'Type' is valid on a GPU output pin. */
+	bool IsTypeAllowedAsOutput(EPCGDataType Type);
+
+	/** True if 'Type' is valid in a GPU data collection. Some types are only supported as DataInterfaces, and cannot be uploaded in data collections. */
+	bool IsTypeAllowedInDataCollection(EPCGDataType Type);
 }
diff --git a/Engine/Plugins/PCG/Source/PCG/Private/Compute/PCGDataForGPU.cpp b/Engine/Plugins/PCG/Source/PCG/Private/Compute/PCGDataForGPU.cpp
--- a/Engine/Plugins/PCG/Source/PCG/Private/Compute/PCGDataForGPU.cpp
+++ b/Engine/Plugins/PCG/Source/PCG/Private/Compute/PCGDataForGPU.cpp
@@ -3,6 +3,7 @@
 #include "Compute/PCGDataForGPU.h"
 
 #include "PCGData.h"
+#include "PCGParamData.h"
 #include "PCGPoint.h"
 #include "Compute/PCGComputeCommon.h"
 #include "Data/PCGPointData.h"
@@ -412,13 +413,7 @@ FPCGDataDesc::FPCGDataDesc(const UPCGData* Data, const TMap<FPCGKernelAttributeK
 	check(Data);
 
 	Type = Data->GetDataType();
-
-	if (Type == EPCGDataType::Point)
-	{
-		const UPCGPointData* PointData = CastChecked<UPCGPointData>(Data);
-		ElementCount = PointData->GetPoints().Num();
-	}
-	else { /* TODO: More types! */ }
+	ElementCount = PCGComputeHelpers::GetElementCount(Data);
 
 	InitializeAttributeDescs(Data->ConstMetadata(), GlobalAttributeLookupTable);
 }
@@ -431,6 +426,10 @@ uint32 FPCGDataDesc::ComputePackedSize() const
 	{
 		DataSizeBytes += POINT_DATA_HEADER_SIZE_BYTES;
 	}
+	else if (Type == EPCGDataType::Param)
+	{
+		DataSizeBytes += PARAM_DATA_HEADER_SIZE_BYTES;
+	}
 	else
 	{
 		// TODO: Support more types
@@ -508,13 +507,11 @@ FPCGDataCollectionDesc FPCGDataCollectionDesc::BuildFromInputDataCollectionAndIn
 	const TMap<FPCGKernelAttributeKey, int32>& InAttributeLookupTable)
 {
 	FPCGDataCollectionDesc Desc;
-
 	TArray<FPCGTaggedData> DataForPin = InDataCollection.GetInputsByPin(InputPinLabel);
 
 	for (const FPCGTaggedData& Data : DataForPin)
 	{
-		// Only support Point data right now
-		if (!Data.Data || Data.Data->GetDataType() != EPCGDataType::Point)
+		if (!Data.Data || !PCGComputeHelpers::IsTypeAllowedInDataCollection(Data.Data->GetDataType()))
 		{
 			continue;
 		}
@@ -709,6 +706,54 @@ void FPCGDataCollectionDesc::PackDataCollection(const FPCGDataCollection& InData
 				CurrentAttributeAddress += NumElements * AttributeNumComponents * 4;
 			}
 		}
+		else if (const UPCGParamData* ParamData = Cast<UPCGParamData>(InputData[DataIndex].Data))
+		{
+			const UPCGMetadata* Metadata = ParamData->ConstMetadata();
+
+			const FPCGDataDesc& DataDesc = DataDescs[DataIndex];
+			const uint32 NumElements = DataDesc.ElementCount;
+
+			const TArray<FPCGKernelAttributeDesc>& AttributeDescs = DataDesc.AttributeDescs;
+			const uint32 NumAttributes = AttributeDescs.Num();
+
+			OutPackedDataCollection[CurrentDataIndex + 0] = /*ParamDataTypeId=*/PARAM_DATA_TYPE_ID;
+			OutPackedDataCollection[CurrentDataIndex + 1] = NumAttributes;
+			OutPackedDataCollection[CurrentDataIndex + 2] = PARAM_DATA_HEADER_PREAMBLE_SIZE_BYTES;
+			OutPackedDataCollection[CurrentDataIndex + 3] = NumElements; // TypeInfo for ParamData is # of elements
+
+			const uint32 BaseAttributeHeaderAddress = CurrentDataAddress + PARAM_DATA_HEADER_PREAMBLE_SIZE_BYTES;
+			uint32 CurrentAttributeAddress = CurrentDataAddress + PARAM_DATA_HEADER_SIZE_BYTES;
+
+			for (const FPCGKernelAttributeDesc& AttributeDesc : AttributeDescs)
+			{
+				const uint32 AttributeId = AttributeDesc.Index;
+				const uint32 AttributeStrideBytes = PCGDataForGPUHelpers::GetAttributeTypeStrideBytes(AttributeDesc.Type);
+				const uint32 AttributeNumComponents = AttributeStrideBytes / sizeof(uint32); // E.g. float3 has 3 components
+
+				// Pack Position (24 bits for AttributeId, 8 bits for Stride)
+				const uint32 PackedIdAndStride = (AttributeId << 8) + AttributeStrideBytes;
+				const uint32 AttributeIndex = CurrentAttributeAddress / sizeof(uint32);
+
+				const uint32 AttributeHeaderIndex = (BaseAttributeHeaderAddress + AttributeId * ATTRIBUTE_HEADER_SIZE_BYTES) / sizeof(uint32);
+				OutPackedDataCollection[AttributeHeaderIndex + 0] = PackedIdAndStride;
+				OutPackedDataCollection[AttributeHeaderIndex + 1] = CurrentAttributeAddress;
+
+				const FPCGMetadataAttributeBase* AttributeBase = Metadata->GetConstAttribute(AttributeDesc.Name);
+
+				for (uint32 ElementIndex = 0; ElementIndex < NumElements; ++ElementIndex)
+				{
+					const uint32 PackedDataElementIndex = AttributeIndex + (ElementIndex * AttributeNumComponents);
+					const int64 MetadataKey = ElementIndex;
+
+					if (AttributeBase) // Pack attribute
+					{
+						ensure(PCGDataForGPUHelpers::PackAttributeHelper(AttributeBase, AttributeDesc, MetadataKey, OutPackedDataCollection, PackedDataElementIndex));
+					}
+				}
+
+				CurrentAttributeAddress += NumElements * AttributeNumComponents * 4;
+			}
+		}
 		else { /* TODO: Support non-point data. */ }
 	}
 }
@@ -766,6 +811,36 @@ void FPCGDataCollectionDesc::PrepareBufferForKernelOutput(TArray<uint32>& OutPac
 				CurrentAttributeAddress += NumElements * AttributeNumComponents * 4;
 			}
 		}
+		if (DataDescs[DataIndex].Type == EPCGDataType::Param)
+		{
+			const uint32 NumElements = DataDescs[DataIndex].ElementCount;
+
+			const TArray<FPCGKernelAttributeDesc>& AttributeDescs = DataDescs[DataIndex].AttributeDescs;
+			const uint32 NumAttributes = AttributeDescs.Num();
+
+			OutPackedDataCollection[CurrentDataIndex + 0] = /*ParamDataTypeId=*/PARAM_DATA_TYPE_ID;
+			OutPackedDataCollection[CurrentDataIndex + 1] = NumAttributes;
+			OutPackedDataCollection[CurrentDataIndex + 2] = PARAM_DATA_HEADER_PREAMBLE_SIZE_BYTES;
+			OutPackedDataCollection[CurrentDataIndex + 3] = NumElements; // TypeInfo for ParamData is # of elems
+
+			const uint32 BaseAttributeHeaderAddress = CurrentDataAddress + PARAM_DATA_HEADER_PREAMBLE_SIZE_BYTES;
+			uint32 CurrentAttributeAddress = CurrentDataAddress + PARAM_DATA_HEADER_SIZE_BYTES;
+
+			for (const FPCGKernelAttributeDesc& AttributeDesc : AttributeDescs)
+			{
+				const uint32 AttributeId = AttributeDesc.Index;
+				const uint32 AttributeStrideBytes = PCGDataForGPUHelpers::GetAttributeTypeStrideBytes(AttributeDesc.Type);
+				const uint32 AttributeNumComponents = AttributeStrideBytes / sizeof(uint32); // E.g. float3 has 3 components
+				const uint32 AttributeHeaderIndex = (BaseAttributeHeaderAddress + AttributeId * ATTRIBUTE_HEADER_SIZE_BYTES) / sizeof(uint32);
+
+				// Pack Position (24 bits for AttributeId, 8 bits for Stride)
+				const uint32 PackedIdAndStride = (AttributeId << 8) + AttributeStrideBytes;
+				OutPackedDataCollection[AttributeHeaderIndex + 0] = PackedIdAndStride;
+
+				OutPackedDataCollection[AttributeHeaderIndex + 1] = CurrentAttributeAddress;
+				CurrentAttributeAddress += NumElements * AttributeNumComponents * 4;
+			}
+		}
 		else { /* TODO: Support non-point data. */ }
 	}
 }
@@ -974,7 +1049,50 @@ EPCGUnpackDataCollectionResult FPCGDataCollectionDesc::UnpackDataCollection(cons
 			};
 
 			FPCGAsync::AsyncPointProcessing(/*Context=*/nullptr, OutPoints.Num(), OutPoints, DiscardInvalidPoints);
+		}
+		else if (TypeId == PARAM_DATA_TYPE_ID)
+		{
+			UPCGParamData* OutParamData = NewObject<UPCGParamData>();
+			UPCGMetadata* Metadata = OutParamData->MutableMetadata();
 
+			FPCGTaggedData& OutTaggedData = OutData.Emplace_GetRef();
+			OutTaggedData.Data = OutParamData;
+			OutTaggedData.Pin = InPin;
+
+			const uint32 AttributeHeadersIndex = CurrentDataIndex + PARAM_DATA_HEADER_PREAMBLE_SIZE_BYTES / sizeof(uint32);
+
+			TArray<TTuple</*EntryKey=*/int64, /*ParentEntryKey=*/int64>> AllMetadataEntries;
+			AllMetadataEntries.SetNumUninitialized(NumElements);
+
+			ParallelFor(NumElements, [&](int32 ElementIndex)
+			{
+				AllMetadataEntries[ElementIndex] = MakeTuple(Metadata->AddEntryPlaceholder(), PCGInvalidEntryKey);
+			});
+
+			Metadata->AddDelayedEntries(AllMetadataEntries);
+
+			// Loop over attributes.
+			for (const FPCGKernelAttributeDesc& AttributeDesc : AttributeDescs)
+			{
+				TRACE_CPUPROFILER_EVENT_SCOPE(WriteAttribute);
+
+				const uint32 AttributeNumComponents = PCGDataForGPUHelpers::GetAttributeTypeStrideBytes(AttributeDesc.Type) / sizeof(uint32);
+				const uint32 AttributeHeaderIndex = AttributeHeadersIndex + AttributeDesc.Index * ATTRIBUTE_HEADER_SIZE_BYTES / sizeof(uint32);
+				const uint32 AttributeIndex = DataAsUint[AttributeHeaderIndex + 1] / sizeof(uint32);
+
+				FPCGMetadataAttributeBase* AttributeBase = PCGDataForGPUHelpers::CreateAttributeFromAttributeDesc(Metadata, AttributeDesc);
+
+				ParallelFor(NumElements, [&](int32 ElementIndex)
+				{
+					if (AttributeBase)
+					{
+						const uint32 PackedDataElementIndex = AttributeIndex + ElementIndex * AttributeNumComponents;
+						check(PackedDataElementIndex + AttributeNumComponents <= NumPackedFloats);
+
+						ensure(PCGDataForGPUHelpers::UnpackAttributeHelper(PackedData, PackedDataElementIndex, AttributeBase, AttributeDesc, ElementIndex));
+					}
+				});
+			}
 		}
 		else { /* TODO: Support non-point data. */ }
 	}
diff --git a/Engine/Plugins/PCG/Source/PCG/Private/Graph/PCGGraphCompilerGPU.cpp b/Engine/Plugins/PCG/Source/PCG/Private/Graph/PCGGraphCompilerGPU.cpp
--- a/Engine/Plugins/PCG/Source/PCG/Private/Graph/PCGGraphCompilerGPU.cpp
+++ b/Engine/Plugins/PCG/Source/PCG/Private/Graph/PCGGraphCompilerGPU.cpp
@@ -507,7 +507,7 @@ void FPCGGraphCompilerGPU::BuildGPUGraphTask(
 		switch (PinType)
 		{
 		case EPCGDataType::Point:
-		case EPCGDataType::PointOrParam:
+		case EPCGDataType::Param:
 		{
 			UPCGDataCollectionDataInterface* DataInterfacePCGData = nullptr;
 
diff --git a/Engine/Plugins/PCG/Source/PCG/Public/Compute/Elements/PCGCustomHLSL.h b/Engine/Plugins/PCG/Source/PCG/Public/Compute/Elements/PCGCustomHLSL.h
--- a/Engine/Plugins/PCG/Source/PCG/Public/Compute/Elements/PCGCustomHLSL.h
+++ b/Engine/Plugins/PCG/Source/PCG/Public/Compute/Elements/PCGCustomHLSL.h
@@ -31,20 +31,31 @@ public:
 	{
 	}
 
+#if WITH_EDITOR
+	bool CanEditChange(const FEditPropertyChain& PropertyChain) const;
+#endif
+
 public:
 	/** Compute graphs use this to calculate the buffer size of output pins. */
-	UPROPERTY(BlueprintReadWrite, EditAnywhere, Category = "GPU Buffer Size", meta = (EditCondition = "bDisplayBufferSizeSettings", EditConditionHides))
+	UPROPERTY(BlueprintReadWrite, EditAnywhere, Category = "GPU Buffer Size", meta = (EditCondition = "bDisplayBufferSizeSettings", EditConditionHides, HideEditConditionToggle))
 	EPCGPinBufferSizeMode BufferSizeMode = EPCGPinBufferSizeMode::FromFirstPin;
 
-	UPROPERTY(BlueprintReadWrite, EditAnywhere, Category = "GPU Buffer Size", meta = (EditCondition = "bDisplayBufferSizeSettings && BufferSizeMode == EPCGPinBufferSizeMode::FixedElementCount", EditConditionHides))
+	UPROPERTY(BlueprintReadWrite, EditAnywhere, Category = "GPU Buffer Size", meta = (EditCondition = "(bDisplayBufferSizeSettings && BufferSizeMode == EPCGPinBufferSizeMode::FixedElementCount) || AllowedTypes == EPCGDataType::Param", EditConditionHides))
 	int FixedBufferElementCount = 4;
 
 	UPROPERTY(BlueprintReadWrite, EditAnywhere, DisplayName = "Input Pins", Category = "GPU Buffer Size", meta = (EditCondition = "bDisplayBufferSizeSettings && BufferSizeMode == EPCGPinBufferSizeMode::FromProductOfInputPins", EditConditionHides))
 	TArray<FName> BufferSizeInputPinLabels;
 
+	/** Select an input pin to copy attributes from. If left as 'None', this will be ignored. Note, this will copy attribute names only, not their values. */
+	UPROPERTY(BlueprintReadWrite, EditAnywhere, Category = "GPU Buffer Size", meta = (EditCondition = "bAllowEditInitializationPin", EditConditionHides, HideEditConditionToggle))
+	FName InitializeFromPin = NAME_None;
+
 #if WITH_EDITORONLY_DATA
 	UPROPERTY(Transient)
 	bool bDisplayBufferSizeSettings = true;
+
+	UPROPERTY(Transient)
+	bool bAllowEditInitializationPin = false;
 #endif // WITH_EDITORONLY_DATA
 
 	/** Add entries to create new attributes on data emitted by this pin. */
@@ -52,6 +63,15 @@ public:
 	TArray<FPCGKernelAttributeKey> CreatedKernelAttributeKeys;
 };
 
+template<>
+struct TStructOpsTypeTraits<FPCGPinPropertiesGPU> : public TStructOpsTypeTraitsBase2<FPCGPinPropertiesGPU>
+{
+	enum
+	{
+		WithCanEditChange = true,
+	};
+};
+
 /** Type of kernel allows us to make decisions about execution automatically, streamlining authoring. */
 UENUM()
 enum class EPCGKernelType : uint8
@@ -120,7 +140,7 @@ protected:
 #endif
 
 public:
-	bool IsKernelValid(FPCGContext* InContext = nullptr) const;
+	bool IsKernelValid(FPCGContext* InContext = nullptr, bool bQuiet = true) const;
 	
 	FString GetCookedKernelSource(const TMap<FPCGKernelAttributeKey, int>& GlobalAttributeLookupTable) const;
 	FString GetKernelEntryPoint() const { return TEXT("Main"); }
@@ -133,7 +153,6 @@ public:
 	const UPCGPin* GetOutputPin(FName Label) const;
 	const UPCGPin* GetFirstInputPin() const;
 	const UPCGPin* GetPointProcessingInputPin() const;
-	const UPCGPin* GetSecondPointProcessingInputPin() const;
 	const UPCGPin* GetFirstOutputPin() const;
 	const UPCGPin* GetFirstPointOutputPin() const;
 
