[PCG] Change recursive QueueNextTask/ProcessCacheResults into a loop

#rnx

diff --git a/Engine/Plugins/PCG/Source/PCG/Private/Graph/PCGGraphExecutor.cpp b/Engine/Plugins/PCG/Source/PCG/Private/Graph/PCGGraphExecutor.cpp
--- a/Engine/Plugins/PCG/Source/PCG/Private/Graph/PCGGraphExecutor.cpp
+++ b/Engine/Plugins/PCG/Source/PCG/Private/Graph/PCGGraphExecutor.cpp
@@ -926,14 +926,14 @@ bool FPCGGraphExecutor::ProcessScheduledTasks()
 		}
 	}
 
-	ProcessCachedResults(CachedResults);
+	ProcessCachedResults(MoveTemp(CachedResults));
 		
 	if (bIsInGameThread)
 	{
 		PCGGraphExecutor::TScopeLock ScopeLock(TasksLock);
 		PCGGraphExecutor::TScopeLock ChildScopeLock(LiveTasksLock);
 		// This is a safeguard to check if we're in a stuck state
-		if (CachingResults.Num() == 0 && ReadyTasks.Num() == 0 && ActiveTasks.Num() == 0 && ActiveTasksGameThreadOnly.Num() == 0 && SleepingTasks.Num() == 0 && Tasks.Num() > 0)
+		if (CollectGCCachingResults.Num() == 0 && ReadyTasks.Num() == 0 && ActiveTasks.Num() == 0 && ActiveTasksGameThreadOnly.Num() == 0 && SleepingTasks.Num() == 0 && Tasks.Num() > 0)
 		{
 			UE_LOG(LogPCG, Error, TEXT("PCG Graph executor error: tasks are in a deadlocked state. Will drop all tasks."));
 			ClearAllTasks();
@@ -1248,8 +1248,8 @@ void FPCGGraphExecutor::PrepareForExecute(FPCGGraphTask& Task, FCachedResult*& O
 		const bool bFoundInCache = GraphCache.GetFromCache(Task.Node, Task.Element.Get(), DependenciesCrc, Task.SourceComponent.Get(), LocalCachedResult.Output);
 		if (bFoundInCache)
 		{
-			check(!CachingResults.Contains(Task.NodeId));
-			TUniquePtr<FCachedResult>& CachingResultPtr = CachingResults.Add(Task.NodeId, MakeUnique<FCachedResult>());
+			check(!CollectGCCachingResults.Contains(Task.NodeId));
+			TUniquePtr<FCachedResult>& CachingResultPtr = CollectGCCachingResults.Add(Task.NodeId, MakeUnique<FCachedResult>());
 			OutCachedResult = CachingResultPtr.Get();
 			OutCachedResult->TaskId = Task.NodeId;
 			OutCachedResult->Output = MoveTemp(LocalCachedResult.Output);
@@ -1473,7 +1473,7 @@ bool FPCGGraphExecutor::ExecuteScheduling(double EndTime, TSharedPtr<FPCGGraphAc
 		}
 	}
 		
-	ProcessCachedResults(CachedResults);
+	ProcessCachedResults(MoveTemp(CachedResults));
 
 	// Dispatch Tasks
 	{
@@ -1961,7 +1961,7 @@ void FPCGGraphExecutor::ClearAllTasks()
 	SleepingTasks.Reset();
 }
 
-void FPCGGraphExecutor::QueueNextTasks(FPCGTaskId FinishedTask)
+TArray<FPCGGraphExecutor::FCachedResult*> FPCGGraphExecutor::QueueNextTasksInternal(FPCGTaskId FinishedTask)
 {
 	TRACE_CPUPROFILER_EVENT_SCOPE(FPCGGraphExecutor::QueueNextTasks);
 	TArray<FCachedResult*> CachedResults;
@@ -2003,12 +2003,35 @@ void FPCGGraphExecutor::QueueNextTasks(FPCGTaskId FinishedTask)
 		}
 	}
 
-	ProcessCachedResults(CachedResults);
+	return CachedResults;
+}
+
+void FPCGGraphExecutor::QueueNextTasks(FPCGTaskId FinishedTask)
+{
+	TArray<FCachedResult*> CachedResults = QueueNextTasksInternal(FinishedTask);
+
+	ProcessCachedResults(MoveTemp(CachedResults));
 }
 
-void FPCGGraphExecutor::ProcessCachedResults(const TArray<FCachedResult*>& CachedResults)
+void FPCGGraphExecutor::ProcessCachedResults(TArray<FCachedResult*> CachedResults)
+{
+	while (CachedResults.Num() > 0)
+	{
+		TArray<FPCGTaskId> NextTasks = ProcessCachedResultsInternal(MoveTemp(CachedResults));
+		check(CachedResults.IsEmpty());
+
+		for (const FPCGTaskId& NextTask : NextTasks)
+		{
+			CachedResults.Append(QueueNextTasksInternal(NextTask));
+		}
+	}
+}
+
+TArray<FPCGTaskId> FPCGGraphExecutor::ProcessCachedResultsInternal(TArray<FCachedResult*> CachedResults)
 {
 	TRACE_CPUPROFILER_EVENT_SCOPE(FPCGGraphExecutor::ProcessCachedResults);
+	TArray<FPCGTaskId> NextTasks;
+	NextTasks.Reserve(CachedResults.Num());
 
 	for (FCachedResult* CachedResult : CachedResults)
 	{
@@ -2030,7 +2053,7 @@ void FPCGGraphExecutor::ProcessCachedResults(const TArray<FCachedResult*>& Cache
 		}
 
 		StoreResults(CachedResult->TaskId, CachedResult->Output, CachedResult->bIsPostGraphTask);
-		QueueNextTasks(CachedResult->TaskId);
+		NextTasks.Add(CachedResult->TaskId);
 	}
 
 	if (CachedResults.Num() > 0)
@@ -2038,13 +2061,15 @@ void FPCGGraphExecutor::ProcessCachedResults(const TArray<FCachedResult*>& Cache
 		PCGGraphExecutor::TScopeLock ScopeLock(LiveTasksLock);
 		for (const FCachedResult* CachedResult : CachedResults)
 		{
-			check(CachingResults.Contains(CachedResult->TaskId));
-			CachingResults.Remove(CachedResult->TaskId);
+			check(CollectGCCachingResults.Contains(CachedResult->TaskId));
+			CollectGCCachingResults.Remove(CachedResult->TaskId);
 		}
 
 		// Next Scheduling call needs to check if this task completion unblocked some sleeping task(s)
 		bNeedToCheckSleepingTasks = true;
 	}
+
+	return NextTasks;
 }
 
 bool FPCGGraphExecutor::CancelNextTasks(FPCGTaskId CancelledTask, TSet<UPCGComponent*>& OutCancelledComponents)
@@ -2548,7 +2573,7 @@ void FPCGGraphExecutor::AddReferencedObjects(FReferenceCollector& Collector)
 		Algo::ForEach(ActiveTasksGameThreadOnly, AddReferencesActiveTask);
 		Algo::ForEach(SleepingTasks, AddReferencesActiveTask);
 
-		Algo::ForEach(CachingResults, [&Collector](auto& CachingResult)
+		Algo::ForEach(CollectGCCachingResults, [&Collector](auto& CachingResult)
 		{
 			CachingResult.Value->Output.AddReferences(Collector);
 		});
diff --git a/Engine/Plugins/PCG/Source/PCG/Private/Graph/PCGGraphExecutor.h b/Engine/Plugins/PCG/Source/PCG/Private/Graph/PCGGraphExecutor.h
--- a/Engine/Plugins/PCG/Source/PCG/Private/Graph/PCGGraphExecutor.h
+++ b/Engine/Plugins/PCG/Source/PCG/Private/Graph/PCGGraphExecutor.h
@@ -279,6 +279,8 @@ private:
 	TSet<UPCGComponent*> Cancel(TFunctionRef<bool(TWeakObjectPtr<UPCGComponent>)> CancelFilter);
 	void ClearAllTasks();
 	void QueueNextTasks(FPCGTaskId FinishedTask);
+	TArray<FCachedResult*> QueueNextTasksInternal(FPCGTaskId FinishedTask);
+
 	bool CancelNextTasks(FPCGTaskId CancelledTask, TSet<UPCGComponent*>& OutCancelledComponents);
 	void RemoveTaskFromInputSuccessors(FPCGTaskId CancelledTask, const TArray<FPCGGraphTaskInput>& CancelledTaskInputs);
 	void RemoveTaskFromInputSuccessorsNoLock(FPCGTaskId CancelledTask, const TArray<FPCGGraphTaskInput>& CancelledTaskInputs);
@@ -293,7 +295,8 @@ private:
 	void PrepareForExecute(FPCGGraphTask& Task, FCachedResult*& OutCachedResult, bool bLiveTasksLockAlreadyLocked);
 
 	/** Store cache results and Queue next tasks */
-	void ProcessCachedResults(const TArray<FCachedResult*>& CachedResults);
+	void ProcessCachedResults(TArray<FCachedResult*> CachedResults);
+	TArray<FPCGTaskId> ProcessCachedResultsInternal(TArray<FCachedResult*> CachedResults);
 
 	/** Combine all param data into one on the Params pin, if any.*/
 	void CombineParams(FPCGGraphTask& Task);
@@ -355,8 +358,7 @@ private:
 	bool bNeedToCheckSleepingTasks = false;
 
 	// Used to keep GC references to in flight caching results (not yet stored to output and might not be in cache anymore)
-	TMap<FPCGTaskId, TUniquePtr<FCachedResult>> CachingResults;
-	int32 NumWorkerTasks = 0;
+	TMap<FPCGTaskId, TUniquePtr<FCachedResult>> CollectGCCachingResults;
 
 	/** Lock level 3 */
 	UE::FSpinLock CollectGCReferenceTasksLock;
