Outliner: Add a way to add a delay before the next requested UI refresh executes
PCG: Add a delay to the outliner whenever PCG actors are created/destroyed to prevent multiple UI refreshes in a row

diff --git a/Engine/Plugins/PCG/Source/PCG/Private/Helpers/PCGActorHelpers.cpp b/Engine/Plugins/PCG/Source/PCG/Private/Helpers/PCGActorHelpers.cpp
--- a/Engine/Plugins/PCG/Source/PCG/Private/Helpers/PCGActorHelpers.cpp
+++ b/Engine/Plugins/PCG/Source/PCG/Private/Helpers/PCGActorHelpers.cpp
@@ -19,11 +19,26 @@
 #if WITH_EDITOR
 #include "Editor.h"
 #include "Engine/Level.h"
+#include "Editor/IPCGEditorModule.h"
+#include "Modules/ModuleManager.h"
 #include "ScopedTransaction.h"
 #include "Subsystems/ActorEditorContextSubsystem.h"
 #include "WorldPartition/DataLayer/DataLayerInstance.h"
 #include "WorldPartition/DataLayer/ExternalDataLayerAsset.h"
 #include "WorldPartition/DataLayer/ExternalDataLayerInstance.h"
+
+namespace UE::PCGActorHelpers::Local
+{
+	static float OutlinerUIRefreshDelay = 1.0f;
+	static FAutoConsoleVariableRef CvarOutlinerUIRefreshDelay
+	(
+		TEXT("PCG.Editor.OutlinerRefreshDelay"),
+		OutlinerUIRefreshDelay,
+		TEXT("The delay (in seconds) before refreshing the Outliner after executing PCG tasks."),
+		ECVF_Default
+	);
+}
+
 #endif
 
 UInstancedStaticMeshComponent* UPCGActorHelpers::GetOrCreateISMC(AActor* InTargetActor, UPCGComponent* InSourceComponent, uint64 SettingsUID, const FPCGISMCBuilderParameters& InParams)
@@ -348,6 +363,12 @@ bool UPCGActorHelpers::DeleteActors(UWorld* World, const TArray<TSoftObjectPtr<A
 #endif
 	{
 #if WITH_EDITOR
+
+		if (IPCGEditorModule* EditorModule = FModuleManager::GetModulePtr<IPCGEditorModule>("PCGEditor"))
+		{
+			EditorModule->SetOutlinerUIRefreshDelay(UE::PCGActorHelpers::Local::OutlinerUIRefreshDelay);
+		}
+
 		// Create TX so that dirty actor packages are tracked
 		// 
 		// Without tracking deleted actor packages will get unloaded on the next GC with no chance to save them first 
@@ -496,6 +517,11 @@ AActor* UPCGActorHelpers::SpawnDefaultActor(const FSpawnDefaultActorParams& Para
 
 	// Specify EDL we want to use if any for spawning this actor
 	FScopedOverrideSpawningLevelMountPointObject EDLScope(ExternalDataLayerInstance ? ExternalDataLayerInstance->GetExternalDataLayerAsset() : nullptr);
+
+	if (IPCGEditorModule* EditorModule = FModuleManager::GetModulePtr<IPCGEditorModule>("PCGEditor"))
+	{
+		EditorModule->SetOutlinerUIRefreshDelay(UE::PCGActorHelpers::Local::OutlinerUIRefreshDelay);
+	}
 #endif
 
 	AActor* NewActor = Params.World->SpawnActor(*Params.ActorClass, &Params.Transform, SpawnParams);
diff --git a/Engine/Plugins/PCG/Source/PCG/Public/Editor/IPCGEditorModule.h b/Engine/Plugins/PCG/Source/PCG/Public/Editor/IPCGEditorModule.h
--- a/Engine/Plugins/PCG/Source/PCG/Public/Editor/IPCGEditorModule.h
+++ b/Engine/Plugins/PCG/Source/PCG/Public/Editor/IPCGEditorModule.h
@@ -12,4 +12,5 @@ class IPCGEditorModule : public IModuleInterface
 public:
 	virtual TWeakPtr<IPCGEditorProgressNotification> CreateProgressNotification(const FTextFormat& TextFormat, bool bCanCancel) = 0;
 	virtual void ReleaseProgressNotification(TWeakPtr<IPCGEditorProgressNotification> InNotification) = 0;
+	virtual void SetOutlinerUIRefreshDelay(float InDelay) = 0;
 };
\ No newline at end of file
diff --git a/Engine/Plugins/PCG/Source/PCGEditor/PCGEditor.Build.cs b/Engine/Plugins/PCG/Source/PCGEditor/PCGEditor.Build.cs
--- a/Engine/Plugins/PCG/Source/PCGEditor/PCGEditor.Build.cs
+++ b/Engine/Plugins/PCG/Source/PCGEditor/PCGEditor.Build.cs
@@ -61,7 +61,8 @@ namespace UnrealBuildTool.Rules
 					"TypedElementFramework",
 					"TypedElementRuntime",
 					"UnrealEd",
-					"LevelEditor"
+					"LevelEditor",
+					"SceneOutliner"
 				});
 		}
 	}
diff --git a/Engine/Plugins/PCG/Source/PCGEditor/Private/PCGEditorModule.cpp b/Engine/Plugins/PCG/Source/PCGEditor/Private/PCGEditorModule.cpp
--- a/Engine/Plugins/PCG/Source/PCGEditor/Private/PCGEditorModule.cpp
+++ b/Engine/Plugins/PCG/Source/PCGEditor/Private/PCGEditorModule.cpp
@@ -44,6 +44,7 @@
 #include "Details/PCGVolumeDetails.h"
 #include "Framework/MultiBox/MultiBoxBuilder.h"
 #include "LevelEditor.h"
+#include "Editor/SceneOutliner/Public/ISceneOutliner.h"
 #include "Framework/Notifications/NotificationManager.h"
 
 #define LOCTEXT_NAMESPACE "FPCGEditorModule"
@@ -175,6 +176,23 @@ void FPCGEditorModule::ReleaseProgressNotification(TWeakPtr<IPCGEditorProgressNo
 	}
 }
 
+void FPCGEditorModule::SetOutlinerUIRefreshDelay(float InDelay)
+{
+	TWeakPtr<class ILevelEditor> LevelEditor = FModuleManager::GetModuleChecked<FLevelEditorModule>(TEXT("LevelEditor")).GetLevelEditorInstance();
+	if (LevelEditor.IsValid())
+	{
+		TArray<TWeakPtr<class ISceneOutliner>> SceneOutlinerPtrs = LevelEditor.Pin()->GetAllSceneOutliners();
+
+		for (TWeakPtr<class ISceneOutliner> SceneOutlinerPtr : SceneOutlinerPtrs)
+		{
+			if (TSharedPtr<class ISceneOutliner> SceneOutlinerPin = SceneOutlinerPtr.Pin())
+			{
+				SceneOutlinerPin->SetNextUIRefreshDelay(InDelay);
+			}
+		}
+	}
+}
+
 void FPCGEditorModule::RegisterOnEditorModeChange()
 {
 	// Have a callback that catches changes in the Editor modes, to catch when we exit the landscape edit mode.
diff --git a/Engine/Plugins/PCG/Source/PCGEditor/Private/PCGEditorModule.h b/Engine/Plugins/PCG/Source/PCGEditor/Private/PCGEditorModule.h
--- a/Engine/Plugins/PCG/Source/PCGEditor/Private/PCGEditorModule.h
+++ b/Engine/Plugins/PCG/Source/PCGEditor/Private/PCGEditorModule.h
@@ -28,6 +28,7 @@ public:
 	// ~IPCGEditorModule implementation
 	virtual TWeakPtr<IPCGEditorProgressNotification> CreateProgressNotification(const FTextFormat& TextFormat, bool bCanCancel) override;
 	virtual void ReleaseProgressNotification(TWeakPtr<IPCGEditorProgressNotification> InNotification) override;
+	virtual void SetOutlinerUIRefreshDelay(float InDelay) override;
 	// ~End IPCGEditorModule implementation
 
 protected:
diff --git a/Engine/Source/Editor/SceneOutliner/Private/SSceneOutliner.cpp b/Engine/Source/Editor/SceneOutliner/Private/SSceneOutliner.cpp
--- a/Engine/Source/Editor/SceneOutliner/Private/SSceneOutliner.cpp
+++ b/Engine/Source/Editor/SceneOutliner/Private/SSceneOutliner.cpp
@@ -937,6 +937,11 @@ FSceneOutlinerTreeItemPtr SSceneOutliner::GetTreeItem(FSceneOutlinerTreeItemID I
 	return Result;
 }
 
+void SSceneOutliner::SetNextUIRefreshDelay(float InDelay)
+{
+	UIRefreshDelay = InDelay;
+}
+
 void SSceneOutliner::RemoveItemFromTree(FSceneOutlinerTreeItemRef ReferenceItem)
 {
 	if (const FSceneOutlinerTreeItemPtr* ItemInTree = TreeItemMap.Find(ReferenceItem->GetID()))
@@ -2366,10 +2371,24 @@ void SSceneOutliner::Tick(const FGeometry& AllottedGeometry, const double InCurr
 			Pair.Value->Flags.bChildrenRequireSort = true;
 		}
 
-		OutlinerTreeView->RequestTreeRefresh();
+		bNeedsUIRefresh = true;
 		bSortDirty = false;
 	}
 
+	// If we are pending a UI refresh
+	if(bNeedsUIRefresh)
+	{
+		UIRefreshDelay -= InDeltaTime;
+
+		// if we are currently pending a sort, don't refresh until that is completed
+		if(UIRefreshDelay <= 0.0f && !bSortDirty)
+		{
+			OutlinerTreeView->RequestTreeRefresh();
+			bNeedsUIRefresh = false;
+			UIRefreshDelay = 0.0f;
+		}
+	}
+
 	if (SortOutlinerTimer <= 0)
 	{
 		SortOutlinerTimer = SCENE_OUTLINER_RESORT_TIMER;
diff --git a/Engine/Source/Editor/SceneOutliner/Public/ISceneOutliner.h b/Engine/Source/Editor/SceneOutliner/Public/ISceneOutliner.h
--- a/Engine/Source/Editor/SceneOutliner/Public/ISceneOutliner.h
+++ b/Engine/Source/Editor/SceneOutliner/Public/ISceneOutliner.h
@@ -188,6 +188,10 @@ public:
 
 	/** Retrieve an ISceneOutlinerTreeItem by its ID if it exists in the tree */
 	virtual FSceneOutlinerTreeItemPtr GetTreeItem(FSceneOutlinerTreeItemID, bool bIncludePending = false) = 0;
+
+	/** Set a delay (in seconds) before the the next requested UI refresh executes. Note: This only applies until the next refresh that executes
+	 * and any subsequent ones will be immediate. */
+	virtual void SetNextUIRefreshDelay(float InDelay) = 0;
 protected:
 	ISceneOutlinerMode* Mode;
 };
diff --git a/Engine/Source/Editor/SceneOutliner/Public/SSceneOutliner.h b/Engine/Source/Editor/SceneOutliner/Public/SSceneOutliner.h
--- a/Engine/Source/Editor/SceneOutliner/Public/SSceneOutliner.h
+++ b/Engine/Source/Editor/SceneOutliner/Public/SSceneOutliner.h
@@ -410,6 +410,11 @@ public:
 
 	/** Retrieve an ISceneOutlinerTreeItem by its ID if it exists in the tree */
 	virtual FSceneOutlinerTreeItemPtr GetTreeItem(FSceneOutlinerTreeItemID, bool bIncludePending = false) override;
+	
+	/** Set a delay (in seconds) before the the next requested UI refresh executes. Note: This only applies until the next refresh that executes
+	 * and any subsequent ones will be immediate. */
+	virtual void SetNextUIRefreshDelay(float InDelay) override;
+
 public:
 	/** Event to react to a user double click on a item */
 	SceneOutliner::FTreeItemPtrEvent& GetDoubleClickEvent() { return OnDoubleClickOnTreeEvent; }
@@ -956,6 +961,12 @@ private:
 	/** Reentrancy guard */
 	bool bIsReentrant;
 
+	/** The delay (in seconds) before a UI refresh is executed after being requested. 0 means the next refresh will be immediate */
+	float UIRefreshDelay = 0.0f;
+
+	/** If true, the TreeView UI is pending a refresh */
+	bool bNeedsUIRefresh = false;
+
 	/* Widget containing the filtering text box */
 	TSharedPtr< SFilterSearchBox > FilterTextBoxWidget;
 
