Added per-data caching in elements that loop over their required pins.
Light refactoring around cache usage (incl. adding an interface) and moved most of the cache interaction to the element.
Requires opt-in per element, which is done in a separate CL.

diff --git a/Engine/Plugins/PCG/Source/PCG/Private/Graph/PCGGraphCache.cpp b/Engine/Plugins/PCG/Source/PCG/Private/Graph/PCGGraphCache.cpp
--- a/Engine/Plugins/PCG/Source/PCG/Private/Graph/PCGGraphCache.cpp
+++ b/Engine/Plugins/PCG/Source/PCG/Private/Graph/PCGGraphCache.cpp
@@ -47,13 +47,18 @@ FPCGGraphCache::~FPCGGraphCache()
 	ClearCache();
 }
 
-bool FPCGGraphCache::GetFromCache(const UPCGNode* InNode, const IPCGElement* InElement, const FPCGCrc& InDependenciesCrc, const UPCGComponent* InComponent, FPCGDataCollection& OutOutput) const
+bool FPCGGraphCache::GetFromCache(const FPCGGetFromCacheParams& Params, FPCGDataCollection& OutOutput) const
 {
 	if (!CVarCacheEnabled.GetValueOnAnyThread())
 	{
 		return false;
 	}
 
+	const UPCGNode* InNode = Params.Node;
+	const IPCGElement* InElement = Params.Element;
+	const UPCGComponent* InComponent = Params.Component;
+	const FPCGCrc& InDependenciesCrc = Params.Crc;
+
 	if(!InDependenciesCrc.IsValid())
 	{
 		UE_LOG(LogPCG, Warning, TEXT("Invalid dependencies passed to FPCGGraphCache::GetFromCache(), lookup aborted."));
@@ -64,7 +69,7 @@ bool FPCGGraphCache::GetFromCache(const UPCGNode* InNode, const IPCGElement* InE
 
 	{
 		TRACE_CPUPROFILER_EVENT_SCOPE(FPCGGraphCache::GetFromCache);
-		FReadScopeLock ScopedReadLock(CacheLock);
+		UE::TScopeLock ScopedLock(CacheLock);
 
 		FPCGCacheEntryKey CacheKey(InElement, InDependenciesCrc);
 		if (const FPCGDataCollection* Value = const_cast<FPCGGraphCache*>(this)->CacheData.FindAndTouch(CacheKey))
@@ -91,13 +96,16 @@ bool FPCGGraphCache::GetFromCache(const UPCGNode* InNode, const IPCGElement* InE
 	}
 }
 
-void FPCGGraphCache::StoreInCache(const IPCGElement* InElement, const FPCGCrc& InDependenciesCrc, const FPCGDataCollection& InOutput)
+void FPCGGraphCache::StoreInCache(const FPCGStoreInCacheParams& Params, const FPCGDataCollection& InOutput)
 {
 	if (!CVarCacheEnabled.GetValueOnAnyThread())
 	{
 		return;
 	}
 
+	const IPCGElement* InElement = Params.Element;
+	const FPCGCrc& InDependenciesCrc = Params.Crc;
+
 	if (!ensure(InDependenciesCrc.IsValid()))
 	{
 		return;
@@ -105,7 +113,7 @@ void FPCGGraphCache::StoreInCache(const IPCGElement* InElement, const FPCGCrc& I
 
 	{
 		TRACE_CPUPROFILER_EVENT_SCOPE(FPCGGraphCache::StoreInCache);
-		FWriteScopeLock ScopedWriteLock(CacheLock);
+		UE::TScopeLock ScopedLock(CacheLock);
 
 		if (CacheData.Num() == CacheData.Max())
 		{
@@ -121,7 +129,7 @@ void FPCGGraphCache::StoreInCache(const IPCGElement* InElement, const FPCGCrc& I
 
 void FPCGGraphCache::ClearCache()
 {
-	FWriteScopeLock ScopedWriteLock(CacheLock);
+	UE::TScopeLock ScopedLock(CacheLock);
 
 	MemoryRecords.Empty();
 	TotalMemoryUsed = 0;
@@ -150,7 +158,7 @@ bool FPCGGraphCache::EnforceMemoryBudget()
 	}
 
 	{
-		FWriteScopeLock ScopeWriteLock(CacheLock);
+		UE::TScopeLock ScopedLock(CacheLock);
 		const float MemoryCleanupRatio = FMath::Clamp(CVarCacheMemoryCleanupRatio.GetValueOnAnyThread(), 0.0f, 1.0f);
 		const uint64 TargetCacheMemoryUsage = static_cast<uint64>(MemoryCleanupRatio * MemoryBudget);
 
@@ -178,7 +186,7 @@ void FPCGGraphCache::CleanFromCache(const IPCGElement* InElement, const UPCGSett
 	}
 
 	{
-		FWriteScopeLock ScopeWriteLock(CacheLock);
+		UE::TScopeLock ScopedLock(CacheLock);
 
 		TArray<FPCGCacheEntryKey> Keys;
 		CacheData.GetKeys(Keys);
@@ -202,7 +210,7 @@ void FPCGGraphCache::CleanFromCache(const IPCGElement* InElement, const UPCGSett
 uint32 FPCGGraphCache::GetGraphCacheEntryCount(IPCGElement* InElement) const
 {
 	TRACE_CPUPROFILER_EVENT_SCOPE(FPCGGraphCache::GetFromCache);
-	FReadScopeLock ScopedReadLock(CacheLock);
+	UE::TScopeLock ScopedLock(CacheLock);
 
 	uint32 Count = 0;
 
@@ -224,7 +232,7 @@ uint32 FPCGGraphCache::GetGraphCacheEntryCount(IPCGElement* InElement) const
 void FPCGGraphCache::AddReferencedObjects(FReferenceCollector& Collector)
 {
 	TRACE_CPUPROFILER_EVENT_SCOPE(FPCGGraphCache::AddReferencedObjects);
-	FReadScopeLock ScopedReadLock(CacheLock);
+	UE::TScopeLock ScopedLock(CacheLock);
 
 	for (FPCGDataCollection& CacheEntry : CacheData)
 	{
diff --git a/Engine/Plugins/PCG/Source/PCG/Private/Graph/PCGGraphCache.h b/Engine/Plugins/PCG/Source/PCG/Private/Graph/PCGGraphCache.h
--- a/Engine/Plugins/PCG/Source/PCG/Private/Graph/PCGGraphCache.h
+++ b/Engine/Plugins/PCG/Source/PCG/Private/Graph/PCGGraphCache.h
@@ -4,8 +4,10 @@
 
 #include "PCGCrc.h"
 #include "PCGData.h"
+#include "Graph/IPCGGraphCache.h"
 
 #include "Containers/LruCache.h"
+#include "Misc/SpinLock.h"
 #include "UObject/GCObject.h"
 
 class IPCGElement;
@@ -61,16 +63,17 @@ private:
 * In cases where we have some subgraph reuse. Under that premise, we can then
 * instead store by element, as we will never recreate elements (except for arbitrary tasks)
 */
-struct FPCGGraphCache : public FGCObject
+class FPCGGraphCache : public FGCObject, public IPCGGraphCache
 {
+public:
 	FPCGGraphCache();
 	~FPCGGraphCache();
 
 	/** Returns true if data was found from the cache, in which case the outputs are written in OutOutput. InNode is optional and for logging only. */
-	bool GetFromCache(const UPCGNode* InNode, const IPCGElement* InElement, const FPCGCrc& InCrc, const UPCGComponent* InComponent, FPCGDataCollection& OutOutput) const;
+	virtual bool GetFromCache(const FPCGGetFromCacheParams& Params, FPCGDataCollection& OutCollection) const override;
 
 	/** Stores data in the cache for later use */
-	void StoreInCache(const IPCGElement* InElement, const FPCGCrc& InCrc, const FPCGDataCollection& InOutput);
+	virtual void StoreInCache(const FPCGStoreInCacheParams& Params, const FPCGDataCollection& InCollection) override;
 
 	/** Removes all entries from the cache, unroots data, etc. */
 	void ClearCache();
@@ -114,5 +117,5 @@ private:
 	/** Total memory usage by all data objects in cache. */
 	uint64 TotalMemoryUsed = 0;
 
-	mutable FRWLock CacheLock;
+	mutable UE::FSpinLock CacheLock;
 };
diff --git a/Engine/Plugins/PCG/Source/PCG/Private/Graph/PCGGraphExecutor.cpp b/Engine/Plugins/PCG/Source/PCG/Private/Graph/PCGGraphExecutor.cpp
--- a/Engine/Plugins/PCG/Source/PCG/Private/Graph/PCGGraphExecutor.cpp
+++ b/Engine/Plugins/PCG/Source/PCG/Private/Graph/PCGGraphExecutor.cpp
@@ -268,27 +268,31 @@ FPCGGraphActiveTask::~FPCGGraphActiveTask()
 void FPCGGraphActiveTask::StartExecuting()
 {
 	check(!bIsExecutingTask);
-	bIsExecutingTask = true;
 	++NumExecuting;
 
-	// @todo_pcg this should be merged through some kind of visitor pattern with AddStructReferencedObjects
-	// We need to keep a immutable list of references prior to execution mainly so that TaggedData contents don't change while we are running GC
+	auto GatherReferences = [this](const FPCGDataCollection& Collection)
 	{
-		ExecutingReferences.Reserve(Context->InputData.TaggedData.Num() + Context->OutputData.TaggedData.Num() + 1);
-		for (const FPCGTaggedData& TaggedInputData : Context->InputData.TaggedData)
+		for (const FPCGTaggedData& TaggedInputData : Collection.TaggedData)
 		{
 			if (TaggedInputData.Data)
 			{
 				ExecutingReferences.Add(TaggedInputData.Data);
 			}
 		}
+	};
+
+	// @todo_pcg this should be merged through some kind of visitor pattern with AddStructReferencedObjects
+	// We need to keep a immutable list of references prior to execution mainly so that TaggedData contents don't change while we are running GC
+	{
+		ExecutingReferences.Reserve(Context->InputData.TaggedData.Num() + Context->OutputData.TaggedData.Num() + 1);
+		GatherReferences(Context->InputData);
+		GatherReferences(Context->OutputData);
 
-		for (const FPCGTaggedData& TaggedOutputData : Context->OutputData.TaggedData)
+		// Additional cached per-data
+		for (TPair<FPCGDataCollection, FPCGDataCollection>& CachedInputToOutput : Context->CachedInputToOutputInternalResults)
 		{
-			if (TaggedOutputData.Data)
-			{
-				ExecutingReferences.Add(TaggedOutputData.Data);
-			}
+			GatherReferences(CachedInputToOutput.Key);
+			GatherReferences(CachedInputToOutput.Value);
 		}
 
 		if (Context->SettingsWithOverride)
@@ -296,6 +300,8 @@ void FPCGGraphActiveTask::StartExecuting()
 			ExecutingReferences.Add(Context->SettingsWithOverride);
 		}
 	}
+
+	bIsExecutingTask = true;
 }
 
 void FPCGGraphActiveTask::StopExecuting()
@@ -932,6 +938,7 @@ bool FPCGGraphExecutor::ProcessScheduledTasks()
 	{
 		PCGGraphExecutor::TScopeLock ScopeLock(TasksLock);
 		PCGGraphExecutor::TScopeLock ChildScopeLock(LiveTasksLock);
+		PCGGraphExecutor::TScopeLock CachingResultsScopeLock(CachingResultsLock);
 		// This is a safeguard to check if we're in a stuck state
 		if (CollectGCCachingResults.Num() == 0 && ReadyTasks.Num() == 0 && ActiveTasks.Num() == 0 && ActiveTasksGameThreadOnly.Num() == 0 && SleepingTasks.Num() == 0 && Tasks.Num() > 0)
 		{
@@ -1054,23 +1061,6 @@ void FPCGGraphExecutor::PostTaskExecute(TSharedPtr<FPCGGraphActiveTask> ActiveTa
 		}
 	}
 
-	if (bTaskFullyExecuted && !bTaskWasBypassed)
-	{
-		// Store result in cache as needed - done here because it needs to be done on the main thread
-
-		// Don't store if errors or warnings present
-#if WITH_EDITOR
-		const bool bHasErrorOrWarning = ActiveTask.Context->Node && (ActiveTask.Context->HasVisualLogs());
-#else
-		const bool bHasErrorOrWarning = false;
-#endif
-
-		if (ActiveTaskSettingsInterface && !bHasErrorOrWarning && ActiveTask.Element->IsCacheableInstance(ActiveTaskSettingsInterface))
-		{
-			GraphCache.StoreInCache(ActiveTask.Element.Get(), ActiveTask.Context->DependenciesCrc, ActiveTask.Context->OutputData);
-		}
-	}
-
 	if (ActiveTask.Context->AsyncState.NumAvailableTasks > 0)
 	{
 		CurrentlyUsedThreads -= ActiveTask.Context->AsyncState.NumAvailableTasks;
@@ -1222,31 +1212,20 @@ void FPCGGraphExecutor::PrepareForExecute(FPCGGraphTask& Task, FCachedResult*& O
 	// there is an execution mode that would prevent us from doing so.
 	const UPCGSettingsInterface* TaskSettingsInterface = Task.TaskInput.GetSettingsInterface(Task.Node ? Task.Node->GetSettingsInterface() : nullptr);
 	const UPCGSettings* TaskSettings = TaskSettingsInterface ? TaskSettingsInterface->GetSettings() : nullptr;
-	const bool bCacheable = Task.Element->IsCacheableInstance(TaskSettingsInterface);
 
 	// Calculate Crc of dependencies (input data Crcs, settings) and use this as the key in the cache lookup
 	FPCGCrc DependenciesCrc;
-	if (TaskSettings && bCacheable)
-	{
-		Task.Element->GetDependenciesCrc(Task.TaskInput, TaskSettings, Task.SourceComponent.Get(), DependenciesCrc);
-	}
+	EPCGCachingStatus CacheStatus = EPCGCachingStatus::NotInCache;
+	bool bResultAlreadyInCache = false;
 
-	if (!bCacheable)
+	// This section requires the lock on the caching results to prevent interaction between different tasks and the GC.
 	{
-		PCGGraphExecutionLogging::LogTaskExecuteCachingDisabled(Task);
-	}
-		
-	auto GetFromCache = [this, bLiveTasksLockAlreadyLocked, &Task, &DependenciesCrc, &OutCachedResult]() -> bool
-	{
-		// Lock LiveTasksLock if not already done so that we don't have a window where Cache can get cleaned on mainthread and GC runs while we have a unreferenced FPCGDataCollection
-		if(!bLiveTasksLockAlreadyLocked)
-		{
-			LiveTasksLock.Lock();
-		}
-		
+		PCGGraphExecutor::TScopeLock ScopeLock(CachingResultsLock);
 		FCachedResult LocalCachedResult;
-		const bool bFoundInCache = GraphCache.GetFromCache(Task.Node, Task.Element.Get(), DependenciesCrc, Task.SourceComponent.Get(), LocalCachedResult.Output);
-		if (bFoundInCache)
+		CacheStatus = Task.Element->RetrieveResultsFromCache(&GraphCache, Task.Node, Task.TaskInput, Task.SourceComponent.Get(), LocalCachedResult.Output, &DependenciesCrc);
+
+		bResultAlreadyInCache = (CacheStatus == EPCGCachingStatus::Cached);
+		if (bResultAlreadyInCache)
 		{
 			check(!CollectGCCachingResults.Contains(Task.NodeId));
 			TUniquePtr<FCachedResult>& CachingResultPtr = CollectGCCachingResults.Add(Task.NodeId, MakeUnique<FCachedResult>());
@@ -1254,16 +1233,13 @@ void FPCGGraphExecutor::PrepareForExecute(FPCGGraphTask& Task, FCachedResult*& O
 			OutCachedResult->TaskId = Task.NodeId;
 			OutCachedResult->Output = MoveTemp(LocalCachedResult.Output);
 		}
+	}
 
-		if (!bLiveTasksLockAlreadyLocked)
-		{
-			LiveTasksLock.Unlock();
-		}
-
-		return bFoundInCache;
-	};
+	if (CacheStatus == EPCGCachingStatus::NotCacheable)
+	{
+		PCGGraphExecutionLogging::LogTaskExecuteCachingDisabled(Task);
+	}
 
-	const bool bResultAlreadyInCache = bCacheable && DependenciesCrc.IsValid() && GetFromCache();
 #if WITH_EDITOR
 	const bool bNeedsToCreateActiveTask = !bResultAlreadyInCache || TaskSettingsInterface->bDebug;
 #else
@@ -2058,7 +2034,7 @@ TArray<FPCGTaskId> FPCGGraphExecutor::ProcessCachedResultsInternal(TArray<FCache
 
 	if (CachedResults.Num() > 0)
 	{
-		PCGGraphExecutor::TScopeLock ScopeLock(LiveTasksLock);
+		PCGGraphExecutor::TScopeLock ScopeLock(CachingResultsLock);
 		for (const FCachedResult* CachedResult : CachedResults)
 		{
 			check(CollectGCCachingResults.Contains(CachedResult->TaskId));
@@ -2572,7 +2548,10 @@ void FPCGGraphExecutor::AddReferencedObjects(FReferenceCollector& Collector)
 		Algo::ForEach(ActiveTasks, AddReferencesActiveTask);
 		Algo::ForEach(ActiveTasksGameThreadOnly, AddReferencesActiveTask);
 		Algo::ForEach(SleepingTasks, AddReferencesActiveTask);
+	}
 
+	{
+		PCGGraphExecutor::TScopeLock ScopeLock(CachingResultsLock);
 		Algo::ForEach(CollectGCCachingResults, [&Collector](auto& CachingResult)
 		{
 			CachingResult.Value->Output.AddReferences(Collector);
diff --git a/Engine/Plugins/PCG/Source/PCG/Private/Graph/PCGGraphExecutor.h b/Engine/Plugins/PCG/Source/PCG/Private/Graph/PCGGraphExecutor.h
--- a/Engine/Plugins/PCG/Source/PCG/Private/Graph/PCGGraphExecutor.h
+++ b/Engine/Plugins/PCG/Source/PCG/Private/Graph/PCGGraphExecutor.h
@@ -357,9 +357,6 @@ private:
 	TArray<TSharedPtr<FPCGGraphActiveTask>> SleepingTasks;
 	bool bNeedToCheckSleepingTasks = false;
 
-	// Used to keep GC references to in flight caching results (not yet stored to output and might not be in cache anymore)
-	TMap<FPCGTaskId, TUniquePtr<FCachedResult>> CollectGCCachingResults;
-
 	/** Lock level 3 */
 	UE::FSpinLock CollectGCReferenceTasksLock;
 	TSet<TSharedPtr<FPCGGraphActiveTask>> CollectGCReferenceTasks;
@@ -377,6 +374,11 @@ private:
 		bool bCulled = false;
 	};
 
+	/** Lock level 4 */
+	UE::FSpinLock CachingResultsLock;
+	// Used to keep GC references to in flight caching results (not yet stored to output and might not be in cache anymore)
+	TMap<FPCGTaskId, TUniquePtr<FCachedResult>> CollectGCCachingResults;
+
 	/** Lock level 4 */
 	UE::FSpinLock TaskOutputsLock;
 	TMap<FPCGTaskId, FOutputDataInfo> TaskOutputs;
diff --git a/Engine/Plugins/PCG/Source/PCG/Private/PCGContext.cpp b/Engine/Plugins/PCG/Source/PCG/Private/PCGContext.cpp
--- a/Engine/Plugins/PCG/Source/PCG/Private/PCGContext.cpp
+++ b/Engine/Plugins/PCG/Source/PCG/Private/PCGContext.cpp
@@ -401,6 +401,12 @@ void FPCGContext::AddStructReferencedObjects(FReferenceCollector& Collector)
 	InputData.AddReferences(Collector);
 	OutputData.AddReferences(Collector);
 
+	for (TPair<FPCGDataCollection, FPCGDataCollection>& CachedIOResult : CachedInputToOutputInternalResults)
+	{
+		CachedIOResult.Key.AddReferences(Collector);
+		CachedIOResult.Value.AddReferences(Collector);
+	}
+
 	if (SettingsWithOverride)
 	{
 		Collector.AddReferencedObject(SettingsWithOverride);
diff --git a/Engine/Plugins/PCG/Source/PCG/Private/PCGData.cpp b/Engine/Plugins/PCG/Source/PCG/Private/PCGData.cpp
--- a/Engine/Plugins/PCG/Source/PCG/Private/PCGData.cpp
+++ b/Engine/Plugins/PCG/Source/PCG/Private/PCGData.cpp
@@ -437,6 +437,12 @@ void FPCGDataCollection::ComputeCrcs(bool bFullDataCrc)
 	}
 }
 
+void FPCGDataCollection::AddData(const FPCGTaggedData& InData, const FPCGCrc& InDataCrc)
+{
+	TaggedData.Add(InData);
+	DataCrcs.Add(InDataCrc);
+}
+
 void FPCGDataCollection::AddData(const TConstArrayView<FPCGTaggedData>& InData, const TConstArrayView<FPCGCrc>& InDataCrcs)
 {
 	TaggedData.Append(InData);
diff --git a/Engine/Plugins/PCG/Source/PCG/Private/PCGElement.cpp b/Engine/Plugins/PCG/Source/PCG/Private/PCGElement.cpp
--- a/Engine/Plugins/PCG/Source/PCG/Private/PCGElement.cpp
+++ b/Engine/Plugins/PCG/Source/PCG/Private/PCGElement.cpp
@@ -7,6 +7,7 @@
 #include "PCGGraph.h"
 #include "PCGSubsystem.h"
 #include "Data/PCGPointData.h"
+#include "Graph/PCGGraphCache.h"
 
 #if WITH_EDITOR
 #include "PCGDataVisualization.h"
@@ -22,6 +23,11 @@ static TAutoConsoleVariable<bool> CVarPCGValidatePointMetadata(
 	true,
 	TEXT("Controls whether we validate that the metadata entry keys on the output point data are consistent"));
 
+static TAutoConsoleVariable<bool> CVarPCGAllowPerDataCaching(
+	TEXT("pcg.AllowPerDataCaching"),
+	false,
+	TEXT("Controls whether we test & split down inputs to check caching per input on primary loop nodes."));
+
 #if WITH_EDITOR
 #define PCG_ELEMENT_EXECUTION_BREAKPOINT() \
 	if (Context && Context->GetInputSettingsInterface() && Context->GetInputSettingsInterface()->bBreakDebugger) \
@@ -32,6 +38,113 @@ static TAutoConsoleVariable<bool> CVarPCGValidatePointMetadata(
 #define PCG_ELEMENT_EXECUTION_BREAKPOINT()
 #endif
 
+namespace PCGElementHelpers
+{
+	bool SplitDataPerPrimaryPin(const UPCGSettings* Settings, const FPCGDataCollection& Collection, EPCGElementExecutionLoopMode Mode, TArray<FPCGDataCollection>& OutPrimaryCollections, FPCGDataCollection& OutCommonCollection)
+	{
+		check(Settings);
+		OutPrimaryCollections.Reset();
+		OutCommonCollection.TaggedData.Reset();
+
+		TArray<FPCGPinProperties> RequiredPins = Settings->AllInputPinProperties().FilterByPredicate([](const FPCGPinProperties& Props) { return Props.IsRequiredPin(); });
+
+		// Early out
+		if (Mode == EPCGElementExecutionLoopMode::SinglePrimaryPin && RequiredPins.Num() != 1)
+		{
+			return false;
+		}
+
+		TArray<FName> RequiredPinLabels;
+		Algo::Transform(RequiredPins, RequiredPinLabels, [](const FPCGPinProperties& Props) { return Props.Label; });
+
+		TArray<FPCGDataCollection> DataPerRequiredPin;
+		DataPerRequiredPin.SetNum(RequiredPinLabels.Num());
+
+		for (int32 DataIndex = 0; DataIndex < Collection.TaggedData.Num(); ++DataIndex)
+		{
+			const FPCGTaggedData& TaggedData = Collection.TaggedData[DataIndex];
+			int32 RequiredPinIndex = RequiredPinLabels.IndexOfByKey(TaggedData.Pin);
+
+			if (RequiredPinIndex == INDEX_NONE)
+			{
+				OutCommonCollection.TaggedData.Add(TaggedData);
+			}
+			else
+			{
+				DataPerRequiredPin[RequiredPinIndex].TaggedData.Add(TaggedData);
+			}
+		}
+
+		if (DataPerRequiredPin.IsEmpty())
+		{
+			return true;
+		}
+
+		// Broadcast to final primary collections
+		if (Mode == EPCGElementExecutionLoopMode::SinglePrimaryPin)
+		{
+			check(DataPerRequiredPin.Num() == 1);
+			OutPrimaryCollections.Reserve(DataPerRequiredPin[0].TaggedData.Num());
+
+			for(int32 DataIndex = 0; DataIndex < DataPerRequiredPin[0].TaggedData.Num(); ++DataIndex)
+			{
+				FPCGDataCollection& OutPrimaryCollection = OutPrimaryCollections.Emplace_GetRef();
+				OutPrimaryCollection.TaggedData.Add(DataPerRequiredPin[0].TaggedData[DataIndex]);
+			}
+		}
+		else if (Mode == EPCGElementExecutionLoopMode::MatchingPrimaryPins)
+		{
+			const int32 NumberOfData = DataPerRequiredPin[0].TaggedData.Num();
+
+			// Validate matching number of entries
+			for (int32 RequiredPinIndex = 1; RequiredPinIndex < DataPerRequiredPin.Num(); ++RequiredPinIndex)
+			{
+				if (DataPerRequiredPin[RequiredPinIndex].TaggedData.Num() != NumberOfData)
+				{
+					return false;
+				}
+			}
+
+			OutPrimaryCollections.SetNum(NumberOfData);
+
+			for (int32 DataIndex = 0; DataIndex < NumberOfData; ++DataIndex)
+			{
+				for (int32 RequiredPinIndex = 0; RequiredPinIndex < DataPerRequiredPin.Num(); ++RequiredPinIndex)
+				{
+					OutPrimaryCollections[DataIndex].TaggedData.Add(DataPerRequiredPin[RequiredPinIndex].TaggedData[DataIndex]);
+				}
+			}
+		}
+		/*else if (Mode == EPCGElementExecutionLoopMode::CartesianPins)
+		{
+			int32 NumberOfCollections = 1;
+			for (int32 RequiredPinIndex = 0; RequiredPinIndex < DataPerRequiredPin.Num(); ++RequiredPinIndex)
+			{
+				NumberOfCollections *= DataPerRequiredPin[RequiredPinIndex].TaggedData.Num();
+			}
+
+			OutPrimaryCollections.SetNum(NumberOfCollections);
+
+			for (int32 RequiredPinIndex = 0; RequiredPinIndex < DataPerRequiredPin.Num(); ++RequiredPinIndex)
+			{
+				const FPCGDataCollection& PinData = DataPerRequiredPin[RequiredPinIndex];
+
+				for (int32 OutCollectionIndex = 0; OutCollectionIndex < OutPrimaryCollections.Num(); ++OutCollectionIndex)
+				{
+					OutPrimaryCollections[OutCollectionIndex].TaggedData.Add(PinData.TaggedData[OutCollectionIndex % PinData.TaggedData.Num()]);
+				}
+			}
+		}*/
+		else
+		{
+			// Invalid mode
+			return false;
+		}
+
+		return true;
+	}
+}
+
 bool IPCGElement::Execute(FPCGContext* Context) const
 {
 	check(Context && Context->AsyncState.NumAvailableTasks != 0 && Context->CurrentPhase < EPCGExecutionPhase::Done);
@@ -51,9 +164,6 @@ bool IPCGElement::Execute(FPCGContext* Context) const
 
 				PreExecute(Context);
 
-				// Will override the settings if there is any override.
-				Context->OverrideSettings();
-
 				break;
 			}
 
@@ -118,6 +228,7 @@ bool IPCGElement::Execute(FPCGContext* Context) const
 
 void IPCGElement::PreExecute(FPCGContext* Context) const
 {
+	check(Context);
 	// Check for early outs (task cancelled + node disabled)
 	// Early out to stop execution
 	if (Context->InputData.bCancelExecution || (!Context->SourceComponent.IsExplicitlyNull() && !Context->SourceComponent.IsValid()))
@@ -136,9 +247,8 @@ void IPCGElement::PreExecute(FPCGContext* Context) const
 	Context->CurrentPhase = EPCGExecutionPhase::PrepareData;
 
 	const UPCGSettingsInterface* SettingsInterface = Context->GetInputSettingsInterface();
-	const UPCGSettings* Settings = SettingsInterface ? SettingsInterface->GetSettings() : nullptr;
 
-	if (!SettingsInterface || !Settings)
+	if (!SettingsInterface)
 	{
 		return;
 	}
@@ -148,9 +258,94 @@ void IPCGElement::PreExecute(FPCGContext* Context) const
 		//Pass-through - no execution
 		DisabledPassThroughData(Context);
 		Context->CurrentPhase = EPCGExecutionPhase::PostExecute;
+		return;
+	}
+
+	// Will override the settings if there is any override.
+	Context->OverrideSettings();
+
+	if (CVarPCGAllowPerDataCaching.GetValueOnAnyThread())
+	{
+		const UPCGSettings* Settings = Context->GetInputSettings<UPCGSettings>();
+
+		// Default implementation when the entries in a primary loop can be processed independently, e.g. they can appear in the cache separately
+		// Implementation note: this supposes that the current node has only ONE required pin, and not multiple.
+		// For more complex cases (such as multiple required pins, whether cartesian or matching), the implementation should use common code instead to streamline this process -
+		// both when getting the results from the cache but also when writing them
+		if (ExecutionLoopMode(Settings) != EPCGElementExecutionLoopMode::NotALoop && IsCacheable(Settings))
+		{
+			PreExecutePrimaryLoopElement(Context, Settings);
+		}
 	}
 }
 
+void IPCGElement::PreExecutePrimaryLoopElement(FPCGContext* Context, const UPCGSettings* Settings) const
+{
+	check(Context);
+
+	UPCGSubsystem* Subsystem = Context->SourceComponent.Get() ? Context->SourceComponent->GetSubsystem() : nullptr;
+	if (!Subsystem || !Settings)
+	{
+		return;
+	}
+
+	// Mark inputs in the order they're presented so we can appropriately find the relation from output to input after the execution
+	// TODO: this is not sufficient to do a proper mapping from output to input when we have a cartesian loop
+	for (int32 DataIndex = 0; DataIndex < Context->InputData.TaggedData.Num(); ++DataIndex)
+	{
+		Context->InputData.TaggedData[DataIndex].OriginalIndex = DataIndex;
+	}
+
+	TArray<FPCGDataCollection> PrimaryDataCollections;
+	FPCGDataCollection OtherData;
+	if (!PCGElementHelpers::SplitDataPerPrimaryPin(Settings, Context->InputData, ExecutionLoopMode(Settings), PrimaryDataCollections, OtherData))
+	{
+		return;
+	}
+
+	// Implementation note: if there is a single primary data collection, then there's no point checking in the cache again, since that has already been done.
+	if (PrimaryDataCollections.Num() <= 1)
+	{
+		return;
+	}
+
+	const bool bShouldComputeFullOutputDataCrc = ShouldComputeFullOutputDataCrc(Context);
+
+	// Check against the cache if subcollections of one data from the primary data collection + the other data is found already in the cache.
+	// If so, we can remove the matching input data - note that this is somewhat trivial in the single pin & matching pin cases, but in general is not going to work for cartesian cases.
+	for (int32 PrimaryDataIndex = PrimaryDataCollections.Num() - 1; PrimaryDataIndex >= 0; --PrimaryDataIndex)
+	{
+		const FPCGDataCollection& PrimaryDataCollection = PrimaryDataCollections[PrimaryDataIndex];
+		FPCGDataCollection SubCollection = PrimaryDataCollection;
+		SubCollection.TaggedData.Append(OtherData.TaggedData);
+
+		SubCollection.ComputeCrcs(bShouldComputeFullOutputDataCrc);
+
+		FPCGGetFromCacheParams CacheParams = { .Node = Context->Node, .Element = this, .Component = Context->SourceComponent.Get() };
+		GetDependenciesCrc(SubCollection, Settings, Context->SourceComponent.Get(), CacheParams.Crc);
+
+		FPCGDataCollection SubCollectionOutput;
+		if(Subsystem->GetCache()->GetFromCache(CacheParams, SubCollectionOutput))
+		{
+			// Found a match in the cache, add it to the output, and remove the matching inputs.
+			// Note that in the input part we'll take only the data present in the PrimaryDataCollection here, as we will reintroduce only those then.
+			// IMPLEMENTATION NOTE: the order is important here; if we can't guarantee the hint index, then we should just do a remove!
+			TPair<FPCGDataCollection, FPCGDataCollection>& InputToOutputResults = Context->CachedInputToOutputInternalResults.Emplace_GetRef();
+			InputToOutputResults.Key = PrimaryDataCollection;
+			InputToOutputResults.Value = SubCollectionOutput;
+
+			for (int32 SubDataIndex = PrimaryDataCollection.TaggedData.Num() - 1; SubDataIndex >= 0; --SubDataIndex)
+			{
+				Context->InputData.TaggedData.RemoveAt(PrimaryDataCollection.TaggedData[SubDataIndex].OriginalIndex);
+			}
+		}
+	}
+
+	//
+	// TODO : if there are no inputs left, then we could skip the execute phase
+	// 
+}
+
 bool IPCGElement::PrepareDataInternal(FPCGContext* Context) const
 {
 	return true;
@@ -163,10 +358,17 @@ void IPCGElement::PostExecute(FPCGContext* Context) const
 
 	// Cleanup and validate output
 	CleanupAndValidateOutput(Context);
-
-	const UPCGSettingsInterface* SettingsInterface = Context->GetInputSettingsInterface();
-	const UPCGSettings* Settings = SettingsInterface ? SettingsInterface->GetSettings() : nullptr;
 	
+	const UPCGSettings* Settings = Context->GetInputSettings<UPCGSettings>();
+
+	if (CVarPCGAllowPerDataCaching.GetValueOnAnyThread())
+	{
+		if (!Context->OutputData.bCancelExecution && ExecutionLoopMode(Settings) != EPCGElementExecutionLoopMode::NotALoop && IsCacheable(Settings))
+		{
+			PostExecutePrimaryLoopElement(Context, Settings);
+		}
+	}
+
 	// Output data Crc
 	{
 		TRACE_CPUPROFILER_EVENT_SCOPE(IPCGElement::PostExecute::CRC);
@@ -180,6 +382,25 @@ void IPCGElement::PostExecute(FPCGContext* Context) const
 		Context->OutputData.ComputeCrcs(bShouldComputeFullOutputDataCrc);
 	}
 
+#if WITH_EDITOR
+	const bool bHasErrorsOrWarnings = Context->Node && Context->HasVisualLogs();
+#else
+	const bool bHasErrorsOrWarnings = false;
+#endif
+
+	// Store result in cache
+	// TODO - There is a potential mismatch here between using the Settings (incl. overrides) and the input settings interface (pre-overrides), as done in the graph executor.
+	// TODO - The dependencies CRC here should always be valid except in the indirection case, which we should normalize to allow caching (tested here, otherwise it can ensure in the graph cache)
+	if (!Context->OutputData.bCancelExecution && !bHasErrorsOrWarnings && Context->DependenciesCrc.IsValid() && IsCacheableInstance(Settings))
+	{
+		UPCGSubsystem* Subsystem = Context->SourceComponent.Get() ? Context->SourceComponent->GetSubsystem() : nullptr;
+		if (Subsystem)
+		{
+			FPCGStoreInCacheParams Params = { .Element = this, .Crc = Context->DependenciesCrc };
+			Subsystem->GetCache()->StoreInCache(Params, Context->OutputData);
+		}
+	}
+
 #if WITH_EDITOR
 	// Register the element to the component indicating the element has run and can have dynamic tracked keys.
 	if (Settings && Settings->CanDynamicallyTrackKeys() && Context->SourceComponent.IsValid())
@@ -191,6 +412,98 @@ void IPCGElement::PostExecute(FPCGContext* Context) const
 	Context->CurrentPhase = EPCGExecutionPhase::Done;
 }
 
+void IPCGElement::PostExecutePrimaryLoopElement(FPCGContext* Context, const UPCGSettings* Settings) const
+{
+	check(Context);
+
+	// In the case of primary pin-loops, write back individual results to the cache, and reinsert the cached results in the output as needed.
+	UPCGSubsystem* Subsystem = Context->SourceComponent.Get() ? Context->SourceComponent->GetSubsystem() : nullptr;
+
+	if (!Subsystem || !Settings)
+	{
+		return;
+	}
+
+#if WITH_EDITOR
+	const bool bHasErrorsOrWarnings = Context->Node && Context->HasVisualLogs();
+#else
+	const bool bHasErrorsOrWarnings = false;
+#endif
+
+	// Store individual results in the cache; here we will try to match the remaining hint indices from the input data with the ones given at the output.
+	TArray<FPCGDataCollection> PrimaryDataCollections;
+	FPCGDataCollection OtherData;
+	if (!bHasErrorsOrWarnings && PCGElementHelpers::SplitDataPerPrimaryPin(Settings, Context->InputData, ExecutionLoopMode(Settings), PrimaryDataCollections, OtherData))
+	{
+		const bool bShouldComputeFullOutputDataCrc = ShouldComputeFullOutputDataCrc(Context);
+
+		for (const FPCGDataCollection& PrimaryDataCollection : PrimaryDataCollections)
+		{
+			if (PrimaryDataCollection.TaggedData.IsEmpty() || PrimaryDataCollection.TaggedData[0].OriginalIndex == INDEX_NONE)
+			{
+				continue;
+			}
+
+			FPCGDataCollection SubCollectionOutput;
+			for (int32 DataIndex = 0; DataIndex < Context->OutputData.TaggedData.Num(); ++DataIndex)
+			{
+				const FPCGTaggedData& TaggedData = Context->OutputData.TaggedData[DataIndex];
+				if (TaggedData.OriginalIndex == PrimaryDataCollection.TaggedData[0].OriginalIndex)
+				{
+					SubCollectionOutput.TaggedData.Add(Context->OutputData.TaggedData[DataIndex]);
+				}
+			}
+
+			FPCGDataCollection SubCollection = PrimaryDataCollection;
+			SubCollection.ComputeCrcs(bShouldComputeFullOutputDataCrc);
+
+			FPCGCrc DependenciesCrc;
+			GetDependenciesCrc(SubCollection, Settings, Context->SourceComponent.Get(), DependenciesCrc);
+
+			SubCollectionOutput.ComputeCrcs(bShouldComputeFullOutputDataCrc);
+
+			FPCGStoreInCacheParams Params = { .Element = this, .Crc = DependenciesCrc };
+			Subsystem->GetCache()->StoreInCache(Params, SubCollectionOutput);
+		}
+	}
+
+	// Put back cached results and set aside input (needed for inspection) if any
+	if (!Context->CachedInputToOutputInternalResults.IsEmpty())
+	{
+		// Push cached results back to the final output data, from the last to the first, at the right place.
+		for (int CachedCollectionIndex = Context->CachedInputToOutputInternalResults.Num() - 1; CachedCollectionIndex >= 0; --CachedCollectionIndex)
+		{
+			const FPCGDataCollection& CachedInputData = Context->CachedInputToOutputInternalResults[CachedCollectionIndex].Key;
+			int32 CacheInputOriginalIndex = (CachedInputData.TaggedData.IsEmpty() ? INDEX_NONE : CachedInputData.TaggedData[0].OriginalIndex);
+			int32 InsertInputIndex = CacheInputOriginalIndex != INDEX_NONE ? Context->InputData.TaggedData.IndexOfByPredicate([CacheInputOriginalIndex](const FPCGTaggedData& TaggedData) { return TaggedData.OriginalIndex > CacheInputOriginalIndex; }) : INDEX_NONE;
+
+			if (InsertInputIndex != INDEX_NONE)
+			{
+				Context->InputData.TaggedData.Insert(CachedInputData.TaggedData, InsertInputIndex);
+			}
+			else
+			{
+				Context->InputData.TaggedData.Append(CachedInputData.TaggedData);
+			}
+
+			const FPCGDataCollection& CachedOutputData = Context->CachedInputToOutputInternalResults[CachedCollectionIndex].Value;
+
+			// This assumes we have hinted properly on the output data; we'll look at the hinted value on the first data THEN insert before the next hint value
+			int32 CacheOutputOriginalIndex = (CachedOutputData.TaggedData.IsEmpty() ? INDEX_NONE : CachedOutputData.TaggedData[0].OriginalIndex);
+			int32 InsertOutputIndex = CacheOutputOriginalIndex != INDEX_NONE ? Context->OutputData.TaggedData.IndexOfByPredicate([CacheOutputOriginalIndex](const FPCGTaggedData& TaggedData) { return TaggedData.OriginalIndex > CacheOutputOriginalIndex; }) : INDEX_NONE;
+
+			if (InsertOutputIndex != INDEX_NONE)
+			{
+				Context->OutputData.TaggedData.Insert(CachedOutputData.TaggedData, InsertOutputIndex);
+			}
+			else
+			{
+				Context->OutputData.TaggedData.Append(CachedOutputData.TaggedData);
+			}
+		}
+	}
+}
+
 void IPCGElement::Abort(FPCGContext* Context) const
 {
 	AbortInternal(Context);
@@ -525,4 +838,37 @@ void IPCGElement::GetDependenciesCrc(const FPCGDataCollection& InInput, const UP
 	OutCrc = Crc;
 }
 
+EPCGCachingStatus IPCGElement::RetrieveResultsFromCache(IPCGGraphCache* Cache, const UPCGNode* Node, const FPCGDataCollection& Input, UPCGComponent* Component, FPCGDataCollection& Output, FPCGCrc* OutCrc) const
+{
+	if (!Cache)
+	{
+		return EPCGCachingStatus::NotInCache;
+	}
+
+	const UPCGSettingsInterface* SettingsInterface = Input.GetSettingsInterface(Node ? Node->GetSettingsInterface() : nullptr);
+	const UPCGSettings* Settings = SettingsInterface ? SettingsInterface->GetSettings() : nullptr;
+	const bool bCacheable = IsCacheableInstance(SettingsInterface);
+
+	FPCGGetFromCacheParams Params = { .Node = Node, .Element = this, .Component = Component };
+
+	if (Settings && bCacheable)
+	{
+		GetDependenciesCrc(Input, Settings, Component, Params.Crc);
+
+		if (OutCrc)
+		{
+			*OutCrc = Params.Crc;
+		}
+	}
+
+	if(Params.Crc.IsValid() && Cache->GetFromCache(Params, Output))
+	{
+		return EPCGCachingStatus::Cached;
+	}
+	else
+	{
+		return EPCGCachingStatus::NotInCache;
+	}
+}
+
 #undef LOCTEXT_NAMESPACE
diff --git a/Engine/Plugins/PCG/Source/PCG/Private/PCGSubsystem.cpp b/Engine/Plugins/PCG/Source/PCG/Private/PCGSubsystem.cpp
--- a/Engine/Plugins/PCG/Source/PCG/Private/PCGSubsystem.cpp
+++ b/Engine/Plugins/PCG/Source/PCG/Private/PCGSubsystem.cpp
@@ -1826,6 +1826,11 @@ void UPCGSubsystem::ExecuteBeginTickActions()
 	}
 }
 
+IPCGGraphCache* UPCGSubsystem::GetCache()
+{
+	return GraphExecutor ? &(GraphExecutor->GetCache()) : nullptr;
+}
+
 void UPCGSubsystem::FlushCache()
 {
 	if (GraphExecutor)
diff --git a/Engine/Plugins/PCG/Source/PCG/Public/Graph/IPCGGraphCache.h b/Engine/Plugins/PCG/Source/PCG/Public/Graph/IPCGGraphCache.h
new file mode 100644
--- /dev/null
+++ b/Engine/Plugins/PCG/Source/PCG/Public/Graph/IPCGGraphCache.h
@@ -0,0 +1,32 @@
+// Copyright Epic Games, Inc. All Rights Reserved.
+
+#pragma once
+
+#include "PCGCrc.h"
+
+class IPCGElement;
+class UPCGComponent;
+class UPCGNode;
+struct FPCGDataCollection;
+
+struct FPCGGetFromCacheParams
+{
+	const UPCGNode* Node = nullptr;
+	const IPCGElement* Element = nullptr;
+	const UPCGComponent* Component = nullptr;
+	FPCGCrc Crc;
+};
+
+struct FPCGStoreInCacheParams
+{
+	const IPCGElement* Element = nullptr;
+	FPCGCrc Crc;
+};
+
+/** Interface to encapsulate use of the cache in PCG elements */
+class IPCGGraphCache
+{
+public:
+	virtual bool GetFromCache(const FPCGGetFromCacheParams& Params, FPCGDataCollection& OutCollection) const = 0;
+	virtual void StoreInCache(const FPCGStoreInCacheParams& Params, const FPCGDataCollection& InCollection) = 0;
+};
\ No newline at end of file
diff --git a/Engine/Plugins/PCG/Source/PCG/Public/PCGContext.h b/Engine/Plugins/PCG/Source/PCG/Public/PCGContext.h
--- a/Engine/Plugins/PCG/Source/PCG/Public/PCGContext.h
+++ b/Engine/Plugins/PCG/Source/PCG/Public/PCGContext.h
@@ -16,7 +16,6 @@ class UPCGComponent;
 class UPCGGraphInterface;
 class UPCGSettingsInterface;
 class UPCGSpatialData;
-struct FPCGGraphCache;
 struct FPCGSettingsOverridableParam;
 struct FPCGStack;
 
@@ -71,7 +70,7 @@ struct PCG_API FPCGContext
 	bool bIsPaused = false;
 
 	EPCGExecutionPhase CurrentPhase = EPCGExecutionPhase::NotExecuted;
-	int32 BypassedOutputCount = 0;
+	TArray<TPair<FPCGDataCollection, FPCGDataCollection>> CachedInputToOutputInternalResults;
 
 	/** The current call stack. */
 	const FPCGStack* Stack = nullptr;
diff --git a/Engine/Plugins/PCG/Source/PCG/Public/PCGData.h b/Engine/Plugins/PCG/Source/PCG/Public/PCGData.h
--- a/Engine/Plugins/PCG/Source/PCG/Public/PCGData.h
+++ b/Engine/Plugins/PCG/Source/PCG/Public/PCGData.h
@@ -104,6 +104,9 @@ struct PCG_API FPCGTaggedData
 	UPROPERTY()
 	bool bPinlessData = false;
 
+	// Used to re-order elements when using per-data caching
+	int32 OriginalIndex = INDEX_NONE;
+
 	bool operator==(const FPCGTaggedData& Other) const;
 	bool operator!=(const FPCGTaggedData& Other) const;
 
@@ -197,6 +200,8 @@ struct PCG_API FPCGDataCollection
 	/** Computes CRCs for all data items. */
 	void ComputeCrcs(bool bFullDataCrc);
 
+	/** Add data and CRC to collection. */
+	void AddData(const FPCGTaggedData& InData, const FPCGCrc& InDataCrc);
 	/** Add data and CRCs to collection. */
 	void AddData(const TConstArrayView<FPCGTaggedData>& InData, const TConstArrayView<FPCGCrc>& InDataCrcs);
 	/** Add data and CRCs to collection with pin label combined into the CRC. */
diff --git a/Engine/Plugins/PCG/Source/PCG/Public/PCGElement.h b/Engine/Plugins/PCG/Source/PCG/Public/PCGElement.h
--- a/Engine/Plugins/PCG/Source/PCG/Public/PCGElement.h
+++ b/Engine/Plugins/PCG/Source/PCG/Public/PCGElement.h
@@ -12,6 +12,7 @@ struct FPCGCrc;
 struct FPCGDataCollection;
 
 class IPCGElement;
+class IPCGGraphCache;
 class UPCGComponent;
 class UPCGSettings;
 class UPCGNode;
@@ -61,6 +62,34 @@ namespace EPCGElementLogMode
 #define PCGE_LOG_C(Verbosity, LogMode, CustomContext, Message) PCGE_LOG_BASE(Verbosity, CustomContext, Message)
 #endif
 
+enum EPCGElementExecutionLoopMode : uint8
+{
+	/** Not a trivial input -> output mapping, with respect to caching. */
+	NotALoop,
+	/** Loops on (singular) required pin. */
+	SinglePrimaryPin, 
+	/** Loops on matching indices on required pin(s). */
+	MatchingPrimaryPins,
+	/** Cartesian loop on required pins. */
+	// CartesianPins // TODO
+};
+
+enum EPCGCachingStatus : uint8
+{
+	NotCacheable,
+	NotInCache,
+	Cached
+};
+
+namespace PCGElementHelpers
+{
+	/** Breaks down the input data collection (InCollection) in a set of primary inputs (OutPrimaryCollections) and a set of fixed data per-itereation (OutCommonCollection), based on the provided mode. 
+	* This is needed to perform per-data caching in nodes that support it, either in single-primary-pin-loops or in matching-pins loops.
+	* Note that this uses the settings to drive the selection of the "primary" pins by testing if they are required.
+	*/
+	bool PCG_API SplitDataPerPrimaryPin(const UPCGSettings* Settings, const FPCGDataCollection& InCollection, EPCGElementExecutionLoopMode Mode, TArray<FPCGDataCollection>& OutPrimaryCollections, FPCGDataCollection& OutCommonCollection);
+};
+
 /**
 * Base class for the processing bit of a PCG node/settings
 */
@@ -91,6 +120,9 @@ public:
 	 */
 	virtual void GetDependenciesCrc(const FPCGDataCollection& InInput, const UPCGSettings* InSettings, UPCGComponent* InComponent, FPCGCrc& OutCrc) const;
 
+	/** Gather input data (pre-context creation) and tries to retrieve matching data from the cache if the element is cacheable. */
+	EPCGCachingStatus RetrieveResultsFromCache(IPCGGraphCache* Cache, const UPCGNode* Node, const FPCGDataCollection& Input, UPCGComponent* Component, FPCGDataCollection& Output, FPCGCrc* OutCrc = nullptr) const;
+
 	/** Public function that executes the element on the appropriately created context.
 	* The caller should call the Execute function until it returns true.
 	*/
@@ -132,6 +164,15 @@ protected:
 	/** Passes through data when the element is Disabled. Can be implemented to override what gets passed through. */
 	virtual void DisabledPassThroughData(FPCGContext* Context) const;
 
+	/** Describes internal execution behavior, which is used to break down inputs/outputs for caching purposes. */
+	virtual EPCGElementExecutionLoopMode ExecutionLoopMode(const UPCGSettings* Settings) const { return EPCGElementExecutionLoopMode::NotALoop; }
+
+	/** Implements input breakdown for caching purposes, if the ExecutionLoopMode is set to something else than not a loop. */
+	virtual void PreExecutePrimaryLoopElement(FPCGContext* Context, const UPCGSettings* Settings) const;
+
+	/** Implements output breakdown for caching purposes, if the ExecutionLoopMode is set to something else than not a loop. */
+	virtual void PostExecutePrimaryLoopElement(FPCGContext* Context, const UPCGSettings* Settings) const;
+
 	/** Let each element optionally act as a concrete factory for its own context */
 	virtual FPCGContext* CreateContext();
 
diff --git a/Engine/Plugins/PCG/Source/PCG/Public/PCGSubsystem.h b/Engine/Plugins/PCG/Source/PCG/Public/PCGSubsystem.h
--- a/Engine/Plugins/PCG/Source/PCG/Public/PCGSubsystem.h
+++ b/Engine/Plugins/PCG/Source/PCG/Public/PCGSubsystem.h
@@ -22,6 +22,7 @@ class FPCGRuntimeGenScheduler;
 enum class EPCGComponentDirtyFlag : uint8;
 enum class ETickableTickType : uint8;
 
+class IPCGGraphCache;
 class FPCGGraphCompiler;
 class FPCGGraphExecutor;
 struct FPCGContext;
@@ -210,6 +211,9 @@ public:
 	TSet<UPCGComponent*> GetAllRegisteredPartitionedComponents() const { return ActorAndComponentMapping.GetAllRegisteredPartitionedComponents(); }
 	TSet<UPCGComponent*> GetAllRegisteredComponents() const { return ActorAndComponentMapping.GetAllRegisteredComponents(); }
 
+	/** Returns the interface to the cache, required for element per-data caching */
+	IPCGGraphCache* GetCache();
+
 	/** Flushes the graph cache completely, use only for debugging */
 	void FlushCache();
 
