Fixed caching error in the world ray cast node

#rnx

diff --git a/Engine/Plugins/PCG/Source/PCG/Public/Elements/PCGWorldRaycast.h b/Engine/Plugins/PCG/Source/PCG/Public/Elements/PCGWorldRaycast.h
--- a/Engine/Plugins/PCG/Source/PCG/Public/Elements/PCGWorldRaycast.h
+++ b/Engine/Plugins/PCG/Source/PCG/Public/Elements/PCGWorldRaycast.h
@@ -104,6 +104,8 @@ namespace FPCGWorldRaycast
 
 class FPCGWorldRaycastElement : public TPCGTimeSlicedElementBase<FPCGWorldRaycast::FExecutionState, FPCGWorldRaycast::FIterationState>
 {
+public:
+	virtual bool IsCacheable(const UPCGSettings* InSettings) const override { return false; }
 protected:
 	virtual bool PrepareDataInternal(FPCGContext* InContext) const override;
 	virtual bool ExecuteInternal(FPCGContext* InContext) const override;
