Update content to support per-data crc caching

diff --git a/Engine/Plugins/PCG/Source/PCG/Private/Elements/PCGAttributeNoise.cpp b/Engine/Plugins/PCG/Source/PCG/Private/Elements/PCGAttributeNoise.cpp
--- a/Engine/Plugins/PCG/Source/PCG/Private/Elements/PCGAttributeNoise.cpp
+++ b/Engine/Plugins/PCG/Source/PCG/Private/Elements/PCGAttributeNoise.cpp
@@ -230,11 +230,6 @@ EPCGDataType UPCGAttributeNoiseSettings::GetCurrentPinTypes(const UPCGPin* InPin
 	}
 }
 
-FPCGContext* FPCGAttributeNoiseElement::CreateContext()
-{
-	return new FPCGAttributeNoiseContext();
-}
-
 bool FPCGAttributeNoiseElement::ExecuteInternal(FPCGContext* InContext) const
 {
 	TRACE_CPUPROFILER_EVENT_SCOPE(FPCGAttributeNoiseElement::Execute);
diff --git a/Engine/Plugins/PCG/Source/PCG/Private/Elements/PCGAttributeReduceElement.cpp b/Engine/Plugins/PCG/Source/PCG/Private/Elements/PCGAttributeReduceElement.cpp
--- a/Engine/Plugins/PCG/Source/PCG/Private/Elements/PCGAttributeReduceElement.cpp
+++ b/Engine/Plugins/PCG/Source/PCG/Private/Elements/PCGAttributeReduceElement.cpp
@@ -240,7 +240,8 @@ FString UPCGAttributeReduceSettings::GetAdditionalTitleInformation() const
 TArray<FPCGPinProperties> UPCGAttributeReduceSettings::InputPinProperties() const
 {
 	TArray<FPCGPinProperties> PinProperties;
-	PinProperties.Emplace(PCGPinConstants::DefaultInputLabel, EPCGDataType::Any);
+	FPCGPinProperties& InputPin = PinProperties.Emplace_GetRef(PCGPinConstants::DefaultInputLabel, EPCGDataType::Any);
+	InputPin.SetRequiredPin();
 
 	return PinProperties;
 }
@@ -403,4 +404,10 @@ bool FPCGAttributeReduceElement::ExecuteInternal(FPCGContext* Context) const
 	return true;
 }
 
+EPCGElementExecutionLoopMode FPCGAttributeReduceElement::ExecutionLoopMode(const UPCGSettings* InSettings) const
+{
+	const UPCGAttributeReduceSettings* Settings = Cast<UPCGAttributeReduceSettings>(InSettings);
+	return (!Settings || !Settings->bMergeOutputAttributes) ? EPCGElementExecutionLoopMode::SinglePrimaryPin : EPCGElementExecutionLoopMode::NotALoop;
+}
+
 #undef LOCTEXT_NAMESPACE
diff --git a/Engine/Plugins/PCG/Source/PCG/Private/Elements/PCGCreateSplineMesh.cpp b/Engine/Plugins/PCG/Source/PCG/Private/Elements/PCGCreateSplineMesh.cpp
--- a/Engine/Plugins/PCG/Source/PCG/Private/Elements/PCGCreateSplineMesh.cpp
+++ b/Engine/Plugins/PCG/Source/PCG/Private/Elements/PCGCreateSplineMesh.cpp
@@ -111,7 +111,6 @@ bool FPCGCreateSplineMeshElement::ExecuteInternal(FPCGContext* Context) const
 	SMCBuilderParams.SplineMeshParams = Settings->SplineMeshParams;
 
 	TArray<FPCGTaggedData> Inputs = Context->InputData.GetInputs();
-	TArray<FPCGTaggedData>& Outputs = Context->OutputData.TaggedData;
 
 	for (const FPCGTaggedData& Input : Inputs)
 	{
diff --git a/Engine/Plugins/PCG/Source/PCG/Private/Elements/PCGCreateSurfaceFromSpline.cpp b/Engine/Plugins/PCG/Source/PCG/Private/Elements/PCGCreateSurfaceFromSpline.cpp
--- a/Engine/Plugins/PCG/Source/PCG/Private/Elements/PCGCreateSurfaceFromSpline.cpp
+++ b/Engine/Plugins/PCG/Source/PCG/Private/Elements/PCGCreateSurfaceFromSpline.cpp
@@ -22,6 +22,8 @@ TArray<FPCGPinProperties> UPCGCreateSurfaceFromSplineSettings::InputPinPropertie
 {
 	TArray<FPCGPinProperties> PinProperties;
 	PinProperties.Emplace(PCGPinConstants::DefaultInputLabel, EPCGDataType::Spline, /*bAllowMultipleConnections=*/true, /*bAllowMultipleData=*/true);
+	PinProperties[0].SetRequiredPin();
+
 	return PinProperties;
 }
 
diff --git a/Engine/Plugins/PCG/Source/PCG/Private/Elements/PCGRandomChoice.cpp b/Engine/Plugins/PCG/Source/PCG/Private/Elements/PCGRandomChoice.cpp
--- a/Engine/Plugins/PCG/Source/PCG/Private/Elements/PCGRandomChoice.cpp
+++ b/Engine/Plugins/PCG/Source/PCG/Private/Elements/PCGRandomChoice.cpp
@@ -216,13 +216,6 @@ bool FPCGRandomChoiceElement::ExecuteInternal(FPCGContext* Context) const
 			}
 		}
 
-		auto AddToOutput = [&ShuffledIndexes, &CurrentInput, &Outputs, &ChooseFunc, Context](const int StartIndex, const int Count, const FName Label)
-		{
-			FPCGTaggedData& ChosenOutput = Outputs.Add_GetRef(CurrentInput);
-			ChosenOutput.Data = ChooseFunc(CurrentInput.Data, MakeArrayView(ShuffledIndexes.GetData() + StartIndex, Count), Context);
-			ChosenOutput.Pin = PCGRandomChoiceConstants::ChosenEntriesLabel;
-		};
-
 		FPCGTaggedData& ChosenOutput = Outputs.Add_GetRef(CurrentInput);
 		ChosenOutput.Data = ChooseFunc(CurrentInput.Data, MakeArrayView(ShuffledIndexes.GetData(), NumOfElementsToKeep), Context);
 		ChosenOutput.Pin = PCGRandomChoiceConstants::ChosenEntriesLabel;
diff --git a/Engine/Plugins/PCG/Source/PCG/Private/Elements/PCGSampleTexture.cpp b/Engine/Plugins/PCG/Source/PCG/Private/Elements/PCGSampleTexture.cpp
--- a/Engine/Plugins/PCG/Source/PCG/Private/Elements/PCGSampleTexture.cpp
+++ b/Engine/Plugins/PCG/Source/PCG/Private/Elements/PCGSampleTexture.cpp
@@ -16,10 +16,11 @@ TArray<FPCGPinProperties> UPCGSampleTextureSettings::InputPinProperties() const
 {
 	TArray<FPCGPinProperties> PinProperties;
 
-	PinProperties.Emplace(PCGSampleTextureConstants::InputPointLabel,
+	FPCGPinProperties& InputPin = PinProperties.Emplace_GetRef(PCGSampleTextureConstants::InputPointLabel,
 		EPCGDataType::Point,
 		/*bAllowMultipleConnections=*/true,
 		/*bAllowMultipleData=*/true);
+	InputPin.SetRequiredPin();
 
 	PinProperties.Emplace(PCGSampleTextureConstants::InputTextureLabel,
 		EPCGDataType::BaseTexture,
diff --git a/Engine/Plugins/PCG/Source/PCG/Private/Elements/PCGSelfPruning.cpp b/Engine/Plugins/PCG/Source/PCG/Private/Elements/PCGSelfPruning.cpp
--- a/Engine/Plugins/PCG/Source/PCG/Private/Elements/PCGSelfPruning.cpp
+++ b/Engine/Plugins/PCG/Source/PCG/Private/Elements/PCGSelfPruning.cpp
@@ -442,8 +442,6 @@ namespace PCGSelfPruningElement
 			return;
 		}
 
-		TArray<FPCGTaggedData>& Outputs = Context->OutputData.TaggedData;
-
 		TArray<FPCGTaggedData> Inputs = Context->InputData.GetInputs();
 		for (const FPCGTaggedData& Input : Inputs)
 		{
diff --git a/Engine/Plugins/PCG/Source/PCG/Private/Elements/PCGSplineSampler.cpp b/Engine/Plugins/PCG/Source/PCG/Private/Elements/PCGSplineSampler.cpp
--- a/Engine/Plugins/PCG/Source/PCG/Private/Elements/PCGSplineSampler.cpp
+++ b/Engine/Plugins/PCG/Source/PCG/Private/Elements/PCGSplineSampler.cpp
@@ -1556,8 +1556,7 @@ bool FPCGSplineSamplerElement::ExecuteInternal(FPCGContext* Context) const
 			ProjectionParams = SplineProjection->GetProjectionParams();
 		}
 
-		FPCGTaggedData& Output = Outputs.Emplace_GetRef();
-		Output = Input;
+		FPCGTaggedData& Output = Outputs.Emplace_GetRef(Input);
 
 		UPCGPointData* SampledPointData = FPCGContext::NewObject_AnyThread<UPCGPointData>(Context);
 		SampledPointData->InitializeFromData(SpatialData);
diff --git a/Engine/Plugins/PCG/Source/PCG/Private/Elements/PCGSurfaceSampler.cpp b/Engine/Plugins/PCG/Source/PCG/Private/Elements/PCGSurfaceSampler.cpp
--- a/Engine/Plugins/PCG/Source/PCG/Private/Elements/PCGSurfaceSampler.cpp
+++ b/Engine/Plugins/PCG/Source/PCG/Private/Elements/PCGSurfaceSampler.cpp
@@ -518,7 +518,7 @@ bool FPCGSurfaceSamplerElement::ExecuteInternal(FPCGContext* InContext) const
 	{
 		for (FPCGTaggedData& Input : TimeSlicedContext->InputData.GetInputs())
 		{
-			FPCGTaggedData& Output = TimeSlicedContext->OutputData.TaggedData.Emplace_GetRef();
+			FPCGTaggedData& Output = TimeSlicedContext->OutputData.TaggedData.Add_GetRef(Input);
 			UPCGPointData* PointData = FPCGContext::NewObject_AnyThread<UPCGPointData>(InContext);
 			PointData->InitializeFromData(Cast<UPCGSpatialData>(Input.Data));
 			Output.Data = PointData;
diff --git a/Engine/Plugins/PCG/Source/PCG/Private/Elements/PCGVolumeSampler.cpp b/Engine/Plugins/PCG/Source/PCG/Private/Elements/PCGVolumeSampler.cpp
--- a/Engine/Plugins/PCG/Source/PCG/Private/Elements/PCGVolumeSampler.cpp
+++ b/Engine/Plugins/PCG/Source/PCG/Private/Elements/PCGVolumeSampler.cpp
@@ -343,7 +343,7 @@ bool FPCGVolumeSamplerElement::ExecuteInternal(FPCGContext* Context) const
 		for (FPCGTaggedData& Input : TimeSlicedContext->InputData.GetInputs())
 		{
 			// TODO: Empty point data (to preserve previous behavior). Eventually, should be replaced with no output at all
-			FPCGTaggedData& Output = TimeSlicedContext->OutputData.TaggedData.Emplace_GetRef();
+			FPCGTaggedData& Output = TimeSlicedContext->OutputData.TaggedData.Add_GetRef(Input);
 			UPCGPointData* PointData = FPCGContext::NewObject_AnyThread<UPCGPointData>(Context);
 			PointData->InitializeFromData(Cast<UPCGSpatialData>(Input.Data));
 			Output.Data = PointData;
diff --git a/Engine/Plugins/PCG/Source/PCG/Public/Elements/Grammar/PCGSplineToSegment.h b/Engine/Plugins/PCG/Source/PCG/Public/Elements/Grammar/PCGSplineToSegment.h
--- a/Engine/Plugins/PCG/Source/PCG/Public/Elements/Grammar/PCGSplineToSegment.h
+++ b/Engine/Plugins/PCG/Source/PCG/Public/Elements/Grammar/PCGSplineToSegment.h
@@ -52,5 +52,6 @@ class FPCGSplineToSegmentElement : public IPCGElement
 {
 protected:
 	virtual bool ExecuteInternal(FPCGContext* InContext) const override;
+	virtual EPCGElementExecutionLoopMode ExecutionLoopMode(const UPCGSettings* Settings) const override { return EPCGElementExecutionLoopMode::SinglePrimaryPin; }
 };
 
diff --git a/Engine/Plugins/PCG/Source/PCG/Public/Elements/Metadata/PCGMetadataOpElementBase.h b/Engine/Plugins/PCG/Source/PCG/Public/Elements/Metadata/PCGMetadataOpElementBase.h
--- a/Engine/Plugins/PCG/Source/PCG/Public/Elements/Metadata/PCGMetadataOpElementBase.h
+++ b/Engine/Plugins/PCG/Source/PCG/Public/Elements/Metadata/PCGMetadataOpElementBase.h
@@ -208,6 +208,7 @@ class FPCGMetadataElementBase : public TPCGTimeSlicedElementBase<PCGTimeSlice::F
 protected:
 	virtual bool PrepareDataInternal(FPCGContext* Context) const override;
 	virtual bool ExecuteInternal(FPCGContext* Context) const override;
+	virtual EPCGElementExecutionLoopMode ExecutionLoopMode(const UPCGSettings* Settings) const override { return EPCGElementExecutionLoopMode::SinglePrimaryPin; }
 
 	virtual bool DoOperation(PCGMetadataOps::FOperationData& InOperationData) const = 0;
 
diff --git a/Engine/Plugins/PCG/Source/PCG/Public/Elements/Metadata/PCGMetadataPartition.h b/Engine/Plugins/PCG/Source/PCG/Public/Elements/Metadata/PCGMetadataPartition.h
--- a/Engine/Plugins/PCG/Source/PCG/Public/Elements/Metadata/PCGMetadataPartition.h
+++ b/Engine/Plugins/PCG/Source/PCG/Public/Elements/Metadata/PCGMetadataPartition.h
@@ -64,4 +64,5 @@ class FPCGMetadataPartitionElement : public IPCGElement
 {
 protected:
 	virtual bool ExecuteInternal(FPCGContext* Context) const override;
+	virtual EPCGElementExecutionLoopMode ExecutionLoopMode(const UPCGSettings* Settings) const override { return EPCGElementExecutionLoopMode::SinglePrimaryPin; }
 };
diff --git a/Engine/Plugins/PCG/Source/PCG/Public/Elements/Metadata/PCGMetadataRenameElement.h b/Engine/Plugins/PCG/Source/PCG/Public/Elements/Metadata/PCGMetadataRenameElement.h
--- a/Engine/Plugins/PCG/Source/PCG/Public/Elements/Metadata/PCGMetadataRenameElement.h
+++ b/Engine/Plugins/PCG/Source/PCG/Public/Elements/Metadata/PCGMetadataRenameElement.h
@@ -39,4 +39,5 @@ class FPCGMetadataRenameElement : public IPCGElement
 {
 protected:
 	virtual bool ExecuteInternal(FPCGContext* Context) const override;
+	virtual EPCGElementExecutionLoopMode ExecutionLoopMode(const UPCGSettings* Settings) const override { return EPCGElementExecutionLoopMode::SinglePrimaryPin; }
 };
diff --git a/Engine/Plugins/PCG/Source/PCG/Public/Elements/PCGAddTag.h b/Engine/Plugins/PCG/Source/PCG/Public/Elements/PCGAddTag.h
--- a/Engine/Plugins/PCG/Source/PCG/Public/Elements/PCGAddTag.h
+++ b/Engine/Plugins/PCG/Source/PCG/Public/Elements/PCGAddTag.h
@@ -58,4 +58,5 @@ class FPCGAddTagElement : public IPCGElement
 {
 protected:
 	virtual bool ExecuteInternal(FPCGContext* Context) const override;
+	virtual EPCGElementExecutionLoopMode ExecutionLoopMode(const UPCGSettings* Settings) const override { return EPCGElementExecutionLoopMode::SinglePrimaryPin; }
 };
\ No newline at end of file
diff --git a/Engine/Plugins/PCG/Source/PCG/Public/Elements/PCGAttributeCast.h b/Engine/Plugins/PCG/Source/PCG/Public/Elements/PCGAttributeCast.h
--- a/Engine/Plugins/PCG/Source/PCG/Public/Elements/PCGAttributeCast.h
+++ b/Engine/Plugins/PCG/Source/PCG/Public/Elements/PCGAttributeCast.h
@@ -51,4 +51,5 @@ class FPCGAttributeCastElement : public IPCGElement
 {
 protected:
 	virtual bool ExecuteInternal(FPCGContext* Context) const override;
+	virtual EPCGElementExecutionLoopMode ExecutionLoopMode(const UPCGSettings* Settings) const override { return EPCGElementExecutionLoopMode::SinglePrimaryPin; }
 };
diff --git a/Engine/Plugins/PCG/Source/PCG/Public/Elements/PCGAttributeFilter.h b/Engine/Plugins/PCG/Source/PCG/Public/Elements/PCGAttributeFilter.h
--- a/Engine/Plugins/PCG/Source/PCG/Public/Elements/PCGAttributeFilter.h
+++ b/Engine/Plugins/PCG/Source/PCG/Public/Elements/PCGAttributeFilter.h
@@ -191,6 +191,7 @@ public:
 class FPCGAttributeFilterElementBase : public IPCGElement
 {
 protected:
+	virtual EPCGElementExecutionLoopMode ExecutionLoopMode(const UPCGSettings* Settings) const override { return EPCGElementExecutionLoopMode::SinglePrimaryPin; }
 	bool DoFiltering(FPCGContext* Context, EPCGAttributeFilterOperator InOperation, const FPCGAttributePropertyInputSelector& TargetAttribute, bool bHasSpatialToPointDeprecation, const FPCGAttributeFilterThresholdSettings& FirstThreshold, const FPCGAttributeFilterThresholdSettings* SecondThreshold = nullptr) const;
 };
 
diff --git a/Engine/Plugins/PCG/Source/PCG/Public/Elements/PCGAttributeGetFromIndexElement.h b/Engine/Plugins/PCG/Source/PCG/Public/Elements/PCGAttributeGetFromIndexElement.h
--- a/Engine/Plugins/PCG/Source/PCG/Public/Elements/PCGAttributeGetFromIndexElement.h
+++ b/Engine/Plugins/PCG/Source/PCG/Public/Elements/PCGAttributeGetFromIndexElement.h
@@ -36,4 +36,5 @@ class FPCGAttributeGetFromIndexElement : public IPCGElement
 {
 protected:
 	virtual bool ExecuteInternal(FPCGContext* Context) const override;
+	virtual EPCGElementExecutionLoopMode ExecutionLoopMode(const UPCGSettings* Settings) const override { return EPCGElementExecutionLoopMode::SinglePrimaryPin; }
 };
\ No newline at end of file
diff --git a/Engine/Plugins/PCG/Source/PCG/Public/Elements/PCGAttributeGetFromPointIndexElement.h b/Engine/Plugins/PCG/Source/PCG/Public/Elements/PCGAttributeGetFromPointIndexElement.h
--- a/Engine/Plugins/PCG/Source/PCG/Public/Elements/PCGAttributeGetFromPointIndexElement.h
+++ b/Engine/Plugins/PCG/Source/PCG/Public/Elements/PCGAttributeGetFromPointIndexElement.h
@@ -67,4 +67,5 @@ class FPCGAttributeGetFromPointIndexElement : public IPCGElement
 {
 protected:
 	virtual bool ExecuteInternal(FPCGContext* Context) const override;
+	virtual EPCGElementExecutionLoopMode ExecutionLoopMode(const UPCGSettings* Settings) const override { return EPCGElementExecutionLoopMode::SinglePrimaryPin; }
 };
diff --git a/Engine/Plugins/PCG/Source/PCG/Public/Elements/PCGAttributeNoise.h b/Engine/Plugins/PCG/Source/PCG/Public/Elements/PCGAttributeNoise.h
--- a/Engine/Plugins/PCG/Source/PCG/Public/Elements/PCGAttributeNoise.h
+++ b/Engine/Plugins/PCG/Source/PCG/Public/Elements/PCGAttributeNoise.h
@@ -122,11 +122,11 @@ struct FPCGAttributeNoiseContext : public FPCGContext
 	TUniquePtr<IPCGAttributeAccessorKeys> OutputKeys;
 };
 
-class FPCGAttributeNoiseElement : public IPCGElement
+class FPCGAttributeNoiseElement : public IPCGElementWithCustomContext<FPCGAttributeNoiseContext>
 {
 protected:
-	virtual FPCGContext* CreateContext() override;
 	virtual bool ExecuteInternal(FPCGContext* Context) const override;
+	virtual EPCGElementExecutionLoopMode ExecutionLoopMode(const UPCGSettings* Settings) const override { return EPCGElementExecutionLoopMode::SinglePrimaryPin; }
 };
 
 #if UE_ENABLE_INCLUDE_ORDER_DEPRECATED_IN_5_2
diff --git a/Engine/Plugins/PCG/Source/PCG/Public/Elements/PCGAttributeReduceElement.h b/Engine/Plugins/PCG/Source/PCG/Public/Elements/PCGAttributeReduceElement.h
--- a/Engine/Plugins/PCG/Source/PCG/Public/Elements/PCGAttributeReduceElement.h
+++ b/Engine/Plugins/PCG/Source/PCG/Public/Elements/PCGAttributeReduceElement.h
@@ -88,4 +88,5 @@ class FPCGAttributeReduceElement : public IPCGElement
 {
 protected:
 	virtual bool ExecuteInternal(FPCGContext* Context) const override;
+	virtual EPCGElementExecutionLoopMode ExecutionLoopMode(const UPCGSettings* Settings) const override;
 };
diff --git a/Engine/Plugins/PCG/Source/PCG/Public/Elements/PCGBlurElement.h b/Engine/Plugins/PCG/Source/PCG/Public/Elements/PCGBlurElement.h
--- a/Engine/Plugins/PCG/Source/PCG/Public/Elements/PCGBlurElement.h
+++ b/Engine/Plugins/PCG/Source/PCG/Public/Elements/PCGBlurElement.h
@@ -92,5 +92,6 @@ class FPCGBlurElement : public TPCGTimeSlicedElementBase<PCGTimeSlice::FEmptyStr
 protected:
 	virtual bool PrepareDataInternal(FPCGContext* InContext) const override;
 	virtual bool ExecuteInternal(FPCGContext* InContext) const override;
+	virtual EPCGElementExecutionLoopMode ExecutionLoopMode(const UPCGSettings* Settings) const override { return EPCGElementExecutionLoopMode::SinglePrimaryPin; }
 };
 
diff --git a/Engine/Plugins/PCG/Source/PCG/Public/Elements/PCGBoundsFromMesh.h b/Engine/Plugins/PCG/Source/PCG/Public/Elements/PCGBoundsFromMesh.h
--- a/Engine/Plugins/PCG/Source/PCG/Public/Elements/PCGBoundsFromMesh.h
+++ b/Engine/Plugins/PCG/Source/PCG/Public/Elements/PCGBoundsFromMesh.h
@@ -77,4 +77,5 @@ public:
 protected:
 	virtual bool PrepareDataInternal(FPCGContext* Context) const override;
 	virtual bool ExecuteInternal(FPCGContext* Context) const override;
+	virtual EPCGElementExecutionLoopMode ExecutionLoopMode(const UPCGSettings* Settings) const override { return EPCGElementExecutionLoopMode::SinglePrimaryPin; }
 };
diff --git a/Engine/Plugins/PCG/Source/PCG/Public/Elements/PCGCollapseElement.h b/Engine/Plugins/PCG/Source/PCG/Public/Elements/PCGCollapseElement.h
--- a/Engine/Plugins/PCG/Source/PCG/Public/Elements/PCGCollapseElement.h
+++ b/Engine/Plugins/PCG/Source/PCG/Public/Elements/PCGCollapseElement.h
@@ -56,4 +56,5 @@ public:
 
 protected:
 	virtual bool ExecuteInternal(FPCGContext* Context) const override;
+	virtual EPCGElementExecutionLoopMode ExecutionLoopMode(const UPCGSettings* Settings) const override { return EPCGElementExecutionLoopMode::SinglePrimaryPin; }
 };
diff --git a/Engine/Plugins/PCG/Source/PCG/Public/Elements/PCGCombinePoints.h b/Engine/Plugins/PCG/Source/PCG/Public/Elements/PCGCombinePoints.h
--- a/Engine/Plugins/PCG/Source/PCG/Public/Elements/PCGCombinePoints.h
+++ b/Engine/Plugins/PCG/Source/PCG/Public/Elements/PCGCombinePoints.h
@@ -49,4 +49,5 @@ class FPCGCombinePointsElement : public IPCGElement
 {
 protected:
 	virtual bool ExecuteInternal(FPCGContext* Context) const override;
+	virtual EPCGElementExecutionLoopMode ExecutionLoopMode(const UPCGSettings* Settings) const override { return EPCGElementExecutionLoopMode::SinglePrimaryPin; }
 };
\ No newline at end of file
diff --git a/Engine/Plugins/PCG/Source/PCG/Public/Elements/PCGConvertToAttributeSet.h b/Engine/Plugins/PCG/Source/PCG/Public/Elements/PCGConvertToAttributeSet.h
--- a/Engine/Plugins/PCG/Source/PCG/Public/Elements/PCGConvertToAttributeSet.h
+++ b/Engine/Plugins/PCG/Source/PCG/Public/Elements/PCGConvertToAttributeSet.h
@@ -38,6 +38,7 @@ public:
 
 protected:
 	virtual bool ExecuteInternal(FPCGContext* Context) const override;
+	virtual EPCGElementExecutionLoopMode ExecutionLoopMode(const UPCGSettings* Settings) const override { return EPCGElementExecutionLoopMode::SinglePrimaryPin; }
 };
 
 /** Converts tags on data collections to an attribute set (per input data), akin to how it's done on Get Actor Data with the Single Point option. */
@@ -65,4 +66,5 @@ class FPCGTagsToAttributeSetElement : public IPCGElement
 {
 protected:
 	virtual bool ExecuteInternal(FPCGContext* Context) const override;
+	virtual EPCGElementExecutionLoopMode ExecutionLoopMode(const UPCGSettings* Settings) const override { return EPCGElementExecutionLoopMode::SinglePrimaryPin; }
 };
\ No newline at end of file
diff --git a/Engine/Plugins/PCG/Source/PCG/Public/Elements/PCGConvexHull2D.h b/Engine/Plugins/PCG/Source/PCG/Public/Elements/PCGConvexHull2D.h
--- a/Engine/Plugins/PCG/Source/PCG/Public/Elements/PCGConvexHull2D.h
+++ b/Engine/Plugins/PCG/Source/PCG/Public/Elements/PCGConvexHull2D.h
@@ -35,4 +35,5 @@ class FPCGConvexHull2DElement : public IPCGElement
 {
 protected:
 	virtual bool ExecuteInternal(FPCGContext* Context) const override;
+	virtual EPCGElementExecutionLoopMode ExecutionLoopMode(const UPCGSettings* Settings) const override { return EPCGElementExecutionLoopMode::SinglePrimaryPin; }
 };
\ No newline at end of file
diff --git a/Engine/Plugins/PCG/Source/PCG/Public/Elements/PCGCreateAttribute.h b/Engine/Plugins/PCG/Source/PCG/Public/Elements/PCGCreateAttribute.h
--- a/Engine/Plugins/PCG/Source/PCG/Public/Elements/PCGCreateAttribute.h
+++ b/Engine/Plugins/PCG/Source/PCG/Public/Elements/PCGCreateAttribute.h
@@ -129,4 +129,5 @@ class FPCGAddAttributeElement : public IPCGElement
 {
 protected:
 	virtual bool ExecuteInternal(FPCGContext* Context) const override;
+	virtual EPCGElementExecutionLoopMode ExecutionLoopMode(const UPCGSettings* Settings) const override { return EPCGElementExecutionLoopMode::SinglePrimaryPin; }
 };
\ No newline at end of file
diff --git a/Engine/Plugins/PCG/Source/PCG/Public/Elements/PCGCreateCollisionData.h b/Engine/Plugins/PCG/Source/PCG/Public/Elements/PCGCreateCollisionData.h
--- a/Engine/Plugins/PCG/Source/PCG/Public/Elements/PCGCreateCollisionData.h
+++ b/Engine/Plugins/PCG/Source/PCG/Public/Elements/PCGCreateCollisionData.h
@@ -78,4 +78,5 @@ public:
 protected:
 	virtual bool PrepareDataInternal(FPCGContext* Context) const override;
 	virtual bool ExecuteInternal(FPCGContext* Context) const override;
+	virtual EPCGElementExecutionLoopMode ExecutionLoopMode(const UPCGSettings* Settings) const override { return EPCGElementExecutionLoopMode::SinglePrimaryPin; }
 };
\ No newline at end of file
diff --git a/Engine/Plugins/PCG/Source/PCG/Public/Elements/PCGCreateSpline.h b/Engine/Plugins/PCG/Source/PCG/Public/Elements/PCGCreateSpline.h
--- a/Engine/Plugins/PCG/Source/PCG/Public/Elements/PCGCreateSpline.h
+++ b/Engine/Plugins/PCG/Source/PCG/Public/Elements/PCGCreateSpline.h
@@ -77,4 +77,5 @@ public:
 
 protected:	
 	virtual bool ExecuteInternal(FPCGContext* Context) const override;
+	virtual EPCGElementExecutionLoopMode ExecutionLoopMode(const UPCGSettings* Settings) const override { return EPCGElementExecutionLoopMode::SinglePrimaryPin; }
 };
\ No newline at end of file
diff --git a/Engine/Plugins/PCG/Source/PCG/Public/Elements/PCGCreateSurfaceFromSpline.h b/Engine/Plugins/PCG/Source/PCG/Public/Elements/PCGCreateSurfaceFromSpline.h
--- a/Engine/Plugins/PCG/Source/PCG/Public/Elements/PCGCreateSurfaceFromSpline.h
+++ b/Engine/Plugins/PCG/Source/PCG/Public/Elements/PCGCreateSurfaceFromSpline.h
@@ -44,4 +44,5 @@ class FPCGCreateSurfaceFromSplineElement : public IPCGElement
 {
 protected:
 	virtual bool ExecuteInternal(FPCGContext* Context) const override;
+	virtual EPCGElementExecutionLoopMode ExecutionLoopMode(const UPCGSettings* Settings) const override { return EPCGElementExecutionLoopMode::SinglePrimaryPin; }
 };
diff --git a/Engine/Plugins/PCG/Source/PCG/Public/Elements/PCGCullPointsOutsideActorBounds.h b/Engine/Plugins/PCG/Source/PCG/Public/Elements/PCGCullPointsOutsideActorBounds.h
--- a/Engine/Plugins/PCG/Source/PCG/Public/Elements/PCGCullPointsOutsideActorBounds.h
+++ b/Engine/Plugins/PCG/Source/PCG/Public/Elements/PCGCullPointsOutsideActorBounds.h
@@ -45,4 +45,5 @@ public:
 
 protected:
 	virtual bool ExecuteInternal(FPCGContext* Context) const override;
+	virtual EPCGElementExecutionLoopMode ExecutionLoopMode(const UPCGSettings* Settings) const override { return EPCGElementExecutionLoopMode::SinglePrimaryPin; }
 };
diff --git a/Engine/Plugins/PCG/Source/PCG/Public/Elements/PCGDeleteAttributesElement.h b/Engine/Plugins/PCG/Source/PCG/Public/Elements/PCGDeleteAttributesElement.h
--- a/Engine/Plugins/PCG/Source/PCG/Public/Elements/PCGDeleteAttributesElement.h
+++ b/Engine/Plugins/PCG/Source/PCG/Public/Elements/PCGDeleteAttributesElement.h
@@ -80,4 +80,5 @@ class FPCGDeleteAttributesElement : public IPCGElement
 {
 protected:
 	virtual bool ExecuteInternal(FPCGContext* Context) const override;
+	virtual EPCGElementExecutionLoopMode ExecutionLoopMode(const UPCGSettings* Settings) const override { return EPCGElementExecutionLoopMode::SinglePrimaryPin; }
 };
diff --git a/Engine/Plugins/PCG/Source/PCG/Public/Elements/PCGDeleteTags.h b/Engine/Plugins/PCG/Source/PCG/Public/Elements/PCGDeleteTags.h
--- a/Engine/Plugins/PCG/Source/PCG/Public/Elements/PCGDeleteTags.h
+++ b/Engine/Plugins/PCG/Source/PCG/Public/Elements/PCGDeleteTags.h
@@ -61,4 +61,5 @@ class FPCGDeleteTagsElement : public IPCGElement
 {
 protected:
 	virtual bool ExecuteInternal(FPCGContext* Context) const override;
+	virtual EPCGElementExecutionLoopMode ExecutionLoopMode(const UPCGSettings* Settings) const override { return EPCGElementExecutionLoopMode::SinglePrimaryPin; }
 };
\ No newline at end of file
diff --git a/Engine/Plugins/PCG/Source/PCG/Public/Elements/PCGDensityFilter.h b/Engine/Plugins/PCG/Source/PCG/Public/Elements/PCGDensityFilter.h
--- a/Engine/Plugins/PCG/Source/PCG/Public/Elements/PCGDensityFilter.h
+++ b/Engine/Plugins/PCG/Source/PCG/Public/Elements/PCGDensityFilter.h
@@ -46,6 +46,7 @@ class FPCGDensityFilterElement : public IPCGElement
 {
 protected:
 	virtual bool ExecuteInternal(FPCGContext* Context) const override;
+	virtual EPCGElementExecutionLoopMode ExecutionLoopMode(const UPCGSettings* Settings) const override { return EPCGElementExecutionLoopMode::SinglePrimaryPin; }
 };
 
 #if UE_ENABLE_INCLUDE_ORDER_DEPRECATED_IN_5_2
diff --git a/Engine/Plugins/PCG/Source/PCG/Public/Elements/PCGDifferenceElement.h b/Engine/Plugins/PCG/Source/PCG/Public/Elements/PCGDifferenceElement.h
--- a/Engine/Plugins/PCG/Source/PCG/Public/Elements/PCGDifferenceElement.h
+++ b/Engine/Plugins/PCG/Source/PCG/Public/Elements/PCGDifferenceElement.h
@@ -60,6 +60,7 @@ class FPCGDifferenceElement : public IPCGElement
 {
 protected:
 	virtual bool ExecuteInternal(FPCGContext* Context) const override;
+	virtual EPCGElementExecutionLoopMode ExecutionLoopMode(const UPCGSettings* Settings) const override { return EPCGElementExecutionLoopMode::SinglePrimaryPin; }
 };
 
 #if UE_ENABLE_INCLUDE_ORDER_DEPRECATED_IN_5_2
diff --git a/Engine/Plugins/PCG/Source/PCG/Public/Elements/PCGDistance.h b/Engine/Plugins/PCG/Source/PCG/Public/Elements/PCGDistance.h
--- a/Engine/Plugins/PCG/Source/PCG/Public/Elements/PCGDistance.h
+++ b/Engine/Plugins/PCG/Source/PCG/Public/Elements/PCGDistance.h
@@ -90,4 +90,5 @@ class FPCGDistanceElement : public IPCGElement
 {
 protected:
 	virtual bool ExecuteInternal(FPCGContext* Context) const override;
+	virtual EPCGElementExecutionLoopMode ExecutionLoopMode(const UPCGSettings* Settings) const override { return EPCGElementExecutionLoopMode::SinglePrimaryPin; }
 };
\ No newline at end of file
diff --git a/Engine/Plugins/PCG/Source/PCG/Public/Elements/PCGDuplicatePoint.h b/Engine/Plugins/PCG/Source/PCG/Public/Elements/PCGDuplicatePoint.h
--- a/Engine/Plugins/PCG/Source/PCG/Public/Elements/PCGDuplicatePoint.h
+++ b/Engine/Plugins/PCG/Source/PCG/Public/Elements/PCGDuplicatePoint.h
@@ -56,4 +56,5 @@ class FPCGDuplicatePointElement : public IPCGElement
 {
 protected:
 	virtual bool ExecuteInternal(FPCGContext* Context) const override;
+	virtual EPCGElementExecutionLoopMode ExecutionLoopMode(const UPCGSettings* Settings) const override { return EPCGElementExecutionLoopMode::SinglePrimaryPin; }
 };
\ No newline at end of file
diff --git a/Engine/Plugins/PCG/Source/PCG/Public/Elements/PCGElevationIsolines.h b/Engine/Plugins/PCG/Source/PCG/Public/Elements/PCGElevationIsolines.h
--- a/Engine/Plugins/PCG/Source/PCG/Public/Elements/PCGElevationIsolines.h
+++ b/Engine/Plugins/PCG/Source/PCG/Public/Elements/PCGElevationIsolines.h
@@ -70,5 +70,6 @@ public:
 
 protected:
 	virtual bool ExecuteInternal(FPCGContext* InContext) const override;
+	virtual EPCGElementExecutionLoopMode ExecutionLoopMode(const UPCGSettings* Settings) const override { return EPCGElementExecutionLoopMode::SinglePrimaryPin; }
 };
 
diff --git a/Engine/Plugins/PCG/Source/PCG/Public/Elements/PCGFilterByAttribute.h b/Engine/Plugins/PCG/Source/PCG/Public/Elements/PCGFilterByAttribute.h
--- a/Engine/Plugins/PCG/Source/PCG/Public/Elements/PCGFilterByAttribute.h
+++ b/Engine/Plugins/PCG/Source/PCG/Public/Elements/PCGFilterByAttribute.h
@@ -46,4 +46,5 @@ class FPCGFilterByAttributeElement : public IPCGElement
 {
 protected:
 	virtual bool ExecuteInternal(FPCGContext* Context) const override;
+	virtual EPCGElementExecutionLoopMode ExecutionLoopMode(const UPCGSettings* Settings) const override { return EPCGElementExecutionLoopMode::SinglePrimaryPin; }
 };
diff --git a/Engine/Plugins/PCG/Source/PCG/Public/Elements/PCGFilterByIndex.h b/Engine/Plugins/PCG/Source/PCG/Public/Elements/PCGFilterByIndex.h
--- a/Engine/Plugins/PCG/Source/PCG/Public/Elements/PCGFilterByIndex.h
+++ b/Engine/Plugins/PCG/Source/PCG/Public/Elements/PCGFilterByIndex.h
@@ -45,4 +45,5 @@ class FPCGFilterByIndexElement : public IPCGElement
 {
 protected:
 	virtual bool ExecuteInternal(FPCGContext* Context) const override;
+	virtual EPCGElementExecutionLoopMode ExecutionLoopMode(const UPCGSettings* Settings) const override { return EPCGElementExecutionLoopMode::SinglePrimaryPin; }
 };
diff --git a/Engine/Plugins/PCG/Source/PCG/Public/Elements/PCGFilterByTag.h b/Engine/Plugins/PCG/Source/PCG/Public/Elements/PCGFilterByTag.h
--- a/Engine/Plugins/PCG/Source/PCG/Public/Elements/PCGFilterByTag.h
+++ b/Engine/Plugins/PCG/Source/PCG/Public/Elements/PCGFilterByTag.h
@@ -59,4 +59,5 @@ class FPCGFilterByTagElement : public IPCGElement
 {
 protected:
 	virtual bool ExecuteInternal(FPCGContext* Context) const override;
+	virtual EPCGElementExecutionLoopMode ExecutionLoopMode(const UPCGSettings* Settings) const override { return EPCGElementExecutionLoopMode::SinglePrimaryPin; }
 };
diff --git a/Engine/Plugins/PCG/Source/PCG/Public/Elements/PCGFilterByType.h b/Engine/Plugins/PCG/Source/PCG/Public/Elements/PCGFilterByType.h
--- a/Engine/Plugins/PCG/Source/PCG/Public/Elements/PCGFilterByType.h
+++ b/Engine/Plugins/PCG/Source/PCG/Public/Elements/PCGFilterByType.h
@@ -41,4 +41,5 @@ class FPCGFilterByTypeElement : public IPCGElement
 {
 protected:
 	virtual bool ExecuteInternal(FPCGContext* InContext) const override;
+	virtual EPCGElementExecutionLoopMode ExecutionLoopMode(const UPCGSettings* Settings) const override { return EPCGElementExecutionLoopMode::SinglePrimaryPin; }
 };
diff --git a/Engine/Plugins/PCG/Source/PCG/Public/Elements/PCGGetBounds.h b/Engine/Plugins/PCG/Source/PCG/Public/Elements/PCGGetBounds.h
--- a/Engine/Plugins/PCG/Source/PCG/Public/Elements/PCGGetBounds.h
+++ b/Engine/Plugins/PCG/Source/PCG/Public/Elements/PCGGetBounds.h
@@ -35,4 +35,5 @@ class FPCGGetBoundsElement : public IPCGElement
 {
 protected:
 	virtual bool ExecuteInternal(FPCGContext* Context) const override;
+	virtual EPCGElementExecutionLoopMode ExecutionLoopMode(const UPCGSettings* Settings) const override { return EPCGElementExecutionLoopMode::SinglePrimaryPin; }
 };
diff --git a/Engine/Plugins/PCG/Source/PCG/Public/Elements/PCGGetDataInfo.h b/Engine/Plugins/PCG/Source/PCG/Public/Elements/PCGGetDataInfo.h
--- a/Engine/Plugins/PCG/Source/PCG/Public/Elements/PCGGetDataInfo.h
+++ b/Engine/Plugins/PCG/Source/PCG/Public/Elements/PCGGetDataInfo.h
@@ -36,6 +36,7 @@ class FPCGGetTagsElement : public IPCGElement
 {
 protected:
 	virtual bool ExecuteInternal(FPCGContext* Context) const override;
+	virtual EPCGElementExecutionLoopMode ExecutionLoopMode(const UPCGSettings* Settings) const override { return EPCGElementExecutionLoopMode::SinglePrimaryPin; }
 };
 
 /** Extracts the list of attribute on data to an attribute set with one entry per attribute. */
@@ -72,4 +73,5 @@ class FPCGGetAttributesElement : public IPCGElement
 {
 protected:
 	virtual bool ExecuteInternal(FPCGContext* Context) const override;
+	virtual EPCGElementExecutionLoopMode ExecutionLoopMode(const UPCGSettings* Settings) const override { return EPCGElementExecutionLoopMode::SinglePrimaryPin; }
 };
\ No newline at end of file
diff --git a/Engine/Plugins/PCG/Source/PCG/Public/Elements/PCGMakeConcreteElement.h b/Engine/Plugins/PCG/Source/PCG/Public/Elements/PCGMakeConcreteElement.h
--- a/Engine/Plugins/PCG/Source/PCG/Public/Elements/PCGMakeConcreteElement.h
+++ b/Engine/Plugins/PCG/Source/PCG/Public/Elements/PCGMakeConcreteElement.h
@@ -37,4 +37,5 @@ public:
 
 protected:
 	virtual bool ExecuteInternal(FPCGContext* InContext) const override;
+	virtual EPCGElementExecutionLoopMode ExecutionLoopMode(const UPCGSettings* Settings) const override { return EPCGElementExecutionLoopMode::SinglePrimaryPin; }
 };
diff --git a/Engine/Plugins/PCG/Source/PCG/Public/Elements/PCGMatchAndSetAttributes.h b/Engine/Plugins/PCG/Source/PCG/Public/Elements/PCGMatchAndSetAttributes.h
--- a/Engine/Plugins/PCG/Source/PCG/Public/Elements/PCGMatchAndSetAttributes.h
+++ b/Engine/Plugins/PCG/Source/PCG/Public/Elements/PCGMatchAndSetAttributes.h
@@ -121,4 +121,5 @@ class FPCGMatchAndSetAttributesElement : public TPCGTimeSlicedElementBase<FPCGMa
 protected:
 	virtual bool PrepareDataInternal(FPCGContext* InContext) const override;
 	virtual bool ExecuteInternal(FPCGContext* InContext) const override;
+	virtual EPCGElementExecutionLoopMode ExecutionLoopMode(const UPCGSettings* Settings) const override { return EPCGElementExecutionLoopMode::SinglePrimaryPin; }
 };
diff --git a/Engine/Plugins/PCG/Source/PCG/Public/Elements/PCGOuterIntersectionElement.h b/Engine/Plugins/PCG/Source/PCG/Public/Elements/PCGOuterIntersectionElement.h
--- a/Engine/Plugins/PCG/Source/PCG/Public/Elements/PCGOuterIntersectionElement.h
+++ b/Engine/Plugins/PCG/Source/PCG/Public/Elements/PCGOuterIntersectionElement.h
@@ -67,4 +67,5 @@ class FPCGOuterIntersectionElement : public IPCGElement
 {
 protected:
 	virtual bool ExecuteInternal(FPCGContext* Context) const override;
+	virtual EPCGElementExecutionLoopMode ExecutionLoopMode(const UPCGSettings* Settings) const override { return EPCGElementExecutionLoopMode::SinglePrimaryPin; }
 };
\ No newline at end of file
diff --git a/Engine/Plugins/PCG/Source/PCG/Public/Elements/PCGPointMatchAndSet.h b/Engine/Plugins/PCG/Source/PCG/Public/Elements/PCGPointMatchAndSet.h
--- a/Engine/Plugins/PCG/Source/PCG/Public/Elements/PCGPointMatchAndSet.h
+++ b/Engine/Plugins/PCG/Source/PCG/Public/Elements/PCGPointMatchAndSet.h
@@ -86,6 +86,7 @@ class FPCGPointMatchAndSetElement : public IPCGElement
 {
 public:
 	virtual bool ExecuteInternal(FPCGContext* Context) const override;
+	virtual EPCGElementExecutionLoopMode ExecutionLoopMode(const UPCGSettings* Settings) const override { return EPCGElementExecutionLoopMode::SinglePrimaryPin; }
 };
 
 #if UE_ENABLE_INCLUDE_ORDER_DEPRECATED_IN_5_2
diff --git a/Engine/Plugins/PCG/Source/PCG/Public/Elements/PCGPointNeighborhood.h b/Engine/Plugins/PCG/Source/PCG/Public/Elements/PCGPointNeighborhood.h
--- a/Engine/Plugins/PCG/Source/PCG/Public/Elements/PCGPointNeighborhood.h
+++ b/Engine/Plugins/PCG/Source/PCG/Public/Elements/PCGPointNeighborhood.h
@@ -78,4 +78,5 @@ class FPCGPointNeighborhoodElement : public IPCGElement
 {
 protected:
 	virtual bool ExecuteInternal(FPCGContext* Context) const override;
+	virtual EPCGElementExecutionLoopMode ExecutionLoopMode(const UPCGSettings* Settings) const override { return EPCGElementExecutionLoopMode::SinglePrimaryPin; }
 };
diff --git a/Engine/Plugins/PCG/Source/PCG/Public/Elements/PCGPointOperationElementBase.h b/Engine/Plugins/PCG/Source/PCG/Public/Elements/PCGPointOperationElementBase.h
--- a/Engine/Plugins/PCG/Source/PCG/Public/Elements/PCGPointOperationElementBase.h
+++ b/Engine/Plugins/PCG/Source/PCG/Public/Elements/PCGPointOperationElementBase.h
@@ -35,6 +35,7 @@ protected:
 	//~Begin IPCGElement interface
 	/** Conveniently calls PreparePointOperationData to prepare the time sliced element for execution. May be overridden, but PreparePointOperationData must be called. */
 	virtual bool PrepareDataInternal(FPCGContext* Context) const override;
+	virtual EPCGElementExecutionLoopMode ExecutionLoopMode(const UPCGSettings* Settings) const override { return EPCGElementExecutionLoopMode::SinglePrimaryPin; }
 	//~End IPCGElement interface
 
 	/** Executes the PointFunction function/lambda for every point copied from PreparePointOperationData */
diff --git a/Engine/Plugins/PCG/Source/PCG/Public/Elements/PCGRandomChoice.h b/Engine/Plugins/PCG/Source/PCG/Public/Elements/PCGRandomChoice.h
--- a/Engine/Plugins/PCG/Source/PCG/Public/Elements/PCGRandomChoice.h
+++ b/Engine/Plugins/PCG/Source/PCG/Public/Elements/PCGRandomChoice.h
@@ -63,4 +63,5 @@ class FPCGRandomChoiceElement : public IPCGElement
 {
 protected:
 	virtual bool ExecuteInternal(FPCGContext* Context) const override;
+	virtual EPCGElementExecutionLoopMode ExecutionLoopMode(const UPCGSettings* Settings) const override { return EPCGElementExecutionLoopMode::SinglePrimaryPin; }
 };
diff --git a/Engine/Plugins/PCG/Source/PCG/Public/Elements/PCGReplaceTags.h b/Engine/Plugins/PCG/Source/PCG/Public/Elements/PCGReplaceTags.h
--- a/Engine/Plugins/PCG/Source/PCG/Public/Elements/PCGReplaceTags.h
+++ b/Engine/Plugins/PCG/Source/PCG/Public/Elements/PCGReplaceTags.h
@@ -54,4 +54,5 @@ class FPCGReplaceTagsElement : public IPCGElement
 {
 protected:
 	virtual bool ExecuteInternal(FPCGContext* Context) const override;
+	virtual EPCGElementExecutionLoopMode ExecutionLoopMode(const UPCGSettings* Settings) const override { return EPCGElementExecutionLoopMode::SinglePrimaryPin; }
 };
\ No newline at end of file
diff --git a/Engine/Plugins/PCG/Source/PCG/Public/Elements/PCGReverseSpline.h b/Engine/Plugins/PCG/Source/PCG/Public/Elements/PCGReverseSpline.h
--- a/Engine/Plugins/PCG/Source/PCG/Public/Elements/PCGReverseSpline.h
+++ b/Engine/Plugins/PCG/Source/PCG/Public/Elements/PCGReverseSpline.h
@@ -62,5 +62,6 @@ class FPCGReverseSplineElement : public IPCGElement
 {
 protected:
 	virtual bool ExecuteInternal(FPCGContext* InContext) const override;
+	virtual EPCGElementExecutionLoopMode ExecutionLoopMode(const UPCGSettings* Settings) const override { return EPCGElementExecutionLoopMode::SinglePrimaryPin; }
 };
 
diff --git a/Engine/Plugins/PCG/Source/PCG/Public/Elements/PCGSampleTexture.h b/Engine/Plugins/PCG/Source/PCG/Public/Elements/PCGSampleTexture.h
--- a/Engine/Plugins/PCG/Source/PCG/Public/Elements/PCGSampleTexture.h
+++ b/Engine/Plugins/PCG/Source/PCG/Public/Elements/PCGSampleTexture.h
@@ -71,4 +71,5 @@ public:
 
 protected:
 	virtual bool ExecuteInternal(FPCGContext* Context) const override;
+	virtual EPCGElementExecutionLoopMode ExecutionLoopMode(const UPCGSettings* Settings) const override { return EPCGElementExecutionLoopMode::SinglePrimaryPin; }
 };
diff --git a/Engine/Plugins/PCG/Source/PCG/Public/Elements/PCGSelectPoints.h b/Engine/Plugins/PCG/Source/PCG/Public/Elements/PCGSelectPoints.h
--- a/Engine/Plugins/PCG/Source/PCG/Public/Elements/PCGSelectPoints.h
+++ b/Engine/Plugins/PCG/Source/PCG/Public/Elements/PCGSelectPoints.h
@@ -43,6 +43,7 @@ class FPCGSelectPointsElement : public IPCGElement
 {
 protected:
 	virtual bool ExecuteInternal(FPCGContext* Context) const override;
+	virtual EPCGElementExecutionLoopMode ExecutionLoopMode(const UPCGSettings* Settings) const override { return EPCGElementExecutionLoopMode::SinglePrimaryPin; }
 };
 
 #if UE_ENABLE_INCLUDE_ORDER_DEPRECATED_IN_5_2
diff --git a/Engine/Plugins/PCG/Source/PCG/Public/Elements/PCGSelfPruning.h b/Engine/Plugins/PCG/Source/PCG/Public/Elements/PCGSelfPruning.h
--- a/Engine/Plugins/PCG/Source/PCG/Public/Elements/PCGSelfPruning.h
+++ b/Engine/Plugins/PCG/Source/PCG/Public/Elements/PCGSelfPruning.h
@@ -150,4 +150,5 @@ protected:
 	virtual bool PrepareDataInternal(FPCGContext* Context) const override;
 	virtual bool ExecuteInternal(FPCGContext* Context) const override;
 	virtual bool CanExecuteOnlyOnMainThread(FPCGContext* Context) const override;
+	virtual EPCGElementExecutionLoopMode ExecutionLoopMode(const UPCGSettings* Settings) const override { return EPCGElementExecutionLoopMode::SinglePrimaryPin; }
 };
diff --git a/Engine/Plugins/PCG/Source/PCG/Public/Elements/PCGSortAttributes.h b/Engine/Plugins/PCG/Source/PCG/Public/Elements/PCGSortAttributes.h
--- a/Engine/Plugins/PCG/Source/PCG/Public/Elements/PCGSortAttributes.h
+++ b/Engine/Plugins/PCG/Source/PCG/Public/Elements/PCGSortAttributes.h
@@ -51,4 +51,5 @@ class FPCGSortAttributesElement : public IPCGElement
 {
 protected:
 	virtual bool ExecuteInternal(FPCGContext* Context) const override;
+	virtual EPCGElementExecutionLoopMode ExecutionLoopMode(const UPCGSettings* Settings) const override { return EPCGElementExecutionLoopMode::SinglePrimaryPin; }
 };
\ No newline at end of file
diff --git a/Engine/Plugins/PCG/Source/PCG/Public/Elements/PCGSpatialNoise.h b/Engine/Plugins/PCG/Source/PCG/Public/Elements/PCGSpatialNoise.h
--- a/Engine/Plugins/PCG/Source/PCG/Public/Elements/PCGSpatialNoise.h
+++ b/Engine/Plugins/PCG/Source/PCG/Public/Elements/PCGSpatialNoise.h
@@ -159,4 +159,5 @@ class FPCGSpatialNoise : public IPCGElement
 {
 protected:
 	virtual bool ExecuteInternal(FPCGContext* Context) const override;
+	virtual EPCGElementExecutionLoopMode ExecutionLoopMode(const UPCGSettings* Settings) const override { return EPCGElementExecutionLoopMode::SinglePrimaryPin; }
 };
diff --git a/Engine/Plugins/PCG/Source/PCG/Public/Elements/PCGSplineSampler.h b/Engine/Plugins/PCG/Source/PCG/Public/Elements/PCGSplineSampler.h
--- a/Engine/Plugins/PCG/Source/PCG/Public/Elements/PCGSplineSampler.h
+++ b/Engine/Plugins/PCG/Source/PCG/Public/Elements/PCGSplineSampler.h
@@ -290,4 +290,5 @@ public:
 
 protected:
 	virtual bool ExecuteInternal(FPCGContext* Context) const override;
+	virtual EPCGElementExecutionLoopMode ExecutionLoopMode(const UPCGSettings* Settings) const override { return EPCGElementExecutionLoopMode::SinglePrimaryPin; }
 };
diff --git a/Engine/Plugins/PCG/Source/PCG/Public/Elements/PCGSplitPoints.h b/Engine/Plugins/PCG/Source/PCG/Public/Elements/PCGSplitPoints.h
--- a/Engine/Plugins/PCG/Source/PCG/Public/Elements/PCGSplitPoints.h
+++ b/Engine/Plugins/PCG/Source/PCG/Public/Elements/PCGSplitPoints.h
@@ -50,4 +50,5 @@ class FPCGSplitPointsElement : public IPCGElement
 {
 protected:
 	virtual bool ExecuteInternal(FPCGContext* Context) const override;
+	virtual EPCGElementExecutionLoopMode ExecutionLoopMode(const UPCGSettings* Settings) const override { return EPCGElementExecutionLoopMode::SinglePrimaryPin; }
 };
\ No newline at end of file
diff --git a/Engine/Plugins/PCG/Source/PCG/Public/Elements/PCGSurfaceSampler.h b/Engine/Plugins/PCG/Source/PCG/Public/Elements/PCGSurfaceSampler.h
--- a/Engine/Plugins/PCG/Source/PCG/Public/Elements/PCGSurfaceSampler.h
+++ b/Engine/Plugins/PCG/Source/PCG/Public/Elements/PCGSurfaceSampler.h
@@ -168,6 +168,7 @@ public:
 protected:
 	virtual bool PrepareDataInternal(FPCGContext* InContext) const override;
 	virtual bool ExecuteInternal(FPCGContext* InContext) const override;
+	virtual EPCGElementExecutionLoopMode ExecutionLoopMode(const UPCGSettings* Settings) const override { return EPCGElementExecutionLoopMode::SinglePrimaryPin; }
 };
 
 #if UE_ENABLE_INCLUDE_ORDER_DEPRECATED_IN_5_2
diff --git a/Engine/Plugins/PCG/Source/PCG/Public/Elements/PCGTransformPoints.h b/Engine/Plugins/PCG/Source/PCG/Public/Elements/PCGTransformPoints.h
--- a/Engine/Plugins/PCG/Source/PCG/Public/Elements/PCGTransformPoints.h
+++ b/Engine/Plugins/PCG/Source/PCG/Public/Elements/PCGTransformPoints.h
@@ -78,6 +78,7 @@ class FPCGTransformPointsElement : public IPCGElement
 {
 protected:
 	virtual bool ExecuteInternal(FPCGContext* Context) const override;
+	virtual EPCGElementExecutionLoopMode ExecutionLoopMode(const UPCGSettings* Settings) const override { return EPCGElementExecutionLoopMode::SinglePrimaryPin; }
 };
 
 #if UE_ENABLE_INCLUDE_ORDER_DEPRECATED_IN_5_2
diff --git a/Engine/Plugins/PCG/Source/PCG/Public/Elements/PCGVolumeSampler.h b/Engine/Plugins/PCG/Source/PCG/Public/Elements/PCGVolumeSampler.h
--- a/Engine/Plugins/PCG/Source/PCG/Public/Elements/PCGVolumeSampler.h
+++ b/Engine/Plugins/PCG/Source/PCG/Public/Elements/PCGVolumeSampler.h
@@ -105,4 +105,5 @@ public:
 protected:
 	virtual bool PrepareDataInternal(FPCGContext* Context) const override;
 	virtual bool ExecuteInternal(FPCGContext* Context) const override;
+	virtual EPCGElementExecutionLoopMode ExecutionLoopMode(const UPCGSettings* Settings) const override { return EPCGElementExecutionLoopMode::SinglePrimaryPin; }
 };
