[PCG] Fix minor oversight in PCG Parsing, now allowing selection of the final element with a single -1

#rnx

diff --git a/Engine/Plugins/PCG/Source/PCG/Private/Helpers/Parsing/PCGIndexing.cpp b/Engine/Plugins/PCG/Source/PCG/Private/Helpers/Parsing/PCGIndexing.cpp
--- a/Engine/Plugins/PCG/Source/PCG/Private/Helpers/Parsing/PCGIndexing.cpp
+++ b/Engine/Plugins/PCG/Source/PCG/Private/Helpers/Parsing/PCGIndexing.cpp
@@ -31,7 +31,7 @@ namespace PCGIndexing
 
 	bool FPCGIndexCollection::AddRange(const int32 StartIndex, const int32 EndIndex)
 	{
-		const FPCGIndexRange NewRange(AdjustIndex(StartIndex), AdjustIndex(EndIndex));
+		const FPCGIndexRange NewRange = AdjustIndicesAndCreateRange(StartIndex, EndIndex);
 
 		if (!RangeIsValid(NewRange))
 		{
@@ -125,14 +125,13 @@ namespace PCGIndexing
 		return (ArraySize == Other.ArraySize) && (IndexRanges == Other.IndexRanges);
 	}
 
-	int32 FPCGIndexCollection::AdjustIndex(const int32 Index) const
+	FPCGIndexRange FPCGIndexCollection::AdjustIndicesAndCreateRange(int32 StartIndex, int32 EndIndex) const
 	{
-		if (Index >= 0)
-		{
-			return FMath::Min(Index, ArraySize);
-		}
+		// Convert both indices to a positive range
+		StartIndex = (StartIndex >= 0) ? FMath::Min(StartIndex, ArraySize) : FMath::Min(ArraySize + StartIndex, ArraySize);
+		EndIndex = (EndIndex >= 0) ? FMath::Min(EndIndex, ArraySize) : FMath::Min(ArraySize + EndIndex, ArraySize);
 
-		return FMath::Max(ArraySize + Index, 0);
+		return FPCGIndexRange(StartIndex, (EndIndex == StartIndex) ? EndIndex + 1 : EndIndex);
 	}
 
 	bool FPCGIndexCollection::CheckOverlap(const FPCGIndexRange& FirstRange, const FPCGIndexRange& SecondRange) const
diff --git a/Engine/Plugins/PCG/Source/PCG/Private/Helpers/Parsing/PCGParsing.cpp b/Engine/Plugins/PCG/Source/PCG/Private/Helpers/Parsing/PCGParsing.cpp
--- a/Engine/Plugins/PCG/Source/PCG/Private/Helpers/Parsing/PCGParsing.cpp
+++ b/Engine/Plugins/PCG/Source/PCG/Private/Helpers/Parsing/PCGParsing.cpp
@@ -69,7 +69,8 @@ namespace PCGParser
 				}
 
 				const int32 Index = FCString::Atoi(*ElementSubstring);
-				if (!OutIndexCollection.AddRange(Index, Index + 1))
+				// Adding the same index twice will adjust the end index appropriately
+				if (!OutIndexCollection.AddRange(Index, Index))
 				{
 					return EPCGParserResult::InvalidExpression;
 				}
diff --git a/Engine/Plugins/PCG/Source/PCG/Private/Tests/Parsing/PCGParsingTest.cpp b/Engine/Plugins/PCG/Source/PCG/Private/Tests/Parsing/PCGParsingTest.cpp
--- a/Engine/Plugins/PCG/Source/PCG/Private/Tests/Parsing/PCGParsingTest.cpp
+++ b/Engine/Plugins/PCG/Source/PCG/Private/Tests/Parsing/PCGParsingTest.cpp
@@ -33,6 +33,14 @@ bool FPCGParsingTest_Basic::RunTest(const FString& Parameters)
 		UTEST_TRUE("Singular non-range index is correct", Results.ContainsIndex(4));
 	}
 
+	// Final index
+	{
+		const PCGIndexing::FPCGIndexCollection Results = PCGParsingTestPrivate::ParseAndRetrieveIndices("-1", ArraySize);
+		UTEST_EQUAL("Single final index has 1 index", Results.GetTotalIndexCount(), 1);
+		UTEST_EQUAL("Single final index has 1 range", Results.GetTotalRangeCount(), 1);
+		UTEST_TRUE("Single final index is correct", Results.ContainsIndex(ArraySize - 1));
+	}
+
 	// Simple negative index
 	{
 		const PCGIndexing::FPCGIndexCollection Results = PCGParsingTestPrivate::ParseAndRetrieveIndices("-4", ArraySize);
@@ -265,9 +273,15 @@ bool FPCGParsingTest_Robustness::RunTest(const FString& Parameters)
 	Results = PCGParsingTestPrivate::ParseAndRetrieveIndices("5:2", ArraySize);
 	UTEST_EQUAL("Inverted range has no result", Results.GetTotalIndexCount(), 0);
 
+	Results = PCGParsingTestPrivate::ParseAndRetrieveIndices(FString::FromInt(ArraySize), ArraySize);
+	UTEST_EQUAL("Selected array max size", Results.GetTotalIndexCount(), 0);
+
 	Results = PCGParsingTestPrivate::ParseAndRetrieveIndices(FString::FromInt(ArraySize + 1), ArraySize);
 	UTEST_EQUAL("Index outside the bounds of the array", Results.GetTotalIndexCount(), 0);
 
+	Results = PCGParsingTestPrivate::ParseAndRetrieveIndices(FString::FromInt(-1 * (ArraySize + 1)), ArraySize);
+	UTEST_EQUAL("Index outside the bounds of the array", Results.GetTotalIndexCount(), 0);
+
 	Results = PCGParsingTestPrivate::ParseAndRetrieveIndices("65465465498475309485730495873409587654684986146162342348438", std::numeric_limits<int32>::max());
 	UTEST_EQUAL("Large value converts to int32 max which is the max size of the array which fails with no result", Results.GetTotalIndexCount(), 0);
 
@@ -284,4 +298,4 @@ bool FPCGParsingTest_Robustness::RunTest(const FString& Parameters)
 	UTEST_EQUAL("Redundant delimiter between indices has no result", Results.GetTotalIndexCount(), 0);
 
 	return true;
-}
\ No newline at end of file
+}
diff --git a/Engine/Plugins/PCG/Source/PCG/Public/Helpers/Parsing/PCGIndexing.h b/Engine/Plugins/PCG/Source/PCG/Public/Helpers/Parsing/PCGIndexing.h
--- a/Engine/Plugins/PCG/Source/PCG/Public/Helpers/Parsing/PCGIndexing.h
+++ b/Engine/Plugins/PCG/Source/PCG/Public/Helpers/Parsing/PCGIndexing.h
@@ -62,7 +62,8 @@ namespace PCGIndexing
 		bool operator==(const FPCGIndexCollection& Other) const;
 
 	private:
-		int32 AdjustIndex(int32 Index) const;
+		/** Adjusts indices into a positive range, including offsetting the EndIndex if they are equal. */
+		FPCGIndexRange AdjustIndicesAndCreateRange(int32 StartIndex, int32 EndIndex) const;
 
 		/** Checks two ranges for an overlap and returns true if they overlap. */
 		bool CheckOverlap(const FPCGIndexRange& FirstRange, const FPCGIndexRange& SecondRange) const;
