Added editor config to change the label on the pause button.
Added alternate icon and it's related config choice as well.

#rnx

diff --git a/Engine/Plugins/PCG/Source/PCGEditor/Private/PCGEditorModule.cpp b/Engine/Plugins/PCG/Source/PCGEditor/Private/PCGEditorModule.cpp
--- a/Engine/Plugins/PCG/Source/PCGEditor/Private/PCGEditorModule.cpp
+++ b/Engine/Plugins/PCG/Source/PCGEditor/Private/PCGEditorModule.cpp
@@ -371,17 +371,44 @@ void FPCGEditorModule::RegisterMenuExtensions()
 					),
 					TAttribute<FText>::CreateLambda([]()
 					{
+						const UPCGEditorSettings* EditorSettings = GetDefault<UPCGEditorSettings>();
+
 						if(PCGSystemSwitches::CVarPausePCGExecution.GetValueOnAnyThread())
 						{
-							return LOCTEXT("PCGPauseButton_Off", "Paused");
+							if (EditorSettings && EditorSettings->OverridePausedButtonLabel != NAME_None)
+							{
+								return FText::FromName(EditorSettings->OverridePausedButtonLabel);
+							}
+							else
+							{
+								return LOCTEXT("PCGPauseButton_Off", "Paused");
+							}
 						}
 						else
 						{
-							return LOCTEXT("PCGPauseButton_On", "PCG");
+							if (EditorSettings && EditorSettings->OverrideNotPausedButtonLabel != NAME_None)
+							{
+								return FText::FromName(EditorSettings->OverrideNotPausedButtonLabel);
+							}
+							else
+							{
+								return LOCTEXT("PCGPauseButton_On", "PCG");
+							}
 						}
 					}),
 					LOCTEXT("PCGPauseButton_Tooltip", "Toggles PCG processing on/off and will cancel tasks depending on settings.\nUse Ctrl to unpause and cancel all tasks.\nUse Alt to unpause without cancelling tasks."),
-					FSlateIcon(FPCGEditorStyle::Get().GetStyleSetName(), "PCG.EditorIcon"),
+					TAttribute<FSlateIcon>::CreateLambda([]()
+					{
+						const UPCGEditorSettings* EditorSettings = GetDefault<UPCGEditorSettings>();
+						if (EditorSettings && EditorSettings->bUseAlternatePauseButton)
+						{
+							return FSlateIcon(FPCGEditorStyle::Get().GetStyleSetName(), "PCG.Editor.AlternatePause");
+						}
+						else
+						{
+							return FSlateIcon(FPCGEditorStyle::Get().GetStyleSetName(), "PCG.Editor.Pause");
+						}
+					}),
 					EUserInterfaceActionType::ToggleButton
 				);
 				PCGPauseButton.StyleNameOverride = "CalloutToolbar"; // used to show the button text
diff --git a/Engine/Plugins/PCG/Source/PCGEditor/Private/PCGEditorStyle.cpp b/Engine/Plugins/PCG/Source/PCGEditor/Private/PCGEditorStyle.cpp
--- a/Engine/Plugins/PCG/Source/PCGEditor/Private/PCGEditorStyle.cpp
+++ b/Engine/Plugins/PCG/Source/PCGEditor/Private/PCGEditorStyle.cpp
@@ -127,6 +127,8 @@ FPCGEditorStyle::FPCGEditorStyle() : FSlateStyleSet("PCGEditorStyle")
 
 	Set("PCG.Editor.CopyToClipboard", new IMAGE_BRUSH("Style/PCG_Editor_CopyToClipboard", Icon32x32));
 	Set("PCG.Editor.ZoomToSelection", new IMAGE_BRUSH_SVG("Style/PCG_Editor_ZoomToSelection", Icon16x16));
+	Set("PCG.Editor.Pause", new IMAGE_BRUSH_SVG("Icons/PCG_16", Icon16x16));
+	Set("PCG.Editor.AlternatePause", new IMAGE_BRUSH_SVG("Style/PCG_Editor_AlternatePause", Icon16x16));
 }
 
 const FPCGEditorStyle& FPCGEditorStyle::Get()
diff --git a/Engine/Plugins/PCG/Source/PCGEditor/Public/PCGEditorSettings.h b/Engine/Plugins/PCG/Source/PCGEditor/Public/PCGEditorSettings.h
--- a/Engine/Plugins/PCG/Source/PCGEditor/Public/PCGEditorSettings.h
+++ b/Engine/Plugins/PCG/Source/PCGEditor/Public/PCGEditorSettings.h
@@ -177,11 +177,25 @@ public:
 	UPROPERTY(EditAnywhere, config, Category = "Editor Performance", meta = (DisplayName = "Disable CPU throttling during graph execution"))
 	bool bDisableCPUThrottlingDuringGraphExecution = false;
 
+	/** Controls whether the "Pause PCG" button will be display in the editor menu */
 	UPROPERTY(EditAnywhere, config, Category = "Editor Performance", meta = (DisplayName = "Show PCG pause button in the editor (requires restart)"))
 	bool bShowPauseButton = false;
 
+	/** Controls whether PCG tasks will be cancelled when unpausing PCG, which might result in dirty/stale content, but will not require significant time. */
 	UPROPERTY(EditAnywhere, config, Category = "Editor Performance", meta = (DisplayName = "Unpausing cancels all PCG tasks"))
 	bool bUnpauseCancelsAll = false;
+
+	/** Controls whether the alternate PCG pause button is shown instead of the default PCG icon. */
+	UPROPERTY(EditAnywhere, config, Category = "Editor Performance", meta = (DisplayName = "Use alternate pause button icon", EditCondition = "bShowPauseButton", EditConditionHides))
+	bool bUseAlternatePauseButton = false;
+
+	/** Overrides the label of the pause button while PCG is not currently being paused. The default is empty or None, which will use the default label then. */
+	UPROPERTY(EditAnywhere, config, Category = "Editor Performance", meta = (DisplayName = "Overrides default name for the 'currently paused' button", EditCondition = "bShowPauseButton", EditConditionHides))
+	FName OverridePausedButtonLabel = NAME_None;
+
+	/** Overrides the label of the pause button while PCG is currently paused. The default is empty or None, which will use the default label then. */
+	UPROPERTY(EditAnywhere, config, Category = "Editor Performance", meta = (DisplayName = "Overrides default name for the 'not currently paused' button", EditCondition = "bShowPauseButton", EditConditionHides))
+	FName OverrideNotPausedButtonLabel = NAME_None;
 };
 
 #if UE_ENABLE_INCLUDE_ORDER_DEPRECATED_IN_5_2
