Added right-click -> copy column name in the attribute list view.

diff --git a/Engine/Plugins/PCG/Source/PCGEditor/Private/SPCGEditorGraphAttributeListView.cpp b/Engine/Plugins/PCG/Source/PCGEditor/Private/SPCGEditorGraphAttributeListView.cpp
--- a/Engine/Plugins/PCG/Source/PCGEditor/Private/SPCGEditorGraphAttributeListView.cpp
+++ b/Engine/Plugins/PCG/Source/PCGEditor/Private/SPCGEditorGraphAttributeListView.cpp
@@ -1217,12 +1217,39 @@ void SPCGEditorGraphAttributeListView::AddColumn(const UPCGData* InPCGData, cons
 	Arguments.SortMode(this, &SPCGEditorGraphAttributeListView::GetColumnSortMode, InColumnInfo.Id);
 	Arguments.OnSort(this, &SPCGEditorGraphAttributeListView::OnColumnSortModeChanged);
 	Arguments.OverflowPolicy(ETextOverflowPolicy::Ellipsis);
+	Arguments.HeaderComboVisibility(EHeaderComboVisibility::Never);
+	Arguments.MenuContent()
+	[
+		GenerateColumnMenu(InColumnInfo.Id)
+	];
 
 	SHeaderRow::FColumn* NewColumn = new SHeaderRow::FColumn(Arguments);
 	NewColumn->bIsVisible = !HiddenAttributes.Contains(InColumnInfo.Id);
 	ListViewHeader->AddColumn(*NewColumn);
 }
 
+TSharedRef<SWidget> SPCGEditorGraphAttributeListView::GenerateColumnMenu(FName ColumnId)
+{
+	FMenuBuilder MenuBuilder(/*bShouldCloseAfterMenuSelection=*/true, nullptr);
+	MenuBuilder.AddMenuEntry(
+		LOCTEXT("CopyAttributeNameToClipboard", "Copy attribute name"),
+		LOCTEXT("CopyAttributeNameToClipboardTooltip", "Copies the attribute name to the clipboard."),
+		FSlateIcon(),
+		FUIAction(FExecuteAction::CreateLambda([ColumnId]()
+		{
+			FString ColumnIdString = ColumnId.ToString();
+
+			// TODO[UE-221219]: until we support @None as an actual valid token, we need to replace it with None.
+			ColumnIdString = ColumnIdString.Replace(TEXT("@None"), TEXT("None"));
+
+			FPlatformApplicationMisc::ClipboardCopy(*ColumnIdString);
+		})),
+		NAME_None,
+		EUserInterfaceActionType::Button);
+
+	return MenuBuilder.MakeWidget();
+}
+
 void SPCGEditorGraphAttributeListView::CopySelectionToClipboard() const
 {
 	constexpr TCHAR Delimiter = TEXT(',');
diff --git a/Engine/Plugins/PCG/Source/PCGEditor/Private/SPCGEditorGraphAttributeListView.h b/Engine/Plugins/PCG/Source/PCGEditor/Private/SPCGEditorGraphAttributeListView.h
--- a/Engine/Plugins/PCG/Source/PCGEditor/Private/SPCGEditorGraphAttributeListView.h
+++ b/Engine/Plugins/PCG/Source/PCGEditor/Private/SPCGEditorGraphAttributeListView.h
@@ -209,6 +209,7 @@ private:
 
 	void OnColumnSortModeChanged(const EColumnSortPriority::Type SortPriority, const FName& ColumnId, const EColumnSortMode::Type InSortMode);
 	EColumnSortMode::Type GetColumnSortMode(const FName InColumnId) const;
+	TSharedRef<SWidget> GenerateColumnMenu(FName ColumnId);
 
 	void OnFilterTextChanged(const FText& InFilterText);
 	void OnFilterTextCommitted(const FText& NewText, ETextCommit::Type CommitInfo);
