[PCG] GPU - actually use ThreadCountMultiplier setting on Custom kernels

#rnx

diff --git a/Engine/Plugins/PCG/Source/PCG/Private/Compute/Elements/PCGCustomHLSL.cpp b/Engine/Plugins/PCG/Source/PCG/Private/Compute/Elements/PCGCustomHLSL.cpp
--- a/Engine/Plugins/PCG/Source/PCG/Private/Compute/Elements/PCGCustomHLSL.cpp
+++ b/Engine/Plugins/PCG/Source/PCG/Private/Compute/Elements/PCGCustomHLSL.cpp
@@ -246,6 +246,11 @@ int UPCGCustomHLSLSettings::ComputeKernelThreadCount(const UPCGDataBinding* Bind
 		checkNoEntry();
 	}
 
+	if (IsThreadCountMultiplierInUse())
+	{
+		ThreadCount *= ThreadCountMultiplier;
+	}
+
 	return ThreadCount;
 }
 
@@ -1085,6 +1090,15 @@ bool UPCGCustomHLSLSettings::IsKernelValid(FPCGContext* InContext, bool bQuiet)
 		}
 	}
 
+	if (IsThreadCountMultiplierInUse())
+	{
+		if (ThreadCountMultiplier < 1)
+		{
+			PCG_LOG_VALIDATION(FText::Format(LOCTEXT("InvalidThreadCountMultiplier", "Thread Count Multiplier has invalid value ({0}). Must be greater than 0."), ThreadCountMultiplier));
+			return false;
+		}
+	}
+
 	// Validate attributes
 	for (const FPCGKernelAttributeKey& AttributeKey : KernelAttributeKeys)
 	{
diff --git a/Engine/Plugins/PCG/Source/PCG/Public/Compute/Elements/PCGCustomHLSL.h b/Engine/Plugins/PCG/Source/PCG/Public/Compute/Elements/PCGCustomHLSL.h
--- a/Engine/Plugins/PCG/Source/PCG/Public/Compute/Elements/PCGCustomHLSL.h
+++ b/Engine/Plugins/PCG/Source/PCG/Public/Compute/Elements/PCGCustomHLSL.h
@@ -162,6 +162,9 @@ public:
 protected:
 	bool AreKernelAttributesValid(FPCGContext* InContext, FText* OutErrorText) const;
 
+	/** Will the ThreadCountMultiplier value be applied when calculating the dispatch thread count. */
+	bool IsThreadCountMultiplierInUse() const { return KernelType == EPCGKernelType::Custom && DispatchThreadCount != EPCGDispatchThreadCount::Fixed; }
+
 protected:
 	UPROPERTY(BlueprintReadWrite, EditAnywhere, Category = "Settings")
 	EPCGKernelType KernelType = EPCGKernelType::PointProcessor;
