Fixed static analysis warning but also potentially prevented some edge case issues related to cancellation

#rnx

diff --git a/Engine/Plugins/PCG/Source/PCG/Private/PCGElement.cpp b/Engine/Plugins/PCG/Source/PCG/Private/PCGElement.cpp
--- a/Engine/Plugins/PCG/Source/PCG/Private/PCGElement.cpp
+++ b/Engine/Plugins/PCG/Source/PCG/Private/PCGElement.cpp
@@ -433,7 +433,7 @@ void IPCGElement::PostExecutePrimaryLoopElement(FPCGContext* Context, const UPCG
 	// Store individual results in the cache; here we will try to match the remaining hint indices from the input data with the ones given at the output.
 	TArray<FPCGDataCollection> PrimaryDataCollections;
 	FPCGDataCollection OtherData;
-	if (!bHasErrorsOrWarnings && PCGElementHelpers::SplitDataPerPrimaryPin(Settings, Context->InputData, ExecutionLoopMode(Settings), PrimaryDataCollections, OtherData))
+	if (!Context->OutputData.bCancelExecution && !bHasErrorsOrWarnings && PCGElementHelpers::SplitDataPerPrimaryPin(Settings, Context->InputData, ExecutionLoopMode(Settings), PrimaryDataCollections, OtherData))
 	{
 		const bool bShouldComputeFullOutputDataCrc = ShouldComputeFullOutputDataCrc(Context);
 
