[PCG] Address broken actor bounds computation in presence of negative scales.

The result is still out by a scale factor when there are negative scales, and irregular (non 90 degree multiple) rotations.

#rn[PCG] Fixed broken bounds computation when obtaining a single point from an actor which has negative and non-uniform scale.

diff --git a/Engine/Plugins/PCG/Source/PCG/Private/Helpers/PCGHelpers.cpp b/Engine/Plugins/PCG/Source/PCG/Private/Helpers/PCGHelpers.cpp
--- a/Engine/Plugins/PCG/Source/PCG/Private/Helpers/PCGHelpers.cpp
+++ b/Engine/Plugins/PCG/Source/PCG/Private/Helpers/PCGHelpers.cpp
@@ -160,7 +160,20 @@ namespace PCGHelpers
 			{
 				const bool bNonColliding = true;
 
-				const FTransform& ActorToWorld = InActor->GetTransform();
+				// The following code does a bounds computation in local actor space so that Box can capture the tight bounds. 
+				FTransform ActorToWorld = InActor->GetTransform();
+
+				// The matrix inverse below seems to work well for positive scales, but seems to break down badly for non uniform
+				// scales (to see, compare ActorToWorld*WorldToActor to identity) - UE-221283. The following workaround removes mirroring
+				// from the actor transform, does the bounds computation, and then re-mirrors afterwards. This works well for close-to-90deg
+				// actor transform rotations and relatively-uniform scales, but can result in artificial dilation of the bounds in some cases.
+				const FVector ScaleSign = FVector(
+					FMath::Sign(InActor->GetTransform().GetScale3D().X),
+					FMath::Sign(InActor->GetTransform().GetScale3D().Y),
+					FMath::Sign(InActor->GetTransform().GetScale3D().Z));
+
+				ActorToWorld.SetScale3D(InActor->GetTransform().GetScale3D().GetAbs());
+
 				const FTransform WorldToActor = ActorToWorld.Inverse();
 
 				InActor->ForEachComponent<UPrimitiveComponent>(/*bIncludeFromChildActors=*/true, [bNonColliding, bIgnorePCGCreatedComponents, &WorldToActor, &Box](const UPrimitiveComponent* InPrimComp)
@@ -176,6 +189,17 @@ namespace PCGHelpers
 						}
 					}
 				});
+
+				// Un-mirror - see notes above.
+				for (int Axis = 0; Axis < 3; ++Axis)
+				{
+					if (ScaleSign[Axis] < 0.0)
+					{
+						Box.Min[Axis] *= -1.0;
+						Box.Max[Axis] *= -1.0;
+						Swap(Box.Min[Axis], Box.Max[Axis]);
+					}
+				}
 			}
 		}
 		else
