[PCG] Fix bug with World Raycast Node and tag filtering

#rnx

diff --git a/Engine/Plugins/PCG/Source/PCG/Private/Elements/PCGWorldRaycast.cpp b/Engine/Plugins/PCG/Source/PCG/Private/Elements/PCGWorldRaycast.cpp
--- a/Engine/Plugins/PCG/Source/PCG/Private/Elements/PCGWorldRaycast.cpp
+++ b/Engine/Plugins/PCG/Source/PCG/Private/Elements/PCGWorldRaycast.cpp
@@ -27,7 +27,6 @@ namespace PCGWorldRaycastElementConstants
 }
 
 UPCGWorldRaycastElementSettings::UPCGWorldRaycastElementSettings()
-	: WorldQueryParams()
 {
 	OriginInputAttribute.SetPointProperty(EPCGPointProperties::Position);
 	EndPointAttribute.SetPointProperty(EPCGPointProperties::Position);
@@ -298,7 +297,9 @@ bool FPCGWorldRaycastElement::ExecuteInternal(FPCGContext* InContext) const
 			OutputPoints.Reserve(IterState.CachedRayOrigins.Num());
 		}
 
-		const FPCGWorldRaycastQueryParams& WorldQueryParams = Settings->WorldQueryParams;
+		FPCGWorldRaycastQueryParams WorldQueryParams = Settings->WorldQueryParams;
+		WorldQueryParams.Initialize();
+
 
 		if (!PCGWorldQueryHelpers::CreateRayHitAttributes(WorldQueryParams, OutMetadata))
 		{
diff --git a/Engine/Plugins/PCG/Source/PCG/Public/Data/PCGWorldData.h b/Engine/Plugins/PCG/Source/PCG/Public/Data/PCGWorldData.h
--- a/Engine/Plugins/PCG/Source/PCG/Public/Data/PCGWorldData.h
+++ b/Engine/Plugins/PCG/Source/PCG/Public/Data/PCGWorldData.h
@@ -134,7 +134,7 @@ struct FPCGWorldRaycastQueryParams : public FPCGWorldCommonQueryParams
 	UPROPERTY(BlueprintReadWrite, EditAnywhere, Category = "Data|Attributes", meta = (PCG_Overridable, EditCondition = "bTraceComplex"))
 	uint8 bGetFaceIndex : 1 = false;
 
-	/** Create an attribute for UV Coordinates of the surface hit. Note: Will only work in complex traces. */
+	/** Create an attribute for UV Coordinates of the surface hit. Note: Will only work in complex traces and must have 'Project Settings->Physics->Support UV From Hit Results' set to true. */
 	UPROPERTY(BlueprintReadWrite, EditAnywhere, Category = "Data|Attributes", meta = (PCG_Overridable, EditCondition = "bTraceComplex"))
 	uint8 bGetUVCoords : 1 = false;
 
