Fix a NULL pointer dereference when calling GetResourceSizeEx.
#tests Tested in Editor and Cooked Client build.

diff --git a/Engine/Plugins/PCG/Source/PCG/Private/Data/PCGCollisionWrapperData.cpp b/Engine/Plugins/PCG/Source/PCG/Private/Data/PCGCollisionWrapperData.cpp
--- a/Engine/Plugins/PCG/Source/PCG/Private/Data/PCGCollisionWrapperData.cpp
+++ b/Engine/Plugins/PCG/Source/PCG/Private/Data/PCGCollisionWrapperData.cpp
@@ -272,7 +272,10 @@ void UPCGCollisionWrapperData::AddToCrc(FArchiveCrc32& Ar, bool bFullDataCrc) co
 
 void UPCGCollisionWrapperData::GetResourceSizeEx(FResourceSizeEx& CumulativeResourceSize)
 {
-	const_cast<UPCGPointData*>(GetPointData())->GetResourceSizeEx(CumulativeResourceSize);
+	if (GetPointData())
+	{
+		const_cast<UPCGPointData*>(GetPointData())->GetResourceSizeEx(CumulativeResourceSize);
+	}
 }
 
 FBox UPCGCollisionWrapperData::GetBounds() const
