Fixed comparing TSoftObjectPtr with TObjectPtr.

diff --git a/Engine/Plugins/PCG/Source/PCG/Private/Tests/Elements/PCGStaticMeshSpawnerTest.cpp b/Engine/Plugins/PCG/Source/PCG/Private/Tests/Elements/PCGStaticMeshSpawnerTest.cpp
--- a/Engine/Plugins/PCG/Source/PCG/Private/Tests/Elements/PCGStaticMeshSpawnerTest.cpp
+++ b/Engine/Plugins/PCG/Source/PCG/Private/Tests/Elements/PCGStaticMeshSpawnerTest.cpp
@@ -304,7 +304,7 @@ void TestMeshSelectorWeighted(
 		TObjectPtr<UStaticMesh> StaticMesh = ISMC->GetStaticMesh();
 
 		FPCGMeshSelectorWeightedEntry* Entry = MeshSelector->MeshEntries.FindByPredicate([StaticMesh](const FPCGMeshSelectorWeightedEntry& Entry) {
-			return Entry.Descriptor.StaticMesh == StaticMesh;
+			return Entry.Descriptor.StaticMesh == StaticMesh.Get();
 			});
 
 		if (Test->TestNotNull("Validate instanced mesh exists in MeshEntries", Entry))
