[PCG] Create a new "Clean Spline" node

+ Fuse co-located control points
+ Remove collinear control points

diff --git a/Engine/Plugins/PCG/Source/PCG/Private/Elements/PCGCleanSpline.cpp b/Engine/Plugins/PCG/Source/PCG/Private/Elements/PCGCleanSpline.cpp
new file mode 100644
--- /dev/null
+++ b/Engine/Plugins/PCG/Source/PCG/Private/Elements/PCGCleanSpline.cpp
@@ -0,0 +1,225 @@
+// Copyright Epic Games, Inc. All Rights Reserved.
+
+#include "Elements/PCGCleanSpline.h"
+
+#include "PCGContext.h"
+#include "Data/PCGSplineData.h"
+
+#define LOCTEXT_NAMESPACE "PCGCleanSplineElement"
+
+namespace PCGCleanSplineHelpers
+{
+	bool VectorsAreCollinear(const FVector& FirstVector, const FVector& SecondVector, const double Threshold)
+	{
+		return FMath::Abs(FirstVector.GetSafeNormal() | SecondVector.GetSafeNormal()) >= Threshold;
+	}
+
+	bool ControlPointsAreColocated(const FSplinePoint& Point1, const FSplinePoint& Point2, const FTransform& SplineTransform, const double Threshold, const bool bUseLocalSpace = false)
+	{
+		const FVector Location1 = bUseLocalSpace ? Point1.Position : SplineTransform.TransformPosition(Point1.Position);
+		const FVector Location2 = bUseLocalSpace ? Point2.Position : SplineTransform.TransformPosition(Point2.Position);
+		return FVector::DistSquared(Location1, Location2) < Threshold * Threshold;
+	}
+
+	// Implementation note: Co-located points will have a vector dot product of zero, regardless of tangents, and will be thus collinear.
+	bool ControlPointsAreCollinear(const FSplinePoint& Point1, const FSplinePoint& Point2, const FSplinePoint& Point3, const double Threshold)
+	{
+		const FVector Segment = Point3.Position - Point1.Position;
+
+		/* Need to check all four tangents against the segment to guarantee collinearity. Note: This works for linear
+		 * segments only. It is possible to find more control points on curves that would have no effect on the final
+		 * result, but it would be extremely rare for a user to wind up in that situation. It would also provide little
+		 * to no benefit to do so as well.
+		 */
+		return (
+			VectorsAreCollinear(Point1.LeaveTangent, Segment, Threshold) &&
+			VectorsAreCollinear(Point2.ArriveTangent, Segment, Threshold) &&
+			VectorsAreCollinear(Point2.LeaveTangent, Segment, Threshold) &&
+			VectorsAreCollinear(Point3.ArriveTangent, Segment, Threshold)
+		);
+	}
+}
+
+FPCGElementPtr UPCGCleanSplineSettings::CreateElement() const
+{
+	return MakeShared<FPCGCleanSplineElement>();
+}
+
+#if WITH_EDITOR
+void UPCGCleanSplineSettings::PostEditChangeProperty(FPropertyChangedEvent& PropertyChangedEvent)
+{
+	if (PropertyChangedEvent.Property)
+	{
+		if (PropertyChangedEvent.Property->GetFName() == GET_MEMBER_NAME_CHECKED(UPCGCleanSplineSettings, bUseRadians))
+		{
+			CollinearAngleThreshold = bUseRadians ? FMath::DegreesToRadians(CollinearAngleThreshold) : FMath::RadiansToDegrees(CollinearAngleThreshold);
+		}
+		else if (PropertyChangedEvent.Property->GetFName() == GET_MEMBER_NAME_CHECKED(UPCGCleanSplineSettings, CollinearAngleThreshold))
+		{
+			static constexpr double MaxCollinearAngleToleranceDegrees = 89;
+			CollinearAngleThreshold = FMath::Clamp(CollinearAngleThreshold, 0, bUseRadians ? FMath::DegreesToRadians(MaxCollinearAngleToleranceDegrees) : MaxCollinearAngleToleranceDegrees);
+		}
+	}
+}
+#endif // WITH_EDITOR
+
+TArray<FPCGPinProperties> UPCGCleanSplineSettings::InputPinProperties() const
+{
+	TArray<FPCGPinProperties> Properties;
+	Properties.Emplace_GetRef(PCGPinConstants::DefaultInputLabel, EPCGDataType::Spline).SetRequiredPin();
+	return Properties;
+}
+
+TArray<FPCGPinProperties> UPCGCleanSplineSettings::OutputPinProperties() const
+{
+	TArray<FPCGPinProperties> Properties;
+	Properties.Emplace(PCGPinConstants::DefaultOutputLabel, EPCGDataType::Spline);
+	return Properties;
+}
+
+bool FPCGCleanSplineElement::ExecuteInternal(FPCGContext* InContext) const
+{
+	TRACE_CPUPROFILER_EVENT_SCOPE(FPCGCleanSplineElement::Execute);
+
+	check(InContext);
+
+	const UPCGCleanSplineSettings* Settings = InContext->GetInputSettings<UPCGCleanSplineSettings>();
+	check(Settings);
+
+	// Nothing to do. Forward the output.
+	if (!Settings->bFuseColocatedControlPoints && !Settings->bRemoveCollinearControlPoints)
+	{
+		InContext->OutputData = InContext->InputData;
+		PCGLog::LogWarningOnGraph(LOCTEXT("NoOperation", "No Clean Spline operations selected. Input will be forwarded"), InContext);
+		return true;
+	}
+
+	// Pre-calculate the tolerance value from the user defined tolerance in radians/degrees to cross-product comparable value.
+	const double DotProductToleranceRad = Settings->bUseRadians ? Settings->CollinearAngleThreshold : FMath::DegreesToRadians(Settings->CollinearAngleThreshold);
+	// Max the tolerance with an epsilon to accomodate rounding errors.
+	const double DotProductTolerance = FMath::Max(FMath::Abs(FMath::Cos(DotProductToleranceRad)), UE_DOUBLE_SMALL_NUMBER);
+
+	for (const FPCGTaggedData& InputData : InContext->InputData.GetInputsByPin(PCGPinConstants::DefaultInputLabel))
+	{
+		const UPCGSplineData* InputSplineData = Cast<const UPCGSplineData>(InputData.Data);
+		if (!InputSplineData || InputSplineData->GetNumSegments() < 1)
+		{
+			continue;
+		}
+
+		const FTransform& SplineTransform = InputSplineData->SplineStruct.Transform;
+		const FInterpCurveVector& ControlPointsPosition = InputSplineData->SplineStruct.SplineCurves.Position;
+		const FInterpCurveQuat& ControlPointsRotation = InputSplineData->SplineStruct.SplineCurves.Rotation;
+		const FInterpCurveVector& ControlPointsScale = InputSplineData->SplineStruct.SplineCurves.Scale;
+		const int32 NumControlPoints = ControlPointsPosition.Points.Num();
+		bool bControlPointWasRemoved = false;
+
+		TArray<FSplinePoint> ControlPoints;
+		ControlPoints.Reserve(NumControlPoints);
+
+		// For code clarity, generate the points first and remove them as needed.
+		for (int32 i = 0; i < NumControlPoints; ++i)
+		{
+			/* Implementation note: Decay to custom tangents. They assist the user in building the spline, but the
+			 * interpolation modes will affect the recalculations unpredictably when control points are removed.
+			 */
+			ControlPoints.Emplace(/*InputKey=*/i,
+			   ControlPointsPosition.Points[i].OutVal,
+			   ControlPointsPosition.Points[i].ArriveTangent,
+			   ControlPointsPosition.Points[i].LeaveTangent,
+			   ControlPointsRotation.Points[i].OutVal.Rotator(),
+			   ControlPointsScale.Points[i].OutVal,
+			   ESplinePointType::CurveCustomTangent);
+		}
+
+		if (Settings->bFuseColocatedControlPoints)
+		{
+			// Will evaluate by pairs. Reverse order for optimizing RemoveAt.
+			for (int Index = ControlPoints.Num() - 1; Index > 0; --Index)
+			{
+				// Evaluate by pairs for colocated.
+				if (PCGCleanSplineHelpers::ControlPointsAreColocated(
+						ControlPoints[Index - 1],
+						ControlPoints[Index],
+						SplineTransform,
+						Settings->ColocationDistanceThreshold,
+						Settings->bUseSplineLocalSpace))
+				{
+					EPCGControlPointFuseMode FuseMode = Settings->FuseMode;
+					// Generally, keep the first control point, but preserve the final control point to maintain the spline's length
+					if (FuseMode == EPCGControlPointFuseMode::Auto)
+					{
+						FuseMode = (Index == ControlPoints.Num() - 1) ? EPCGControlPointFuseMode::KeepFirst : FuseMode = EPCGControlPointFuseMode::KeepSecond;
+					}
+
+					switch (FuseMode)
+					{
+						// Keep the first control point
+						case EPCGControlPointFuseMode::KeepFirst:
+							ControlPoints[Index - 1].LeaveTangent = ControlPoints[Index].LeaveTangent;
+							ControlPoints.RemoveAt(Index, EAllowShrinking::No);
+							break;
+						// Keep the second control point
+						case EPCGControlPointFuseMode::KeepSecond:
+							ControlPoints[Index].ArriveTangent = ControlPoints[Index - 1].ArriveTangent;
+							ControlPoints.RemoveAt(Index - 1, EAllowShrinking::No);
+							break;
+						// Average the two control points' transforms and update the leave tangent
+						case EPCGControlPointFuseMode::Merge:
+							ControlPoints[Index - 1].Position = FMath::Lerp(ControlPoints[Index - 1].Position, ControlPoints[Index].Position, 0.5);
+							ControlPoints[Index - 1].Rotation = FMath::Lerp(ControlPoints[Index - 1].Rotation, ControlPoints[Index].Rotation, 0.5);
+							ControlPoints[Index - 1].Scale = FMath::Lerp(ControlPoints[Index - 1].Scale, ControlPoints[Index].Scale, 0.5);
+							ControlPoints[Index - 1].LeaveTangent = ControlPoints[Index].LeaveTangent;
+							ControlPoints.RemoveAt(Index, EAllowShrinking::No);
+							break;
+						case EPCGControlPointFuseMode::Auto: // Should've been picked up by now. Fallthrough...
+						default:
+							checkNoEntry();
+							break;
+					}
+
+					bControlPointWasRemoved = true;
+				}
+			}
+		}
+
+		if (Settings->bRemoveCollinearControlPoints)
+		{
+			// Will evaluate by triplets. Reverse order for optimizing RemoveAt.
+			for (int Index = ControlPoints.Num() - 1; Index > 1; --Index)
+			{
+				if (PCGCleanSplineHelpers::ControlPointsAreCollinear(
+						ControlPoints[Index - 2],
+						ControlPoints[Index - 1],
+						ControlPoints[Index],
+						DotProductTolerance))
+				{
+					ControlPoints[Index - 2].LeaveTangent = ControlPoints[Index].Position - ControlPoints[Index - 2].Position;
+					ControlPoints[Index].ArriveTangent = ControlPoints[Index - 2].LeaveTangent;
+					ControlPoints.RemoveAt(Index - 1, EAllowShrinking::No);
+					bControlPointWasRemoved = true;
+				}
+			}
+		}
+
+		// Update input keys if a point was removed, to keep them monotonically incremental
+		if (bControlPointWasRemoved)
+		{
+			for (int i = 0; i < ControlPoints.Num(); ++i)
+			{
+				ControlPoints[i].InputKey = i;
+			}
+		}
+
+		UPCGSplineData* NewSplineData = FPCGContext::NewObject_AnyThread<UPCGSplineData>(InContext);
+		NewSplineData->InitializeFromData(InputSplineData);
+		NewSplineData->Initialize(ControlPoints, InputSplineData->IsClosed(), InputSplineData->GetTransform());
+
+		FPCGTaggedData& Output = InContext->OutputData.TaggedData.Emplace_GetRef(InputData);
+		Output.Data = NewSplineData;
+	}
+
+	return true;
+}
+
+#undef LOCTEXT_NAMESPACE
diff --git a/Engine/Plugins/PCG/Source/PCG/Public/Elements/PCGCleanSpline.h b/Engine/Plugins/PCG/Source/PCG/Public/Elements/PCGCleanSpline.h
new file mode 100644
--- /dev/null
+++ b/Engine/Plugins/PCG/Source/PCG/Public/Elements/PCGCleanSpline.h
@@ -0,0 +1,74 @@
+// Copyright Epic Games, Inc. All Rights Reserved.
+
+#pragma once
+
+#include "PCGSettings.h"
+
+#include "PCGCleanSpline.generated.h"
+
+UENUM()
+enum class EPCGControlPointFuseMode
+{
+	KeepFirst UMETA(Tooltip = "Keep the first of two control points."),
+	KeepSecond UMETA(Tooltip = "Keep the second of two control points."),
+	Merge UMETA(Tooltip = "Fuse the two control points into a new control point, located halfway between the original two."),
+	Auto UMETA(Tooltip = "Will generally keep the second of both control points, except keep the first when it would otherwise alter the length of the spline--ie. final control point.")
+};
+
+/* Remove superfluous control points along the spline, such as those that are co-located or collinear. */
+UCLASS(MinimalAPI, BlueprintType, ClassGroup = (Procedural))
+class UPCGCleanSplineSettings : public UPCGSettings
+{
+	GENERATED_BODY()
+
+public:
+	//~Begin UPCGSettings interface
+#if WITH_EDITOR
+	virtual FName GetDefaultNodeName() const override { return FName(TEXT("CleanSpline")); }
+	virtual FText GetDefaultNodeTitle() const override { return NSLOCTEXT("PCGCleanSplineElement", "NodeTitle", "Clean Spline"); }
+	virtual FText GetNodeTooltipText() const override { return NSLOCTEXT("PCGCleanSplineElement", "NodeTooltip", "Remove superfluous control points along the spline, such as those that are co-located or collinear."); }
+	virtual EPCGSettingsType GetType() const override { return EPCGSettingsType::Spatial; }
+	virtual void PostEditChangeProperty(FPropertyChangedEvent& PropertyChangedEvent) override;
+#endif // WITH_EDITOR
+
+protected:
+	virtual TArray<FPCGPinProperties> InputPinProperties() const override;
+	virtual TArray<FPCGPinProperties> OutputPinProperties() const override;
+	virtual FPCGElementPtr CreateElement() const override;
+	//~End UPCGSettings interface
+
+public:
+	/** Fuse control points that share the same location in world space, within a distance threshold. Colocated points are inherently collinear, so this will automatically be applied when removing collinear points. */
+	UPROPERTY(EditAnywhere, BlueprintReadWrite, Category = "Settings", meta = (PCG_Overridable))
+	bool bFuseColocatedControlPoints = true;
+
+	/** Control points will be considered co-located if they are within this distance from one another. */
+	UPROPERTY(EditAnywhere, BlueprintReadWrite, Category = "Settings", meta = (EditCondition = "bFuseColocatedControlPoints", EditConditionHides, PCG_Overridable))
+	double ColocationDistanceThreshold = UE_KINDA_SMALL_NUMBER;
+
+	/** Use spline local space for the distance calculation, rather than world space. */
+	UPROPERTY(EditAnywhere, BlueprintReadWrite, Category = "Settings", meta = (EditCondition = "bFuseColocatedControlPoints", EditConditionHides, PCG_Overridable))
+	bool bUseSplineLocalSpace = false;
+
+	/** Controls how two co-located points will be fused together. */
+	UPROPERTY(EditAnywhere, BlueprintReadWrite, Category = "Settings", meta = (EditCondition = "bFuseColocatedControlPoints", EditConditionHides, PCG_Overridable))
+	EPCGControlPointFuseMode FuseMode = EPCGControlPointFuseMode::Auto;
+
+	/** Remove control points on linear sections of the spline that would otherwise have no effect on the final spline calculation. */
+	UPROPERTY(EditAnywhere, BlueprintReadWrite, Category = "Settings", meta = (PCG_Overridable))
+	bool bRemoveCollinearControlPoints = true;
+
+	/** A control point will be considered collinear if it is within this angle from the segment between its previous and next control points. */
+	UPROPERTY(EditAnywhere, BlueprintReadWrite, Category = "Settings", meta = (EditCondition = "bRemoveCollinearControlPoints", EditConditionHides, PCG_Overridable, ClampMin = "0"))
+	double CollinearAngleThreshold = 5;
+
+	/** Use radians directly, instead of degrees. The current 'CollinearAngleThreshold' value will be automatically converted when toggled. */
+	UPROPERTY(EditAnywhere, BlueprintReadWrite, Category = "Settings", meta = (EditCondition = "bRemoveCollinearControlPoints", EditConditionHides))
+	bool bUseRadians = false;
+};
+
+class FPCGCleanSplineElement : public IPCGElement
+{
+protected:
+	virtual bool ExecuteInternal(FPCGContext* InContext) const override;
+};
