[PCG] Fix invalid test for Filter By Index

- Specifiying a same-value range such as [2:2] is now a valid option, which just selected the 2nd index directly.

#rnx

diff --git a/Engine/Plugins/PCG/Source/PCG/Private/Tests/Elements/PCGFilterByIndexTest.cpp b/Engine/Plugins/PCG/Source/PCG/Private/Tests/Elements/PCGFilterByIndexTest.cpp
--- a/Engine/Plugins/PCG/Source/PCG/Private/Tests/Elements/PCGFilterByIndexTest.cpp
+++ b/Engine/Plugins/PCG/Source/PCG/Private/Tests/Elements/PCGFilterByIndexTest.cpp
@@ -121,6 +121,21 @@ bool FPCGFilterByIndexTest_Basic::RunTest(const FString& Parameters)
 		UTEST_EQUAL("Output 5 filtered exclusively", Outputs[4].Pin, PCGPinConstants::DefaultOutFilterLabel);
 	}
 
+	// Range index selection [2:2] = [2]
+	{
+		Settings->SelectedIndices = FString("2:2");
+		const TArray<FPCGTaggedData> Outputs = PCGFilterByIndexTestPrivate::ResetAndExecute(TestData, Settings);
+
+		UTEST_EQUAL("Output count", Outputs.Num(), InputNum);
+		UTEST_TRUE("Output is valid", PCGFilterByIndexTestPrivate::ValidateData(Outputs));
+
+		UTEST_EQUAL("Output 1 filtered exclusively", Outputs[0].Pin, PCGPinConstants::DefaultOutFilterLabel);
+		UTEST_EQUAL("Output 2 filtered exclusively", Outputs[1].Pin, PCGPinConstants::DefaultOutFilterLabel);
+		UTEST_EQUAL("Output 3 filtered inclusively", Outputs[2].Pin, PCGPinConstants::DefaultInFilterLabel);
+		UTEST_EQUAL("Output 4 filtered exclusively", Outputs[3].Pin, PCGPinConstants::DefaultOutFilterLabel);
+		UTEST_EQUAL("Output 5 filtered exclusively", Outputs[4].Pin, PCGPinConstants::DefaultOutFilterLabel);
+	}
+
 	// Range index selection (should be [5-3,5-1) = [2,4) = [2:3]
 	{
 		Settings->SelectedIndices = FString("-3:-1");
@@ -177,7 +192,7 @@ bool FPCGFilterByIndexTest_InvalidSelection::RunTest(const FString& Parameters)
 	PCGTestsCommon::FTestData TestData;
 	UPCGFilterByIndexSettings* Settings = PCGTestsCommon::GenerateSettings<UPCGFilterByIndexSettings>(TestData);
 
-	AddExpectedError("Invalid expression in index selection string", EAutomationExpectedMessageFlags::Contains, 3);
+	AddExpectedError("Invalid expression in index selection string", EAutomationExpectedMessageFlags::Contains, 2);
 
 	// Test inverted range
 	{
@@ -195,14 +210,6 @@ bool FPCGFilterByIndexTest_InvalidSelection::RunTest(const FString& Parameters)
 		UTEST_TRUE("Output count", Outputs.IsEmpty());
 	}
 
-	// Duplicate selection for range
-	{
-		Settings->SelectedIndices = FString("2:2");
-		const TArray<FPCGTaggedData> Outputs = PCGFilterByIndexTestPrivate::ResetAndExecute(TestData, Settings);
-
-		UTEST_TRUE("Output count", Outputs.IsEmpty());
-	}
-
 	AddExpectedError("Invalid character in index selection string");
 
 	// Test invalid character
@@ -214,4 +221,4 @@ bool FPCGFilterByIndexTest_InvalidSelection::RunTest(const FString& Parameters)
 	}
 
 	return true;
-}
\ No newline at end of file
+}
diff --git a/Engine/Plugins/PCG/Source/PCG/Private/Tests/Parsing/PCGParsingTest.cpp b/Engine/Plugins/PCG/Source/PCG/Private/Tests/Parsing/PCGParsingTest.cpp
--- a/Engine/Plugins/PCG/Source/PCG/Private/Tests/Parsing/PCGParsingTest.cpp
+++ b/Engine/Plugins/PCG/Source/PCG/Private/Tests/Parsing/PCGParsingTest.cpp
@@ -57,6 +57,14 @@ bool FPCGParsingTest_Basic::RunTest(const FString& Parameters)
 		UTEST_TRUE("Multiple index (first) is correct", Results.ContainsIndex(4) && Results.ContainsIndex(ArraySize - 4));
 	}
 
+	// Repeated index
+	{
+		const PCGIndexing::FPCGIndexCollection Results = PCGParsingTestPrivate::ParseAndRetrieveIndices("2:2", ArraySize);
+		UTEST_EQUAL("Repeated index has 1 index", Results.GetTotalIndexCount(), 1);
+		UTEST_EQUAL("Repeated index has 1 range", Results.GetTotalRangeCount(), 1);
+		UTEST_TRUE("Repeated index is correct", Results.ContainsIndex(2));
+	}
+
 	// More complex use case
 	{
 		TArray<int32> TestArray;
@@ -127,8 +135,8 @@ bool FPCGParsingTest_Ranges::RunTest(const FString& Parameters)
 	// ------------
 
 	Results = PCGParsingTestPrivate::ParseAndRetrieveIndices("0:1", ArraySize);
-	UTEST_EQUAL("Repeated single index results in a single range", Results.GetTotalRangeCount(), 1);
-	UTEST_EQUAL("Repeated single index results in a single index", Results.GetTotalIndexCount(), 1);
+	UTEST_EQUAL("Sequential single index results in a single range", Results.GetTotalRangeCount(), 1);
+	UTEST_EQUAL("Sequential single index results in a single index", Results.GetTotalIndexCount(), 1);
 	UTEST_TRUE("Included indices verified", CheckSplitTestArray(1));
 
 	Results = PCGParsingTestPrivate::ParseAndRetrieveIndices("0:10", ArraySize);
