Added option to quiet attribute filter on data that's missing an attribute

diff --git a/Engine/Plugins/PCG/Source/PCG/Private/Elements/PCGAttributeFilter.cpp b/Engine/Plugins/PCG/Source/PCG/Private/Elements/PCGAttributeFilter.cpp
--- a/Engine/Plugins/PCG/Source/PCG/Private/Elements/PCGAttributeFilter.cpp
+++ b/Engine/Plugins/PCG/Source/PCG/Private/Elements/PCGAttributeFilter.cpp
@@ -153,7 +153,7 @@ namespace PCGAttributeFilterHelpers
 		return true;
 	}
 
-	bool PrepareThresholdInfoFromInput(FPCGContext* InContext, const UPCGData* InputData, const int32 NumInput, const FPCGAttributeFilterThresholdSettings& ThresholdSettings, ThresholdInfo& InOutThresholdInfo, int16 TargetType, bool bCheckCompare, bool bCheckStringSearch, const ThresholdInfo* OtherInfo = nullptr)
+	bool PrepareThresholdInfoFromInput(FPCGContext* InContext, const UPCGData* InputData, const int32 NumInput, const FPCGAttributeFilterThresholdSettings& ThresholdSettings, ThresholdInfo& InOutThresholdInfo, int16 TargetType, bool bCheckCompare, bool bCheckStringSearch, bool bWarnOnDataMissingAttribute, const ThresholdInfo* OtherInfo = nullptr)
 	{
 		check(InContext && InputData);
 
@@ -189,7 +189,11 @@ namespace PCGAttributeFilterHelpers
 
 		if (!InOutThresholdInfo.ThresholdAccessor.IsValid() || !InOutThresholdInfo.ThresholdKeys.IsValid())
 		{
-			PCGE_LOG_C(Warning, GraphAndLog, InContext, FText::Format(LOCTEXT("AttributeMissingForFilter", "DataToFilter does not have '{0}' threshold attribute/property"), FText::FromName(ThresholdSettings.ThresholdAttribute.GetName())));
+			if (bWarnOnDataMissingAttribute)
+			{
+				PCGE_LOG_C(Warning, GraphAndLog, InContext, FText::Format(LOCTEXT("AttributeMissingForFilter", "Filter data does not have '{0}' threshold attribute/property"), FText::FromName(ThresholdSettings.ThresholdAttribute.GetName())));
+			}
+
 			return false;
 		}
 
@@ -522,7 +526,7 @@ TArray<FPCGPinProperties> UPCGAttributeFilteringRangeSettings::OutputPinProperti
 // FPCGAttributeFilterElementBase
 ////////////////////////////////////////
 
-bool FPCGAttributeFilterElementBase::DoFiltering(FPCGContext* Context, EPCGAttributeFilterOperator InOperation, const FPCGAttributePropertyInputSelector& InTargetAttribute, bool bHasSpatialToPointDeprecation, const FPCGAttributeFilterThresholdSettings& FirstThreshold, const FPCGAttributeFilterThresholdSettings* SecondThreshold) const
+bool FPCGAttributeFilterElementBase::DoFiltering(FPCGContext* Context, EPCGAttributeFilterOperator InOperation, const FPCGAttributePropertyInputSelector& InTargetAttribute, bool bHasSpatialToPointDeprecation, bool bWarnOnDataMissingAttribute, const FPCGAttributeFilterThresholdSettings& FirstThreshold, const FPCGAttributeFilterThresholdSettings* SecondThreshold) const
 {
 	TRACE_CPUPROFILER_EVENT_SCOPE(FPCGAttributeFilterElementBase::DoFiltering);
 	check(Context);
@@ -644,7 +648,11 @@ bool FPCGAttributeFilterElementBase::DoFiltering(FPCGContext* Context, EPCGAttri
 
 		if (!TargetAccessor.IsValid() || !TargetKeys.IsValid())
 		{
-			PCGE_LOG(Warning, GraphAndLog, FText::Format(LOCTEXT("TargetMissingAttribute", "TargetData doesn't have target attribute/property '{0}'"), FText::FromName(TargetAttribute.GetName())));
+			if (bWarnOnDataMissingAttribute)
+			{
+				PCGE_LOG(Warning, GraphAndLog, FText::Format(LOCTEXT("TargetMissingAttribute", "Input data doesn't have target attribute/property '{0}'"), FText::FromName(TargetAttribute.GetName())));
+			}
+
 			ForwardInputToOutFilterPin();
 			continue;
 		}
@@ -660,13 +668,13 @@ bool FPCGAttributeFilterElementBase::DoFiltering(FPCGContext* Context, EPCGAttri
 			continue;
 		}
 
-		if (!PCGAttributeFilterHelpers::PrepareThresholdInfoFromInput(Context, OriginalData, NumInput, FirstThreshold, FirstThresholdInfo, TargetType, bCheckCompare, bCheckStringSearch))
+		if (!PCGAttributeFilterHelpers::PrepareThresholdInfoFromInput(Context, OriginalData, NumInput, FirstThreshold, FirstThresholdInfo, TargetType, bCheckCompare, bCheckStringSearch, bWarnOnDataMissingAttribute))
 		{
 			ForwardInputToInFilterPin();
 			continue;
 		}
 
-		if (SecondThreshold && !PCGAttributeFilterHelpers::PrepareThresholdInfoFromInput(Context, OriginalData, NumInput , *SecondThreshold, SecondThresholdInfo, TargetType, bCheckCompare, bCheckStringSearch, &FirstThresholdInfo))
+		if (SecondThreshold && !PCGAttributeFilterHelpers::PrepareThresholdInfoFromInput(Context, OriginalData, NumInput , *SecondThreshold, SecondThresholdInfo, TargetType, bCheckCompare, bCheckStringSearch, bWarnOnDataMissingAttribute, &FirstThresholdInfo))
 		{
 			ForwardInputToInFilterPin();
 			continue;
@@ -893,7 +901,7 @@ bool FPCGAttributeFilterElement::ExecuteInternal(FPCGContext* Context) const
 	ThresholdSettings.ThresholdAttribute = Settings->ThresholdAttribute;
 	ThresholdSettings.AttributeTypes = Settings->AttributeTypes;
 
-	return DoFiltering(Context, Settings->Operator, Settings->TargetAttribute, Settings->bHasSpatialToPointDeprecation, ThresholdSettings);
+	return DoFiltering(Context, Settings->Operator, Settings->TargetAttribute, Settings->bHasSpatialToPointDeprecation, Settings->bWarnOnDataMissingAttribute, ThresholdSettings);
 }
 
 ////////////////////////////////////////
@@ -918,7 +926,7 @@ bool FPCGAttributeFilterRangeElement::ExecuteInternal(FPCGContext* Context) cons
 	const UPCGAttributeFilteringRangeSettings* Settings = Context->GetInputSettings<UPCGAttributeFilteringRangeSettings>();
 	check(Settings);
 
-	return DoFiltering(Context, EPCGAttributeFilterOperator::InRange, Settings->TargetAttribute, Settings->bHasSpatialToPointDeprecation, Settings->MinThreshold, &Settings->MaxThreshold);
+	return DoFiltering(Context, EPCGAttributeFilterOperator::InRange, Settings->TargetAttribute, Settings->bHasSpatialToPointDeprecation, Settings->bWarnOnDataMissingAttribute, Settings->MinThreshold, &Settings->MaxThreshold);
 }
 
 #undef LOCTEXT_NAMESPACE
diff --git a/Engine/Plugins/PCG/Source/PCG/Public/Elements/PCGAttributeFilter.h b/Engine/Plugins/PCG/Source/PCG/Public/Elements/PCGAttributeFilter.h
--- a/Engine/Plugins/PCG/Source/PCG/Public/Elements/PCGAttributeFilter.h
+++ b/Engine/Plugins/PCG/Source/PCG/Public/Elements/PCGAttributeFilter.h
@@ -119,6 +119,10 @@ public:
 	UPROPERTY(BlueprintReadWrite, EditAnywhere, Category = Settings, meta = (EditCondition = "bUseConstantThreshold", EditConditionHides, ShowOnlyInnerProperties, DisplayAfter = "bUseConstantThreshold", PCG_NotOverridable))
 	FPCGMetadataTypesConstantStruct AttributeTypes;
 
+	/** Controls whether the node will emit a warning when the input data or the filter data doesn't have the attribute to filter on. */
+	UPROPERTY(BlueprintReadWrite, EditAnywhere, Category = Settings)
+	bool bWarnOnDataMissingAttribute = true;
+
 	// Hidden value to indicate that Spatial -> Point deprecation is on where pins are not explicitly points.
 	UPROPERTY()
 	bool bHasSpatialToPointDeprecation = false;
@@ -183,6 +187,10 @@ public:
 	UPROPERTY(BlueprintReadWrite, EditAnywhere, Category = Settings, meta = (PCG_Overridable))
 	FPCGAttributeFilterThresholdSettings MaxThreshold;
 
+	/** Controls whether the node will emit a warning when the input data or the filter data doesn't have the attribute to filter on. */
+	UPROPERTY(BlueprintReadWrite, EditAnywhere, Category = Settings)
+	bool bWarnOnDataMissingAttribute = true;
+
 	// Hidden value to indicate that Spatial -> Point deprecation is on where pins are not explicitly points.
 	UPROPERTY()
 	bool bHasSpatialToPointDeprecation = false;
@@ -192,7 +200,7 @@ class FPCGAttributeFilterElementBase : public IPCGElement
 {
 protected:
 	virtual EPCGElementExecutionLoopMode ExecutionLoopMode(const UPCGSettings* Settings) const override { return EPCGElementExecutionLoopMode::SinglePrimaryPin; }
-	bool DoFiltering(FPCGContext* Context, EPCGAttributeFilterOperator InOperation, const FPCGAttributePropertyInputSelector& TargetAttribute, bool bHasSpatialToPointDeprecation, const FPCGAttributeFilterThresholdSettings& FirstThreshold, const FPCGAttributeFilterThresholdSettings* SecondThreshold = nullptr) const;
+	bool DoFiltering(FPCGContext* Context, EPCGAttributeFilterOperator InOperation, const FPCGAttributePropertyInputSelector& TargetAttribute, bool bHasSpatialToPointDeprecation, bool bWarnOnDataMissingAttribute, const FPCGAttributeFilterThresholdSettings& FirstThreshold, const FPCGAttributeFilterThresholdSettings* SecondThreshold = nullptr) const;
 };
 
 class FPCGAttributeFilterElement : public FPCGAttributeFilterElementBase
