[PCG] Update InstanceDataPackers to use attribute accessors

diff --git a/Engine/Plugins/PCG/Source/PCG/Private/InstanceDataPackers/PCGInstanceDataPackerBase.cpp b/Engine/Plugins/PCG/Source/PCG/Private/InstanceDataPackers/PCGInstanceDataPackerBase.cpp
--- a/Engine/Plugins/PCG/Source/PCG/Private/InstanceDataPackers/PCGInstanceDataPackerBase.cpp
+++ b/Engine/Plugins/PCG/Source/PCG/Private/InstanceDataPackers/PCGInstanceDataPackerBase.cpp
@@ -3,7 +3,10 @@
 #include "InstanceDataPackers/PCGInstanceDataPackerBase.h"
 
 #include "PCGElement.h"
+#include "Elements/Metadata/PCGMetadataElementCommon.h"
 #include "Metadata/PCGMetadata.h"
+#include "Metadata/Accessors/IPCGAttributeAccessor.h"
+#include "Metadata/Accessors/PCGAttributeAccessorKeys.h"
 
 #include UE_INLINE_GENERATED_CPP_BY_NAME(PCGInstanceDataPackerBase)
 
@@ -16,6 +19,7 @@ bool UPCGInstanceDataPackerBase::AddTypeToPacking(int TypeId, FPCGPackedCustomDa
 {
 	switch (TypeId)
 	{
+	case PCG::Private::MetadataTypes<bool>::Id: // fall-through
 	case PCG::Private::MetadataTypes<float>::Id: // fall-through
 	case PCG::Private::MetadataTypes<double>::Id: // fall-through
 	case PCG::Private::MetadataTypes<int32>::Id: // fall-through
@@ -29,7 +33,8 @@ bool UPCGInstanceDataPackerBase::AddTypeToPacking(int TypeId, FPCGPackedCustomDa
 	case PCG::Private::MetadataTypes<FVector>::Id:
 		OutPackedCustomData.NumCustomDataFloats += 3;
 		break;
-	case PCG::Private::MetadataTypes<FVector4>::Id:
+	case PCG::Private::MetadataTypes<FVector4>::Id: // fall-through
+	case PCG::Private::MetadataTypes<FQuat>::Id:
 		OutPackedCustomData.NumCustomDataFloats += 4;
 		break;
 	default:
@@ -66,6 +71,15 @@ void UPCGInstanceDataPackerBase::PackCustomDataFromAttributes(const FPCGMeshInst
 
 			switch (AttributeBase->GetTypeId())
 			{
+			case PCG::Private::MetadataTypes<bool>::Id:
+			{
+				const FPCGMetadataAttribute<bool>* Attribute = static_cast<const FPCGMetadataAttribute<bool>*>(AttributeBase);
+				check(Attribute);
+
+				const bool Value = Attribute->GetValueFromItemKey(PointMetadataEntry);
+				OutPackedCustomData.CustomData.Add(Value);
+				break;
+			}
 			case PCG::Private::MetadataTypes<float>::Id:
 			{
 				const FPCGMetadataAttribute<float>* Attribute = static_cast<const FPCGMetadataAttribute<float>*>(AttributeBase);
@@ -146,6 +160,18 @@ void UPCGInstanceDataPackerBase::PackCustomDataFromAttributes(const FPCGMeshInst
 				OutPackedCustomData.CustomData.Add(Value.W);
 				break;
 			}
+			case PCG::Private::MetadataTypes<FQuat>::Id:
+			{
+				const FPCGMetadataAttribute<FQuat>* Attribute = static_cast<const FPCGMetadataAttribute<FQuat>*>(AttributeBase);
+				check(Attribute);
+
+				const FQuat Value = Attribute->GetValueFromItemKey(PointMetadataEntry);
+				OutPackedCustomData.CustomData.Add(Value.X);
+				OutPackedCustomData.CustomData.Add(Value.Y);
+				OutPackedCustomData.CustomData.Add(Value.Z);
+				OutPackedCustomData.CustomData.Add(Value.W);
+				break;
+			}
 			default:
 				break;
 			}
@@ -153,4 +179,139 @@ void UPCGInstanceDataPackerBase::PackCustomDataFromAttributes(const FPCGMeshInst
 	}
 }
 
+void UPCGInstanceDataPackerBase::PackCustomDataFromAccessors(const FPCGMeshInstanceList& InstanceList, TArray<TUniquePtr<const IPCGAttributeAccessor>> Accessors, TArray<TUniquePtr<const IPCGAttributeAccessorKeys>> AccessorKeys, FPCGPackedCustomData& OutPackedCustomData) const
+{
+	const int TotalNumCustomDataFloats = OutPackedCustomData.NumCustomDataFloats * InstanceList.Instances.Num();
+
+	if (!ensure(Accessors.Num() == AccessorKeys.Num()))
+	{
+		// Zero out the custom data to avoid sending garbage data to the GPU.
+		OutPackedCustomData.CustomData.SetNumZeroed(TotalNumCustomDataFloats);
+		return;
+	}
+
+	OutPackedCustomData.CustomData.SetNumUninitialized(TotalNumCustomDataFloats);
+
+	// Index of the accessor's attribute/property for a single instance (e.g. if the
+	// accessors are [float, FVector, float], AccessorPackedIndex will be [0, 1, 4])
+	int AccessorPackedIndex = 0;
+	const EPCGAttributeAccessorFlags Flags = EPCGAttributeAccessorFlags::AllowBroadcast | EPCGAttributeAccessorFlags::AllowConstructible;
+
+	for (int AccessorIndex = 0; AccessorIndex < Accessors.Num(); ++AccessorIndex)
+	{
+		const IPCGAttributeAccessor* Accessor = Accessors[AccessorIndex].Get();
+		const IPCGAttributeAccessorKeys* Keys = AccessorKeys[AccessorIndex].Get();
+
+		check(Accessor && Keys);
+
+		switch (Accessor->GetUnderlyingType())
+		{
+		case PCG::Private::MetadataTypes<bool>::Id: // fall-through
+		case PCG::Private::MetadataTypes<float>::Id: // fall-through
+		case PCG::Private::MetadataTypes<double>::Id: // fall-through
+		case PCG::Private::MetadataTypes<int32>::Id: // fall-through
+		case PCG::Private::MetadataTypes<int64>::Id:
+		{
+			auto PackFloats = [&OutPackedCustomData, AccessorPackedIndex](float Value, int32 Index)
+			{
+				const int PackedIndex = Index * OutPackedCustomData.NumCustomDataFloats + AccessorPackedIndex;
+				OutPackedCustomData.CustomData[PackedIndex] = Value;
+			};
+
+			ensure(PCGMetadataElementCommon::ApplyOnAccessor<float>(*Keys, *Accessor, PackFloats, Flags));
+
+			AccessorPackedIndex += 1;
+
+			break;
+		}
+		case PCG::Private::MetadataTypes<FRotator>::Id:
+		{
+			auto PackFloats = [&OutPackedCustomData, AccessorPackedIndex](const FRotator& Value, int32 Index)
+			{
+				const int PackedIndex = Index * OutPackedCustomData.NumCustomDataFloats + AccessorPackedIndex;
+
+				OutPackedCustomData.CustomData[PackedIndex + 0] = static_cast<float>(Value.Roll);
+				OutPackedCustomData.CustomData[PackedIndex + 1] = static_cast<float>(Value.Pitch);
+				OutPackedCustomData.CustomData[PackedIndex + 2] = static_cast<float>(Value.Yaw);
+			};
+
+			ensure(PCGMetadataElementCommon::ApplyOnAccessor<FRotator>(*Keys, *Accessor, PackFloats, Flags));
+
+			AccessorPackedIndex += 3;
+
+			break;
+		}
+		case PCG::Private::MetadataTypes<FVector2D>::Id:
+		{
+			auto PackFloats = [&OutPackedCustomData, AccessorPackedIndex](const FVector2D& Value, int32 Index)
+			{
+				const int PackedIndex = Index * OutPackedCustomData.NumCustomDataFloats + AccessorPackedIndex;
+
+				OutPackedCustomData.CustomData[PackedIndex + 0] = static_cast<float>(Value.X);
+				OutPackedCustomData.CustomData[PackedIndex + 1] = static_cast<float>(Value.Y);
+			};
+
+			ensure(PCGMetadataElementCommon::ApplyOnAccessor<FVector2D>(*Keys, *Accessor, PackFloats, Flags));
 
+			AccessorPackedIndex += 2;
+
+			break;
+		}
+		case PCG::Private::MetadataTypes<FVector>::Id:
+		{
+			auto PackFloats = [&OutPackedCustomData, AccessorPackedIndex](const FVector& Value, int32 Index)
+			{
+				const int PackedIndex = Index * OutPackedCustomData.NumCustomDataFloats + AccessorPackedIndex;
+
+				OutPackedCustomData.CustomData[PackedIndex + 0] = static_cast<float>(Value.X);
+				OutPackedCustomData.CustomData[PackedIndex + 1] = static_cast<float>(Value.Y);
+				OutPackedCustomData.CustomData[PackedIndex + 2] = static_cast<float>(Value.Z);
+			};
+
+			ensure(PCGMetadataElementCommon::ApplyOnAccessor<FVector>(*Keys, *Accessor, PackFloats, Flags));
+
+			AccessorPackedIndex += 3;
+
+			break;
+		}
+		case PCG::Private::MetadataTypes<FVector4>::Id:
+		{
+			auto PackFloats = [&OutPackedCustomData, AccessorPackedIndex](const FVector4& Value, int32 Index)
+			{
+				const int PackedIndex = Index * OutPackedCustomData.NumCustomDataFloats + AccessorPackedIndex;
+
+				OutPackedCustomData.CustomData[PackedIndex + 0] = static_cast<float>(Value.X);
+				OutPackedCustomData.CustomData[PackedIndex + 1] = static_cast<float>(Value.Y);
+				OutPackedCustomData.CustomData[PackedIndex + 2] = static_cast<float>(Value.Z);
+				OutPackedCustomData.CustomData[PackedIndex + 3] = static_cast<float>(Value.W);
+			};
+
+			ensure(PCGMetadataElementCommon::ApplyOnAccessor<FVector4>(*Keys, *Accessor, PackFloats, Flags));
+
+			AccessorPackedIndex += 4;
+
+			break;
+		}
+		case PCG::Private::MetadataTypes<FQuat>::Id:
+		{
+			auto PackFloats = [&OutPackedCustomData, AccessorPackedIndex](const FQuat& Value, int32 Index)
+			{
+				const int PackedIndex = Index * OutPackedCustomData.NumCustomDataFloats + AccessorPackedIndex;
+
+				OutPackedCustomData.CustomData[PackedIndex + 0] = static_cast<float>(Value.X);
+				OutPackedCustomData.CustomData[PackedIndex + 1] = static_cast<float>(Value.Y);
+				OutPackedCustomData.CustomData[PackedIndex + 2] = static_cast<float>(Value.Z);
+				OutPackedCustomData.CustomData[PackedIndex + 3] = static_cast<float>(Value.W);
+			};
+
+			ensure(PCGMetadataElementCommon::ApplyOnAccessor<FQuat>(*Keys, *Accessor, PackFloats, Flags));
+
+			AccessorPackedIndex += 4;
+
+			break;
+		}
+		default:
+			break;
+		}
+	}
+}
diff --git a/Engine/Plugins/PCG/Source/PCG/Private/InstanceDataPackers/PCGInstanceDataPackerByAttribute.cpp b/Engine/Plugins/PCG/Source/PCG/Private/InstanceDataPackers/PCGInstanceDataPackerByAttribute.cpp
--- a/Engine/Plugins/PCG/Source/PCG/Private/InstanceDataPackers/PCGInstanceDataPackerByAttribute.cpp
+++ b/Engine/Plugins/PCG/Source/PCG/Private/InstanceDataPackers/PCGInstanceDataPackerByAttribute.cpp
@@ -6,43 +6,66 @@
 #include "PCGElement.h"
 #include "Data/PCGSpatialData.h"
 #include "InstanceDataPackers/PCGInstanceDataPackerBase.h"
+#include "Metadata/Accessors/PCGAttributeAccessorHelpers.h"
 
 #include UE_INLINE_GENERATED_CPP_BY_NAME(PCGInstanceDataPackerByAttribute)
 
 #define LOCTEXT_NAMESPACE "PCGInstanceDataPackerByAttribute"
 
+#if WITH_EDITOR
+void UPCGInstanceDataPackerByAttribute::PostLoad()
+{
+	Super::PostLoad();
+
+	PRAGMA_DISABLE_DEPRECATION_WARNINGS
+	for (FName AttributeName : AttributeNames)
+	{
+		FPCGAttributePropertyInputSelector& Selector = AttributeSelectors.Emplace_GetRef();
+		Selector.SetAttributeName(AttributeName);
+	}
+
+	AttributeNames.Reset();
+	PRAGMA_ENABLE_DEPRECATION_WARNINGS
+}
+#endif // WITH_EDITOR
+
 void UPCGInstanceDataPackerByAttribute::PackInstances_Implementation(FPCGContext& Context, const UPCGSpatialData* InSpatialData, const FPCGMeshInstanceList& InstanceList, FPCGPackedCustomData& OutPackedCustomData) const
 {
 	if (!InSpatialData || !InSpatialData->Metadata)
 	{
-		PCGE_LOG_C(Error, GraphAndLog, &Context, LOCTEXT("InvalidInputData", "Invalid input data"));
+		PCGLog::InputOutput::LogInvalidInputDataError(&Context);
 		return;
 	}
 
-	TArray<const FPCGMetadataAttributeBase*> SelectedAttributes;
+	TArray<TUniquePtr<const IPCGAttributeAccessor>> SelectedAccessors;
+	TArray<TUniquePtr<const IPCGAttributeAccessorKeys>> SelectedKeys;
+
+	SelectedAccessors.Reserve(AttributeSelectors.Num());
+	SelectedKeys.Reserve(AttributeSelectors.Num());
 
-	// Find Attributes by name and calculate NumCustomDataFloats
-	for (const FName& AttributeName : AttributeNames)
+	// Find attributes and calculate NumCustomDataFloats
+	for (const FPCGAttributePropertyInputSelector& Selector : AttributeSelectors)
 	{
-		if (!InSpatialData->Metadata->HasAttribute(AttributeName)) 
+		TUniquePtr<const IPCGAttributeAccessor> Accessor = PCGAttributeAccessorHelpers::CreateConstAccessor(InSpatialData, Selector);
+		TUniquePtr<const IPCGAttributeAccessorKeys> Keys = PCGAttributeAccessorHelpers::CreateConstKeys(InSpatialData, Selector);
+
+		if (!Accessor.IsValid() || !Keys.IsValid())
 		{
-			PCGE_LOG_C(Warning, GraphAndLog, &Context, FText::Format(LOCTEXT("AttributeNotInMetadata", "Attribute '{0}' is not in the metadata"), FText::FromName(AttributeName)));
+			PCGLog::Metadata::LogFailToCreateAccessor(Selector, &Context);
 			continue;
 		}
 
-		const FPCGMetadataAttributeBase* AttributeBase = InSpatialData->Metadata->GetConstAttribute(AttributeName);
-		check(AttributeBase);
-
-		if (!AddTypeToPacking(AttributeBase->GetTypeId(), OutPackedCustomData))
+		if (!AddTypeToPacking(Accessor->GetUnderlyingType(), OutPackedCustomData))
 		{
-			PCGE_LOG_C(Warning, GraphAndLog, &Context, FText::Format(LOCTEXT("AttributeInvalidType", "Attribute name '{0}' is not a valid type"), FText::FromName(AttributeName)));
+			PCGLog::LogWarningOnGraph(FText::Format(LOCTEXT("AttributeInvalidType", "Attribute/property '{0}' is not a valid type - skipped."), Selector.GetDisplayText()), &Context);
 			continue;
 		}
 
-		SelectedAttributes.Add(AttributeBase);
+		SelectedAccessors.Add(std::move(Accessor));
+		SelectedKeys.Add(std::move(Keys));
 	}
 
-	PackCustomDataFromAttributes(InstanceList, SelectedAttributes, OutPackedCustomData);
- }
+	PackCustomDataFromAccessors(InstanceList, std::move(SelectedAccessors), std::move(SelectedKeys), OutPackedCustomData);
+}
 
 #undef LOCTEXT_NAMESPACE
diff --git a/Engine/Plugins/PCG/Source/PCG/Private/Tests/Elements/PCGInstancePackerTest.cpp b/Engine/Plugins/PCG/Source/PCG/Private/Tests/Elements/PCGInstancePackerTest.cpp
--- a/Engine/Plugins/PCG/Source/PCG/Private/Tests/Elements/PCGInstancePackerTest.cpp
+++ b/Engine/Plugins/PCG/Source/PCG/Private/Tests/Elements/PCGInstancePackerTest.cpp
@@ -118,43 +118,61 @@ bool FPCGStaticMeshSpawnerInstancePackerByAttributeTest::RunTest(const FString&
 	SourceTaggedData.Data = SourceData;
 	SourceData->TargetActor = TestData.TestActor;
 
+	const FName BoolName = TEXT("Bool");
 	const FName FloatName = TEXT("Float");
 	const FName DoubleName = TEXT("Double");
 	const FName IntName = TEXT("Int");
 	const FName VecName = TEXT("Vec");
 	const FName Vec4Name = TEXT("Vec4");
 	const FName RotatorName = TEXT("Rotator");
+	const FName QuatName = TEXT("Quat");
+
+	const FPCGAttributePropertyInputSelector BoolSelector = FPCGAttributePropertySelector::CreateAttributeSelector<FPCGAttributePropertyInputSelector>(BoolName);
+	const FPCGAttributePropertyInputSelector FloatSelector = FPCGAttributePropertySelector::CreateAttributeSelector<FPCGAttributePropertyInputSelector>(FloatName);
+	const FPCGAttributePropertyInputSelector DoubleSelector = FPCGAttributePropertySelector::CreateAttributeSelector<FPCGAttributePropertyInputSelector>(DoubleName);
+	const FPCGAttributePropertyInputSelector IntSelector = FPCGAttributePropertySelector::CreateAttributeSelector<FPCGAttributePropertyInputSelector>(IntName);
+	const FPCGAttributePropertyInputSelector VecSelector = FPCGAttributePropertySelector::CreateAttributeSelector<FPCGAttributePropertyInputSelector>(VecName);
+	const FPCGAttributePropertyInputSelector Vec4Selector = FPCGAttributePropertySelector::CreateAttributeSelector<FPCGAttributePropertyInputSelector>(Vec4Name);
+	const FPCGAttributePropertyInputSelector RotatorSelector = FPCGAttributePropertySelector::CreateAttributeSelector<FPCGAttributePropertyInputSelector>(RotatorName);
+	const FPCGAttributePropertyInputSelector QuatSelector = FPCGAttributePropertySelector::CreateAttributeSelector<FPCGAttributePropertyInputSelector>(QuatName);
 
 	const bool bAllowsInterpolation = false;
 	const bool bOverrideParent = false;
 
+	SourceData->Metadata->CreateBoolAttribute(BoolName, 0, bAllowsInterpolation, bOverrideParent);
 	SourceData->Metadata->CreateFloatAttribute(FloatName, 0, bAllowsInterpolation, bOverrideParent);
 	SourceData->Metadata->CreateDoubleAttribute(DoubleName, 0, bAllowsInterpolation, bOverrideParent);
 	SourceData->Metadata->CreateInteger64Attribute(IntName, 0, bAllowsInterpolation, bOverrideParent);
 	SourceData->Metadata->CreateVectorAttribute(VecName, FVector::Zero(), bAllowsInterpolation, bOverrideParent);
 	SourceData->Metadata->CreateVector4Attribute(Vec4Name, FVector4::Zero(), bAllowsInterpolation, bOverrideParent);
 	SourceData->Metadata->CreateRotatorAttribute(RotatorName, FRotator::ZeroRotator, bAllowsInterpolation, bOverrideParent);
+	SourceData->Metadata->CreateQuatAttribute(QuatName, FQuat::Identity, bAllowsInterpolation, bOverrideParent);
 
 	TArray<FPCGPoint>& SourcePoints = SourceData->GetMutablePoints();
 
+	TArray<float> BoolValues;
 	TArray<float> FloatValues;
 	TArray<float> DoubleValues;
 	TArray<float> IntValues;
 	TArray<float> VecValues;
 	TArray<float> Vec4Values;
 	TArray<float> RotatorValues;
+	TArray<float> QuatValues;
 
 	FRandomStream RandomSource(TestData.Seed);
 	int NumPoints = 5;
 	for (int I = 0; I < NumPoints; ++I)
 	{
+		const bool BoolValue = I % 2 == 0;
 		const float FloatValue = (I + 0.f) / NumPoints;
 		const double DoubleValue = FloatValue + 1;
 		const int64 IntValue = I;
 		const FVector VecValue = FVector(1, 2, 3) * I;
 		const FVector4 Vec4Value = -FVector(1, 2, 3) * I;
 		const FRotator RotatorValue = FRotator(15, 30, 45) * I;
+		const FQuat QuatValue(RotatorValue);
 
+		BoolValues.Add(BoolValue);
 		FloatValues.Add(FloatValue);
 		DoubleValues.Add(DoubleValue);
 		IntValues.Add(IntValue);
@@ -168,13 +186,19 @@ bool FPCGStaticMeshSpawnerInstancePackerByAttributeTest::RunTest(const FString&
 		RotatorValues.Add(RotatorValue.Roll);
 		RotatorValues.Add(RotatorValue.Pitch);
 		RotatorValues.Add(RotatorValue.Yaw);
+		QuatValues.Add(QuatValue.X);
+		QuatValues.Add(QuatValue.Y);
+		QuatValues.Add(QuatValue.Z);
+		QuatValues.Add(QuatValue.W);
 
+		UPCGMetadataAccessorHelpers::SetBoolAttribute(SourcePoints[I], SourceData->Metadata, BoolName, BoolValue);
 		UPCGMetadataAccessorHelpers::SetFloatAttribute(SourcePoints[I], SourceData->Metadata, FloatName, FloatValue);
 		UPCGMetadataAccessorHelpers::SetDoubleAttribute(SourcePoints[I], SourceData->Metadata, DoubleName, DoubleValue);
 		UPCGMetadataAccessorHelpers::SetInteger64Attribute(SourcePoints[I], SourceData->Metadata, IntName, IntValue);
 		UPCGMetadataAccessorHelpers::SetVectorAttribute(SourcePoints[I], SourceData->Metadata, VecName, VecValue);
 		UPCGMetadataAccessorHelpers::SetVector4Attribute(SourcePoints[I], SourceData->Metadata, Vec4Name, Vec4Value);
 		UPCGMetadataAccessorHelpers::SetRotatorAttribute(SourcePoints[I], SourceData->Metadata, RotatorName, RotatorValue);
+		UPCGMetadataAccessorHelpers::SetQuatAttribute(SourcePoints[I], SourceData->Metadata, QuatName, QuatValue);
 	}
 
 	Settings->SetInstancePackerType(UPCGInstanceDataPackerByAttribute::StaticClass());
@@ -189,51 +213,63 @@ bool FPCGStaticMeshSpawnerInstancePackerByAttributeTest::RunTest(const FString&
 
 	// No attributes
 	{
-		InstancePacker->AttributeNames = {};
+		InstancePacker->AttributeSelectors = {};
 		bTestPassed &= ValidateInstancePacker(this, TestData, Settings, {}, 0);
 	}
 
+	// Bool Attribute
+	{
+		InstancePacker->AttributeSelectors = { BoolSelector };
+		bTestPassed &= ValidateInstancePacker(this, TestData, Settings, BoolValues, 1);
+	}
+
 	// Float Attribute
 	{
-		InstancePacker->AttributeNames = { FloatName };
+		InstancePacker->AttributeSelectors = { FloatSelector };
 		bTestPassed &= ValidateInstancePacker(this, TestData, Settings, FloatValues, 1);
 	}
 
 	// Double Attribute
 	{
-		InstancePacker->AttributeNames = { DoubleName };
+		InstancePacker->AttributeSelectors = { DoubleSelector };
 		bTestPassed &= ValidateInstancePacker(this, TestData, Settings, DoubleValues, 1);
 	}
 
 	// Int Attribute
 	{
-		InstancePacker->AttributeNames = { IntName };
+		InstancePacker->AttributeSelectors = { IntSelector };
 		bTestPassed &= ValidateInstancePacker(this, TestData, Settings, IntValues, 1);
 	}
 
 	// Vector Attribute
 	{
-		InstancePacker->AttributeNames = { VecName };
+		InstancePacker->AttributeSelectors = { VecSelector };
 		bTestPassed &= ValidateInstancePacker(this, TestData, Settings, VecValues, 3);
 	}
 
 	// Vector4 Attribute
 	{
-		InstancePacker->AttributeNames = { Vec4Name };
+		InstancePacker->AttributeSelectors = { Vec4Selector };
 		bTestPassed &= ValidateInstancePacker(this, TestData, Settings, Vec4Values, 4);
 	}
 
 	// Rotator Attribute
 	{
-		InstancePacker->AttributeNames = { RotatorName };
+		InstancePacker->AttributeSelectors = { RotatorSelector };
 		bTestPassed &= ValidateInstancePacker(this, TestData, Settings, RotatorValues, 3);
 	}
 
+	// Quat Attribute
+	{
+		InstancePacker->AttributeSelectors = { QuatSelector };
+		bTestPassed &= ValidateInstancePacker(this, TestData, Settings, QuatValues, 4);
+	}
+
 	// Float + Vector
 	{
 		TArray<float> ExpectedCustomData = LocalPackFloats({ { FloatValues, 1 }, { VecValues, 3 } }, NumPoints);
 
-		InstancePacker->AttributeNames = { FloatName, VecName };
+		InstancePacker->AttributeSelectors = { FloatSelector, VecSelector };
 		bTestPassed &= ValidateInstancePacker(this, TestData, Settings, ExpectedCustomData, 4);
 	}
 
@@ -241,7 +277,7 @@ bool FPCGStaticMeshSpawnerInstancePackerByAttributeTest::RunTest(const FString&
 	{
 		TArray<float> ExpectedCustomData = LocalPackFloats({ { VecValues, 3 }, { FloatValues, 1 } }, NumPoints);
 
-		InstancePacker->AttributeNames = { VecName, FloatName };
+		InstancePacker->AttributeSelectors = { VecSelector, FloatSelector };
 		bTestPassed &= ValidateInstancePacker(this, TestData, Settings, ExpectedCustomData, 4);
 	}
 
@@ -249,7 +285,7 @@ bool FPCGStaticMeshSpawnerInstancePackerByAttributeTest::RunTest(const FString&
 	{
 		TArray<float> ExpectedCustomData = LocalPackFloats({ { VecValues, 3 }, { FloatValues, 1 }, { VecValues, 3 } }, NumPoints);
 
-		InstancePacker->AttributeNames = { VecName, FloatName, VecName };
+		InstancePacker->AttributeSelectors = { VecSelector, FloatSelector, VecSelector };
 		bTestPassed &= ValidateInstancePacker(this, TestData, Settings, ExpectedCustomData, 7);
 	}
 
@@ -257,7 +293,7 @@ bool FPCGStaticMeshSpawnerInstancePackerByAttributeTest::RunTest(const FString&
 	{
 		TArray<float> ExpectedCustomData = LocalPackFloats({ { VecValues, 3 }, { FloatValues, 1 }, { VecValues, 3 }, { Vec4Values, 4 } }, NumPoints);
 
-		InstancePacker->AttributeNames = { VecName, FloatName, VecName, Vec4Name };
+		InstancePacker->AttributeSelectors = { VecSelector, FloatSelector, VecSelector, Vec4Selector };
 		bTestPassed &= ValidateInstancePacker(this, TestData, Settings, ExpectedCustomData, 11);
 	}
 
diff --git a/Engine/Plugins/PCG/Source/PCG/Private/Utils/PCGLogErrors.cpp b/Engine/Plugins/PCG/Source/PCG/Private/Utils/PCGLogErrors.cpp
--- a/Engine/Plugins/PCG/Source/PCG/Private/Utils/PCGLogErrors.cpp
+++ b/Engine/Plugins/PCG/Source/PCG/Private/Utils/PCGLogErrors.cpp
@@ -10,6 +10,7 @@ namespace PCGLog::InputOutput
 	{
 		const FTextFormat TypedInputNotFoundWarning = LOCTEXT("DataInputNotFound", "Data of type {0} not found on pin '{1}'.");
 		const FTextFormat FirstInputOnlyWarning = LOCTEXT("FirstInputOnly", "Multiple inputs found on single-input pin '{0}'. Only the first will be selected.");
+		const FText InvalidInputDataError = LOCTEXT("InvalidInputData", "Invalid input data.");
 	}
 
 	void LogTypedDataNotFoundWarning(EPCGDataType DataType, const FName PinLabel, const FPCGContext* InContext)
@@ -23,6 +24,11 @@ namespace PCGLog::InputOutput
 	{
 		LogWarningOnGraph(FText::Format(ErrorFormat::FirstInputOnlyWarning, FText::FromName(PinLabel)), InContext);
 	}
+
+	void LogInvalidInputDataError(const FPCGContext* InContext)
+	{
+		LogErrorOnGraph(ErrorFormat::InvalidInputDataError, InContext);
+	}
 }
 
 namespace PCGLog::Metadata
diff --git a/Engine/Plugins/PCG/Source/PCG/Public/Elements/PCGStaticMeshSpawner.h b/Engine/Plugins/PCG/Source/PCG/Public/Elements/PCGStaticMeshSpawner.h
--- a/Engine/Plugins/PCG/Source/PCG/Public/Elements/PCGStaticMeshSpawner.h
+++ b/Engine/Plugins/PCG/Source/PCG/Public/Elements/PCGStaticMeshSpawner.h
@@ -70,7 +70,10 @@ public:
 	UPROPERTY(BlueprintReadOnly, EditAnywhere, Category = MeshSelector)
 	bool bAllowDescriptorChanges = true;
 
-	/** Defines the method of custom data packing for spawned (H)ISMCs */
+	/**
+	 * Defines the method of custom data packing for spawned (H)ISMCs. Note, Rotators are treated as 3 floats, while Quaternions are
+	 * treated as 4 floats. You can see an attribute's type in the 'Attribute List View' window, and use an 'Attribute Cast' node to cast to the desired type.
+	 */
 	UPROPERTY(BlueprintReadOnly, EditAnywhere, Category = InstanceDataPacker)
 	TSubclassOf<UPCGInstanceDataPackerBase> InstanceDataPackerType;
 
diff --git a/Engine/Plugins/PCG/Source/PCG/Public/InstanceDataPackers/PCGInstanceDataPackerBase.h b/Engine/Plugins/PCG/Source/PCG/Public/InstanceDataPackers/PCGInstanceDataPackerBase.h
--- a/Engine/Plugins/PCG/Source/PCG/Public/InstanceDataPackers/PCGInstanceDataPackerBase.h
+++ b/Engine/Plugins/PCG/Source/PCG/Public/InstanceDataPackers/PCGInstanceDataPackerBase.h
@@ -4,13 +4,14 @@
 
 #include "PCGInstanceDataPackerBase.generated.h"
 
+class FPCGMetadataAttributeBase;
+class IPCGAttributeAccessor;
+class IPCGAttributeAccessorKeys;
 class UPCGMetadata;
+class UPCGSpatialData;
 struct FPCGContext;
 struct FPCGMeshInstanceList;
 
-class UPCGSpatialData;
-class FPCGMetadataAttributeBase;
-
 USTRUCT(BlueprintType)
 struct FPCGPackedCustomData
 {
@@ -45,6 +46,9 @@ public:
 
 	/** Build a PackedCustomData by processing each attribute in order for each point in the InstanceList */
 	void PackCustomDataFromAttributes(const FPCGMeshInstanceList& InstanceList, const TArray<const FPCGMetadataAttributeBase*>& Attributes, FPCGPackedCustomData& OutPackedCustomData) const;
+
+	/** Build a PackedCustomData by processing each accessor in order for each point in the InstanceList */
+	void PackCustomDataFromAccessors(const FPCGMeshInstanceList& InstanceList, TArray<TUniquePtr<const IPCGAttributeAccessor>> Accessors, TArray<TUniquePtr<const IPCGAttributeAccessorKeys>> AccessorKeys, FPCGPackedCustomData& OutPackedCustomData) const;
 };
 
 #if UE_ENABLE_INCLUDE_ORDER_DEPRECATED_IN_5_2
diff --git a/Engine/Plugins/PCG/Source/PCG/Public/InstanceDataPackers/PCGInstanceDataPackerByAttribute.h b/Engine/Plugins/PCG/Source/PCG/Public/InstanceDataPackers/PCGInstanceDataPackerByAttribute.h
--- a/Engine/Plugins/PCG/Source/PCG/Public/InstanceDataPackers/PCGInstanceDataPackerByAttribute.h
+++ b/Engine/Plugins/PCG/Source/PCG/Public/InstanceDataPackers/PCGInstanceDataPackerByAttribute.h
@@ -3,6 +3,7 @@
 #pragma once
 
 #include "PCGInstanceDataPackerBase.h"
+#include "Metadata/PCGAttributePropertySelector.h"
 
 #include "PCGInstanceDataPackerByAttribute.generated.h"
 
@@ -12,9 +13,21 @@ class UPCGInstanceDataPackerByAttribute : public UPCGInstanceDataPackerBase
 	GENERATED_BODY()
 
 public:
+#if WITH_EDITOR
+	// ~Begin UObject interface
+	virtual void PostLoad() override;
+	// ~End UObject interface
+#endif
+
 	virtual void PackInstances_Implementation(UPARAM(ref) FPCGContext& Context, const UPCGSpatialData* InSpatialData, UPARAM(ref) const FPCGMeshInstanceList& InstanceList, FPCGPackedCustomData& OutPackedCustomData) const override;
 
 public:
 	UPROPERTY(BlueprintReadWrite, EditAnywhere, Category = InstanceDataPacker)
+	TArray<FPCGAttributePropertyInputSelector> AttributeSelectors;
+
+#if WITH_EDITORONLY_DATA
+	UE_DEPRECATED(5.5, "AttributeNames has been deprecated. Use AttributeSelectors instead.")
+	UPROPERTY(BlueprintReadWrite, Category = InstanceDataPacker, meta = (DeprecatedProperty, DeprecationMessage = "AttributeNames has been deprecated. Use AttributeSelectors instead."))
 	TArray<FName> AttributeNames;
+#endif
 };
diff --git a/Engine/Plugins/PCG/Source/PCG/Public/Utils/PCGLogErrors.h b/Engine/Plugins/PCG/Source/PCG/Public/Utils/PCGLogErrors.h
--- a/Engine/Plugins/PCG/Source/PCG/Public/Utils/PCGLogErrors.h
+++ b/Engine/Plugins/PCG/Source/PCG/Public/Utils/PCGLogErrors.h
@@ -23,10 +23,12 @@ namespace PCGLog
 		{
 			extern const FTextFormat TypedInputNotFoundWarning;
 			extern const FTextFormat FirstInputOnlyWarning;
+			extern const FText InvalidInputDataError;
 		}
 
 		void LogTypedDataNotFoundWarning(EPCGDataType DataType, const FName PinLabel, const FPCGContext* InContext = nullptr);
 		void LogFirstInputOnlyWarning(const FName PinLabel, const FPCGContext* InContext = nullptr);
+		void LogInvalidInputDataError(const FPCGContext* InContext = nullptr);
 	}
 
 	namespace Metadata
