Made the get parameter by name more flexible by allowing to go through the graph stack.
Changed a bit the way we embed static subgraphs to accomodate linking to pretask in an easy manner.

diff --git a/Engine/Plugins/PCG/Source/PCG/Private/Elements/PCGUserParameterGet.cpp b/Engine/Plugins/PCG/Source/PCG/Private/Elements/PCGUserParameterGet.cpp
--- a/Engine/Plugins/PCG/Source/PCG/Private/Elements/PCGUserParameterGet.cpp
+++ b/Engine/Plugins/PCG/Source/PCG/Private/Elements/PCGUserParameterGet.cpp
@@ -27,9 +27,9 @@ namespace PCGUserParameterGetSettings
 	* but didn't propagate the changes.
 	* If the property bags aren't the same, we are traversing the graph instance hierarchy to find the first graph/graph instance that matches.
 	*/
-	FConstStructView GetFirstValidLayout(FPCGContext& InContext)
+	TArray<FConstStructView, TInlineAllocator<16>> GetValidLayouts(FPCGContext& InContext)
 	{
-		FConstStructView GraphParameters{};
+		TArray<FConstStructView, TInlineAllocator<16>> Layouts;
 		const UPCGGraphInterface* GraphInterface = nullptr;
 
 		// First we will read from the input if we find an input for this graph, set by the subgraph element, then use it.
@@ -42,12 +42,19 @@ namespace PCGUserParameterGetSettings
 #endif // WITH_EDITOR
 			if (const UPCGUserParametersData* OverrideParametersData = CastChecked<UPCGUserParametersData>(UserParameterData[0].Data, ECastCheckedType::NullAllowed))
 			{
-				return FConstStructView{ OverrideParametersData->UserParameters };
+				Layouts.Emplace(OverrideParametersData->UserParameters);
+
+				const UPCGUserParametersData* UpstreamData = OverrideParametersData->UpstreamData.Get();
+				while (UpstreamData)
+				{
+					Layouts.Emplace(UpstreamData->UserParameters);
+					UpstreamData = UpstreamData->UpstreamData.Get();
+				}
 			}
 			else
 			{
 				PCGE_LOG_C(Error, LogOnly, &InContext, LOCTEXT("InvalidUserParameterData", "Internal error, PCG User Parameters Data is null"));
-				return FConstStructView{};
+				return {};
 			}
 		}
 
@@ -55,58 +62,33 @@ namespace PCGUserParameterGetSettings
 		const UPCGGraph* GraphFromNode = Cast<UPCGGraph>(InContext.Node->GetOuter());
 		check(GraphFromNode);
 		FConstStructView GraphFromNodeParameters = GraphFromNode->GetUserParametersStruct() ? GraphFromNode->GetUserParametersStruct()->GetValue() : FConstStructView{};
-		if (!GraphFromNodeParameters.IsValid())
-		{
-			PCGE_LOG_C(Error, LogOnly, &InContext, LOCTEXT("NoBaseGraphParameters", "Node's graph owner don't have any parameters."));
-			return FConstStructView{};
-		}
 
 		// If we don't have a graph instance, we just use the user parameters from the node graph owner.
 		const UPCGComponent* SourceComponent = InContext.SourceComponent.Get();
 		const UPCGGraphInstance* GraphInstance = SourceComponent ? SourceComponent->GetGraphInstance() : nullptr;
-		if (!GraphInstance)
-		{
-			return GraphFromNodeParameters;
-		}
-
-		// Making sure the graph matches.
-		if (GraphFromNode != GraphInstance->GetGraph())
-		{
-			PCGE_LOG_C(Error, LogOnly, &InContext, LOCTEXT("GraphDontMatch", "Node's graph owner and base graph of the source component doesn't match."));
-			return FConstStructView{};
-		}
-
-		// If we don't have overridden parameters from the input, get the parameters of the graph instance.
-		if (!GraphParameters.IsValid())
-		{
-			GraphParameters = GraphInstance->GetUserParametersStruct() ? GraphInstance->GetUserParametersStruct()->GetValue() : FConstStructView{};
-		}
 
 		// If there is no user parameters, just use the one from the graph.
-		if (!GraphParameters.IsValid())
-		{
-			return GraphFromNodeParameters;
-		}
+		// Go down the graph instance chain and take all layouts. Revert back to original graph node params if none were found
+		bool bAddedNonTrivialLayout = false;
 
-		const UScriptStruct* WantedScriptStruct = GraphFromNodeParameters.GetScriptStruct();
-
-		// Finally, we go down the graph instance chain to find the first valid layout.
-		// In most cases, we should never enter this loop, as layouts should match. But it is a safeguard.
-		while (GraphInstance && GraphParameters.GetScriptStruct() != WantedScriptStruct)
+		while (GraphInstance)
 		{
-			GraphInstance = Cast<UPCGGraphInstance>(GraphInstance->Graph);
-			if (!GraphInstance)
+			FConstStructView GraphParameters = GraphInstance->GetUserParametersStruct() ? GraphInstance->GetUserParametersStruct()->GetValue() : FConstStructView{};
+			if (GraphParameters.IsValid())
 			{
-				// We reached the end, use the user parameters default values. 
-				GraphParameters = GraphFromNodeParameters;
-				break;
+				bAddedNonTrivialLayout = true;
+				Layouts.Add(GraphParameters);
 			}
 
-			GraphParameters = GraphInstance->GetUserParametersStruct()->GetValue();
+			GraphInstance = Cast<UPCGGraphInstance>(GraphInstance->Graph);
 		}
 
+		if (!bAddedNonTrivialLayout && GraphFromNodeParameters.IsValid())
+		{
+			Layouts.Add(GraphFromNodeParameters);
+		}
 
-		return GraphParameters;
+		return Layouts;
 	}
 }
 
@@ -162,76 +144,112 @@ bool FPCGUserParameterGetElement::ExecuteInternal(FPCGContext* Context) const
 	TRACE_CPUPROFILER_EVENT_SCOPE(FPCGUserParameterGetElement::Execute);
 
 	check(Context);
+	const UPCGUserParameterGetSettings* Settings = Context->GetInputSettings<UPCGUserParameterGetSettings>();
+	const UPCGGenericUserParameterGetSettings* GenericSettings = Context->GetInputSettings<UPCGGenericUserParameterGetSettings>();
+	check(Settings || GenericSettings);
 
-	FConstStructView Parameters = PCGUserParameterGetSettings::GetFirstValidLayout(*Context);
-	const UScriptStruct* PropertyBag = Parameters.GetScriptStruct();
+	bool bForceQuiet = (GenericSettings && GenericSettings->bQuiet);
 
-	PCGPropertyHelpers::FExtractorParameters ExtractorParameters;
+	TArray<FConstStructView, TInlineAllocator<16>> ValidLayouts = PCGUserParameterGetSettings::GetValidLayouts(*Context);
 
-	if (const UPCGUserParameterGetSettings* Settings = Context->GetInputSettings<UPCGUserParameterGetSettings>())
-	{
-		ExtractorParameters = PCGPropertyHelpers::FExtractorParameters(Parameters.GetMemory(),
-			PropertyBag,
-			FPCGAttributePropertySelector::CreateAttributeSelector(Settings->PropertyName),
-			Settings->PropertyName,
-			Settings->bForceObjectAndStructExtraction,
-			/*bPropertyNeedsToBeVisible=*/false);
-	}
-	else if (const UPCGGenericUserParameterGetSettings* GenericSettings = Context->GetInputSettings<UPCGGenericUserParameterGetSettings>())
+	// Remove layouts we're not interested in, as per settings
+	if (ValidLayouts.Num() > 1)
 	{
-		ExtractorParameters = PCGPropertyHelpers::FExtractorParameters(Parameters.GetMemory(),
-			PropertyBag,
-			GenericSettings->PropertyPath,
-			GenericSettings->OutputAttributeName,
-			GenericSettings->bForceObjectAndStructExtraction,
-			/*bPropertyNeedsToBeVisible=*/false);
+		if (GenericSettings)
+		{
+			if (GenericSettings->Source == EPCGUserParameterSource::Current)
+			{
+				ValidLayouts.SetNum(1);
+			}
+			else if (GenericSettings->Source == EPCGUserParameterSource::Upstream)
+			{
+				ValidLayouts.RemoveAt(0);
+			}
+			else if (GenericSettings->Source == EPCGUserParameterSource::Root)
+			{
+				ValidLayouts.RemoveAtSwap(0);
+				ValidLayouts.SetNum(1);
+			}
+		}
+		else
+		{
+			// Not a generic get settings -> default to taking the current layout only
+			ValidLayouts.SetNum(1);
+		}
 	}
-	else
+
+	for(int CurrentLayoutIndex = 0; CurrentLayoutIndex < ValidLayouts.Num(); ++CurrentLayoutIndex)
 	{
-		checkNoEntry();
-		return true;
-	}
+		const bool bIsLastIteration = (CurrentLayoutIndex == ValidLayouts.Num() - 1);
+		FConstStructView Parameters = ValidLayouts[CurrentLayoutIndex];
+		const UScriptStruct* PropertyBag = Parameters.GetScriptStruct();
+
+		PCGPropertyHelpers::FExtractorParameters ExtractorParameters;
 
-	// Don't care for object traversed in non-editor build, since it is only useful for tracking.
-	TSet<FSoftObjectPath>* ObjectTraversedPtr = nullptr;
+		if (Settings)
+		{
+			ExtractorParameters = PCGPropertyHelpers::FExtractorParameters(Parameters.GetMemory(),
+				PropertyBag,
+				FPCGAttributePropertySelector::CreateAttributeSelector(Settings->PropertyName),
+				Settings->PropertyName,
+				Settings->bForceObjectAndStructExtraction,
+				/*bPropertyNeedsToBeVisible=*/false);
+		}
+		else if (GenericSettings)
+		{
+			ExtractorParameters = PCGPropertyHelpers::FExtractorParameters(Parameters.GetMemory(),
+				PropertyBag,
+				GenericSettings->PropertyPath,
+				GenericSettings->OutputAttributeName,
+				GenericSettings->bForceObjectAndStructExtraction,
+				/*bPropertyNeedsToBeVisible=*/false);
+		}
+
+		// Don't care for object traversed in non-editor build, since it is only useful for tracking.
+		TSet<FSoftObjectPath>* ObjectTraversedPtr = nullptr;
 #if WITH_EDITOR
-	TSet<FSoftObjectPath> ObjectTraversed;
-	ObjectTraversedPtr = &ObjectTraversed;
+		TSet<FSoftObjectPath> ObjectTraversed;
+		ObjectTraversedPtr = &ObjectTraversed;
 #endif // WITH_EDITOR
 
-	if (UPCGParamData* ParamData = PCGPropertyHelpers::ExtractPropertyAsAttributeSet(ExtractorParameters, Context, ObjectTraversedPtr))
-	{
-		Context->OutputData.TaggedData.Emplace_GetRef().Data = ParamData;
-	}
-	else
-	{
-		if (!ExtractorParameters.PropertySelectors.IsEmpty())
+		if (UPCGParamData* ParamData = PCGPropertyHelpers::ExtractPropertyAsAttributeSet(ExtractorParameters, Context, ObjectTraversedPtr, /*bQuiet=*/bForceQuiet || !bIsLastIteration))
 		{
-			PCGE_LOG(Error, GraphAndLog, FText::Format(LOCTEXT("InvalidProperty", "Could not find the property '{0}' in the user parameters"), ExtractorParameters.PropertySelectors[0].GetDisplayText()));
+			Context->OutputData.TaggedData.Emplace_GetRef().Data = ParamData;
 		}
-		else
+		else if (!bIsLastIteration)
 		{
-			PCGE_LOG(Error, GraphAndLog, LOCTEXT("InvalidNamelessProperty", "Could not find nameless property in the user parameters"));
+			continue;
+		}
+		else // Final iteration, should break & report errors
+		{
+			if (!ExtractorParameters.PropertySelectors.IsEmpty())
+			{
+				PCGE_LOG(Error, GraphAndLog, FText::Format(LOCTEXT("InvalidProperty", "Could not find the property '{0}' in the user parameters"), ExtractorParameters.PropertySelectors[0].GetDisplayText()));
+			}
+			else
+			{
+				PCGE_LOG(Error, GraphAndLog, LOCTEXT("InvalidNamelessProperty", "Could not find nameless property in the user parameters"));
+			}
 		}
 
-		return true;
-	}
-
-	// Register dynamic tracking
+		// Register dynamic tracking
 #if WITH_EDITOR
-	if (!ObjectTraversed.IsEmpty())
-	{
-		FPCGDynamicTrackingHelper DynamicTracking;
-		DynamicTracking.EnableAndInitialize(Context, ObjectTraversed.Num());
-		for (FSoftObjectPath& Path : ObjectTraversed)
+		if (!ObjectTraversed.IsEmpty())
 		{
-			DynamicTracking.AddToTracking(FPCGSelectionKey::CreateFromPath(std::move(Path)), /*bCulled=*/false);
-		}
+			FPCGDynamicTrackingHelper DynamicTracking;
+			DynamicTracking.EnableAndInitialize(Context, ObjectTraversed.Num());
+			for (FSoftObjectPath& Path : ObjectTraversed)
+			{
+				DynamicTracking.AddToTracking(FPCGSelectionKey::CreateFromPath(std::move(Path)), /*bCulled=*/false);
+			}
 
-		DynamicTracking.Finalize(Context);
-	}
+			DynamicTracking.Finalize(Context);
+		}
 #endif // WITH_EDITOR
 
+		break;
+	}
+	
 	return true;
 }
 
diff --git a/Engine/Plugins/PCG/Source/PCG/Private/Graph/PCGGraphCompiler.cpp b/Engine/Plugins/PCG/Source/PCG/Private/Graph/PCGGraphCompiler.cpp
--- a/Engine/Plugins/PCG/Source/PCG/Private/Graph/PCGGraphCompiler.cpp
+++ b/Engine/Plugins/PCG/Source/PCG/Private/Graph/PCGGraphCompiler.cpp
@@ -125,11 +125,24 @@ TArray<FPCGGraphTask> FPCGGraphCompiler::CompileGraph(UPCGGraph* InGraph, FPCGTa
 			for (const UPCGPin* InputPin : Node->InputPins)
 			{
 				check(InputPin);
+				// Params have to be funneled in the subgraph element (so we can pass down user parameters)
+				const bool bIsParamPin = SubgraphSettings->HasOverridableParam(InputPin->Properties.Label);
+
 				for (const UPCGEdge* InboundEdge : InputPin->Edges)
 				{
 					if (InboundEdge->IsValid())
 					{
-						PreTask.Inputs.Emplace(IdMapping[InboundEdge->InputPin->Node], InboundEdge->InputPin->Properties, InboundEdge->OutputPin->Properties);
+						// Implementation note: conceptually, the non-param inputs need to be connected only to the input node task.
+						// However, because of the static/dynamic culling which happen on the subgraph node, we need to have the connections on the pretask too (e.g. the subgraph node) even if they don't provide data.
+						if (bIsParamPin)
+						{
+							PreTask.Inputs.Emplace(IdMapping[InboundEdge->InputPin->Node], InboundEdge->InputPin->Properties, InboundEdge->OutputPin->Properties, /*bInProvideData=*/true);
+						}
+						else
+						{
+							PreTask.Inputs.Emplace(IdMapping[InboundEdge->InputPin->Node], InboundEdge->InputPin->Properties, InboundEdge->OutputPin->Properties, /*bInProvideData=*/false);
+							InputNodeTask->Inputs.Emplace(IdMapping[InboundEdge->InputPin->Node], InboundEdge->InputPin->Properties, InboundEdge->OutputPin->Properties, /*bInProvideData=*/true);
+						}
 					}
 					else
 					{
@@ -138,10 +151,10 @@ TArray<FPCGGraphTask> FPCGGraphCompiler::CompileGraph(UPCGGraph* InGraph, FPCGTa
 				}
 			}
 
-			// Add pre-task as input to subgraph input node task
+			// Add pre-task as input to subgraph input node task, without data dependency
 			if (InputNodeTask)
 			{
-				InputNodeTask->Inputs.Emplace(PreId);
+				InputNodeTask->Inputs.Emplace(PreId,/*InUpstreamPin=*/FPCGGraphTaskInput::NoPin, /*InDownstreamPin=*/FPCGGraphTaskInput::NoPin, /*bInProvideData=*/false);
 			}
 
 			// Hook nodes to the PreTask if they require so.
diff --git a/Engine/Plugins/PCG/Source/PCG/Private/Helpers/PCGPropertyHelpers.cpp b/Engine/Plugins/PCG/Source/PCG/Private/Helpers/PCGPropertyHelpers.cpp
--- a/Engine/Plugins/PCG/Source/PCG/Private/Helpers/PCGPropertyHelpers.cpp
+++ b/Engine/Plugins/PCG/Source/PCG/Private/Helpers/PCGPropertyHelpers.cpp
@@ -77,7 +77,7 @@ namespace PCGPropertyHelpers
 	* @param OptionalObjectTraversed Optional set to store all object that we traversed, to be able to react to those objects changes.
 	* @returns                 The last property of the chain (and its container address is in OutContainer)
 	*/
-	const FProperty* ExtractPropertyChain(const UStruct* CurrentClass, const FName CurrentName, TArrayView<const FString> NextNames, const bool bNeedsToBeVisible, TArray<const void*>& OutContainers, FPCGContext* OptionalContext, TSet<FSoftObjectPath>* OptionalObjectTraversed)
+	const FProperty* ExtractPropertyChain(const UStruct* CurrentClass, const FName CurrentName, TArrayView<const FString> NextNames, const bool bNeedsToBeVisible, TArray<const void*>& OutContainers, FPCGContext* OptionalContext, TSet<FSoftObjectPath>* OptionalObjectTraversed, bool bQuiet)
 	{
 		check(CurrentClass);
 
@@ -102,14 +102,22 @@ namespace PCGPropertyHelpers
 
 		if (!Property)
 		{
-			LogError(FText::Format(LOCTEXT("PropertyDoesNotExist", "Property '{0}' does not exist in {1}."), FText::FromName(CurrentName), FText::FromName(CurrentClass->GetFName())), OptionalContext);
+			if (!bQuiet)
+			{
+				LogError(FText::Format(LOCTEXT("PropertyDoesNotExist", "Property '{0}' does not exist in {1}."), FText::FromName(CurrentName), FText::FromName(CurrentClass->GetFName())), OptionalContext);
+			}
+
 			return nullptr;
 		}
 
 		// Make sure the property is visible, if requested
 		if (bNeedsToBeVisible && (Property->HasAnyPropertyFlags(ExcludePropertyFlags) || !Property->HasAnyPropertyFlags(IncludePropertyFlags)))
 		{
-			LogError(FText::Format(LOCTEXT("PropertyExistsButNotVisible", "Property '{0}' does exist in {1}, but is not visible."), FText::FromName(CurrentName), FText::FromName(CurrentClass->GetFName())), OptionalContext);
+			if (!bQuiet)
+			{
+				LogError(FText::Format(LOCTEXT("PropertyExistsButNotVisible", "Property '{0}' does exist in {1}, but is not visible."), FText::FromName(CurrentName), FText::FromName(CurrentClass->GetFName())), OptionalContext);
+			}
+
 			return nullptr;
 		}
 
@@ -165,11 +173,15 @@ namespace PCGPropertyHelpers
 			
 			if(bPropertyNotExtractable)
 			{
-				LogError(FText::Format(LOCTEXT("PropertyIsNotExtractable", "Property '{0}' does exist in {1}, but is not extractable."), FText::FromName(CurrentName), FText::FromName(CurrentClass->GetFName())), OptionalContext);
+				if (!bQuiet)
+				{
+					LogError(FText::Format(LOCTEXT("PropertyIsNotExtractable", "Property '{0}' does exist in {1}, but is not extractable."), FText::FromName(CurrentName), FText::FromName(CurrentClass->GetFName())), OptionalContext);
+				}
+
 				return nullptr;
 			}
 
-			return ExtractPropertyChain(NextClass, FName(NextNames[0]), NextNames.RightChop(1), bNeedsToBeVisible, OutContainers, OptionalContext, OptionalObjectTraversed);
+			return ExtractPropertyChain(NextClass, FName(NextNames[0]), NextNames.RightChop(1), bNeedsToBeVisible, OutContainers, OptionalContext, OptionalObjectTraversed, bQuiet);
 		}
 		else
 		{
@@ -218,7 +230,7 @@ EPCGMetadataTypes PCGPropertyHelpers::GetMetadataTypeFromProperty(const FPropert
 	return PropertyAccessor.IsValid() ? EPCGMetadataTypes(PropertyAccessor->GetUnderlyingType()) : EPCGMetadataTypes::Unknown;
 }
 
-UPCGParamData* PCGPropertyHelpers::ExtractPropertyAsAttributeSet(const PCGPropertyHelpers::FExtractorParameters& Parameters, FPCGContext* OptionalContext, TSet<FSoftObjectPath>* OptionalObjectTraversed)
+UPCGParamData* PCGPropertyHelpers::ExtractPropertyAsAttributeSet(const PCGPropertyHelpers::FExtractorParameters& Parameters, FPCGContext* OptionalContext, TSet<FSoftObjectPath>* OptionalObjectTraversed, bool bQuiet)
 {
 	check(Parameters.Container && Parameters.Class);
 
@@ -251,7 +263,7 @@ UPCGParamData* PCGPropertyHelpers::ExtractPropertyAsAttributeSet(const PCGProper
 		// If Name is none, extract the container as-is, using Parameters.Class, otherwise, extract the chain.
 		if (!ExtractRoot)
 		{
-			Property = ExtractPropertyChain(Parameters.Class, PropertyName, PropertySelector.GetExtraNames(), Parameters.bPropertyNeedsToBeVisible, Containers, OptionalContext, OptionalObjectTraversed);
+			Property = ExtractPropertyChain(Parameters.Class, PropertyName, PropertySelector.GetExtraNames(), Parameters.bPropertyNeedsToBeVisible, Containers, OptionalContext, OptionalObjectTraversed, bQuiet);
 			if (!Property)
 			{
 				return nullptr;
@@ -342,7 +354,11 @@ UPCGParamData* PCGPropertyHelpers::ExtractPropertyAsAttributeSet(const PCGProper
 
 		if (ExtractableProperties.IsEmpty())
 		{
-			LogError(LOCTEXT("NoPropertiesFound", "No properties found to extract"), OptionalContext);
+			if (!bQuiet)
+			{
+				LogError(LOCTEXT("NoPropertiesFound", "No properties found to extract"), OptionalContext);
+			}
+
 			return nullptr;
 		}
 
@@ -367,7 +383,11 @@ UPCGParamData* PCGPropertyHelpers::ExtractPropertyAsAttributeSet(const PCGProper
 		{
 			if (Entries.Num() != ElementAddresses.Num())
 			{
-				LogError(LOCTEXT("InvalidCardinality", "Unable to extract because some properties are of mismatched sizes"), OptionalContext);
+				if (!bQuiet)
+				{
+					LogError(LOCTEXT("InvalidCardinality", "Unable to extract because some properties are of mismatched sizes"), OptionalContext);
+				}
+
 				return nullptr;
 			}
 		}
@@ -390,7 +410,11 @@ UPCGParamData* PCGPropertyHelpers::ExtractPropertyAsAttributeSet(const PCGProper
 
 				if (!Metadata->SetAttributeFromDataProperty(FName(AttributeName), EntryKey, ContainerPtr, FinalProperty, /*bCreate=*/ true))
 				{
-					LogError(FText::Format(LOCTEXT("ErrorCreatingAttribute", "Error while creating an attribute for property '{0}'. Either the property type is not supported by PCG or attribute creation failed."), FText::FromString(FinalProperty->GetName())), OptionalContext);
+					if (!bQuiet)
+					{
+						LogError(FText::Format(LOCTEXT("ErrorCreatingAttribute", "Error while creating an attribute for property '{0}'. Either the property type is not supported by PCG or attribute creation failed."), FText::FromString(FinalProperty->GetName())), OptionalContext);
+					}
+
 					bValidOperation = false;
 					break;
 				}
diff --git a/Engine/Plugins/PCG/Source/PCG/Private/PCGComponent.cpp b/Engine/Plugins/PCG/Source/PCG/Private/PCGComponent.cpp
--- a/Engine/Plugins/PCG/Source/PCG/Private/PCGComponent.cpp
+++ b/Engine/Plugins/PCG/Source/PCG/Private/PCGComponent.cpp
@@ -9,6 +9,7 @@
 #include "PCGInputOutputSettings.h"
 #include "PCGManagedResource.h"
 #include "PCGPin.h"
+#include "PCGSubgraph.h"
 #include "PCGSubsystem.h"
 #include "Data/PCGIntersectionData.h"
 #include "Data/PCGLandscapeData.h"
@@ -2434,7 +2435,7 @@ void UPCGComponent::StoreInspectionData(const FPCGStack* InStack, const UPCGNode
 
 	if (IsInspecting())
 	{
-		auto StorePinInspectionData = [InStack, InNode](const TArray<TObjectPtr<UPCGPin>>& InPins, const FPCGDataCollection& InData, TMap<FPCGStack, FPCGDataCollection>& InOutInspectionCache)
+		auto StorePinInspectionDataFromNode = [](const FPCGStack* InStack, const TArray<TObjectPtr<UPCGPin>>& InPins, const FPCGDataCollection& InData, TMap<FPCGStack, FPCGDataCollection>& InOutInspectionCache)
 		{
 			for (const UPCGPin* Pin : InPins)
 			{
@@ -2442,8 +2443,6 @@ void UPCGComponent::StoreInspectionData(const FPCGStack* InStack, const UPCGNode
 
 				// Append the Node and Pin to the current Stack to uniquely identify each DataCollection
 				TArray<FPCGStackFrame>& StackFrames = Stack.GetStackFramesMutable();
-				StackFrames.Reserve(StackFrames.Num() + 2);
-				StackFrames.Emplace(InNode);
 				StackFrames.Emplace(Pin);
 
 				FPCGDataCollection PinDataCollection;
@@ -2452,13 +2451,49 @@ void UPCGComponent::StoreInspectionData(const FPCGStack* InStack, const UPCGNode
 				// Implementation note: since static subgraphs actually are visited twice and the second time the input doesn't match the input pins, we don't clear the data.
 				if (!PinDataCollection.TaggedData.IsEmpty())
 				{
-					InOutInspectionCache.Add(Stack, PinDataCollection);
+					if (FPCGDataCollection* CollectionInCache = InOutInspectionCache.Find(Stack))
+					{
+						CollectionInCache->TaggedData.Append(PinDataCollection.TaggedData);
+					}
+					else
+					{
+						InOutInspectionCache.Add(Stack, PinDataCollection);
+					}
 				}
 			}
 		};
 
+		auto StorePinInspectionData = [InStack, InNode, &StorePinInspectionDataFromNode](const TArray<TObjectPtr<UPCGPin>>& InPins, const FPCGDataCollection& InData, TMap<FPCGStack, FPCGDataCollection>& InOutInspectionCache)
+		{
+			FPCGStack Stack = *InStack;
+
+			// Append the Node (here) and Pin (in call) to the current Stack to uniquely identify each DataCollection
+			TArray<FPCGStackFrame>& StackFrames = Stack.GetStackFramesMutable();
+			StackFrames.Reserve(StackFrames.Num() + 2);
+			StackFrames.Emplace(InNode);
+
+			StorePinInspectionDataFromNode(&Stack, InPins, InData, InOutInspectionCache);
+		};
+
 		FWriteScopeLock Lock(InspectionCacheLock);
 
+		// Special case: if we have a static (embedded) subgraph, then the actual data inputs (not params) of the subgraph will be on the input node.
+		// Considering we don't allow inspection on input pins of the input node, then we can move that data up the chain.
+		if (InNode->GetSettings()->IsA<UPCGGraphInputOutputSettings>() && InStack->GetStackFrames().Num() > 2)
+		{
+			// We're expecting the last frame to be the graph
+			// Then, if the graph was statically dispatched, it will be the subgraph node.
+			// In the case of a dynamic subgraph or loop, it will be the loop index instead.
+			FPCGStack StackToSubgraphNode = *InStack;
+			TArray<FPCGStackFrame>& StackFrames = StackToSubgraphNode.GetStackFramesMutable();
+			StackFrames.Pop();
+
+			if (StackFrames.Last().Object.IsValid() && StackFrames.Last().Object->IsA<UPCGSubgraphNode>())
+			{
+				StorePinInspectionDataFromNode(&StackToSubgraphNode, Cast<const UPCGSubgraphNode>(StackFrames.Last().Object)->GetInputPins(), InInputData, InspectionCache);
+			}
+		}
+
 		StorePinInspectionData(InNode->GetInputPins(), InInputData, InspectionCache);
 		StorePinInspectionData(InNode->GetOutputPins(), InOutputData, InspectionCache);
 	}
diff --git a/Engine/Plugins/PCG/Source/PCG/Private/PCGGraph.cpp b/Engine/Plugins/PCG/Source/PCG/Private/PCGGraph.cpp
--- a/Engine/Plugins/PCG/Source/PCG/Private/PCGGraph.cpp
+++ b/Engine/Plugins/PCG/Source/PCG/Private/PCGGraph.cpp
@@ -223,7 +223,8 @@ EPCGChangeType UPCGGraphInterface::GetChangeTypeForGraphParameterChange(EPCGGrap
 	}
 
 	// Finally if anything change on a property that has an impact for the graph, look for GetUserParameters nodes for this property, to only refresh if the property is used.
-	for (const UPCGNode* Node : Graph->GetNodes())
+	// TODO: add tracking for user parameters from subgraphs
+	/*for (const UPCGNode* Node : Graph->GetNodes())
 	{
 		if (!Node)
 		{
@@ -240,7 +241,8 @@ EPCGChangeType UPCGGraphInterface::GetChangeTypeForGraphParameterChange(EPCGGrap
 	}
 
 	// At this point, we didn't find any node that use our property, so no refresh needed.
-	return EPCGChangeType::Cosmetic;
+	return EPCGChangeType::Cosmetic;*/
+	return EPCGChangeType::Settings;
 }
 
 /****************************
diff --git a/Engine/Plugins/PCG/Source/PCG/Private/PCGSettings.cpp b/Engine/Plugins/PCG/Source/PCG/Private/PCGSettings.cpp
--- a/Engine/Plugins/PCG/Source/PCG/Private/PCGSettings.cpp
+++ b/Engine/Plugins/PCG/Source/PCG/Private/PCGSettings.cpp
@@ -438,6 +438,14 @@ TArray<FPCGPinProperties> UPCGSettings::OutputPinProperties() const
 	return PinProperties;
 }
 
+bool UPCGSettings::HasOverridableParam(FName InParamName) const
+{
+	return CachedOverridableParams.FindByPredicate([InParamName](const FPCGSettingsOverridableParam& ParamToCheck)
+	{
+		return !ParamToCheck.PropertiesNames.IsEmpty() && ParamToCheck.PropertiesNames.Last() == InParamName;
+	}) != nullptr;
+}
+
 void UPCGSettings::FillOverridableParamsPins(TArray<FPCGPinProperties>& OutPins) const
 {
 	if (!HasOverridableParams())
diff --git a/Engine/Plugins/PCG/Source/PCG/Private/PCGSubgraph.cpp b/Engine/Plugins/PCG/Source/PCG/Private/PCGSubgraph.cpp
--- a/Engine/Plugins/PCG/Source/PCG/Private/PCGSubgraph.cpp
+++ b/Engine/Plugins/PCG/Source/PCG/Private/PCGSubgraph.cpp
@@ -614,6 +614,21 @@ void FPCGSubgraphElement::PrepareSubgraphUserParameters(const UPCGSubgraphSettin
 			// Do nothing, we still want to have a User Parameter Data to indicate we are in a subgraph context.
 		}
 
+		// Hook up user parameter data from upstream
+		TArray<FPCGTaggedData> UpstreamUserParameterData = Context->InputData.GetTaggedTypedInputs<UPCGUserParametersData>(PCGBaseSubgraphConstants::UserParameterTagData);
+		if (!UpstreamUserParameterData.IsEmpty())
+		{
+#if WITH_EDITOR
+			// Safeguard to make sure we always have one and only one data of this type
+			ensure(UpstreamUserParameterData.Num() == 1);
+#endif
+
+			if (UPCGUserParametersData* UpstreamData = Cast<UPCGUserParametersData>(const_cast<UPCGData*>(UpstreamUserParameterData[0].Data.Get())))
+			{
+				UserParamData->UpstreamData = UpstreamData;
+			}
+		}
+
 		FPCGTaggedData& TaggedData = OutputData.TaggedData.Emplace_GetRef();
 		TaggedData.Data = UserParamData;
 		TaggedData.Tags.Add(PCGBaseSubgraphConstants::UserParameterTagData);
@@ -690,9 +705,10 @@ bool FPCGSubgraphElement::ExecuteInternal(FPCGContext* InContext) const
 					return true;
 				}
 
-				// Prepare the invocation stack - which is the stack up to this node, and then this node
+				// Prepare the invocation stack - which is the stack up to this node, and then this node and the 'not-a-loop index' which we use to differentiate dynamic vs static subgraphs
 				FPCGStack InvocationStack = ensure(Context->Stack) ? *Context->Stack : FPCGStack();
 				InvocationStack.GetStackFramesMutable().Emplace(Context->Node);
+				InvocationStack.GetStackFramesMutable().Emplace(INDEX_NONE); // not a loop index
 
 				// Higen is not allowed in dynamic subgraphs, entire subgraph is executed on the same grid as this subgraph node.
 				FPCGTaskId SubgraphTaskId = Subsystem->ScheduleGraph(
@@ -784,8 +800,7 @@ bool FPCGSubgraphElement::ExecuteInternal(FPCGContext* InContext) const
 	}
 	else
 	{
-		// This node acts as both the pre-graph node and the input node so it should have both the user parameters & the actual inputs
-		PrepareSubgraphData(Settings, Context, Context->InputData, Context->OutputData);
+		// This node acts as the pre-graph node only.
 		PrepareSubgraphUserParameters(Settings, Context, Context->OutputData);
 		return true;
 	}
diff --git a/Engine/Plugins/PCG/Source/PCG/Public/Data/PCGUserParametersData.h b/Engine/Plugins/PCG/Source/PCG/Public/Data/PCGUserParametersData.h
--- a/Engine/Plugins/PCG/Source/PCG/Public/Data/PCGUserParametersData.h
+++ b/Engine/Plugins/PCG/Source/PCG/Public/Data/PCGUserParametersData.h
@@ -30,4 +30,7 @@ public:
 
 	UPROPERTY()
 	FInstancedStruct UserParameters;
+
+	UPROPERTY()
+	TSoftObjectPtr<UPCGUserParametersData> UpstreamData = nullptr;
 };
\ No newline at end of file
diff --git a/Engine/Plugins/PCG/Source/PCG/Public/Elements/PCGUserParameterGet.h b/Engine/Plugins/PCG/Source/PCG/Public/Elements/PCGUserParameterGet.h
--- a/Engine/Plugins/PCG/Source/PCG/Public/Elements/PCGUserParameterGet.h
+++ b/Engine/Plugins/PCG/Source/PCG/Public/Elements/PCGUserParameterGet.h
@@ -9,6 +9,14 @@
 
 #include "PCGUserParameterGet.generated.h"
 
+UENUM()
+enum class EPCGUserParameterSource : uint8
+{
+	Current,
+	Upstream,
+	Root
+};
+
 /**
 * Getter for user parameters defined in PCGGraph, by the user.
 * Will pick up the value from the graph instance.
@@ -73,6 +81,12 @@ public:
 	UPROPERTY(BlueprintReadWrite, EditAnywhere, Category = Settings, meta = (PCG_Overridable))
 	FName OutputAttributeName = NAME_None;
 
+	UPROPERTY(BlueprintReadWrite, EditAnywhere, Category = Settings)
+	EPCGUserParameterSource Source = EPCGUserParameterSource::Current;
+
+	UPROPERTY(BlueprintReadWrite, EditAnywhere, Category = Settings)
+	bool bQuiet = false;
+
 	//~Begin UPCGSettings interface
 	virtual bool RequiresDataFromPreTask() const override { return true; }
 	virtual FString GetAdditionalTitleInformation() const override;
diff --git a/Engine/Plugins/PCG/Source/PCG/Public/Helpers/PCGPropertyHelpers.h b/Engine/Plugins/PCG/Source/PCG/Public/Helpers/PCGPropertyHelpers.h
--- a/Engine/Plugins/PCG/Source/PCG/Public/Helpers/PCGPropertyHelpers.h
+++ b/Engine/Plugins/PCG/Source/PCG/Public/Helpers/PCGPropertyHelpers.h
@@ -90,7 +90,7 @@ namespace PCGPropertyHelpers
 	* @param OptionalContext - Optional context if the extraction is done in a PCG Node, so errors are using the context to log.
 	* @param OptionalObjectTraversed - Optional set to store all objects that we traversed, to be able to react to those objects changes.
 	*/
-	PCG_API UPCGParamData* ExtractPropertyAsAttributeSet(const FExtractorParameters& Parameters, FPCGContext* OptionalContext = nullptr, TSet<FSoftObjectPath>* OptionalObjectTraversed = nullptr);
+	PCG_API UPCGParamData* ExtractPropertyAsAttributeSet(const FExtractorParameters& Parameters, FPCGContext* OptionalContext = nullptr, TSet<FSoftObjectPath>* OptionalObjectTraversed = nullptr, bool bQuiet = false);
 
 	/**
 	* Extract an attribute set in a array of structures. T MUST be a UStruct. Also, it must only contain supported types (so no arrays nor other structures)
diff --git a/Engine/Plugins/PCG/Source/PCG/Public/PCGSettings.h b/Engine/Plugins/PCG/Source/PCG/Public/PCGSettings.h
--- a/Engine/Plugins/PCG/Source/PCG/Public/PCGSettings.h
+++ b/Engine/Plugins/PCG/Source/PCG/Public/PCGSettings.h
@@ -504,6 +504,9 @@ public:
 	/** Check if we have some override. Can be overriden to force params pin for example */
 	virtual bool HasOverridableParams() const { return !CachedOverridableParams.IsEmpty(); }
 
+	/** Checks if a label matches an overridable param */
+	virtual bool HasOverridableParam(FName InParamName) const;
+
 	/** Check if we need to hook the output of the pre-task to this. One use is to compute overrides in the subgraph element and pass the overrides as data, to all nodes that needs it. */
 	virtual bool RequiresDataFromPreTask() const { return false; }
 protected:
diff --git a/Engine/Plugins/PCG/Source/PCG/Public/PCGSubgraph.h b/Engine/Plugins/PCG/Source/PCG/Public/PCGSubgraph.h
--- a/Engine/Plugins/PCG/Source/PCG/Public/PCGSubgraph.h
+++ b/Engine/Plugins/PCG/Source/PCG/Public/PCGSubgraph.h
@@ -47,6 +47,7 @@ protected:
 
 public:
 	//~Begin UPCGSettings interface
+	virtual bool RequiresDataFromPreTask() const override { return true; }
 	virtual bool HasFlippedTitleLines() const override { return true; }
 	// The graph may contain nodes that have side effects, don't assume we can cull even when unwired.
 	// TODO: For static SGs we could probably compute this value based on the subgraph nodes.
