Fixed static analysis warning in the pcg graph compiler

#rnx

diff --git a/Engine/Plugins/PCG/Source/PCG/Private/Graph/PCGGraphCompiler.cpp b/Engine/Plugins/PCG/Source/PCG/Private/Graph/PCGGraphCompiler.cpp
--- a/Engine/Plugins/PCG/Source/PCG/Private/Graph/PCGGraphCompiler.cpp
+++ b/Engine/Plugins/PCG/Source/PCG/Private/Graph/PCGGraphCompiler.cpp
@@ -116,6 +116,8 @@ TArray<FPCGGraphTask> FPCGGraphCompiler::CompileGraph(UPCGGraph* InGraph, FPCGTa
 				return Subtask.Node == SubgraphOutputNode;
 				});
 
+			check(InputNodeTask && OutputNodeTask);
+
 			// Build pre-task
 			FPCGGraphTask& PreTask = CompiledTasks.Emplace_GetRef();
 			PreTask.Node = Node;
@@ -152,10 +154,7 @@ TArray<FPCGGraphTask> FPCGGraphCompiler::CompileGraph(UPCGGraph* InGraph, FPCGTa
 			}
 
 			// Add pre-task as input to subgraph input node task, without data dependency
-			if (InputNodeTask)
-			{
-				InputNodeTask->Inputs.Emplace(PreId,/*InUpstreamPin=*/FPCGGraphTaskInput::NoPin, /*InDownstreamPin=*/FPCGGraphTaskInput::NoPin, /*bInProvideData=*/false);
-			}
+			InputNodeTask->Inputs.Emplace(PreId,/*InUpstreamPin=*/FPCGGraphTaskInput::NoPin, /*InDownstreamPin=*/FPCGGraphTaskInput::NoPin, /*bInProvideData=*/false);
 
 			// Hook nodes to the PreTask if they require so.
 			// Only do it for nodes that are directly under the subgraph, not in subsequent subgraphs.
@@ -196,10 +195,7 @@ TArray<FPCGGraphTask> FPCGGraphCompiler::CompileGraph(UPCGGraph* InGraph, FPCGTa
 			PostTask.Inputs.Emplace(PreId, /*InUpstreamPin=*/FPCGGraphTaskInput::NoPin, /*InDownstreamPin=*/FPCGGraphTaskInput::NoPin, /*bInProvideData=*/false);
 
 			// Add subgraph output node task as input to the post-task
-			if (OutputNodeTask)
-			{
-				PostTask.Inputs.Emplace(OutputNodeTask->NodeId);
-			}
+			PostTask.Inputs.Emplace(OutputNodeTask->NodeId);
 
 			check(!IdMapping.Contains(Node));
 			IdMapping.Add(Node, PostId);
