[PCG] Add option on Attribute partition to assign a partition index instead of partitioning.

* Also fix an issue for overriding selectors, as it was not working.
* Add tests for issue + new feature.

#rnx

diff --git a/Engine/Plugins/PCG/Source/PCG/Private/Elements/Metadata/PCGMetadataPartition.cpp b/Engine/Plugins/PCG/Source/PCG/Private/Elements/Metadata/PCGMetadataPartition.cpp
--- a/Engine/Plugins/PCG/Source/PCG/Private/Elements/Metadata/PCGMetadataPartition.cpp
+++ b/Engine/Plugins/PCG/Source/PCG/Private/Elements/Metadata/PCGMetadataPartition.cpp
@@ -5,8 +5,13 @@
 #include "PCGContext.h"
 #include "PCGCustomVersion.h"
 #include "Helpers/PCGHelpers.h"
+#include "Metadata/PCGMetadata.h"
 #include "Metadata/PCGMetadataAttribute.h"
 #include "Metadata/PCGMetadataPartitionCommon.h"
+#include "Metadata/Accessors/IPCGAttributeAccessor.h"
+#include "Metadata/Accessors/PCGAttributeAccessorHelpers.h"
+#include "Metadata/Accessors/PCGAttributeAccessorKeys.h"
+#include "Utils/PCGLogErrors.h"
 
 #include UE_INLINE_GENERATED_CPP_BY_NAME(PCGMetadataPartition)
 
@@ -105,10 +110,10 @@ bool FPCGMetadataPartitionElement::ExecuteInternal(FPCGContext* Context) const
 	PRAGMA_ENABLE_DEPRECATION_WARNINGS
 
 	// If the names are overridden by the user, generate the selectors with them
-	OverriddenSelectors.SetNum(AttributeNames.Num());
+	OverriddenSelectors.Reserve(AttributeNames.Num());
 	for (const FString& AttributeName : AttributeNames)
 	{
-		OverriddenSelectors.Emplace_GetRef().SetAttributeName(FName(AttributeName));
+		OverriddenSelectors.Emplace_GetRef().Update(AttributeName);
 	}
 
 	const TArray<FPCGAttributePropertyInputSelector>& ActiveSelectors = OverriddenSelectors.IsEmpty() ? Settings->PartitionAttributeSelectors : OverriddenSelectors;
@@ -125,25 +130,89 @@ bool FPCGMetadataPartitionElement::ExecuteInternal(FPCGContext* Context) const
 		return true;
 	}
 
+	const bool bAssignIndexPartition = Settings->bAssignIndexPartition;
+	// bDoNotPartition is only exposed to the user if bAssignIndexPartition is true, so make sure that this condition is respected
+	// as no partition and no assign index is a trivial operation.
+	const bool bDoPartition = !bAssignIndexPartition || !Settings->bDoNotPartition;
+
 	for (const FPCGTaggedData& Input : Inputs)
 	{
 		const UPCGData* InData = Input.Data;
 
-		for (int I = 0; I < PartitionAttributeSources.Num(); ++I)
+		for (int32 i = 0; i < PartitionAttributeSources.Num(); ++i)
 		{
-			PartitionAttributeSources[I] = static_cast<FPCGAttributePropertySelector>(ActiveSelectors[I].CopyAndFixLast(InData));
+			PartitionAttributeSources[i] = static_cast<FPCGAttributePropertySelector>(ActiveSelectors[i].CopyAndFixLast(InData));
 		}
 
-		TArray<UPCGData*> PartitionDataArray;
-		PartitionDataArray = PCGMetadataPartitionCommon::AttributePartition(InData, PartitionAttributeSources, Context);
-
-		for (UPCGData* PartitionData : PartitionDataArray)
+		if (bDoPartition)
 		{
-			if (PartitionData)
+			TArray<UPCGData*> PartitionDataArray = PCGMetadataPartitionCommon::AttributePartition(InData, PartitionAttributeSources, Context);
+
+			for (int32 i = 0; i < PartitionDataArray.Num(); ++i)
 			{
+				UPCGData* PartitionData = PartitionDataArray[i];
+				// By construction, it must be valid
+				check(PartitionData);
+
 				FPCGTaggedData& Output = Outputs.Add_GetRef(Input);
 				Output.Data = PartitionData;
+
+				if (bAssignIndexPartition)
+				{
+					UPCGMetadata* Metadata = PartitionData->MutableMetadata();
+					if (ensure(Metadata))
+					{
+						Metadata->CreateAttribute<int32>(Settings->PartitionIndexAttributeName, i, /*bAllowInterpolation=*/false, /*bOverrideParent=*/true);
+					}
+				}
+			}
+		}
+		else
+		{
+			const TArray<TArray<int32>> PartitionDataArray = PCGMetadataPartitionCommon::AttributeGenericPartition(InData, PartitionAttributeSources, Context);
+			FPCGTaggedData& Output = Outputs.Add_GetRef(Input);
+
+			if (PartitionDataArray.IsEmpty())
+			{
+				continue;
 			}
+
+			UPCGData* OutputData = InData->DuplicateData(Context);
+			check(OutputData);
+			UPCGMetadata* Metadata = OutputData->MutableMetadata();
+			check(Metadata);
+
+			FPCGMetadataAttributeBase* AttributePartitionIndex = Metadata->CreateAttribute<int32>(Settings->PartitionIndexAttributeName, 0, /*bAllowInterpolation=*/false, /*bOverrideParent=*/true);
+			if (!AttributePartitionIndex)
+			{
+				PCGLog::Metadata::LogFailToCreateAttribute<int32>(Settings->PartitionIndexAttributeName, Context);
+				continue;
+			}
+
+			const FPCGAttributePropertySelector Selector = FPCGAttributePropertySelector::CreateAttributeSelector(Settings->PartitionIndexAttributeName);
+			TUniquePtr<IPCGAttributeAccessor> Accessor = PCGAttributeAccessorHelpers::CreateAccessor(OutputData, Selector);
+			TUniquePtr<IPCGAttributeAccessorKeys> Keys = PCGAttributeAccessorHelpers::CreateKeys(OutputData, Selector);
+
+			// Attribute exists, so Accessor and keys should exist
+			if (!ensure(Accessor && Keys))
+			{
+				continue;
+			}
+
+			// Flatten the array to write it in one go.
+			TArray<int32> AllIndexes;
+			AllIndexes.SetNumUninitialized(Keys->GetNum());
+
+			for (int32 i = 0; i < PartitionDataArray.Num(); ++i)
+			{
+				for (int32 j : PartitionDataArray[i])
+				{
+					AllIndexes[j] = i;
+				}
+			}
+
+			Accessor->SetRange<int32>(AllIndexes, 0, *Keys);
+			Output.Data = OutputData;
 		}
 	}
 
diff --git a/Engine/Plugins/PCG/Source/PCG/Private/Tests/Elements/PCGAttributePartitionTest.cpp b/Engine/Plugins/PCG/Source/PCG/Private/Tests/Elements/PCGAttributePartitionTest.cpp
--- a/Engine/Plugins/PCG/Source/PCG/Private/Tests/Elements/PCGAttributePartitionTest.cpp
+++ b/Engine/Plugins/PCG/Source/PCG/Private/Tests/Elements/PCGAttributePartitionTest.cpp
@@ -12,51 +12,180 @@
 IMPLEMENT_CUSTOM_SIMPLE_AUTOMATION_TEST(FPCGAttributePartition_Points, FPCGTestBaseClass, "Plugins.PCG.AttributePartition.Points", PCGTestsCommon::TestFlags)
 IMPLEMENT_CUSTOM_SIMPLE_AUTOMATION_TEST(FPCGAttributePartition_AttributeSet, FPCGTestBaseClass, "Plugins.PCG.AttributePartition.AttributeSet", PCGTestsCommon::TestFlags)
 IMPLEMENT_CUSTOM_SIMPLE_AUTOMATION_TEST(FPCGAttributePartition_Order, FPCGTestBaseClass, "Plugins.PCG.AttributePartition.Order", PCGTestsCommon::TestFlags)
+IMPLEMENT_CUSTOM_SIMPLE_AUTOMATION_TEST(FPCGAttributePartition_MultiPartition, FPCGTestBaseClass, "Plugins.PCG.AttributePartition.MultiPartition", PCGTestsCommon::TestFlags)
+IMPLEMENT_CUSTOM_SIMPLE_AUTOMATION_TEST(FPCGAttributePartition_MultiPartitionOverride, FPCGTestBaseClass, "Plugins.PCG.AttributePartition.MultiPartitionOverride", PCGTestsCommon::TestFlags)
+IMPLEMENT_CUSTOM_SIMPLE_AUTOMATION_TEST(FPCGAttributePartition_WithPartitionIndex, FPCGTestBaseClass, "Plugins.PCG.AttributePartition.WithPartitionIndex", PCGTestsCommon::TestFlags)
+IMPLEMENT_CUSTOM_SIMPLE_AUTOMATION_TEST(FPCGAttributePartition_NoPartitionWithPartitionIndex, FPCGTestBaseClass, "Plugins.PCG.AttributePartition.NoPartitionWithPartitionIndex", PCGTestsCommon::TestFlags)
 
-bool FPCGAttributePartition_Points::RunTest(const FString& Parameters)
+namespace PCGAttributePartitionTest
 {
-	PCGTestsCommon::FTestData TestData;
-	UPCGMetadataPartitionSettings* Settings = PCGTestsCommon::GenerateSettings<UPCGMetadataPartitionSettings>(TestData);
-	check(Settings);
+	const static FName PartitionIndexAttributeName = TEXT("PartitionIndex");
 
-	// By default there should be one selector with @Last attribute
-	UTEST_EQUAL("There is one Partition Attribute Selector by default", Settings->PartitionAttributeSelectors.Num(), 1);
+	bool MultiPartitionTest(FPCGTestBaseClass& TestClass, bool bWithOverride)
+	{
+		PCGTestsCommon::FTestData TestData;
+		UPCGMetadataPartitionSettings* Settings = PCGTestsCommon::GenerateSettings<UPCGMetadataPartitionSettings>(TestData);
+		check(Settings);
 
-	Settings->PartitionAttributeSelectors[0].SetPointProperty(EPCGPointProperties::Density);
+		if (!bWithOverride)
+		{
+			Settings->PartitionAttributeSelectors.Empty(2);
+			Settings->PartitionAttributeSelectors.Emplace_GetRef().Update(TEXT("$Position.X"));
+			Settings->PartitionAttributeSelectors.Emplace_GetRef().Update(TEXT("$Position.Y"));
+		}
+		else
+		{
+			UPCGParamData* OverrideParamData = NewObject<UPCGParamData>();
+			FPCGMetadataAttribute<FString>* OverrideAttribute = OverrideParamData->Metadata->CreateAttribute<FString>(TEXT(""), TEXT("$Position.X,$Position.Y"), false, false);
+			OverrideParamData->Metadata->AddEntry();
 
-	UPCGPointData* InputPointData = NewObject<UPCGPointData>();
-	TArray<FPCGPoint>& Points = InputPointData->GetMutablePoints();
-	Points.Reserve(100);
-	for (int32 i = 0; i < 100; ++i)
-	{
-		FPCGPoint& Point = Points.Emplace_GetRef();
-		Point.Density = (i % 10) / 10.0f;
-	}
+			FPCGTaggedData& OverrideTaggedData = TestData.InputData.TaggedData.Emplace_GetRef();
+			OverrideTaggedData.Data = OverrideParamData;
+			OverrideTaggedData.Pin = GET_MEMBER_NAME_CHECKED(UPCGMetadataPartitionSettings, PartitionAttributeNames);
+		}
 
-	FPCGTaggedData& InputTaggedData = TestData.InputData.TaggedData.Emplace_GetRef();
-	InputTaggedData.Data = InputPointData;
-	InputTaggedData.Pin = PCGPinConstants::DefaultInputLabel;
+		UPCGPointData* InputPointData = NewObject<UPCGPointData>();
+		TArray<FPCGPoint>& Points = InputPointData->GetMutablePoints();
+		Points.Reserve(20);
+		for (int32 i = 0; i < 20; ++i)
+		{
+			FPCGPoint& Point = Points.Emplace_GetRef();
+			// Also add a very small number on Y, to verify that approximation errors are not giving different partitions.
+			const double Jitter = i >= 10 ? UE_DOUBLE_SMALL_NUMBER : 0.0;
+			Point.Transform.SetLocation(FVector(i % 10, i % 5 + Jitter, 0.0));
+		}
 
-	FPCGElementPtr TestElement = TestData.Settings->GetElement();
-	TUniquePtr<FPCGContext> Context = TestData.InitializeTestContext();
+		FPCGTaggedData& InputTaggedData = TestData.InputData.TaggedData.Emplace_GetRef();
+		InputTaggedData.Data = InputPointData;
+		InputTaggedData.Pin = PCGPinConstants::DefaultInputLabel;
 
-	while (!TestElement->Execute(Context.Get())) {}
+		FPCGElementPtr TestElement = TestData.Settings->GetElement();
+		TUniquePtr<FPCGContext> Context = TestData.InitializeTestContext();
 
-	UTEST_EQUAL("There are 10 outputs", Context->OutputData.TaggedData.Num(), 10);
+		while (!TestElement->Execute(Context.Get())) {}
 
-	for (int32 i = 0; i < Context->OutputData.TaggedData.Num(); ++i)
+		if (!TestClass.TestEqual(TEXT("There are 10 outputs"), Context->OutputData.TaggedData.Num(), 10))
+		{
+			return false;
+		}
+
+		for (int32 i = 0; i < Context->OutputData.TaggedData.Num(); ++i)
+		{
+			const UPCGPointData* OutputPointData = Cast<const UPCGPointData>(Context->OutputData.TaggedData[i].Data);
+			if (!TestClass.TestNotNull(*FString::Printf(TEXT("Output %d is a point data"), i), OutputPointData))
+			{
+				return false;
+			}
+
+			const TArray<FPCGPoint>& OutputPoints = OutputPointData->GetPoints();
+
+			if (!TestClass.TestEqual(*FString::Printf(TEXT("Output %d has 2 points"), i), OutputPoints.Num(), 2))
+			{
+				return false;
+			}
+
+			FVector TempValue;
+			bool bFirstPoint = true;
+			bool bAllEquals = true;
+			for (const FPCGPoint& Point : OutputPoints)
+			{
+				if (bFirstPoint)
+				{
+					TempValue = Point.Transform.GetLocation();
+					bFirstPoint = false;
+				}
+				else
+				{
+					bAllEquals &= TempValue.Equals(Point.Transform.GetLocation());
+				}
+			}
+
+			if (!TestClass.TestTrue(*FString::Printf(TEXT("Output points for output %d have all the same position"), i), bAllEquals))
+			{
+				return false;
+			}
+		}
+
+		return true;
+	}
+
+	bool PointPartitionTest(FPCGTestBaseClass& TestClass, bool bWithPartitionIndex, bool bNoPartition, auto VerifyFunction)
 	{
-		const UPCGPointData* OutputPointData = Cast<const UPCGPointData>(Context->OutputData.TaggedData[i].Data);
-		UTEST_NOT_NULL(*FString::Printf(TEXT("Output %d is a point data"), i), OutputPointData);
+		PCGTestsCommon::FTestData TestData;
+		UPCGMetadataPartitionSettings* Settings = PCGTestsCommon::GenerateSettings<UPCGMetadataPartitionSettings>(TestData);
+		check(Settings);
 
-		const TArray<FPCGPoint>& OutputPoints = OutputPointData->GetPoints();
+		// By default there should be one selector with @Last attribute
+		if (!TestClass.TestEqual(TEXT("There is one Partition Attribute Selector by default"), Settings->PartitionAttributeSelectors.Num(), 1))
+		{
+			return false;
+		}
+
+		Settings->PartitionAttributeSelectors[0].SetPointProperty(EPCGPointProperties::Density);
+		Settings->bAssignIndexPartition = bWithPartitionIndex;
+		Settings->bDoNotPartition = bNoPartition;
+		Settings->PartitionIndexAttributeName = PartitionIndexAttributeName;
+
+		UPCGPointData* InputPointData = NewObject<UPCGPointData>();
+		TArray<FPCGPoint>& Points = InputPointData->GetMutablePoints();
+		Points.Reserve(100);
+		for (int32 i = 0; i < 100; ++i)
+		{
+			FPCGPoint& Point = Points.Emplace_GetRef();
+			Point.Density = i % 10;
+		}
+
+		const int32 ExpectedPointsPerData = bNoPartition ? 100 : 10;
+		const int32 ExpectedDataNum = bNoPartition ? 1 : 10;
+
+		FPCGTaggedData& InputTaggedData = TestData.InputData.TaggedData.Emplace_GetRef();
+		InputTaggedData.Data = InputPointData;
+		InputTaggedData.Pin = PCGPinConstants::DefaultInputLabel;
+
+		FPCGElementPtr TestElement = TestData.Settings->GetElement();
+		TUniquePtr<FPCGContext> Context = TestData.InitializeTestContext();
+
+		while (!TestElement->Execute(Context.Get())) {}
+
+		if (!TestClass.TestEqual(TEXT("There are the right number outputs"), Context->OutputData.TaggedData.Num(), ExpectedDataNum))
+		{
+			return false;
+		}
 
-		UTEST_EQUAL(*FString::Printf(TEXT("Output %d has 10 points"), i), OutputPoints.Num(), 10);
+		for (int32 i = 0; i < Context->OutputData.TaggedData.Num(); ++i)
+		{
+			const UPCGPointData* OutputPointData = Cast<const UPCGPointData>(Context->OutputData.TaggedData[i].Data);
+			if (!TestClass.TestNotNull(*FString::Printf(TEXT("Output %d is a point data"), i), OutputPointData))
+			{
+				return false;
+			}
+
+			const TArray<FPCGPoint>& OutputPoints = OutputPointData->GetPoints();
 
+			if (!TestClass.TestEqual(*FString::Printf(TEXT("Output %d has the right number of points"), i), OutputPoints.Num(), ExpectedPointsPerData))
+			{
+				return false;
+			}
+
+			if (!VerifyFunction(OutputPointData, i))
+			{
+				return false;
+			}
+		}
+
+		return true;
+	}
+}
+
+bool FPCGAttributePartition_Points::RunTest(const FString& Parameters)
+{
+	return PCGAttributePartitionTest::PointPartitionTest(*this, /*bWithPartitionIndex=*/ false, /*bNoPartition=*/ false, [this](const UPCGPointData* OutputPointData, int32 OutputIndex) -> bool
+	{
+		check(OutputPointData);
 		float TempValue = -1.0f;
 		bool bFirstPoint = true;
 		bool bAllEquals = true;
-		for (const FPCGPoint& Point : OutputPoints)
+		for (const FPCGPoint& Point : OutputPointData->GetPoints())
 		{
 			if (bFirstPoint)
 			{
@@ -69,10 +198,8 @@ bool FPCGAttributePartition_Points::RunTest(const FString& Parameters)
 			}
 		}
 
-		UTEST_TRUE(*FString::Printf(TEXT("Output points for output %d have all the same density"), i), bAllEquals);
-	}
-
-	return true;
+		return TestTrue(*FString::Printf(TEXT("Output points for output %d have all the same density"), OutputIndex), bAllEquals);
+	});
 }
 
 /**
@@ -212,4 +339,64 @@ bool FPCGAttributePartition_AttributeSet::RunTest(const FString& Parameters)
 	}
 
 	return true;
+}
+
+bool FPCGAttributePartition_MultiPartition::RunTest(const FString& Parameters)
+{
+	return PCGAttributePartitionTest::MultiPartitionTest(*this, /*bWithOverride=*/false);
+}
+
+bool FPCGAttributePartition_MultiPartitionOverride::RunTest(const FString& Parameters)
+{
+	return PCGAttributePartitionTest::MultiPartitionTest(*this, /*bWithOverride=*/true);
+}
+
+bool FPCGAttributePartition_WithPartitionIndex::RunTest(const FString& Parameters)
+{
+	return PCGAttributePartitionTest::PointPartitionTest(*this, /*bWithPartitionIndex=*/ true, /*bNoPartition=*/ false, [this](const UPCGPointData* OutputPointData, int32 OutputIndex) -> bool
+		{
+			check(OutputPointData && OutputPointData->Metadata);
+			const FPCGMetadataAttribute<int32>* PartitionIndexAttribute = OutputPointData->Metadata->GetConstTypedAttribute<int32>(PCGAttributePartitionTest::PartitionIndexAttributeName);
+
+			UTEST_NOT_NULL("Partition Index attribute exists", PartitionIndexAttribute);
+
+			float TempValue = -1.0f;
+			bool bFirstPoint = true;
+			bool bAllEquals = true;
+			for (const FPCGPoint& Point : OutputPointData->GetPoints())
+			{
+				if (bFirstPoint)
+				{
+					TempValue = Point.Density;
+					bFirstPoint = false;
+				}
+				else
+				{
+					bAllEquals &= TempValue == Point.Density;
+				}
+
+				UTEST_EQUAL("Partition index value is equal to output index", PartitionIndexAttribute->GetValueFromItemKey(Point.MetadataEntry), OutputIndex);
+			}
+
+			return TestTrue(*FString::Printf(TEXT("Output points for output %d have all the same density"), OutputIndex), bAllEquals);
+		});
+}
+
+bool FPCGAttributePartition_NoPartitionWithPartitionIndex::RunTest(const FString& Parameters)
+{
+	return PCGAttributePartitionTest::PointPartitionTest(*this, /*bWithPartitionIndex=*/ true, /*bNoPartition=*/ true, [this](const UPCGPointData* OutputPointData, int32 OutputIndex) -> bool
+		{
+			check(OutputPointData && OutputPointData->Metadata);
+			const FPCGMetadataAttribute<int32>* PartitionIndexAttribute = OutputPointData->Metadata->GetConstTypedAttribute<int32>(PCGAttributePartitionTest::PartitionIndexAttributeName);
+
+			UTEST_NOT_NULL("Partition Index attribute exists", PartitionIndexAttribute);
+
+			for (const FPCGPoint& Point : OutputPointData->GetPoints())
+			{
+				const int32 PartitionIndex = FMath::FloorToInt32(Point.Density);
+				UTEST_EQUAL("Partition index has the right value", PartitionIndexAttribute->GetValueFromItemKey(Point.MetadataEntry), PartitionIndex);
+			}
+
+			return true;
+		});
 }
\ No newline at end of file
diff --git a/Engine/Plugins/PCG/Source/PCG/Public/Elements/Metadata/PCGMetadataPartition.h b/Engine/Plugins/PCG/Source/PCG/Public/Elements/Metadata/PCGMetadataPartition.h
--- a/Engine/Plugins/PCG/Source/PCG/Public/Elements/Metadata/PCGMetadataPartition.h
+++ b/Engine/Plugins/PCG/Source/PCG/Public/Elements/Metadata/PCGMetadataPartition.h
@@ -49,6 +49,17 @@ public:
 	UPROPERTY(EditAnywhere, Category = Settings, meta = (EditCondition = "bTokenizeOnWhiteSpace", EditConditionHides, DeprecationMessage = "bTokenizeOnWhiteSpace has been deprecated."))
 	bool bTokenizeOnWhiteSpace = false;
 
+	/** Assign an index partition as an extra attribute. */
+	UPROPERTY(EditAnywhere, Category = Settings, meta = (PCG_Overridable))
+	bool bAssignIndexPartition = false;
+
+	/** If we assign an index, we can also not partition (and only assign the partition index to the original data). */
+	UPROPERTY(EditAnywhere, Category = Settings, meta = (PCG_Overridable, EditCondition = "bAssignIndexPartition", EditConditionHides))
+	bool bDoNotPartition = true;
+
+	UPROPERTY(EditAnywhere, Category = Settings, meta = (PCG_Overridable, EditCondition = "bAssignIndexPartition", EditConditionHides))
+	FName PartitionIndexAttributeName = TEXT("PartitionIndex");
+
 #if WITH_EDITORONLY_DATA
 	UE_DEPRECATED(5.5, "PartitionAttribute has been deprecated.")
 	UPROPERTY(meta = (DeprecatedProperty, DeprecationMessage = "PartitionAttribute has been deprecated."))
