[PCG] Add Remove Duplicates node

#rnx

diff --git a/Engine/Plugins/PCG/Source/PCG/Private/Elements/PCGAttributeRemoveDuplicates.cpp b/Engine/Plugins/PCG/Source/PCG/Private/Elements/PCGAttributeRemoveDuplicates.cpp
new file mode 100644
--- /dev/null
+++ b/Engine/Plugins/PCG/Source/PCG/Private/Elements/PCGAttributeRemoveDuplicates.cpp
@@ -0,0 +1,102 @@
+// Copyright Epic Games, Inc. All Rights Reserved.
+
+#include "Elements/PCGAttributeRemoveDuplicates.h"
+
+#include "PCGContext.h"
+#include "Helpers/PCGHelpers.h"
+#include "Metadata/PCGMetadataPartitionCommon.h"
+
+#define LOCTEXT_NAMESPACE "PCGAttributeRemoveDuplicatesElement"
+
+TArray<FPCGPinProperties> UPCGAttributeRemoveDuplicatesSettings::OutputPinProperties() const
+{
+	TArray<FPCGPinProperties> PinProperties;
+	PinProperties.Emplace(PCGPinConstants::DefaultOutputLabel, EPCGDataType::Any);
+
+	return PinProperties;
+}
+
+FPCGElementPtr UPCGAttributeRemoveDuplicatesSettings::CreateElement() const
+{
+	return MakeShared<FPCGAttributeRemoveDuplicatesElement>();
+}
+
+FString UPCGAttributeRemoveDuplicatesSettings::GetAdditionalTitleInformation() const
+{
+#if WITH_EDITOR
+	if (IsPropertyOverriddenByPin(GET_MEMBER_NAME_CHECKED(UPCGAttributeRemoveDuplicatesSettings, AttributeNamesToRemoveDuplicates)))
+	{
+		return FString();
+	}
+	else
+#endif
+	{
+		if (AttributeSelectors.IsEmpty())
+		{
+			return FString();
+		}
+
+		FString OutString = AttributeSelectors[0].GetDisplayText().ToString();
+		for (int i = 1; i < AttributeSelectors.Num(); ++i)
+		{
+			OutString += ", " + AttributeSelectors[i].GetDisplayText().ToString();
+		}
+
+		return OutString;
+	}
+}
+
+bool FPCGAttributeRemoveDuplicatesElement::ExecuteInternal(FPCGContext* Context) const
+{
+	TRACE_CPUPROFILER_EVENT_SCOPE(FPCGMetadataPartitionElement::Execute);
+	check(Context);
+
+	const UPCGAttributeRemoveDuplicatesSettings* Settings = Context->GetInputSettings<UPCGAttributeRemoveDuplicatesSettings>();
+	check(Settings);
+
+	// TODO: This is a temporary solution for overrides until arrays are supported
+	TArray<FPCGAttributePropertyInputSelector> OverriddenSelectors;
+
+	const TArray<FString> AttributeNames = PCGHelpers::GetStringArrayFromCommaSeparatedList(Settings->AttributeNamesToRemoveDuplicates);
+
+	// If the names are overridden by the user, generate the selectors with them
+	OverriddenSelectors.Reserve(AttributeNames.Num());
+	for (const FString& AttributeName : AttributeNames)
+	{
+		OverriddenSelectors.Emplace_GetRef().Update(AttributeName);
+	}
+
+	const TArray<FPCGAttributePropertyInputSelector>& ActiveSelectors = OverriddenSelectors.IsEmpty() ? Settings->AttributeSelectors : OverriddenSelectors;
+
+	TArray<FPCGTaggedData> Inputs = Context->InputData.GetInputsByPin(PCGPinConstants::DefaultInputLabel);
+	TArray<FPCGTaggedData>& Outputs = Context->OutputData.TaggedData;
+
+	TArray<FPCGAttributePropertySelector> PartitionAttributeSources;
+	PartitionAttributeSources.SetNum(ActiveSelectors.Num());
+
+	if (PartitionAttributeSources.IsEmpty())
+	{
+		Outputs = Inputs;
+		return true;
+	}
+
+	for (const FPCGTaggedData& Input : Inputs)
+	{
+		const UPCGData* InData = Input.Data;
+
+		for (int32 i = 0; i < PartitionAttributeSources.Num(); ++i)
+		{
+			PartitionAttributeSources[i] = static_cast<FPCGAttributePropertySelector>(ActiveSelectors[i].CopyAndFixLast(InData));
+		}
+
+		FPCGTaggedData& Output = Outputs.Add_GetRef(Input);
+		if (UPCGData* DataWithoutDuplicates = PCGMetadataPartitionCommon::RemoveDuplicates(InData, PartitionAttributeSources, Context))
+		{
+			Output.Data = DataWithoutDuplicates;
+		}
+	}
+
+	return true;
+}
+
+#undef LOCTEXT_NAMESPACE
diff --git a/Engine/Plugins/PCG/Source/PCG/Private/Metadata/PCGMetadataPartitionCommon.cpp b/Engine/Plugins/PCG/Source/PCG/Private/Metadata/PCGMetadataPartitionCommon.cpp
--- a/Engine/Plugins/PCG/Source/PCG/Private/Metadata/PCGMetadataPartitionCommon.cpp
+++ b/Engine/Plugins/PCG/Source/PCG/Private/Metadata/PCGMetadataPartitionCommon.cpp
@@ -453,6 +453,40 @@ namespace PCGMetadataPartitionCommon
 		return PartitionedData;
 	}
 
+	UPCGData* RemoveDuplicatesPoint(const UPCGPointData* InData, const TArrayView<const FPCGAttributePropertySelector>& InSelectorArrayView, FPCGContext* InOptionalContext, bool bSilenceMissingAttributeErrors)
+	{
+		TArray<TArray<int32>> Partition = AttributeGenericPartition(InData, InSelectorArrayView, InOptionalContext, bSilenceMissingAttributeErrors);
+		if (Partition.IsEmpty())
+		{
+			return nullptr;
+		}
+
+		const TArray<FPCGPoint>& Points = InData->GetPoints();
+
+		UPCGPointData* OutputPointData = nullptr;
+
+		for (TArray<int32>& Indices : Partition)
+		{
+			if (Indices.IsEmpty())
+			{
+				continue;
+			}
+
+			if (!OutputPointData)
+			{
+				OutputPointData = FPCGContext::NewObject_AnyThread<UPCGPointData>(InOptionalContext);
+				OutputPointData->InitializeFromData(InData);
+
+				OutputPointData->GetMutablePoints().Reserve(Partition.Num());
+			}
+
+			check(OutputPointData);
+			OutputPointData->GetMutablePoints().Add(Points[Indices[0]]);
+		}
+
+		return OutputPointData;
+	}
+
 	TArray<UPCGData*> AttributeParamSpatialPartition(const UPCGData* InData, const TArrayView<const FPCGAttributePropertySelector>& InSelectorArray, FPCGContext* InOptionalContext, bool bSilenceMissingAttributeErrors)
 	{
 		if (!InData->IsA<UPCGSpatialData>() && !InData->IsA<UPCGParamData>())
@@ -530,17 +564,81 @@ namespace PCGMetadataPartitionCommon
 		return PartitionedData;
 	}
 
-	TArray<UPCGData*> AttributePartition(const UPCGData* InData, const FPCGAttributePropertySelector& InSelector, FPCGContext* InOptionalContext, bool bSilenceMissingAttributeErrors)
+	UPCGData* RemoveDuplicatesParamSpatial(const UPCGData* InData, const TArrayView<const FPCGAttributePropertySelector>& InSelectorArray, FPCGContext* InOptionalContext, bool bSilenceMissingAttributeErrors)
 	{
-		const TArrayView<const FPCGAttributePropertySelector> ArrayView(&InSelector, 1);
-		if (const UPCGPointData* InPointData = Cast<UPCGPointData>(InData))
+		if (!InData->IsA<UPCGSpatialData>() && !InData->IsA<UPCGParamData>())
 		{
-			return AttributePointPartition(InPointData, ArrayView, InOptionalContext, bSilenceMissingAttributeErrors);
+			PCGLog::LogErrorOnGraph(LOCTEXT("InvalidDataType", "Input data is not an attribute set nor a spatial data. Operation not supported."), InOptionalContext);
+			return nullptr;
 		}
-		else
+
+		const TArray<TArray<int32>> Partition = AttributeGenericPartition(InData, InSelectorArray, InOptionalContext, bSilenceMissingAttributeErrors);
+
+		if (Partition.IsEmpty())
 		{
-			return AttributeParamSpatialPartition(InData, ArrayView, InOptionalContext, bSilenceMissingAttributeErrors);
+			return nullptr;
 		}
+
+		const UPCGSpatialData* InSpatialData = Cast<const UPCGSpatialData>(InData);
+
+		UPCGData* OutputData = nullptr;
+		UPCGMetadata* OutMetadata = nullptr;
+
+		TArray<FName> AttributeNames;
+		TArray<EPCGMetadataTypes> AttributeTypes;
+		const UPCGMetadata* OriginalMetadata = InData->ConstMetadata();
+		OriginalMetadata->GetAttributes(AttributeNames, AttributeTypes);
+
+		for (const TArray<int32>& Indices : Partition)
+		{
+			if (Indices.IsEmpty())
+			{
+				continue;
+			}
+
+			if (!OutputData)
+			{
+				if (InSpatialData)
+				{
+					UPCGSpatialData* NewData = FPCGContext::NewObject_AnyThread<UPCGSpatialData>(InOptionalContext);
+					NewData->InitializeFromData(InSpatialData);
+					OutputData = NewData;
+					OutMetadata = NewData->Metadata;
+				}
+				else
+				{
+					UPCGParamData* NewData = FPCGContext::NewObject_AnyThread<UPCGParamData>(InOptionalContext);
+					NewData->Metadata->AddAttributes(OriginalMetadata);
+					OutputData = NewData;
+					OutMetadata = NewData->Metadata;
+				}
+			}
+
+			check(OutMetadata);
+
+			const PCGMetadataEntryKey CurrentEntryKey = OutMetadata->AddEntry();
+
+			for (const FName AttributeName : AttributeNames)
+			{
+				const FPCGMetadataAttributeBase* OriginalAttribute = OriginalMetadata->GetConstAttribute(AttributeName);
+				FPCGMetadataAttributeBase* NewAttribute = OutMetadata->GetMutableAttribute(AttributeName);
+				check(OriginalAttribute && NewAttribute);
+
+				NewAttribute->SetValue(CurrentEntryKey, OriginalAttribute, Indices[0]);
+
+				if (Partition.Num() == 1)
+				{
+					NewAttribute->SetDefaultValueToFirstEntry();
+				}
+			}
+		}
+
+		return OutputData;
+	}
+
+	TArray<UPCGData*> AttributePartition(const UPCGData* InData, const FPCGAttributePropertySelector& InSelector, FPCGContext* InOptionalContext, bool bSilenceMissingAttributeErrors)
+	{
+		return AttributePartition(InData, TArrayView<const FPCGAttributePropertySelector>(&InSelector, 1), InOptionalContext, bSilenceMissingAttributeErrors);
 	}
 
 	TArray<UPCGData*> AttributePartition(const UPCGData* InData, const TArrayView<const FPCGAttributePropertySelector>& InSelectorArrayView, FPCGContext* InOptionalContext, bool bSilenceMissingAttributeErrors)
@@ -554,6 +652,23 @@ namespace PCGMetadataPartitionCommon
 			return AttributeParamSpatialPartition(InData, InSelectorArrayView, InOptionalContext, bSilenceMissingAttributeErrors);
 		}
 	}
+
+	UPCGData* RemoveDuplicates(const UPCGData* InData, const FPCGAttributePropertySelector& InSelector, FPCGContext* InOptionalContext, bool bSilenceMissingAttributeErrors)
+	{
+		return RemoveDuplicates(InData, TArrayView<const FPCGAttributePropertySelector>(&InSelector, 1), InOptionalContext, bSilenceMissingAttributeErrors);
+	}
+
+	UPCGData* RemoveDuplicates(const UPCGData* InData, const TArrayView<const FPCGAttributePropertySelector>& InSelectorArrayView, FPCGContext* InOptionalContext, bool bSilenceMissingAttributeErrors)
+	{
+		if (const UPCGPointData* InPointData = Cast<UPCGPointData>(InData))
+		{
+			return RemoveDuplicatesPoint(InPointData, InSelectorArrayView, InOptionalContext, bSilenceMissingAttributeErrors);
+		}
+		else
+		{
+			return RemoveDuplicatesParamSpatial(InData, InSelectorArrayView, InOptionalContext, bSilenceMissingAttributeErrors);
+		}
+	}
 }
 
 #undef LOCTEXT_NAMESPACE
\ No newline at end of file
diff --git a/Engine/Plugins/PCG/Source/PCG/Public/Elements/PCGAttributeRemoveDuplicates.h b/Engine/Plugins/PCG/Source/PCG/Public/Elements/PCGAttributeRemoveDuplicates.h
new file mode 100644
--- /dev/null
+++ b/Engine/Plugins/PCG/Source/PCG/Public/Elements/PCGAttributeRemoveDuplicates.h
@@ -0,0 +1,49 @@
+// Copyright Epic Games, Inc. All Rights Reserved.
+
+#pragma once
+
+#include "PCGSettings.h"
+
+#include "PCGAttributeRemoveDuplicates.generated.h"
+
+/**
+* Remove duplicates for given attributes
+*/
+UCLASS(MinimalAPI, BlueprintType, ClassGroup = (Procedural))
+class UPCGAttributeRemoveDuplicatesSettings : public UPCGSettings
+{
+	GENERATED_BODY()
+
+public:
+	//~Begin UPCGSettings interface
+#if WITH_EDITOR
+	virtual FName GetDefaultNodeName() const override { return TEXT("AttributeRemoveDuplicates"); }
+	virtual FText GetDefaultNodeTitle() const override { return NSLOCTEXT("PCGAttributeRemoveDuplicatesElement", "NodeTitle", "Attribute Remove Duplicates"); }
+	virtual EPCGSettingsType GetType() const override { return EPCGSettingsType::Metadata; }
+	virtual bool HasDynamicPins() const override { return true; }
+#endif
+
+	virtual FString GetAdditionalTitleInformation() const override;
+
+protected:
+	virtual TArray<FPCGPinProperties> OutputPinProperties() const override;
+	virtual FPCGElementPtr CreateElement() const override;
+	//~End UPCGSettings interface
+
+public:
+	// TODO: Should be overridable once array override is supported
+	/** The data will be partitioned on these selected attributes, and we will only keep the first entry for each partition. */
+	UPROPERTY(BlueprintReadWrite, EditAnywhere, Category = Settings)
+	TArray<FPCGAttributePropertyInputSelector> AttributeSelectors = { FPCGAttributePropertyInputSelector() };
+
+	// TODO: Should be deprecated once array override is supported
+	UPROPERTY(meta = (PCG_Overridable))
+	FString AttributeNamesToRemoveDuplicates;
+};
+
+class FPCGAttributeRemoveDuplicatesElement : public IPCGElement
+{
+protected:
+	virtual bool ExecuteInternal(FPCGContext* Context) const override;
+	virtual EPCGElementExecutionLoopMode ExecutionLoopMode(const UPCGSettings* Settings) const override { return EPCGElementExecutionLoopMode::SinglePrimaryPin; }
+};
diff --git a/Engine/Plugins/PCG/Source/PCG/Public/Metadata/PCGMetadataPartitionCommon.h b/Engine/Plugins/PCG/Source/PCG/Public/Metadata/PCGMetadataPartitionCommon.h
--- a/Engine/Plugins/PCG/Source/PCG/Public/Metadata/PCGMetadataPartitionCommon.h
+++ b/Engine/Plugins/PCG/Source/PCG/Public/Metadata/PCGMetadataPartitionCommon.h
@@ -54,4 +54,26 @@ namespace PCGMetadataPartitionCommon
 	* @returns Array of bucket indices.
 	*/
 	PCG_API TArray<TArray<int32>> AttributeGenericPartition(const UPCGData* InData, const TArrayView<const FPCGAttributePropertySelector>& InSelectorArrayView, FPCGContext* InOptionalContext = nullptr, bool bSilenceMissingAttributeErrors = false);
+
+	/**
+	* Remove duplicates for the incoming data on the given attribute.
+	 *Will return a single data, containing the first occurrence of every unique value
+	* @param InData - Data to partition, need to support attributes and multi entries (point data or attribute set).
+	* @param InSelector - Selector on the attribute to partition.
+	* @param InOptionalContext - Optional context for logging.
+	* @param bSilenceMissingAttributeErrors - Do not log errors to the context or log.
+	* @returns Resulting data, stripped of duplicates
+	*/
+	PCG_API UPCGData* RemoveDuplicates(const UPCGData* InData, const FPCGAttributePropertySelector& InSelector, FPCGContext* InOptionalContext = nullptr, bool bSilenceMissingAttributeErrors = false);
+
+	/**
+	* Remove duplicates for the incoming data on the given array of attributes.
+	* Will return a single data, containing the first occurrence of every unique tuple of values.
+	* @param InData - Data to partition, need to support attributes and multi entries (point data or attribute set).
+	* @param InSelectorArrayView - ArrayView of selectors on the attribute to partition.
+	* @param InOptionalContext - Optional context for logging.
+	* @param bSilenceMissingAttributeErrors - Do not log errors to the context or log.
+	* @returns Resulting data, stripped of duplicates
+	*/
+	PCG_API UPCGData* RemoveDuplicates(const UPCGData* InData, const TArrayView<const FPCGAttributePropertySelector>& InSelectorArrayView, FPCGContext* InOptionalContext = nullptr, bool bSilenceMissingAttributeErrors = false);
 }
\ No newline at end of file
