Fix an invalidation serialization path that occurs when an editor loads a cooked asset that contains a PCGComponent

diff --git a/Engine/Plugins/PCG/Source/PCG/Private/PCGComponent.cpp b/Engine/Plugins/PCG/Source/PCG/Private/PCGComponent.cpp
--- a/Engine/Plugins/PCG/Source/PCG/Private/PCGComponent.cpp
+++ b/Engine/Plugins/PCG/Source/PCG/Private/PCGComponent.cpp
@@ -1624,7 +1624,7 @@ void UPCGComponent::Serialize(FArchive& Ar)
 	Super::Serialize(Ar);
 
 #if WITH_EDITORONLY_DATA
-	if (!Ar.IsCooking())
+	if (!Ar.IsCooking() && !Ar.IsLoadingFromCookedPackage())
 	{
 		int32 DataVersion = FPCGCustomVersion::LatestVersion;
 		if (Ar.IsLoading())
