[PCG] Native GPU Copy Points implementation
* A full implementation of the Copy Points node for GPU execution (except for tag inheritance, which is not supported on the GPU)
* Moved GPU settings and functions from CustomHLSL settings to the base class settings
* Added option for settings to provide additional data interfaces to expose their parameters on the GPU

diff --git a/Engine/Plugins/PCG/Shaders/Private/Elements/PCGCopyPoints.usf b/Engine/Plugins/PCG/Shaders/Private/Elements/PCGCopyPoints.usf
new file mode 100644
--- /dev/null
+++ b/Engine/Plugins/PCG/Shaders/Private/Elements/PCGCopyPoints.usf
@@ -0,0 +1,307 @@
+// Copyright Epic Games, Inc. All Rights Reserved.
+
+#include "/Engine/Private/ComputeShaderUtils.ush"
+#include "/Engine/Private/Quaternion.ush"
+#include "/Plugin/PCG/Private/PCGShaderUtils.ush"
+
+// Mirrors EPCGCopyPointsInheritanceMode
+#define COPY_POINTS_INHERITANCE_MODE_Relative 0
+#define COPY_POINTS_INHERITANCE_MODE_Source 1
+#define COPY_POINTS_INHERITANCE_MODE_Target 2
+
+// Mirrors EPCGCopyPointsMetadataInheritanceMode
+#define COPY_POINTS_METADATA_INHERITANCE_MODE_SourceFirst 0
+#define COPY_POINTS_METADATA_INHERITANCE_MODE_TargetFirst 1
+#define COPY_POINTS_METADATA_INHERITANCE_MODE_SourceOnly 2
+#define COPY_POINTS_METADATA_INHERITANCE_MODE_TargetOnly 3
+#define COPY_POINTS_METADATA_INHERITANCE_MODE_None 4
+
+// TODO: should be code-gen'd from PCGComputeCommon.h
+#define MAX_NUM_ATTRS 128
+#define NUM_RESERVED_ATTRS 32
+
+void CopyAttributesFromSource(uint SourceDataIndex, uint SourceElemIndex, uint OutDataIndex, uint OutElemIndex)
+{
+	const uint SourceDataAddress = Source_GetDataAddress(SourceDataIndex);
+	const uint SourceHeadersAddress = Source_GetDataAttributeHeadersAddress(SourceDataAddress);
+
+	const uint OutDataAddress = Out_GetDataAddress(OutDataIndex);
+	const uint OutHeadersAddress = Out_GetDataAttributeHeadersAddress(OutDataAddress);
+
+	const uint NumAttributes = Source_GetDataNumAttributes(SourceDataAddress);
+	uint NumAttributesProcessed = 0;
+	
+	// Loop over all attribute headers, if the address is non-zero, then copy it from pin Source to pin Out.
+	for (int AttributeIndex = NUM_RESERVED_ATTRS; AttributeIndex < MAX_NUM_ATTRS; ++AttributeIndex)
+	{
+		const uint SourceHeaderAddress = Source_GetAttributeHeaderAddress(SourceHeadersAddress, AttributeIndex);
+		const uint SourceAddress = Source_GetAttributeAddress(SourceHeaderAddress);
+
+		const uint OutHeaderAddress = Out_GetAttributeHeaderAddress(OutHeadersAddress, AttributeIndex);
+		const uint OutAddress = Out_GetAttributeAddress(OutHeaderAddress);
+
+		const uint Stride = Source_GetAttributeStride(SourceHeaderAddress);
+
+		if (SourceAddress != 0 && OutAddress != 0)
+		{
+			const uint SourceBaseElementAddress = SourceAddress + SourceElemIndex * Stride;
+			const uint OutBaseElementAddress = OutAddress + OutElemIndex * Stride;
+
+			for (int I = 0; I < Stride; I += 4)
+			{
+				const uint SourceElementAddress = SourceBaseElementAddress + I;
+				const uint OutElementAddress = OutBaseElementAddress + I;
+				Out_StoreBufferInternal(OutElementAddress, Source_LoadBufferInternal(SourceElementAddress));
+			}
+
+			// We can early-out when we've looked at all the possible attributes
+			if (++NumAttributesProcessed >= NumAttributes)
+			{
+				break;
+			}
+		}
+	}
+}
+
+void CopyAttributesFromTarget(uint TargetDataIndex, uint TargetElemIndex, uint OutDataIndex, uint OutElemIndex)
+{
+	const uint TargetDataAddress = Target_GetDataAddress(TargetDataIndex);
+	const uint TargetHeadersAddress = Target_GetDataAttributeHeadersAddress(TargetDataAddress);
+
+	const uint OutDataAddress = Out_GetDataAddress(OutDataIndex);
+	const uint OutHeadersAddress = Out_GetDataAttributeHeadersAddress(OutDataAddress);
+
+	const uint NumAttributes = Target_GetDataNumAttributes(TargetDataAddress);
+	uint NumAttributesProcessed = 0;
+	
+	// Loop over all attribute headers, if the address is non-zero, then copy it from pin Target to pin Out.
+	for (int AttributeIndex = NUM_RESERVED_ATTRS; AttributeIndex < MAX_NUM_ATTRS; ++AttributeIndex)
+	{
+		const uint TargetHeaderAddress = Target_GetAttributeHeaderAddress(TargetHeadersAddress, AttributeIndex);
+		const uint TargetAddress = Target_GetAttributeAddress(TargetHeaderAddress);
+
+		const uint OutHeaderAddress = Out_GetAttributeHeaderAddress(OutHeadersAddress, AttributeIndex);
+		const uint OutAddress = Out_GetAttributeAddress(OutHeaderAddress);
+
+		const uint Stride = Target_GetAttributeStride(TargetHeaderAddress);
+
+		if (TargetAddress != 0 && OutAddress != 0)
+		{
+			const uint TargetBaseElementAddress = TargetAddress + TargetElemIndex * Stride;
+			const uint OutBaseElementAddress = OutAddress + OutElemIndex * Stride;
+
+			for (int I = 0; I < Stride; I += 4)
+			{
+				const uint TargetElementAddress = TargetBaseElementAddress + I;
+				const uint OutElementAddress = OutBaseElementAddress + I;
+				Out_StoreBufferInternal(OutElementAddress, Target_LoadBufferInternal(TargetElementAddress));
+			}
+
+			// We can early-out when we've looked at all the possible attributes
+			if (++NumAttributesProcessed >= NumAttributes)
+			{
+				break;
+			}
+		}
+	}
+}
+
+[numthreads(64, 1, 1)]
+void Main(uint3 GroupId : SV_GroupID, uint GroupIndex : SV_GroupIndex)
+{
+	const uint ThreadIndex = GetUnWrappedDispatchThreadId(GroupId, GroupIndex, 64);
+	if (ThreadIndex >= GetNumThreads().x) return;
+
+	const uint SourceNumData = Source_GetNumData();
+	const uint TargetNumData = Target_GetNumData();
+	const uint SourcePointsCount = Source_GetNumElements();
+	const uint bCopyEachSourceOnEveryTarget = GetCopyEachSourceOnEveryTarget();
+
+	// Data must be non-zero, and NxM (when bCopyEachSourceOnEveryTarget == true), N:N, 1:N, or N:1
+	if (SourceNumData == 0 || TargetNumData == 0 || (!bCopyEachSourceOnEveryTarget && SourceNumData != TargetNumData && SourceNumData != 1 && TargetNumData != 1))
+	{
+		// TODO: Runtime error diagnostics
+		return;
+	}
+	
+	uint SourceDataIndex, SourcePointIndexInData;
+	uint TargetDataIndex, TargetPointIndexInData;
+
+	if (bCopyEachSourceOnEveryTarget)
+	{
+		const uint TargetPointIndex = ThreadIndex / SourcePointsCount;
+		const uint SourcePointIndex = ThreadIndex % SourcePointsCount;
+
+		if (!Source_GetThreadData(SourcePointIndex, SourceDataIndex, SourcePointIndexInData))
+		{
+			return;
+		}
+
+		if (!Target_GetThreadData(TargetPointIndex, TargetDataIndex, TargetPointIndexInData))
+		{
+			return;
+		}
+	}
+	else
+	{
+		const uint NumIterations = max(SourceNumData, TargetNumData);
+
+		bool bFoundInData = false;
+		uint TotalElementCount = 0;
+
+		for (uint DataIndex = 0; DataIndex < NumIterations; ++DataIndex)
+		{
+			SourceDataIndex = min(DataIndex, SourceNumData - 1);
+			TargetDataIndex = min(DataIndex, TargetNumData - 1);
+
+			const uint SourceDataAddress = Source_GetDataAddress(SourceDataIndex);
+			const uint TargetDataAddress = Target_GetDataAddress(TargetDataIndex);
+
+			const uint SourceElemCount = Source_GetDataInfo(SourceDataAddress);
+			const uint TargetElemCount = Target_GetDataInfo(TargetDataAddress);
+
+			const uint ElemCount = SourceElemCount * TargetElemCount;
+
+			if (TotalElementCount + ElemCount > ThreadIndex)
+			{
+				const uint ElementIndex = ThreadIndex - TotalElementCount;
+				SourcePointIndexInData = ElementIndex / TargetElemCount;
+				TargetPointIndexInData = ElementIndex % TargetElemCount;
+				bFoundInData = true;
+				break;
+			}
+
+			TotalElementCount += ElemCount;
+		}
+
+		if (!bFoundInData)
+		{
+			return;
+		}
+	}
+
+	float3 SourcePosition = Source_GetPosition(SourceDataIndex,SourcePointIndexInData);
+	float4 SourceRotation = Source_GetRotation(SourceDataIndex,SourcePointIndexInData);
+	float3 SourceScale = Source_GetScale(SourceDataIndex,SourcePointIndexInData);
+	float4 SourceColor = Source_GetColor(SourceDataIndex,SourcePointIndexInData);
+	uint SourceSeed = Source_GetSeed(SourceDataIndex,SourcePointIndexInData);
+	
+	float3 TargetPosition = Target_GetPosition(TargetDataIndex,TargetPointIndexInData);
+	float4 TargetRotation = Target_GetRotation(TargetDataIndex,TargetPointIndexInData);
+	float3 TargetScale = Target_GetScale(TargetDataIndex,TargetPointIndexInData);
+	float4 TargetColor = Target_GetColor(TargetDataIndex,TargetPointIndexInData);
+	uint TargetSeed = Target_GetSeed(TargetDataIndex,TargetPointIndexInData);
+	
+	float3 OutPosition = TargetPosition + (QuatRotateVector(TargetRotation, (SourcePosition * TargetScale)));
+	float4 OutRotation = float4(0, 0, 0, 1);
+	float3 OutScale = (float3)0;
+	float4 OutColor = (float4)0;
+	uint OutSeed = 0;
+
+	const uint RotationInheritance = GetRotationInheritance();
+	const uint ScaleInheritance = GetScaleInheritance();
+	const uint ColorInheritance = GetColorInheritance();
+	const uint SeedInheritance = GetSeedInheritance();
+	const uint AttributeInheritance = GetAttributeInheritance();
+
+	if (RotationInheritance == COPY_POINTS_INHERITANCE_MODE_Relative)
+	{
+	  OutRotation = QuatMultiply(SourceRotation,TargetRotation);
+	}
+	else if (RotationInheritance == COPY_POINTS_INHERITANCE_MODE_Source)
+	{
+		OutRotation = SourceRotation;
+	}
+	else if (RotationInheritance == COPY_POINTS_INHERITANCE_MODE_Target)
+	{
+		OutRotation = TargetRotation;
+	}
+
+	if (ScaleInheritance == COPY_POINTS_INHERITANCE_MODE_Relative)
+	{
+		OutScale = SourceScale * TargetScale;
+	}
+	else if (ScaleInheritance == COPY_POINTS_INHERITANCE_MODE_Source)
+	{
+		OutScale = SourceScale;
+	}
+	else if (ScaleInheritance == COPY_POINTS_INHERITANCE_MODE_Target)
+	{
+		OutScale = TargetScale;
+	}
+
+	if (ColorInheritance == COPY_POINTS_INHERITANCE_MODE_Relative)
+	{
+		OutColor = SourceColor * TargetColor;
+	}
+	else if (ColorInheritance == COPY_POINTS_INHERITANCE_MODE_Source)
+	{
+		OutColor = SourceColor;
+	}
+	else if (ColorInheritance == COPY_POINTS_INHERITANCE_MODE_Target)
+	{
+		OutColor = TargetColor;
+	}
+
+	if (SeedInheritance == COPY_POINTS_INHERITANCE_MODE_Relative)
+	{
+		OutSeed = ComputeSeed(SourceSeed, TargetSeed);
+	}
+	else if (SeedInheritance == COPY_POINTS_INHERITANCE_MODE_Source)
+	{
+		OutSeed = SourceSeed;
+	}
+	else if (SeedInheritance == COPY_POINTS_INHERITANCE_MODE_Target)
+	{
+		OutSeed = TargetSeed;
+	}
+
+	uint OutNumData = 0;
+
+	if (bCopyEachSourceOnEveryTarget)
+	{
+		OutNumData = SourceNumData * TargetNumData;
+	}
+	else
+	{
+		OutNumData = max(SourceNumData, TargetNumData);
+	}
+
+	Out_SetNumDataInternal(OutNumData);
+
+	uint OutDataIndex, OutElemIndex;
+	if (!Out_GetThreadData(ThreadIndex, OutDataIndex, OutElemIndex))
+	{
+		return;
+	}
+	
+	// Initialize all fields
+	Out_InitializePoint(OutDataIndex, OutElemIndex);
+	
+	// Write output
+	Out_SetPosition(OutDataIndex, OutElemIndex, OutPosition);
+	Out_SetRotation(OutDataIndex, OutElemIndex, OutRotation);
+	Out_SetScale(OutDataIndex, OutElemIndex, OutScale);
+	Out_SetColor(OutDataIndex, OutElemIndex, OutColor);
+	Out_SetSeed(OutDataIndex, OutElemIndex, OutSeed);
+
+	// Copy attributes
+	if (AttributeInheritance == COPY_POINTS_METADATA_INHERITANCE_MODE_SourceOnly)
+	{
+		CopyAttributesFromSource(SourceDataIndex, SourcePointIndexInData, OutDataIndex, OutElemIndex);
+	}
+	else if (AttributeInheritance == COPY_POINTS_METADATA_INHERITANCE_MODE_TargetOnly)
+	{
+		CopyAttributesFromTarget(TargetDataIndex, TargetPointIndexInData, OutDataIndex, OutElemIndex);
+	}
+	else if (AttributeInheritance == COPY_POINTS_METADATA_INHERITANCE_MODE_SourceFirst)
+	{
+		CopyAttributesFromTarget(TargetDataIndex, TargetPointIndexInData, OutDataIndex, OutElemIndex);
+		CopyAttributesFromSource(SourceDataIndex, SourcePointIndexInData, OutDataIndex, OutElemIndex);
+	}
+	else if (AttributeInheritance == COPY_POINTS_METADATA_INHERITANCE_MODE_TargetFirst)
+	{
+		CopyAttributesFromSource(SourceDataIndex, SourcePointIndexInData, OutDataIndex, OutElemIndex);
+		CopyAttributesFromTarget(TargetDataIndex, TargetPointIndexInData, OutDataIndex, OutElemIndex);
+	}
+}
\ No newline at end of file
diff --git a/Engine/Plugins/PCG/Source/PCG/Private/Compute/DataInterfaces/PCGCopyPointsDataInterface.cpp b/Engine/Plugins/PCG/Source/PCG/Private/Compute/DataInterfaces/PCGCopyPointsDataInterface.cpp
new file mode 100644
--- /dev/null
+++ b/Engine/Plugins/PCG/Source/PCG/Private/Compute/DataInterfaces/PCGCopyPointsDataInterface.cpp
@@ -0,0 +1,118 @@
+// Copyright Epic Games, Inc. All Rights Reserved.
+
+#include "PCGCopyPointsDataInterface.h"
+
+#include "Elements/PCGCopyPoints.h"
+
+#include "GlobalRenderResources.h"
+#include "RHIResources.h"
+#include "ShaderCompilerCore.h"
+#include "ShaderCore.h"
+#include "ShaderParameterMetadataBuilder.h"
+#include "ComputeFramework/ShaderParamTypeDefinition.h"
+
+#include UE_INLINE_GENERATED_CPP_BY_NAME(PCGCopyPointsDataInterface)
+
+void UPCGCopyPointsDataInterface::GetSupportedInputs(TArray<FShaderFunctionDefinition>& OutFunctions) const
+{
+	OutFunctions.AddDefaulted_GetRef()
+		.SetName(TEXT("GetRotationInheritance"))
+		.AddReturnType(FShaderValueType::Get(EShaderFundamentalType::Uint));
+
+	OutFunctions.AddDefaulted_GetRef()
+		.SetName(TEXT("GetScaleInheritance"))
+		.AddReturnType(FShaderValueType::Get(EShaderFundamentalType::Uint));
+
+	OutFunctions.AddDefaulted_GetRef()
+		.SetName(TEXT("GetColorInheritance"))
+		.AddReturnType(FShaderValueType::Get(EShaderFundamentalType::Uint));
+
+	OutFunctions.AddDefaulted_GetRef()
+		.SetName(TEXT("GetSeedInheritance"))
+		.AddReturnType(FShaderValueType::Get(EShaderFundamentalType::Uint));
+
+	OutFunctions.AddDefaulted_GetRef()
+		.SetName(TEXT("GetAttributeInheritance"))
+		.AddReturnType(FShaderValueType::Get(EShaderFundamentalType::Uint));
+
+	OutFunctions.AddDefaulted_GetRef()
+		.SetName(TEXT("GetCopyEachSourceOnEveryTarget"))
+		.AddReturnType(FShaderValueType::Get(EShaderFundamentalType::Bool));
+}
+
+BEGIN_SHADER_PARAMETER_STRUCT(FPCGCopyPointsDataInterfaceParameters,)
+	SHADER_PARAMETER(uint32, RotationInheritance)
+	SHADER_PARAMETER(uint32, ScaleInheritance)
+	SHADER_PARAMETER(uint32, ColorInheritance)
+	SHADER_PARAMETER(uint32, SeedInheritance)
+	SHADER_PARAMETER(uint32, AttributeInheritance)
+	SHADER_PARAMETER(uint32, bCopyEachSourceOnEveryTarget)
+END_SHADER_PARAMETER_STRUCT()
+
+void UPCGCopyPointsDataInterface::GetShaderParameters(TCHAR const* UID, FShaderParametersMetadataBuilder& InOutBuilder, FShaderParametersMetadataAllocations& InOutAllocations) const
+{
+	InOutBuilder.AddNestedStruct<FPCGCopyPointsDataInterfaceParameters>(UID);
+}
+
+void UPCGCopyPointsDataInterface::GetHLSL(FString& OutHLSL, FString const& InDataInterfaceName) const
+{
+	TMap<FString, FStringFormatArg> TemplateArgs =
+	{
+		{ TEXT("DataInterfaceName"), InDataInterfaceName },
+	};
+
+	OutHLSL += FString::Format(TEXT(
+		"uint {DataInterfaceName}_RotationInheritance;\n"
+		"uint {DataInterfaceName}_ScaleInheritance;\n"
+		"uint {DataInterfaceName}_ColorInheritance;\n"
+		"uint {DataInterfaceName}_SeedInheritance;\n"
+		"uint {DataInterfaceName}_AttributeInheritance;\n"
+		"uint {DataInterfaceName}_bCopyEachSourceOnEveryTarget;\n"
+		"\n"
+		"uint GetRotationInheritance_{DataInterfaceName}() { return {DataInterfaceName}_RotationInheritance; }\n"
+		"uint GetScaleInheritance_{DataInterfaceName}() { return {DataInterfaceName}_ScaleInheritance; }\n"
+		"uint GetColorInheritance_{DataInterfaceName}() { return {DataInterfaceName}_ColorInheritance; }\n"
+		"uint GetSeedInheritance_{DataInterfaceName}() { return {DataInterfaceName}_SeedInheritance; }\n"
+		"uint GetAttributeInheritance_{DataInterfaceName}() { return {DataInterfaceName}_AttributeInheritance; }\n"
+		"uint GetCopyEachSourceOnEveryTarget_{DataInterfaceName}() { return {DataInterfaceName}_bCopyEachSourceOnEveryTarget; }\n"
+		), TemplateArgs);
+}
+
+UComputeDataProvider* UPCGCopyPointsDataInterface::CreateDataProvider(TObjectPtr<UObject> InBinding, uint64 InInputMask, uint64 InOutputMask) const
+{
+	TObjectPtr<UPCGCopyPointsDataProvider> DataProvider = NewObject<UPCGCopyPointsDataProvider>();
+	DataProvider->Settings = Cast<UPCGCopyPointsSettings>(Settings);
+
+	return DataProvider;
+}
+
+FComputeDataProviderRenderProxy* UPCGCopyPointsDataProvider::GetRenderProxy()
+{
+	FPCGCopyPointsDataProviderProxy* Proxy = new FPCGCopyPointsDataProviderProxy(Settings);
+	return Proxy;
+}
+
+bool FPCGCopyPointsDataProviderProxy::IsValid(FValidationData const& InValidationData) const
+{
+	if (InValidationData.ParameterStructSize != sizeof(FParameters))
+	{
+		return false;
+	}
+
+	return true;
+}
+
+void FPCGCopyPointsDataProviderProxy::GatherDispatchData(FDispatchData const& InDispatchData)
+{
+	const TStridedView<FParameters> ParameterArray = MakeStridedParameterView<FParameters>(InDispatchData);
+	for (int32 InvocationIndex = 0; InvocationIndex < ParameterArray.Num(); ++InvocationIndex)
+	{
+		FParameters& Parameters = ParameterArray[InvocationIndex];
+		Parameters.RotationInheritance = static_cast<uint32>(Settings->RotationInheritance);
+		Parameters.ScaleInheritance = static_cast<uint32>(Settings->ScaleInheritance);
+		Parameters.ColorInheritance = static_cast<uint32>(Settings->ColorInheritance);
+		Parameters.SeedInheritance = static_cast<uint32>(Settings->SeedInheritance);
+		Parameters.AttributeInheritance = static_cast<uint32>(Settings->AttributeInheritance);
+		Parameters.bCopyEachSourceOnEveryTarget = static_cast<uint32>(Settings->bCopyEachSourceOnEveryTarget);
+	}
+}
diff --git a/Engine/Plugins/PCG/Source/PCG/Private/Compute/DataInterfaces/PCGCopyPointsDataInterface.h b/Engine/Plugins/PCG/Source/PCG/Private/Compute/DataInterfaces/PCGCopyPointsDataInterface.h
new file mode 100644
--- /dev/null
+++ b/Engine/Plugins/PCG/Source/PCG/Private/Compute/DataInterfaces/PCGCopyPointsDataInterface.h
@@ -0,0 +1,62 @@
+// Copyright Epic Games, Inc. All Rights Reserved.
+
+#pragma once
+
+#include "Compute/DataInterfaces/PCGComputeDataInterface.h"
+#include "ComputeFramework/ComputeDataProvider.h"
+
+#include "PCGCopyPointsDataInterface.generated.h"
+
+class FPCGCopyPointsDataInterfaceParameters;
+class UPCGCopyPointsSettings;
+class UPCGSettings;
+
+/** Data Interface to marshal Copy Points settings to the GPU. */
+UCLASS(ClassGroup = (Procedural))
+class UPCGCopyPointsDataInterface : public UPCGComputeDataInterface
+{
+	GENERATED_BODY()
+
+public:
+	//~ Begin UComputeDataInterface Interface
+	TCHAR const* GetClassName() const override { return TEXT("PCGCopyPoints"); }
+	void GetSupportedInputs(TArray<FShaderFunctionDefinition>& OutFunctions) const override;
+	void GetShaderParameters(TCHAR const* UID, FShaderParametersMetadataBuilder& InOutBuilder, FShaderParametersMetadataAllocations& InOutAllocations) const override;
+	void GetHLSL(FString& OutHLSL, FString const& InDataInterfaceName) const override;
+	UComputeDataProvider* CreateDataProvider(TObjectPtr<UObject> InBinding, uint64 InInputMask, uint64 InOutputMask) const override;
+	//~ End UComputeDataInterface Interface
+
+	UPROPERTY()
+	TObjectPtr<const UPCGSettings> Settings = nullptr;
+};
+
+UCLASS()
+class UPCGCopyPointsDataProvider : public UComputeDataProvider
+{
+	GENERATED_BODY()
+
+public:
+	//~ Begin UComputeDataProvider Interface
+	FComputeDataProviderRenderProxy* GetRenderProxy() override;
+	//~ End UComputeDataProvider Interface
+
+	const UPCGCopyPointsSettings* Settings;
+};
+
+class FPCGCopyPointsDataProviderProxy : public FComputeDataProviderRenderProxy
+{
+public:
+	FPCGCopyPointsDataProviderProxy(const UPCGCopyPointsSettings* InSettings)
+		: Settings(InSettings)
+	{}
+
+	//~ Begin FComputeDataProviderRenderProxy Interface
+	bool IsValid(FValidationData const& InValidationData) const override;
+	void GatherDispatchData(FDispatchData const& InDispatchData) override;
+	//~ End FComputeDataProviderRenderProxy Interface
+
+protected:
+	using FParameters = FPCGCopyPointsDataInterfaceParameters;
+
+	const UPCGCopyPointsSettings* Settings;
+};
diff --git a/Engine/Plugins/PCG/Source/PCG/Private/Compute/DataInterfaces/PCGCustomKernelDataInterface.cpp b/Engine/Plugins/PCG/Source/PCG/Private/Compute/DataInterfaces/PCGCustomKernelDataInterface.cpp
--- a/Engine/Plugins/PCG/Source/PCG/Private/Compute/DataInterfaces/PCGCustomKernelDataInterface.cpp
+++ b/Engine/Plugins/PCG/Source/PCG/Private/Compute/DataInterfaces/PCGCustomKernelDataInterface.cpp
@@ -35,7 +35,7 @@ void UPCGCustomKernelDataInterface::GetSupportedInputs(TArray<FShaderFunctionDef
 		.AddReturnType(EShaderFundamentalType::Float, 3);
 }
 
-BEGIN_SHADER_PARAMETER_STRUCT(FPCGKernelDataInterfaceParameters, )
+BEGIN_SHADER_PARAMETER_STRUCT(FPCGCustomKernelDataInterfaceParameters, )
 	SHADER_PARAMETER(FIntVector3, NumThreads)
 	SHADER_PARAMETER(uint32, Seed)
 	SHADER_PARAMETER(FVector3f, ComponentBoundsMin)
@@ -44,7 +44,7 @@ END_SHADER_PARAMETER_STRUCT()
 
 void UPCGCustomKernelDataInterface::GetShaderParameters(TCHAR const* UID, FShaderParametersMetadataBuilder& InOutBuilder, FShaderParametersMetadataAllocations& InOutAllocations) const
 {
-	InOutBuilder.AddNestedStruct<FPCGKernelDataInterfaceParameters>(UID);
+	InOutBuilder.AddNestedStruct<FPCGCustomKernelDataInterfaceParameters>(UID);
 }
 
 void UPCGCustomKernelDataInterface::GetHLSL(FString& OutHLSL, FString const& InDataInterfaceName) const
diff --git a/Engine/Plugins/PCG/Source/PCG/Private/Compute/DataInterfaces/PCGCustomKernelDataInterface.h b/Engine/Plugins/PCG/Source/PCG/Private/Compute/DataInterfaces/PCGCustomKernelDataInterface.h
--- a/Engine/Plugins/PCG/Source/PCG/Private/Compute/DataInterfaces/PCGCustomKernelDataInterface.h
+++ b/Engine/Plugins/PCG/Source/PCG/Private/Compute/DataInterfaces/PCGCustomKernelDataInterface.h
@@ -9,7 +9,7 @@
 
 #include "PCGCustomKernelDataInterface.generated.h"
 
-class FPCGKernelDataInterfaceParameters;
+class FPCGCustomKernelDataInterfaceParameters;
 
 /** Interface for any meta data provided to the compute kernel, such as num threads. */
 UCLASS(ClassGroup = (Procedural))
@@ -21,7 +21,7 @@ public:
 	static const TCHAR* NumThreadsReservedName;
 	
 	//~ Begin UComputeDataInterface Interface
-	TCHAR const* GetClassName() const override { return TEXT("CustomComputeKernelData"); }
+	TCHAR const* GetClassName() const override { return TEXT("PCGCustomKernel"); }
 	bool CanSupportUnifiedDispatch() const override { return true; }
 	void GetSupportedInputs(TArray<FShaderFunctionDefinition>& OutFunctions) const override;
 	void GetShaderParameters(TCHAR const* UID, FShaderParametersMetadataBuilder& InOutBuilder, FShaderParametersMetadataAllocations& InOutAllocations) const override;
@@ -33,7 +33,7 @@ public:
 	//~ End UComputeDataInterface Interface
 
 	UPROPERTY()
-	TObjectPtr<const UPCGSettings> Settings;
+	TObjectPtr<const UPCGSettings> Settings = nullptr;
 };
 
 /** Compute Framework Data Provider for each custom compute kernel. */
@@ -74,7 +74,7 @@ public:
 	//~ End FComputeDataProviderRenderProxy Interface
 
 protected:
-	using FParameters = FPCGKernelDataInterfaceParameters;
+	using FParameters = FPCGCustomKernelDataInterfaceParameters;
 
 	TArray<int32> InvocationThreadCounts;
 	int32 TotalThreadCount;
diff --git a/Engine/Plugins/PCG/Source/PCG/Private/Compute/Elements/PCGComputeGraphElement.cpp b/Engine/Plugins/PCG/Source/PCG/Private/Compute/Elements/PCGComputeGraphElement.cpp
--- a/Engine/Plugins/PCG/Source/PCG/Private/Compute/Elements/PCGComputeGraphElement.cpp
+++ b/Engine/Plugins/PCG/Source/PCG/Private/Compute/Elements/PCGComputeGraphElement.cpp
@@ -154,14 +154,9 @@ bool FPCGComputeGraphElement::ExecuteInternal(FPCGContext* InContext) const
 		{
 			const UPCGSettings* Settings = Node.Get() ? Node->GetSettings() : nullptr;
 
-			if (Settings)
+			if (!Settings || !Settings->IsKernelValid(Context, /*bQuiet=*/false))
 			{
-				const UPCGCustomHLSLSettings* KernelSettings = CastChecked<UPCGCustomHLSLSettings>(Settings);
-
-				if (!KernelSettings->IsKernelValid(Context, /*bQuiet=*/false))
-				{
-					return true;
-				}
+				return true;
 			}
 		}
 	}
diff --git a/Engine/Plugins/PCG/Source/PCG/Private/Compute/Elements/PCGComputeGraphElement.h b/Engine/Plugins/PCG/Source/PCG/Private/Compute/Elements/PCGComputeGraphElement.h
--- a/Engine/Plugins/PCG/Source/PCG/Private/Compute/Elements/PCGComputeGraphElement.h
+++ b/Engine/Plugins/PCG/Source/PCG/Private/Compute/Elements/PCGComputeGraphElement.h
@@ -5,7 +5,6 @@
 #include "PCGElement.h"
 
 #include "PCGContext.h"
-#include "Compute/Elements/PCGCustomHLSL.h"
 
 #include "Compute/PCGComputeGraph.h"
 #include "Compute/PCGDataBinding.h"
diff --git a/Engine/Plugins/PCG/Source/PCG/Private/Compute/Elements/PCGCustomHLSL.cpp b/Engine/Plugins/PCG/Source/PCG/Private/Compute/Elements/PCGCustomHLSL.cpp
--- a/Engine/Plugins/PCG/Source/PCG/Private/Compute/Elements/PCGCustomHLSL.cpp
+++ b/Engine/Plugins/PCG/Source/PCG/Private/Compute/Elements/PCGCustomHLSL.cpp
@@ -19,13 +19,6 @@
 #include UE_INLINE_GENERATED_CPP_BY_NAME(PCGCustomHLSL)
 
 #define LOCTEXT_NAMESPACE "PCGCustomHLSLElement"
-#define PCG_LOGGING_ENABLED (!(UE_BUILD_SHIPPING || UE_BUILD_TEST) || USE_LOGGING_IN_SHIPPING)
-
-#if PCG_LOGGING_ENABLED
-#define PCG_LOG_VALIDATION(ValidationMessage) LogGraphError(ValidationMessage);
-#else
-#define PCG_LOG_VALIDATION(ValidationMessage) // Log removed
-#endif
 
 namespace PCGHLSLElement
 {
@@ -851,50 +844,29 @@ void UPCGCustomHLSLSettings::UpdateAttributeKeys()
 
 bool UPCGCustomHLSLSettings::IsKernelValid(FPCGContext* InContext, bool bQuiet) const
 {
-	auto LogGraphError = [InContext, bQuiet, This=this](const FText& InText)
-	{
-		if (!bQuiet)
-		{
-#if WITH_EDITOR
-			if (InContext && ensure(InContext->SourceComponent.IsValid() && InContext->SourceComponent.Get()))
-			{
-				if (UPCGSubsystem* Subsystem = InContext->SourceComponent->GetSubsystem())
-				{
-					FPCGStack StackWithNode = InContext->Stack ? *InContext->Stack : FPCGStack();
-					StackWithNode.PushFrame(This->GetOuter());
-
-					Subsystem->GetNodeVisualLogsMutable().Log(StackWithNode, ELogVerbosity::Error, InText);
-				}
-			}
-#endif
-
-			PCGE_LOG_C(Error, LogOnly, InContext, InText);
-		}
-	};
-
 	if (OutputPins.IsEmpty())
 	{
-		PCG_LOG_VALIDATION(LOCTEXT("NoOutputs", "Custom HLSL nodes must have at least one output."));
+		PCG_KERNEL_VALIDATION(InContext, this, bQuiet, LOCTEXT("NoOutputs", "Custom HLSL nodes must have at least one output."));
 		return false;
 	}
 
-	auto CheckPinLabel = [InContext, &InPins = InputPins, OutPins = OutputPinProperties(), &LogGraphError](FName PinLabel)
+	auto CheckPinLabel = [This = this, InContext, bQuiet, &InPins = InputPins, OutPins = OutputPinProperties()](FName PinLabel)
 	{
 		if (PinLabel == NAME_None)
 		{
-			PCG_LOG_VALIDATION(LOCTEXT("InvalidPinLabelNone", "Pin label 'None' is not a valid pin label."));
+			PCG_KERNEL_VALIDATION(InContext, This, bQuiet, LOCTEXT("InvalidPinLabelNone", "Pin label 'None' is not a valid pin label."));
 			return false;
 		}
 
 		bool bFoundPinLabel = false;
 
-		auto IsAlreadyFound = [InContext, &LogGraphError, PinLabel, &bFoundPinLabel](const FPCGPinProperties PinProps)
+		auto IsAlreadyFound = [This, InContext, bQuiet, PinLabel, &bFoundPinLabel](const FPCGPinProperties PinProps)
 		{
 			if (PinProps.Label == PinLabel)
 			{
 				if (bFoundPinLabel)
 				{
-					PCG_LOG_VALIDATION(FText::Format(LOCTEXT("DuplicatedPinLabels", "Duplicate pin label '{0}', all labels must be unique."), FText::FromName(PinLabel)));
+					PCG_KERNEL_VALIDATION(InContext, This, bQuiet, FText::Format(LOCTEXT("DuplicatedPinLabels", "Duplicate pin label '{0}', all labels must be unique."), FText::FromName(PinLabel)));
 					return true;
 				}
 
@@ -936,7 +908,7 @@ bool UPCGCustomHLSLSettings::IsKernelValid(FPCGContext* InContext, bool bQuiet)
 		{
 			if (Properties.AllowedTypes != EPCGDataType::Point)
 			{
-				PCG_LOG_VALIDATION(FText::Format(
+				PCG_KERNEL_VALIDATION(InContext, this, bQuiet, FText::Format(
 					LOCTEXT("InvalidNonPointPrimaryInput", "'Point Processor' nodes require primary input pin to be of type 'Point', but found '{0}'."),
 					FText::FromString(PCGHLSLElement::GetDataTypeString(Properties.AllowedTypes))));
 
@@ -946,7 +918,7 @@ bool UPCGCustomHLSLSettings::IsKernelValid(FPCGContext* InContext, bool bQuiet)
 
 		if (!PCGComputeHelpers::IsTypeAllowedAsInput(Properties.AllowedTypes))
 		{
-			PCG_LOG_VALIDATION(FText::Format(
+			PCG_KERNEL_VALIDATION(InContext, this, bQuiet, FText::Format(
 				LOCTEXT("InvalidInputType", "Unsupported input type '{0}', found on pin '{1}'."),
 				FText::FromString(PCGHLSLElement::GetDataTypeString(Properties.AllowedTypes)),
 				FText::FromName(Properties.Label)));
@@ -972,7 +944,7 @@ bool UPCGCustomHLSLSettings::IsKernelValid(FPCGContext* InContext, bool bQuiet)
 		{
 			if (Properties.AllowedTypes != EPCGDataType::Point)
 			{
-				PCG_LOG_VALIDATION(FText::Format(
+				PCG_KERNEL_VALIDATION(InContext, this, bQuiet, FText::Format(
 					LOCTEXT("InvalidNonPointPrimaryOutput", "'Point Processor' and 'Point Generator' nodes require primary output pin to be of type 'Point', but found '{0}'."),
 					FText::FromString(PCGHLSLElement::GetDataTypeString(Properties.AllowedTypes))));
 
@@ -982,7 +954,7 @@ bool UPCGCustomHLSLSettings::IsKernelValid(FPCGContext* InContext, bool bQuiet)
 
 		if (!PCGComputeHelpers::IsTypeAllowedAsOutput(Properties.AllowedTypes))
 		{
-			PCG_LOG_VALIDATION(FText::Format(
+			PCG_KERNEL_VALIDATION(InContext, this, bQuiet, FText::Format(
 				LOCTEXT("InvalidOutputType", "Unsupported output type '{0}', found on pin '{1}'."),
 				FText::FromString(PCGHLSLElement::GetDataTypeString(Properties.AllowedTypes)),
 				FText::FromName(Properties.Label)));
@@ -996,7 +968,7 @@ bool UPCGCustomHLSLSettings::IsKernelValid(FPCGContext* InContext, bool bQuiet)
 			{
 				if (Properties.FixedBufferElementCount <= 0)
 				{
-					PCG_LOG_VALIDATION(FText::Format(
+					PCG_KERNEL_VALIDATION(InContext, this, bQuiet, FText::Format(
 						LOCTEXT("InvalidFixedBufferSize", "Fixed GPU buffer size on '{0}' was invalid (%d)."),
 						FText::FromName(Properties.Label),
 						Properties.FixedBufferElementCount));
@@ -1006,7 +978,7 @@ bool UPCGCustomHLSLSettings::IsKernelValid(FPCGContext* InContext, bool bQuiet)
 
 				if (Properties.AllowedTypes == EPCGDataType::Param && Properties.InitializeFromPin != NAME_None && !GetInputPin(Properties.InitializeFromPin))
 				{
-					PCG_LOG_VALIDATION(FText::Format(
+					PCG_KERNEL_VALIDATION(InContext, this, bQuiet, FText::Format(
 						LOCTEXT("InvalidInitFromPin", "Tried to initialize attribute set pin '{0}' from non-existent pin '{1}'. Must reference a valid input pin or be 'None'."),
 						FText::FromName(Properties.Label),
 						FText::FromName(Properties.InitializeFromPin)));
@@ -1018,7 +990,7 @@ bool UPCGCustomHLSLSettings::IsKernelValid(FPCGContext* InContext, bool bQuiet)
 			{
 				if (InputPins.IsEmpty())
 				{
-					PCG_LOG_VALIDATION(FText::Format(
+					PCG_KERNEL_VALIDATION(InContext, this, bQuiet, FText::Format(
 						LOCTEXT("InvalidBufferSizeNoInputPin", "GPU buffer size for pin '{0}' could not be computed as there are no input pins."),
 						FText::FromName(Properties.Label)));
 
@@ -1027,7 +999,7 @@ bool UPCGCustomHLSLSettings::IsKernelValid(FPCGContext* InContext, bool bQuiet)
 
 				if (!GetFirstInputPin())
 				{
-					PCG_LOG_VALIDATION(FText::Format(
+					PCG_KERNEL_VALIDATION(InContext, this, bQuiet, FText::Format(
 						LOCTEXT("MissingPrimaryInputPin", "GPU buffer size for pin '{0}' could not be computed, because it refers to the primary input pin, which does not exist."),
 						FText::FromName(Properties.Label)));
 
@@ -1038,7 +1010,7 @@ bool UPCGCustomHLSLSettings::IsKernelValid(FPCGContext* InContext, bool bQuiet)
 			{
 				if (InputPins.IsEmpty())
 				{
-					PCG_LOG_VALIDATION(FText::Format(
+					PCG_KERNEL_VALIDATION(InContext, this, bQuiet, FText::Format(
 						LOCTEXT("InvalidBufferSizeNoInputPins", "GPU buffer size for pin '{0}' could not be computed as there are no input pins on this node."),
 						FText::FromName(Properties.Label)));
 
@@ -1047,7 +1019,7 @@ bool UPCGCustomHLSLSettings::IsKernelValid(FPCGContext* InContext, bool bQuiet)
 
 				if (Properties.BufferSizeInputPinLabels.IsEmpty())
 				{
-					PCG_LOG_VALIDATION(FText::Format(
+					PCG_KERNEL_VALIDATION(InContext, this, bQuiet, FText::Format(
 						LOCTEXT("InvalidBufferSizeNoBufferPins", "GPU buffer size for pin '{0}' could not be computed as input pins are specified in the pin settings."),
 						FText::FromName(Properties.Label)));
 
@@ -1058,7 +1030,7 @@ bool UPCGCustomHLSLSettings::IsKernelValid(FPCGContext* InContext, bool bQuiet)
 				{
 					if (!GetInputPin(Label))
 					{
-						PCG_LOG_VALIDATION(FText::Format(
+						PCG_KERNEL_VALIDATION(InContext, this, bQuiet, FText::Format(
 							LOCTEXT("MissingBufferSizePin", "GPU buffer size for pin '{0}' could not be computed. Invalid pin specified in Input Pins array: '{1}'."),
 							FText::FromName(Properties.Label),
 							FText::FromName(Label)));
@@ -1076,7 +1048,7 @@ bool UPCGCustomHLSLSettings::IsKernelValid(FPCGContext* InContext, bool bQuiet)
 	{
 		if (ThreadCountInputPinLabels.IsEmpty())
 		{
-			PCG_LOG_VALIDATION(LOCTEXT("MissingThreadCountPins", "Dispatch thread count is based on input pins but no labels have been set in Input Pins array."));
+			PCG_KERNEL_VALIDATION(InContext, this, bQuiet, LOCTEXT("MissingThreadCountPins", "Dispatch thread count is based on input pins but no labels have been set in Input Pins array."));
 			return false;
 		}
 
@@ -1084,7 +1056,7 @@ bool UPCGCustomHLSLSettings::IsKernelValid(FPCGContext* InContext, bool bQuiet)
 		{
 			if (!GetInputPin(Label))
 			{
-				PCG_LOG_VALIDATION(FText::Format(LOCTEXT("MissingThreadCountPin", "Invalid pin specified in Input Pins array: '{0}'."), FText::FromName(Label)));
+				PCG_KERNEL_VALIDATION(InContext, this, bQuiet, FText::Format(LOCTEXT("MissingThreadCountPin", "Invalid pin specified in Input Pins array: '{0}'."), FText::FromName(Label)));
 				return false;
 			}
 		}
@@ -1094,7 +1066,7 @@ bool UPCGCustomHLSLSettings::IsKernelValid(FPCGContext* InContext, bool bQuiet)
 	{
 		if (ThreadCountMultiplier < 1)
 		{
-			PCG_LOG_VALIDATION(FText::Format(LOCTEXT("InvalidThreadCountMultiplier", "Thread Count Multiplier has invalid value ({0}). Must be greater than 0."), ThreadCountMultiplier));
+			PCG_KERNEL_VALIDATION(InContext, this, bQuiet, FText::Format(LOCTEXT("InvalidThreadCountMultiplier", "Thread Count Multiplier has invalid value ({0}). Must be greater than 0."), ThreadCountMultiplier));
 			return false;
 		}
 	}
@@ -1104,7 +1076,7 @@ bool UPCGCustomHLSLSettings::IsKernelValid(FPCGContext* InContext, bool bQuiet)
 	{
 		if (AttributeKey.Name == NAME_None)
 		{
-			PCG_LOG_VALIDATION(LOCTEXT("InvalidAttributeNameNone", "'None' is not a valid GPU attribute name, check the 'Attributes to Create' array on your pins."));
+			PCG_KERNEL_VALIDATION(InContext, this, bQuiet, LOCTEXT("InvalidAttributeNameNone", "'None' is not a valid GPU attribute name, check the 'Attributes to Create' array on your pins."));
 			return false;
 		}
 	}
@@ -1112,7 +1084,7 @@ bool UPCGCustomHLSLSettings::IsKernelValid(FPCGContext* InContext, bool bQuiet)
 	if (InContext)
 	{
 		FText* ErrorTextPtr = nullptr;
-#if PCG_LOGGING_ENABLED
+#if PCG_KERNEL_LOGGING_ENABLED
 		FText ErrorText;
 		ErrorTextPtr = &ErrorText;
 #endif
@@ -1121,7 +1093,7 @@ bool UPCGCustomHLSLSettings::IsKernelValid(FPCGContext* InContext, bool bQuiet)
 		{
 			if (ErrorTextPtr)
 			{
-				PCG_LOG_VALIDATION(*ErrorTextPtr);
+				PCG_KERNEL_VALIDATION(InContext, this, bQuiet, *ErrorTextPtr);
 			}
 			return false;
 		}
@@ -1187,7 +1159,7 @@ bool UPCGCustomHLSLSettings::AreKernelAttributesValid(FPCGContext* InContext, FT
 
 				if (!PinDesc && InputPinDescs.Find(PinName))
 				{
-#if PCG_LOGGING_ENABLED
+#if PCG_KERNEL_LOGGING_ENABLED
 					if (OutErrorText)
 					{
 						*OutErrorText = FText::Format(
@@ -1207,7 +1179,7 @@ bool UPCGCustomHLSLSettings::AreKernelAttributesValid(FPCGContext* InContext, FT
 
 			if (!PinDesc)
 			{
-#if PCG_LOGGING_ENABLED
+#if PCG_KERNEL_LOGGING_ENABLED
 				if (OutErrorText)
 				{
 					*OutErrorText = FText::Format(
@@ -1227,7 +1199,7 @@ bool UPCGCustomHLSLSettings::AreKernelAttributesValid(FPCGContext* InContext, FT
 
 			if (AttributeType == INDEX_NONE)
 			{
-#if PCG_LOGGING_ENABLED
+#if PCG_KERNEL_LOGGING_ENABLED
 				if (OutErrorText)
 				{
 					*OutErrorText = FText::Format(
@@ -1252,7 +1224,7 @@ bool UPCGCustomHLSLSettings::AreKernelAttributesValid(FPCGContext* InContext, FT
 
 			if (!AttrDesc)
 			{
-#if PCG_LOGGING_ENABLED
+#if PCG_KERNEL_LOGGING_ENABLED
 				if (OutErrorText)
 				{
 					*OutErrorText = FText::Format(
@@ -1267,7 +1239,7 @@ bool UPCGCustomHLSLSettings::AreKernelAttributesValid(FPCGContext* InContext, FT
 
 			if (AttrDesc->Type != static_cast<EPCGKernelAttributeType>(AttributeType))
 			{
-#if PCG_LOGGING_ENABLED
+#if PCG_KERNEL_LOGGING_ENABLED
 				if (OutErrorText)
 				{
 					const FString ActualTypeStr = AttributeTypeEnum->GetNameStringByIndex(static_cast<int64>(AttrDesc->Type));
@@ -1558,11 +1530,6 @@ FString UPCGCustomHLSLSettings::GetCookedKernelSource(const TMap<FPCGKernelAttri
 			*Includes, *Functions, *KernelFunc, *UnWrappedDispatchThreadId, *HeaderWriters, *KernelSpecificPreamble, *ShaderPathName, *Source);
 	}
 
-	if (bDumpCookedHLSL)
-	{
-		UE_LOG(LogPCG, Log, TEXT("Cooked HLSL:\n%s\n"), *Result);
-	}
-
 	return Result;
 }
 
@@ -1578,6 +1545,4 @@ bool FPCGCustomHLSLElement::ExecuteInternal(FPCGContext* Context) const
 	return true;
 }
 
-#undef PCG_LOG_VALIDATION
-#undef PCG_LOGGING_ENABLED
 #undef LOCTEXT_NAMESPACE
diff --git a/Engine/Plugins/PCG/Source/PCG/Private/Compute/PCGComputeCommon.cpp b/Engine/Plugins/PCG/Source/PCG/Private/Compute/PCGComputeCommon.cpp
--- a/Engine/Plugins/PCG/Source/PCG/Private/Compute/PCGComputeCommon.cpp
+++ b/Engine/Plugins/PCG/Source/PCG/Private/Compute/PCGComputeCommon.cpp
@@ -2,7 +2,10 @@
 
 #include "PCGComputeCommon.h"
 
+#include "PCGComponent.h"
+#include "PCGContext.h"
 #include "PCGParamData.h"
+#include "PCGSubsystem.h"
 #include "Data/PCGPointData.h"
 
 namespace PCGComputeHelpers
@@ -38,4 +41,23 @@ namespace PCGComputeHelpers
 	{
 		return (Type | PCGComputeConstants::AllowedDataCollectionTypes) == PCGComputeConstants::AllowedDataCollectionTypes;
 	}
+
+#if PCG_KERNEL_LOGGING_ENABLED
+	void LogKernelError(const FPCGContext* Context, const UPCGSettings* Settings, const FText& InText)
+	{
+#if WITH_EDITOR
+		if (Context && ensure(Context->SourceComponent.IsValid() && Context->SourceComponent.Get()))
+		{
+			if (UPCGSubsystem* Subsystem = Context->SourceComponent->GetSubsystem())
+			{
+				FPCGStack StackWithNode = Context->Stack ? *Context->Stack : FPCGStack();
+				StackWithNode.PushFrame(Settings->GetOuter());
+
+				Subsystem->GetNodeVisualLogsMutable().Log(StackWithNode, ELogVerbosity::Error, InText);
+			}
+		}
+#endif
+		PCGE_LOG_C(Error, LogOnly, Context, InText);
+	}
+#endif
 }
diff --git a/Engine/Plugins/PCG/Source/PCG/Private/Compute/PCGComputeCommon.h b/Engine/Plugins/PCG/Source/PCG/Private/Compute/PCGComputeCommon.h
--- a/Engine/Plugins/PCG/Source/PCG/Private/Compute/PCGComputeCommon.h
+++ b/Engine/Plugins/PCG/Source/PCG/Private/Compute/PCGComputeCommon.h
@@ -5,6 +5,16 @@
 #include "PCGCommon.h"
 
 class UPCGData;
+class UPCGSettings;
+struct FPCGContext;
+
+#define PCG_KERNEL_LOGGING_ENABLED (!(UE_BUILD_SHIPPING || UE_BUILD_TEST) || USE_LOGGING_IN_SHIPPING)
+
+#if PCG_KERNEL_LOGGING_ENABLED
+#define PCG_KERNEL_VALIDATION(Context, Settings, bQuiet, ValidationMessage) if (!bQuiet) PCGComputeHelpers::LogKernelError(Context, Settings, ValidationMessage);
+#else
+#define PCG_KERNEL_VALIDATION(Context, Settings, bQuiet, ValidationMessage) // Log removed
+#endif
 
 namespace PCGComputeConstants
 {
@@ -56,4 +66,8 @@ namespace PCGComputeHelpers
 
 	/** True if 'Type' is valid in a GPU data collection. Some types are only supported as DataInterfaces, and cannot be uploaded in data collections. */
 	bool IsTypeAllowedInDataCollection(EPCGDataType Type);
+
+#if PCG_KERNEL_LOGGING_ENABLED
+	void LogKernelError(const FPCGContext* Context, const UPCGSettings* Settings, const FText& InText);
+#endif
 }
diff --git a/Engine/Plugins/PCG/Source/PCG/Private/Elements/PCGCopyPoints.cpp b/Engine/Plugins/PCG/Source/PCG/Private/Elements/PCGCopyPoints.cpp
--- a/Engine/Plugins/PCG/Source/PCG/Private/Elements/PCGCopyPoints.cpp
+++ b/Engine/Plugins/PCG/Source/PCG/Private/Elements/PCGCopyPoints.cpp
@@ -3,11 +3,15 @@
 #include "Elements/PCGCopyPoints.h"
 
 #include "PCGContext.h"
+#include "Compute/DataInterfaces/PCGCopyPointsDataInterface.h"
+#include "Compute/PCGComputeCommon.h"
 #include "Data/PCGPointData.h"
 #include "Data/PCGSpatialData.h"
 #include "Helpers/PCGAsync.h"
 #include "Helpers/PCGHelpers.h"
 
+#include "RHIShaderPlatform.h"
+#include "ShaderCompilerCore.h"
 #include "Async/ParallelFor.h"
 
 #include UE_INLINE_GENERATED_CPP_BY_NAME(PCGCopyPoints)
@@ -21,6 +25,99 @@ FText UPCGCopyPointsSettings::GetNodeTooltipText() const
 }
 #endif
 
+FString UPCGCopyPointsSettings::GetCookedKernelSource(const TMap<FPCGKernelAttributeKey, int>& GlobalAttributeLookupTable) const
+{
+	FString TemplateFile;
+	LoadShaderSourceFile(TEXT("/Plugin/PCG/Private/Elements/PCGCopyPoints.usf"), EShaderPlatform::SP_PCD3D_SM5, &TemplateFile, nullptr);
+	return TemplateFile;
+}
+
+FPCGDataCollectionDesc UPCGCopyPointsSettings::ComputeOutputPinDataDesc(const UPCGPin* OutputPin, const UPCGDataBinding* Binding) const
+{
+	check(OutputPin);
+
+	const UPCGNode* Node = CastChecked<UPCGNode>(GetOuter());
+	const UPCGPin* SourcePin = Node->GetInputPin(PCGCopyPointsConstants::SourcePointsLabel);
+	const UPCGPin* TargetPin = Node->GetInputPin(PCGCopyPointsConstants::TargetPointsLabel);
+
+	FPCGDataCollectionDesc PinDesc;
+
+	if (ensure(SourcePin && TargetPin))
+	{
+		const FPCGDataCollectionDesc SourcePinDesc = ComputeInputPinDataDesc(SourcePin, Binding);
+		const FPCGDataCollectionDesc TargetPinDesc = ComputeInputPinDataDesc(TargetPin, Binding);
+
+		const int32 NumSources = SourcePinDesc.DataDescs.Num();
+		const int32 NumTargets = TargetPinDesc.DataDescs.Num();
+		const int32 NumIterations = bCopyEachSourceOnEveryTarget ? NumSources * NumTargets : FMath::Max(NumSources, NumTargets);
+
+		if (NumSources > 0 && NumTargets > 0 && (bCopyEachSourceOnEveryTarget || NumSources == NumTargets || NumSources == 1 || NumTargets == 1))
+		{
+			for (int32 I = 0; I < NumIterations; ++I)
+			{
+				const int32 SourceIndex = bCopyEachSourceOnEveryTarget ? (I / NumTargets) : FMath::Min(I, NumSources - 1);
+				const int32 TargetIndex = bCopyEachSourceOnEveryTarget ? (I % NumTargets) : FMath::Min(I, NumTargets - 1);
+
+				const FPCGDataDesc& SourceDesc = SourcePinDesc.DataDescs[SourceIndex];
+				const FPCGDataDesc& TargetDesc = TargetPinDesc.DataDescs[TargetIndex];
+
+				FPCGDataDesc& ResultDataDesc = PinDesc.DataDescs.Emplace_GetRef(EPCGDataType::Point, SourceDesc.ElementCount * TargetDesc.ElementCount);
+
+				if (AttributeInheritance == EPCGCopyPointsMetadataInheritanceMode::SourceFirst)
+				{
+					ResultDataDesc.AttributeDescs = SourceDesc.AttributeDescs;
+
+					for (const FPCGKernelAttributeDesc& AttrDesc : TargetDesc.AttributeDescs)
+					{
+						ResultDataDesc.AttributeDescs.AddUnique(AttrDesc);
+					}
+				}
+				else if (AttributeInheritance == EPCGCopyPointsMetadataInheritanceMode::TargetFirst)
+				{
+					ResultDataDesc.AttributeDescs = TargetDesc.AttributeDescs;
+
+					for (const FPCGKernelAttributeDesc& AttrDesc : SourceDesc.AttributeDescs)
+					{
+						ResultDataDesc.AttributeDescs.AddUnique(AttrDesc);
+					}
+				}
+				else if (AttributeInheritance == EPCGCopyPointsMetadataInheritanceMode::SourceOnly)
+				{
+					ResultDataDesc.AttributeDescs = SourceDesc.AttributeDescs;
+				}
+				else if (AttributeInheritance == EPCGCopyPointsMetadataInheritanceMode::TargetOnly)
+				{
+					ResultDataDesc.AttributeDescs = TargetDesc.AttributeDescs;
+				}
+			}
+		}
+	}
+
+	return PinDesc;
+}
+
+int UPCGCopyPointsSettings::ComputeKernelThreadCount(const UPCGDataBinding* Binding) const
+{
+	const UPCGNode* Node = CastChecked<UPCGNode>(GetOuter());
+	const UPCGPin* OutPin = Node->GetOutputPin(PCGPinConstants::DefaultOutputLabel);
+
+	const FPCGDataCollectionDesc OutputPinDesc = ComputeOutputPinDataDesc(OutPin, Binding);
+	int ThreadCount = 0;
+
+	for (const FPCGDataDesc& DataDesc : OutputPinDesc.DataDescs)
+	{
+		ThreadCount += DataDesc.ElementCount;
+	}
+
+	return ThreadCount;
+}
+
+void UPCGCopyPointsSettings::CreateAdditionalInputDataInterfaces(TArray<TObjectPtr<UComputeDataInterface>>& OutDataInterfaces) const
+{
+	TObjectPtr<UPCGCopyPointsDataInterface> DataInterface = Cast<UPCGCopyPointsDataInterface>(OutDataInterfaces.Add_GetRef(NewObject<UPCGCopyPointsDataInterface>()));
+	DataInterface->Settings = this;
+}
+
 TArray<FPCGPinProperties> UPCGCopyPointsSettings::InputPinProperties() const
 {
 	TArray<FPCGPinProperties> PinProperties;
@@ -74,7 +171,7 @@ bool FPCGCopyPointsElement::ExecuteInternal(FPCGContext* Context) const
 	{
 		if (NumSources != NumTargets && NumSources != 1 && NumTargets != 1)
 		{
-			PCGLog::LogErrorOnGraph(FText::Format(LOCTEXT("MismatchNum", "Num Sources ({0}) mismatches with Num Targets ({1}). Only supports N:N, 1:N and N:1 operation."), NumSources, NumTargets), Context);
+			PCGLog::LogErrorOnGraph(FText::Format(LOCTEXT("NumDataMismatch", "Num Sources ({0}) mismatches with Num Targets ({1}). Only supports N:N, 1:N and N:1 operation."), NumSources, NumTargets), Context);
 			Outputs = Sources;
 			return true;
 		}
diff --git a/Engine/Plugins/PCG/Source/PCG/Private/Graph/PCGGraphCompilerGPU.cpp b/Engine/Plugins/PCG/Source/PCG/Private/Graph/PCGGraphCompilerGPU.cpp
--- a/Engine/Plugins/PCG/Source/PCG/Private/Graph/PCGGraphCompilerGPU.cpp
+++ b/Engine/Plugins/PCG/Source/PCG/Private/Graph/PCGGraphCompilerGPU.cpp
@@ -5,9 +5,6 @@
 #include "PCGGraph.h"
 #include "PCGModule.h"
 #include "PCGPin.h"
-#include "Compute/PCGComputeCommon.h"
-#include "Compute/PCGComputeKernelSource.h"
-#include "Compute/DataInterfaces/PCGComputeDataInterface.h"
 #include "Compute/DataInterfaces/PCGCustomKernelDataInterface.h"
 #include "Compute/DataInterfaces/PCGDataCollectionDataInterface.h"
 #include "Compute/DataInterfaces/PCGDataCollectionUploadDataInterface.h"
@@ -15,7 +12,8 @@
 #include "Compute/DataInterfaces/PCGLandscapeDataInterface.h"
 #include "Compute/DataInterfaces/PCGTextureDataInterface.h"
 #include "Compute/Elements/PCGComputeGraphElement.h"
-#include "Compute/Elements/PCGCustomHLSL.h"
+#include "Compute/PCGComputeCommon.h"
+#include "Compute/PCGComputeKernelSource.h"
 #include "Graph/PCGGraphCompiler.h"
 #include "Graph/PCGGraphExecutor.h"
 
@@ -238,7 +236,7 @@ void FPCGGraphCompilerGPU::CollectGPUNodeSubsets(
 			bool bAllNodesValid = true;
 			for (FPCGTaskId& TaskId : GPUSubsetTaskIds)
 			{
-				const UPCGCustomHLSLSettings* Settings = Cast<UPCGCustomHLSLSettings>(InCompiledTasks[TaskId].Node ? InCompiledTasks[TaskId].Node->GetSettings() : nullptr);
+				const UPCGSettings* Settings = InCompiledTasks[TaskId].Node ? InCompiledTasks[TaskId].Node->GetSettings() : nullptr;
 				if (Settings && !Settings->IsKernelValid())
 				{
 					bAllNodesValid = false;
@@ -680,7 +678,7 @@ void FPCGGraphCompilerGPU::BuildGPUGraphTask(
 
 		const UPCGNode* Node = InOutCompiledTasks[TaskId].Node;
 
-		const UPCGCustomHLSLSettings* Settings = Cast<UPCGCustomHLSLSettings>(Node ? Node->GetSettings() : nullptr);
+		const UPCGSettings* Settings = Node ? Node->GetSettings() : nullptr;
 		check(Settings && Settings->bEnabled && Settings->ShouldExecuteOnGPU());
 
 		// For every usage of a DI, get the original (non-aliased) pin label.
@@ -786,22 +784,64 @@ void FPCGGraphCompilerGPU::BuildGPUGraphTask(
 
 		UPCGCustomKernelDataInterface* KernelDI = NewObject<UPCGCustomKernelDataInterface>(ComputeGraph);
 		KernelDI->Settings = Settings;
-		const int KernelDIIndex = ComputeGraph->DataInterfaces.Num();
+
+		const int32 KernelDIIndex = ComputeGraph->DataInterfaces.Num();
 		ComputeGraph->DataInterfaces.Add(KernelDI);
+		InputDataInterfaceIndices.Add(KernelDIIndex);
+
+		TArray<TObjectPtr<UComputeDataInterface>> AdditionalInputDIs, AdditionalOutputDIs;
+		Settings->CreateAdditionalInputDataInterfaces(AdditionalInputDIs);
+		Settings->CreateAdditionalOutputDataInterfaces(AdditionalOutputDIs);
+
+		auto AddAdditionalDataInterface = [ComputeGraph](TObjectPtr<UComputeDataInterface> DataInterface) -> int32
+		{
+			if (ensure(DataInterface))
+			{
+				// Re-outer to the ComputeGraph
+				DataInterface->Rename(/*Name=*/nullptr, ComputeGraph);
+
+				const int32 DataInterfaceIndex = ComputeGraph->DataInterfaces.Num();
+				ComputeGraph->DataInterfaces.Add(DataInterface);
+
+				return DataInterfaceIndex;
+			}
+
+			return static_cast<int32>(INDEX_NONE);
+		};
+
+		for (TObjectPtr<UComputeDataInterface> DataInterface : AdditionalInputDIs)
+		{
+			const int32 DataInterfaceIndex = AddAdditionalDataInterface(DataInterface);
+
+			if (DataInterfaceIndex != INDEX_NONE)
+			{
+				InputDataInterfaceIndices.Add(DataInterfaceIndex);
+			}
+		}
+
+		for (TObjectPtr<UComputeDataInterface> DataInterface : AdditionalOutputDIs)
+		{
+			const int32 DataInterfaceIndex = AddAdditionalDataInterface(DataInterface);
+
+			if (DataInterfaceIndex != INDEX_NONE)
+			{
+				OutputDataInterfaceIndices.Add(DataInterfaceIndex);
+			}
+		}
 
 		// TODO add graph data interface (graph params). Reference: UOptimusGraphDataInterface.
 		//Element->Graph->DataInterfaces.Add(NewObject< UPCGDataCollectionDataInterface>(InGraph));
 		//Element->Graph->DataInterfaceToBinding.Add(0);
 
 		// TODO once we support cooking for different platforms/configs, don't create the interface if logging is not present.
-		int DebugDIIndex = INDEX_NONE;
 		if (Settings->bPrintShaderDebugValues)
 		{
 			UPCGDebugDataInterface* DebugDI = NewObject<UPCGDebugDataInterface>(ComputeGraph);
 			DebugDI->SetDebugBufferSize(Settings->DebugBufferSize);
 
-			DebugDIIndex = ComputeGraph->DataInterfaces.Num();
+			const int32 DebugDIIndex = ComputeGraph->DataInterfaces.Num();
 			ComputeGraph->DataInterfaces.Add(DebugDI);
+			OutputDataInterfaceIndices.Add_GetRef(DebugDIIndex);
 		}
 
 		// Now that all data interfaces added, create the (trivial) binding mapping. All map to primary binding, index 0.
@@ -867,26 +907,23 @@ void FPCGGraphCompilerGPU::BuildGPUGraphTask(
 			SetupAllInputBindings(InputDataInterfaceIndex);
 		}
 
-		SetupAllInputBindings(KernelDIIndex);
-
 		for (int OutputDataInterfaceIndex : OutputDataInterfaceIndices)
 		{
 			SetupAllOutputBindings(OutputDataInterfaceIndex);
 		}
 
-		if (DebugDIIndex != INDEX_NONE)
-		{
-			SetupAllOutputBindings(DebugDIIndex);
-		}
-
 		{
 			UPCGComputeKernelSource* KernelSource = NewObject<UPCGComputeKernelSource>(KernelWithBindings.Kernel); // is outer to kernel fine?
 			KernelWithBindings.Kernel->KernelSource = KernelSource;
 			KernelSource->EntryPoint = Settings->GetKernelEntryPoint();
 			KernelSource->GroupSize = Settings->GetThreadGroupSize();
-
 			KernelSource->SetSource(Settings->GetCookedKernelSource(ComputeGraph->GlobalAttributeLookupTable));
 
+			if (Settings->bDumpCookedHLSL)
+			{
+				UE_LOG(LogPCG, Log, TEXT("Cooked HLSL:\n%s\n"), *KernelSource->GetSource());
+			}
+
 #if WITH_EDITOR
 			if (PCGGraphCompilerGPU::CVarEnableGPUDebugging.GetValueOnAnyThread())
 			{
diff --git a/Engine/Plugins/PCG/Source/PCG/Public/Compute/Elements/PCGCustomHLSL.h b/Engine/Plugins/PCG/Source/PCG/Public/Compute/Elements/PCGCustomHLSL.h
--- a/Engine/Plugins/PCG/Source/PCG/Public/Compute/Elements/PCGCustomHLSL.h
+++ b/Engine/Plugins/PCG/Source/PCG/Public/Compute/Elements/PCGCustomHLSL.h
@@ -111,15 +111,16 @@ public:
 	virtual TArray<FPCGPinProperties> OutputPinProperties() const override;
 	virtual bool IsInputPinRequiredByExecution(const UPCGPin* InPin) const { return true; }
 #if WITH_EDITOR
-	virtual bool DisplayExecuteOnGPUSetting() const override { return false; }
 	virtual FName GetDefaultNodeName() const override { return FName(TEXT("CustomHLSL")); }
 	virtual FText GetDefaultNodeTitle() const override { return NSLOCTEXT("PCGCustomHLSLElement", "NodeTitle", "Custom HLSL"); }
 	virtual FText GetNodeTooltipText() const override { return NSLOCTEXT("PCGCustomHLSLElement", "NodeTooltip", "Produces a HLSL compute shader which will be executed on the GPU."); }
 	virtual EPCGSettingsType GetType() const override { return EPCGSettingsType::GPU; }
 #endif
 
-	virtual FPCGDataCollectionDesc ComputeOutputPinDataDesc(const UPCGPin* OutputPin, const UPCGDataBinding* Binding) const override;
+	virtual bool IsKernelValid(FPCGContext* InContext = nullptr, bool bQuiet = true) const override;
+	virtual FString GetCookedKernelSource(const TMap<FPCGKernelAttributeKey, int>& GlobalAttributeLookupTable) const override;
 	virtual int ComputeKernelThreadCount(const UPCGDataBinding* Binding) const override;
+	virtual FPCGDataCollectionDesc ComputeOutputPinDataDesc(const UPCGPin* OutputPin, const UPCGDataBinding* Binding) const override;
 
 protected:
 #if WITH_EDITOR
@@ -139,27 +140,14 @@ protected:
 	void UpdateAttributeKeys();
 #endif
 
-public:
-	bool IsKernelValid(FPCGContext* InContext = nullptr, bool bQuiet = true) const;
-	
-	FString GetCookedKernelSource(const TMap<FPCGKernelAttributeKey, int>& GlobalAttributeLookupTable) const;
-	FString GetKernelEntryPoint() const { return TEXT("Main"); }
-	FIntVector GetThreadGroupSize() const { return FIntVector(64, 1, 1); }
-
-	int GetPointCount() const { return PointCount; }
-	int GetFixedThreadCount() const { return FixedThreadCount; }
-
+protected:
 	const UPCGPin* GetInputPin(FName Label) const;
 	const UPCGPin* GetOutputPin(FName Label) const;
 	const UPCGPin* GetFirstInputPin() const;
 	const UPCGPin* GetPointProcessingInputPin() const;
 	const UPCGPin* GetFirstOutputPin() const;
 	const UPCGPin* GetFirstPointOutputPin() const;
-
 	int GetProcessingElemCountForInputPin(const UPCGPin* InputPin, const UPCGDataBinding* Binding) const;
-	virtual const UPCGPin* GetExecutionPin() const { return GetPointProcessingInputPin(); }
-
-protected:
 	bool AreKernelAttributesValid(FPCGContext* InContext, FText* OutErrorText) const;
 
 	/** Will the ThreadCountMultiplier value be applied when calculating the dispatch thread count. */
@@ -184,19 +172,6 @@ protected:
 	UPROPERTY(BlueprintReadWrite, EditAnywhere, DisplayName = "Input Pins", Category = "Thread Count", meta = (EditCondition = "KernelType == EPCGKernelType::Custom && DispatchThreadCount == EPCGDispatchThreadCount::FromProductOfInputPins", EditConditionHides))
 	TArray<FName> ThreadCountInputPinLabels;
 
-public:
-	/** Dump the cooked HLSL into the log after it is generated. */
-	UPROPERTY(EditAnywhere, Category = "Debug")
-	bool bDumpCookedHLSL = false;
-
-	/** Enable use of 'WriteDebugValue(uint Index, float Value)' function in your kernel. Allows you to write float values to a buffer for logging on the CPU. */
-	UPROPERTY(EditAnywhere, Category = "Debug")
-	bool bPrintShaderDebugValues = false;
-
-	/** Size (in number of floats) of the shader debug print buffer. */
-	UPROPERTY(EditAnywhere, Category = "Debug", meta = (EditCondition="bPrintShaderDebugValues", EditConditionHides))
-	int DebugBufferSize = 16;
-
 public:
 	UPROPERTY(BlueprintReadWrite, EditAnywhere, Category = "Settings")
 	TArray<FPCGPinProperties> InputPins = Super::DefaultPointInputPinProperties();
diff --git a/Engine/Plugins/PCG/Source/PCG/Public/Elements/PCGCopyPoints.h b/Engine/Plugins/PCG/Source/PCG/Public/Elements/PCGCopyPoints.h
--- a/Engine/Plugins/PCG/Source/PCG/Public/Elements/PCGCopyPoints.h
+++ b/Engine/Plugins/PCG/Source/PCG/Public/Elements/PCGCopyPoints.h
@@ -50,8 +50,13 @@ public:
 	virtual FText GetDefaultNodeTitle() const override { return NSLOCTEXT("PCGCopyPointSettings", "NodeTitle", "Copy Points"); }
 	virtual FText GetNodeTooltipText() const override;
 	virtual EPCGSettingsType GetType() const override { return EPCGSettingsType::Sampler; }
+	virtual bool DisplayExecuteOnGPUSetting() const override { return true; }
 #endif
-	
+
+	virtual FString GetCookedKernelSource(const TMap<FPCGKernelAttributeKey, int>& GlobalAttributeLookupTable) const override;
+	virtual int ComputeKernelThreadCount(const UPCGDataBinding* Binding) const override;
+	virtual FPCGDataCollectionDesc ComputeOutputPinDataDesc(const UPCGPin* OutputPin, const UPCGDataBinding* Binding) const override;
+	virtual void CreateAdditionalInputDataInterfaces(TArray<TObjectPtr<UComputeDataInterface>>& OutDataInterfaces) const override;
 
 protected:
 	virtual TArray<FPCGPinProperties> InputPinProperties() const override;
diff --git a/Engine/Plugins/PCG/Source/PCG/Public/PCGSettings.h b/Engine/Plugins/PCG/Source/PCG/Public/PCGSettings.h
--- a/Engine/Plugins/PCG/Source/PCG/Public/PCGSettings.h
+++ b/Engine/Plugins/PCG/Source/PCG/Public/PCGSettings.h
@@ -19,10 +19,12 @@ class UPCGPin;
 struct FPCGPinProperties;
 struct FPropertyChangedEvent;
 
+class UPCGComputeGraph;
+class UPCGDataBinding;
 class UPCGGraph;
+class UComputeDataInterface;
 class UPCGNode;
 class UPCGSettings;
-class UPCGDataBinding;
 
 using FPCGSettingsAndCulling = TPair<TSoftObjectPtr<const UPCGSettings>, bool>;
 using FPCGSelectionKeyToSettingsMap = TMap<FPCGSelectionKey, TArray<FPCGSettingsAndCulling>>;
@@ -158,7 +160,6 @@ struct FPCGPreConfiguredSettingsInfo
 #endif // WITH_EDITORONLY_DATA
 };
 
-
 UCLASS(Abstract)
 class PCG_API UPCGSettingsInterface : public UPCGData
 {
@@ -264,9 +265,6 @@ public:
 	
 	bool UseSeed() const { return bUseSeed; }
 
-	/** Whether this node should be executed on the GPU. */
-	virtual bool ShouldExecuteOnGPU() const { return bExecuteOnGPU; }
-
 	// Get the seed, combined with optional PCGComponent seed
 	int GetSeed(const UPCGComponent* InSourceComponent = nullptr) const;
 
@@ -405,24 +403,11 @@ public:
 	// Holds the original settings used to duplicate this object if it was overridden
 	const UPCGSettings* OriginalSettings = nullptr;
 
-	/** Get a list of the attributes read or written by this node. */
-	virtual const TArray<FPCGKernelAttributeKey>& GetKernelAttributeKeys() const { return KernelAttributeKeys; }
-
-	/** Compute a description of all data arriving on InputPin. */
-	FPCGDataCollectionDesc ComputeInputPinDataDesc(const UPCGPin* InputPin, const UPCGDataBinding* Binding) const;
-
-	/** Compute a description of data that will be output from OutputPinLabel/OutputPin. */
-	FPCGDataCollectionDesc ComputeOutputPinDataDesc(const FName& OutputPinLabel, const UPCGDataBinding* Binding) const;
-	virtual FPCGDataCollectionDesc ComputeOutputPinDataDesc(const UPCGPin* OutputPin, const UPCGDataBinding* Binding) const;
-
-	/** Compute how many threads should be dispatched to execute this node on the GPU. */
-	virtual int ComputeKernelThreadCount(const UPCGDataBinding* Binding) const { return 0; };
-
-protected:
 	// Returns an array of all the input pin properties. You should not add manually a "params" pin, it is handled automatically by FillOverridableParamsPins
 	virtual TArray<FPCGPinProperties> InputPinProperties() const;
 	virtual TArray<FPCGPinProperties> OutputPinProperties() const;
 
+protected:
 	// BP version since EPCGDataType is uint32 and BP only supports uint8 enums.
 	/** Bitwise union of the allowed types of each incident edge on pin. Returns None type if no common bits, or no edges. Use the BP function helpers to extract the types from the result. */
 	UFUNCTION(BlueprintCallable, Category = "Settings|DynamicPins", DisplayName = "Get Type Union Of Incident Edges")
@@ -459,16 +444,8 @@ protected:
 
 	// Passthrough for the simpler method, to avoid modifying the child settings already overriding this method.
 	virtual bool CanEditChange(const FProperty* InProperty) const override;
-
-	/** Whether to display GPU execution option in node settings UI. */
-	UFUNCTION()
-	virtual bool DisplayExecuteOnGPUSetting() const { return ShouldExecuteOnGPU(); }
 #endif
 
-	/** Whether this node should be executed on the GPU. */
-	UPROPERTY(BlueprintReadWrite, EditAnywhere, Category = GPU, meta = (EditCondition = "DisplayExecuteOnGPUSetting()", EditConditionHides, HideEditConditionToggle))
-	bool bExecuteOnGPU = false;
-
 	// By default, settings won't use a seed. Set this bool to true in the child ctor to allow edition and use it.
 	UPROPERTY(VisibleAnywhere, Transient, Category = Settings, meta = (EditCondition = false, EditConditionHides))
 	bool bUseSeed = false;
@@ -535,6 +512,56 @@ protected:
 	UPROPERTY()
 	TArray<FPCGSettingsOverridableParam> CachedOverridableParams;
 
+	// GPU section
+public:
+	/** Whether this node should be executed on the GPU. */
+	virtual bool ShouldExecuteOnGPU() const { return bExecuteOnGPU; }
+
+	virtual bool IsKernelValid(FPCGContext* InContext = nullptr, bool bQuiet = true) const { return true; }
+	virtual FString GetCookedKernelSource(const TMap<FPCGKernelAttributeKey, int>& GlobalAttributeLookupTable) const { return TEXT(""); }
+	virtual FString GetKernelEntryPoint() const { return TEXT("Main"); }
+	virtual FIntVector GetThreadGroupSize() const { return FIntVector(64, 1, 1); }
+
+	/** Get a list of the attributes read or written by this node. */
+	virtual const TArray<FPCGKernelAttributeKey>& GetKernelAttributeKeys() const { return KernelAttributeKeys; }
+
+	/** Compute how many threads should be dispatched to execute this node on the GPU. */
+	virtual int ComputeKernelThreadCount(const UPCGDataBinding* Binding) const { return 0; };
+
+	/** Compute a description of all data arriving on InputPin. */
+	FPCGDataCollectionDesc ComputeInputPinDataDesc(const UPCGPin* InputPin, const UPCGDataBinding* Binding) const;
+
+	/** Compute a description of data that will be output from OutputPinLabel/OutputPin. */
+	FPCGDataCollectionDesc ComputeOutputPinDataDesc(const FName& OutputPinLabel, const UPCGDataBinding* Binding) const;
+	virtual FPCGDataCollectionDesc ComputeOutputPinDataDesc(const UPCGPin* OutputPin, const UPCGDataBinding* Binding) const;
+
+	/** Create additional data interfaces to marshal any data required by this settings. */
+	virtual void CreateAdditionalInputDataInterfaces(TArray<TObjectPtr<UComputeDataInterface>>& OutDataInterfaces) const {}
+	virtual void CreateAdditionalOutputDataInterfaces(TArray<TObjectPtr<UComputeDataInterface>>& OutDataInterfaces) const {}
+
+#if WITH_EDITOR
+	/** Whether to display GPU execution option in node settings UI. */
+	UFUNCTION()
+	virtual bool DisplayExecuteOnGPUSetting() const { return false; }
+#endif
+
+public:
+	/** Whether this node should be executed on the GPU. */
+	UPROPERTY(BlueprintReadWrite, EditAnywhere, Category = GPU, meta = (EditCondition = "DisplayExecuteOnGPUSetting()", EditConditionHides, HideEditConditionToggle))
+	bool bExecuteOnGPU = false;
+
+	/** Dump the cooked HLSL into the log after it is generated. */
+	UPROPERTY(EditAnywhere, Category = "GPU", meta = (EditCondition = "bExecuteOnGPU", EditConditionHides))
+	bool bDumpCookedHLSL = false;
+
+	/** Enable use of 'WriteDebugValue(uint Index, float Value)' function in your kernel. Allows you to write float values to a buffer for logging on the CPU. */
+	UPROPERTY(EditAnywhere, Category = "GPU", meta = (EditCondition = "bExecuteOnGPU", EditConditionHides))
+	bool bPrintShaderDebugValues = false;
+
+	/** Size (in number of floats) of the shader debug print buffer. */
+	UPROPERTY(EditAnywhere, Category = "GPU", meta = (EditCondition="bExecuteOnGPU && bPrintShaderDebugValues", EditConditionHides))
+	int DebugBufferSize = 16;
+
 private:
 	/** Calculate Crc for these settings and save it. */
 	void CacheCrc();
