Fixed wrong assumption about the output task node, which depending on timing might not be fully set yet.

#rnx

diff --git a/Engine/Plugins/PCG/Source/PCG/Private/Graph/PCGGraphCompiler.cpp b/Engine/Plugins/PCG/Source/PCG/Private/Graph/PCGGraphCompiler.cpp
--- a/Engine/Plugins/PCG/Source/PCG/Private/Graph/PCGGraphCompiler.cpp
+++ b/Engine/Plugins/PCG/Source/PCG/Private/Graph/PCGGraphCompiler.cpp
@@ -116,8 +116,6 @@ TArray<FPCGGraphTask> FPCGGraphCompiler::CompileGraph(UPCGGraph* InGraph, FPCGTa
 				return Subtask.Node == SubgraphOutputNode;
 				});
 
-			check(InputNodeTask && OutputNodeTask);
-
 			// Build pre-task
 			FPCGGraphTask& PreTask = CompiledTasks.Emplace_GetRef();
 			PreTask.Node = Node;
@@ -136,7 +134,7 @@ TArray<FPCGGraphTask> FPCGGraphCompiler::CompileGraph(UPCGGraph* InGraph, FPCGTa
 					{
 						// Implementation note: conceptually, the non-param inputs need to be connected only to the input node task.
 						// However, because of the static/dynamic culling which happen on the subgraph node, we need to have the connections on the pretask too (e.g. the subgraph node) even if they don't provide data.
-						if (bIsParamPin)
+						if (bIsParamPin || !InputNodeTask)
 						{
 							PreTask.Inputs.Emplace(IdMapping[InboundEdge->InputPin->Node], InboundEdge->InputPin->Properties, InboundEdge->OutputPin->Properties, /*bInProvideData=*/true);
 						}
@@ -154,7 +152,10 @@ TArray<FPCGGraphTask> FPCGGraphCompiler::CompileGraph(UPCGGraph* InGraph, FPCGTa
 			}
 
 			// Add pre-task as input to subgraph input node task, without data dependency
-			InputNodeTask->Inputs.Emplace(PreId,/*InUpstreamPin=*/FPCGGraphTaskInput::NoPin, /*InDownstreamPin=*/FPCGGraphTaskInput::NoPin, /*bInProvideData=*/false);
+			if (InputNodeTask)
+			{
+				InputNodeTask->Inputs.Emplace(PreId,/*InUpstreamPin=*/FPCGGraphTaskInput::NoPin, /*InDownstreamPin=*/FPCGGraphTaskInput::NoPin, /*bInProvideData=*/false);
+			}
 
 			// Hook nodes to the PreTask if they require so.
 			// Only do it for nodes that are directly under the subgraph, not in subsequent subgraphs.
@@ -195,7 +196,10 @@ TArray<FPCGGraphTask> FPCGGraphCompiler::CompileGraph(UPCGGraph* InGraph, FPCGTa
 			PostTask.Inputs.Emplace(PreId, /*InUpstreamPin=*/FPCGGraphTaskInput::NoPin, /*InDownstreamPin=*/FPCGGraphTaskInput::NoPin, /*bInProvideData=*/false);
 
 			// Add subgraph output node task as input to the post-task
-			PostTask.Inputs.Emplace(OutputNodeTask->NodeId);
+			if (OutputNodeTask)
+			{
+				PostTask.Inputs.Emplace(OutputNodeTask->NodeId);
+			}
 
 			check(!IdMapping.Contains(Node));
 			IdMapping.Add(Node, PostId);
