[PCG] UPCGBuilderSettings + Build PCG menu & dialog

- New UPCGBuilderSettings asset allows configuring PCG World Partition Builder  instead of using switches
- Switches can still be used to override the asset when needed
- Build PCG menu will prompt for a UPCGBuilderSettings to use but will default to a new Project Settings config that allows setting a default UPCGBuilderSettings asset to use

#rnx

diff --git a/Engine/Plugins/PCG/Source/PCGEditor/PCGEditor.Build.cs b/Engine/Plugins/PCG/Source/PCGEditor/PCGEditor.Build.cs
--- a/Engine/Plugins/PCG/Source/PCGEditor/PCGEditor.Build.cs
+++ b/Engine/Plugins/PCG/Source/PCGEditor/PCGEditor.Build.cs
@@ -39,6 +39,7 @@ namespace UnrealBuildTool.Rules
 					"ContentBrowser",
 					"DesktopWidgets",
 					"DetailCustomizations",
+					"DeveloperSettings",
 					"EditorFramework",
 					"EditorScriptingUtilities",
 					"EditorStyle",
diff --git a/Engine/Plugins/PCG/Source/PCGEditor/Private/AssetDefinitions/AssetDefinition_PCGAsset.cpp b/Engine/Plugins/PCG/Source/PCGEditor/Private/AssetDefinitions/AssetDefinition_PCGAsset.cpp
--- a/Engine/Plugins/PCG/Source/PCGEditor/Private/AssetDefinitions/AssetDefinition_PCGAsset.cpp
+++ b/Engine/Plugins/PCG/Source/PCGEditor/Private/AssetDefinitions/AssetDefinition_PCGAsset.cpp
@@ -3,8 +3,7 @@
 #include "AssetDefinition_PCGAsset.h"
 
 #include "PCGDataAsset.h"
-
-#include "Misc/AssetCategoryPath.h"
+#include "PCGEditorCommon.h"
 
 #define LOCTEXT_NAMESPACE "AssetDefinition_PCGAsset"
 
@@ -25,7 +24,7 @@ TSoftClassPtr<UObject> UAssetDefinition_PCGAsset::GetAssetClass() const
 
 TConstArrayView<FAssetCategoryPath> UAssetDefinition_PCGAsset::GetAssetCategories() const
 {
-	static const FAssetCategoryPath Categories[] = { FAssetCategoryPath(LOCTEXT("PCGCategory", "PCG")) };
+	static const FAssetCategoryPath Categories[] = { FPCGEditorCommon::PCGAssetCategoryPath };
 	return Categories;
 }
 
diff --git a/Engine/Plugins/PCG/Source/PCGEditor/Private/AssetDefinitions/AssetDefinition_PCGBuilderSettings.cpp b/Engine/Plugins/PCG/Source/PCGEditor/Private/AssetDefinitions/AssetDefinition_PCGBuilderSettings.cpp
new file mode 100644
--- /dev/null
+++ b/Engine/Plugins/PCG/Source/PCGEditor/Private/AssetDefinitions/AssetDefinition_PCGBuilderSettings.cpp
@@ -0,0 +1,32 @@
+// Copyright Epic Games, Inc. All Rights Reserved.
+
+#include "AssetDefinition_PCGBuilderSettings.h"
+
+#include "PCGEditorCommon.h"
+
+#include "WorldPartitionBuilder/PCGBuilderSettings.h"
+
+#define LOCTEXT_NAMESPACE "AssetDefinition_PCGBuilderSettings"
+
+FText UAssetDefinition_PCGBuilderSettings::GetAssetDisplayName() const
+{
+	return LOCTEXT("DisplayName", "PCG Builder Settings");
+}
+
+FLinearColor UAssetDefinition_PCGBuilderSettings::GetAssetColor() const
+{
+	return FColor::Purple;
+}
+
+TSoftClassPtr<UObject> UAssetDefinition_PCGBuilderSettings::GetAssetClass() const
+{
+	return UPCGBuilderSettings::StaticClass();
+}
+
+TConstArrayView<FAssetCategoryPath> UAssetDefinition_PCGBuilderSettings::GetAssetCategories() const
+{
+	static const FAssetCategoryPath Categories[] = { FPCGEditorCommon::PCGAdvancedAssetCategoryPath };
+	return Categories;
+}
+
+#undef LOCTEXT_NAMESPACE
diff --git a/Engine/Plugins/PCG/Source/PCGEditor/Private/AssetDefinitions/AssetDefinition_PCGBuilderSettings.h b/Engine/Plugins/PCG/Source/PCGEditor/Private/AssetDefinitions/AssetDefinition_PCGBuilderSettings.h
new file mode 100644
--- /dev/null
+++ b/Engine/Plugins/PCG/Source/PCGEditor/Private/AssetDefinitions/AssetDefinition_PCGBuilderSettings.h
@@ -0,0 +1,20 @@
+// Copyright Epic Games, Inc. All Rights Reserved.
+
+#pragma once
+
+#include "AssetDefinitionDefault.h"
+
+#include "AssetDefinition_PCGBuilderSettings.generated.h"
+
+UCLASS()
+class UAssetDefinition_PCGBuilderSettings : public UAssetDefinitionDefault
+{
+	GENERATED_BODY()
+
+public:
+	virtual FText GetAssetDisplayName() const override;
+	virtual FLinearColor GetAssetColor() const override;
+	virtual TSoftClassPtr<UObject> GetAssetClass() const override;
+	virtual TConstArrayView<FAssetCategoryPath> GetAssetCategories() const override;
+	virtual FAssetSupportResponse CanLocalize(const FAssetData& InAsset) const override { return FAssetSupportResponse::NotSupported(); }
+};
\ No newline at end of file
diff --git a/Engine/Plugins/PCG/Source/PCGEditor/Private/AssetDefinitions/AssetDefinition_PCGGraphInterface.cpp b/Engine/Plugins/PCG/Source/PCGEditor/Private/AssetDefinitions/AssetDefinition_PCGGraphInterface.cpp
--- a/Engine/Plugins/PCG/Source/PCGEditor/Private/AssetDefinitions/AssetDefinition_PCGGraphInterface.cpp
+++ b/Engine/Plugins/PCG/Source/PCGEditor/Private/AssetDefinitions/AssetDefinition_PCGGraphInterface.cpp
@@ -4,6 +4,7 @@
 
 #include "PCGGraph.h"
 
+#include "PCGEditorCommon.h"
 #include "PCGEditorStyle.h"
 #include "PCGGraphFactory.h"
 
@@ -13,7 +14,6 @@
 #include "ToolMenuEntry.h"
 #include "ToolMenuSection.h"
 #include "ToolMenus.h"
-#include "Misc/AssetCategoryPath.h"
 #include "Misc/DelayedAutoRegister.h"
 
 #define LOCTEXT_NAMESPACE "AssetDefinition_PCGGraphInterface"
@@ -35,7 +35,7 @@ TSoftClassPtr<UObject> UAssetDefinition_PCGGraphInterface::GetAssetClass() const
 
 TConstArrayView<FAssetCategoryPath> UAssetDefinition_PCGGraphInterface::GetAssetCategories() const
 {
-	static const FAssetCategoryPath Categories[] = { FAssetCategoryPath(LOCTEXT("PCGCategory", "PCG")) };
+	static const FAssetCategoryPath Categories[] = { FPCGEditorCommon::PCGAssetCategoryPath };
 	return Categories;
 }
 
diff --git a/Engine/Plugins/PCG/Source/PCGEditor/Private/AssetDefinitions/AssetDefinition_PCGSettings.cpp b/Engine/Plugins/PCG/Source/PCGEditor/Private/AssetDefinitions/AssetDefinition_PCGSettings.cpp
--- a/Engine/Plugins/PCG/Source/PCGEditor/Private/AssetDefinitions/AssetDefinition_PCGSettings.cpp
+++ b/Engine/Plugins/PCG/Source/PCGEditor/Private/AssetDefinitions/AssetDefinition_PCGSettings.cpp
@@ -2,10 +2,9 @@
 
 #include "AssetDefinition_PCGSettings.h"
 
+#include "PCGEditorCommon.h"
 #include "PCGSettings.h"
 
-#include "Misc/AssetCategoryPath.h"
-
 #define LOCTEXT_NAMESPACE "AssetDefinition_PCGSettings"
 
 FText UAssetDefinition_PCGSettings::GetAssetDisplayName() const
@@ -25,7 +24,7 @@ TSoftClassPtr<UObject> UAssetDefinition_PCGSettings::GetAssetClass() const
 
 TConstArrayView<FAssetCategoryPath> UAssetDefinition_PCGSettings::GetAssetCategories() const
 {
-	static const FAssetCategoryPath Categories[] = { FAssetCategoryPath(LOCTEXT("PCGCategory", "PCG")) }; 
+	static const FAssetCategoryPath Categories[] = { FPCGEditorCommon::PCGAssetCategoryPath };
 	return Categories;
 }
 
diff --git a/Engine/Plugins/PCG/Source/PCGEditor/Private/Factories/PCGBuilderSettingsFactory.cpp b/Engine/Plugins/PCG/Source/PCGEditor/Private/Factories/PCGBuilderSettingsFactory.cpp
new file mode 100644
--- /dev/null
+++ b/Engine/Plugins/PCG/Source/PCGEditor/Private/Factories/PCGBuilderSettingsFactory.cpp
@@ -0,0 +1,21 @@
+// Copyright Epic Games, Inc. All Rights Reserved.
+
+#include "PCGBuilderSettingsFactory.h"
+#include "WorldPartitionBuilder/PCGBuilderSettings.h"
+
+UPCGBuilderSettingsFactory::UPCGBuilderSettingsFactory()
+{
+	SupportedClass = UPCGBuilderSettings::StaticClass();
+	bCreateNew = true;
+	bEditAfterNew = true;
+}
+
+UObject* UPCGBuilderSettingsFactory::FactoryCreateNew(UClass* InClass, UObject* InParent, FName InName, EObjectFlags Flags, UObject* Context, FFeedbackContext* Warn)
+{
+	return NewObject<UPCGBuilderSettings>(InParent, InClass, InName, Flags);
+}
+
+bool UPCGBuilderSettingsFactory::ShouldShowInNewMenu() const
+{
+	return true;
+}
\ No newline at end of file
diff --git a/Engine/Plugins/PCG/Source/PCGEditor/Private/Factories/PCGBuilderSettingsFactory.h b/Engine/Plugins/PCG/Source/PCGEditor/Private/Factories/PCGBuilderSettingsFactory.h
new file mode 100644
--- /dev/null
+++ b/Engine/Plugins/PCG/Source/PCGEditor/Private/Factories/PCGBuilderSettingsFactory.h
@@ -0,0 +1,20 @@
+// Copyright Epic Games, Inc. All Rights Reserved.
+
+#pragma once
+
+#include "Factories/Factory.h"
+
+#include "PCGBuilderSettingsFactory.generated.h"
+
+UCLASS(hidecategories = Object)
+class UPCGBuilderSettingsFactory : public UFactory
+{
+	GENERATED_BODY()
+
+public:
+	UPCGBuilderSettingsFactory();
+
+	//~UFactory interface
+	virtual UObject* FactoryCreateNew(UClass* InClass, UObject* InParent, FName InName, EObjectFlags Flags, UObject* Context, FFeedbackContext* Warn) override;
+	virtual bool ShouldShowInNewMenu() const override;
+};
\ No newline at end of file
diff --git a/Engine/Plugins/PCG/Source/PCGEditor/Private/PCGEditorCommon.h b/Engine/Plugins/PCG/Source/PCGEditor/Private/PCGEditorCommon.h
--- a/Engine/Plugins/PCG/Source/PCGEditor/Private/PCGEditorCommon.h
+++ b/Engine/Plugins/PCG/Source/PCGEditor/Private/PCGEditorCommon.h
@@ -3,9 +3,12 @@
 #pragma once
 
 #include "HAL/IConsoleManager.h"
+#include "Misc/AssetCategoryPath.h"
 
 #include "PCGEditorCommon.generated.h"
 
+#define LOCTEXT_NAMESPACE "PCGEditorCommon"
+
 UENUM()
 enum class EPCGElementType : uint8
 {
@@ -47,4 +50,9 @@ namespace FPCGEditorCommon
 	const FName PrimitiveDataType = FName(TEXT("Primitive Data"));
 
 	extern TAutoConsoleVariable<bool> CVarShowAdvancedAttributesFields;
+
+	const FAssetCategoryPath PCGAssetCategoryPath(LOCTEXT("PCGCategory", "PCG"));
+	const FAssetCategoryPath PCGAdvancedAssetCategoryPath(FPCGEditorCommon::PCGAssetCategoryPath, LOCTEXT("PCGAvancedCategory", "Advanced"));
 };
+
+#undef LOCTEXT_NAMESPACE
diff --git a/Engine/Plugins/PCG/Source/PCGEditor/Private/PCGEditorModule.cpp b/Engine/Plugins/PCG/Source/PCGEditor/Private/PCGEditorModule.cpp
--- a/Engine/Plugins/PCG/Source/PCGEditor/Private/PCGEditorModule.cpp
+++ b/Engine/Plugins/PCG/Source/PCGEditor/Private/PCGEditorModule.cpp
@@ -21,10 +21,12 @@
 #include "DataVisualizations/PCGSpatialDataVisualization.h"
 #include "DataVisualizations/PCGSplineDataVisualization.h"
 #include "Grid/PCGPartitionActor.h"
+#include "WorldPartitionBuilder/PCGWorldPartitionBuilder.h"
 
 #include "ContentBrowserMenuContexts.h"
 #include "ContentBrowserModule.h"
 #include "Editor.h"
+#include "EditorBuildUtils.h"
 #include "EditorModeManager.h"
 #include "EditorModes.h"
 #include "IContentBrowserSingleton.h"
@@ -49,6 +51,11 @@
 
 #define LOCTEXT_NAMESPACE "FPCGEditorModule"
 
+namespace PCGEditorModule
+{
+	static const FName PCGBuildType(TEXT("PCG"));
+}
+
 void FPCGEditorModule::StartupModule()
 {
 	RegisterDetailsCustomizations();
@@ -72,6 +79,13 @@ void FPCGEditorModule::StartupModule()
 	{
 		LevelEditorModule.OnLevelEditorCreated().AddRaw(this, &FPCGEditorModule::OnLevelEditorCreated);
 	}
+
+	FEditorBuildUtils::RegisterCustomBuildType(PCGEditorModule::PCGBuildType,
+		FCanDoEditorBuildDelegate::CreateStatic(&UPCGWorldPartitionBuilder::CanBuild),
+		FDoEditorBuildDelegate::CreateStatic(&UPCGWorldPartitionBuilder::Build),
+		/*BuildAllExtensionPoint*/NAME_None,
+		/*MenuEntryLabel*/LOCTEXT("BuildPCG", "Build PCG"),
+		/*MenuSectionLabel*/LOCTEXT("PCG", "PCG"));
 }
 
 void FPCGEditorModule::ShutdownModule()
@@ -105,6 +119,8 @@ void FPCGEditorModule::ShutdownModule()
 			FirstLevelEditor->GetEditorModeManager().OnEditorModeIDChanged().RemoveAll(this);
 		}
 	}
+
+	FEditorBuildUtils::UnregisterCustomBuildType(PCGEditorModule::PCGBuildType);
 }
 
 void FPCGEditorModule::OnLevelEditorCreated(TSharedPtr<ILevelEditor> InLevelEditor)
diff --git a/Engine/Plugins/PCG/Source/PCGEditor/Private/PCGEditorSettings.cpp b/Engine/Plugins/PCG/Source/PCGEditor/Private/PCGEditorSettings.cpp
--- a/Engine/Plugins/PCG/Source/PCGEditor/Private/PCGEditorSettings.cpp
+++ b/Engine/Plugins/PCG/Source/PCGEditor/Private/PCGEditorSettings.cpp
@@ -6,6 +6,11 @@
 #include "EdGraph/EdGraphPin.h"
 #include "PCGSettings.h"
 
+UPCGEditorProjectSettings::UPCGEditorProjectSettings(const FObjectInitializer& ObjectInitializer)
+	: Super(ObjectInitializer)
+{
+}
+
 UPCGEditorSettings::UPCGEditorSettings(const FObjectInitializer& ObjectInitializer)
 	: Super(ObjectInitializer)
 {
diff --git a/Engine/Plugins/PCG/Source/PCGEditor/Private/WorldPartitionBuilder/PCGBuilderSettings.cpp b/Engine/Plugins/PCG/Source/PCGEditor/Private/WorldPartitionBuilder/PCGBuilderSettings.cpp
new file mode 100644
--- /dev/null
+++ b/Engine/Plugins/PCG/Source/PCGEditor/Private/WorldPartitionBuilder/PCGBuilderSettings.cpp
@@ -0,0 +1,9 @@
+// Copyright Epic Games, Inc. All Rights Reserved.
+
+#include "PCGBuilderSettings.h"
+
+UPCGBuilderSettings::UPCGBuilderSettings()
+{
+	// By default include this mode
+	EditingModes.Add(EPCGEditorDirtyMode::LoadAsPreview);
+}
diff --git a/Engine/Plugins/PCG/Source/PCGEditor/Private/WorldPartitionBuilder/PCGBuilderSettings.h b/Engine/Plugins/PCG/Source/PCGEditor/Private/WorldPartitionBuilder/PCGBuilderSettings.h
new file mode 100644
--- /dev/null
+++ b/Engine/Plugins/PCG/Source/PCGEditor/Private/WorldPartitionBuilder/PCGBuilderSettings.h
@@ -0,0 +1,39 @@
+// Copyright Epic Games, Inc. All Rights Reserved.
+
+#pragma once
+
+#include "PCGCommon.h"
+#include "UObject/Object.h"
+
+#include "PCGBuilderSettings.generated.h"
+
+class UPCGGraphInterface;
+
+UCLASS(hidecategories = Object)
+class UPCGBuilderSettings : public UObject
+{
+	GENERATED_BODY()
+
+public:
+	UPCGBuilderSettings();
+
+	/** Graphs that will be generated in the specified order. When left empty, all Graphs will get generated in the order they are discovered. */
+	UPROPERTY(EditAnywhere, Category = Build)
+	TArray<TSoftObjectPtr<UPCGGraphInterface>> Graphs;
+
+	/** Include components which have the following editing mode. */
+	UPROPERTY(EditAnywhere, Category = Build)
+	TArray<EPCGEditorDirtyMode> EditingModes;
+
+	/** Filter generated components by owner actor name. */
+	UPROPERTY(EditAnywhere, Category = Filter)
+	TArray<FString>  FilterByActorNames;
+
+	/** Call generate on each component and wait until completion and any async processes before generating the next. */
+	UPROPERTY(EditAnywhere, Category = Build, meta = (DisplayName = "One Component at a Time"))
+	bool bOneComponentAtATime = false;
+
+	/** Submit dirty files even if errors occurred during generation. */
+	UPROPERTY(EditAnywhere, Category = Advanced)
+	bool bIgnoreGenerationErrors = false;
+};
\ No newline at end of file
diff --git a/Engine/Plugins/PCG/Source/PCGEditor/Private/WorldPartitionBuilder/PCGWorldPartitionBuilder.cpp b/Engine/Plugins/PCG/Source/PCGEditor/Private/WorldPartitionBuilder/PCGWorldPartitionBuilder.cpp
--- a/Engine/Plugins/PCG/Source/PCGEditor/Private/WorldPartitionBuilder/PCGWorldPartitionBuilder.cpp
+++ b/Engine/Plugins/PCG/Source/PCGEditor/Private/WorldPartitionBuilder/PCGWorldPartitionBuilder.cpp
@@ -2,12 +2,18 @@
 
 #include "WorldPartitionBuilder/PCGWorldPartitionBuilder.h"
 
+#include "PCGBuilderSettings.h"
 #include "PCGComponent.h"
+#include "PCGEditorSettings.h"
 #include "PCGGraph.h"
+#include "SPCGBuilderDialog.h"
 
 #include "AssetCompilingManager.h"
 #include "Editor.h"
 #include "FileHelpers.h"
+
+#include "Framework/Application/SlateApplication.h"
+#include "Interfaces/IMainFrameModule.h"
 #include "Misc/OutputDevice.h"
 #include "Modules/ModuleManager.h"
 #include "UObject/Linker.h"
@@ -16,6 +22,8 @@
 
 DEFINE_LOG_CATEGORY_STATIC(LogPCGWorldPartitionBuilder, All, All);
 
+#define LOCTEXT_NAMESPACE "PCGWorldPartitionBulder"
+
 namespace PCGWorldPartitionBuilder
 {
 	/** Grab all original PCG components in world, applying given filter. */
@@ -38,18 +46,20 @@ namespace PCGWorldPartitionBuilder
 	void WaitForAllAsyncEditorProcesses(UWorld* InWorld);
 
 	/** Runs builder on current editor world. */
-	void Build(const TArray<FString>& Args);
+	bool Build(UWorld* InWorld, const TArray<FString>& Args);
 
 	static FAutoConsoleCommand CommandBuildComponents(
 		TEXT("pcg.BuildComponents"),
-		TEXT("Runs PCG world builder on PCG components in current world. Arguments (multiple values separated with ';'):\n"
-			"\t[-IncludeGraphNames=PCG_GraphA;PCG_GraphB]\n"
-			"\t[-GenerateComponentEditingModeNormal]\n"
-			"\t[-GenerateComponentEditingModePreview]\n"
-			"\t[-IgnoreGenerationErrors]\n"
-			"\t[-IncludeActorIDs=MyActor1_UAID1234678;MyActor2_UAID1234678]\n"
-			"\t[-OneComponentAtATime]\n"),
-		FConsoleCommandWithArgsDelegate::CreateLambda([](const TArray<FString>& Args) { PCGWorldPartitionBuilder::Build(Args); }));
+		TEXT("Runs PCG world builder on PCG components in current world. Arguments (multiple values separated with ';'), Flags can optionally have a specified '=true or false', true if none provided:\n"
+			"\t[-PCGBuilderSettings=/Game/Path/BuilderSettingsAssetName]\n"
+			"\t[-IncludeGraphNames=PCG_GraphA;PCG_GraphB] (default: all graphs)\n"
+			"\t[-GenerateComponentEditingModeLoadAsPreview] (default: true)\n"
+			"\t[-GenerateComponentEditingModeNormal] (default: false)\n"
+			"\t[-GenerateComponentEditingModePreview] (default: false)\n"
+			"\t[-IgnoreGenerationErrors] (default: false)\n"
+			"\t[-IncludeActorIDs=MyActor1_UAID1234678;MyActor2_UAID1234678] (default: all actors)\n"
+			"\t[-OneComponentAtATime] (default: false)\n"),
+		FConsoleCommandWithArgsDelegate::CreateLambda([](const TArray<FString>& Args) { PCGWorldPartitionBuilder::Build(GEditor ? GEditor->GetEditorWorldContext().World() : nullptr, Args); }));
 };
 
 /** Output device to capture presence of errors during generation. */
@@ -90,39 +100,202 @@ private:
 	std::atomic<bool> bErrorOccurred = false;
 };
 
-UPCGWorldPartitionBuilder::UPCGWorldPartitionBuilder(const FObjectInitializer& ObjectInitializer)
-	: Super(ObjectInitializer)
+FPCGWorldPartitionBuilderArgs FPCGWorldPartitionBuilderArgs::InitializeFrom(const FPCGWorldPartitionCommandlineArgs& CommandlineArgs)
 {
-	if (!HasAnyFlags(RF_ClassDefaultObject))
+	FPCGWorldPartitionBuilderArgs BuilderArgs;
+	TSoftObjectPtr<UPCGBuilderSettings> SettingsPath = CommandlineArgs.SettingAsset.IsSet() ? CommandlineArgs.SettingAsset.GetValue() : GetDefault<UPCGEditorProjectSettings>()->DefaultBuilderSetting;
+	
+	UE_LOG(LogPCGWorldPartitionBuilder, Display, TEXT("Initialize Builder Args"));
+
+	bool bUsingBuilderSettings = false;
+	// Load Settings Asset if any exists
+	if (SettingsPath.IsValid() || SettingsPath.IsPending())
+	{
+		if (UPCGBuilderSettings* BuilderSettings = Cast<UPCGBuilderSettings>(SettingsPath.LoadSynchronous()))
+		{
+			BuilderArgs.bGenerateEditingModeLoadAsPreviewComponents = BuilderSettings->EditingModes.Contains(EPCGEditorDirtyMode::LoadAsPreview);
+			BuilderArgs.bGenerateEditingModeNormalComponents = BuilderSettings->EditingModes.Contains(EPCGEditorDirtyMode::Normal);
+			BuilderArgs.bGenerateEditingModePreviewComponents = BuilderSettings->EditingModes.Contains(EPCGEditorDirtyMode::Preview);
+			BuilderArgs.bIgnoreGenerationErrors = BuilderSettings->bIgnoreGenerationErrors;
+			BuilderArgs.bOneComponentAtATime = BuilderSettings->bOneComponentAtATime;
+			BuilderArgs.IncludeActorIDs = BuilderSettings->FilterByActorNames;
+			BuilderArgs.IncludeGraphs = BuilderSettings->Graphs;
+
+			UE_LOG(LogPCGWorldPartitionBuilder, Display, TEXT("Loaded PCGBuilderSettings '%s'"), *SettingsPath.ToString());
+			bUsingBuilderSettings = true;
+		}
+		else
+		{
+			UE_LOG(LogPCGWorldPartitionBuilder, Warning, TEXT("Failed to load PCGBuilderSettings '%s'"), *SettingsPath.ToString());
+		}
+	}
+	
+	if (!bUsingBuilderSettings)
 	{
-		bGenerateEditingModeNormalComponents = HasParam("GenerateComponentEditingModeNormal");
+		UE_LOG(LogPCGWorldPartitionBuilder, Display, TEXT("Not using any PCGBuilderSettings"));
+	}
 
-		bGenerateEditingModePreviewComponents = HasParam("GenerateComponentEditingModePreview");
+	// Then apply Commandline since it has priority
+	const TCHAR* FromCommandline = TEXT("Commandline");
+	const TCHAR* FromDefault = bUsingBuilderSettings ? TEXT("PCGBuilderSettings") : TEXT("Default");
+		
+	// Init 'bGenerateEditingModeLoadAsPreviewComponents'
+	bool bFromCommandline = CommandlineArgs.bGenerateEditingModeLoadAsPreviewComponents.IsSet();
+	if (bFromCommandline)
+	{
+		BuilderArgs.bGenerateEditingModeLoadAsPreviewComponents = CommandlineArgs.bGenerateEditingModeLoadAsPreviewComponents.GetValue();
+	}
 
-		bOneComponentAtATime = HasParam("OneComponentAtATime");
+	UE_LOG(LogPCGWorldPartitionBuilder, Display, TEXT("bGenerateEditingModeLoadAsPreviewComponents : '%d' (%s)"), 
+		BuilderArgs.bGenerateEditingModeLoadAsPreviewComponents ? 1 : 0,
+		bFromCommandline ? FromCommandline : FromDefault);
+	
+	// Init 'bGenerateEditingModeNormalComponents'
+	bFromCommandline = CommandlineArgs.bGenerateEditingModeNormalComponents.IsSet();
+	if (bFromCommandline)
+	{
+		BuilderArgs.bGenerateEditingModeNormalComponents = CommandlineArgs.bGenerateEditingModeNormalComponents.GetValue();
+	}
+
+	UE_LOG(LogPCGWorldPartitionBuilder, Display, TEXT("bGenerateEditingModeNormalComponents : '%d' (%s)"),
+		BuilderArgs.bGenerateEditingModeNormalComponents ? 1 : 0,
+		bFromCommandline ? FromCommandline : FromDefault);
+	
+	// Init 'bGenerateEditingModePreviewComponents'
+	bFromCommandline = CommandlineArgs.bGenerateEditingModePreviewComponents.IsSet();
+	if (bFromCommandline)
+	{
+		BuilderArgs.bGenerateEditingModePreviewComponents = CommandlineArgs.bGenerateEditingModePreviewComponents.GetValue();
+	}
 
-		bIgnoreGenerationErrors = HasParam("IgnoreGenerationErrors");
+	UE_LOG(LogPCGWorldPartitionBuilder, Display, TEXT("bGenerateEditingModePreviewComponents : '%d' (%s)"),
+		BuilderArgs.bGenerateEditingModePreviewComponents ? 1 : 0,
+		bFromCommandline ? FromCommandline : FromDefault);
 
+
+	// Init 'bOneComponentAtATime'
+	bFromCommandline = CommandlineArgs.bOneComponentAtATime.IsSet();
+	if (bFromCommandline)
+	{
+		BuilderArgs.bOneComponentAtATime = CommandlineArgs.bOneComponentAtATime.GetValue();
+	}
+
+	UE_LOG(LogPCGWorldPartitionBuilder, Display, TEXT("bOneComponentAtATime : '%d' (%s)"),
+		BuilderArgs.bOneComponentAtATime ? 1 : 0,
+		bFromCommandline ? FromCommandline : FromDefault);
+	
+	// Init 'bIgnoreGenerationErrors'
+	bFromCommandline = CommandlineArgs.bIgnoreGenerationErrors.IsSet();
+	if (bFromCommandline)
+	{
+		BuilderArgs.bIgnoreGenerationErrors = CommandlineArgs.bIgnoreGenerationErrors.GetValue();
+	}
+
+	UE_LOG(LogPCGWorldPartitionBuilder, Display, TEXT("bIgnoreGenerationErrors : '%d' (%s)"),
+		BuilderArgs.bIgnoreGenerationErrors ? 1 : 0,
+		bFromCommandline ? FromCommandline : FromDefault);
+		
+	// Init 'IncludeActorIDs'
+	bFromCommandline = CommandlineArgs.IncludeActorIDs.IsSet();
+	if (bFromCommandline)
+	{
+		BuilderArgs.IncludeActorIDs = CommandlineArgs.IncludeActorIDs.GetValue();
+	}
+
+	UE_LOG(LogPCGWorldPartitionBuilder, Display, TEXT("IncludeActorIDs : '%s' (%s)"),
+		BuilderArgs.IncludeActorIDs.IsEmpty() ? TEXT("None") : *FString::Join(BuilderArgs.IncludeActorIDs, TEXT(",")),
+		bFromCommandline ? FromCommandline : FromDefault);
+		
+	// Init 'IncludeGraphNames'
+	bFromCommandline = CommandlineArgs.IncludeGraphNames.IsSet();
+	if (bFromCommandline)
+	{
+		// Commandline overrides using graph names so we empty whatever graph specified in the settings (it one or the other)
+		BuilderArgs.IncludeGraphs.Empty();
+		BuilderArgs.IncludeGraphNames = CommandlineArgs.IncludeGraphNames.GetValue();
+	}
+
+	UE_LOG(LogPCGWorldPartitionBuilder, Display, TEXT("IncludeGraphNames : '%s' (%s)"),
+		BuilderArgs.IncludeGraphNames.IsEmpty() ? TEXT("None") : *FString::Join(BuilderArgs.IncludeGraphNames, TEXT(",")),
+		bFromCommandline ? FromCommandline : TEXT("Default"));
+
+	TArray<FString> GraphsToString;
+	Algo::Transform(BuilderArgs.IncludeGraphs, GraphsToString, [](const TSoftObjectPtr<UPCGGraphInterface>& Graph) { return Graph.ToString(); });
+		
+	UE_LOG(LogPCGWorldPartitionBuilder, Display, TEXT("IncludeGraphs '%s' (%s)"),
+		GraphsToString.IsEmpty() ? TEXT("None") : *FString::Join(GraphsToString, TEXT(",")),
+		FromDefault);
+
+	return BuilderArgs;
+}
+
+UPCGWorldPartitionBuilder::UPCGWorldPartitionBuilder(const FObjectInitializer& ObjectInitializer)
+	: Super(ObjectInitializer)
+{
+	if (!HasAnyFlags(RF_ClassDefaultObject))
+	{
+		auto GetBoolParam = [this](const FString& ParamName, TOptional<bool>& OutOptionalParam)
+		{
+			const FString ParamNameAssignment = ParamName + "=";
+			bool bParamValue;
+			if (HasParam(ParamName))
+			{
+				OutOptionalParam = true;
+			}
+			else if (FParse::Bool(*GetBuilderArgs(), *ParamNameAssignment, bParamValue))
+			{
+				OutOptionalParam = bParamValue;
+			}
+		};
+
+		GetBoolParam("GenerateComponentEditingModeLoadAsPreview", CommandlineArgs.bGenerateEditingModeLoadAsPreviewComponents);
+		GetBoolParam("GenerateComponentEditingModeNormal", CommandlineArgs.bGenerateEditingModeNormalComponents);
+		GetBoolParam("GenerateComponentEditingModePreview", CommandlineArgs.bGenerateEditingModePreviewComponents);
+		GetBoolParam("OneComponentAtATime", CommandlineArgs.bOneComponentAtATime);
+		GetBoolParam("IgnoreGenerationErrors", CommandlineArgs.bIgnoreGenerationErrors);
+	
 		FString IncludeGraphNamesValue;
-		if (GetParamValue("IncludeGraphNames=", IncludeGraphNamesValue) && !IncludeGraphNamesValue.IsEmpty())
+		if (GetParamValue("IncludeGraphNames=", IncludeGraphNamesValue))
 		{
 			TArray<FString> IncludeGraphNameStrings;
-			IncludeGraphNamesValue.ParseIntoArray(IncludeGraphNameStrings, TEXT(";"), true);
-
-			IncludeGraphNames.Append(IncludeGraphNameStrings);
+			if (!IncludeGraphNamesValue.IsEmpty())
+			{
+				IncludeGraphNamesValue.ParseIntoArray(IncludeGraphNameStrings, TEXT(";"), true);
+			}
+			CommandlineArgs.IncludeGraphNames = IncludeGraphNameStrings;
 		}
 
 		FString IncludeActorIDsValue;
-		if (GetParamValue("IncludeActorIDs=", IncludeActorIDsValue) && !IncludeActorIDsValue.IsEmpty())
+		if (GetParamValue("IncludeActorIDs=", IncludeActorIDsValue))
 		{
 			TArray<FString> IncludeActorIDStrings;
-			IncludeActorIDsValue.ParseIntoArray(IncludeActorIDStrings, TEXT(";"), true);
+			if (!IncludeActorIDsValue.IsEmpty())
+			{
+				IncludeActorIDsValue.ParseIntoArray(IncludeActorIDStrings, TEXT(";"), true);
+			}
+			CommandlineArgs.IncludeActorIDs = IncludeActorIDStrings;
+		}
 
-			IncludeActorIDs.Append(IncludeActorIDStrings);
+		FString SettingAsset;
+		if (GetParamValue("PCGBuilderSettings=", SettingAsset) && !SettingAsset.IsEmpty())
+		{
+			CommandlineArgs.SettingAsset = TSoftObjectPtr<UPCGBuilderSettings>(FSoftObjectPath(SettingAsset));
 		}
 	}
 }
 
+bool UPCGWorldPartitionBuilder::PreRun(UWorld* World, FPackageSourceControlHelper& PackageHelper)
+{
+	if (!Super::PreRun(World, PackageHelper))
+	{
+		return false;
+	}
+
+	Args = FPCGWorldPartitionBuilderArgs::InitializeFrom(CommandlineArgs);
+		
+	return true;
+}
+
 bool UPCGWorldPartitionBuilder::RunInternal(UWorld* World, const FCellInfo& InCellInfo, FPackageSourceControlHelper& PackageHelper)
 {
 	if (!ensure(GEngine) || !ensure(World))
@@ -145,15 +318,15 @@ bool UPCGWorldPartitionBuilder::RunInternal(UWorld* World, const FCellInfo& InCe
 			}
 
 			// Check actor in inclusion list if provided.
-			if (!IncludeActorIDs.IsEmpty() && !IncludeActorIDs.Contains(InComponent->GetOwner()->GetName()))
+			if (!Args.IncludeActorIDs.IsEmpty() && !Args.IncludeActorIDs.Contains(InComponent->GetOwner()->GetName()))
 			{
 				return false;
 			}
 
 			// Accept based on editing mode.
-			return (InComponent->GetSerializedEditingMode() == EPCGEditorDirtyMode::LoadAsPreview)
-				|| (bGenerateEditingModeNormalComponents && InComponent->GetSerializedEditingMode() == EPCGEditorDirtyMode::Normal)
-				|| (bGenerateEditingModePreviewComponents && InComponent->GetSerializedEditingMode() == EPCGEditorDirtyMode::Preview);
+			return (Args.bGenerateEditingModeLoadAsPreviewComponents && InComponent->GetSerializedEditingMode() == EPCGEditorDirtyMode::LoadAsPreview)
+				|| (Args.bGenerateEditingModeNormalComponents && InComponent->GetSerializedEditingMode() == EPCGEditorDirtyMode::Normal)
+				|| (Args.bGenerateEditingModePreviewComponents && InComponent->GetSerializedEditingMode() == EPCGEditorDirtyMode::Preview);
 		};
 
 		PCGWorldPartitionBuilder::CollectComponentsToGenerate(World, ComponentFilter, ComponentsToGenerate);
@@ -182,10 +355,10 @@ bool UPCGWorldPartitionBuilder::RunInternal(UWorld* World, const FCellInfo& InCe
 
 	bool bGeneratedAnyComponent = false;
 
-	if (!IncludeGraphNames.IsEmpty())
+	if (!Args.IncludeGraphNames.IsEmpty())
 	{
 		// Generate all components with graph of each name, in specified order.
-		for (const FName& IncludeGraphName : IncludeGraphNames)
+		for (const FString& IncludeGraphName : Args.IncludeGraphNames)
 		{
 			auto FilterOnGraphName = [IncludeGraphName](const UPCGComponent* InComponent)
 			{
@@ -194,7 +367,24 @@ bool UPCGWorldPartitionBuilder::RunInternal(UWorld* World, const FCellInfo& InCe
 			};
 
 			bool bErrorsOccurred = false;
-			bGeneratedAnyComponent |= PCGWorldPartitionBuilder::GenerateComponents(ComponentsToGenerate, World, bOneComponentAtATime, FilterOnGraphName, DeletedActorPackages, bErrorsOccurred);
+			bGeneratedAnyComponent |= PCGWorldPartitionBuilder::GenerateComponents(ComponentsToGenerate, World, Args.bOneComponentAtATime, FilterOnGraphName, DeletedActorPackages, bErrorsOccurred);
+
+			bErrorOccurredWhileGenerating |= bErrorsOccurred;
+		}
+	}
+	else if (!Args.IncludeGraphs.IsEmpty())
+	{
+		// Generate all components with graph of each name, in specified order.
+		for (const TSoftObjectPtr<UPCGGraphInterface>& IncludeGraphPtr : Args.IncludeGraphs)
+		{
+			auto FilterOnGraphPtr = [&IncludeGraphPtr](const UPCGComponent* InComponent)
+				{
+					TSoftObjectPtr<UPCGGraphInterface> ComponentGraphPtr(InComponent ? InComponent->GetGraph() : nullptr);
+					return ComponentGraphPtr.IsValid() && ComponentGraphPtr == IncludeGraphPtr;
+				};
+
+			bool bErrorsOccurred = false;
+			bGeneratedAnyComponent |= PCGWorldPartitionBuilder::GenerateComponents(ComponentsToGenerate, World, Args.bOneComponentAtATime, FilterOnGraphPtr, DeletedActorPackages, bErrorsOccurred);
 
 			bErrorOccurredWhileGenerating |= bErrorsOccurred;
 		}
@@ -202,7 +392,7 @@ bool UPCGWorldPartitionBuilder::RunInternal(UWorld* World, const FCellInfo& InCe
 	else
 	{
 		bool bErrorsOccurred = false;
-		bGeneratedAnyComponent |= PCGWorldPartitionBuilder::GenerateComponents(ComponentsToGenerate, World, bOneComponentAtATime, DeletedActorPackages, bErrorsOccurred);
+		bGeneratedAnyComponent |= PCGWorldPartitionBuilder::GenerateComponents(ComponentsToGenerate, World, Args.bOneComponentAtATime, DeletedActorPackages, bErrorsOccurred);
 
 		bErrorOccurredWhileGenerating |= bErrorsOccurred;
 	}
@@ -250,7 +440,7 @@ bool UPCGWorldPartitionBuilder::SaveDirtyPackages(UWorld* World, FPackageSourceC
 	// Check whether an error was thrown while generating components and fail the builder if the ignore argument is not provided.
 	if (bErrorOccurredWhileGenerating)
 	{
-		if (!bIgnoreGenerationErrors)
+		if (!Args.bIgnoreGenerationErrors)
 		{
 			UE_LOG(LogPCGWorldPartitionBuilder, Display, TEXT("Dirty package detection and save skipped due to errors during generation."));
 
@@ -570,22 +760,65 @@ void PCGWorldPartitionBuilder::WaitForAllAsyncEditorProcesses(UWorld* InWorld)
 	FAssetCompilingManager::Get().ProcessAsyncTasks();
 }
 
-void PCGWorldPartitionBuilder::Build(const TArray<FString>& Args)
+bool PCGWorldPartitionBuilder::Build(UWorld* InWorld, const TArray<FString>& Args)
 {
-	if (UWorld* World = (GEditor ? GEditor->GetEditorWorldContext().World() : nullptr))
+	if (InWorld)
 	{
 		IWorldPartitionEditorModule::FRunBuilderParams Params;
 		Params.BuilderClass = UPCGWorldPartitionBuilder::StaticClass();
-		Params.World = World;
+		Params.World = InWorld;
 		Params.OperationDescription = FText::FromString("Generating PCG Components...");
 		
-		Params.ExtraArgs = TEXT("-AllowCommandletRendering");
+		Params.ExtraArgs = TEXT("-AllowCommandletRendering -AllowSoftwareRendering -AssetGatherAll=true");
 		for (const FString& Arg : Args)
 		{
 			Params.ExtraArgs += " ";
 			Params.ExtraArgs += Arg;
 		}
 
-		IWorldPartitionEditorModule::Get().RunBuilder(Params);
+		return IWorldPartitionEditorModule::Get().RunBuilder(Params);
 	}
+
+	return false;
+}
+
+bool UPCGWorldPartitionBuilder::CanBuild(const UWorld* InWorld, FName InBuildOption)
+{
+	return InWorld && InWorld->IsPartitionedWorld();
 }
+
+EEditorBuildResult UPCGWorldPartitionBuilder::Build(UWorld* InWorld, FName InBuildOption)
+{
+	TSharedPtr<SWindow> DlgWindow =
+		SNew(SWindow)
+		.Title(LOCTEXT("BuildPCGWindowTitle", "Build PCG"))
+		.ClientSize(SPCGBuilderDialog::DEFAULT_WINDOW_SIZE)
+		.SizingRule(ESizingRule::UserSized)
+		.SupportsMinimize(false)
+		.SupportsMaximize(false)
+		.SizingRule(ESizingRule::FixedSize);
+
+	TSharedRef<SPCGBuilderDialog> BuildDialog =
+		SNew(SPCGBuilderDialog)
+		.ParentWindow(DlgWindow);
+
+	DlgWindow->SetContent(BuildDialog);
+
+	IMainFrameModule& MainFrameModule = FModuleManager::LoadModuleChecked<IMainFrameModule>(TEXT("MainFrame"));
+	FSlateApplication::Get().AddModalWindow(DlgWindow.ToSharedRef(), MainFrameModule.GetParentWindow());
+
+	if (BuildDialog->GetDialogResult() != SPCGBuilderDialog::DialogResult::Cancel)
+	{
+		TArray<FString> Args;
+		FString BuilderSettingStr = BuildDialog->GetBuilderSetting().ToString();
+		if(!BuilderSettingStr.IsEmpty())
+		{
+			Args.Add("-PCGBuilderSettings=" + BuilderSettingStr);
+		}
+		return PCGWorldPartitionBuilder::Build(InWorld, Args) ? EEditorBuildResult::Success : EEditorBuildResult::Skipped;
+	}
+
+	return EEditorBuildResult::Skipped;
+}
+
+#undef LOCTEXT_NAMESPACE
diff --git a/Engine/Plugins/PCG/Source/PCGEditor/Private/WorldPartitionBuilder/PCGWorldPartitionBuilder.h b/Engine/Plugins/PCG/Source/PCGEditor/Private/WorldPartitionBuilder/PCGWorldPartitionBuilder.h
--- a/Engine/Plugins/PCG/Source/PCGEditor/Private/WorldPartitionBuilder/PCGWorldPartitionBuilder.h
+++ b/Engine/Plugins/PCG/Source/PCGEditor/Private/WorldPartitionBuilder/PCGWorldPartitionBuilder.h
@@ -4,9 +4,69 @@
 
 #include "WorldPartition/WorldPartitionBuilder.h"
 
+#include "EditorBuildUtils.h"
+
 #include "PCGWorldPartitionBuilder.generated.h"
 
 class UWorld;
+class UPCGBuilderSettings;
+class UPCGGraphInterface;
+
+struct FPCGWorldPartitionCommandlineArgs
+{
+	/** Include components which have editing mode set to LoadAsPreview. */
+	TOptional<bool> bGenerateEditingModeLoadAsPreviewComponents;
+
+	/** Include components which have editing mode set to Normal. */
+	TOptional<bool> bGenerateEditingModeNormalComponents;
+
+	/** Include components which have editing mode set to Preview. */
+	TOptional<bool> bGenerateEditingModePreviewComponents;
+		
+	/** If non empty, components with each graph name will be generated (if editing mode is included), in order. */
+	TOptional<TArray<FString>> IncludeGraphNames;
+
+	/** Call generate on each component and wait until completion and any async processes before generating the next. */
+	TOptional<bool> bOneComponentAtATime;
+
+	/** If non empty, only components on actors with given ID(s) will be generated. */
+	TOptional<TArray<FString>> IncludeActorIDs;
+
+	/** Submit dirty files even if errors occurred during generation. */
+	TOptional<bool> bIgnoreGenerationErrors;
+		
+	/** Path to builder settings asset */
+	TOptional<TSoftObjectPtr<UPCGBuilderSettings>> SettingAsset;
+};
+
+struct FPCGWorldPartitionBuilderArgs
+{
+	/** Include components which have editing mode set to LoadAsPreview. */
+	bool bGenerateEditingModeLoadAsPreviewComponents = true;
+
+	/** Include components which have editing mode set to Normal. */
+	bool bGenerateEditingModeNormalComponents = false;
+
+	/** Include components which have editing mode set to Preview. */
+	bool bGenerateEditingModePreviewComponents = false;
+
+	/** If non empty, components with each graph name will be generated (if editing mode is included), in order. */
+	TArray<FString> IncludeGraphNames;
+
+	/** If non empty, components with each graph asset will be generated (if editing mode is included), in order */
+	TArray<TSoftObjectPtr<UPCGGraphInterface>> IncludeGraphs;
+
+	/** Call generate on each component and wait until completion and any async processes before generating the next. */
+	bool bOneComponentAtATime = false;
+
+	/** If non empty, only components on actors with given ID(s) will be generated. */
+	TArray<FString> IncludeActorIDs;
+
+	/** Submit dirty files even if errors occurred during generation. */
+	bool bIgnoreGenerationErrors = false;
+
+	static FPCGWorldPartitionBuilderArgs InitializeFrom(const FPCGWorldPartitionCommandlineArgs& CommandlineArgs);
+};
 
 /**
 * Builder that triggers generation on PCG components.
@@ -30,12 +90,16 @@ public:
 	virtual ELoadingMode GetLoadingMode() const override { return ELoadingMode::EntireWorld; }
 
 protected:
+	virtual bool PreRun(UWorld* World, FPackageSourceControlHelper& PackageHelper) override;
 	virtual bool RunInternal(UWorld* World, const FCellInfo& InCellInfo, FPackageSourceControlHelper& PackageHelper) override;
 	virtual bool CanProcessNonPartitionedWorlds() const override { return true; }
 
 	/** Save all the pending dirty and deleted packages. */
 	virtual bool SaveDirtyPackages(UWorld* World, FPackageSourceControlHelper& PackageHelper);
 
+	friend class FPCGEditorModule;
+	static bool CanBuild(const UWorld* InWorld, FName InBuildOption);
+	static EEditorBuildResult Build(UWorld* InWorld, FName InBuildOption);
 private:
 	/** The packages dirtied while generating all components. */
 	UPROPERTY(Transient)
@@ -45,23 +109,8 @@ private:
 	UPROPERTY(Transient)
 	TArray<TObjectPtr<UPackage>> DeletedActorPackages;
 
-	/** Include components which have editing mode set to Normal. */
-	bool bGenerateEditingModeNormalComponents = false;
-
-	/** Include components which have editing mode set to Preview. */
-	bool bGenerateEditingModePreviewComponents = false;
-
-	/** If non empty, components with each graph name will be generated (if editing mode is included), in order. */
-	TArray<FName> IncludeGraphNames;
-
-	/** Call generate on each component and wait until completion and any async processes before generating the next. */
-	bool bOneComponentAtATime = false;
-
-	/** If non empty, only components on actors with given ID(s) will be generated. */
-	TArray<FString> IncludeActorIDs;
-
-	/** Submit dirty files even if errors occurred during generation. */
-	bool bIgnoreGenerationErrors = false;
+	FPCGWorldPartitionCommandlineArgs CommandlineArgs;
+	FPCGWorldPartitionBuilderArgs Args;
 
 	/** Flag to register if any error occurred during generation. */
 	bool bErrorOccurredWhileGenerating = false;
diff --git a/Engine/Plugins/PCG/Source/PCGEditor/Private/WorldPartitionBuilder/SPCGBuilderDialog.cpp b/Engine/Plugins/PCG/Source/PCGEditor/Private/WorldPartitionBuilder/SPCGBuilderDialog.cpp
new file mode 100644
--- /dev/null
+++ b/Engine/Plugins/PCG/Source/PCGEditor/Private/WorldPartitionBuilder/SPCGBuilderDialog.cpp
@@ -0,0 +1,141 @@
+// Copyright Epic Games, Inc. All Rights Reserved.
+#include "SPCGBuilderDialog.h"
+
+#include "PCGEditorSettings.h"
+
+#include "IDetailsView.h"
+#include "IStructureDetailsView.h"
+#include "PropertyEditorModule.h"
+
+#include "Internationalization/Internationalization.h"
+#include "Layout/Children.h"
+#include "Layout/Margin.h"
+#include "Modules/ModuleManager.h"
+#include "SlotBase.h"
+#include "Styling/AppStyle.h"
+#include "Types/SlateEnums.h"
+#include "Widgets/Input/SButton.h"
+#include "Widgets/Layout/SBorder.h"
+#include "Widgets/Layout/SSpacer.h"
+#include "Widgets/SBoxPanel.h"
+#include "Widgets/SWindow.h"
+#include "Widgets/Text/STextBlock.h"
+
+#define LOCTEXT_NAMESPACE "PCGBulderDialog"
+
+const FVector2D SPCGBuilderDialog::DEFAULT_WINDOW_SIZE = FVector2D(600, 180);
+
+void SPCGBuilderDialog::Construct(const FArguments& InArgs)
+{
+	ParentWindowPtr = InArgs._ParentWindow.Get();
+	Result = DialogResult::Cancel;
+
+	FPropertyEditorModule& PropertyEditorModule = FModuleManager::GetModuleChecked<FPropertyEditorModule>("PropertyEditor");
+
+	TSharedPtr<IStructureDetailsView> StructureDetailsView;
+
+	FDetailsViewArgs DetailsViewArgs;
+	{
+		DetailsViewArgs.bAllowSearch = false;
+		DetailsViewArgs.bHideSelectionTip = true;
+		DetailsViewArgs.bLockable = false;
+		DetailsViewArgs.bSearchInitialKeyFocus = true;
+		DetailsViewArgs.bUpdatesFromSelection = false;
+		DetailsViewArgs.NotifyHook = nullptr;
+		DetailsViewArgs.bShowOptions = true;
+		DetailsViewArgs.bShowModifiedPropertiesOption = false;
+		DetailsViewArgs.bShowScrollBar = false;
+	}
+
+	FStructureDetailsViewArgs StructureViewArgs;
+	{
+		StructureViewArgs.bShowObjects = true;
+		StructureViewArgs.bShowAssets = true;
+		StructureViewArgs.bShowClasses = true;
+		StructureViewArgs.bShowInterfaces = true;
+	}
+
+	StructureDetailsView = PropertyEditorModule.CreateStructureDetailView(DetailsViewArgs, StructureViewArgs, nullptr);
+	
+	Params.BuilderSetting = GetDefault<UPCGEditorProjectSettings>()->DefaultBuilderSetting;
+
+	FStructOnScope* Struct = new FStructOnScope(FPCGBuilderDialogParams::StaticStruct(), (uint8*)&Params);
+	StructureDetailsView->SetStructureData(MakeShareable(Struct));
+	
+	this->ChildSlot
+	[
+		SNew(SBorder)
+		.BorderImage(FAppStyle::GetBrush("ToolPanel.GroupBorder"))
+		[
+			SNew(SVerticalBox)
+			+ SVerticalBox::Slot()
+			.AutoHeight()
+			[
+				SNew(SHorizontalBox)
+				+ SHorizontalBox::Slot()
+				.Padding(10.f)
+				[
+					SNew(STextBlock)
+					.AutoWrapText(true)
+					.Text(LOCTEXT("BuildPCGDescription", "Note that this will require the current map be saved & unloaded. It will be reopened once the operation is completed."))
+				]
+			]
+			+ SVerticalBox::Slot()
+			.AutoHeight()
+			[
+				SNew(SHorizontalBox)
+				+ SHorizontalBox::Slot()
+				.FillWidth(1.f)
+				[
+					StructureDetailsView->GetWidget()->AsShared()
+				]
+			]
+			+ SVerticalBox::Slot()
+			.AutoHeight()
+			[
+				SNew(SHorizontalBox)
+				+ SHorizontalBox::Slot()
+				.FillWidth(1.f)
+				[
+					SNew(SSpacer)
+				]
+				+ SHorizontalBox::Slot()
+				.AutoWidth()
+				.Padding(2.f)
+				[
+					SNew(SButton)
+					.HAlign(HAlign_Center)
+					.ContentPadding(FAppStyle::GetMargin("StandardDialog.ContentPadding"))
+					.OnClicked(this, &SPCGBuilderDialog::OnBuildClicked)
+					.Text(LOCTEXT("BuildButton", "Build"))
+				]
+				+ SHorizontalBox::Slot()
+				.AutoWidth()
+				.Padding(2.f)
+				[
+					SNew(SButton)
+					.HAlign(HAlign_Center)
+					.ContentPadding(FAppStyle::GetMargin("StandardDialog.ContentPadding"))
+					.OnClicked(this, &SPCGBuilderDialog::OnCancelClicked)
+					.Text(LOCTEXT("CancelButton", "Cancel"))
+				]
+			]
+		]
+	];
+}
+
+FReply SPCGBuilderDialog::OnBuildClicked()
+{
+	Result = DialogResult::Build;
+	ParentWindowPtr.Pin()->RequestDestroyWindow();
+	return FReply::Handled();
+}
+
+FReply SPCGBuilderDialog::OnCancelClicked()
+{
+	Result = DialogResult::Cancel;
+	ParentWindowPtr.Pin()->RequestDestroyWindow();
+	return FReply::Handled();
+}
+
+#undef LOCTEXT_NAMESPACE
diff --git a/Engine/Plugins/PCG/Source/PCGEditor/Private/WorldPartitionBuilder/SPCGBuilderDialog.h b/Engine/Plugins/PCG/Source/PCGEditor/Private/WorldPartitionBuilder/SPCGBuilderDialog.h
new file mode 100644
--- /dev/null
+++ b/Engine/Plugins/PCG/Source/PCGEditor/Private/WorldPartitionBuilder/SPCGBuilderDialog.h
@@ -0,0 +1,68 @@
+// Copyright Epic Games, Inc. All Rights Reserved.
+#pragma once
+
+#include "Input/Reply.h"
+#include "Math/Vector2D.h"
+#include "Misc/Attribute.h"
+#include "Templates/SharedPointer.h"
+#include "Widgets/DeclarativeSyntaxSupport.h"
+#include "Widgets/SCompoundWidget.h"
+
+#include "SPCGBuilderDialog.generated.h"
+
+class UPCGBuilderSettings;
+class SWindow;
+
+USTRUCT()
+struct FPCGBuilderDialogParams
+{
+	GENERATED_BODY()
+
+	UPROPERTY(EditAnywhere, Category = Default)
+	TSoftObjectPtr<UPCGBuilderSettings> BuilderSetting;
+};
+
+//////////////////////////////////////////////////////////////////////////
+// SPCGBuilderDialog
+
+class SPCGBuilderDialog : public SCompoundWidget
+{
+public:
+	enum DialogResult
+	{
+		Build,
+		Cancel
+	};
+
+	SLATE_BEGIN_ARGS(SPCGBuilderDialog) {}
+	/** A pointer to the parent window */
+	SLATE_ATTRIBUTE(TSharedPtr<SWindow>, ParentWindow)
+	SLATE_END_ARGS()
+
+	~SPCGBuilderDialog() {}
+
+	/**
+	 * Construct this widget.  Called by the SNew() Slate macro.
+	 *
+	 * @param	InArgs				Declaration used by the SNew() macro to construct this widget
+	 */
+	void Construct(const FArguments& InArgs);
+
+	static const FVector2D DEFAULT_WINDOW_SIZE;
+
+	DialogResult GetDialogResult() const { return Result; }
+	TSoftObjectPtr<UPCGBuilderSettings> GetBuilderSetting() const { return Params.BuilderSetting; }
+
+private:
+	FReply OnBuildClicked();
+	FReply OnCancelClicked();
+
+	/** Pointer to the parent window, so we know to destroy it when done */
+	TWeakPtr<SWindow> ParentWindowPtr;
+
+	/** Dialog Result */
+	DialogResult Result;
+
+	/** Dialog Params */
+	FPCGBuilderDialogParams Params;
+};
diff --git a/Engine/Plugins/PCG/Source/PCGEditor/Public/PCGEditorSettings.h b/Engine/Plugins/PCG/Source/PCGEditor/Public/PCGEditorSettings.h
--- a/Engine/Plugins/PCG/Source/PCGEditor/Public/PCGEditorSettings.h
+++ b/Engine/Plugins/PCG/Source/PCGEditor/Public/PCGEditorSettings.h
@@ -2,14 +2,29 @@
 
 #pragma once
 
+#include "Engine/DeveloperSettings.h"
 #include "Templates/SubclassOf.h"
 
 #include "PCGEditorSettings.generated.h"
 
 class UPCGSettings;
+class UPCGBuilderSettings;
 
 struct FEdGraphPinType;
 
+UCLASS(config=Editor, meta=(DisplayName="PCG"))
+class PCGEDITOR_API UPCGEditorProjectSettings : public UDeveloperSettings
+{
+	GENERATED_BODY()
+
+public:
+	UPCGEditorProjectSettings(const FObjectInitializer& ObjectInitializer);
+		
+	/** Default Builder Settings to use when running the PCGWorldPartitionBuilder */
+	UPROPERTY(EditAnywhere, config, Category = Builder)
+	TSoftObjectPtr<UPCGBuilderSettings> DefaultBuilderSetting;
+};
+
 UCLASS(config=EditorPerProjectUserSettings)
 class PCGEDITOR_API UPCGEditorSettings : public UObject
 {
