Made it so the selector-based instance packer uses only attributes, and replaced wrong usage of keys. This disables using properties until we refactor the code and pass indices or point pointers.

#rnx

diff --git a/Engine/Plugins/PCG/Source/PCG/Private/InstanceDataPackers/PCGInstanceDataPackerBase.cpp b/Engine/Plugins/PCG/Source/PCG/Private/InstanceDataPackers/PCGInstanceDataPackerBase.cpp
--- a/Engine/Plugins/PCG/Source/PCG/Private/InstanceDataPackers/PCGInstanceDataPackerBase.cpp
+++ b/Engine/Plugins/PCG/Source/PCG/Private/InstanceDataPackers/PCGInstanceDataPackerBase.cpp
@@ -192,6 +192,11 @@ void UPCGInstanceDataPackerBase::PackCustomDataFromAccessors(const FPCGMeshInsta
 
 	OutPackedCustomData.CustomData.SetNumUninitialized(TotalNumCustomDataFloats);
 
+	if (TotalNumCustomDataFloats == 0)
+	{
+		return;
+	}
+
 	// Index of the accessor's attribute/property for a single instance (e.g. if the
 	// accessors are [float, FVector, float], AccessorPackedIndex will be [0, 1, 4])
 	int AccessorPackedIndex = 0;
diff --git a/Engine/Plugins/PCG/Source/PCG/Private/InstanceDataPackers/PCGInstanceDataPackerByAttribute.cpp b/Engine/Plugins/PCG/Source/PCG/Private/InstanceDataPackers/PCGInstanceDataPackerByAttribute.cpp
--- a/Engine/Plugins/PCG/Source/PCG/Private/InstanceDataPackers/PCGInstanceDataPackerByAttribute.cpp
+++ b/Engine/Plugins/PCG/Source/PCG/Private/InstanceDataPackers/PCGInstanceDataPackerByAttribute.cpp
@@ -7,6 +7,7 @@
 #include "Data/PCGSpatialData.h"
 #include "InstanceDataPackers/PCGInstanceDataPackerBase.h"
 #include "Metadata/Accessors/PCGAttributeAccessorHelpers.h"
+#include "MeshSelectors/PCGMeshSelectorBase.h"
 
 #include UE_INLINE_GENERATED_CPP_BY_NAME(PCGInstanceDataPackerByAttribute)
 
@@ -47,7 +48,7 @@ void UPCGInstanceDataPackerByAttribute::PackInstances_Implementation(FPCGContext
 	for (const FPCGAttributePropertyInputSelector& Selector : AttributeSelectors)
 	{
 		TUniquePtr<const IPCGAttributeAccessor> Accessor = PCGAttributeAccessorHelpers::CreateConstAccessor(InSpatialData, Selector);
-		TUniquePtr<const IPCGAttributeAccessorKeys> Keys = PCGAttributeAccessorHelpers::CreateConstKeys(InSpatialData, Selector);
+		TUniquePtr<const IPCGAttributeAccessorKeys> Keys = MakeUnique<FPCGAttributeAccessorKeysEntries>(InstanceList.InstancesMetadataEntry);
 
 		if (!Accessor.IsValid() || !Keys.IsValid())
 		{
diff --git a/Engine/Plugins/PCG/Source/PCG/Public/InstanceDataPackers/PCGInstanceDataPackerByAttribute.h b/Engine/Plugins/PCG/Source/PCG/Public/InstanceDataPackers/PCGInstanceDataPackerByAttribute.h
--- a/Engine/Plugins/PCG/Source/PCG/Public/InstanceDataPackers/PCGInstanceDataPackerByAttribute.h
+++ b/Engine/Plugins/PCG/Source/PCG/Public/InstanceDataPackers/PCGInstanceDataPackerByAttribute.h
@@ -22,7 +22,7 @@ public:
 	virtual void PackInstances_Implementation(UPARAM(ref) FPCGContext& Context, const UPCGSpatialData* InSpatialData, UPARAM(ref) const FPCGMeshInstanceList& InstanceList, FPCGPackedCustomData& OutPackedCustomData) const override;
 
 public:
-	UPROPERTY(BlueprintReadWrite, EditAnywhere, Category = InstanceDataPacker)
+	UPROPERTY(BlueprintReadWrite, EditAnywhere, Category = InstanceDataPacker, meta = (PCG_DiscardPropertySelection))
 	TArray<FPCGAttributePropertyInputSelector> AttributeSelectors;
 
 #if WITH_EDITORONLY_DATA
