Fixed fallout in the debug object tree view following changes in the stack for dynamic subgraphs

#rnx

diff --git a/Engine/Plugins/PCG/Source/PCGEditor/Private/SPCGEditorGraphDebugObjectTree.cpp b/Engine/Plugins/PCG/Source/PCGEditor/Private/SPCGEditorGraphDebugObjectTree.cpp
--- a/Engine/Plugins/PCG/Source/PCGEditor/Private/SPCGEditorGraphDebugObjectTree.cpp
+++ b/Engine/Plugins/PCG/Source/PCGEditor/Private/SPCGEditorGraphDebugObjectTree.cpp
@@ -604,9 +604,45 @@ void SPCGEditorGraphDebugObjectTree::AddStacksToTree(const TArray<FPCGStack>& St
 
 		const TArray<FPCGStackFrame>& StackFrames = Stack.GetStackFrames();
 
+		// Generic function to attack an item to the parent at the correct place
+		auto AddToGraphItemsAndAttachItemToParent = [&InOutStackToItem, this](FPCGStack& GraphStack, FPCGEditorGraphDebugObjectItemPtr GraphItem)
+		{
+			AllGraphItems.Add(GraphItem);
+
+			TArray<FPCGStackFrame>& LoopGraphStackFrames = GraphStack.GetStackFramesMutable();
+			while (LoopGraphStackFrames.Num() > 0)
+			{
+				LoopGraphStackFrames.SetNum(LoopGraphStackFrames.Num() - 1, EAllowShrinking::No);
+
+				if (FPCGEditorGraphDebugObjectItemPtr* ParentItem = InOutStackToItem.Find(GraphStack))
+				{
+					(*ParentItem)->AddChild(GraphItem.ToSharedRef());
+					break;
+				}
+			}
+		};
+
+		// Generic function to create the subgraph/loop item and hook it up to the parent properly
+		auto AddSubgraphOrLoopItemToStack = [&Stack, &InOutStackToItem, &AddToGraphItemsAndAttachItemToParent](const UPCGGraph* InGraph, const UPCGNode* InSubgraphNode, int FrameCutoff, bool bInIsDebuggable, bool bInDisplayGrayedOut)
+		{
+			FPCGStack GraphStack = Stack;
+			GraphStack.GetStackFramesMutable().SetNum(FrameCutoff);
+
+			if (!InOutStackToItem.Contains(GraphStack))
+			{
+				FPCGEditorGraphDebugObjectItemPtr GraphItem = InOutStackToItem.Emplace(
+					GraphStack,
+					MakeShared<FPCGEditorGraphDebugObjectItem_PCGSubgraph>(InSubgraphNode, InGraph, GraphStack, bInIsDebuggable, bInDisplayGrayedOut));
+
+				AddToGraphItemsAndAttachItemToParent(GraphStack, GraphItem);
+			}
+		};
+
 		// Example stack:
 		//     Component/TopGraph/SubgraphNode/Subgraph/LoopSubgraphNode/LoopIndex/LoopSubgraph
-		// 
+		//                                       ^ static subgraph
+		//     Component/TopGraph/SubgraphNode/INDEX_NONE/Subgraph/...
+		//                                       ^ dynamic subgraph
 		// The loop below adds tree items for component & top graph, and then whenever a graph is encountered
 		// we look a previous frames to determine whether to add a subgraph item or loop subgraph item.
 		for (int FrameIndex = 1; FrameIndex < StackFrames.Num() && bRemainingStackContainsEditedGraph; FrameIndex++)
@@ -637,63 +673,19 @@ void SPCGEditorGraphDebugObjectTree::AddStacksToTree(const TArray<FPCGStack>& St
 						ActorItem->AddChild(TopGraphItem.ToSharedRef());
 					}
 				}
-				// Previous stack was node, therefore subgraph.
+				// Previous stack was node, therefore static subgraph.
 				else if (const UPCGNode* SubgraphNode = Cast<const UPCGNode>(PreviousStackFrame.Object))
 				{
-					FPCGStack GraphStack = Stack;
-					GraphStack.GetStackFramesMutable().SetNum(FrameIndex + 1);
-
-					if (!InOutStackToItem.Contains(GraphStack))
-					{
-						FPCGEditorGraphDebugObjectItemPtr GraphItem = InOutStackToItem.Emplace(
-							GraphStack,
-							MakeShared<FPCGEditorGraphDebugObjectItem_PCGSubgraph>(SubgraphNode, StackGraph, GraphStack, bIsDebuggable, bDisplayGrayedOut));
-
-						AllGraphItems.Add(GraphItem);
-
-						TArray<FPCGStackFrame>& GraphStackFrames = GraphStack.GetStackFramesMutable();
-						while (GraphStackFrames.Num() > 0)
-						{
-							GraphStackFrames.SetNum(GraphStackFrames.Num() - 1, EAllowShrinking::No);
-
-							if (FPCGEditorGraphDebugObjectItemPtr* ParentItem = InOutStackToItem.Find(GraphStack))
-							{
-								(*ParentItem)->AddChild(GraphItem.ToSharedRef());
-								break;
-							}
-						}
-					}
+					AddSubgraphOrLoopItemToStack(StackGraph, SubgraphNode, FrameIndex + 1, bIsDebuggable, bDisplayGrayedOut);
 				}
 				// Previous stack was loop index, therefore loop subgraph.
-				else if (FrameIndex >= 2 && PreviousStackFrame.LoopIndex != INDEX_NONE)
+				else if (FrameIndex >= 2 && (PreviousStackFrame.LoopIndex != INDEX_NONE))
 				{
 					const UPCGNode* LoopSubgraphNode = Cast<const UPCGNode>(StackFrames[FrameIndex - 2].Object);
 					if (ensure(LoopSubgraphNode))
 					{
-						// Take the stack up to the looped subgraph node, add a item for the node + graph.
-						FPCGStack LoopGraphStack = Stack;
-						LoopGraphStack.GetStackFramesMutable().SetNum(FrameIndex - 1);
-
-						if (!InOutStackToItem.Contains(LoopGraphStack))
-						{
-							FPCGEditorGraphDebugObjectItemPtr LoopGraphItem = InOutStackToItem.Emplace(
-								LoopGraphStack,
-								MakeShared<FPCGEditorGraphDebugObjectItem_PCGSubgraph>(LoopSubgraphNode, StackGraph, LoopGraphStack, /*bIsDebuggable=*/false, false));
-
-							AllGraphItems.Add(LoopGraphItem);
-
-							TArray<FPCGStackFrame>& LoopGraphStackFrames = LoopGraphStack.GetStackFramesMutable();
-							while (LoopGraphStackFrames.Num() > 0)
-							{
-								LoopGraphStackFrames.SetNum(LoopGraphStackFrames.Num() - 1, EAllowShrinking::No);
-
-								if (FPCGEditorGraphDebugObjectItemPtr* ParentItem = InOutStackToItem.Find(LoopGraphStack))
-								{
-									(*ParentItem)->AddChild(LoopGraphItem.ToSharedRef());
-									break;
-								}
-							}
-						}
+						// Take the stack up to the looped subgraph node, add a item for the node + graph
+						AddSubgraphOrLoopItemToStack(StackGraph, LoopSubgraphNode, FrameIndex - 1, /*bIsDebuggable=*/false, /*bIsGrayedOut=*/false);
 
 						// Take full stack up until this point which will be the unique stack for the loop iteration.
 						FPCGStack LoopIterationStack = Stack;
@@ -705,22 +697,19 @@ void SPCGEditorGraphDebugObjectTree::AddStacksToTree(const TArray<FPCGStack>& St
 								LoopIterationStack,
 								MakeShared<FPCGEditorGraphDebugObjectItem_PCGLoopIndex>(PreviousStackFrame.LoopIndex, StackGraph, LoopIterationStack, bIsDebuggable, bDisplayGrayedOut));
 
-							AllGraphItems.Add(LoopIterationItem);
-
-							TArray<FPCGStackFrame>& GraphStackFrames = LoopIterationStack.GetStackFramesMutable();
-							while (GraphStackFrames.Num() > 0)
-							{
-								GraphStackFrames.SetNum(GraphStackFrames.Num() - 1, EAllowShrinking::No);
-
-								if (FPCGEditorGraphDebugObjectItemPtr* ParentItem = InOutStackToItem.Find(LoopIterationStack))
-								{
-									(*ParentItem)->AddChild(LoopIterationItem.ToSharedRef());
-									break;
-								}
-							}
+							AddToGraphItemsAndAttachItemToParent(LoopIterationStack, LoopIterationItem);
 						}
 					}
 				}
+				// Previous stack was invalid node / node with a INDEX_NONE loop, therefore most likely a dynamic subgraph
+				else if (FrameIndex >= 2 && !PreviousStackFrame.IsValid())
+				{
+					const UPCGNode* DynamicSubgraphNode = Cast<const UPCGNode>(StackFrames[FrameIndex - 2].Object);
+					if (DynamicSubgraphNode)
+					{
+						AddSubgraphOrLoopItemToStack(StackGraph, DynamicSubgraphNode, FrameIndex + 1, bIsDebuggable, bDisplayGrayedOut);
+					}
+				}
 			}
 
 			// Check if any of the remaining stack frames contain valid debug targets - invocations of the current edited graph.
