Added additional pin status type to be able to differentiate overrides/user parameter pins, which allow for easier processing in subgraphs.
Adjusted behavior so we fill in explicit pins before override pins if there is a collision.
Added cvar to track those, should be made to be always activated once we've gone through the data.

#rnx

diff --git a/Engine/Plugins/PCG/Source/PCG/Private/Graph/PCGGraphCompiler.cpp b/Engine/Plugins/PCG/Source/PCG/Private/Graph/PCGGraphCompiler.cpp
--- a/Engine/Plugins/PCG/Source/PCG/Private/Graph/PCGGraphCompiler.cpp
+++ b/Engine/Plugins/PCG/Source/PCG/Private/Graph/PCGGraphCompiler.cpp
@@ -126,7 +126,7 @@ TArray<FPCGGraphTask> FPCGGraphCompiler::CompileGraph(UPCGGraph* InGraph, FPCGTa
 			{
 				check(InputPin);
 				// Params have to be funneled in the subgraph element (so we can pass down user parameters)
-				const bool bIsParamPin = SubgraphSettings->HasOverridableParam(InputPin->Properties.Label);
+				const bool bIsOverrideOrUserParamPin = InputPin->Properties.IsOverrideOrUserParamPin();
 
 				for (const UPCGEdge* InboundEdge : InputPin->Edges)
 				{
@@ -134,7 +134,7 @@ TArray<FPCGGraphTask> FPCGGraphCompiler::CompileGraph(UPCGGraph* InGraph, FPCGTa
 					{
 						// Implementation note: conceptually, the non-param inputs need to be connected only to the input node task.
 						// However, because of the static/dynamic culling which happen on the subgraph node, we need to have the connections on the pretask too (e.g. the subgraph node) even if they don't provide data.
-						if (bIsParamPin || !InputNodeTask)
+						if (bIsOverrideOrUserParamPin || !InputNodeTask)
 						{
 							PreTask.Inputs.Emplace(IdMapping[InboundEdge->InputPin->Node], InboundEdge->InputPin->Properties, InboundEdge->OutputPin->Properties, /*bInProvideData=*/true);
 						}
diff --git a/Engine/Plugins/PCG/Source/PCG/Private/PCGSettings.cpp b/Engine/Plugins/PCG/Source/PCG/Private/PCGSettings.cpp
--- a/Engine/Plugins/PCG/Source/PCG/Private/PCGSettings.cpp
+++ b/Engine/Plugins/PCG/Source/PCG/Private/PCGSettings.cpp
@@ -40,6 +40,11 @@ namespace PCGSettings
 	static constexpr int DeprecationAliasIndex = -2;
 	// For property path concatenation
 	static constexpr const TCHAR* PropertyPathSeparator = TEXT("/");
+
+	static TAutoConsoleVariable<bool> CVarWarnOverridePinsCollideWithInputPins(
+		TEXT("pcg.Graph.WarnPinNameCollisions"),
+		false,
+		TEXT("Enables warnings when there are name collision between pin names and overrides/user parameters."));
 }
 
 /** Custom Crc computation that ignores properties that will not affect the computed result of a node. */
@@ -495,7 +500,7 @@ void UPCGSettings::FillOverridableParamsPins(TArray<FPCGPinProperties>& OutPins)
 	else
 	{
 		FPCGPinProperties& ParamPin = OutPins.Emplace_GetRef(PCGPinConstants::DefaultParamsLabel, EPCGDataType::Param, /*bInAllowMultipleConnections=*/ true, /*bAllowMultipleData=*/ true);
-		ParamPin.SetAdvancedPin();
+		ParamPin.SetOverrideOrUserParamPin();
 
 #if WITH_EDITOR
 		ParamPin.Tooltip = LOCTEXT("GlobalParamPinTooltip", "Atribute Set containing multiple parameters to override. Names must match perfectly.");
@@ -508,17 +513,20 @@ void UPCGSettings::FillOverridableParamsPins(TArray<FPCGPinProperties>& OutPins)
 	{
 		if (InputPinsLabelsAndTypes.Contains(OverridableParam.Label))
 		{
-			//const FString ParamsName = OverridableParam.Label.ToString();
-			//UE_LOG(LogPCG, Warning, TEXT("[%s-%s] While automatically adding override pins, an existing pin was found with conflicting name '%s'. "
-			//	"Rename or remove this pin to allow the automatic override pin to be added. Automatic override pin '%s' skipped."),
-			//	*GraphName, *NodeName, *ParamsName, *ParamsName);
+			if (PCGSettings::CVarWarnOverridePinsCollideWithInputPins.GetValueOnAnyThread())
+			{
+				const FString ParamsName = OverridableParam.Label.ToString();
+				UE_LOG(LogPCG, Warning, TEXT("[%s-%s] While automatically adding override pins, an existing pin was found with conflicting name '%s'. "
+					"Rename or remove this pin to allow the automatic override pin to be added. Automatic override pin '%s' skipped."),
+					*GraphName, *NodeName, *ParamsName, *ParamsName);
+			}
 			continue;
 		}
 
 		InputPinsLabelsAndTypes.Emplace(OverridableParam.Label, EPCGDataType::Param);
 
 		FPCGPinProperties& ParamPin = OutPins.Emplace_GetRef(OverridableParam.Label, EPCGDataType::Param, /*bInAllowMultipleConnections=*/ false, /*bAllowMultipleData=*/ false);
-		ParamPin.SetAdvancedPin();
+		ParamPin.SetOverrideOrUserParamPin();
 #if WITH_EDITOR
 
 		if (!OverridableParam.Properties.IsEmpty())
diff --git a/Engine/Plugins/PCG/Source/PCG/Public/PCGPin.h b/Engine/Plugins/PCG/Source/PCG/Public/PCGPin.h
--- a/Engine/Plugins/PCG/Source/PCG/Public/PCGPin.h
+++ b/Engine/Plugins/PCG/Source/PCG/Public/PCGPin.h
@@ -31,7 +31,9 @@ enum class EPCGPinStatus : uint8
 	/** Only for input pins, mark this pin as required.If set on an output pin, behave as Normal. */
 	Required,
 	/** Advanced pin will be hidden by default in the UIand will be shown only if the user extends the node(in the UI) to see advanced pins.Pins can't be required and advanced at the same time. */
-	Advanced
+	Advanced,
+	/** Pins generated by overridable parameters (normal nodes) or user parameters (subgraph nodes) */
+	OverrideOrUserParam UMETA(Hidden)
 };
 
 USTRUCT(BlueprintType, meta=(HasNativeBreak="/Script/PCG.PCGBlueprintPinHelpers.BreakPinProperty", HasNativeMake="/Script/PCG.PCGBlueprintPinHelpers.MakePinProperty"))
@@ -74,9 +76,12 @@ struct PCG_API FPCGPinProperties
 	// Allowing multiple connections will automatically enable multi data.
 	void SetAllowMultipleConnections(bool bInAllowMultipleConnectons);
 
-	bool IsAdvancedPin() const { return PinStatus == EPCGPinStatus::Advanced; }
+	bool IsAdvancedPin() const { return PinStatus == EPCGPinStatus::Advanced || PinStatus == EPCGPinStatus::OverrideOrUserParam; }
 	void SetAdvancedPin() { PinStatus = EPCGPinStatus::Advanced; }
 
+	bool IsOverrideOrUserParamPin() const { return PinStatus == EPCGPinStatus::OverrideOrUserParam; }
+	void SetOverrideOrUserParamPin() { PinStatus = EPCGPinStatus::OverrideOrUserParam; }
+
 	bool IsRequiredPin() const { return PinStatus == EPCGPinStatus::Required; }
 	void SetRequiredPin() { PinStatus = EPCGPinStatus::Required; }
 
diff --git a/Engine/Plugins/PCG/Source/PCGEditor/Private/SPCGEditorGraphNode.cpp b/Engine/Plugins/PCG/Source/PCGEditor/Private/SPCGEditorGraphNode.cpp
--- a/Engine/Plugins/PCG/Source/PCGEditor/Private/SPCGEditorGraphNode.cpp
+++ b/Engine/Plugins/PCG/Source/PCGEditor/Private/SPCGEditorGraphNode.cpp
@@ -423,7 +423,7 @@ bool SPCGEditorGraphNodePin::ShouldDisplayAsRequiredForExecution() const
 		// as required, because it effectively is. So return false if there are other pins which are not advanced.
 		return !Algo::AnyOf(PCGNode->GetInputPins(), [PCGPin](UPCGPin* InOtherPin)
 		{
-			return InOtherPin && InOtherPin != PCGPin && InOtherPin->Properties.PinStatus != EPCGPinStatus::Advanced;
+			return InOtherPin && InOtherPin != PCGPin && !InOtherPin->Properties.IsAdvancedPin();
 		});
 	}
 
