[PCG] Graph Executor V2: Fix out of tick budget being set multiple times if taks calls AsyncProcessing multiple times in its Execute

#rnx

diff --git a/Engine/Plugins/PCG/Source/PCG/Private/Graph/PCGGraphExecutor.cpp b/Engine/Plugins/PCG/Source/PCG/Private/Graph/PCGGraphExecutor.cpp
--- a/Engine/Plugins/PCG/Source/PCG/Private/Graph/PCGGraphExecutor.cpp
+++ b/Engine/Plugins/PCG/Source/PCG/Private/Graph/PCGGraphExecutor.cpp
@@ -1525,6 +1525,7 @@ bool FPCGGraphExecutor::ExecuteScheduling(double EndTime, TSharedPtr<FPCGGraphAc
 						ActiveTask->Context->AsyncState.EndTime = EndTime;
 						ActiveTask->Context->AsyncState.bIsRunningOnMainThread = false;
 						ActiveTask->Context->AsyncState.bIsRunningOutOfTick = true;
+						ActiveTask->Context->AsyncState.bIsOutOfTickBudgetSet = false;
 
 						// Capture copy of SharedPtr so task stays valid while Async task is alive
 						ActiveTask->StartExecuting();
diff --git a/Engine/Plugins/PCG/Source/PCG/Public/Helpers/PCGAsync.h b/Engine/Plugins/PCG/Source/PCG/Public/Helpers/PCGAsync.h
--- a/Engine/Plugins/PCG/Source/PCG/Public/Helpers/PCGAsync.h
+++ b/Engine/Plugins/PCG/Source/PCG/Public/Helpers/PCGAsync.h
@@ -99,9 +99,12 @@ namespace FPCGAsync
 
 			double EndTime = AsyncState.EndTime;
 			
-			if (AsyncState.bIsRunningOutOfTick && bEnableTimeSlicing && OutOfTickBudgetInSeconds > 0.f)
+			if (AsyncState.bIsRunningOutOfTick && !AsyncState.bIsOutOfTickBudgetSet && bEnableTimeSlicing && OutOfTickBudgetInSeconds > 0.f)
 			{
 				EndTime = FPlatformTime::Seconds() + OutOfTickBudgetInSeconds;
+
+				// Prevent multiple AsyncProcessing calls in the same node execute to reset the budget
+				AsyncState.bIsOutOfTickBudgetSet = true;
 			}
 
 			AsyncState.SetStarted(EndTime);
diff --git a/Engine/Plugins/PCG/Source/PCG/Public/Helpers/PCGAsyncState.h b/Engine/Plugins/PCG/Source/PCG/Public/Helpers/PCGAsyncState.h
--- a/Engine/Plugins/PCG/Source/PCG/Public/Helpers/PCGAsyncState.h
+++ b/Engine/Plugins/PCG/Source/PCG/Public/Helpers/PCGAsyncState.h
@@ -56,6 +56,9 @@ struct FPCGAsyncState
 	/** True if we can run out of tick */
 	bool bIsRunningOutOfTick = false;
 
+	/** True if out of tick buddget was set */
+	bool bIsOutOfTickBudgetSet = false;
+
 	/** Returns true if we reached end time. */
 	PCG_API bool ShouldStop() const;
 
