[PCG] Fixing possible nullptr dereference on metadata accessors

#rnx

diff --git a/Engine/Plugins/PCG/Source/PCG/Private/Metadata/Accessors/PCGAttributeAccessorHelpers.cpp b/Engine/Plugins/PCG/Source/PCG/Private/Metadata/Accessors/PCGAttributeAccessorHelpers.cpp
--- a/Engine/Plugins/PCG/Source/PCG/Private/Metadata/Accessors/PCGAttributeAccessorHelpers.cpp
+++ b/Engine/Plugins/PCG/Source/PCG/Private/Metadata/Accessors/PCGAttributeAccessorHelpers.cpp
@@ -24,7 +24,7 @@ namespace PCGAttributeAccessorHelpers
 	/** Retrieves metadata object and the corresponding attribute if any for the provided data. */
 	void ExtractMetadataAttribute(UPCGData* InData, FName Name, UPCGMetadata*& OutMetadata, FPCGMetadataAttributeBase*& OutAttribute)
 	{
-		OutMetadata = InData->MutableMetadata();
+		OutMetadata = InData ? InData->MutableMetadata() : nullptr;
 		OutAttribute = OutMetadata ? OutMetadata->GetMutableAttribute(Name) : nullptr;
 	}
 
@@ -169,7 +169,10 @@ namespace PCGAttributeAccessorHelpers
 			}
 
 			// Set the cached selector
-			InData->SetLastSelector(InSelector);
+			if (InData)
+			{
+				InData->SetLastSelector(InSelector);
+			}
 		}
 
 		return Accessor;
