[PCG] Clean up Custom HLSL UI
* Added pin name dropdowns instead of typing pin names by hand
* Re-organized declarations (added Helper section & moved into sub-categories)
* Now avoids templating declarations by pin when there is only one pin
* Fixed a bug where declarations were not initialized properly on load

#rnx

diff --git a/Engine/Plugins/PCG/Source/PCG/Private/Compute/Elements/PCGCustomHLSL.cpp b/Engine/Plugins/PCG/Source/PCG/Private/Compute/Elements/PCGCustomHLSL.cpp
--- a/Engine/Plugins/PCG/Source/PCG/Private/Compute/Elements/PCGCustomHLSL.cpp
+++ b/Engine/Plugins/PCG/Source/PCG/Private/Compute/Elements/PCGCustomHLSL.cpp
@@ -22,6 +22,8 @@
 
 namespace PCGHLSLElement
 {
+	const FString PinDeclTemplateStr = TEXT("<pin>");
+
 	/** First capture: Pin name (supports a - z, A - Z, and 0 - 9) */
 	constexpr int AttributePinCaptureGroup = 1;
 
@@ -91,13 +93,6 @@ void UPCGCustomHLSLSettings::PostLoad()
 
 	UpdatePinSettings();
 	UpdateAttributeKeys();
-}
-
-void UPCGCustomHLSLSettings::PostInitProperties()
-{
-	Super::PostInitProperties();
-
-	UpdatePinSettings();
 	UpdateDeclarations();
 }
 #endif
@@ -436,263 +431,364 @@ const FPCGPinPropertiesGPU* UPCGCustomHLSLSettings::GetOutputPinPropertiesGPU(co
 void UPCGCustomHLSLSettings::UpdateDeclarations()
 {
 	// Reference: UOptimusNode_CustomComputeKernel::UpdatePreamble
+	UpdateInputDeclarations();
+	UpdateOutputDeclarations();
+	UpdateHelperDeclarations();
+}
+
+void UPCGCustomHLSLSettings::UpdateInputDeclarations()
+{
 	InputDeclarations.Reset();
-	OutputDeclarations.Reset();
 
-	// Add constants category
-	{
-		if (KernelType == EPCGKernelType::PointGenerator)
-		{
-			const UPCGPin* PointProcessingOutputPin = GetFirstPointOutputPin();
-			InputDeclarations += FString::Format(
-				TEXT("// Constants\n")
-				TEXT("uint PointCount = {0};\n")
-				TEXT("\n"),
-				{ PointCount });
-		}
-	}
+	TArray<FString> DataPins;
+	TArray<FString> PointDataPins;
+	TArray<FString> SplineDataPins;
+	TArray<FString> LandscapeDataPins;
+	TArray<FString> TextureDataPins;
+	TArray<FString> RawBufferDataPins;
 
-	// Add resource indexing category
+	for (const FPCGPinProperties& Pin : InputPinProperties())
 	{
-		InputDeclarations += TEXT("// Resource Indexing\n");
-		InputDeclarations += TEXT("uint ThreadIndex;\n");
+		DataPins.Add(Pin.Label.ToString());
 
-		if (KernelType == EPCGKernelType::PointProcessor)
+		if (!!(Pin.AllowedTypes & EPCGDataType::Point))
 		{
-			const UPCGPin* PointProcessingInputPin = GetPointProcessingInputPin();
-			const UPCGPin* PointProcessingOutputPin = GetFirstPointOutputPin();
+			PointDataPins.Add(Pin.Label.ToString());
+		}
 
-			if (PointProcessingInputPin && PointProcessingOutputPin)
-			{
-				InputDeclarations += FString::Format(
-					TEXT("uint {0}_DataIndex;\n")
-					TEXT("uint {1}_DataIndex;\n"),
-					{ PointProcessingInputPin->Properties.Label.ToString(),  PointProcessingOutputPin->Properties.Label.ToString() });
-			}
+		if (!!(Pin.AllowedTypes & EPCGDataType::Spline))
+		{
+			SplineDataPins.Add(Pin.Label.ToString());
 		}
-		else if (KernelType == EPCGKernelType::PointGenerator)
+
+		if (!!(Pin.AllowedTypes & EPCGDataType::Landscape))
 		{
-			if (const UPCGPin* PointProcessingOutputPin = GetFirstPointOutputPin())
-			{
-				InputDeclarations += FString::Format(
-					TEXT("uint {0}_DataIndex;\n"),
-					{ PointProcessingOutputPin->Properties.Label.ToString() });
-			}
+			LandscapeDataPins.Add(Pin.Label.ToString());
 		}
-		else if (KernelType == EPCGKernelType::Custom)
+
+		if (!!(Pin.AllowedTypes & EPCGDataType::Texture))
 		{
-			auto EmitGetThreadElement = [&InInputDeclarations = InputDeclarations](const FPCGPinProperties& Properties)
-			{
-				InInputDeclarations += FString::Format(TEXT(
-					"// Returns false if thread has no data to operate on.\n"
-					"bool {0}_GetThreadData(uint ThreadIndex, out uint OutDataIndex, out uint OutElementIndex);\n"
-					"bool {0}_GetThreadData(uint ThreadIndex, out uint OutDataIndex, out uint OutDataAddress, out uint OutElementIndex);\n"),
-					{ Properties.Label.ToString() });
-			};
+			TextureDataPins.Add(Pin.Label.ToString());
+		}
+	}
 
-			for (const FPCGPinProperties& Properties : InputPinProperties())
-			{
-				EmitGetThreadElement(Properties);
-			}
+	if (!DataPins.IsEmpty())
+	{
+		InputDeclarations += TEXT("### DATA FUNCTIONS ###\n\n");
 
-			for (const FPCGPinProperties& Properties : OutputPinProperties())
-			{
-				EmitGetThreadElement(Properties);
-			}
-		}
-		else
+		const bool bMultiPin = DataPins.Num() > 1;
+
+		if (bMultiPin)
 		{
-			checkNoEntry();
+			InputDeclarations += TEXT("// Valid pins: ") + FString::Join(DataPins, TEXT(", ")) + TEXT("\n");
 		}
 
-		InputDeclarations += TEXT("uint ElementIndex;\n");
-		InputDeclarations += TEXT("int3 GetNumThreads();\n");
+		InputDeclarations += FString::Format(TEXT(
+			"uint {0}_GetNumData();\n"
+			"uint {0}_GetNumElements();\n"
+			"\n"
+			"// Valid types: bool, int, float, float2, float3, float4, Rotator (float3), Quat (float4), Transform (float4x4)\n"
+			"\n"
+			"<type> {0}_Get<type>(uint DataIndex, uint ElementIndex, uint AttributeId);\n"
+			"<type> {0}_Get<type>(uint DataIndex, uint ElementIndex, 'AttributeName');\n"),
+			{ bMultiPin ? PCGHLSLElement::PinDeclTemplateStr : DataPins[0] });
+
+		InputDeclarations += TEXT("\n");
 	}
 
-	// Add debug category
+	if (!PointDataPins.IsEmpty())
 	{
-		if (bPrintShaderDebugValues)
+		InputDeclarations += TEXT("### POINT DATA FUNCTIONS ###\n\n");
+
+		const bool bMultiPin = PointDataPins.Num() > 1;
+
+		if (bMultiPin)
 		{
-			InputDeclarations += FString::Format(
-				TEXT("\n// Debug\n")
-				TEXT("void WriteDebugValue(uint Index, float Value); // Index in [0, {0}] (set from 'Debug Buffer Size' property)\n"),
-				{ DebugBufferSize - 1 });
+			InputDeclarations += TEXT("// Valid pins: ") + FString::Join(PointDataPins, TEXT(", ")) + TEXT("\n\n");
 		}
+
+		InputDeclarations += FString::Format(TEXT(
+			"uint {0}_GetNumPoints(uint DataIndex);\n"
+			"float3 {0}_GetPosition(uint DataIndex, uint ElementIndex);\n"
+			"float4 {0}_GetRotation(uint DataIndex, uint ElementIndex);\n"
+			"float3 {0}_GetScale(uint DataIndex, uint ElementIndex);\n"
+			"float3 {0}_GetBoundsMin(uint DataIndex, uint ElementIndex);\n"
+			"float3 {0}_GetBoundsMax(uint DataIndex, uint ElementIndex);\n"
+			"float3 {0}_GetColor(uint DataIndex, uint ElementIndex);\n"
+			"float {0}_GetDensity(uint DataIndex, uint ElementIndex);\n"
+			"int {0}_GetSeed(uint DataIndex, uint ElementIndex);\n"
+			"float {0}_GetSteepness(uint DataIndex, uint ElementIndex);\n"
+			"float4x4 {0}_GetPointTransform(uint DataIndex, uint ElementIndex);\n"
+			"bool {0}_IsValid(uint DataIndex, uint ElementIndex);\n"),
+			{ bMultiPin ? PCGHLSLElement::PinDeclTemplateStr : PointDataPins[0] });
+
+		InputDeclarations += TEXT("\n");
 	}
 
-	// Per-pin input category
+	if (!LandscapeDataPins.IsEmpty())
 	{
-		TArray<FString> DataPins;
-		TArray<FString> PointDataPins;
-		TArray<FString> SplineDataPins;
-		TArray<FString> LandscapeDataPins;
-		TArray<FString> TextureDataPins;
-		TArray<FString> RawBufferDataPins;
+		InputDeclarations += TEXT("### LANDSCAPE DATA FUNCTIONS ###\n\n");
+
+		const bool bMultiPin = LandscapeDataPins.Num() > 1;
 
-		for (const FPCGPinProperties& Pin : InputPinProperties())
+		if (bMultiPin)
 		{
-			DataPins.Add(Pin.Label.ToString());
+			InputDeclarations += TEXT("// Valid pins: ") + FString::Join(LandscapeDataPins, TEXT(", ")) + TEXT("\n\n");
+		}
 
-			if (!!(Pin.AllowedTypes & EPCGDataType::Point))
-			{
-				PointDataPins.Add(Pin.Label.ToString());
-			}
+		InputDeclarations += FString::Format(TEXT(
+			"float {0}_GetHeight(float3 WorldPos);\n"
+			"float3 {0}_GetNormal(float3 WorldPos);\n"),
+			{ bMultiPin ? PCGHLSLElement::PinDeclTemplateStr : LandscapeDataPins[0] });
 
-			if (!!(Pin.AllowedTypes & EPCGDataType::Spline))
-			{
-				SplineDataPins.Add(Pin.Label.ToString());
-			}
+		InputDeclarations += TEXT("\n");
+	}
 
-			if (!!(Pin.AllowedTypes & EPCGDataType::Landscape))
-			{
-				LandscapeDataPins.Add(Pin.Label.ToString());
-			}
+	if (!TextureDataPins.IsEmpty())
+	{
+		InputDeclarations += TEXT("### TEXTURE DATA FUNCTIONS ###\n\n");
 
-			if (!!(Pin.AllowedTypes & EPCGDataType::Texture))
-			{
-				TextureDataPins.Add(Pin.Label.ToString());
-			}
+		const bool bMultiPin = TextureDataPins.Num() > 1;
+
+		if (bMultiPin)
+		{
+			InputDeclarations += TEXT("// Valid pins: ") + FString::Join(TextureDataPins, TEXT(", ")) + TEXT("\n\n");
 		}
 
-		InputDeclarations += TEXT("\n### HELPER FUNCTIONS ###\n");
-		InputDeclarations += TEXT("\nfloat3 GetComponentBoundsMin(); // World-space\n");
-		InputDeclarations += TEXT("float3 GetComponentBoundsMax();\n");
-		InputDeclarations += TEXT("uint GetSeed();\n");
+		InputDeclarations += FString::Format(TEXT(
+			"float2 {0}_GetTexCoords(float2 WorldPos, float2 Min, float2 Max);\n"
+			"float4 {0}_Sample(float2 TexCoords);\n"),
+			{ bMultiPin ? PCGHLSLElement::PinDeclTemplateStr : TextureDataPins[0] });
+
+		InputDeclarations += TEXT("\n");
+	}
 
-		InputDeclarations += TEXT("\nfloat FRand(inout uint Seed); // Returns random float between 0 and 1.\n");
-		InputDeclarations += TEXT("uint ComputeSeed(uint A, uint B);\n");
-		InputDeclarations += TEXT("uint ComputeSeed(uint A, uint B, uint C);\n");
-		InputDeclarations += TEXT("uint ComputeSeedFromPosition(float3 Position);\n");
+	if (!RawBufferDataPins.IsEmpty())
+	{
+		InputDeclarations += TEXT("### BYTE ADDRESS BUFFER DATA FUNCTIONS ###\n\n");
 
-		InputDeclarations += TEXT("\nfloat3 CreateGrid2D(int ElementIndex, int NumPoints, float3 Min, float3 Max);\n");
-		InputDeclarations += TEXT("float3 CreateGrid2D(int ElementIndex, int NumPoints, int NumRows, float3 Min, float3 Max);\n");
-		InputDeclarations += TEXT("float3 CreateGrid3D(int ElementIndex, int NumPoints, float3 Min, float3 Max);\n");
-		InputDeclarations += TEXT("float3 CreateGrid3D(int ElementIndex, int NumPoints, int NumRows, int NumCols, float3 Min, float3 Max);\n");
+		const bool bMultiPin = RawBufferDataPins.Num() > 1;
 
-		if (!DataPins.IsEmpty())
+		if (bMultiPin)
 		{
-			InputDeclarations += TEXT("\n### DATA FUNCTIONS ###\n");
-			InputDeclarations += TEXT("\n// Valid pins: ") + FString::Join(DataPins, TEXT(", ")) + TEXT("\n");
-			InputDeclarations += TEXT("// Valid types: bool, int, float, float2, float3, float4, Rotator (float3), Quat (float4), Transform (float4x4)\n");
-
-			InputDeclarations += TEXT("\nuint <pin>_GetNumData();\n");
-			InputDeclarations += TEXT("uint <pin>_GetNumElements();\n");
-			InputDeclarations += TEXT("<type> <pin>_Get<type>(uint DataIndex, uint ElementIndex, uint AttributeId);\n");
-			InputDeclarations += TEXT("<type> <pin>_Get<type>(uint DataIndex, uint ElementIndex, 'AttributeName');\n");
+			InputDeclarations += TEXT("// Valid pins: ") + FString::Join(RawBufferDataPins, TEXT(", ")) + TEXT("\n\n");
 		}
 
-		if (!PointDataPins.IsEmpty())
+		InputDeclarations += FString::Format(TEXT(
+			"uint {0}_ReadNumValues();\n"
+			"uint {0}_ReadValue(uint Index);\n"),
+			{ bMultiPin ? PCGHLSLElement::PinDeclTemplateStr : RawBufferDataPins[0] });
+
+		InputDeclarations += TEXT("\n");
+	}
+
+	InputDeclarations.TrimStartAndEndInline();
+}
+
+void UPCGCustomHLSLSettings::UpdateOutputDeclarations()
+{
+	OutputDeclarations.Reset();
+
+	TArray<FString> DataPins;
+	TArray<FString> PointDataPins;
+	TArray<FString> SplineDataPins;
+	TArray<FString> RawBufferDataPins;
+
+	for (const FPCGPinProperties& Pin : OutputPinProperties())
+	{
+		DataPins.Add(Pin.Label.ToString());
+
+		if (!!(Pin.AllowedTypes & EPCGDataType::Point))
 		{
-			InputDeclarations += TEXT("\n### POINT DATA FUNCTIONS ###\n");
-			InputDeclarations += TEXT("\n// Valid pins: ") + FString::Join(PointDataPins, TEXT(", ")) + TEXT("\n");
+			PointDataPins.Add(Pin.Label.ToString());
+		}
 
-			InputDeclarations += TEXT("\nuint <pin>_GetNumPoints(uint DataIndex);\n");
-			InputDeclarations += TEXT("float3 <pin>_GetPosition(uint DataIndex, uint ElementIndex);\n");
-			InputDeclarations += TEXT("float4 <pin>_GetRotation(uint DataIndex, uint ElementIndex);\n");
-			InputDeclarations += TEXT("float3 <pin>_GetScale(uint DataIndex, uint ElementIndex);\n");
-			InputDeclarations += TEXT("float3 <pin>_GetBoundsMin(uint DataIndex, uint ElementIndex);\n");
-			InputDeclarations += TEXT("float3 <pin>_GetBoundsMax(uint DataIndex, uint ElementIndex);\n");
-			InputDeclarations += TEXT("float3 <pin>_GetColor(uint DataIndex, uint ElementIndex);\n");
-			InputDeclarations += TEXT("float <pin>_GetDensity(uint DataIndex, uint ElementIndex);\n");
-			InputDeclarations += TEXT("int <pin>_GetSeed(uint DataIndex, uint ElementIndex);\n");
-			InputDeclarations += TEXT("float <pin>_GetSteepness(uint DataIndex, uint ElementIndex);\n");
-			InputDeclarations += TEXT("float4x4 <pin>_GetPointTransform(uint DataIndex, uint ElementIndex);\n");
-			InputDeclarations += TEXT("bool <pin>_IsValid(uint DataIndex, uint ElementIndex);\n");
+		if (!!(Pin.AllowedTypes & EPCGDataType::Spline))
+		{
+			SplineDataPins.Add(Pin.Label.ToString());
 		}
+	}
+
+	if (!DataPins.IsEmpty())
+	{
+		OutputDeclarations += TEXT("### DATA FUNCTIONS ###\n\n");
 
-		if (!LandscapeDataPins.IsEmpty())
+		const bool bMultiPin = DataPins.Num() > 1;
+
+		if (bMultiPin)
 		{
-			InputDeclarations += TEXT("\n### LANDSCAPE DATA FUNCTIONS ###\n");
-			InputDeclarations += TEXT("\n// Valid pins: ") + FString::Join(LandscapeDataPins, TEXT(", ")) + TEXT("\n");
+			OutputDeclarations += TEXT("// Valid pins: ") + FString::Join(DataPins, TEXT(", ")) + TEXT("\n");
+		}
 
-			InputDeclarations += TEXT("\nfloat <pin>_GetHeight(float3 WorldPos);\n");
-			InputDeclarations += TEXT("float3 <pin>_GetNormal(float3 WorldPos);\n");
+		OutputDeclarations += FString::Format(TEXT(
+			"// Valid types: bool, int, float, float2, float3, float4, Rotator (float3), Quat (float4), Transform (float4x4)\n"
+			"\n"
+			"void {0}_Set<type>(uint DataIndex, uint ElementIndex, uint AttributeId, <type> Value);\n"
+			"void {0}_Set<type>(uint DataIndex, uint ElementIndex, 'AttributeName', <type> Value);\n"),
+			{ bMultiPin ? PCGHLSLElement::PinDeclTemplateStr : DataPins[0] });
+
+		OutputDeclarations += TEXT("\n");
+	}
+
+	if (!PointDataPins.IsEmpty())
+	{
+		OutputDeclarations += TEXT("### POINT DATA FUNCTIONS ###\n\n");
+
+		const bool bMultiPin = PointDataPins.Num() > 1;
+
+		if (bMultiPin)
+		{
+			OutputDeclarations += TEXT("// Valid pins: ") + FString::Join(PointDataPins, TEXT(", ")) + TEXT("\n\n");
 		}
 
-		if (!TextureDataPins.IsEmpty())
+		OutputDeclarations += FString::Format(TEXT(
+			"void {0}_SetPosition(uint DataIndex, uint ElementIndex, float3 Position);\n"
+			"void {0}_SetRotation(uint DataIndex, uint ElementIndex, float4 Rotation);\n"
+			"void {0}_SetScale(uint DataIndex, uint ElementIndex, float3 Scale);\n"
+			"void {0}_SetBoundsMin(uint DataIndex, uint ElementIndex, float3 BoundsMin);\n"
+			"void {0}_SetBoundsMax(uint DataIndex, uint ElementIndex, float3 BoundsMax);\n"
+			"void {0}_SetColor(uint DataIndex, uint ElementIndex, float4 Color);\n"
+			"void {0}_SetDensity(uint DataIndex, uint ElementIndex, float Density);\n"
+			"void {0}_SetSeed(uint DataIndex, uint ElementIndex, int Seed);\n"
+			"void {0}_SetSteepness(uint DataIndex, uint ElementIndex, float Steepness);\n"
+			"void {0}_SetPointTransform(uint DataIndex, uint ElementIndex, float4x4 Transform);\n"
+			"bool {0}_RemovePoint(uint DataIndex, uint ElementIndex);\n"),
+			{ bMultiPin ? PCGHLSLElement::PinDeclTemplateStr : PointDataPins[0] });
+
+		OutputDeclarations += TEXT("\n");
+	}
+
+	if (!RawBufferDataPins.IsEmpty())
+	{
+		OutputDeclarations += TEXT("### BYTE ADDRESS BUFFER DATA FUNCTIONS ###\n\n");
+
+		const bool bMultiPin = RawBufferDataPins.Num() > 1;
+
+		if (bMultiPin)
 		{
-			InputDeclarations += TEXT("\n### TEXTURE DATA FUNCTIONS ###\n");
-			InputDeclarations += TEXT("\n// Valid pins: ") + FString::Join(TextureDataPins, TEXT(", ")) + TEXT("\n");
+			OutputDeclarations += TEXT("// Valid pins: ") + FString::Join(RawBufferDataPins, TEXT(", ")) + TEXT("\n\n");
+		}
+
+		OutputDeclarations += FString::Format(TEXT(
+			"uint {0}_WriteValue(uint Index, uint Value);\n"),
+			{ bMultiPin ? PCGHLSLElement::PinDeclTemplateStr : RawBufferDataPins[0] });
+
+		OutputDeclarations += TEXT("\n");
+	}
+
+	OutputDeclarations.TrimStartAndEndInline();
+}
+
+void UPCGCustomHLSLSettings::UpdateHelperDeclarations()
+{
+	HelperDeclarations.Reset();
 
-			InputDeclarations += TEXT("\nfloat2 <pin>_GetTexCoords(float2 WorldPos, float2 Min, float2 Max);\n");
-			InputDeclarations += TEXT("float4 <pin>_Sample(float2 TexCoords);\n");
+	// Constants category
+	{
+		if (KernelType == EPCGKernelType::PointGenerator)
+		{
+			HelperDeclarations += TEXT("### CONSTANTS ###\n\n");
+			HelperDeclarations += FString::Format(TEXT("const uint PointCount = {0};\n\n"), { PointCount });
 		}
 
-		if (!RawBufferDataPins.IsEmpty())
+		HelperDeclarations += TEXT("### PER-THREAD CONSTANTS ###\n\n");
+		HelperDeclarations += TEXT("const uint ThreadIndex;\n");
+
+		if (KernelType == EPCGKernelType::PointProcessor)
 		{
-			InputDeclarations += TEXT("\n### BYTE ADDRESS BUFFER DATA FUNCTIONS ###\n");
-			InputDeclarations += TEXT("\n// Valid pins: ") + FString::Join(RawBufferDataPins, TEXT(", ")) + TEXT("\n");
+			const UPCGPin* PointProcessingInputPin = GetPointProcessingInputPin();
+			const UPCGPin* PointProcessingOutputPin = GetFirstPointOutputPin();
 
-			InputDeclarations += TEXT("\nuint <pin>_ReadNumValues();\n");
-			InputDeclarations += TEXT("uint <pin>_ReadValue(uint Index);\n");
+			if (PointProcessingInputPin && PointProcessingOutputPin)
+			{
+				HelperDeclarations += FString::Format(TEXT(
+					"const uint {0}_DataIndex;\n"
+					"const uint {1}_DataIndex;\n"),
+					{ PointProcessingInputPin->Properties.Label.ToString(),  PointProcessingOutputPin->Properties.Label.ToString() });
+			}
+		}
+		else if (KernelType == EPCGKernelType::PointGenerator)
+		{
+			if (const UPCGPin* PointProcessingOutputPin = GetFirstPointOutputPin())
+			{
+				HelperDeclarations += FString::Format(
+					TEXT("const uint {0}_DataIndex;\n"),
+					{ PointProcessingOutputPin->Properties.Label.ToString() });
+			}
 		}
+
+		HelperDeclarations += TEXT("const uint ElementIndex;\n");
+		HelperDeclarations += TEXT("\n");
 	}
 
-	// Per-pin output category
+	// Helper funcs category
 	{
-		TArray<FString> DataPins;
-		TArray<FString> PointDataPins;
-		TArray<FString> SplineDataPins;
-		TArray<FString> RawBufferDataPins;
+		HelperDeclarations += TEXT(
+			"### HELPER FUNCTIONS ###\n"
+			"\n"
+			"int3 GetNumThreads();\n");
 
-		for (const FPCGPinProperties& Pin : OutputPinProperties())
+		if (KernelType == EPCGKernelType::Custom)
 		{
-			DataPins.Add(Pin.Label.ToString());
+			TArray<FString> PinNames;
 
-			if (!!(Pin.AllowedTypes & EPCGDataType::Point))
+			for (const FPCGPinProperties& Properties : InputPinProperties())
 			{
-				PointDataPins.Add(Pin.Label.ToString());
+				PinNames.Add(Properties.Label.ToString());
 			}
 
-			if (!!(Pin.AllowedTypes & EPCGDataType::Spline))
+			for (const FPCGPinProperties& Properties : OutputPinProperties())
 			{
-				SplineDataPins.Add(Pin.Label.ToString());
+				PinNames.Add(Properties.Label.ToString());
 			}
-		}
 
-		if (!DataPins.IsEmpty())
-		{
-			OutputDeclarations += TEXT("### DATA FUNCTIONS ###\n");
-			OutputDeclarations += TEXT("\n// Valid pins: ") + FString::Join(DataPins, TEXT(", ")) + TEXT("\n");
-			OutputDeclarations += TEXT("// Valid types: bool, int, float, float2, float3, float4, Rotator (float3), Quat (float4), Transform (float4x4)\n");
+			if (!PinNames.IsEmpty())
+			{
+				HelperDeclarations += TEXT("\n// Returns false if thread has no data to operate on.\n");
 
-			OutputDeclarations += TEXT("\nvoid <pin>_Set<type>(uint DataIndex, uint ElementIndex, uint AttributeId, <type> Value);\n");
-			OutputDeclarations += TEXT("void <pin>_Set<type>(uint DataIndex, uint ElementIndex, 'AttributeName', <type> Value);\n");
-		}
+				const bool bMultiPin = PinNames.Num() > 1;
 
-		if (!PointDataPins.IsEmpty())
-		{
-			OutputDeclarations += TEXT("\n### POINT DATA FUNCTIONS ###\n");
-			OutputDeclarations += TEXT("\n// Valid pins: ") + FString::Join(PointDataPins, TEXT(", ")) + TEXT("\n");
+				if (bMultiPin)
+				{
+					HelperDeclarations += TEXT("// Valid pins: ") + FString::Join(PinNames, TEXT(", ")) + TEXT("\n");
+				}
+
+				HelperDeclarations += FString::Format(TEXT(
+					"bool {0}_GetThreadData(uint ThreadIndex, out uint OutDataIndex, out uint OutElementIndex);\n"
+					"bool {0}_GetThreadData(uint ThreadIndex, out uint OutDataIndex, out uint OutDataAddress, out uint OutElementIndex);\n"),
+					{ bMultiPin ? PCGHLSLElement::PinDeclTemplateStr : PinNames[0] });
 
-			OutputDeclarations += TEXT("\nvoid <pin>_SetPosition(uint DataIndex, uint ElementIndex, float3 Position);\n");
-			OutputDeclarations += TEXT("void <pin>_SetRotation(uint DataIndex, uint ElementIndex, float4 Rotation);\n");
-			OutputDeclarations += TEXT("void <pin>_SetScale(uint DataIndex, uint ElementIndex, float3 Scale);\n");
-			OutputDeclarations += TEXT("void <pin>_SetBoundsMin(uint DataIndex, uint ElementIndex, float3 BoundsMin);\n");
-			OutputDeclarations += TEXT("void <pin>_SetBoundsMax(uint DataIndex, uint ElementIndex, float3 BoundsMax);\n");
-			OutputDeclarations += TEXT("void <pin>_SetColor(uint DataIndex, uint ElementIndex, float4 Color);\n");
-			OutputDeclarations += TEXT("void <pin>_SetDensity(uint DataIndex, uint ElementIndex, float Density);\n");
-			OutputDeclarations += TEXT("void <pin>_SetSeed(uint DataIndex, uint ElementIndex, int Seed);\n");
-			OutputDeclarations += TEXT("void <pin>_SetSeedFromPosition(uint DataIndex, uint ElementIndex, float3 Position);\n");
-			OutputDeclarations += TEXT("void <pin>_SetSteepness(uint DataIndex, uint ElementIndex, float Steepness);\n");
-			OutputDeclarations += TEXT("void <pin>_SetPointTransform(uint DataIndex, uint ElementIndex, float4x4 Transform);\n");
-			OutputDeclarations += TEXT("bool <pin>_RemovePoint(uint DataIndex, uint ElementIndex);\n");
+				HelperDeclarations += TEXT("\n");
+			}
 		}
 
-		if (!RawBufferDataPins.IsEmpty())
-		{
-			OutputDeclarations += TEXT("\n### BYTE ADDRESS BUFFER DATA FUNCTIONS ###\n\n");
-			OutputDeclarations += TEXT("// Valid pins: ") + FString::Join(RawBufferDataPins, TEXT(", ")) + TEXT("\n");
+		HelperDeclarations += TEXT(
+			"float3 GetComponentBoundsMin(); // World-space\n"
+			"float3 GetComponentBoundsMax();\n"
+			"uint GetSeed();\n"
+			"\n"
+			"float FRand(inout uint Seed); // Returns random float between 0 and 1.\n"
+			"uint ComputeSeed(uint A, uint B);\n"
+			"uint ComputeSeed(uint A, uint B, uint C);\n"
+			"uint ComputeSeedFromPosition(float3 Position);\n"
+			"\n"
+			"// Returns the position of the Nth point in a 2D or 3D grid with the given constraints.\n"
+			"float3 CreateGrid2D(int ElementIndex, int NumPoints, float3 Min, float3 Max);\n"
+			"float3 CreateGrid2D(int ElementIndex, int NumPoints, int NumRows, float3 Min, float3 Max);\n"
+			"float3 CreateGrid3D(int ElementIndex, int NumPoints, float3 Min, float3 Max);\n"
+			"float3 CreateGrid3D(int ElementIndex, int NumPoints, int NumRows, int NumCols, float3 Min, float3 Max);\n");
 
-			OutputDeclarations += TEXT("\nuint <pin>_WriteValue(uint Index, uint Value);\n");
+		if (bPrintShaderDebugValues)
+		{
+			HelperDeclarations += FString::Format(TEXT(
+				"\n"
+				"// Writes floats to the debug buffer array, which will be readback and logged in the console for inspection."
+				"void WriteDebugValue(uint Index, float Value); // Index in [0, {0}] (set from 'Debug Buffer Size' property)\n"),
+				{ DebugBufferSize - 1 });
 		}
 	}
 
-	if (!OutputDeclarations.IsEmpty())
-	{
-		// Remove final newline as a small UI improvement.
-		OutputDeclarations = OutputDeclarations.LeftChop(1);
-	}
+	HelperDeclarations.TrimStartAndEndInline();
 }
 
 void UPCGCustomHLSLSettings::UpdatePinSettings()
@@ -840,6 +936,26 @@ void UPCGCustomHLSLSettings::UpdateAttributeKeys()
 		}
 	}
 }
+
+TArray<FName> UPCGCustomHLSLSettings::GetInputPinNames() const
+{
+	TArray<FName> PinNames;
+
+	for (const FPCGPinProperties& PinProps : InputPins)
+	{
+		PinNames.Add(PinProps.Label);
+	}
+
+	return PinNames;
+}
+
+TArray<FName> UPCGCustomHLSLSettings::GetInputPinNamesAndNone() const
+{
+	TArray<FName> PinNames = GetInputPinNames();
+	PinNames.Insert(NAME_None, 0);
+
+	return PinNames;
+}
 #endif
 
 bool UPCGCustomHLSLSettings::IsKernelValid(FPCGContext* InContext, bool bQuiet) const
diff --git a/Engine/Plugins/PCG/Source/PCG/Public/Compute/Elements/PCGCustomHLSL.h b/Engine/Plugins/PCG/Source/PCG/Public/Compute/Elements/PCGCustomHLSL.h
--- a/Engine/Plugins/PCG/Source/PCG/Public/Compute/Elements/PCGCustomHLSL.h
+++ b/Engine/Plugins/PCG/Source/PCG/Public/Compute/Elements/PCGCustomHLSL.h
@@ -37,19 +37,23 @@ public:
 
 public:
 	/** Compute graphs use this to calculate the buffer size of output pins. */
-	UPROPERTY(BlueprintReadWrite, EditAnywhere, Category = "GPU Buffer Size", meta = (EditCondition = "bDisplayBufferSizeSettings", EditConditionHides, HideEditConditionToggle))
+	UPROPERTY(BlueprintReadWrite, EditAnywhere, Category = Settings, meta = (DisplayAfter = "Tooltip", EditCondition = "bDisplayBufferSizeSettings", EditConditionHides, HideEditConditionToggle))
 	EPCGPinBufferSizeMode BufferSizeMode = EPCGPinBufferSizeMode::FromFirstPin;
 
-	UPROPERTY(BlueprintReadWrite, EditAnywhere, Category = "GPU Buffer Size", meta = (EditCondition = "(bDisplayBufferSizeSettings && BufferSizeMode == EPCGPinBufferSizeMode::FixedElementCount) || AllowedTypes == EPCGDataType::Param", EditConditionHides))
+	UPROPERTY(BlueprintReadWrite, EditAnywhere, Category = Settings, meta = (DisplayAfter = "BufferSizeMode", EditCondition = "(bDisplayBufferSizeSettings && BufferSizeMode == EPCGPinBufferSizeMode::FixedElementCount) || AllowedTypes == EPCGDataType::Param", EditConditionHides))
 	int FixedBufferElementCount = 4;
 
-	UPROPERTY(BlueprintReadWrite, EditAnywhere, DisplayName = "Input Pins", Category = "GPU Buffer Size", meta = (EditCondition = "bDisplayBufferSizeSettings && BufferSizeMode == EPCGPinBufferSizeMode::FromProductOfInputPins", EditConditionHides))
+	UPROPERTY(BlueprintReadWrite, EditAnywhere, DisplayName = "Buffer Size Input Pins", Category = Settings, meta = (DisplayAfter = "BufferSizeMode", EditCondition = "bDisplayBufferSizeSettings && BufferSizeMode == EPCGPinBufferSizeMode::FromProductOfInputPins", EditConditionHides, GetOptions = "GetInputPinNames"))
 	TArray<FName> BufferSizeInputPinLabels;
 
 	/** Select an input pin to copy attributes from. If left as 'None', this will be ignored. Note, this will copy attribute names only, not their values. */
-	UPROPERTY(BlueprintReadWrite, EditAnywhere, Category = "GPU Buffer Size", meta = (EditCondition = "bAllowEditInitializationPin", EditConditionHides, HideEditConditionToggle))
+	UPROPERTY(BlueprintReadWrite, EditAnywhere, Category = Settings, meta = (DisplayAfter = "BufferSizeMode", EditCondition = "bAllowEditInitializationPin", EditConditionHides, HideEditConditionToggle, GetOptions = "GetInputPinNamesAndNone"))
 	FName InitializeFromPin = NAME_None;
 
+	/** Add entries to create new attributes on data emitted by this pin. */
+	UPROPERTY(EditAnywhere, DisplayName = "Attributes to Create", Category = Settings, meta = (DisplayAfter = "BufferSizeMode"))
+	TArray<FPCGKernelAttributeKey> CreatedKernelAttributeKeys;
+
 #if WITH_EDITORONLY_DATA
 	UPROPERTY(Transient)
 	bool bDisplayBufferSizeSettings = true;
@@ -57,10 +61,6 @@ public:
 	UPROPERTY(Transient)
 	bool bAllowEditInitializationPin = false;
 #endif // WITH_EDITORONLY_DATA
-
-	/** Add entries to create new attributes on data emitted by this pin. */
-	UPROPERTY(EditAnywhere, DisplayName = "Attributes to Create", Category = Settings)
-	TArray<FPCGKernelAttributeKey> CreatedKernelAttributeKeys;
 };
 
 template<>
@@ -102,7 +102,6 @@ public:
 	//~Begin UObject interface
 #if WITH_EDITOR
 	virtual void PostLoad() override;
-	virtual void PostInitProperties() override;
 #endif
 	//~End UObject interface
 
@@ -130,17 +129,26 @@ protected:
 	virtual FPCGElementPtr CreateElement() const override;
 	//~End UPCGSettings interface
 
-protected:
 	/** Gets the GPU pin properties for the output pin with the given label. */
 	const FPCGPinPropertiesGPU* GetOutputPinPropertiesGPU(const FName& InPinLabel) const;
 
 #if WITH_EDITOR
 	void UpdateDeclarations();
+	void UpdateInputDeclarations();
+	void UpdateOutputDeclarations();
+	void UpdateHelperDeclarations();
 	void UpdatePinSettings();
 	void UpdateAttributeKeys();
+
+	/** List of all non-advanced input pin names. */
+	UFUNCTION()
+	TArray<FName> GetInputPinNames() const;
+
+	/** List of all non-advanced input pin names, prepended with 'Name_NONE'. */
+	UFUNCTION()
+	TArray<FName> GetInputPinNamesAndNone() const;
 #endif
 
-protected:
 	const UPCGPin* GetInputPin(FName Label) const;
 	const UPCGPin* GetOutputPin(FName Label) const;
 	const UPCGPin* GetFirstInputPin() const;
@@ -160,16 +168,16 @@ protected:
 	UPROPERTY(BlueprintReadWrite, EditAnywhere, Category = "Settings", meta = (EditCondition = "KernelType == EPCGKernelType::PointGenerator", EditConditionHides))
 	int PointCount = 256;
 
-	UPROPERTY(BlueprintReadWrite, EditAnywhere, Category = "Thread Count", meta = (EditCondition = "KernelType == EPCGKernelType::Custom", EditConditionHides))
+	UPROPERTY(BlueprintReadWrite, EditAnywhere, Category = "Settings|Thread Count", meta = (EditCondition = "KernelType == EPCGKernelType::Custom", EditConditionHides))
 	EPCGDispatchThreadCount DispatchThreadCount = EPCGDispatchThreadCount::FromFirstOutputPin;
 
-	UPROPERTY(BlueprintReadWrite, EditAnywhere, Category = "Thread Count", meta = (EditCondition = "KernelType == EPCGKernelType::Custom && DispatchThreadCount != EPCGDispatchThreadCount::Fixed", EditConditionHides))
+	UPROPERTY(BlueprintReadWrite, EditAnywhere, Category = "Settings|Thread Count", meta = (EditCondition = "KernelType == EPCGKernelType::Custom && DispatchThreadCount != EPCGDispatchThreadCount::Fixed", EditConditionHides))
 	int ThreadCountMultiplier = 1;
 
-	UPROPERTY(BlueprintReadWrite, EditAnywhere, Category = "Thread Count", meta = (EditCondition = "KernelType == EPCGKernelType::Custom && DispatchThreadCount == EPCGDispatchThreadCount::Fixed", EditConditionHides))
+	UPROPERTY(BlueprintReadWrite, EditAnywhere, Category = "Settings|Thread Count", meta = (EditCondition = "KernelType == EPCGKernelType::Custom && DispatchThreadCount == EPCGDispatchThreadCount::Fixed", EditConditionHides))
 	int FixedThreadCount = 1;
 
-	UPROPERTY(BlueprintReadWrite, EditAnywhere, DisplayName = "Input Pins", Category = "Thread Count", meta = (EditCondition = "KernelType == EPCGKernelType::Custom && DispatchThreadCount == EPCGDispatchThreadCount::FromProductOfInputPins", EditConditionHides))
+	UPROPERTY(BlueprintReadWrite, EditAnywhere, DisplayName = "Input Pins", Category = "Settings|Thread Count", meta = (EditCondition = "KernelType == EPCGKernelType::Custom && DispatchThreadCount == EPCGDispatchThreadCount::FromProductOfInputPins", EditConditionHides, GetOptions = "GetInputPinNames"))
 	TArray<FName> ThreadCountInputPinLabels;
 
 public:
@@ -180,17 +188,20 @@ public:
 	TArray<FPCGPinPropertiesGPU> OutputPins = { FPCGPinPropertiesGPU(PCGPinConstants::DefaultOutputLabel, EPCGDataType::Point) };
 
 protected:
-	UPROPERTY(Transient, VisibleAnywhere, Category = "Settings|Declarations", meta = (MultiLine = true))
+	UPROPERTY(BlueprintReadWrite, EditAnywhere, Category = "Source", meta = (MultiLine = true, Tooltip = ""))
+	FString ShaderFunctions;
+
+	UPROPERTY(BlueprintReadWrite, EditAnywhere, Category = "Source", meta = (MultiLine = true, Tooltip = ""))
+	FString ShaderSource;
+
+	UPROPERTY(Transient, VisibleAnywhere, Category = "Declarations|Inputs", meta = (MultiLine = true))
 	FString InputDeclarations;
 
-	UPROPERTY(Transient, VisibleAnywhere, Category = "Settings|Declarations", meta = (MultiLine = true))
+	UPROPERTY(Transient, VisibleAnywhere, Category = "Declarations|Outputs", meta = (MultiLine = true))
 	FString OutputDeclarations;
 
-	UPROPERTY(BlueprintReadWrite, EditAnywhere, Category = "Settings", meta = (MultiLine = true, Tooltip = ""))
-	FString ShaderFunctions;
-
-	UPROPERTY(BlueprintReadWrite, EditAnywhere, Category = "Settings", meta = (MultiLine = true, Tooltip = ""))
-	FString ShaderSource;
+	UPROPERTY(Transient, VisibleAnywhere, Category = "Declarations|Helpers", meta = (MultiLine = true))
+	FString HelperDeclarations;
 };
 
 class FPCGCustomHLSLElement : public IPCGElement
