[PCG] Break Custom HLSL node into multiple entries in the graph palette, one for each value of EPCGKernelType
* Include KernelType as an additional title under the 'Custom HLSL' node title

#rnx

diff --git a/Engine/Plugins/PCG/Source/PCG/Private/Compute/Elements/PCGCustomHLSL.cpp b/Engine/Plugins/PCG/Source/PCG/Private/Compute/Elements/PCGCustomHLSL.cpp
--- a/Engine/Plugins/PCG/Source/PCG/Private/Compute/Elements/PCGCustomHLSL.cpp
+++ b/Engine/Plugins/PCG/Source/PCG/Private/Compute/Elements/PCGCustomHLSL.cpp
@@ -13,6 +13,7 @@
 #include "Compute/PCGComputeGraph.h"
 #include "Compute/PCGDataBinding.h"
 #include "Data/PCGPointData.h"
+#include "Elements/Metadata/PCGMetadataElementCommon.h"
 
 #include "Internationalization/Regex.h"
 
@@ -342,6 +343,32 @@ FPCGDataCollectionDesc UPCGCustomHLSLSettings::ComputeOutputPinDataDesc(const UP
 }
 
 #if WITH_EDITOR
+TArray<FPCGPreConfiguredSettingsInfo> UPCGCustomHLSLSettings::GetPreconfiguredInfo() const
+{
+	return PCGMetadataElementCommon::FillPreconfiguredSettingsInfoFromEnum<EPCGKernelType>();
+}
+
+void UPCGCustomHLSLSettings::ApplyPreconfiguredSettings(const FPCGPreConfiguredSettingsInfo& PreconfiguredInfo)
+{
+	if (const UEnum* EnumPtr = StaticEnum<EPCGKernelType>())
+	{
+		if (EnumPtr->IsValidEnumValue(PreconfiguredInfo.PreconfiguredIndex))
+		{
+			KernelType = EPCGKernelType(PreconfiguredInfo.PreconfiguredIndex);
+		}
+	}
+}
+
+FString UPCGCustomHLSLSettings::GetAdditionalTitleInformation() const
+{
+	if (const UEnum* EnumPtr = StaticEnum<EPCGKernelType>())
+	{
+		return EnumPtr->GetNameStringByValue(static_cast<int>(KernelType));
+	}
+
+	return FString();
+}
+
 EPCGChangeType UPCGCustomHLSLSettings::GetChangeTypeForProperty(const FName& InPropertyName) const
 {
 	EPCGChangeType ChangeType = Super::GetChangeTypeForProperty(InPropertyName);
diff --git a/Engine/Plugins/PCG/Source/PCG/Public/Compute/Elements/PCGCustomHLSL.h b/Engine/Plugins/PCG/Source/PCG/Public/Compute/Elements/PCGCustomHLSL.h
--- a/Engine/Plugins/PCG/Source/PCG/Public/Compute/Elements/PCGCustomHLSL.h
+++ b/Engine/Plugins/PCG/Source/PCG/Public/Compute/Elements/PCGCustomHLSL.h
@@ -114,8 +114,14 @@ public:
 	virtual FText GetDefaultNodeTitle() const override { return NSLOCTEXT("PCGCustomHLSLElement", "NodeTitle", "Custom HLSL"); }
 	virtual FText GetNodeTooltipText() const override { return NSLOCTEXT("PCGCustomHLSLElement", "NodeTooltip", "Produces a HLSL compute shader which will be executed on the GPU."); }
 	virtual EPCGSettingsType GetType() const override { return EPCGSettingsType::GPU; }
+
+	virtual TArray<FPCGPreConfiguredSettingsInfo> GetPreconfiguredInfo() const override;
+	virtual bool OnlyExposePreconfiguredSettings() const override { return true; }
 #endif
 
+	virtual FString GetAdditionalTitleInformation() const override;
+	virtual void ApplyPreconfiguredSettings(const FPCGPreConfiguredSettingsInfo& PreconfigureInfo) override;
+
 	virtual bool IsKernelValid(FPCGContext* InContext = nullptr, bool bQuiet = true) const override;
 	virtual FString GetCookedKernelSource(const TMap<FPCGKernelAttributeKey, int>& GlobalAttributeLookupTable) const override;
 	virtual int ComputeKernelThreadCount(const UPCGDataBinding* Binding) const override;
