[PCG] Fix misc GPU bugs
* Properly handle Color point property when unpacking data collections
* Properly support consuming invalid points from Source or Target pin in GPU Copy Points
* Properly initialize point properties from Source data in GPU Copy Points
* Fix incorrect loop order in GPU Copy Points when copying each source to every target

#rnx

diff --git a/Engine/Plugins/PCG/Shaders/Private/Elements/PCGCopyPoints.usf b/Engine/Plugins/PCG/Shaders/Private/Elements/PCGCopyPoints.usf
--- a/Engine/Plugins/PCG/Shaders/Private/Elements/PCGCopyPoints.usf
+++ b/Engine/Plugins/PCG/Shaders/Private/Elements/PCGCopyPoints.usf
@@ -114,7 +114,7 @@ void Main(uint3 GroupId : SV_GroupID, uint GroupIndex : SV_GroupIndex)
 
 	const uint SourceNumData = Source_GetNumData();
 	const uint TargetNumData = Target_GetNumData();
-	const uint SourcePointsCount = Source_GetNumElements();
+	const uint TargetPointsCount = Target_GetNumElements();
 	const uint bCopyEachSourceOnEveryTarget = GetCopyEachSourceOnEveryTarget();
 
 	// Data must be non-zero, and NxM (when bCopyEachSourceOnEveryTarget == true), N:N, 1:N, or N:1
@@ -129,8 +129,8 @@ void Main(uint3 GroupId : SV_GroupID, uint GroupIndex : SV_GroupIndex)
 
 	if (bCopyEachSourceOnEveryTarget)
 	{
-		const uint TargetPointIndex = ThreadIndex / SourcePointsCount;
-		const uint SourcePointIndex = ThreadIndex % SourcePointsCount;
+		const uint SourcePointIndex = ThreadIndex / TargetPointsCount;
+		const uint TargetPointIndex = ThreadIndex % TargetPointsCount;
 
 		if (!Source_GetThreadData(SourcePointIndex, SourceDataIndex, SourcePointIndexInData))
 		{
@@ -180,17 +180,44 @@ void Main(uint3 GroupId : SV_GroupID, uint GroupIndex : SV_GroupIndex)
 		}
 	}
 
-	float3 SourcePosition = Source_GetPosition(SourceDataIndex,SourcePointIndexInData);
-	float4 SourceRotation = Source_GetRotation(SourceDataIndex,SourcePointIndexInData);
-	float3 SourceScale = Source_GetScale(SourceDataIndex,SourcePointIndexInData);
-	float4 SourceColor = Source_GetColor(SourceDataIndex,SourcePointIndexInData);
-	uint SourceSeed = Source_GetSeed(SourceDataIndex,SourcePointIndexInData);
+	uint OutNumData = 0;
+
+	if (bCopyEachSourceOnEveryTarget)
+	{
+		OutNumData = SourceNumData * TargetNumData;
+	}
+	else
+	{
+		OutNumData = max(SourceNumData, TargetNumData);
+	}
+
+	Out_SetNumDataInternal(OutNumData);
+
+	uint OutDataIndex, OutElemIndex;
+	if (!Out_GetThreadData(ThreadIndex, OutDataIndex, OutElemIndex))
+	{
+		return;
+	}
+
+	// Propagate the 'Removed' status to the output point if either source or target point is invalid. If we don't do this, then
+	// the output point will be considered valid and not culled even if the source or target point were culled upstream.
+	if (!Source_IsValid(SourceDataIndex, SourcePointIndexInData) || !Target_IsValid(TargetDataIndex, TargetPointIndexInData))
+	{
+		Out_RemovePoint(OutDataIndex, OutElemIndex);
+		return;
+	}
+
+	float3 SourcePosition = Source_GetPosition(SourceDataIndex, SourcePointIndexInData);
+	float4 SourceRotation = Source_GetRotation(SourceDataIndex, SourcePointIndexInData);
+	float3 SourceScale = Source_GetScale(SourceDataIndex, SourcePointIndexInData);
+	float4 SourceColor = Source_GetColor(SourceDataIndex, SourcePointIndexInData);
+	uint SourceSeed = Source_GetSeed(SourceDataIndex, SourcePointIndexInData);
 	
-	float3 TargetPosition = Target_GetPosition(TargetDataIndex,TargetPointIndexInData);
-	float4 TargetRotation = Target_GetRotation(TargetDataIndex,TargetPointIndexInData);
-	float3 TargetScale = Target_GetScale(TargetDataIndex,TargetPointIndexInData);
-	float4 TargetColor = Target_GetColor(TargetDataIndex,TargetPointIndexInData);
-	uint TargetSeed = Target_GetSeed(TargetDataIndex,TargetPointIndexInData);
+	float3 TargetPosition = Target_GetPosition(TargetDataIndex, TargetPointIndexInData);
+	float4 TargetRotation = Target_GetRotation(TargetDataIndex, TargetPointIndexInData);
+	float3 TargetScale = Target_GetScale(TargetDataIndex, TargetPointIndexInData);
+	float4 TargetColor = Target_GetColor(TargetDataIndex, TargetPointIndexInData);
+	uint TargetSeed = Target_GetSeed(TargetDataIndex, TargetPointIndexInData);
 	
 	float3 OutPosition = TargetPosition + (QuatRotateVector(TargetRotation, (SourcePosition * TargetScale)));
 	float4 OutRotation = float4(0, 0, 0, 1);
@@ -206,7 +233,7 @@ void Main(uint3 GroupId : SV_GroupID, uint GroupIndex : SV_GroupIndex)
 
 	if (RotationInheritance == COPY_POINTS_INHERITANCE_MODE_Relative)
 	{
-	  OutRotation = QuatMultiply(SourceRotation,TargetRotation);
+	  OutRotation = QuatMultiply(SourceRotation, TargetRotation);
 	}
 	else if (RotationInheritance == COPY_POINTS_INHERITANCE_MODE_Source)
 	{
@@ -255,28 +282,6 @@ void Main(uint3 GroupId : SV_GroupID, uint GroupIndex : SV_GroupIndex)
 	{
 		OutSeed = TargetSeed;
 	}
-
-	uint OutNumData = 0;
-
-	if (bCopyEachSourceOnEveryTarget)
-	{
-		OutNumData = SourceNumData * TargetNumData;
-	}
-	else
-	{
-		OutNumData = max(SourceNumData, TargetNumData);
-	}
-
-	Out_SetNumDataInternal(OutNumData);
-
-	uint OutDataIndex, OutElemIndex;
-	if (!Out_GetThreadData(ThreadIndex, OutDataIndex, OutElemIndex))
-	{
-		return;
-	}
-	
-	// Initialize all fields
-	Out_InitializePoint(OutDataIndex, OutElemIndex);
 	
 	// Write output
 	Out_SetPosition(OutDataIndex, OutElemIndex, OutPosition);
@@ -285,6 +290,12 @@ void Main(uint3 GroupId : SV_GroupID, uint GroupIndex : SV_GroupIndex)
 	Out_SetColor(OutDataIndex, OutElemIndex, OutColor);
 	Out_SetSeed(OutDataIndex, OutElemIndex, OutSeed);
 
+	// These properties do not have an inheritance mode, simply copy from the source data.
+	Out_SetBoundsMin(OutDataIndex, OutElemIndex, Source_GetBoundsMin(SourceDataIndex, SourcePointIndexInData));
+	Out_SetBoundsMax(OutDataIndex, OutElemIndex, Source_GetBoundsMax(SourceDataIndex, SourcePointIndexInData));
+	Out_SetDensity(OutDataIndex, OutElemIndex, Source_GetDensity(SourceDataIndex, SourcePointIndexInData));
+	Out_SetSteepness(OutDataIndex, OutElemIndex, Source_GetSteepness(SourceDataIndex, SourcePointIndexInData));
+
 	// Copy attributes
 	if (AttributeInheritance == COPY_POINTS_METADATA_INHERITANCE_MODE_SourceOnly)
 	{
diff --git a/Engine/Plugins/PCG/Source/PCG/Private/Compute/PCGDataForGPU.cpp b/Engine/Plugins/PCG/Source/PCG/Private/Compute/PCGDataForGPU.cpp
--- a/Engine/Plugins/PCG/Source/PCG/Private/Compute/PCGDataForGPU.cpp
+++ b/Engine/Plugins/PCG/Source/PCG/Private/Compute/PCGDataForGPU.cpp
@@ -1001,7 +1001,7 @@ EPCGUnpackDataCollectionResult FPCGDataCollectionDesc::UnpackDataCollection(cons
 								DataAsFloat[PackedDataElementIndex + 0],
 								DataAsFloat[PackedDataElementIndex + 1],
 								DataAsFloat[PackedDataElementIndex + 2],
-								DataAsFloat[PackedDataElementIndex + 2]
+								DataAsFloat[PackedDataElementIndex + 3]
 							);
 
 							OutPoints[ElementIndex].Color = Color;
