[PCG] Fix minor issue where Select Grammar comparator was showing up as BinaryOps, even though it was intended to be hidden.

#rnx

diff --git a/Engine/Plugins/PCG/Source/PCG/Private/Elements/Grammar/PCGSelectGrammar.cpp b/Engine/Plugins/PCG/Source/PCG/Private/Elements/Grammar/PCGSelectGrammar.cpp
--- a/Engine/Plugins/PCG/Source/PCG/Private/Elements/Grammar/PCGSelectGrammar.cpp
+++ b/Engine/Plugins/PCG/Source/PCG/Private/Elements/Grammar/PCGSelectGrammar.cpp
@@ -91,6 +91,28 @@ FPCGElementPtr UPCGSelectGrammarSettings::CreateElement() const
 	return MakeShared<FPCGSelectGrammarElement>();
 }
 
+void UPCGSelectGrammarSettings::PostLoad()
+{
+	Super::PostLoad();
+
+#if WITH_EDITOR
+	if (!Criteria.IsEmpty())
+	{
+		for (FPCGSelectGrammarCriterion& Criterion : Criteria)
+		{
+			if (Criterion.Comparator == EPCGSelectGrammarComparator::BinaryOps)
+			{
+				Criterion.Comparator = EPCGSelectGrammarComparator::Select;
+			}
+			else if (Criterion.Comparator == EPCGSelectGrammarComparator::TernaryOps)
+			{
+				Criterion.Comparator = EPCGSelectGrammarComparator::RangeExclusive;
+			}
+		}
+	}
+#endif // WITH_EDITOR
+}
+
 TArray<FPCGPinProperties> UPCGSelectGrammarSettings::InputPinProperties() const
 {
 	TArray<FPCGPinProperties> PinProperties;
diff --git a/Engine/Plugins/PCG/Source/PCG/Public/Elements/Grammar/PCGSelectGrammar.h b/Engine/Plugins/PCG/Source/PCG/Public/Elements/Grammar/PCGSelectGrammar.h
--- a/Engine/Plugins/PCG/Source/PCG/Public/Elements/Grammar/PCGSelectGrammar.h
+++ b/Engine/Plugins/PCG/Source/PCG/Public/Elements/Grammar/PCGSelectGrammar.h
@@ -13,7 +13,7 @@ enum class EPCGSelectGrammarComparator : uint16
 {
 	// 32 options
 	BinaryOps = 1 << 5 UMETA(Hidden),
-	Select = BinaryOps UMETA(DisplayName = "Select", ToolTip = "Select this grammar. Always evaluates to true."),
+	Select UMETA(DisplayName = "Select", ToolTip = "Select this grammar. Always evaluates to true."),
 	LessThan UMETA(DisplayName = "<", ToolTip = "Select if 'less than' (A < B)"),
 	LessThanEqualTo UMETA(DisplayName = "<=", ToolTip = "Select if 'less than or equal to' (A <= B)"),
 	EqualTo UMETA(DisplayName = "==", ToolTip = "Select if 'equal to' (A == B)"),
@@ -22,7 +22,7 @@ enum class EPCGSelectGrammarComparator : uint16
 
 	// 64 options
 	TernaryOps = 1 << 6 UMETA(Hidden),
-	RangeExclusive = TernaryOps UMETA(DisplayName = "Range (exclusive)", ToolTip = "Select if 'within exclusive range' (B < A < C)"),
+	RangeExclusive UMETA(DisplayName = "Range (exclusive)", ToolTip = "Select if 'within exclusive range' (B < A < C)"),
 	RangeInclusive UMETA(DisplayName = "Range (inclusive)", ToolTip = "Select if 'within inclusive range' (B <= A <= C)"),
 };
 
@@ -91,6 +91,7 @@ public:
 	virtual FText GetNodeTooltipText() const override;
 	virtual EPCGSettingsType GetType() const override { return EPCGSettingsType::Generic; }
 #endif
+	virtual void PostLoad() override;
 
 protected:
 	virtual TArray<FPCGPinProperties> InputPinProperties() const override;
