[PCG] Add new metadata operations: Inc/Dec/Negate

+ Unit Tests

#rnx

diff --git a/Engine/Plugins/PCG/Source/PCG/Private/Elements/Metadata/PCGMetadataMaths.inl b/Engine/Plugins/PCG/Source/PCG/Private/Elements/Metadata/PCGMetadataMaths.inl
--- a/Engine/Plugins/PCG/Source/PCG/Private/Elements/Metadata/PCGMetadataMaths.inl
+++ b/Engine/Plugins/PCG/Source/PCG/Private/Elements/Metadata/PCGMetadataMaths.inl
@@ -468,4 +468,40 @@ namespace PCGMetadataMaths
 	{
 		return FVector4::One() - Value;
 	}
+
+	////////////////////
+	// Inc/Dec
+	////////////////////
+	template <typename T>
+	inline T Inc(const T& Value)
+	{
+		return Value + T{1};
+	}
+
+	template <>
+	inline FVector4 Inc(const FVector4& Value)
+	{
+		return Value + FVector4::One();
+	}
+
+	template <typename T>
+	inline T Dec(const T& Value)
+	{
+		return Value - T{1};
+	}
+
+	template <>
+	inline FVector4 Dec(const FVector4& Value)
+	{
+		return Value - FVector4::One();
+	}
+
+	////////////////////
+	// Negate
+	////////////////////
+	template <typename T>
+	inline T Negate(const T& Value)
+	{
+		return -Value;
+	}
 }
diff --git a/Engine/Plugins/PCG/Source/PCG/Private/Elements/Metadata/PCGMetadataMathsOpElement.cpp b/Engine/Plugins/PCG/Source/PCG/Private/Elements/Metadata/PCGMetadataMathsOpElement.cpp
--- a/Engine/Plugins/PCG/Source/PCG/Private/Elements/Metadata/PCGMetadataMathsOpElement.cpp
+++ b/Engine/Plugins/PCG/Source/PCG/Private/Elements/Metadata/PCGMetadataMathsOpElement.cpp
@@ -103,6 +103,12 @@ namespace PCGMetadataMathsSettings
 			return PCGMetadataMaths::Ceil(Value);
 		case EPCGMetadataMathsOperation::OneMinus:
 			return PCGMetadataMaths::OneMinus(Value);
+		case EPCGMetadataMathsOperation::Inc:
+			return PCGMetadataMaths::Inc(Value);
+		case EPCGMetadataMathsOperation::Dec:
+			return PCGMetadataMaths::Dec(Value);
+		case EPCGMetadataMathsOperation::Negate:
+			return PCGMetadataMaths::Negate(Value);
 		default:
 			return T{};
 		}
diff --git a/Engine/Plugins/PCG/Source/PCG/Private/Tests/Elements/Metadata/PCGMetadataMathsOpTest.cpp b/Engine/Plugins/PCG/Source/PCG/Private/Tests/Elements/Metadata/PCGMetadataMathsOpTest.cpp
--- a/Engine/Plugins/PCG/Source/PCG/Private/Tests/Elements/Metadata/PCGMetadataMathsOpTest.cpp
+++ b/Engine/Plugins/PCG/Source/PCG/Private/Tests/Elements/Metadata/PCGMetadataMathsOpTest.cpp
@@ -39,7 +39,7 @@ protected:
 		int32 PinToForward = 0;
 	};
 
-	// Expects Input Data to be be an array of Tagged Data on each index per pin and a single value in each Input Data (1 attribute or 1 point data)
+	// Expects Input Data to be an array of Tagged Data on each index per pin and a single value in each Input Data (1 attribute or 1 point data)
 	template <typename T>
 	bool ExecuteTest(TestParams<T>&& Params, const FPCGAttributePropertyInputSelector InputSelector = DefaultSelector)
 	{
@@ -447,6 +447,136 @@ bool FPCGMetadataMathsOp_MultipleInput::RunTest(const FString& Parameters)
 	return bSuccess;
 }
 
-// TODO: Round, Sqrt, Abs, Floor, Ceil, Max, Min, Pow, ClampMin, ClampMax, Modulo, Set, Clamp, Lerp
+IMPLEMENT_CUSTOM_SIMPLE_AUTOMATION_TEST(FPCGMetadataMathsOp_Abs, FPCGMetadataMathsOpTest, "Plugins.PCG.Metadata.MathsOp.Abs", PCGTestsCommon::TestFlags)
 
-#endif // WITH_EDITOR
\ No newline at end of file
+bool FPCGMetadataMathsOp_Abs::RunTest(const FString& Parameters)
+{
+	bool bSuccess = true;
+
+	using EPCGMetadataMathsOperation::Abs;
+
+	// Param Data
+	int32 P1 = 5;
+	bSuccess &= ExecuteTest(TestParams(Abs, GenerateParamData<int32>({{P1}}), TArray({0}), TArray({P1})));
+	bSuccess &= ExecuteTest(TestParams(Abs, GenerateParamData<int32>({{-1 * P1}}), TArray({0}), TArray({P1})));
+
+	int64 P2 = 5ll;
+	bSuccess &= ExecuteTest(TestParams(Abs, GenerateParamData<int64>({{P2}}), TArray({0ll}), TArray({P2})));
+	bSuccess &= ExecuteTest(TestParams(Abs, GenerateParamData<int64>({{-1ll * P2}}), TArray({0ll}), TArray({P2})));
+
+	float P3 = 5.2f;
+	bSuccess &= ExecuteTest(TestParams(Abs, GenerateParamData<float>({{P3}}), TArray({0.f}), TArray({P3})));
+	bSuccess &= ExecuteTest(TestParams(Abs, GenerateParamData<float>({{-1.0f * P3}}), TArray({0.f}), TArray({P3})));
+
+	double V1[] = {0.3, 512.2, 0.0, 15.0};
+	bSuccess &= ExecuteTest(TestParams(Abs, GenerateParamData<double>({{V1[0]}}), TArray({.0}), TArray({V1[0]})));
+	bSuccess &= ExecuteTest(TestParams(Abs, GenerateParamData<double>({{-1.0 * V1[0]}}), TArray({.0}), TArray({V1[0]})));
+	bSuccess &= ExecuteTest(TestParams(Abs, GenerateParamData<FVector2D>({{FVector2D(V1[0], V1[1])}}), TArray({FVector2D::Zero()}), TArray({FVector2D(V1[0], V1[1])})));
+	bSuccess &= ExecuteTest(TestParams(Abs, GenerateParamData<FVector2D>({{-1.0 * FVector2D(V1[0], V1[1])}}), TArray({FVector2D::Zero()}), TArray({FVector2D(V1[0], V1[1])})));
+	bSuccess &= ExecuteTest(TestParams(Abs, GenerateParamData<FVector>({{FVector(V1[0], V1[1], V1[2])}}), TArray({FVector::Zero()}), TArray({FVector(V1[0], V1[1], V1[2])})));
+	bSuccess &= ExecuteTest(TestParams(Abs, GenerateParamData<FVector>({{-1.0 * FVector(V1[0], V1[1], V1[2])}}), TArray({FVector::Zero()}), TArray({FVector(V1[0], V1[1], V1[2])})));
+	bSuccess &= ExecuteTest(TestParams(Abs, GenerateParamData<FVector4>({{FVector4(V1[0], V1[1], V1[2], V1[3])}}), TArray({FVector4::Zero()}), TArray({FVector4(V1[0], V1[1], V1[2], V1[3])})));
+	bSuccess &= ExecuteTest(TestParams(Abs, GenerateParamData<FVector4>({{-1.0 * FVector4(V1[0], V1[1], V1[2], V1[3])}}), TArray({FVector4::Zero()}), TArray({FVector4(V1[0], V1[1], V1[2], V1[3])})));
+
+	// Point Data
+	FPCGAttributePropertyInputSelector Selector;
+	Selector.SetPointProperty(EPCGPointProperties::Seed);
+	bSuccess &= ExecuteTest(TestParams(Abs, GeneratePointData(), TArray<int64>({0}), TArray<int64>({1})), Selector);
+	Selector.SetPointProperty(EPCGPointProperties::Density);
+	bSuccess &= ExecuteTest(TestParams(Abs, GeneratePointData(), TArray({0.0}), TArray({1.0})), Selector);
+	Selector.SetPointProperty(EPCGPointProperties::Position);
+	bSuccess &= ExecuteTest(TestParams(Abs, GeneratePointData(), TArray({FVector::Zero()}), TArray({FVector(1.0, 10.0, 100.0)})), Selector);
+
+	return bSuccess;
+}
+
+IMPLEMENT_CUSTOM_SIMPLE_AUTOMATION_TEST(FPCGMetadataMathsOp_Inc, FPCGMetadataMathsOpTest, "Plugins.PCG.Metadata.MathsOp.Inc", PCGTestsCommon::TestFlags)
+
+bool FPCGMetadataMathsOp_Inc::RunTest(const FString& Parameters)
+{
+	bool bSuccess = true;
+
+	using EPCGMetadataMathsOperation::Inc;
+
+	// Param Data
+	bSuccess &= ExecuteTest(TestParams(Inc, GenerateParamData<int32>({{5}}), TArray({1}), TArray({6})));
+	bSuccess &= ExecuteTest(TestParams(Inc, GenerateParamData<int64>({{6ll}}), TArray({1ll}), TArray({7ll})));
+	bSuccess &= ExecuteTest(TestParams(Inc, GenerateParamData<float>({{0.1f}}), TArray({1.f}), TArray({1.1f})));
+	bSuccess &= ExecuteTest(TestParams(Inc, GenerateParamData<double>({{0.2}}), TArray({1.0}), TArray({1.2})));
+	bSuccess &= ExecuteTest(TestParams(Inc, GenerateParamData<FVector2D>({{FVector2D(0.3, 1.4)}}), TArray({FVector2D::One()}), TArray({FVector2D(1.3, 2.4)})));
+	bSuccess &= ExecuteTest(TestParams(Inc, GenerateParamData<FVector>({{FVector(0.5, 1.6, 0.7)}}), TArray({FVector::One()}), TArray({FVector(1.5, 2.6, 1.7)})));
+	bSuccess &= ExecuteTest(TestParams(Inc, GenerateParamData<FVector4>({{FVector4(0.5, 1.6, 0.7, 0.8)}}), TArray({FVector4::One()}), TArray({FVector4(1.5, 2.6, 1.7, 1.8)})));
+
+	// Point Data
+	FPCGAttributePropertyInputSelector Selector;
+	Selector.SetPointProperty(EPCGPointProperties::Seed);
+	bSuccess &= ExecuteTest(TestParams(Inc, GeneratePointData(), TArray<int64>({0}), TArray<int64>({2})), Selector);
+	Selector.SetPointProperty(EPCGPointProperties::Density);
+	bSuccess &= ExecuteTest(TestParams(Inc, GeneratePointData(), TArray({0.0}), TArray({2.0})), Selector);
+	Selector.SetPointProperty(EPCGPointProperties::Position);
+	bSuccess &= ExecuteTest(TestParams(Inc, GeneratePointData(), TArray({FVector(0.0, 0.0, 0.0)}), TArray({FVector(2.0, 11.0, 101.0)})), Selector);
+
+	return bSuccess;
+}
+
+IMPLEMENT_CUSTOM_SIMPLE_AUTOMATION_TEST(FPCGMetadataMathsOp_Dec, FPCGMetadataMathsOpTest, "Plugins.PCG.Metadata.MathsOp.Dec", PCGTestsCommon::TestFlags)
+
+bool FPCGMetadataMathsOp_Dec::RunTest(const FString& Parameters)
+{
+	bool bSuccess = true;
+
+	using EPCGMetadataMathsOperation::Dec;
+
+	// Param Data
+	bSuccess &= ExecuteTest(TestParams(Dec, GenerateParamData<int32>({{5}}), TArray({-1}), TArray({4})));
+	bSuccess &= ExecuteTest(TestParams(Dec, GenerateParamData<int64>({{6ll}}), TArray({-1ll}), TArray({5ll})));
+	bSuccess &= ExecuteTest(TestParams(Dec, GenerateParamData<float>({{0.1f}}), TArray({-1.f}), TArray({-0.9f})));
+	bSuccess &= ExecuteTest(TestParams(Dec, GenerateParamData<double>({{0.2}}), TArray({-1.0}), TArray({-0.8})));
+	bSuccess &= ExecuteTest(TestParams(Dec, GenerateParamData<FVector2D>({{FVector2D(0.3, 1.4)}}), TArray({-FVector2D::One()}), TArray({FVector2D(-0.7, 0.4)})));
+	bSuccess &= ExecuteTest(TestParams(Dec, GenerateParamData<FVector>({{FVector(0.5, 1.6, 0.7)}}), TArray({-FVector::One()}), TArray({FVector(-0.5, 0.6, -0.3)})));
+	bSuccess &= ExecuteTest(TestParams(Dec, GenerateParamData<FVector4>({{FVector4(0.5, 1.6, 0.7, 0.8)}}), TArray({-FVector4::One()}), TArray({FVector4(-0.5, 0.6, -0.3, -0.2)})));
+
+	// Point Data
+	FPCGAttributePropertyInputSelector Selector;
+	Selector.SetPointProperty(EPCGPointProperties::Seed);
+	bSuccess &= ExecuteTest(TestParams(Dec, GeneratePointData(), TArray<int64>({0}), TArray<int64>({0})), Selector);
+	Selector.SetPointProperty(EPCGPointProperties::Density);
+	bSuccess &= ExecuteTest(TestParams(Dec, GeneratePointData(), TArray({0.0}), TArray({0.0})), Selector);
+	Selector.SetPointProperty(EPCGPointProperties::Position);
+	bSuccess &= ExecuteTest(TestParams(Dec, GeneratePointData(), TArray({FVector(0.0, 0.0, 0.0)}), TArray({FVector(0.0, 9.0, 99.0)})), Selector);
+
+	return bSuccess;
+}
+
+IMPLEMENT_CUSTOM_SIMPLE_AUTOMATION_TEST(FPCGMetadataMathsOp_Negate, FPCGMetadataMathsOpTest, "Plugins.PCG.Metadata.MathsOp.Negate", PCGTestsCommon::TestFlags)
+
+bool FPCGMetadataMathsOp_Negate::RunTest(const FString& Parameters)
+{
+	bool bSuccess = true;
+
+	using EPCGMetadataMathsOperation::Negate;
+
+	// Param Data
+	bSuccess &= ExecuteTest(TestParams(Negate, GenerateParamData<int32>({{5}}), TArray({0}), TArray({-5})));
+	bSuccess &= ExecuteTest(TestParams(Negate, GenerateParamData<int64>({{6ll}}), TArray({0ll}), TArray({-6ll})));
+	bSuccess &= ExecuteTest(TestParams(Negate, GenerateParamData<float>({{0.1f}}), TArray({0.f}), TArray({-0.1f})));
+	bSuccess &= ExecuteTest(TestParams(Negate, GenerateParamData<double>({{-0.2}}), TArray({0.0}), TArray({0.2})));
+	bSuccess &= ExecuteTest(TestParams(Negate, GenerateParamData<FVector2D>({{FVector2D(0.3, -1.4)}}), TArray({FVector2D::Zero()}), TArray({FVector2D(-0.3, 1.4)})));
+	bSuccess &= ExecuteTest(TestParams(Negate, GenerateParamData<FVector>({{FVector(0.5, -1.6, 0.7)}}), TArray({FVector::Zero()}), TArray({FVector(-0.5, 1.6, -0.7)})));
+	bSuccess &= ExecuteTest(TestParams(Negate, GenerateParamData<FVector4>({{FVector4(0.5, -1.6, 0.7, -0.8)}}), TArray({FVector4::Zero()}), TArray({FVector4(-0.5, 1.6, -0.7, 0.8)})));
+
+	// Point Data
+	FPCGAttributePropertyInputSelector Selector;
+	Selector.SetPointProperty(EPCGPointProperties::Seed);
+	bSuccess &= ExecuteTest(TestParams(Negate, GeneratePointData(), TArray<int64>({0}), TArray<int64>({-1})), Selector);
+	Selector.SetPointProperty(EPCGPointProperties::Density);
+	bSuccess &= ExecuteTest(TestParams(Negate, GeneratePointData(), TArray({0.0}), TArray({-1.0})), Selector);
+	Selector.SetPointProperty(EPCGPointProperties::Position);
+	bSuccess &= ExecuteTest(TestParams(Negate, GeneratePointData(), TArray({FVector(0.0, 0.0, 0.0)}), TArray({FVector(-1.0, -10.0, -100.0)})), Selector);
+
+	return bSuccess;
+}
+
+// TODO: Round, Sqrt, Floor, Ceil, Max, Min, Pow, ClampMin, ClampMax, Modulo, Set, Clamp, Lerp
+
+#endif // WITH_EDITOR
diff --git a/Engine/Plugins/PCG/Source/PCG/Public/Elements/Metadata/PCGMetadataMathsOpElement.h b/Engine/Plugins/PCG/Source/PCG/Public/Elements/Metadata/PCGMetadataMathsOpElement.h
--- a/Engine/Plugins/PCG/Source/PCG/Public/Elements/Metadata/PCGMetadataMathsOpElement.h
+++ b/Engine/Plugins/PCG/Source/PCG/Public/Elements/Metadata/PCGMetadataMathsOpElement.h
@@ -20,6 +20,9 @@ enum class EPCGMetadataMathsOperation : uint16
 	Floor,
 	Ceil,
 	OneMinus UMETA(Tooltip = "1 - X operation"),
+	Inc UMETA(Tooltip = "X + 1 operation"),
+	Dec UMETA(Tooltip = "X - 1 operation"),
+	Negate UMETA(Tooltip = "-X operation"),
 
 	// Binary op
 	BinaryOp = 1 << 11 UMETA(Hidden),
