Added a cvar to control behavior of dirty-on-load for load as preview components (default: on)

#rnx

diff --git a/Engine/Plugins/PCG/Source/PCG/Private/PCGCommon.cpp b/Engine/Plugins/PCG/Source/PCG/Private/PCGCommon.cpp
--- a/Engine/Plugins/PCG/Source/PCG/Private/PCGCommon.cpp
+++ b/Engine/Plugins/PCG/Source/PCG/Private/PCGCommon.cpp
@@ -35,9 +35,14 @@ namespace PCGSystemSwitches
 		TEXT("Pauses all execution of PCG but does not cancel tasks."));
 
 	TAutoConsoleVariable<bool> CVarGlobalDisableRefresh(
-			TEXT("pcg.GlobalDisableRefresh"),
-			false,
-			TEXT("Disable refresh for all PCG Components."));
+		TEXT("pcg.GlobalDisableRefresh"),
+		false,
+		TEXT("Disable refresh for all PCG Components."));
+
+	TAutoConsoleVariable<bool> CVarDirtyLoadAsPreviewOnLoad(
+		TEXT("pcg.DirtyLoadAsPreviewOnLoad"),
+		true,
+		TEXT("Enables or not dirtying on load for load as preview components.\nTurning off this option will require to force generate or apply a change before this component is regenerated."));
 #endif
 }
 
diff --git a/Engine/Plugins/PCG/Source/PCG/Private/PCGComponent.cpp b/Engine/Plugins/PCG/Source/PCG/Private/PCGComponent.cpp
--- a/Engine/Plugins/PCG/Source/PCG/Private/PCGComponent.cpp
+++ b/Engine/Plugins/PCG/Source/PCG/Private/PCGComponent.cpp
@@ -1734,7 +1734,7 @@ void UPCGComponent::PostLoad()
 	{
 		CurrentEditingMode = EPCGEditorDirtyMode::Preview;
 		MarkResourcesAsTransientOnLoad();
-		bDirtyGenerated = true;
+		bDirtyGenerated = PCGSystemSwitches::CVarDirtyLoadAsPreviewOnLoad.GetValueOnAnyThread();
 	}
 #endif
 
diff --git a/Engine/Plugins/PCG/Source/PCG/Public/PCGCommon.h b/Engine/Plugins/PCG/Source/PCG/Public/PCGCommon.h
--- a/Engine/Plugins/PCG/Source/PCG/Public/PCGCommon.h
+++ b/Engine/Plugins/PCG/Source/PCG/Public/PCGCommon.h
@@ -193,6 +193,7 @@ namespace PCGSystemSwitches
 #if WITH_EDITOR
 	extern PCG_API TAutoConsoleVariable<bool> CVarPausePCGExecution;
 	extern TAutoConsoleVariable<bool> CVarGlobalDisableRefresh;
+	extern TAutoConsoleVariable<bool> CVarDirtyLoadAsPreviewOnLoad;
 #endif
 }
 
