Refactored editor dialogs so they support object replacement, update less often etc.

#rnx

diff --git a/Engine/Plugins/PCG/Source/PCG/Private/PCGComponent.cpp b/Engine/Plugins/PCG/Source/PCG/Private/PCGComponent.cpp
--- a/Engine/Plugins/PCG/Source/PCG/Private/PCGComponent.cpp
+++ b/Engine/Plugins/PCG/Source/PCG/Private/PCGComponent.cpp
@@ -650,7 +650,7 @@ void UPCGComponent::PostProcessGraph(const FBox& InNewBounds, bool bInGenerated,
 	StopGenerationInProgress();
 
 	UPCGSubsystem* Subsystem = GetSubsystem();
-	Subsystem->OnComponentGenerationCompleteOrCancelled.Broadcast(Subsystem);
+	Subsystem->OnComponentGenerationCompleteOrCancelled.Broadcast(Subsystem, IsValid(this) ? this : nullptr);
 #endif
 }
 
@@ -758,7 +758,7 @@ void UPCGComponent::OnProcessGraphAborted(bool bQuiet)
 	OnPCGGraphCancelledDelegate.Broadcast(this);
 
 	UPCGSubsystem* Subsystem = GetSubsystem();
-	Subsystem->OnComponentGenerationCompleteOrCancelled.Broadcast(Subsystem);
+	Subsystem->OnComponentGenerationCompleteOrCancelled.Broadcast(Subsystem, this);
 #endif
 
 	PCGComponent::BroadcastDynamicDelegate(OnPCGGraphCancelledExternal, this);
@@ -1555,7 +1555,7 @@ void UPCGComponent::OnUnregister()
 			Subsystem->CancelGeneration(this);
 		}
 
-		Subsystem->OnComponentUnregistered.Broadcast();
+		Subsystem->OnComponentUnregistered.Broadcast(IsValid(this) ? this : nullptr);
 	}
 #endif // WITH_EDITOR
 
@@ -3657,6 +3657,9 @@ void FPCGComponentInstanceData::ApplyToComponent(UActorComponent* Component, con
 			PCGComponent->OnPCGGraphCancelledExternal = ConstructionSourceComponent->OnPCGGraphCancelledExternal;
 			PCGComponent->OnPCGGraphGeneratedExternal = ConstructionSourceComponent->OnPCGGraphGeneratedExternal;
 			PCGComponent->OnPCGGraphCleanedExternal = ConstructionSourceComponent->OnPCGGraphCleanedExternal;
+
+			PCGComponent->bWasGeneratedThisSession = ConstructionSourceComponent->bWasGeneratedThisSession;
+			PCGComponent->InspectionCounter = ConstructionSourceComponent->InspectionCounter;
 #endif
 		}
 
diff --git a/Engine/Plugins/PCG/Source/PCG/Public/Graph/PCGStackContext.h b/Engine/Plugins/PCG/Source/PCG/Public/Graph/PCGStackContext.h
--- a/Engine/Plugins/PCG/Source/PCG/Public/Graph/PCGStackContext.h
+++ b/Engine/Plugins/PCG/Source/PCG/Public/Graph/PCGStackContext.h
@@ -25,15 +25,13 @@ struct PCG_API FPCGStackFrame
 	FPCGStackFrame() {}
 
 	explicit FPCGStackFrame(TWeakObjectPtr<const UObject> InObject)
-		: Object(InObject)
 	{
-		Hash = PointerHash(Object.Get());
+		SetObject(InObject);
 	}
 
 	explicit FPCGStackFrame(int32 InLoopIndex)
-		: LoopIndex(InLoopIndex)
 	{
-		Hash = GetTypeHash(LoopIndex);
+		SetLoopIndex(InLoopIndex);
 	}
 
 	bool operator==(const FPCGStackFrame& Other) const { return Object == Other.Object && LoopIndex == Other.LoopIndex; }
@@ -44,6 +42,20 @@ struct PCG_API FPCGStackFrame
 	// A valid frame should either point to an object or have a loop index >= 0.
 	bool IsValid() const { return LoopIndex != INDEX_NONE || Object.IsValid(); }
 
+	void SetObject(TWeakObjectPtr<const UObject> InObject)
+	{
+		Object = InObject;
+		LoopIndex = INDEX_NONE;
+		Hash = PointerHash(Object.Get());
+	}
+
+	void SetLoopIndex(int32 InLoopIndex)
+	{
+		Object = nullptr;
+		LoopIndex = InLoopIndex;
+		Hash = GetTypeHash(LoopIndex);
+	}
+
 	TWeakObjectPtr<const UObject> Object;
 	int32 LoopIndex = INDEX_NONE;
 
@@ -124,14 +136,15 @@ public:
 		return Hash;
 	}
 
+private:
+	TArray<FPCGStackFrame> StackFrames;
+
+public:
 #if WITH_EDITOR
 	// Used to store node & hierarchy information
 	UE_DEPRECATED(5.5, "The timer has been moved to the FPCGContext struct.")
 	PCGUtils::FCallTime Timer;
 #endif
-
-private:
-	TArray<FPCGStackFrame> StackFrames;
 };
 
 /** A collection of call stacks. */
diff --git a/Engine/Plugins/PCG/Source/PCG/Public/PCGSubsystem.h b/Engine/Plugins/PCG/Source/PCG/Public/PCGSubsystem.h
--- a/Engine/Plugins/PCG/Source/PCG/Public/PCGSubsystem.h
+++ b/Engine/Plugins/PCG/Source/PCG/Public/PCGSubsystem.h
@@ -36,8 +36,8 @@ typedef TSharedPtr<IPCGElement, ESPMode::ThreadSafe> FPCGElementPtr;
 class UWorld;
 
 #if WITH_EDITOR
-DECLARE_MULTICAST_DELEGATE(FPCGOnComponentUnregistered);
-DECLARE_MULTICAST_DELEGATE_OneParam(FPCGOnComponentGenerationCompleteOrCancelled, UPCGSubsystem*);
+DECLARE_MULTICAST_DELEGATE_OneParam(FPCGOnComponentUnregistered, UPCGComponent*);
+DECLARE_MULTICAST_DELEGATE_TwoParams(FPCGOnComponentGenerationCompleteOrCancelled, UPCGSubsystem*, UPCGComponent*);
 #endif // WITH_EDITOR
 
 /**
diff --git a/Engine/Plugins/PCG/Source/PCGEditor/Private/PCGEditor.cpp b/Engine/Plugins/PCG/Source/PCGEditor/Private/PCGEditor.cpp
--- a/Engine/Plugins/PCG/Source/PCGEditor/Private/PCGEditor.cpp
+++ b/Engine/Plugins/PCG/Source/PCGEditor/Private/PCGEditor.cpp
@@ -112,6 +112,7 @@ UPCGEditorGraph* FPCGEditor::GetPCGEditorGraph(UPCGGraph* InGraph)
 
 void FPCGEditor::Initialize(const EToolkitMode::Type InMode, const TSharedPtr<class IToolkitHost>& InToolkitHost, UPCGGraph* InPCGGraph)
 {
+	FCoreUObjectDelegates::OnObjectsReplaced.AddSP(this, &FPCGEditor::OnObjectsReplaced);
 	PCGGraphBeingEdited = InPCGGraph;
 
 	// Initializes the UPCGEditorGraph if needed
@@ -234,6 +235,34 @@ void FPCGEditor::Initialize(const EToolkitMode::Type InMode, const TSharedPtr<cl
 	}
 }
 
+void FPCGEditor::OnObjectsReplaced(const TMap<UObject*, UObject*>& ReplacementMap)
+{
+	if (UObject* NewComponentBeingInspected = ReplacementMap.FindRef(PCGComponentBeingInspected.GetEvenIfUnreachable()))
+	{
+		PCGComponentBeingInspected = Cast<UPCGComponent>(NewComponentBeingInspected);
+	}
+
+	if (UObject* NewLastValidPCGComponentBeingInspected = ReplacementMap.FindRef(LastValidPCGComponentBeingInspected.GetEvenIfUnreachable()))
+	{
+		LastValidPCGComponentBeingInspected = Cast<UPCGComponent>(NewLastValidPCGComponentBeingInspected);
+	}
+
+	TArray<FPCGStackFrame>& StackFrames = StackBeingInspected.GetStackFramesMutable();
+	if (!StackFrames.IsEmpty())
+	{
+		if (UObject* NewStackRoot = ReplacementMap.FindRef(StackFrames[0].Object.GetEvenIfUnreachable()))
+		{
+			StackFrames[0].SetObject(NewStackRoot);
+		}
+	}
+
+	// Propagate object replacement to the debug object tree view too
+	if (DebugObjectTreeWidget)
+	{
+		DebugObjectTreeWidget->OnObjectsReplaced(ReplacementMap);
+	}
+}
+
 UPCGEditorGraph* FPCGEditor::GetPCGEditorGraph()
 {
 	return PCGEditorGraph;
@@ -278,6 +307,20 @@ void FPCGEditor::SetStackBeingInspected(const FPCGStack& FullStack)
 	UpdateAfterInspectedStackChanged();
 }
 
+void FPCGEditor::OnComponentGenerated(UPCGComponent* InComponent)
+{
+	if (DebugObjectTreeWidget)
+	{
+		DebugObjectTreeWidget->RequestRefresh();
+	}
+
+	if(InComponent == GetPCGComponentBeingInspected())
+	{
+		TRACE_CPUPROFILER_EVENT_SCOPE(FPCGEditor::SetStackBeingInspected::BroadcastStackBeingInspected);
+		OnInspectedStackChangedDelegate.Broadcast(StackBeingInspected);
+	}
+}
+
 void FPCGEditor::UpdateAfterInspectedStackChanged()
 {
 	UPCGComponent* Component = PCGComponentBeingInspected.Get();
@@ -2870,6 +2913,8 @@ void FPCGEditor::OnClose()
 		UnregisterDelegatesForWorld(GEditor->GetEditorWorldContext().World());
 		UnregisterDelegatesForWorld(GEditor->PlayWorld.Get());
 	}
+
+	FCoreUObjectDelegates::OnObjectsReplaced.RemoveAll(this);
 }
 
 void FPCGEditor::InitToolMenuContext(FToolMenuContext& MenuContext)
@@ -3017,15 +3062,23 @@ void FPCGEditor::JumpToDefinition(const UClass* Class) const
 	}
 }
 
-void FPCGEditor::OnComponentUnregistered()
+void FPCGEditor::OnComponentUnregistered(UPCGComponent* Component)
 {
 	// Refresh the debug object tree to avoid stale entries from components that have been unregistered.
-	DebugObjectTreeWidget->RequestRefresh();
+	if (!Component || Component->GetGraph() == PCGGraphBeingEdited)
+	{
+		DebugObjectTreeWidget->RequestRefresh();
+	}
 }
 
-void FPCGEditor::OnComponentGenerationCompleteOrCancelled(UPCGSubsystem* Subsystem)
+void FPCGEditor::OnComponentGenerationCompleteOrCancelled(UPCGSubsystem* Subsystem, UPCGComponent* Component)
 {
-	DebugObjectTreeWidget->RequestRefresh();
+	if(Component && Component->GetGraph() != PCGGraphBeingEdited)
+	{
+		return;
+	}
+
+	OnComponentGenerated(Component);
 
 	const bool CacheDebuggingEnabled = Subsystem && Subsystem->IsGraphCacheDebuggingEnabled();
 
diff --git a/Engine/Plugins/PCG/Source/PCGEditor/Private/PCGEditor.h b/Engine/Plugins/PCG/Source/PCGEditor/Private/PCGEditor.h
--- a/Engine/Plugins/PCG/Source/PCGEditor/Private/PCGEditor.h
+++ b/Engine/Plugins/PCG/Source/PCGEditor/Private/PCGEditor.h
@@ -108,6 +108,9 @@ private:
 	/** Bind commands to delegates */
 	void BindCommands();
 
+	/** Callback to process component changes */
+	void OnObjectsReplaced(const TMap<UObject*, UObject*>& ReplacementMap);
+
 	/** Bring up the find tab */
 	void OnFind();
 
@@ -307,6 +310,9 @@ private:
 	/** Called when the selection changes in the GraphEditor */
 	void OnSelectedNodesChanged(const TSet<UObject*>& NewSelection);
 
+	/** Called when the component inspected is generated/cleaned */
+	void OnComponentGenerated(UPCGComponent* InComponent);
+
 	/** Called when the title of a node is changed */
 	void OnNodeTitleCommitted(const FText& NewText, ETextCommit::Type CommitInfo, UEdGraphNode* NodeBeingChanged);
 
@@ -328,10 +334,10 @@ private:
 	void ReplicateExtraNodes() const;
 
 	/** Called when a PCG component unregisters. */
-	void OnComponentUnregistered();
+	void OnComponentUnregistered(UPCGComponent* Component);
 
 	/** Called when a component finishes executing. Useful for updating debugging tools/UIs. */
-	void OnComponentGenerationCompleteOrCancelled(UPCGSubsystem* Subsystem);
+	void OnComponentGenerationCompleteOrCancelled(UPCGSubsystem* Subsystem, UPCGComponent* Component);
 
 	/** Trigger any generation required to ensure debug display is up to date. */
 	void UpdateDebugAfterComponentSelection(UPCGComponent* InOldComponent, UPCGComponent* InNewComponent, bool bNewComponentStartedInspecting);
diff --git a/Engine/Plugins/PCG/Source/PCGEditor/Private/SPCGEditorGraphAttributeListView.cpp b/Engine/Plugins/PCG/Source/PCGEditor/Private/SPCGEditorGraphAttributeListView.cpp
--- a/Engine/Plugins/PCG/Source/PCGEditor/Private/SPCGEditorGraphAttributeListView.cpp
+++ b/Engine/Plugins/PCG/Source/PCGEditor/Private/SPCGEditorGraphAttributeListView.cpp
@@ -556,22 +556,13 @@ TSharedRef<SHeaderRow> SPCGEditorGraphAttributeListView::CreateHeaderRowWidget()
 	return SNew(SHeaderRow);
 }
 
-void SPCGEditorGraphAttributeListView::OnInspectedStackChanged(const FPCGStack& InPCGStack)
+TWeakObjectPtr<UPCGComponent> SPCGEditorGraphAttributeListView::GetPCGComponent() const
 {
-	if (PCGComponent.IsValid())
-	{
-		PCGComponent->OnPCGGraphGeneratedDelegate.RemoveAll(this);
-		PCGComponent->OnPCGGraphCleanedDelegate.RemoveAll(this);
-	}
-
-	PCGComponent = const_cast<UPCGComponent*>(InPCGStack.GetRootComponent());
-
-	if (PCGComponent.IsValid())
-	{
-		PCGComponent->OnPCGGraphGeneratedDelegate.AddSP(this, &SPCGEditorGraphAttributeListView::OnGenerateUpdated);
-		PCGComponent->OnPCGGraphCleanedDelegate.AddSP(this, &SPCGEditorGraphAttributeListView::OnGenerateUpdated);
-	}
+	return PCGEditorPtr.IsValid() ? PCGEditorPtr.Pin()->GetPCGComponentBeingInspected() : nullptr;
+}
 
+void SPCGEditorGraphAttributeListView::OnInspectedStackChanged(const FPCGStack& InPCGStack)
+{
 	RequestRefresh();
 }
 
@@ -613,6 +604,7 @@ void SPCGEditorGraphAttributeListView::OnGenerateUpdated(UPCGComponent* /*InPCGC
 
 const FPCGDataCollection* SPCGEditorGraphAttributeListView::GetInspectionData(const TSharedPtr<FPinComboBoxItem>& EditorPin) const
 {
+	TWeakObjectPtr<UPCGComponent> PCGComponent = GetPCGComponent();
 	if (!PCGComponent.IsValid())
 	{
 		return nullptr;
diff --git a/Engine/Plugins/PCG/Source/PCGEditor/Private/SPCGEditorGraphAttributeListView.h b/Engine/Plugins/PCG/Source/PCGEditor/Private/SPCGEditorGraphAttributeListView.h
--- a/Engine/Plugins/PCG/Source/PCGEditor/Private/SPCGEditorGraphAttributeListView.h
+++ b/Engine/Plugins/PCG/Source/PCGEditor/Private/SPCGEditorGraphAttributeListView.h
@@ -235,12 +235,12 @@ private:
 	void FocusOnSelection() const;
 	bool CanFocusOnSelection() const;
 
+	/** Returns the current PCG component */
+	TWeakObjectPtr<UPCGComponent> GetPCGComponent() const;
+
 	/** Pointer back to the PCG editor that owns us */
 	TWeakPtr<FPCGEditor> PCGEditorPtr;
 
-	/** Cached PCGComponent being viewed */
-	TWeakObjectPtr<UPCGComponent> PCGComponent;
-
 	/** Cached PCGGraphNode being viewed */
 	TWeakObjectPtr<UPCGEditorGraphNodeBase> PCGEditorGraphNode;
 
diff --git a/Engine/Plugins/PCG/Source/PCGEditor/Private/SPCGEditorGraphDebugObjectTree.cpp b/Engine/Plugins/PCG/Source/PCGEditor/Private/SPCGEditorGraphDebugObjectTree.cpp
--- a/Engine/Plugins/PCG/Source/PCGEditor/Private/SPCGEditorGraphDebugObjectTree.cpp
+++ b/Engine/Plugins/PCG/Source/PCGEditor/Private/SPCGEditorGraphDebugObjectTree.cpp
@@ -134,6 +134,21 @@ void FPCGEditorGraphDebugObjectItem::SortChildren(bool bIsAscending, bool bIsRec
 	}
 }
 
+void FPCGEditorGraphDebugObjectItem::OnObjectsReplaced(const TMap<UObject*, UObject*>& ReplacementMap)
+{
+	if (FPCGStack* Stack = GetMutablePCGStack())
+	{
+		TArray<FPCGStackFrame>& StackFrames = Stack->GetStackFramesMutable();
+		if (!StackFrames.IsEmpty())
+		{
+			if (UObject* NewStackRoot = ReplacementMap.FindRef(StackFrames[0].Object.GetEvenIfUnreachable()))
+			{
+				StackFrames[0].SetObject(NewStackRoot);
+			}
+		}
+	}
+}
+
 FString FPCGEditorGraphDebugObjectItem_Actor::GetLabel() const
 {
 	return Actor.IsValid() ? Actor->GetActorNameOrLabel() : FString();
@@ -149,6 +164,16 @@ FString FPCGEditorGraphDebugObjectItem_PCGComponent::GetLabel() const
 	return FString();
 }
 
+void FPCGEditorGraphDebugObjectItem_PCGComponent::OnObjectsReplaced(const TMap<UObject*, UObject*>& ReplacementMap)
+{
+	FPCGEditorGraphDebugObjectItem::OnObjectsReplaced(ReplacementMap);
+
+	if (UObject* NewComponent = ReplacementMap.FindRef(PCGComponent.GetEvenIfUnreachable()))
+	{
+		PCGComponent = Cast<UPCGComponent>(NewComponent);
+	}
+}
+
 FString FPCGEditorGraphDebugObjectItem_PCGSubgraph::GetLabel() const
 {
 	if (PCGNode.IsValid() && PCGGraph.IsValid())
@@ -383,6 +408,24 @@ void SPCGEditorGraphDebugObjectTree::SetDebugObjectSelection(const FPCGStack& Fu
 	bDisableDebugObjectChangeNotification = false;
 }
 
+void SPCGEditorGraphDebugObjectTree::OnObjectsReplaced(const TMap<UObject*, UObject*>& ReplacementMap)
+{
+	for (FPCGEditorGraphDebugObjectItemPtr& GraphItem : AllGraphItems)
+	{
+		if (GraphItem)
+		{
+			GraphItem->OnObjectsReplaced(ReplacementMap);
+		}
+	}
+
+	// While we'll update the selected original component here (might make stack matching more efficient),
+	// we will not update the previously selected stack, as the selection will be redone anyway.
+	if (UObject* NewComponent = ReplacementMap.FindRef(SelectedOriginalComponent.GetEvenIfUnreachable()))
+	{
+		SelectedOriginalComponent = Cast<UPCGComponent>(NewComponent);
+	}
+}
+
 void SPCGEditorGraphDebugObjectTree::SetNodeBeingInspected(const UPCGNode* InPCGNode)
 {
 	PCGNodeBeingInspected = InPCGNode;
@@ -841,7 +884,11 @@ void SPCGEditorGraphDebugObjectTree::RestoreTreeState()
 
 		if (ItemStack && SelectedStack == *ItemStack)
 		{
-			DebugObjectTreeView->SetItemSelection(Item, true);
+			if (!DebugObjectTreeView->IsItemSelected(Item))
+			{
+				DebugObjectTreeView->SetItemSelection(Item, true);
+			}
+			
 			bFoundMatchingStack = true;
 			break;
 		}
@@ -919,7 +966,11 @@ void SPCGEditorGraphDebugObjectTree::RestoreTreeState()
 					Parent = Parent->GetParent();
 				}
 
-				DebugObjectTreeView->SetItemSelection(Item, true);
+				if (!DebugObjectTreeView->IsItemSelected(Item))
+				{
+					DebugObjectTreeView->SetItemSelection(Item, true);
+				}
+				
 				break;
 			}
 		}
diff --git a/Engine/Plugins/PCG/Source/PCGEditor/Private/SPCGEditorGraphDebugObjectTree.h b/Engine/Plugins/PCG/Source/PCGEditor/Private/SPCGEditorGraphDebugObjectTree.h
--- a/Engine/Plugins/PCG/Source/PCGEditor/Private/SPCGEditorGraphDebugObjectTree.h
+++ b/Engine/Plugins/PCG/Source/PCGEditor/Private/SPCGEditorGraphDebugObjectTree.h
@@ -53,7 +53,11 @@ public:
 	virtual const UPCGGraph* GetPCGGraph() const { return nullptr; }
 	virtual bool IsLoopIteration() const { return false; }
 
+	virtual void OnObjectsReplaced(const TMap<UObject*, UObject*>& ReplacementMap);
+
 protected:
+	virtual FPCGStack* GetMutablePCGStack() { return nullptr; }
+
 	TWeakPtr<FPCGEditorGraphDebugObjectItem> Parent;
 	TSet<TSharedPtr<FPCGEditorGraphDebugObjectItem>> Children;
 	bool bIsExpanded = false;
@@ -76,6 +80,8 @@ public:
 	virtual const FPCGStack* GetPCGStack() const override { return &PCGStack; }
 
 protected:
+	virtual FPCGStack* GetMutablePCGStack() override { return &PCGStack; }
+
 	TWeakObjectPtr<AActor> Actor = nullptr;
 
 	FPCGStack PCGStack;
@@ -104,7 +110,11 @@ public:
 	virtual bool IsDebuggable() const override { return bIsDebuggable; }
 	virtual const UPCGGraph* GetPCGGraph() const override { return PCGGraph.Get(); }
 
+	virtual void OnObjectsReplaced(const TMap<UObject*, UObject*>& ReplacementMap) override;
+
 protected:
+	virtual FPCGStack* GetMutablePCGStack() override { return &PCGStack; }
+
 	TWeakObjectPtr<UPCGComponent> PCGComponent = nullptr;
 	TWeakObjectPtr<const UPCGGraph> PCGGraph = nullptr;
 
@@ -136,6 +146,8 @@ public:
 	virtual const UPCGGraph* GetPCGGraph() const override { return PCGGraph.Get(); }
 
 protected:
+	virtual FPCGStack* GetMutablePCGStack() override { return &PCGStack; }
+
 	TWeakObjectPtr<const UPCGNode> PCGNode = nullptr;
 	TWeakObjectPtr<const UPCGGraph> PCGGraph = nullptr;
 
@@ -171,6 +183,8 @@ public:
 	virtual const UPCGGraph* GetPCGGraph() const override { return Cast<UPCGGraph>(LoopedPCGGraph.Get()); }
 
 protected:
+	virtual FPCGStack* GetMutablePCGStack() override { return &PCGStack; }
+
 	int32 LoopIndex = INDEX_NONE;
 	TWeakObjectPtr<const UObject> LoopedPCGGraph = nullptr;
 
@@ -220,6 +234,8 @@ public:
 
 	void SetDebugObjectSelection(const FPCGStack& FullStack);
 
+	void OnObjectsReplaced(const TMap<UObject*, UObject*>& ReplacementMap);
+
 private:
 	void SelectedDebugObject_OnClicked() const;
 	bool IsSelectDebugObjectButtonEnabled() const;
diff --git a/Engine/Plugins/PCG/Source/PCGEditor/Private/SPCGEditorGraphLogView.cpp b/Engine/Plugins/PCG/Source/PCGEditor/Private/SPCGEditorGraphLogView.cpp
--- a/Engine/Plugins/PCG/Source/PCGEditor/Private/SPCGEditorGraphLogView.cpp
+++ b/Engine/Plugins/PCG/Source/PCGEditor/Private/SPCGEditorGraphLogView.cpp
@@ -120,8 +120,6 @@ void SPCGEditorGraphLogView::Construct(const FArguments& InArgs, TSharedPtr<FPCG
 	if (PCGEditor)
 	{
 		PCGEditorGraph = PCGEditor->GetPCGEditorGraph();
-		PCGComponent = PCGEditor->GetPCGComponentBeingInspected();
-
 		PCGEditor->OnInspectedStackChangedDelegate.AddSP(this, &SPCGEditorGraphLogView::OnDebugStackChanged);
 	}
 
@@ -194,6 +192,17 @@ void SPCGEditorGraphLogView::Construct(const FArguments& InArgs, TSharedPtr<FPCG
 	Refresh();
 }
 
+void SPCGEditorGraphLogView::Tick(const FGeometry& AllottedGeometry, const double InCurrentTime, const float InDeltaTime)
+{
+	SCompoundWidget::Tick(AllottedGeometry, InCurrentTime, InDeltaTime);
+
+	if (bNeedsRefresh)
+	{
+		bNeedsRefresh = false;
+		Refresh();
+	}
+}
+
 TSharedRef<SHeaderRow> SPCGEditorGraphLogView::CreateHeaderRowWidget()
 {
 	return SNew(SHeaderRow)
@@ -268,7 +277,7 @@ FReply SPCGEditorGraphLogView::Clear()
 		return FReply::Handled();
 	}
 
-	if (const UPCGComponent* Component = PCGComponent.Get())
+	if (const UPCGComponent* Component = GetPCGComponent().Get())
 	{
 		Component->ExtraCapture.ResetCapturedMessages();
 		Refresh();
@@ -293,7 +302,7 @@ FReply SPCGEditorGraphLogView::Refresh()
 		return FReply::Handled();
 	}
 
-	const UPCGComponent* Component = PCGComponent.Get();
+	const UPCGComponent* Component = GetPCGComponent().Get();
 
 	if (!Component)
 	{
@@ -349,7 +358,7 @@ void SPCGEditorGraphLogView::CreateAndAddItem(const UPCGEditorGraphNode* InPCGEd
 
 	check(PCGEditor);
 
-	const UPCGComponent* Component = PCGComponent.Get();
+	const UPCGComponent* Component = GetPCGComponent().Get();
 
 	if (!Component || !InPCGEditorNode || !InPCGNode)
 	{
@@ -398,26 +407,14 @@ void SPCGEditorGraphLogView::CreateAndAddItem(const UPCGEditorGraphNode* InPCGEd
 	}
 }
 
-void SPCGEditorGraphLogView::OnDebugStackChanged(const FPCGStack& InPCGStack)
+TWeakObjectPtr<UPCGComponent> SPCGEditorGraphLogView::GetPCGComponent() const
 {
-	if (PCGComponent.IsValid())
-	{
-		PCGComponent->OnPCGGraphGeneratedDelegate.RemoveAll(this);
-	}
-
-	PCGComponent = const_cast<UPCGComponent*>(InPCGStack.GetRootComponent());
-
-	if (PCGComponent.IsValid())
-	{
-		PCGComponent->OnPCGGraphGeneratedDelegate.AddSP(this, &SPCGEditorGraphLogView::OnGenerateUpdated);
-	}
-
-	Refresh();
+	return PCGEditorPtr.IsValid() ? PCGEditorPtr.Pin()->GetPCGComponentBeingInspected() : nullptr;
 }
 
-void SPCGEditorGraphLogView::OnGenerateUpdated(UPCGComponent* InPCGComponent)
+void SPCGEditorGraphLogView::OnDebugStackChanged(const FPCGStack& InPCGStack)
 {
-	Refresh();
+	RequestRefresh();
 }
 
 TSharedRef<ITableRow> SPCGEditorGraphLogView::OnGenerateRow(PCGLogListViewItemPtr Item, const TSharedRef<STableViewBase>& OwnerTable) const
diff --git a/Engine/Plugins/PCG/Source/PCGEditor/Private/SPCGEditorGraphLogView.h b/Engine/Plugins/PCG/Source/PCGEditor/Private/SPCGEditorGraphLogView.h
--- a/Engine/Plugins/PCG/Source/PCGEditor/Private/SPCGEditorGraphLogView.h
+++ b/Engine/Plugins/PCG/Source/PCGEditor/Private/SPCGEditorGraphLogView.h
@@ -59,6 +59,8 @@ public:
 	~SPCGEditorGraphLogView();
 
 	void Construct(const FArguments& InArgs, TSharedPtr<FPCGEditor> InPCGEditor);
+	virtual void Tick(const FGeometry& AllottedGeometry, const double InCurrentTime, const float InDeltaTime) override;
+	void RequestRefresh() { bNeedsRefresh = true; }
 
 private:
 	TSharedRef<SHeaderRow> CreateHeaderRowWidget();
@@ -66,8 +68,6 @@ private:
 
 	void OnDebugStackChanged(const FPCGStack& InPCGStack);
 
-	void OnGenerateUpdated(UPCGComponent* InPCGComponent);
-	
 	// Callbacks
 	FReply Refresh();
 	FReply Clear();
@@ -83,12 +83,14 @@ private:
 	UPCGEditorGraph* PCGEditorGraph = nullptr;
 
 	/** Cached PCGComponent being viewed */
-	TWeakObjectPtr<UPCGComponent> PCGComponent;
+	TWeakObjectPtr<UPCGComponent> GetPCGComponent() const;
 
 	TSharedPtr<SHeaderRow> ListViewHeader;
 	TSharedPtr<SListView<PCGLogListViewItemPtr>> ListView;
 	TArray<PCGLogListViewItemPtr> ListViewItems;
 
+	bool bNeedsRefresh : 1 = false;
+
 	// To allow sorting
 	FName SortingColumn = PCGEditorGraphLogView::NAME_Order;
 	EColumnSortMode::Type SortMode = EColumnSortMode::Type::Ascending;
diff --git a/Engine/Plugins/PCG/Source/PCGEditor/Private/SPCGEditorGraphProfilingView.cpp b/Engine/Plugins/PCG/Source/PCGEditor/Private/SPCGEditorGraphProfilingView.cpp
--- a/Engine/Plugins/PCG/Source/PCGEditor/Private/SPCGEditorGraphProfilingView.cpp
+++ b/Engine/Plugins/PCG/Source/PCGEditor/Private/SPCGEditorGraphProfilingView.cpp
@@ -216,8 +216,6 @@ void SPCGEditorGraphProfilingView::Construct(const FArguments& InArgs, TSharedPt
 	if (PCGEditor)
 	{
 		PCGEditorGraph = PCGEditor->GetPCGEditorGraph();
-		PCGComponent = PCGEditor->GetPCGComponentBeingInspected();
-
 		PCGEditor->OnInspectedStackChangedDelegate.AddSP(this, &SPCGEditorGraphProfilingView::OnDebugStackChanged);
 	}
 
@@ -690,7 +688,7 @@ FReply SPCGEditorGraphProfilingView::Refresh()
 		return FReply::Handled();
 	}
 
-	const UPCGComponent* Component = PCGComponent.Get();
+	const UPCGComponent* Component = GetPCGComponent().Get();
 	if (!Component)
 	{
 		return FReply::Handled();
@@ -785,25 +783,12 @@ FReply SPCGEditorGraphProfilingView::Refresh()
 void SPCGEditorGraphProfilingView::OnDebugStackChanged(const FPCGStack& InPCGStack)
 {
 	PCGStack = InPCGStack;
-
-	if (PCGComponent.IsValid())
-	{
-		PCGComponent->OnPCGGraphGeneratedDelegate.RemoveAll(this);
-	}
-
-	PCGComponent = const_cast<UPCGComponent*>(InPCGStack.GetRootComponent());
-
-	if (PCGComponent.IsValid())
-	{
-		PCGComponent->OnPCGGraphGeneratedDelegate.AddSP(this, &SPCGEditorGraphProfilingView::OnGenerateUpdated);
-	}
-
-	Refresh();
+	RequestRefresh();
 }
 
-void SPCGEditorGraphProfilingView::OnGenerateUpdated(UPCGComponent* InPCGComponent)
+TWeakObjectPtr<UPCGComponent> SPCGEditorGraphProfilingView::GetPCGComponent() const
 {
-	Refresh();
+	return PCGEditorPtr.IsValid() ? PCGEditorPtr.Pin()->GetPCGComponentBeingInspected() : nullptr;
 }
 
 TSharedRef<ITableRow> SPCGEditorGraphProfilingView::OnGenerateRow(PCGProfilingListViewItemPtr Item, const TSharedRef<STableViewBase>& OwnerTable) const
diff --git a/Engine/Plugins/PCG/Source/PCGEditor/Private/SPCGEditorGraphProfilingView.h b/Engine/Plugins/PCG/Source/PCGEditor/Private/SPCGEditorGraphProfilingView.h
--- a/Engine/Plugins/PCG/Source/PCGEditor/Private/SPCGEditorGraphProfilingView.h
+++ b/Engine/Plugins/PCG/Source/PCGEditor/Private/SPCGEditorGraphProfilingView.h
@@ -66,8 +66,6 @@ private:
 
 	void OnDebugStackChanged(const FPCGStack& InPCGStack);
 
-	void OnGenerateUpdated(UPCGComponent* InPCGComponent);
-	
 	// Callbacks
 	void RequestRefresh();
 	FReply Refresh();
@@ -89,15 +87,15 @@ private:
 	/** Called when user changes commits text to the search box */
 	void OnSearchTextCommitted(const FText& InText, ETextCommit::Type InCommitType);
 
+	/** Returns the current PCG component */
+	TWeakObjectPtr<UPCGComponent> GetPCGComponent() const;
+
 	/** Pointer back to the PCG editor that owns us */
 	TWeakPtr<FPCGEditor> PCGEditorPtr;
 
 	/** Cached PCGGraph being viewed */
 	UPCGEditorGraph* PCGEditorGraph = nullptr;
 
-	/** Cached PCGComponent being viewed */
-	TWeakObjectPtr<UPCGComponent> PCGComponent;
-
 	/** Current stack being viewed */
 	FPCGStack PCGStack;
 
