Adjusted text in cvar and added some comments.

#rnx

diff --git a/Engine/Plugins/PCG/Source/PCG/Private/PCGCommon.cpp b/Engine/Plugins/PCG/Source/PCG/Private/PCGCommon.cpp
--- a/Engine/Plugins/PCG/Source/PCG/Private/PCGCommon.cpp
+++ b/Engine/Plugins/PCG/Source/PCG/Private/PCGCommon.cpp
@@ -42,7 +42,7 @@ namespace PCGSystemSwitches
 	TAutoConsoleVariable<bool> CVarDirtyLoadAsPreviewOnLoad(
 		TEXT("pcg.DirtyLoadAsPreviewOnLoad"),
 		true,
-		TEXT("Enables or not dirtying on load for load as preview components.\nTurning off this option will require to force generate or apply a change before this component is regenerated."));
+		TEXT("Enables dirtying on load for load as preview components.\nTurning off this option will require to force generate or apply a change before this component is regenerated."));
 #endif
 }
 
diff --git a/Engine/Plugins/PCG/Source/PCG/Public/Graph/PCGStackContext.h b/Engine/Plugins/PCG/Source/PCG/Public/Graph/PCGStackContext.h
--- a/Engine/Plugins/PCG/Source/PCG/Public/Graph/PCGStackContext.h
+++ b/Engine/Plugins/PCG/Source/PCG/Public/Graph/PCGStackContext.h
@@ -56,7 +56,10 @@ struct PCG_API FPCGStackFrame
 		Hash = GetTypeHash(LoopIndex);
 	}
 
+	// Stores object this frame refers to. Use SetObject to change this and properly update the hash.
 	TWeakObjectPtr<const UObject> Object;
+
+	// Stores the loop index (if any) that this frame refers to. Use SetLoopIndex to change this and properly update the hash.
 	int32 LoopIndex = INDEX_NONE;
 
 private:
