Minor comment adjustment regarding pin status

#rnx

diff --git a/Engine/Plugins/PCG/Source/PCG/Public/PCGPin.h b/Engine/Plugins/PCG/Source/PCG/Public/PCGPin.h
--- a/Engine/Plugins/PCG/Source/PCG/Public/PCGPin.h
+++ b/Engine/Plugins/PCG/Source/PCG/Public/PCGPin.h
@@ -32,7 +32,7 @@ enum class EPCGPinStatus : uint8
 	Required,
 	/** Advanced pin will be hidden by default in the UIand will be shown only if the user extends the node(in the UI) to see advanced pins.Pins can't be required and advanced at the same time. */
 	Advanced,
-	/** Pins generated by overridable parameters (normal nodes) or user parameters (subgraph nodes) */
+	/** Pins generated by overridable parameters (normal nodes) or user parameters (subgraph nodes) - note that these are considered advanced pins. */
 	OverrideOrUserParam UMETA(Hidden)
 };
 
