[PCG] Added WITH_EDITOR guard to PCGAsyncTest.cpp

#rnx

diff --git a/Engine/Plugins/PCG/Source/PCG/Private/Tests/Helpers/PCGAsyncTest.cpp b/Engine/Plugins/PCG/Source/PCG/Private/Tests/Helpers/PCGAsyncTest.cpp
--- a/Engine/Plugins/PCG/Source/PCG/Private/Tests/Helpers/PCGAsyncTest.cpp
+++ b/Engine/Plugins/PCG/Source/PCG/Private/Tests/Helpers/PCGAsyncTest.cpp
@@ -1,10 +1,11 @@
 // Copyright Epic Games, Inc. All Rights Reserved.
 
+#if WITH_EDITOR
+
 #include "Tests/PCGTestsCommon.h"
 
 #include "Helpers/PCGAsync.h"
 
-
 IMPLEMENT_CUSTOM_SIMPLE_AUTOMATION_TEST(FPCGAsyncTest_AsyncProcessing, FPCGTestBaseClass, "Plugins.PCG.Async.AsyncProcessing", PCGTestsCommon::TestFlags)
 
 bool FPCGAsyncTest_AsyncProcessing::RunTest(const FString& Parameters)
@@ -125,3 +126,5 @@ bool FPCGAsyncTest_AsyncProcessingOneToOneEx::RunTest(const FString& Parameters)
 
 	return true;
 }
+
+#endif // WITH_EDITOR
