Added pcg pause button tooltip override

#rnx

diff --git a/Engine/Plugins/PCG/Source/PCGEditor/Private/PCGEditorModule.cpp b/Engine/Plugins/PCG/Source/PCGEditor/Private/PCGEditorModule.cpp
--- a/Engine/Plugins/PCG/Source/PCGEditor/Private/PCGEditorModule.cpp
+++ b/Engine/Plugins/PCG/Source/PCGEditor/Private/PCGEditorModule.cpp
@@ -412,7 +412,18 @@ void FPCGEditorModule::RegisterMenuExtensions()
 							}
 						}
 					}),
-					LOCTEXT("PCGPauseButton_Tooltip", "Toggles PCG processing on/off and will cancel tasks depending on settings.\nUse Ctrl to unpause and cancel all tasks.\nUse Alt to unpause without cancelling tasks."),
+					TAttribute<FText>::CreateLambda([]()
+					{
+						const UPCGEditorSettings* EditorSettings = GetDefault<UPCGEditorSettings>();
+						if (EditorSettings && !EditorSettings->OverridePausedButtonTooltip.IsEmpty())
+						{
+							return FText::FromString(EditorSettings->OverridePausedButtonTooltip);
+						}
+						else
+						{
+							return LOCTEXT("PCGPauseButton_Tooltip", "Toggles PCG processing on/off and will cancel tasks depending on settings.\nUse Ctrl to unpause and cancel all tasks.\nUse Alt to unpause without cancelling tasks.");
+						}
+					}),
 					TAttribute<FSlateIcon>::CreateLambda([]()
 					{
 						const UPCGEditorSettings* EditorSettings = GetDefault<UPCGEditorSettings>();
diff --git a/Engine/Plugins/PCG/Source/PCGEditor/Public/PCGEditorSettings.h b/Engine/Plugins/PCG/Source/PCGEditor/Public/PCGEditorSettings.h
--- a/Engine/Plugins/PCG/Source/PCGEditor/Public/PCGEditorSettings.h
+++ b/Engine/Plugins/PCG/Source/PCGEditor/Public/PCGEditorSettings.h
@@ -211,6 +211,10 @@ public:
 	/** Overrides the label of the pause button while PCG is currently paused. The default is empty or None, which will use the default label then. */
 	UPROPERTY(EditAnywhere, config, Category = "Editor Performance", meta = (DisplayName = "Overrides default name for the 'not currently paused' button", EditCondition = "bShowPauseButton", EditConditionHides))
 	FName OverrideNotPausedButtonLabel = NAME_None;
+
+	/** Overrides the tooltip on the pause button. The default is empty, which will use the default tooltip then. */
+	UPROPERTY(EditAnywhere, config, Category = "Editor Performance", meta = (MultiLine = true, DisplayName = "Overrides default tooltip on the pause button", EditCondition = "bShowPauseButton", EditConditionHides))
+	FString OverridePausedButtonTooltip = FString();
 };
 
 #if UE_ENABLE_INCLUDE_ORDER_DEPRECATED_IN_5_2
