[PCG] Change the node being inspected in the PCG Editor to a weak ptr

#rnx

diff --git a/Engine/Plugins/PCG/Source/PCGEditor/Private/SPCGEditorGraphDebugObjectTree.cpp b/Engine/Plugins/PCG/Source/PCGEditor/Private/SPCGEditorGraphDebugObjectTree.cpp
--- a/Engine/Plugins/PCG/Source/PCGEditor/Private/SPCGEditorGraphDebugObjectTree.cpp
+++ b/Engine/Plugins/PCG/Source/PCGEditor/Private/SPCGEditorGraphDebugObjectTree.cpp
@@ -428,7 +428,7 @@ void SPCGEditorGraphDebugObjectTree::OnObjectsReplaced(const TMap<UObject*, UObj
 
 void SPCGEditorGraphDebugObjectTree::SetNodeBeingInspected(const UPCGNode* InPCGNode)
 {
-	PCGNodeBeingInspected = InPCGNode;
+	PCGNodeBeingInspected = IsValid(InPCGNode) ? InPCGNode : nullptr;
 
 	RequestRefresh();
 }
@@ -625,7 +625,7 @@ void SPCGEditorGraphDebugObjectTree::AddStacksToTree(const TArray<FPCGStack>& St
 		}
 
 		// If we're inspecting a node which has not logged inspection data in a previous execution, display grayed out.
-		const bool bDisplayGrayedOut = PCGNodeBeingInspected && !PCGComponent->HasNodeProducedData(PCGNodeBeingInspected, Stack);
+		const bool bDisplayGrayedOut = PCGNodeBeingInspected.IsValid() && !PCGComponent->HasNodeProducedData(PCGNodeBeingInspected.Get(), Stack);
 
 		AActor* Actor = PCGComponent->GetOwner();
 		if (!Actor)
diff --git a/Engine/Plugins/PCG/Source/PCGEditor/Private/SPCGEditorGraphDebugObjectTree.h b/Engine/Plugins/PCG/Source/PCGEditor/Private/SPCGEditorGraphDebugObjectTree.h
--- a/Engine/Plugins/PCG/Source/PCGEditor/Private/SPCGEditorGraphDebugObjectTree.h
+++ b/Engine/Plugins/PCG/Source/PCGEditor/Private/SPCGEditorGraphDebugObjectTree.h
@@ -305,5 +305,5 @@ private:
 	/** The previous stack that the user selected. */
 	FPCGStack PreviouslySelectedStack;
 
-	const UPCGNode* PCGNodeBeingInspected = nullptr;
+	TWeakObjectPtr<const UPCGNode> PCGNodeBeingInspected = nullptr;
 };
