[PCG] Fix check failure when uploading invalid attribute types to the GPU

#rnx

diff --git a/Engine/Plugins/PCG/Source/PCG/Private/Compute/PCGDataForGPU.cpp b/Engine/Plugins/PCG/Source/PCG/Private/Compute/PCGDataForGPU.cpp
--- a/Engine/Plugins/PCG/Source/PCG/Private/Compute/PCGDataForGPU.cpp
+++ b/Engine/Plugins/PCG/Source/PCG/Private/Compute/PCGDataForGPU.cpp
@@ -57,8 +57,7 @@ namespace PCGDataForGPUHelpers
 		case EPCGMetadataTypes::Transform:
 			return EPCGKernelAttributeType::Transform;
 		default:
-			checkNoEntry();
-			return EPCGKernelAttributeType::Float;
+			return EPCGKernelAttributeType::Invalid;
 		}
 	}
 
@@ -465,6 +464,19 @@ void FPCGDataDesc::InitializeAttributeDescs(const UPCGMetadata* Metadata, const
 		{
 			const FName AttributeName = AttributeNames[CustomAttributeIndex];
 			const EPCGKernelAttributeType AttributeType = PCGDataForGPUHelpers::GetAttributeTypeFromMetadataType(AttributeTypes[CustomAttributeIndex]);
+
+			if (AttributeType == EPCGKernelAttributeType::Invalid)
+			{
+				const UEnum* EnumClass = StaticEnum<EPCGMetadataTypes>();
+				check(EnumClass);
+
+				UE_LOG(LogPCG, Warning, TEXT("Skipping attribute '%s'. '%s' type attributes are not supported on GPU."),
+					*AttributeName.ToString(),
+					*EnumClass->GetNameStringByValue(static_cast<int64>(AttributeTypes[CustomAttributeIndex])));
+
+				continue;
+			}
+
 			const FPCGKernelAttributeKey AttributeKey = { AttributeType, AttributeName };
 
 			// Ignore excess attributes.
diff --git a/Engine/Plugins/PCG/Source/PCG/Public/Compute/PCGDataForGPU.h b/Engine/Plugins/PCG/Source/PCG/Public/Compute/PCGDataForGPU.h
--- a/Engine/Plugins/PCG/Source/PCG/Public/Compute/PCGDataForGPU.h
+++ b/Engine/Plugins/PCG/Source/PCG/Public/Compute/PCGDataForGPU.h
@@ -20,7 +20,7 @@ enum class EPCGUnpackDataCollectionResult
 UENUM()
 enum class EPCGKernelAttributeType : uint8
 {
-	Bool,
+	Bool = 0,
 	Int,
 	Float,
 	Float2,
@@ -28,7 +28,9 @@ enum class EPCGKernelAttributeType : uint8
 	Float4,
 	Rotator,
 	Quat,
-	Transform
+	Transform,
+
+	Invalid = std::numeric_limits<uint8>::max()
 };
 
 USTRUCT()
