[PCG] Allow inverted ranges in Attribute Remap (e.g. 1.0 -> 0.0) (default behavior for new nodes)

#rnx

diff --git a/Engine/Plugins/PCG/Source/PCG/Private/Elements/PCGAttributeRemap.cpp b/Engine/Plugins/PCG/Source/PCG/Private/Elements/PCGAttributeRemap.cpp
--- a/Engine/Plugins/PCG/Source/PCG/Private/Elements/PCGAttributeRemap.cpp
+++ b/Engine/Plugins/PCG/Source/PCG/Private/Elements/PCGAttributeRemap.cpp
@@ -4,6 +4,7 @@
 
 #include "PCGContext.h"
 #include "Helpers/PCGAsync.h"
+#include "Helpers/PCGHelpers.h"
 #include "Metadata/Accessors/PCGAttributeAccessorHelpers.h"
 
 #define LOCTEXT_NAMESPACE "PCGAttributeRemapElement"
@@ -12,7 +13,9 @@ namespace PCGAttributeRemapElement
 {
 	struct FPCGAttributeRemapParams
 	{
-		double InRangeMin;
+		double InRangeStart; // Note: Start and End can be a positive or negative range (e.g. [0, 1] or [1, 0])
+		double InRangeEnd;
+		double InRangeMin; // Note: Min and Max are a true min/max used to reject samples that fall outside the range
 		double InRangeMax;
 		double Slope;
 		double Intercept;
@@ -30,7 +33,7 @@ namespace PCGAttributeRemapElement
 	{
 		if (!InParams.bIgnoreValuesOutsideInputRange || (InOutValue >= InParams.InRangeMin && InOutValue <= InParams.InRangeMax))
 		{
-			InOutValue = InParams.Slope * (InOutValue - InParams.InRangeMin) + InParams.Intercept;
+			InOutValue = InParams.Slope * (InOutValue - InParams.InRangeStart) + InParams.Intercept;
 			if (InParams.bClampToUnitRange)
 			{
 				InOutValue = FMath::Clamp<T>(InOutValue, T(0.0), T(1.0));
@@ -71,6 +74,14 @@ namespace PCGAttributeRemapElement
 	}
 }
 
+UPCGAttributeRemapSettings::UPCGAttributeRemapSettings()
+{
+	if (PCGHelpers::IsNewObjectAndNotDefault(this))
+	{
+		bAllowInverseRange = true;
+	}
+}
+
 #if WITH_EDITOR
 FName UPCGAttributeRemapSettings::GetDefaultNodeName() const
 {
@@ -135,22 +146,37 @@ bool FPCGAttributeRemapElement::DoOperation(PCGMetadataOps::FOperationData& Oper
 	Params.InRangeMin = FMath::Min(Settings->InRangeMin, Settings->InRangeMax);
 	Params.InRangeMax = FMath::Max(Settings->InRangeMin, Settings->InRangeMax);
 
-	const double OutRangeTrueMin = FMath::Min(Settings->OutRangeMin, Settings->OutRangeMax);
-	const double OutRangeTrueMax = FMath::Max(Settings->OutRangeMin, Settings->OutRangeMax);
+	double OutRangeMin = 0.0;
+	double OutRangeMax = 0.0;
+
+	if (Settings->bAllowInverseRange)
+	{
+		Params.InRangeStart = Settings->InRangeMin;
+		Params.InRangeEnd = Settings->InRangeMax;
+		OutRangeMin = Settings->OutRangeMin;
+		OutRangeMax = Settings->OutRangeMax;
+	}
+	else
+	{
+		Params.InRangeStart = Params.InRangeMin;
+		Params.InRangeEnd = Params.InRangeMax;
+		OutRangeMin = FMath::Min(Settings->OutRangeMin, Settings->OutRangeMax);
+		OutRangeMax = FMath::Max(Settings->OutRangeMin, Settings->OutRangeMax);
+	}
 
-	const double InRangeDifference = Params.InRangeMax - Params.InRangeMin;
-	const double OutRangeDifference = OutRangeTrueMax - OutRangeTrueMin;
+	const double InRangeDifference = Params.InRangeEnd - Params.InRangeStart;
+	const double OutRangeDifference = OutRangeMax - OutRangeMin;
 
 	// When InRange is a point leave the Slope at 0 so that Density = Intercept
 	if (InRangeDifference == 0)
 	{
 		Params.Slope = 0;
-		Params.Intercept = (OutRangeTrueMin + OutRangeTrueMax) / 2.f;
+		Params.Intercept = (OutRangeMin + OutRangeMax) / 2.f;
 	}
 	else
 	{
 		Params.Slope = OutRangeDifference / InRangeDifference;
-		Params.Intercept = OutRangeTrueMin;
+		Params.Intercept = OutRangeMin;
 	}
 
 	auto RemapFunc = [this, &Params, &OperationData]<typename AttributeType>(AttributeType) -> bool
diff --git a/Engine/Plugins/PCG/Source/PCG/Public/Elements/PCGAttributeRemap.h b/Engine/Plugins/PCG/Source/PCG/Public/Elements/PCGAttributeRemap.h
--- a/Engine/Plugins/PCG/Source/PCG/Public/Elements/PCGAttributeRemap.h
+++ b/Engine/Plugins/PCG/Source/PCG/Public/Elements/PCGAttributeRemap.h
@@ -21,6 +21,8 @@ class UPCGAttributeRemapSettings : public UPCGMetadataSettingsBase
 	GENERATED_BODY()
 
 public:
+	UPCGAttributeRemapSettings();
+
 	//~Begin UPCGSettings interface
 #if WITH_EDITOR
 	virtual FName GetDefaultNodeName() const override;
@@ -68,6 +70,10 @@ public:
 	/** Attribute values outside of the input range will be unaffected by the remapping */
 	UPROPERTY(BlueprintReadWrite, EditAnywhere, Category = Settings, meta = (PCG_Overridable))
 	bool bIgnoreValuesOutsideInputRange = false;
+
+	/** Allow remapping when Min is larger than Max, e.g. from [0.0, 1.0] -> [1.0, 0.0]. */
+	UPROPERTY(BlueprintReadWrite, EditAnywhere, Category = Settings, meta = (PCG_Overridable))
+	bool bAllowInverseRange = false; // Note that this is no longer the default value for new nodes, it is now 'true'
 };
 
 class FPCGAttributeRemapElement : public FPCGMetadataElementBase
