[PCG] Added/Exposed some PCG Helper functions

+ Random Array Indices - new helper
* Exposed ShuffleArray/ShiftArrayElements from Determinism

#rnx

diff --git a/Engine/Plugins/PCG/Source/PCG/Private/Helpers/PCGHelpers.cpp b/Engine/Plugins/PCG/Source/PCG/Private/Helpers/PCGHelpers.cpp
--- a/Engine/Plugins/PCG/Source/PCG/Private/Helpers/PCGHelpers.cpp
+++ b/Engine/Plugins/PCG/Source/PCG/Private/Helpers/PCGHelpers.cpp
@@ -644,6 +644,34 @@ namespace PCGHelpers
 		default: checkNoEntry(); return [](float, float) { return 0.0f; };
 		}
 	}
+
+	TArray<int32> GetRandomIndices(FRandomStream& RandomStream, const int32 ArraySize, const int32 NumSelections)
+	{
+		if (ArraySize < 1 || NumSelections < 1)
+		{
+			return {};
+		}
+
+		const int32 N = FMath::Min(NumSelections, ArraySize);
+
+		TArray<int32> RandomIndices;
+		RandomIndices.Reserve(N);
+
+		const int32 Max = ArraySize - NumSelections;
+		for (int i = 0; i < N; ++i)
+		{
+			RandomIndices.Emplace(RandomStream.RandRange(0, Max));
+		}
+
+		RandomIndices.Sort();
+
+		for (int i = 0; i < RandomIndices.Num(); ++i)
+		{
+			RandomIndices[i] += i;
+		}
+
+		return RandomIndices;
+	}
 }
 
 #undef LOCTEXT_NAMESPACE
diff --git a/Engine/Plugins/PCG/Source/PCG/Private/Tests/Determinism/PCGDeterminismTestsCommon.cpp b/Engine/Plugins/PCG/Source/PCG/Private/Tests/Determinism/PCGDeterminismTestsCommon.cpp
--- a/Engine/Plugins/PCG/Source/PCG/Private/Tests/Determinism/PCGDeterminismTestsCommon.cpp
+++ b/Engine/Plugins/PCG/Source/PCG/Private/Tests/Determinism/PCGDeterminismTestsCommon.cpp
@@ -1340,12 +1340,12 @@ namespace PCGDeterminismTests
 
 	void ShuffleInputOrder(PCGTestsCommon::FTestData& TestData)
 	{
-		ShuffleArray<FPCGTaggedData>(TestData.InputData.TaggedData, TestData.RandomStream);
+		PCGHelpers::ShuffleArray<FPCGTaggedData>(TestData.RandomStream, TestData.InputData.TaggedData);
 	}
 
 	void ShuffleOutputOrder(PCGTestsCommon::FTestData& TestData)
 	{
-		ShuffleArray<FPCGTaggedData>(TestData.OutputData.TaggedData, TestData.RandomStream);
+		PCGHelpers::ShuffleArray<FPCGTaggedData>(TestData.RandomStream, TestData.OutputData.TaggedData);
 	}
 
 	void ShuffleAllInternalData(PCGTestsCommon::FTestData& TestData)
@@ -1355,14 +1355,14 @@ namespace PCGDeterminismTests
 			if (DataCanBeShuffled(TaggedData.Data))
 			{
 				UPCGPointData* PointData = const_cast<UPCGPointData*>(CastChecked<UPCGPointData>(TaggedData.Data));
-				ShuffleArray<FPCGPoint>(PointData->GetMutablePoints(), TestData.RandomStream);
+				PCGHelpers::ShuffleArray<FPCGPoint>(TestData.RandomStream, PointData->GetMutablePoints());
 			}
 		}
 	}
 
 	void ShiftInputOrder(PCGTestsCommon::FTestData& TestData, int32 NumShifts)
 	{
-		ShiftArrayElements<FPCGTaggedData>(TestData.InputData.TaggedData, NumShifts);
+		PCGHelpers::ShiftArrayElements<FPCGTaggedData>(TestData.InputData.TaggedData, NumShifts);
 	}
 
 	TArray<EPCGDataType> FilterTestableDataTypes(EPCGDataType AllowedTypes, int32 NumMultipleInputs)
@@ -1589,4 +1589,3 @@ namespace PCGDeterminismTests
 }
 
 #undef LOCTEXT_NAMESPACE
-
diff --git a/Engine/Plugins/PCG/Source/PCG/Private/Tests/Helpers/PCGHelpersTest.cpp b/Engine/Plugins/PCG/Source/PCG/Private/Tests/Helpers/PCGHelpersTest.cpp
new file mode 100644
--- /dev/null
+++ b/Engine/Plugins/PCG/Source/PCG/Private/Tests/Helpers/PCGHelpersTest.cpp
@@ -0,0 +1,71 @@
+// Copyright Epic Games, Inc. All Rights Reserved.
+
+#if WITH_EDITOR
+
+#include "Tests/PCGTestsCommon.h"
+#include "PCG/Public/Helpers/PCGHelpers.h"
+
+#include "Algo/MaxElement.h"
+
+IMPLEMENT_CUSTOM_SIMPLE_AUTOMATION_TEST(FPCGHelpersTest_GetRandomIndices, FPCGTestBaseClass, "Plugins.PCG.Helpers.GetRandomIndices", PCGTestsCommon::TestFlags)
+
+bool FPCGHelpersTest_GetRandomIndices::RunTest(const FString& Parameters)
+{
+	FRandomStream RandomStream(42);
+
+	static int32 constexpr Num = 100;
+
+	TArray<int32> RandomIndices = PCGHelpers::GetRandomIndices(RandomStream, Num, 0);
+	UTEST_TRUE("Correct array size", RandomIndices.IsEmpty());
+
+	RandomIndices = PCGHelpers::GetRandomIndices(RandomStream, Num, -1);
+	UTEST_TRUE("Correct array size", RandomIndices.IsEmpty());
+
+	RandomIndices = PCGHelpers::GetRandomIndices(RandomStream, Num, 1);
+	UTEST_EQUAL("Correct array size", RandomIndices.Num(), 1);
+
+	RandomIndices = PCGHelpers::GetRandomIndices(RandomStream, Num, Num + 1);
+	UTEST_EQUAL("Correct array size", RandomIndices.Num(), Num);
+
+	RandomIndices = PCGHelpers::GetRandomIndices(RandomStream, Num, Num / 2);
+	UTEST_EQUAL("Correct array size", RandomIndices.Num(), Num / 2);
+
+	// Uniqueness
+	static int32 constexpr BigNum = 1000000;
+
+	RandomIndices = PCGHelpers::GetRandomIndices(RandomStream, BigNum, BigNum);
+	UTEST_EQUAL("Correct array size", RandomIndices.Num(), BigNum);
+
+	TArray<int32> Counts;
+	Counts.SetNumZeroed(RandomIndices.Num());
+	for (int i = 0; i < RandomIndices.Num(); ++i)
+	{
+		UTEST_TRUE("All values are unique", ++Counts[RandomIndices[i]] < 2);
+	}
+
+	// Distribution
+	static int32 constexpr BinCount = 25;
+	static int32 constexpr IndicesPerBin = BigNum / BinCount;
+	static double constexpr SelectionRatio = 0.2; // Select 1/5 of the array.
+	RandomIndices = PCGHelpers::GetRandomIndices(RandomStream, BigNum, BigNum * SelectionRatio);
+
+	TArray<int32> BinCounts;
+	BinCounts.SetNumZeroed(BinCount);
+	for (int i = 0; i < RandomIndices.Num(); ++i)
+	{
+		const int32 Index = RandomIndices[i];
+		const int32 Bin = Index / IndicesPerBin;
+		++BinCounts[Bin];
+	}
+
+	static constexpr int32 ExpectedCount = BigNum * SelectionRatio / BinCount;
+	static constexpr int32 Tolerance = ExpectedCount * 0.05;
+	for (int i = 0; i < BinCounts.Num(); ++i)
+	{
+		UTEST_LESS("Bins are within tolerance", FMath::Abs(BinCounts[i] - ExpectedCount), Tolerance);
+	}
+
+	return true;
+}
+
+#endif // WITH_EDITOR
diff --git a/Engine/Plugins/PCG/Source/PCG/Public/Helpers/PCGHelpers.h b/Engine/Plugins/PCG/Source/PCG/Public/Helpers/PCGHelpers.h
--- a/Engine/Plugins/PCG/Source/PCG/Public/Helpers/PCGHelpers.h
+++ b/Engine/Plugins/PCG/Source/PCG/Public/Helpers/PCGHelpers.h
@@ -91,7 +91,65 @@ namespace PCGHelpers
 	PCG_API TArray<UFunction*> FindUserFunctions(TSubclassOf<UObject> ActorClass, const TArray<FName>& FunctionNames, const TArray<const UFunction*>& FunctionPrototypes, const FPCGContext* InContext = nullptr);
 
 	PCG_API TFunction<float(float, float)> GetDensityMergeFunction(EPCGDensityMergeOperation InOperation);
-};
+
+	/** Get an array of a randomized, uniformly distributed indices to a provided array view. */
+	PCG_API TArray<int32> GetRandomIndices(FRandomStream& RandomStream, const int32 ArraySize, const int32 NumSelections);
+
+	/** Shuffles the elements of an array randomly and uniformly. */
+	template <typename T>
+	void ShuffleArray(FRandomStream& RandomStream, TArray<T>& Array)
+	{
+		const int32 LastIndex = Array.Num() - 1;
+		for (int32 i = 0; i < LastIndex; ++i)
+		{
+			const int32 Index = RandomStream.RandRange(i, LastIndex);
+
+			if (i != Index)
+			{
+				Array.Swap(i, Index);
+			}
+		}
+	}
+
+	/** Shifts the elements of an array a number of times. */
+	template <typename T>
+	void ShiftArrayElements(TArrayView<T> Array, int32 NumShifts = 1)
+	{
+		if (Array.Num() < 2 || NumShifts == 0)
+		{
+			return;
+		}
+
+		const int32 Count = Array.Num();
+		NumShifts %= Count;
+		if (NumShifts < 0)
+		{
+			NumShifts += Count;
+		}
+
+		TArray<T> TempArray;
+		if constexpr (std::is_trivially_copyable_v<T>)
+		{
+			TempArray.SetNumUninitialized(Count);
+		}
+		else
+		{
+			TempArray.SetNum(Count);
+		}
+
+		for (int32 i = 0; i < NumShifts; ++i)
+		{
+			TempArray[i] = std::move(Array[i + Count - NumShifts]);
+		}
+
+		for (int32 i = NumShifts; i < Count; ++i)
+		{
+			TempArray[i] = std::move(Array[i - NumShifts]);
+		}
+
+		Array = std::move(TempArray);
+	}
+}
 
 /** Holds function prototypes used to match against actor function signatures. */
 UCLASS(MinimalAPI)
diff --git a/Engine/Plugins/PCG/Source/PCG/Public/Tests/Determinism/PCGDeterminismTestsCommon.h b/Engine/Plugins/PCG/Source/PCG/Public/Tests/Determinism/PCGDeterminismTestsCommon.h
--- a/Engine/Plugins/PCG/Source/PCG/Public/Tests/Determinism/PCGDeterminismTestsCommon.h
+++ b/Engine/Plugins/PCG/Source/PCG/Public/Tests/Determinism/PCGDeterminismTestsCommon.h
@@ -305,52 +305,6 @@ namespace PCGDeterminismTests
 		return Cast<const DataType>(FirstData) != nullptr && Cast<const DataType>(SecondData) != nullptr;
 	}
 
-	template<typename DataType>
-	void ShuffleArray(TArray<DataType>& Array, FRandomStream& RandomStream)
-	{
-		const int32 LastIndex = Array.Num() - 1;
-		for (int32 I = 0; I <= LastIndex; ++I)
-		{
-			int32 Index = RandomStream.RandRange(I, LastIndex);
-
-			if (I != Index)
-			{
-				Array.Swap(I, Index);
-			}
-		}
-	}
-
-	template<typename DataType>
-	void ShiftArrayElements(TArray<DataType>& Array, int32 NumShifts = 1)
-	{
-		if (Array.Num() < 2)
-		{
-			return;
-		}
-
-		int32 Count = Array.Num();
-		NumShifts %= Count;
-		if (NumShifts < 0)
-		{
-			NumShifts += Count;
-		}
-
-		TArray<DataType> TempArray;
-		TempArray.SetNum(Count);
-
-		for (int32 I = 0; I < NumShifts; ++I)
-		{
-			TempArray[I] = Array[I + Count - NumShifts];
-		}
-
-		for (int32 I = NumShifts; I < Count; ++I)
-		{
-			TempArray[I] = Array[I - NumShifts];
-		}
-
-		Array = MoveTemp(TempArray);
-	}
-
 	template<typename MetadataAttributeType>
 	bool MetadataAttributesAreEqual(const FPCGMetadataAttributeBase* FirstAttributeBase, const FPCGMetadataAttributeBase* SecondAttributeBase, PCGMetadataValueKey ValueKey)
 	{
