[PCG] Fix a bug in PCGSelectGrammar with a misaligned enum selection

#rnx

diff --git a/Engine/Plugins/PCG/Source/PCG/Private/Elements/Grammar/PCGSelectGrammar.cpp b/Engine/Plugins/PCG/Source/PCG/Private/Elements/Grammar/PCGSelectGrammar.cpp
--- a/Engine/Plugins/PCG/Source/PCG/Private/Elements/Grammar/PCGSelectGrammar.cpp
+++ b/Engine/Plugins/PCG/Source/PCG/Private/Elements/Grammar/PCGSelectGrammar.cpp
@@ -32,8 +32,8 @@ namespace PCGSelectGrammar
 		double Size;
 	};
 
-	constexpr uint16 BinaryEnumIndex = static_cast<uint16>(EPCGSelectGrammarComparator::BinaryOps);
-	constexpr uint16 TernaryEnumIndex = static_cast<uint16>(EPCGSelectGrammarComparator::TernaryOps);
+	constexpr uint16 BinaryEnumIndex = static_cast<uint16>(EPCGSelectGrammarComparator::BinaryOps) + 1;
+	constexpr uint16 TernaryEnumIndex = static_cast<uint16>(EPCGSelectGrammarComparator::TernaryOps) + 1;
 
 	// Function indices must match EPCGSelectGrammarComparator indices directly
 	template <typename T>
