[PCG] Fix Custom HLSL nodes not populating declarations on creation

#rnx

diff --git a/Engine/Plugins/PCG/Source/PCG/Private/Compute/Elements/PCGCustomHLSL.cpp b/Engine/Plugins/PCG/Source/PCG/Private/Compute/Elements/PCGCustomHLSL.cpp
--- a/Engine/Plugins/PCG/Source/PCG/Private/Compute/Elements/PCGCustomHLSL.cpp
+++ b/Engine/Plugins/PCG/Source/PCG/Private/Compute/Elements/PCGCustomHLSL.cpp
@@ -92,10 +92,20 @@ void UPCGCustomHLSLSettings::PostLoad()
 {
 	Super::PostLoad();
 
+	// Note: We update here so that Custom HLSL nodes will have the correct pin settings & declarations on load.
 	UpdatePinSettings();
 	UpdateAttributeKeys();
 	UpdateDeclarations();
 }
+
+void UPCGCustomHLSLSettings::PostInitProperties()
+{
+	Super::PostInitProperties();
+
+	// Note: We update here so that Custom HLSL nodes will have the correct pin settings & declarations on creation.
+	UpdatePinSettings();
+	UpdateDeclarations();
+}
 #endif
 
 TArray<FPCGPinProperties> UPCGCustomHLSLSettings::OutputPinProperties() const
diff --git a/Engine/Plugins/PCG/Source/PCG/Public/Compute/Elements/PCGCustomHLSL.h b/Engine/Plugins/PCG/Source/PCG/Public/Compute/Elements/PCGCustomHLSL.h
--- a/Engine/Plugins/PCG/Source/PCG/Public/Compute/Elements/PCGCustomHLSL.h
+++ b/Engine/Plugins/PCG/Source/PCG/Public/Compute/Elements/PCGCustomHLSL.h
@@ -102,6 +102,7 @@ public:
 	//~Begin UObject interface
 #if WITH_EDITOR
 	virtual void PostLoad() override;
+	virtual void PostInitProperties() override;
 #endif
 	//~End UObject interface
 
