[PCG] Write GPU data collection NumData even when element count is 0

#rnx

diff --git a/Engine/Plugins/PCG/Source/PCG/Private/Compute/PCGDataForGPU.cpp b/Engine/Plugins/PCG/Source/PCG/Private/Compute/PCGDataForGPU.cpp
--- a/Engine/Plugins/PCG/Source/PCG/Private/Compute/PCGDataForGPU.cpp
+++ b/Engine/Plugins/PCG/Source/PCG/Private/Compute/PCGDataForGPU.cpp
@@ -779,8 +779,23 @@ void FPCGDataCollectionDesc::PrepareBufferForKernelOutput(TArray<uint32>& OutPac
 
 	OutPackedDataCollection.SetNumZeroed(PackedDataCollectionSizeBytes / sizeof(uint32));
 	
-	// Num data - set to zero if writing kernel executes. If kernel doesn't execute, 0 means data collection is empty.
-	OutPackedDataCollection[0] = 0;
+	const uint32 TotalNumElements = ComputeDataElementCount(EPCGDataType::Any);
+
+	if (TotalNumElements == 0)
+	{
+		// The kernel won't run with 0 element count, so we have to set the number of data in advance.
+		OutPackedDataCollection[0] = NumData;
+
+		// TODO: This assumption will break dynamic control flow in the future. If we are using 'NumData == 0' to flag whether a kernel ran or not,
+		// we can't just write NumData anyways whenever there are zero elements.
+		// In the future we should dispatch a single thread to execute even when there are zero elements to allow for book keeping, such as setting
+		// the data count.
+	}
+	else
+	{
+		// Num data - set to zero if writing kernel executes. If kernel doesn't execute, 0 means data collection is empty.
+		OutPackedDataCollection[0] = 0;
+	}
 
 	for (uint32 DataIndex = 0; DataIndex < NumData; ++DataIndex)
 	{
@@ -1118,7 +1133,7 @@ uint32 FPCGDataCollectionDesc::ComputeDataElementCount(EPCGDataType InDataType)
 
 	for (const FPCGDataDesc& DataDesc : DataDescs)
 	{
-		if (DataDesc.Type == InDataType)
+		if (!!(DataDesc.Type & InDataType))
 		{
 			ElementCount += DataDesc.ElementCount;
 		}
