[PCG] Fix unsafe use of InputPin ptr in UPCGCustomHLSLSettings::GetProcessingElemCountForInputPin()

#rnx

diff --git a/Engine/Plugins/PCG/Source/PCG/Private/Compute/Elements/PCGCustomHLSL.cpp b/Engine/Plugins/PCG/Source/PCG/Private/Compute/Elements/PCGCustomHLSL.cpp
--- a/Engine/Plugins/PCG/Source/PCG/Private/Compute/Elements/PCGCustomHLSL.cpp
+++ b/Engine/Plugins/PCG/Source/PCG/Private/Compute/Elements/PCGCustomHLSL.cpp
@@ -135,7 +135,7 @@ FPCGElementPtr UPCGCustomHLSLSettings::CreateElement() const
 
 int UPCGCustomHLSLSettings::GetProcessingElemCountForInputPin(const UPCGPin* InputPin, const UPCGDataBinding* Binding) const
 {
-	check(Binding);
+	check(InputPin && Binding);
 	const FPCGDataForGPU& DataForGPU = Binding->DataForGPU;
 
 	// Upper bound estimate of total number of data elements expected to arrive at this pin.
@@ -191,8 +191,10 @@ int UPCGCustomHLSLSettings::ComputeKernelThreadCount(const UPCGDataBinding* Bind
 	else if (KernelType == EPCGKernelType::PointProcessor)
 	{
 		// Processing volume depends on data arriving on primary pin.
-		const UPCGPin* InputPin = GetPointProcessingInputPin();
-		ThreadCount = GetProcessingElemCountForInputPin(InputPin, Binding);
+		if (const UPCGPin* InputPin = GetPointProcessingInputPin())
+		{
+			ThreadCount = GetProcessingElemCountForInputPin(InputPin, Binding);
+		}
 	}
 	else if (KernelType == EPCGKernelType::Custom)
 	{
