[PCG] Disallow overridable params on GPU, since they are not supported right now
* Also fix UnpackDataCollection DataMismatch warning not logging to console
* Rename PCG_KERNEL_VALIDATION to PCG_KERNEL_VALIDATION_ERR and add PCG_KERNEL_VALIDATION_WARN

#rnx

diff --git a/Engine/Plugins/PCG/Source/PCG/Private/Compute/DataInterfaces/PCGDataCollectionDataInterface.cpp b/Engine/Plugins/PCG/Source/PCG/Private/Compute/DataInterfaces/PCGDataCollectionDataInterface.cpp
--- a/Engine/Plugins/PCG/Source/PCG/Private/Compute/DataInterfaces/PCGDataCollectionDataInterface.cpp
+++ b/Engine/Plugins/PCG/Source/PCG/Private/Compute/DataInterfaces/PCGDataCollectionDataInterface.cpp
@@ -7,6 +7,7 @@
 #include "PCGModule.h"
 #include "PCGSettings.h"
 #include "PCGSubsystem.h"
+#include "Compute/PCGComputeCommon.h"
 #include "Compute/PCGDataBinding.h"
 #include "Compute/Elements/PCGComputeGraphElement.h"
 
@@ -581,32 +582,12 @@ bool UPCGDataCollectionDataProvider::ProcessReadBackData(FPCGComputeGraphContext
 	FPCGDataCollection DataFromGPU;
 	const EPCGUnpackDataCollectionResult Result = PinDesc.UnpackDataCollection(RawReadbackData, OutputPinLabelAlias, DataFromGPU);
 
-#if !(UE_BUILD_SHIPPING || UE_BUILD_TEST) || USE_LOGGING_IN_SHIPPING
 	if (Result == EPCGUnpackDataCollectionResult::DataMismatch)
 	{
-		const FText DataMismatchWarningText = FText::Format(
+		PCG_KERNEL_VALIDATION_WARN(InContext, ProducerSettings, /*bQuiet=*/false, FText::Format(
 			LOCTEXT("UnpackDataCollectionDataMismatch", "Mismatch in expected data while unpacking GPU data collection on pin '{0}'. Static analysis does not match the received data. Data collection will be ignored."),
-			FText::FromName(OutputPinLabel));
-
-#if WITH_EDITOR
-		// InContext will be for the compute graph element (injected during compilation). The code below logs against the original
-		// node so that the log will be visible on the graph.
-		if (ensure(InContext->SourceComponent.IsValid() && InContext->SourceComponent.Get()))
-		{
-			if (UPCGSubsystem* Subsystem = InContext->SourceComponent->GetSubsystem())
-			{
-				FPCGStack StackWithNode = InContext->Stack ? *InContext->Stack : FPCGStack();
-				StackWithNode.PushFrame(ProducerSettings->GetOuter());
-				Subsystem->GetNodeVisualLogsMutable().Log(StackWithNode, ELogVerbosity::Warning, DataMismatchWarningText);
-			}
-		}
-		else
-#endif
-		{
-			PCGE_LOG_C(Warning, LogOnly, InContext, DataMismatchWarningText);
-		}
+			FText::FromName(OutputPinLabel)));
 	}
-#endif
 
 	RawReadbackData.Reset();
 
diff --git a/Engine/Plugins/PCG/Source/PCG/Private/Compute/Elements/PCGCustomHLSL.cpp b/Engine/Plugins/PCG/Source/PCG/Private/Compute/Elements/PCGCustomHLSL.cpp
--- a/Engine/Plugins/PCG/Source/PCG/Private/Compute/Elements/PCGCustomHLSL.cpp
+++ b/Engine/Plugins/PCG/Source/PCG/Private/Compute/Elements/PCGCustomHLSL.cpp
@@ -1001,7 +1001,7 @@ bool UPCGCustomHLSLSettings::IsKernelValid(FPCGContext* InContext, bool bQuiet)
 {
 	if (OutputPins.IsEmpty())
 	{
-		PCG_KERNEL_VALIDATION(InContext, this, bQuiet, LOCTEXT("NoOutputs", "Custom HLSL nodes must have at least one output."));
+		PCG_KERNEL_VALIDATION_ERR(InContext, this, bQuiet, LOCTEXT("NoOutputs", "Custom HLSL nodes must have at least one output."));
 		return false;
 	}
 
@@ -1009,7 +1009,7 @@ bool UPCGCustomHLSLSettings::IsKernelValid(FPCGContext* InContext, bool bQuiet)
 	{
 		if (PinLabel == NAME_None)
 		{
-			PCG_KERNEL_VALIDATION(InContext, This, bQuiet, LOCTEXT("InvalidPinLabelNone", "Pin label 'None' is not a valid pin label."));
+			PCG_KERNEL_VALIDATION_ERR(InContext, This, bQuiet, LOCTEXT("InvalidPinLabelNone", "Pin label 'None' is not a valid pin label."));
 			return false;
 		}
 
@@ -1021,7 +1021,7 @@ bool UPCGCustomHLSLSettings::IsKernelValid(FPCGContext* InContext, bool bQuiet)
 			{
 				if (bFoundPinLabel)
 				{
-					PCG_KERNEL_VALIDATION(InContext, This, bQuiet, FText::Format(LOCTEXT("DuplicatedPinLabels", "Duplicate pin label '{0}', all labels must be unique."), FText::FromName(PinLabel)));
+					PCG_KERNEL_VALIDATION_ERR(InContext, This, bQuiet, FText::Format(LOCTEXT("DuplicatedPinLabels", "Duplicate pin label '{0}', all labels must be unique."), FText::FromName(PinLabel)));
 					return true;
 				}
 
@@ -1063,7 +1063,7 @@ bool UPCGCustomHLSLSettings::IsKernelValid(FPCGContext* InContext, bool bQuiet)
 		{
 			if (Properties.AllowedTypes != EPCGDataType::Point)
 			{
-				PCG_KERNEL_VALIDATION(InContext, this, bQuiet, FText::Format(
+				PCG_KERNEL_VALIDATION_ERR(InContext, this, bQuiet, FText::Format(
 					LOCTEXT("InvalidNonPointPrimaryInput", "'Point Processor' nodes require primary input pin to be of type 'Point', but found '{0}'."),
 					FText::FromString(PCGHLSLElement::GetDataTypeString(Properties.AllowedTypes))));
 
@@ -1073,7 +1073,7 @@ bool UPCGCustomHLSLSettings::IsKernelValid(FPCGContext* InContext, bool bQuiet)
 
 		if (!PCGComputeHelpers::IsTypeAllowedAsInput(Properties.AllowedTypes))
 		{
-			PCG_KERNEL_VALIDATION(InContext, this, bQuiet, FText::Format(
+			PCG_KERNEL_VALIDATION_ERR(InContext, this, bQuiet, FText::Format(
 				LOCTEXT("InvalidInputType", "Unsupported input type '{0}', found on pin '{1}'."),
 				FText::FromString(PCGHLSLElement::GetDataTypeString(Properties.AllowedTypes)),
 				FText::FromName(Properties.Label)));
@@ -1099,7 +1099,7 @@ bool UPCGCustomHLSLSettings::IsKernelValid(FPCGContext* InContext, bool bQuiet)
 		{
 			if (Properties.AllowedTypes != EPCGDataType::Point)
 			{
-				PCG_KERNEL_VALIDATION(InContext, this, bQuiet, FText::Format(
+				PCG_KERNEL_VALIDATION_ERR(InContext, this, bQuiet, FText::Format(
 					LOCTEXT("InvalidNonPointPrimaryOutput", "'Point Processor' and 'Point Generator' nodes require primary output pin to be of type 'Point', but found '{0}'."),
 					FText::FromString(PCGHLSLElement::GetDataTypeString(Properties.AllowedTypes))));
 
@@ -1109,7 +1109,7 @@ bool UPCGCustomHLSLSettings::IsKernelValid(FPCGContext* InContext, bool bQuiet)
 
 		if (!PCGComputeHelpers::IsTypeAllowedAsOutput(Properties.AllowedTypes))
 		{
-			PCG_KERNEL_VALIDATION(InContext, this, bQuiet, FText::Format(
+			PCG_KERNEL_VALIDATION_ERR(InContext, this, bQuiet, FText::Format(
 				LOCTEXT("InvalidOutputType", "Unsupported output type '{0}', found on pin '{1}'."),
 				FText::FromString(PCGHLSLElement::GetDataTypeString(Properties.AllowedTypes)),
 				FText::FromName(Properties.Label)));
@@ -1123,7 +1123,7 @@ bool UPCGCustomHLSLSettings::IsKernelValid(FPCGContext* InContext, bool bQuiet)
 			{
 				if (Properties.FixedBufferElementCount <= 0)
 				{
-					PCG_KERNEL_VALIDATION(InContext, this, bQuiet, FText::Format(
+					PCG_KERNEL_VALIDATION_ERR(InContext, this, bQuiet, FText::Format(
 						LOCTEXT("InvalidFixedBufferSize", "Fixed GPU buffer size on '{0}' was invalid (%d)."),
 						FText::FromName(Properties.Label),
 						Properties.FixedBufferElementCount));
@@ -1133,7 +1133,7 @@ bool UPCGCustomHLSLSettings::IsKernelValid(FPCGContext* InContext, bool bQuiet)
 
 				if (Properties.AllowedTypes == EPCGDataType::Param && Properties.InitializeFromPin != NAME_None && !GetInputPin(Properties.InitializeFromPin))
 				{
-					PCG_KERNEL_VALIDATION(InContext, this, bQuiet, FText::Format(
+					PCG_KERNEL_VALIDATION_ERR(InContext, this, bQuiet, FText::Format(
 						LOCTEXT("InvalidInitFromPin", "Tried to initialize attribute set pin '{0}' from non-existent pin '{1}'. Must reference a valid input pin or be 'None'."),
 						FText::FromName(Properties.Label),
 						FText::FromName(Properties.InitializeFromPin)));
@@ -1145,7 +1145,7 @@ bool UPCGCustomHLSLSettings::IsKernelValid(FPCGContext* InContext, bool bQuiet)
 			{
 				if (InputPins.IsEmpty())
 				{
-					PCG_KERNEL_VALIDATION(InContext, this, bQuiet, FText::Format(
+					PCG_KERNEL_VALIDATION_ERR(InContext, this, bQuiet, FText::Format(
 						LOCTEXT("InvalidBufferSizeNoInputPin", "GPU buffer size for pin '{0}' could not be computed as there are no input pins."),
 						FText::FromName(Properties.Label)));
 
@@ -1154,7 +1154,7 @@ bool UPCGCustomHLSLSettings::IsKernelValid(FPCGContext* InContext, bool bQuiet)
 
 				if (!GetFirstInputPin())
 				{
-					PCG_KERNEL_VALIDATION(InContext, this, bQuiet, FText::Format(
+					PCG_KERNEL_VALIDATION_ERR(InContext, this, bQuiet, FText::Format(
 						LOCTEXT("MissingPrimaryInputPin", "GPU buffer size for pin '{0}' could not be computed, because it refers to the primary input pin, which does not exist."),
 						FText::FromName(Properties.Label)));
 
@@ -1165,7 +1165,7 @@ bool UPCGCustomHLSLSettings::IsKernelValid(FPCGContext* InContext, bool bQuiet)
 			{
 				if (InputPins.IsEmpty())
 				{
-					PCG_KERNEL_VALIDATION(InContext, this, bQuiet, FText::Format(
+					PCG_KERNEL_VALIDATION_ERR(InContext, this, bQuiet, FText::Format(
 						LOCTEXT("InvalidBufferSizeNoInputPins", "GPU buffer size for pin '{0}' could not be computed as there are no input pins on this node."),
 						FText::FromName(Properties.Label)));
 
@@ -1174,7 +1174,7 @@ bool UPCGCustomHLSLSettings::IsKernelValid(FPCGContext* InContext, bool bQuiet)
 
 				if (Properties.BufferSizeInputPinLabels.IsEmpty())
 				{
-					PCG_KERNEL_VALIDATION(InContext, this, bQuiet, FText::Format(
+					PCG_KERNEL_VALIDATION_ERR(InContext, this, bQuiet, FText::Format(
 						LOCTEXT("InvalidBufferSizeNoBufferPins", "GPU buffer size for pin '{0}' could not be computed as input pins are specified in the pin settings."),
 						FText::FromName(Properties.Label)));
 
@@ -1185,7 +1185,7 @@ bool UPCGCustomHLSLSettings::IsKernelValid(FPCGContext* InContext, bool bQuiet)
 				{
 					if (!GetInputPin(Label))
 					{
-						PCG_KERNEL_VALIDATION(InContext, this, bQuiet, FText::Format(
+						PCG_KERNEL_VALIDATION_ERR(InContext, this, bQuiet, FText::Format(
 							LOCTEXT("MissingBufferSizePin", "GPU buffer size for pin '{0}' could not be computed. Invalid pin specified in Input Pins array: '{1}'."),
 							FText::FromName(Properties.Label),
 							FText::FromName(Label)));
@@ -1203,7 +1203,7 @@ bool UPCGCustomHLSLSettings::IsKernelValid(FPCGContext* InContext, bool bQuiet)
 	{
 		if (ThreadCountInputPinLabels.IsEmpty())
 		{
-			PCG_KERNEL_VALIDATION(InContext, this, bQuiet, LOCTEXT("MissingThreadCountPins", "Dispatch thread count is based on input pins but no labels have been set in Input Pins array."));
+			PCG_KERNEL_VALIDATION_ERR(InContext, this, bQuiet, LOCTEXT("MissingThreadCountPins", "Dispatch thread count is based on input pins but no labels have been set in Input Pins array."));
 			return false;
 		}
 
@@ -1211,7 +1211,7 @@ bool UPCGCustomHLSLSettings::IsKernelValid(FPCGContext* InContext, bool bQuiet)
 		{
 			if (!GetInputPin(Label))
 			{
-				PCG_KERNEL_VALIDATION(InContext, this, bQuiet, FText::Format(LOCTEXT("MissingThreadCountPin", "Invalid pin specified in Input Pins array: '{0}'."), FText::FromName(Label)));
+				PCG_KERNEL_VALIDATION_ERR(InContext, this, bQuiet, FText::Format(LOCTEXT("MissingThreadCountPin", "Invalid pin specified in Input Pins array: '{0}'."), FText::FromName(Label)));
 				return false;
 			}
 		}
@@ -1221,7 +1221,7 @@ bool UPCGCustomHLSLSettings::IsKernelValid(FPCGContext* InContext, bool bQuiet)
 	{
 		if (ThreadCountMultiplier < 1)
 		{
-			PCG_KERNEL_VALIDATION(InContext, this, bQuiet, FText::Format(LOCTEXT("InvalidThreadCountMultiplier", "Thread Count Multiplier has invalid value ({0}). Must be greater than 0."), ThreadCountMultiplier));
+			PCG_KERNEL_VALIDATION_ERR(InContext, this, bQuiet, FText::Format(LOCTEXT("InvalidThreadCountMultiplier", "Thread Count Multiplier has invalid value ({0}). Must be greater than 0."), ThreadCountMultiplier));
 			return false;
 		}
 	}
@@ -1231,7 +1231,7 @@ bool UPCGCustomHLSLSettings::IsKernelValid(FPCGContext* InContext, bool bQuiet)
 	{
 		if (AttributeKey.Name == NAME_None)
 		{
-			PCG_KERNEL_VALIDATION(InContext, this, bQuiet, LOCTEXT("InvalidAttributeNameNone", "'None' is not a valid GPU attribute name, check the 'Attributes to Create' array on your pins."));
+			PCG_KERNEL_VALIDATION_ERR(InContext, this, bQuiet, LOCTEXT("InvalidAttributeNameNone", "'None' is not a valid GPU attribute name, check the 'Attributes to Create' array on your pins."));
 			return false;
 		}
 	}
@@ -1248,7 +1248,7 @@ bool UPCGCustomHLSLSettings::IsKernelValid(FPCGContext* InContext, bool bQuiet)
 		{
 			if (ErrorTextPtr)
 			{
-				PCG_KERNEL_VALIDATION(InContext, this, bQuiet, *ErrorTextPtr);
+				PCG_KERNEL_VALIDATION_ERR(InContext, this, bQuiet, *ErrorTextPtr);
 			}
 			return false;
 		}
diff --git a/Engine/Plugins/PCG/Source/PCG/Private/Compute/PCGComputeCommon.cpp b/Engine/Plugins/PCG/Source/PCG/Private/Compute/PCGComputeCommon.cpp
--- a/Engine/Plugins/PCG/Source/PCG/Private/Compute/PCGComputeCommon.cpp
+++ b/Engine/Plugins/PCG/Source/PCG/Private/Compute/PCGComputeCommon.cpp
@@ -43,6 +43,23 @@ namespace PCGComputeHelpers
 	}
 
 #if PCG_KERNEL_LOGGING_ENABLED
+	void LogKernelWarning(const FPCGContext* Context, const UPCGSettings* Settings, const FText& InText)
+	{
+#if WITH_EDITOR
+		if (Context && ensure(Context->SourceComponent.IsValid() && Context->SourceComponent.Get()))
+		{
+			if (UPCGSubsystem* Subsystem = Context->SourceComponent->GetSubsystem())
+			{
+				FPCGStack StackWithNode = Context->Stack ? *Context->Stack : FPCGStack();
+				StackWithNode.PushFrame(Settings->GetOuter());
+
+				Subsystem->GetNodeVisualLogsMutable().Log(StackWithNode, ELogVerbosity::Warning, InText);
+			}
+		}
+#endif
+		PCGE_LOG_C(Warning, LogOnly, Context, InText);
+	}
+
 	void LogKernelError(const FPCGContext* Context, const UPCGSettings* Settings, const FText& InText)
 	{
 #if WITH_EDITOR
diff --git a/Engine/Plugins/PCG/Source/PCG/Private/Compute/PCGComputeCommon.h b/Engine/Plugins/PCG/Source/PCG/Private/Compute/PCGComputeCommon.h
--- a/Engine/Plugins/PCG/Source/PCG/Private/Compute/PCGComputeCommon.h
+++ b/Engine/Plugins/PCG/Source/PCG/Private/Compute/PCGComputeCommon.h
@@ -11,9 +11,11 @@ struct FPCGContext;
 #define PCG_KERNEL_LOGGING_ENABLED (!(UE_BUILD_SHIPPING || UE_BUILD_TEST) || USE_LOGGING_IN_SHIPPING)
 
 #if PCG_KERNEL_LOGGING_ENABLED
-#define PCG_KERNEL_VALIDATION(Context, Settings, bQuiet, ValidationMessage) if (!bQuiet) PCGComputeHelpers::LogKernelError(Context, Settings, ValidationMessage);
+#define PCG_KERNEL_VALIDATION_WARN(Context, Settings, bQuiet, ValidationMessage) if (!bQuiet) PCGComputeHelpers::LogKernelWarning(Context, Settings, ValidationMessage);
+#define PCG_KERNEL_VALIDATION_ERR(Context, Settings, bQuiet, ValidationMessage) if (!bQuiet) PCGComputeHelpers::LogKernelError(Context, Settings, ValidationMessage);
 #else
-#define PCG_KERNEL_VALIDATION(Context, Settings, bQuiet, ValidationMessage) // Log removed
+#define PCG_KERNEL_VALIDATION_WARN(Context, Settings, bQuiet, ValidationMessage) // Log removed
+#define PCG_KERNEL_VALIDATION_ERR(Context, Settings, bQuiet, ValidationMessage) // Log removed
 #endif
 
 namespace PCGComputeConstants
@@ -68,6 +70,7 @@ namespace PCGComputeHelpers
 	bool IsTypeAllowedInDataCollection(EPCGDataType Type);
 
 #if PCG_KERNEL_LOGGING_ENABLED
+	void LogKernelWarning(const FPCGContext* Context, const UPCGSettings* Settings, const FText& InText);
 	void LogKernelError(const FPCGContext* Context, const UPCGSettings* Settings, const FText& InText);
 #endif
 }
diff --git a/Engine/Plugins/PCG/Source/PCG/Private/PCGSettings.cpp b/Engine/Plugins/PCG/Source/PCG/Private/PCGSettings.cpp
--- a/Engine/Plugins/PCG/Source/PCG/Private/PCGSettings.cpp
+++ b/Engine/Plugins/PCG/Source/PCG/Private/PCGSettings.cpp
@@ -11,6 +11,7 @@
 #include "PCGPin.h"
 #include "PCGSubgraph.h"
 #include "PCGSubsystem.h"
+#include "Compute/PCGComputeCommon.h"
 #include "Compute/PCGComputeGraph.h"
 #include "Compute/PCGDataBinding.h"
 #include "Compute/PCGDataForGPU.h"
@@ -45,6 +46,11 @@ namespace PCGSettings
 		TEXT("pcg.Graph.WarnPinNameCollisions"),
 		false,
 		TEXT("Enables warnings when there are name collision between pin names and overrides/user parameters."));
+
+	static TAutoConsoleVariable<bool> CVarWarnOnOverridePinUsage(
+		TEXT("pcg.Graph.GPU.WarnOnOverridePinUsage"),
+		true,
+		TEXT("Enables warnings when parameters are overidden on GPU nodes."));
 }
 
 /** Custom Crc computation that ignores properties that will not affect the computed result of a node. */
@@ -195,6 +201,29 @@ void UPCGSettingsInterface::SetEnabled(bool bInEnabled)
 	}
 }
 
+bool UPCGSettings::IsKernelValid(FPCGContext* InContext, bool bQuiet) const
+{
+	if (PCGSettings::CVarWarnOnOverridePinUsage.GetValueOnAnyThread() && !bQuiet)
+	{
+		for (const FPCGSettingsOverridableParam& Param : OverridableParams())
+		{
+			if (ensure(!Param.PropertiesNames.IsEmpty()))
+			{
+				const FName PropertyName = Param.PropertiesNames[0];
+
+				if (IsPropertyOverriddenByPin(PropertyName))
+				{
+					PCG_KERNEL_VALIDATION_WARN(InContext, this, bQuiet, FText::Format(
+						LOCTEXT("ParamOverrideGPU", "Tried to override pin '{0}', but overrides are not supported on GPU nodes."),
+						FText::FromName(PropertyName)));
+				}
+			}
+		}
+	}
+
+	return true;
+}
+
 FPCGDataCollectionDesc UPCGSettings::ComputeInputPinDataDesc(const UPCGPin* InputPin, const UPCGDataBinding* Binding) const
 {
 	check(InputPin && Binding);
diff --git a/Engine/Plugins/PCG/Source/PCG/Public/Compute/Elements/PCGCustomHLSL.h b/Engine/Plugins/PCG/Source/PCG/Public/Compute/Elements/PCGCustomHLSL.h
--- a/Engine/Plugins/PCG/Source/PCG/Public/Compute/Elements/PCGCustomHLSL.h
+++ b/Engine/Plugins/PCG/Source/PCG/Public/Compute/Elements/PCGCustomHLSL.h
@@ -109,7 +109,8 @@ public:
 	//~Begin UPCGSettings interface
 	virtual TArray<FPCGPinProperties> InputPinProperties() const override { return InputPins; }
 	virtual TArray<FPCGPinProperties> OutputPinProperties() const override;
-	virtual bool IsInputPinRequiredByExecution(const UPCGPin* InPin) const { return true; }
+	virtual bool HasOverridableParams() const override { return false; }
+	virtual bool IsInputPinRequiredByExecution(const UPCGPin* InPin) const override { return true; }
 #if WITH_EDITOR
 	virtual FName GetDefaultNodeName() const override { return FName(TEXT("CustomHLSL")); }
 	virtual FText GetDefaultNodeTitle() const override { return NSLOCTEXT("PCGCustomHLSLElement", "NodeTitle", "Custom HLSL"); }
diff --git a/Engine/Plugins/PCG/Source/PCG/Public/PCGSettings.h b/Engine/Plugins/PCG/Source/PCG/Public/PCGSettings.h
--- a/Engine/Plugins/PCG/Source/PCG/Public/PCGSettings.h
+++ b/Engine/Plugins/PCG/Source/PCG/Public/PCGSettings.h
@@ -517,7 +517,7 @@ public:
 	/** Whether this node should be executed on the GPU. */
 	virtual bool ShouldExecuteOnGPU() const { return bExecuteOnGPU; }
 
-	virtual bool IsKernelValid(FPCGContext* InContext = nullptr, bool bQuiet = true) const { return true; }
+	virtual bool IsKernelValid(FPCGContext* InContext = nullptr, bool bQuiet = true) const;
 	virtual FString GetCookedKernelSource(const TMap<FPCGKernelAttributeKey, int>& GlobalAttributeLookupTable) const { return TEXT(""); }
 	virtual FString GetKernelEntryPoint() const { return TEXT("Main"); }
 	virtual FIntVector GetThreadGroupSize() const { return FIntVector(64, 1, 1); }
