[PCG] Support HiGen Grid Linkages on GPU
* Fixup Grid Linkages with dynamically typed upstream pins
* Properly track upstream output pin to virtual input pin relationships for Grid Linkages on the CPU to GPU boundary

#rnx

diff --git a/Engine/Plugins/PCG/Source/PCG/Private/Graph/PCGGraphCompiler.cpp b/Engine/Plugins/PCG/Source/PCG/Private/Graph/PCGGraphCompiler.cpp
--- a/Engine/Plugins/PCG/Source/PCG/Private/Graph/PCGGraphCompiler.cpp
+++ b/Engine/Plugins/PCG/Source/PCG/Private/Graph/PCGGraphCompiler.cpp
@@ -420,22 +420,33 @@ void FPCGGraphCompiler::CreateGridLinkages(
 				LinkTask.NodeId = InOutCompiledTasks.Num() - 1;
 				LinkTask.StackIndex = InOutCompiledTasks[TaskId].StackIndex;
 
-				LinkTask.Inputs.Emplace(TaskInput.TaskId, UpstreamPin->Properties);
+				FPCGPinProperties UpstreamPinProps = UpstreamPin->Properties;
+
+				// Make sure we have the correct type in the case of dynamically typed pins.
+				if (const UPCGSettings* UpstreamSettings = InOutCompiledTasks[TaskInput.TaskId].Node ? InOutCompiledTasks[TaskInput.TaskId].Node->GetSettings() : nullptr)
+				{
+					UpstreamPinProps.AllowedTypes = UpstreamSettings->GetCurrentPinTypes(UpstreamPin);
+				}
+
+				// Create pin properties for the input of the LinkTask to convey type information. Since we also need to give a pin label, use the default input pin label.
+				// So, effectively, Link Elements have a default input pin and output pin. These will forward the type of the upstream pin, which is required to preserve the
+				// type of the upstream node, which will be culled from the graph.
+				FPCGPinProperties DownstreamPinProps = UpstreamPinProps;
+				DownstreamPinProps.Label = PCGPinConstants::DefaultInputLabel;
+
+				LinkTask.Inputs.Emplace(TaskInput.TaskId, UpstreamPinProps, DownstreamPinProps);
 
 				const EPCGHiGenGrid FromGrid = InOutTaskGenerationGrid[TaskInput.TaskId];
 				const EPCGHiGenGrid ToGrid = InOutTaskGenerationGrid[TaskId];
 
 				// This lambda runs at execution time and attempts to retrieve the data from a larger grid. Capture by value is intentional.
-				auto GridLinkageOperation = [FromGrid, ToGrid, ResourceKey, OutputPinLabel = TaskInput.UpstreamPin.GetValue().Label,
-					DownstreamNode = InOutCompiledTasks[TaskId].Node, InGenerationGrid](FPCGContext* InContext)
+				auto GridLinkageOperation = [FromGrid, ToGrid, ResourceKey, InGenerationGrid](FPCGContext* InContext)
 				{
 					return PCGGraphExecutor::ExecuteGridLinkage(
 						InGenerationGrid,
 						FromGrid,
 						ToGrid,
 						ResourceKey,
-						OutputPinLabel,
-						DownstreamNode,
 						static_cast<FPCGGridLinkageContext*>(InContext));
 				};
 
@@ -444,11 +455,14 @@ void FPCGGraphCompiler::CreateGridLinkages(
 					return new FPCGGridLinkageContext();
 				};
 
-				LinkTask.Element = MakeShared<PCGGraphExecutor::FPCGGridLinkageElement>(GridLinkageOperation, ContextAllocator, FromGrid, ToGrid, ResourceKey);
+				LinkTask.Element = MakeShared<PCGGraphExecutor::FPCGGridLinkageElement>(GridLinkageOperation, ContextAllocator, FromGrid, ToGrid, ResourceKey, UpstreamPin);
 
 				// Now splice in the new task - redirect the downstream task to grab its input from the link task.
 				TaskInput.TaskId = LinkTask.NodeId;
 
+				// Set the task input's UpstreamPin so that it has the correct type.
+				TaskInput.UpstreamPin = UpstreamPinProps;
+
 				// The link needs to execute at both FROM grid size (store) and TO grid size (retrieve).
 				InOutTaskGenerationGrid.Add(FromGrid | ToGrid);
 			}
diff --git a/Engine/Plugins/PCG/Source/PCG/Private/Graph/PCGGraphCompilerGPU.cpp b/Engine/Plugins/PCG/Source/PCG/Private/Graph/PCGGraphCompilerGPU.cpp
--- a/Engine/Plugins/PCG/Source/PCG/Private/Graph/PCGGraphCompilerGPU.cpp
+++ b/Engine/Plugins/PCG/Source/PCG/Private/Graph/PCGGraphCompilerGPU.cpp
@@ -394,12 +394,19 @@ void FPCGGraphCompilerGPU::WireGPUGraphNode(
 
 				++InputCount;
 
-				if (const UPCGNode* UpstreamNode = InOutCompiledTasks[Input.TaskId].Node)
+				const FPCGGraphTask& UpstreamTask = InOutCompiledTasks[Input.TaskId];
+				const UPCGPin* OutputPin = UpstreamTask.Node ? UpstreamTask.Node->GetOutputPin(AddedInput.UpstreamPin->Label) : nullptr;
+
+				// Grid Linkages don't have a node associated, so ask the GridLinkageElement for the output pin instead.
+				if (!OutputPin && UpstreamTask.Element && UpstreamTask.Element->IsGridLinkage())
 				{
-					if (const UPCGPin* OutputPin = UpstreamNode->GetOutputPin(AddedInput.UpstreamPin->Label))
-					{
-						OutOutputCPUPinToVirtualPin.Add(OutputPin, VirtualLabel);
-					}
+					const PCGGraphExecutor::FPCGGridLinkageElement* LinkageElement = static_cast<const PCGGraphExecutor::FPCGGridLinkageElement*>(UpstreamTask.Element.Get());
+					OutputPin = LinkageElement->GetUpstreamPin();
+				}
+
+				if (OutputPin)
+				{
+					OutOutputCPUPinToVirtualPin.Add(OutputPin, VirtualLabel);
 				}
 			}
 		}
diff --git a/Engine/Plugins/PCG/Source/PCG/Private/Graph/PCGGraphExecutor.cpp b/Engine/Plugins/PCG/Source/PCG/Private/Graph/PCGGraphExecutor.cpp
--- a/Engine/Plugins/PCG/Source/PCG/Private/Graph/PCGGraphExecutor.cpp
+++ b/Engine/Plugins/PCG/Source/PCG/Private/Graph/PCGGraphExecutor.cpp
@@ -2878,8 +2878,6 @@ namespace PCGGraphExecutor
 		EPCGHiGenGrid InFromGrid,
 		EPCGHiGenGrid InToGrid,
 		const FString& InResourceKey,
-		const FName& InOutputPinLabel,
-		const UPCGNode* InDownstreamNode,
 		FPCGGridLinkageContext* InContext)
 	{
 		check(InContext);
@@ -2939,7 +2937,9 @@ namespace PCGGraphExecutor
 		if (!!(InFromGrid & InGenerationGrid) && FromGridSize != ToGridSize)
 		{
 			FPCGDataCollection Data;
-			Data.TaggedData = InContext->InputData.GetInputsByPin(InOutputPinLabel);
+
+			// Grid Linkages are assumed to have a single input pin with the default input label.
+			Data.TaggedData = InContext->InputData.GetInputsByPin(PCGPinConstants::DefaultInputLabel);
 
 			PCGGraphExecutionLogging::LogGridLinkageTaskExecuteStore(InContext, InGenerationGrid, FromGridSize, ToGridSize, InResourceKey, Data.TaggedData.Num());
 			
@@ -2997,6 +2997,14 @@ namespace PCGGraphExecutor
 				PCGGraphExecutionLogging::LogGridLinkageTaskExecuteRetrieveSuccess(InContext, ComponentWithData, InResourceKey, Data->TaggedData.Num());
 				InContext->OutputData = *Data;
 
+				// Grid Linkages are assumed to have a single output pin with the default output label.
+				// The normal path in IPCGElement::CleanupAndValidateOutput() to fix up the output label does not run because Grid Linkages
+				// have no Settings. Therefore, we have to fix-up by hand here.
+				for (FPCGTaggedData& TaggedData : InContext->OutputData.TaggedData)
+				{
+					TaggedData.Pin = PCGPinConstants::DefaultOutputLabel;
+				}
+
 				return true;
 			}
 
diff --git a/Engine/Plugins/PCG/Source/PCG/Private/Graph/PCGGraphExecutor.h b/Engine/Plugins/PCG/Source/PCG/Private/Graph/PCGGraphExecutor.h
--- a/Engine/Plugins/PCG/Source/PCG/Private/Graph/PCGGraphExecutor.h
+++ b/Engine/Plugins/PCG/Source/PCG/Private/Graph/PCGGraphExecutor.h
@@ -494,19 +494,25 @@ namespace PCGGraphExecutor
 	class FPCGGridLinkageElement : public FPCGGenericElement
 	{
 	public:
-		FPCGGridLinkageElement(TFunction<bool(FPCGContext*)> InOperation, const FContextAllocator& InContextAllocator, EPCGHiGenGrid InFromGrid, EPCGHiGenGrid InToGrid, const FString& InResourceKey)
+		FPCGGridLinkageElement(TFunction<bool(FPCGContext*)> InOperation, const FContextAllocator& InContextAllocator, EPCGHiGenGrid InFromGrid, EPCGHiGenGrid InToGrid, const FString& InResourceKey, const UPCGPin* InUpstreamPin)
 			: FPCGGenericElement(InOperation, InContextAllocator)
 			, FromGrid(InFromGrid)
 			, ToGrid(InToGrid)
 			, ResourceKey(InResourceKey)
 		{
+			if (IsValid(InUpstreamPin))
+			{
+				UpstreamPin = TWeakObjectPtr(InUpstreamPin);
+			}
 		}
 
-#if WITH_EDITOR
 		//~Begin IPCGElement interface
 		virtual bool IsGridLinkage() const override { return true; }
 		//~End IPCGElement interface
 
+		const UPCGPin* GetUpstreamPin() const { return UpstreamPin.Get(); }
+
+#if WITH_EDITOR
 		/** Return true if the grid sizes & path match. */
 		bool operator==(const FPCGGridLinkageElement& Other) const;
 #endif
@@ -516,6 +522,7 @@ namespace PCGGraphExecutor
 		EPCGHiGenGrid FromGrid = EPCGHiGenGrid::Uninitialized;
 		EPCGHiGenGrid ToGrid = EPCGHiGenGrid::Uninitialized;
 		FString ResourceKey;
+		TWeakObjectPtr<const UPCGPin> UpstreamPin = nullptr;
 	};
 
 	/** Compares InFromGrid and InToGrid and performs data storage/retrieval as necessary to marshal data across execution grids. */
@@ -524,7 +531,5 @@ namespace PCGGraphExecutor
 		EPCGHiGenGrid InFromGrid,
 		EPCGHiGenGrid InToGrid,
 		const FString& InResourceKey,
-		const FName& InOutputPinLabel,
-		const UPCGNode* InDownstreamNode,
 		FPCGGridLinkageContext* InContext);
 }
diff --git a/Engine/Plugins/PCG/Source/PCG/Public/PCGElement.h b/Engine/Plugins/PCG/Source/PCG/Public/PCGElement.h
--- a/Engine/Plugins/PCG/Source/PCG/Public/PCGElement.h
+++ b/Engine/Plugins/PCG/Source/PCG/Public/PCGElement.h
@@ -131,10 +131,10 @@ public:
 	/** Public function called when an element is cancelled, passing its current context if any. */
 	void Abort(FPCGContext* Context) const;
 
-#if WITH_EDITOR
-	/** Is this element used for marshalling data across higen grids. Used as rudimentary RTTI for element object comparisons (editor only). */
+	/** Is this element used for marshaling data across higen grids. Used as rudimentary RTTI for element object comparisons (editor only). */
 	virtual bool IsGridLinkage() const { return false; }
 
+#if WITH_EDITOR
 	/** Is this element used for dispatching compute graphs. Used as rudimentary RTTI for element object comparisons (editor only). */
 	virtual bool IsComputeGraphElement() const { return false; }
 
