[PCG] Create a Clustering node.

+ K-Means
+ Expectation-Maximization

diff --git a/Engine/Plugins/PCG/Source/PCG/Private/Elements/PCGClusterElement.cpp b/Engine/Plugins/PCG/Source/PCG/Private/Elements/PCGClusterElement.cpp
new file mode 100644
--- /dev/null
+++ b/Engine/Plugins/PCG/Source/PCG/Private/Elements/PCGClusterElement.cpp
@@ -0,0 +1,698 @@
+// Copyright Epic Games, Inc. All Rights Reserved.
+
+#include "Elements/PCGClusterElement.h"
+
+#include "PCGContext.h"
+#include "Algo/Accumulate.h"
+#include "Data/PCGPointData.h"
+#include "Helpers/PCGAsync.h"
+#include "Helpers/PCGHelpers.h"
+#include "Metadata/Accessors/PCGAttributeAccessorHelpers.h"
+
+#define LOCTEXT_NAMESPACE "PCGClusterElement"
+
+namespace PCGClusterElement
+{
+	namespace Constants
+	{
+		const FName InputCentroidsLabel = TEXT("Initial Centroids");
+	}
+
+	namespace Helpers
+	{
+		FMatrix OuterProduct(const FVector& V)
+		{
+			return FMatrix
+			{
+				{V.X * V.X, V.X * V.Y, V.X * V.Z, 0},
+				{V.Y * V.X, V.Y * V.Y, V.Y * V.Z, 0},
+				{V.Z * V.X, V.Z * V.Y, V.Z * V.Z, 0},
+				{0, 0, 0, 1}
+			};
+		}
+	} // namespace Helpers
+
+	namespace Algorithm
+	{
+		/* Method of vector quantization to try and partition a number of points into a selected number of clusters.
+		 * Reference: https://en.wikipedia.org/wiki/K-means_clustering
+		 */
+		bool KMeans(FPCGContext* Context, FClusteringData& ClusteringData, const int MaxIterations)
+		{
+			TRACE_CPUPROFILER_EVENT_SCOPE(FPCGClusterElement::KMeans);
+
+			// Set up a few references to state data.
+			uint8& ProcessingStage = ClusteringData.ProcessingStage;
+			const TArray<FPCGPoint>& Points = ClusteringData.OutPointData->GetPoints();
+			const TArray<FVector, TInlineAllocator<16>>& Centroids = ClusteringData.Centroids;
+			TArray<int32>& ClusterAssignments = ClusteringData.PointToClusterAssignments;
+
+			while (ClusteringData.Iteration < MaxIterations)
+			{
+				TRACE_CPUPROFILER_EVENT_SCOPE(FPCGClusterElement::KMeans::IterationLoop);
+
+				// Find the centroid that best represents each point.
+				auto ProcessPoints = [&Points, &Centroids, &bConverged = ClusteringData.bConverged](const int32& PointIndex, int32& InOutCentroidIndex)
+				{
+					TRACE_CPUPROFILER_EVENT_SCOPE(FPCGClusterElement::KMeans::ProcessPoints);
+
+					/** Start with the previously selected centroid. The distance will need to be recalculated in case the
+					 * centroid's location changed on the last iteration.
+					 * Note: The centroids should have all been originally initialized to the first cluster (0) for the
+					 * first iteration.
+					 */
+					int CentroidIndex = InOutCentroidIndex;
+					double MinDistanceSquared = FVector::DistSquared(Points[PointIndex].Transform.GetLocation(), Centroids[InOutCentroidIndex]);
+					for (CentroidIndex = 0; CentroidIndex < Centroids.Num(); ++CentroidIndex)
+					{
+						// Stabilizes against flipping centroids in the event of equidistant points
+						if (CentroidIndex == InOutCentroidIndex)
+						{
+							continue;
+						}
+
+						const double DistanceSquared = FVector::DistSquared(Points[PointIndex].Transform.GetLocation(), Centroids[CentroidIndex]);
+						if (DistanceSquared < MinDistanceSquared)
+						{
+							MinDistanceSquared = DistanceSquared;
+							InOutCentroidIndex = CentroidIndex;
+							// Only flag non-convergence if the cluster assignment has changed.
+							bConverged = false;
+						}
+					}
+
+					return true;
+				};
+
+				// Will alternate stages until convergence is found or the max iterations are reached.
+				if (ProcessingStage == 1) // Stage 1 - Process Points
+				{
+					if (!FPCGAsync::AsyncProcessing(&Context->AsyncState, ClusterAssignments.Num(), ClusterAssignments, ProcessPoints, /*bEnableTimeSlicing=*/true, /*ChunkSize=*/64))
+					{
+						return false;
+					}
+
+					if (ClusteringData.bConverged)
+					{
+						break; // We're done.
+					}
+					else // Go to Stage 2 - Process Centroids
+					{
+						ClusteringData.bConverged = true; // Reset until it's recomputed in Stage 1.
+						ProcessingStage = 2;
+					}
+				}
+
+				// Update the centroids with newly added or removed points
+				auto ProcessCentroids = [Points, ClusterAssignments](const int32 CentroidIndex, FVector& OutResult)
+				{
+					TRACE_CPUPROFILER_EVENT_SCOPE(FPCGClusterElement::KMeans::ProcessCentroids);
+
+					FVector CentroidSum = FVector::ZeroVector;
+					int Count = 0;
+
+					for (int PointIndex = 0; PointIndex < Points.Num(); ++PointIndex)
+					{
+						if (ClusterAssignments[PointIndex] == CentroidIndex)
+						{
+							CentroidSum += Points[PointIndex].Transform.GetLocation();
+							++Count;
+						}
+					}
+
+					if (Count > 0)
+					{
+						// Average location of the points
+						OutResult = CentroidSum / Count;
+					}
+
+					return true;
+				};
+
+				// TODO: This can be better optimized for async, by having workers build partial sums for all/one cluster at a time and synchronizing at the end.
+				// Stage 2 - Process Centroids
+				if (!FPCGAsync::AsyncProcessing(&Context->AsyncState, ClusteringData.Centroids.Num(), ClusteringData.Centroids, ProcessCentroids, /*bEnableTimeSlicing=*/true, /*ChunkSize=*/1))
+				{
+					return false;
+				}
+
+				++ClusteringData.Iteration;
+				ProcessingStage = 1; // Stage 1 - Process Points
+			}
+
+			return true;
+		}
+
+		/* Expectation-Maximization: An iterative method to find the local maximum likelihood estimates for clustering.
+		 * Reference: https://en.wikipedia.org/wiki/Expectation-maximization_algorithm
+		 */
+		bool ExpectationMaximization(FPCGContext* Context, FClusteringData& ClusteringData, const int MaxIterations)
+		{
+			TRACE_CPUPROFILER_EVENT_SCOPE(FPCGClusterElement::ExpectationMaximization);
+
+			// Set up a few references to state data
+			uint8& ProcessingStage = ClusteringData.ProcessingStage;
+			const TArray<FVector>& Points = ClusteringData.EMData.NormalizedPointLocations;
+			TArray<int32>& ClusterAssignments = ClusteringData.PointToClusterAssignments;
+
+			using EM = FClusteringData::FExpectationMaximizationData;
+			TArray<EM::FGaussian, TInlineAllocator<16>>& Gaussians = ClusteringData.EMData.Gaussians;
+			EM::FProbabilityArray& Probabilities = ClusteringData.EMData.Probabilities;
+			std::atomic<double>& LogLikelihoodSum = ClusteringData.EMData.LogLikelihoodSum;
+			double& PreviousLogLikelihood = ClusteringData.EMData.PreviousLogLikelihood;
+
+			// Process until convergence via Log Likelihood or max iterations is reached.
+			while (!ClusteringData.bConverged && ClusteringData.Iteration < MaxIterations)
+			{
+				TRACE_CPUPROFILER_EVENT_SCOPE(FPCGClusterElement::ExpectationMaximization::IterationLoop);
+				if (ProcessingStage == 1) // Stage 1 - Expectation Step
+				{
+					// Gaussian Probability Density Function - The likelihood of a point taking a particular value in a continuous distribution.
+					auto GPDFunction = [](const FVector& Point, const FVector& Mean, const FMatrix& InvCovariance, const double NormalizationFactor)
+					{
+						TRACE_CPUPROFILER_EVENT_SCOPE(FPCGClusterElement::ExpectationMaximization::GPDFunction);
+
+						const FVector Diff = Point - Mean;
+						// TODO: Could likely be optimized to handle the dot product and transform in a single step.
+						const double ExpTerm = -0.5 * (Diff | InvCovariance.TransformVector(Diff));
+
+						return NormalizationFactor * FMath::Exp(ExpTerm);
+					};
+
+					auto ProcessExpectationStep = [&Points, &Probabilities, &Gaussians, &LogLikelihoodSum, GPDFunction](const int32 StartIndex, int32, const int32 Count)
+					{
+						TRACE_CPUPROFILER_EVENT_SCOPE(FPCGClusterElement::ExpectationMaximization::ProcessExpectationStep);
+
+						const int32 EndIndex = StartIndex + Count;
+						double LocalLikelihoodSum = 0;
+
+						for (int PointIndex = StartIndex; PointIndex < EndIndex; ++PointIndex)
+						{
+							double Sum = 0.0;
+							for (int i = 0; i < Gaussians.Num(); ++i)
+							{
+								if (Gaussians[i].bDead)
+								{
+									continue;
+								}
+
+								double Probability = GPDFunction(Points[PointIndex], Gaussians[i].Mean, Gaussians[i].InvCovariance, Gaussians[i].NormalizationFactor);
+								Probability *= Gaussians[i].Weight;
+								Probabilities[PointIndex][i] = Probability;
+								Sum += Probability;
+							}
+
+							/* The probabilities may just be extremely small for very far away points. We can
+							 * artificially bump them up a little to avoid precision errors without any problems.
+							 */
+							Sum = FMath::Max(Sum, UE_DOUBLE_SMALL_NUMBER);
+
+							for (int i = 0; i < Gaussians.Num(); ++i)
+							{
+								Probabilities[PointIndex][i] /= Sum;
+							}
+
+							LocalLikelihoodSum += FMath::Loge(Sum);
+						}
+
+						// Atomically += the overall sum. Could be LogLikelihoodSum += LocalLikelihoodSum; once C++20 is enabled.
+						for (double l = LogLikelihoodSum.load(); !LogLikelihoodSum.compare_exchange_strong(l, l + LocalLikelihoodSum);) {}
+
+						return Count;
+					};
+
+					// Implementation Note: This lambda executes independently at conclusion of the expectation stage.
+					auto UpdateLogLikelihood = [&PreviousLogLikelihood, &LogLikelihoodSum, &bConverged = ClusteringData.bConverged, Tolerance = ClusteringData.EMData.Tolerance](int32)
+					{
+						TRACE_CPUPROFILER_EVENT_SCOPE(FPCGClusterElement::ExpectationMaximization::UpdateLogLikelihood);
+
+						if (FMath::Abs(PreviousLogLikelihood - LogLikelihoodSum) < Tolerance)
+						{
+							bConverged = true;
+						}
+
+						PreviousLogLikelihood = LogLikelihoodSum;
+						LogLikelihoodSum = 0.0;
+					};
+
+					// Run Expectation Stage
+					if (!FPCGAsync::AsyncProcessingRangeEx(
+						&Context->AsyncState,
+						Points.Num(),
+						/*Initialize=*/[] {},
+						ProcessExpectationStep,
+						/*MoveDataRange=*/[](int32, int32, int32) {},
+						/*Finished=*/UpdateLogLikelihood,
+						/*bEnableTimeSlicing=*/true,
+						/*ChunkSize=*/256))
+					{
+						return false;
+					}
+
+					// If we converged in the E-Step, we can stop iterating.
+					if (ClusteringData.bConverged)
+					{
+						break;
+					}
+					else // Go to Stage 2 - Maximization Step
+					{
+						ProcessingStage = 2;
+					}
+				}
+
+				if (ProcessingStage == 2) // Stage 2 - Maximization Step
+				{
+					auto ProcessMaximizationStep = [&Points, &Probabilities](const int32 GaussianIndex, EM::FGaussian& OutGaussian)
+					{
+						TRACE_CPUPROFILER_EVENT_SCOPE(FPCGClusterElement::ExpectationMaximization::ProcessMaximizationStep);
+
+						if (OutGaussian.bDead)
+						{
+							return true;
+						}
+
+						double WeightSum = 0.0;
+						FVector MeanSum = FVector::ZeroVector;
+						FMatrix CovarianceSum = FMatrix::Identity;
+
+						for (int i = 0; i < Points.Num(); ++i)
+						{
+							MeanSum += Probabilities[i][GaussianIndex] * Points[i];
+							WeightSum += Probabilities[i][GaussianIndex];
+						}
+
+						// The gaussian is far enough to have nearly no impact on any points. Remove it from further calculation.
+						if (FMath::IsNearlyZero(WeightSum))
+						{
+							OutGaussian.bDead = true;
+							return true;
+						}
+
+						const double WeightFactor = 1.0 / WeightSum;
+						OutGaussian.Mean = MeanSum * WeightFactor;
+						OutGaussian.Weight = WeightSum / Points.Num();
+
+						for (int i = 0; i < Points.Num(); ++i)
+						{
+							const FVector Difference = Points[i] - OutGaussian.Mean;
+							CovarianceSum += Helpers::OuterProduct(Difference) * Probabilities[i][GaussianIndex];
+						}
+
+						OutGaussian.Covariance = CovarianceSum * WeightFactor;
+						if (!ensure(!OutGaussian.Covariance.ContainsNaN()))
+						{
+							// If we have an instability, clamp it back to something insignificant and then remove it.
+							OutGaussian.Covariance = FMatrix::Identity * UE_DOUBLE_SMALL_NUMBER;
+							OutGaussian.bDead = true;
+						}
+
+						// To account for 4D homogenous FMatrix with non-homogenous 3D points.
+						OutGaussian.Covariance.M[3][3] = 1;
+						OutGaussian.InvCovariance = OutGaussian.Covariance.Inverse();
+
+						// Normalization factor tailored for 3D: 1 / ((2 * pi) ^ (3 / 2)) * |sigma| ^ (1 / 2)), where sigma = M_covariance
+						static double constexpr NormalizationFactor3DConstant = 15.74960994572241974429064599;
+						OutGaussian.NormalizationFactor = 1.0 / (NormalizationFactor3DConstant * FMath::Sqrt(OutGaussian.Covariance.Determinant()));
+
+						return true;
+					};
+
+					// Run Maximization Stage
+					if (!FPCGAsync::AsyncProcessing(&Context->AsyncState, Gaussians.Num(), Gaussians, ProcessMaximizationStep, /*bEnableTimeSlicing=*/true, /*ChunkSize=*/1))
+					{
+						return false;
+					}
+					else // Go back to Stage 1 - Expectation Step
+					{
+						checkSlow(FMath::IsNearlyEqual(1.0, Algo::Accumulate(Gaussians, 0.0, [](double PreviousSum, const EM::FGaussian& Gaussian) { return PreviousSum + Gaussian.Weight; })));
+						ProcessingStage = 1;
+					}
+				}
+
+				++ClusteringData.Iteration;
+			}
+
+			// Assign clusters to the points.
+			auto ProcessClusterAssignment = [&Probabilities](const int32 PointIndex, int32& OutClusterAssignment)
+			{
+				TRACE_CPUPROFILER_EVENT_SCOPE(FPCGClusterElement::ExpectationMaximization::ProcessClusterAssignment);
+
+				double MaxProb = Probabilities[PointIndex][0];
+				int32 ClusterAssignment = 0;
+
+				for (int CentroidIndex = 1; CentroidIndex < Probabilities[PointIndex].Num(); ++CentroidIndex)
+				{
+					if (Probabilities[PointIndex][CentroidIndex] > MaxProb)
+					{
+						ClusterAssignment = CentroidIndex;
+						MaxProb = Probabilities[PointIndex][CentroidIndex];
+					}
+				}
+
+				OutClusterAssignment = ClusterAssignment;
+
+				return true;
+			};
+
+			if (!FPCGAsync::AsyncProcessing(&Context->AsyncState, ClusterAssignments.Num(), ClusterAssignments, ProcessClusterAssignment, /*bEnableTimeSlicing=*/true, /*ChunkSize=*/2048))
+			{
+				return false;
+			}
+
+			return true;
+		}
+	} // namespace Algorithm
+} // namespace PCGClusterElement
+
+#if WITH_EDITOR
+void UPCGClusterSettings::PostEditChangeProperty(FPropertyChangedEvent& PropertyChangedEvent)
+{
+	if (PropertyChangedEvent.Property && (PropertyChangedEvent.Property->GetFName() == GET_MEMBER_NAME_CHECKED(UPCGClusterSettings, InitialCentroidSelection)))
+	{
+		bUseSeed = (InitialCentroidSelection == EPCGInitialClusteringCentroidSelection::RandomPoints);
+	}
+
+	Super::PostEditChangeProperty(PropertyChangedEvent);
+}
+
+FText UPCGClusterSettings::GetNodeTooltipText() const
+{
+	return LOCTEXT("NodeTooltip", "Given a desired number of clusters (categories), find the best fit cluster for each point by distance, using one of various clustering algorithms.");
+}
+
+void UPCGClusterSettings::PostDuplicate(bool bDuplicateForPIE)
+{
+	Super::PostDuplicate(bDuplicateForPIE);
+
+	if (InitialCentroidSelection == EPCGInitialClusteringCentroidSelection::RandomPoints)
+	{
+		bUseSeed = true;
+	}
+}
+#endif // WITH_EDITOR
+
+void UPCGClusterSettings::PostLoad()
+{
+	Super::PostLoad();
+
+	if (InitialCentroidSelection == EPCGInitialClusteringCentroidSelection::RandomPoints)
+	{
+		bUseSeed = true;
+	}
+}
+
+TArray<FPCGPinProperties> UPCGClusterSettings::InputPinProperties() const
+{
+	TArray<FPCGPinProperties> Properties;
+	Properties.Emplace_GetRef(PCGPinConstants::DefaultInputLabel, EPCGDataType::Point).SetRequiredPin();
+
+	if (InitialCentroidSelection == EPCGInitialClusteringCentroidSelection::Input)
+	{
+		// TODO: For now, currently only support one input for N:1, but N:N would be useful too.
+		FPCGPinProperties& InputCentroidsPin = Properties.Emplace_GetRef(PCGClusterElement::Constants::InputCentroidsLabel, EPCGDataType::Point, /*bInAllowMultipleConnections=*/false, /*bAllowMultipleData=*/false);
+#if WITH_EDITOR
+		InputCentroidsPin.Tooltip = LOCTEXT("InputCentroidsTooltip", "Points whose locations will be used as both the quantity and location of the initial centroids for the clustering algorithm.");
+#endif // WITH_EDITOR
+		return Properties;
+	}
+
+	return Properties;
+}
+
+FPCGElementPtr UPCGClusterSettings::CreateElement() const
+{
+	return MakeShared<FPCGClusterElement>();
+}
+
+bool FPCGClusterElement::PrepareDataInternal(FPCGContext* InContext) const
+{
+	TRACE_CPUPROFILER_EVENT_SCOPE(FPCGClusterElement::PrepareData);
+
+	ContextType* Context = static_cast<ContextType*>(InContext);
+	check(Context);
+	const UPCGClusterSettings* Settings = Context->GetInputSettings<UPCGClusterSettings>();
+	check(Settings);
+	const TArray<FPCGTaggedData> Inputs = Context->InputData.GetInputsByPin(PCGPinConstants::DefaultInputLabel);
+
+	// Early out if no source to modify
+	if (Inputs.IsEmpty())
+	{
+		return true;
+	}
+
+	Context->InitializePerExecutionState([Settings](const ContextType* Context, ExecStateType& OutState)
+	{
+		TRACE_CPUPROFILER_EVENT_SCOPE(FPCGClusterElement::PrepareData::InitializePerExecutionState);
+
+		switch (Settings->Algorithm)
+		{
+			case EPCGClusterAlgorithm::KMeans:
+				OutState.ClusteringFunction = PCGClusterElement::Algorithm::KMeans;
+				break;
+			case EPCGClusterAlgorithm::EM:
+				OutState.ClusteringFunction = PCGClusterElement::Algorithm::ExpectationMaximization;
+				break;
+			default:
+				checkNoEntry();
+				return EPCGTimeSliceInitResult::AbortExecution;
+		}
+		check(OutState.ClusteringFunction);
+
+		switch (Settings->InitialCentroidSelection)
+		{
+			case EPCGInitialClusteringCentroidSelection::RandomPoints:
+			{
+				// Centroids will be initialized on a per-iteration input basis, but we can allocate them here.
+				OutState.InitialCentroids.SetNum(Settings->NumClusters);
+				break;
+			}
+
+			case EPCGInitialClusteringCentroidSelection::Input:
+			{
+				const TArray<FPCGTaggedData> CentroidInputs = Context->InputData.GetInputsByPin(PCGClusterElement::Constants::InputCentroidsLabel);
+				if (CentroidInputs.IsEmpty())
+				{
+					return EPCGTimeSliceInitResult::NoOperation;
+				}
+
+				if (CentroidInputs.Num() > 1)
+				{
+					PCGLog::InputOutput::LogFirstInputOnlyWarning(PCGClusterElement::Constants::InputCentroidsLabel, Context);
+				}
+
+				const UPCGPointData* CentroidPoints = Cast<UPCGPointData>(CentroidInputs[0].Data);
+				if (!CentroidPoints || CentroidPoints->IsEmpty())
+				{
+					return EPCGTimeSliceInitResult::NoOperation;
+				}
+
+				OutState.InitialCentroids.Reserve(CentroidPoints->GetPoints().Num());
+				for (const FPCGPoint& Point : CentroidPoints->GetPoints())
+				{
+					OutState.InitialCentroids.Emplace(Point.Transform.GetLocation());
+				}
+				break;
+			}
+
+			default:
+				checkNoEntry();
+				return EPCGTimeSliceInitResult::AbortExecution;
+		}
+
+		return EPCGTimeSliceInitResult::Success;
+	});
+
+	auto InitializeIteration = [&Inputs, Settings, Context](IterStateType& OutState, const ExecStateType& ExecState, const uint32 IterationIndex)
+	{
+		TRACE_CPUPROFILER_EVENT_SCOPE(FPCGClusterElement::PrepareData::InitializeIteration);
+
+		const UPCGPointData* InputPointData = CastChecked<UPCGPointData>(Inputs[IterationIndex].Data);
+
+		if (!InputPointData || InputPointData->IsEmpty())
+		{
+			return EPCGTimeSliceInitResult::NoOperation;
+		}
+
+		UPCGPointData* OutPointData = FPCGContext::NewObject_AnyThread<UPCGPointData>(Context);
+		OutPointData->InitializeFromData(InputPointData);
+		OutPointData->GetMutablePoints() = InputPointData->GetPointsCopy();
+
+		if (!OutPointData->Metadata->FindOrCreateAttribute<int32>(Settings->ClusterAttribute, 0, /*bAllowsInterpolation=*/false))
+		{
+			PCGLog::Metadata::LogFailToCreateAttribute<int32>(Settings->ClusterAttribute, Context);
+			return EPCGTimeSliceInitResult::AbortExecution;
+		}
+
+		OutState.OutPointData = OutPointData;
+
+		FPCGAttributePropertySelector Selector;
+		Selector.SetAttributeName(Settings->ClusterAttribute);
+		OutState.Accessor = PCGAttributeAccessorHelpers::CreateAccessor(OutPointData, Selector);
+		OutState.Keys = PCGAttributeAccessorHelpers::CreateKeys(OutPointData, Selector);
+		if (!OutState.Accessor || !OutState.Keys)
+		{
+			PCGLog::Metadata::LogFailToCreateAccessor(Selector, Context);
+			return EPCGTimeSliceInitResult::AbortExecution;
+		}
+
+		const int32 NumPoints = OutPointData->GetNumPoints();
+		const int32 NumClusters = ExecState.InitialCentroids.Num();
+
+		// Early out if centroid count makes this trivial. The default value for the attribute is already set to 0.
+		if (NumClusters < 2)
+		{
+			Context->OutputData.TaggedData.Emplace_GetRef().Data = OutPointData;
+			return EPCGTimeSliceInitResult::NoOperation;
+		}
+
+		// With more clusters than points, each point should be assigned a unique cluster--so this is an acceptable shortcut.
+		if (NumClusters >= NumPoints)
+		{
+			TArray<int32> Clusters;
+			Clusters.Reserve(NumPoints);
+			for (int i = 0; i < NumPoints; ++i)
+			{
+				Clusters.Emplace(i);
+			}
+
+			OutState.Accessor->SetRange(MakeConstArrayView(Clusters), 0, *OutState.Keys);
+			Context->OutputData.TaggedData.Emplace_GetRef().Data = OutPointData;
+			return EPCGTimeSliceInitResult::NoOperation;
+		}
+
+		// Initialize all the cluster assignments to the first cluster (0)
+		OutState.PointToClusterAssignments.SetNumZeroed(NumPoints);
+
+		auto InitializeCentroids = [Settings, NumClusters, NumPoints, &OutState, &ExecState]
+		{
+			if (Settings->InitialCentroidSelection == EPCGInitialClusteringCentroidSelection::RandomPoints)
+			{
+				const TArray<FPCGPoint>& Points = OutState.OutPointData->GetPoints();
+
+				FRandomStream RandomStream(Settings->GetSeed());
+				TArray<int32> StartingIndices = PCGHelpers::GetRandomIndices(RandomStream, NumPoints, NumClusters);
+
+				OutState.Centroids.Reserve(NumClusters);
+				for (int i = 0; i < NumClusters; ++i)
+				{
+					OutState.Centroids.Emplace(Points[StartingIndices[i]].Transform.GetLocation());
+				}
+			}
+			else
+			{
+				OutState.Centroids.Append(ExecState.InitialCentroids);
+			}
+		};
+
+		switch (Settings->Algorithm)
+		{
+			case EPCGClusterAlgorithm::KMeans:
+			{
+				/* TODO: There are some pre-warming algorithms that could be employed to reduce the work required to
+				 * reach convergence. They would need investigation and profiling. Some possibilities may include:
+				 * - KMeans++ - https://en.wikipedia.org/wiki/K-means++
+				 * - Find furthest points using the Octree
+				 * - For large point clouds, Principal Component Analysis might be useful.
+				 */
+
+				// Randomly select indices to start as cluster centroids
+				InitializeCentroids();
+
+				break;
+			}
+
+			case EPCGClusterAlgorithm::EM:
+			{
+				InitializeCentroids();
+
+				const TArray<FPCGPoint>& Points = OutPointData->GetPoints();
+
+				// Normalization of point locations may be required to bring the GPD Function's Exponential Term into the -709.8 -> 709.8 range to prevent numerical overflow.
+				double MaxSquaredLength = std::numeric_limits<double>::min();
+				for (const FPCGPoint& Point : Points)
+				{
+					MaxSquaredLength = FMath::Max(MaxSquaredLength, Point.Transform.GetLocation().SquaredLength());
+				}
+
+				double ScalingFactor = 1.0;
+				// Only normalize the system if needed.
+				if (MaxSquaredLength > 1.0)
+				{
+					const double Length = FMath::Sqrt(MaxSquaredLength);
+					ScalingFactor = UE_DOUBLE_EULERS_NUMBER / Length;
+				}
+
+				OutState.EMData.Gaussians.Reserve(NumClusters);
+				for (int i = 0; i < NumClusters; ++i)
+				{
+					OutState.EMData.Gaussians.Emplace(OutState.Centroids[i] * ScalingFactor, FMatrix::Identity, 1.0 / NumClusters);
+				}
+
+				OutState.EMData.Probabilities.Initialize(Points.Num(), NumClusters);
+
+				OutState.EMData.NormalizedPointLocations.Reserve(Points.Num());
+				Algo::Transform(Points, OutState.EMData.NormalizedPointLocations, [ScalingFactor](const FPCGPoint& Point)
+				{
+					return Point.Transform.GetLocation() * ScalingFactor;
+				});
+
+				break;
+			}
+			default:
+				checkNoEntry();
+				return EPCGTimeSliceInitResult::AbortExecution;
+		}
+
+		return EPCGTimeSliceInitResult::Success;
+	};
+
+	Context->InitializePerIterationStates(Inputs.Num(), InitializeIteration);
+
+	return true;
+}
+
+bool FPCGClusterElement::ExecuteInternal(FPCGContext* InContext) const
+{
+	TRACE_CPUPROFILER_EVENT_SCOPE(FPCGClusterElement::Execute);
+
+	ContextType* TimeSlicedContext = static_cast<ContextType*>(InContext);
+	check(TimeSlicedContext);
+
+	if (!TimeSlicedContext->DataIsPreparedForExecution())
+	{
+		return true;
+	}
+
+	if (TimeSlicedContext->GetExecutionStateResult() == EPCGTimeSliceInitResult::NoOperation)
+	{
+		TimeSlicedContext->OutputData = TimeSlicedContext->InputData;
+		return true;
+	}
+
+	const UPCGClusterSettings* Settings = TimeSlicedContext->GetInputSettings<UPCGClusterSettings>();
+	check(Settings);
+
+	return ExecuteSlice(TimeSlicedContext, [MaxIterations = Settings->MaxIterations](ContextType* Context, const ExecStateType& ExecState, IterStateType& ClusteringData, const uint32 IterIndex)
+	{
+		if (Context->GetIterationStateResult(IterIndex) == EPCGTimeSliceInitResult::NoOperation)
+		{
+			return true;
+		}
+
+		if (!ExecState.ClusteringFunction(Context, ClusteringData, MaxIterations))
+		{
+			return false;
+		}
+
+		// Once finished, set the attribute values.
+		check(ClusteringData.OutPointData);
+		Context->OutputData.TaggedData.Emplace_GetRef().Data = ClusteringData.OutPointData;
+
+		// Finally, update the attribute
+		ClusteringData.Accessor->SetRange(MakeConstArrayView(ClusteringData.PointToClusterAssignments), 0, *ClusteringData.Keys);
+
+		return true;
+	});
+}
+
+#undef LOCTEXT_NAMESPACE
diff --git a/Engine/Plugins/PCG/Source/PCG/Public/Data/PCGPointData.h b/Engine/Plugins/PCG/Source/PCG/Public/Data/PCGPointData.h
--- a/Engine/Plugins/PCG/Source/PCG/Public/Data/PCGPointData.h
+++ b/Engine/Plugins/PCG/Source/PCG/Public/Data/PCGPointData.h
@@ -105,12 +105,18 @@ public:
 	UFUNCTION(BlueprintCallable, BlueprintPure = false, Category = SpatialData)
 	TArray<FPCGPoint> GetPointsCopy() const { return Points; }
 
+	UFUNCTION(BlueprintCallable, Category = SpatialData)
+	int32 GetNumPoints() const { return Points.Num(); }
+
+	UFUNCTION(BlueprintCallable, Category = SpatialData)
+	bool IsEmpty() const { return Points.IsEmpty(); }
+
 	UFUNCTION(BlueprintCallable, Category = SpatialData)
 	FPCGPoint GetPoint(int32 Index) const;
 
 	UFUNCTION(BlueprintCallable, Category = SpatialData)
 	void SetPoints(const TArray<FPCGPoint>& InPoints);
-	
+
 	UFUNCTION(BlueprintCallable, Category = SpatialData)
 	void CopyPointsFrom(const UPCGPointData* InData, const TArray<int>& InDataIndices);
 
@@ -139,7 +145,3 @@ protected:
 	mutable bool bBoundsAreDirty = true;
 	mutable bool bOctreeIsDirty = true;
 };
-
-#if UE_ENABLE_INCLUDE_ORDER_DEPRECATED_IN_5_2
-#include "CoreMinimal.h"
-#endif
diff --git a/Engine/Plugins/PCG/Source/PCG/Public/Elements/PCGClusterElement.h b/Engine/Plugins/PCG/Source/PCG/Public/Elements/PCGClusterElement.h
new file mode 100644
--- /dev/null
+++ b/Engine/Plugins/PCG/Source/PCG/Public/Elements/PCGClusterElement.h
@@ -0,0 +1,163 @@
+// Copyright Epic Games, Inc. All Rights Reserved.
+
+#pragma once
+
+#include "PCGSettings.h"
+
+#include "Elements/PCGTimeSlicedElementBase.h"
+#include "Metadata/Accessors/IPCGAttributeAccessorTpl.h"
+
+#include "PCGClusterElement.generated.h"
+
+class UPCGPointData;
+
+UENUM()
+enum class EPCGClusterAlgorithm
+{
+	KMeans UMETA(Tooltip = "K-Means - Categorizes points by their distance from a predefined number of continuously updated centroids."),
+	EM UMETA(Tooltip = "Expectation-Maximization - Categorizes points by using a Gaussian Mixture Model.")
+};
+
+UENUM()
+enum class EPCGInitialClusteringCentroidSelection
+{
+	RandomPoints UMETA(Tooltip = "Random points from the input will be selected as the starting centroids."),
+	Input UMETA(Tooltip = "Point locations from the secondary pin will be used as the starting centroid locations."),
+};
+
+/** Given a desired number of clusters (categories), find the best fit cluster for each point by distance, using one of various clustering algorithms. */
+UCLASS(MinimalAPI, BlueprintType, ClassGroup = (Procedural))
+class UPCGClusterSettings : public UPCGSettings
+{
+	GENERATED_BODY()
+
+public:
+	//~Begin UPCGSettings interface
+#if WITH_EDITOR
+	virtual FName GetDefaultNodeName() const override { return FName(TEXT("ClusterElement")); }
+	virtual FText GetDefaultNodeTitle() const override { return NSLOCTEXT("PCGClusterElement", "NodeTitle", "Cluster"); }
+	virtual EPCGSettingsType GetType() const override { return EPCGSettingsType::PointOps; }
+	virtual FText GetNodeTooltipText() const override;
+	virtual void PostEditChangeProperty(FPropertyChangedEvent& PropertyChangedEvent) override;
+	virtual void PostDuplicate(bool bDuplicateForPIE) override;
+#endif
+	virtual void PostLoad() override;
+
+protected:
+	virtual TArray<FPCGPinProperties> InputPinProperties() const override;
+	virtual TArray<FPCGPinProperties> OutputPinProperties() const override { return Super::DefaultPointOutputPinProperties(); }
+	virtual FPCGElementPtr CreateElement() const override;
+	//~End UPCGSettings interface
+
+public:
+	/** Mathematical algorithm for selecting clusters. */
+	UPROPERTY(EditAnywhere, BlueprintReadWrite, Category = Settings, meta = (PCG_Overridable))
+	EPCGClusterAlgorithm Algorithm = EPCGClusterAlgorithm::KMeans;
+
+	/** Method for selecting the initial average location of the clusters. A good estimate of these locations can rapidly speed up the process in some cases. */
+	UPROPERTY(EditAnywhere, BlueprintReadWrite, Category = Settings)
+	EPCGInitialClusteringCentroidSelection InitialCentroidSelection = EPCGInitialClusteringCentroidSelection::RandomPoints;
+
+	/** Number of clusters (segments) to group the points into. Each point will be assigned a cluster at the end. */
+	UPROPERTY(EditAnywhere, BlueprintReadWrite, Category = Settings, DisplayName = "Clusters", meta = (ClampMin = "1", EditCondition = "InitialCentroidSelection == EPCGInitialClusteringCentroidSelection::RandomPoints", EditConditionHides, PCG_Overridable))
+	int32 NumClusters = 3;
+
+	/** Cluster IDs will be written to this attribute on the output. */
+	UPROPERTY(EditAnywhere, BlueprintReadWrite, Category = Settings)
+	FName ClusterAttribute = TEXT("ClusterID");
+
+	/** A limit on the number of iterations to run on each algorithm, if it doesn't otherwise converge. A higher number can offer more accuracy at the cost of processing time. */
+	UPROPERTY(EditAnywhere, BlueprintReadWrite, Category = "Settings|Advanced", meta = (ClampMin = "1", PCG_Overridable))
+	int32 MaxIterations = 100;
+
+	/** For EM, the maximum allowed difference between the last two iterations' "Log Likelihood"--which converges from positive infinity to 0 in relation to point-to-cluster probabilities.
+	 * It is exponential, so raising this number can offer faster iteration if exact precision isn't needed.
+	 */
+	UPROPERTY(EditAnywhere, BlueprintReadWrite, Category = "Settings|Advanced", meta = (ClampMin = "0.0001", EditCondition = "Algorithm == EPCGClusterAlgorithm::EM", EditConditionHides, PCG_Overridable))
+	double Tolerance = UE_DOUBLE_KINDA_SMALL_NUMBER;
+};
+
+namespace PCGClusterElement
+{
+	struct FClusteringData
+	{
+		UPCGPointData* OutPointData = nullptr;
+		TUniquePtr<IPCGAttributeAccessor> Accessor;
+		TUniquePtr<IPCGAttributeAccessorKeys> Keys;
+
+		// Shared cached data
+		int32 Iteration = 0;
+		// The stage cached for alternating async processing. Stage 0 == Complete.
+		uint8 ProcessingStage = 1;
+		bool bConverged = false;
+		TArray<FVector, TInlineAllocator<16>> Centroids;
+		TArray<int32> PointToClusterAssignments;
+
+		// Expectation-Maximization cached data
+		struct FExpectationMaximizationData
+		{
+			struct FGaussian
+			{
+				FGaussian() = default;
+
+				FGaussian(const FVector& InMean, const FMatrix& InCovariance, const double InWeight) : Mean(InMean), Covariance(InCovariance), Weight(InWeight)
+				{}
+
+				FVector Mean = FVector::ZeroVector;
+				FMatrix Covariance = FMatrix::Identity;
+				FMatrix InvCovariance = FMatrix::Identity;
+				double NormalizationFactor = 1.0;
+				double Weight = 0.0;
+				bool bDead = false;
+			};
+
+			struct FProbabilityArray
+			{
+				void Initialize(const int32 InNumPoints, const int32 InNumCentroids)
+				{
+					check(InNumPoints > 0 && InNumCentroids > 0);
+					NumPoints = InNumPoints;
+					NumCentroids = InNumCentroids;
+					Array.SetNumZeroed(NumPoints * NumCentroids);
+				}
+
+				TArrayView<double> operator[](const int32 Index)
+				{
+					check(Array.IsValidIndex(Index * NumCentroids));
+					return TArrayView<double>(&Array[Index * NumCentroids], NumCentroids);
+				}
+
+			private:
+				int32 NumPoints = 0;
+				int32 NumCentroids = 0;
+				TArray<double, TInlineAllocator<1024>> Array;
+			};
+
+			TArray<FGaussian, TInlineAllocator<16>> Gaussians;
+
+			FProbabilityArray Probabilities;
+
+			TArray<FVector> NormalizedPointLocations;
+			double Tolerance = UE_DOUBLE_KINDA_SMALL_NUMBER;
+			// A value from -inf to some positive finite value to represent how well the model fits the data. It will be used to recognize convergence.
+			std::atomic<double> LogLikelihoodSum = 0;
+			double PreviousLogLikelihood = std::numeric_limits<double>::max();
+		};
+
+		FExpectationMaximizationData EMData;
+	};
+
+	struct FExecState
+	{
+		using AlgoSignature = bool(*)(FPCGContext* Context, FClusteringData& ClusteringData, const int MaxIterations);
+		AlgoSignature ClusteringFunction = nullptr;
+		TArray<FVector> InitialCentroids;
+	};
+}
+
+class FPCGClusterElement : public TPCGTimeSlicedElementBase<PCGClusterElement::FExecState, PCGClusterElement::FClusteringData>
+{
+protected:
+	virtual bool PrepareDataInternal(FPCGContext* InContext) const override;
+	virtual bool ExecuteInternal(FPCGContext* InContext) const override;
+};
diff --git a/Engine/Plugins/PCG/Source/PCG/Public/Helpers/PCGAsync.h b/Engine/Plugins/PCG/Source/PCG/Public/Helpers/PCGAsync.h
--- a/Engine/Plugins/PCG/Source/PCG/Public/Helpers/PCGAsync.h
+++ b/Engine/Plugins/PCG/Source/PCG/Public/Helpers/PCGAsync.h
@@ -568,12 +568,13 @@ namespace FPCGAsync
 	* @param bAllowChunkSizeOverride - If true, ChunkSize can be overridden by 'pcg.AsyncOverrideChunkSize' CVar
 	* @returns true if the processing is done, false otherwise. Use this to know if you need to reschedule the task.
 	*/
-	template <typename OutputType, typename Func>
-	bool AsyncProcessing(FPCGAsyncState* AsyncState, int32 NumIterations, TArray<OutputType>& OutData, Func&& InFunc, const bool bEnableTimeSlicing, const int32 ChunkSize = 64, bool bAllowChunkSizeOverride = true)
+	template <typename OutputType, typename Func, typename AllocatorType = FDefaultAllocator>
+	bool AsyncProcessing(FPCGAsyncState* AsyncState, int32 NumIterations, TArray<OutputType, AllocatorType>& OutData, Func&& InFunc, const bool bEnableTimeSlicing, const int32 ChunkSize = 64, bool bAllowChunkSizeOverride = true)
 	{
 		auto Initialize = [&OutData, NumIterations]()
 		{
-			OutData.SetNumUninitialized(NumIterations);
+			// Array will be shrunk at the end of the processing.
+			OutData.SetNumUninitialized(NumIterations, EAllowShrinking::No);
 		};
 		
 		auto IterationInnerLoop = [Func = MoveTemp(InFunc), &OutData](int32 ReadIndex, int32 WriteIndex) -> int32
