[PCG] Converted bUseSeed to a virtual override

#rnx

diff --git a/Engine/Plugins/PCG/Source/PCG/Private/Compute/Elements/PCGCustomHLSL.cpp b/Engine/Plugins/PCG/Source/PCG/Private/Compute/Elements/PCGCustomHLSL.cpp
--- a/Engine/Plugins/PCG/Source/PCG/Private/Compute/Elements/PCGCustomHLSL.cpp
+++ b/Engine/Plugins/PCG/Source/PCG/Private/Compute/Elements/PCGCustomHLSL.cpp
@@ -84,7 +84,6 @@ bool FPCGPinPropertiesGPU::CanEditChange(const FEditPropertyChain& PropertyChain
 UPCGCustomHLSLSettings::UPCGCustomHLSLSettings()
 {
 	bExecuteOnGPU = true;
-	bUseSeed = true;
 }
 
 #if WITH_EDITOR
diff --git a/Engine/Plugins/PCG/Source/PCG/Private/Elements/PCGAttributeNoise.cpp b/Engine/Plugins/PCG/Source/PCG/Private/Elements/PCGAttributeNoise.cpp
--- a/Engine/Plugins/PCG/Source/PCG/Private/Elements/PCGAttributeNoise.cpp
+++ b/Engine/Plugins/PCG/Source/PCG/Private/Elements/PCGAttributeNoise.cpp
@@ -93,7 +93,6 @@ namespace PCGAttributeNoiseSettings
 
 UPCGAttributeNoiseSettings::UPCGAttributeNoiseSettings()
 {
-	bUseSeed = true;
 	InputSource.SetPointProperty(EPCGPointProperties::Density);
 }
 
diff --git a/Engine/Plugins/PCG/Source/PCG/Private/Elements/PCGClusterElement.cpp b/Engine/Plugins/PCG/Source/PCG/Private/Elements/PCGClusterElement.cpp
--- a/Engine/Plugins/PCG/Source/PCG/Private/Elements/PCGClusterElement.cpp
+++ b/Engine/Plugins/PCG/Source/PCG/Private/Elements/PCGClusterElement.cpp
@@ -365,40 +365,15 @@ namespace PCGClusterElement
 } // namespace PCGClusterElement
 
 #if WITH_EDITOR
-void UPCGClusterSettings::PostEditChangeProperty(FPropertyChangedEvent& PropertyChangedEvent)
-{
-	if (PropertyChangedEvent.Property && (PropertyChangedEvent.Property->GetFName() == GET_MEMBER_NAME_CHECKED(UPCGClusterSettings, InitialCentroidSelection)))
-	{
-		bUseSeed = (InitialCentroidSelection == EPCGInitialClusteringCentroidSelection::RandomPoints);
-	}
-
-	Super::PostEditChangeProperty(PropertyChangedEvent);
-}
-
 FText UPCGClusterSettings::GetNodeTooltipText() const
 {
 	return LOCTEXT("NodeTooltip", "Given a desired number of clusters (categories), find the best fit cluster for each point by distance, using one of various clustering algorithms.");
 }
-
-void UPCGClusterSettings::PostDuplicate(bool bDuplicateForPIE)
-{
-	Super::PostDuplicate(bDuplicateForPIE);
-
-	if (InitialCentroidSelection == EPCGInitialClusteringCentroidSelection::RandomPoints)
-	{
-		bUseSeed = true;
-	}
-}
 #endif // WITH_EDITOR
 
-void UPCGClusterSettings::PostLoad()
+bool UPCGClusterSettings::UseSeed() const
 {
-	Super::PostLoad();
-
-	if (InitialCentroidSelection == EPCGInitialClusteringCentroidSelection::RandomPoints)
-	{
-		bUseSeed = true;
-	}
+	return InitialCentroidSelection == EPCGInitialClusteringCentroidSelection::RandomPoints;
 }
 
 TArray<FPCGPinProperties> UPCGClusterSettings::InputPinProperties() const
diff --git a/Engine/Plugins/PCG/Source/PCG/Private/Elements/PCGExecuteBlueprint.cpp b/Engine/Plugins/PCG/Source/PCG/Private/Elements/PCGExecuteBlueprint.cpp
--- a/Engine/Plugins/PCG/Source/PCG/Private/Elements/PCGExecuteBlueprint.cpp
+++ b/Engine/Plugins/PCG/Source/PCG/Private/Elements/PCGExecuteBlueprint.cpp
@@ -358,8 +358,6 @@ FRandomStream UPCGBlueprintElement::GetRandomStream(FPCGContext& InContext) cons
 
 UPCGBlueprintSettings::UPCGBlueprintSettings()
 {
-	bUseSeed = true;
-	
 #if WITH_EDITORONLY_DATA
 	bExposeToLibrary = HasAnyFlags(RF_ClassDefaultObject);
 #endif
diff --git a/Engine/Plugins/PCG/Source/PCG/Private/Elements/PCGMatchAndSetAttributes.cpp b/Engine/Plugins/PCG/Source/PCG/Private/Elements/PCGMatchAndSetAttributes.cpp
--- a/Engine/Plugins/PCG/Source/PCG/Private/Elements/PCGMatchAndSetAttributes.cpp
+++ b/Engine/Plugins/PCG/Source/PCG/Private/Elements/PCGMatchAndSetAttributes.cpp
@@ -67,10 +67,10 @@ void UPCGMatchAndSetAttributesSettings::ApplyDeprecationBeforeUpdatePins(UPCGNod
 }
 #endif // WITH_EDITOR
 
-UPCGMatchAndSetAttributesSettings::UPCGMatchAndSetAttributesSettings()
+bool UPCGMatchAndSetAttributesSettings::UseSeed() const
 {
 	// Minor TODO: could mark use seed true only if we don't use the input weight attribute
-	bUseSeed = true;
+	return true;
 }
 
 TArray<FPCGPinProperties> UPCGMatchAndSetAttributesSettings::InputPinProperties() const
@@ -992,4 +992,4 @@ bool FPCGMatchAndSetAttributesElement::ExecuteInternal(FPCGContext* InContext) c
 	});
 }
 
-#undef LOCTEXT_NAMESPACE
\ No newline at end of file
+#undef LOCTEXT_NAMESPACE
diff --git a/Engine/Plugins/PCG/Source/PCG/Private/Elements/PCGMutateSeed.cpp b/Engine/Plugins/PCG/Source/PCG/Private/Elements/PCGMutateSeed.cpp
--- a/Engine/Plugins/PCG/Source/PCG/Private/Elements/PCGMutateSeed.cpp
+++ b/Engine/Plugins/PCG/Source/PCG/Private/Elements/PCGMutateSeed.cpp
@@ -14,11 +14,6 @@ namespace PCGMutateSeedConstants
 	static constexpr int32 PointsPerChunk = 98304;
 }
 
-UPCGMutateSeedSettings::UPCGMutateSeedSettings()
-{
-	bUseSeed = true;
-}
-
 FPCGElementPtr UPCGMutateSeedSettings::CreateElement() const
 {
 	return MakeShared<FPCGMutateSeedElement>();
diff --git a/Engine/Plugins/PCG/Source/PCG/Private/Elements/PCGPointMatchAndSet.cpp b/Engine/Plugins/PCG/Source/PCG/Private/Elements/PCGPointMatchAndSet.cpp
--- a/Engine/Plugins/PCG/Source/PCG/Private/Elements/PCGPointMatchAndSet.cpp
+++ b/Engine/Plugins/PCG/Source/PCG/Private/Elements/PCGPointMatchAndSet.cpp
@@ -25,8 +25,6 @@ UPCGPointMatchAndSetSettings::UPCGPointMatchAndSetSettings(const FObjectInitiali
 	{
 		MatchAndSetInstance = ObjectInitializer.CreateDefaultSubobject<UPCGMatchAndSetWeighted>(this, TEXT("DefaultMatchAndSet"));
 	}
-
-	bUseSeed = MatchAndSetInstance && MatchAndSetInstance->UsesRandomProcess();
 }
 
 #if WITH_EDITOR
@@ -50,6 +48,11 @@ void UPCGPointMatchAndSetSettings::ApplyDeprecation(UPCGNode* InOutNode)
 }
 #endif // WITH_EDITOR
 
+bool UPCGPointMatchAndSetSettings::UseSeed() const
+{
+	return MatchAndSetInstance && MatchAndSetInstance->UsesRandomProcess();
+}
+
 TArray<FPCGPinProperties> UPCGPointMatchAndSetSettings::InputPinProperties() const
 {
 	// TODO Add param support?
@@ -81,7 +84,6 @@ void UPCGPointMatchAndSetSettings::PostLoad()
 	{
 		const EObjectFlags Flags = GetMaskedFlags(RF_PropagateToSubObjects) | RF_Transactional;
 		MatchAndSetInstance->SetFlags(Flags);
-		bUseSeed = MatchAndSetInstance->UsesRandomProcess();
 	}
 
 #if WITH_EDITOR
@@ -178,8 +180,6 @@ void UPCGPointMatchAndSetSettings::RefreshMatchAndSet()
 	{
 		MatchAndSetInstance = nullptr;
 	}
-
-	bUseSeed = MatchAndSetInstance && MatchAndSetInstance->UsesRandomProcess();
 }
 
 bool FPCGPointMatchAndSetElement::ExecuteInternal(FPCGContext* Context) const
@@ -243,4 +243,4 @@ bool FPCGPointMatchAndSetElement::ExecuteInternal(FPCGContext* Context) const
 	return true;
 }
 
-#undef LOCTEXT_NAMESPACE
\ No newline at end of file
+#undef LOCTEXT_NAMESPACE
diff --git a/Engine/Plugins/PCG/Source/PCG/Private/Elements/PCGRandomChoice.cpp b/Engine/Plugins/PCG/Source/PCG/Private/Elements/PCGRandomChoice.cpp
--- a/Engine/Plugins/PCG/Source/PCG/Private/Elements/PCGRandomChoice.cpp
+++ b/Engine/Plugins/PCG/Source/PCG/Private/Elements/PCGRandomChoice.cpp
@@ -53,11 +53,6 @@ namespace PCGRandomChoice
 	}
 }
 
-UPCGRandomChoiceSettings::UPCGRandomChoiceSettings()
-{
-	bUseSeed = true;
-}
-
 #if WITH_EDITOR
 FName UPCGRandomChoiceSettings::GetDefaultNodeName() const
 {
@@ -231,4 +226,4 @@ bool FPCGRandomChoiceElement::ExecuteInternal(FPCGContext* Context) const
 	return true;
 }
 
-#undef LOCTEXT_NAMESPACE
\ No newline at end of file
+#undef LOCTEXT_NAMESPACE
diff --git a/Engine/Plugins/PCG/Source/PCG/Private/Elements/PCGSelectPoints.cpp b/Engine/Plugins/PCG/Source/PCG/Private/Elements/PCGSelectPoints.cpp
--- a/Engine/Plugins/PCG/Source/PCG/Private/Elements/PCGSelectPoints.cpp
+++ b/Engine/Plugins/PCG/Source/PCG/Private/Elements/PCGSelectPoints.cpp
@@ -14,11 +14,6 @@
 
 #define LOCTEXT_NAMESPACE "PCGSelectPointsElement"
 
-UPCGSelectPointsSettings::UPCGSelectPointsSettings()
-{
-	bUseSeed = true;
-}
-
 #if WITH_EDITOR
 FText UPCGSelectPointsSettings::GetNodeTooltipText() const
 {
diff --git a/Engine/Plugins/PCG/Source/PCG/Private/Elements/PCGSpatialNoise.cpp b/Engine/Plugins/PCG/Source/PCG/Private/Elements/PCGSpatialNoise.cpp
--- a/Engine/Plugins/PCG/Source/PCG/Private/Elements/PCGSpatialNoise.cpp
+++ b/Engine/Plugins/PCG/Source/PCG/Private/Elements/PCGSpatialNoise.cpp
@@ -24,30 +24,9 @@
 
 UPCGSpatialNoiseSettings::UPCGSpatialNoiseSettings()
 {
-	bUseSeed = true;
 	ValueTarget.SetPointProperty(EPCGPointProperties::Density);
 }
 
-void UPCGSpatialNoiseSettings::PostLoad()
-{
-	Super::PostLoad();
-
-	if (bForceNoUseSeed)
-	{
-		bUseSeed = false;
-	}
-}
-
-void UPCGSpatialNoiseSettings::PostEditImport()
-{
-	Super::PostEditImport();
-
-	if (bForceNoUseSeed)
-	{
-		bUseSeed = false;
-	}
-}
-
 #if WITH_EDITOR
 void UPCGSpatialNoiseSettings::ApplyDeprecation(UPCGNode* InOutNode)
 {
@@ -59,13 +38,17 @@ void UPCGSpatialNoiseSettings::ApplyDeprecation(UPCGNode* InOutNode)
 	if (DataVersion < FPCGCustomVersion::NoMoreSpatialDataConversionToPointDataByDefaultOnNonPointPins)
 	{
 		bForceNoUseSeed = true;
-		bUseSeed = false;
 	}
 
 	Super::ApplyDeprecation(InOutNode);
 }
 #endif // WITH_EDITOR
 
+bool UPCGSpatialNoiseSettings::UseSeed() const
+{
+	return !bForceNoUseSeed;
+}
+
 TArray<FPCGPinProperties> UPCGSpatialNoiseSettings::InputPinProperties() const
 {
 	return Super::DefaultPointInputPinProperties();
@@ -806,4 +789,4 @@ bool FPCGSpatialNoise::ExecuteInternal(FPCGContext* Context) const
 	return true;
 }
 
-#undef LOCTEXT_NAMESPACE
\ No newline at end of file
+#undef LOCTEXT_NAMESPACE
diff --git a/Engine/Plugins/PCG/Source/PCG/Private/Elements/PCGSplineSampler.cpp b/Engine/Plugins/PCG/Source/PCG/Private/Elements/PCGSplineSampler.cpp
--- a/Engine/Plugins/PCG/Source/PCG/Private/Elements/PCGSplineSampler.cpp
+++ b/Engine/Plugins/PCG/Source/PCG/Private/Elements/PCGSplineSampler.cpp
@@ -1443,11 +1443,6 @@ namespace PCGSplineSamplerHelpers
 	}
 }
 
-UPCGSplineSamplerSettings::UPCGSplineSamplerSettings()
-{
-	bUseSeed = true;
-}
-
 #if WITH_EDITOR
 FText UPCGSplineSamplerSettings::GetNodeTooltipText() const
 {
diff --git a/Engine/Plugins/PCG/Source/PCG/Private/Elements/PCGStaticMeshSpawner.cpp b/Engine/Plugins/PCG/Source/PCG/Private/Elements/PCGStaticMeshSpawner.cpp
--- a/Engine/Plugins/PCG/Source/PCG/Private/Elements/PCGStaticMeshSpawner.cpp
+++ b/Engine/Plugins/PCG/Source/PCG/Private/Elements/PCGStaticMeshSpawner.cpp
@@ -31,8 +31,6 @@ static TAutoConsoleVariable<bool> CVarAllowISMReuse(
 
 UPCGStaticMeshSpawnerSettings::UPCGStaticMeshSpawnerSettings(const FObjectInitializer &ObjectInitializer)
 {
-	bUseSeed = true;
-
 	MeshSelectorType = UPCGMeshSelectorWeighted::StaticClass();
 	// Implementation note: this should not have been done here (it should have been null), as it causes issues with copy & paste
 	// when the thing to paste does not have that class for its instance.
diff --git a/Engine/Plugins/PCG/Source/PCG/Private/Elements/PCGSurfaceSampler.cpp b/Engine/Plugins/PCG/Source/PCG/Private/Elements/PCGSurfaceSampler.cpp
--- a/Engine/Plugins/PCG/Source/PCG/Private/Elements/PCGSurfaceSampler.cpp
+++ b/Engine/Plugins/PCG/Source/PCG/Private/Elements/PCGSurfaceSampler.cpp
@@ -265,11 +265,6 @@ namespace PCGSurfaceSampler
 #endif
 }
 
-UPCGSurfaceSamplerSettings::UPCGSurfaceSamplerSettings()
-{
-	bUseSeed = true;
-}
-
 #if WITH_EDITOR
 FText UPCGSurfaceSamplerSettings::GetNodeTooltipText() const
 {
diff --git a/Engine/Plugins/PCG/Source/PCG/Private/Elements/PCGTransformPoints.cpp b/Engine/Plugins/PCG/Source/PCG/Private/Elements/PCGTransformPoints.cpp
--- a/Engine/Plugins/PCG/Source/PCG/Private/Elements/PCGTransformPoints.cpp
+++ b/Engine/Plugins/PCG/Source/PCG/Private/Elements/PCGTransformPoints.cpp
@@ -15,11 +15,6 @@
 
 #define LOCTEXT_NAMESPACE "PCGTransformPointsElement"
 
-UPCGTransformPointsSettings::UPCGTransformPointsSettings()
-{
-	bUseSeed = true;
-}
-
 FPCGElementPtr UPCGTransformPointsSettings::CreateElement() const
 {
 	return MakeShared<FPCGTransformPointsElement>();
diff --git a/Engine/Plugins/PCG/Source/PCG/Private/PCGSettings.cpp b/Engine/Plugins/PCG/Source/PCG/Private/PCGSettings.cpp
--- a/Engine/Plugins/PCG/Source/PCG/Private/PCGSettings.cpp
+++ b/Engine/Plugins/PCG/Source/PCG/Private/PCGSettings.cpp
@@ -650,7 +650,7 @@ UPCGNode* UPCGSettings::CreateNode() const
 
 int UPCGSettings::GetSeed(const UPCGComponent* InSourceComponent) const
 {
-	return !bUseSeed ? 42 : (InSourceComponent ? PCGHelpers::ComputeSeed(Seed, InSourceComponent->Seed) : Seed);
+	return !UseSeed() ? 42 : (InSourceComponent ? PCGHelpers::ComputeSeed(Seed, InSourceComponent->Seed) : Seed);
 }
 
 #if WITH_EDITOR
@@ -751,7 +751,15 @@ bool UPCGSettings::CanEditChange(const FEditPropertyChain& InPropertyChain) cons
 
 bool UPCGSettings::CanEditChange(const FProperty* InProperty) const
 {
-	return Super::CanEditChange(InProperty);
+	if (!Super::CanEditChange(InProperty))
+		return false;
+
+	if (InProperty->GetFName() == GET_MEMBER_NAME_CHECKED(UPCGSettings, Seed))
+	{
+		return UseSeed();
+	}
+
+	return true;
 }
 
 void UPCGSettings::PostPaste()
@@ -954,7 +962,7 @@ EPCGDataType UPCGSettings::GetCurrentPinTypes(const UPCGPin* InPin) const
 TArray<FPCGSettingsOverridableParam> UPCGSettings::GatherOverridableParams() const
 {
 	PCGSettingsHelpers::FPCGGetAllOverridableParamsConfig Config;
-	Config.bUseSeed = bUseSeed;
+	Config.bUseSeed = UseSeed();
 	Config.IncludeMetadataValues.Add(PCGObjectMetadata::Overridable);
 	Config.ExcludeMetadataValues.Add(PCGObjectMetadata::NotOverridable);
 
@@ -1212,4 +1220,4 @@ bool FPCGTrivialElement::ExecuteInternal(FPCGContext* Context) const
 	return true;
 }
 
-#undef LOCTEXT_NAMESPACE
\ No newline at end of file
+#undef LOCTEXT_NAMESPACE
diff --git a/Engine/Plugins/PCG/Source/PCG/Public/Compute/Elements/PCGCustomHLSL.h b/Engine/Plugins/PCG/Source/PCG/Public/Compute/Elements/PCGCustomHLSL.h
--- a/Engine/Plugins/PCG/Source/PCG/Public/Compute/Elements/PCGCustomHLSL.h
+++ b/Engine/Plugins/PCG/Source/PCG/Public/Compute/Elements/PCGCustomHLSL.h
@@ -111,6 +111,7 @@ public:
 	virtual TArray<FPCGPinProperties> OutputPinProperties() const override;
 	virtual bool HasOverridableParams() const override { return false; }
 	virtual bool IsInputPinRequiredByExecution(const UPCGPin* InPin) const override { return true; }
+	virtual bool UseSeed() const override { return true; }
 #if WITH_EDITOR
 	virtual FName GetDefaultNodeName() const override { return FName(TEXT("CustomHLSL")); }
 	virtual FText GetDefaultNodeTitle() const override { return NSLOCTEXT("PCGCustomHLSLElement", "NodeTitle", "Custom HLSL"); }
diff --git a/Engine/Plugins/PCG/Source/PCG/Public/Elements/PCGAttributeNoise.h b/Engine/Plugins/PCG/Source/PCG/Public/Elements/PCGAttributeNoise.h
--- a/Engine/Plugins/PCG/Source/PCG/Public/Elements/PCGAttributeNoise.h
+++ b/Engine/Plugins/PCG/Source/PCG/Public/Elements/PCGAttributeNoise.h
@@ -51,6 +51,7 @@ public:
 	virtual bool OnlyExposePreconfiguredSettings() const override { return true; }
 	virtual bool GroupPreconfiguredSettings() const override { return false; }
 #endif
+	virtual bool UseSeed() const override { return true; }
 	virtual void ApplyPreconfiguredSettings(const FPCGPreConfiguredSettingsInfo& PreconfigureInfo) override;
 	virtual EPCGDataType GetCurrentPinTypes(const UPCGPin* InPin) const override;
 	virtual bool HasDynamicPins() const override { return true; }
diff --git a/Engine/Plugins/PCG/Source/PCG/Public/Elements/PCGClusterElement.h b/Engine/Plugins/PCG/Source/PCG/Public/Elements/PCGClusterElement.h
--- a/Engine/Plugins/PCG/Source/PCG/Public/Elements/PCGClusterElement.h
+++ b/Engine/Plugins/PCG/Source/PCG/Public/Elements/PCGClusterElement.h
@@ -38,10 +38,8 @@ public:
 	virtual FText GetDefaultNodeTitle() const override { return NSLOCTEXT("PCGClusterElement", "NodeTitle", "Cluster"); }
 	virtual EPCGSettingsType GetType() const override { return EPCGSettingsType::PointOps; }
 	virtual FText GetNodeTooltipText() const override;
-	virtual void PostEditChangeProperty(FPropertyChangedEvent& PropertyChangedEvent) override;
-	virtual void PostDuplicate(bool bDuplicateForPIE) override;
 #endif
-	virtual void PostLoad() override;
+	virtual bool UseSeed() const override;
 
 protected:
 	virtual TArray<FPCGPinProperties> InputPinProperties() const override;
diff --git a/Engine/Plugins/PCG/Source/PCG/Public/Elements/PCGCreatePointsSphere.h b/Engine/Plugins/PCG/Source/PCG/Public/Elements/PCGCreatePointsSphere.h
--- a/Engine/Plugins/PCG/Source/PCG/Public/Elements/PCGCreatePointsSphere.h
+++ b/Engine/Plugins/PCG/Source/PCG/Public/Elements/PCGCreatePointsSphere.h
@@ -41,8 +41,6 @@ class UPCGCreatePointsSphereSettings : public UPCGSettings
 	GENERATED_BODY()
 
 public:
-	UPCGCreatePointsSphereSettings() { bUseSeed = true; }
-
 	//~Begin UPCGSettings interface
 #if WITH_EDITOR
 	virtual FName GetDefaultNodeName() const override { return FName(TEXT("CreatePointsSphere")); }
@@ -50,6 +48,7 @@ public:
 	virtual FText GetNodeTooltipText() const override { return NSLOCTEXT("PCGCreatePointsSphereElement", "NodeTooltip", "Generate points on the surface of a sphere."); }
 	virtual EPCGSettingsType GetType() const override { return EPCGSettingsType::Spatial; }
 #endif // WITH_EDITOR
+	virtual bool UseSeed() const override { return true; }
 
 protected:
 	virtual TArray<FPCGPinProperties> InputPinProperties() const override { return {}; }
diff --git a/Engine/Plugins/PCG/Source/PCG/Public/Elements/PCGExecuteBlueprint.h b/Engine/Plugins/PCG/Source/PCG/Public/Elements/PCGExecuteBlueprint.h
--- a/Engine/Plugins/PCG/Source/PCG/Public/Elements/PCGExecuteBlueprint.h
+++ b/Engine/Plugins/PCG/Source/PCG/Public/Elements/PCGExecuteBlueprint.h
@@ -319,6 +319,7 @@ public:
 	virtual TArray<FPCGPreConfiguredSettingsInfo> GetPreconfiguredInfo() const override;
 	virtual bool OnlyExposePreconfiguredSettings() const override;
 #endif
+	virtual bool UseSeed() const override { return true; }
 	virtual bool HasDynamicPins() const override;
 	virtual EPCGDataType GetCurrentPinTypes(const UPCGPin* InPin) const override;
 
diff --git a/Engine/Plugins/PCG/Source/PCG/Public/Elements/PCGMatchAndSetAttributes.h b/Engine/Plugins/PCG/Source/PCG/Public/Elements/PCGMatchAndSetAttributes.h
--- a/Engine/Plugins/PCG/Source/PCG/Public/Elements/PCGMatchAndSetAttributes.h
+++ b/Engine/Plugins/PCG/Source/PCG/Public/Elements/PCGMatchAndSetAttributes.h
@@ -45,8 +45,6 @@ class UPCGMatchAndSetAttributesSettings : public UPCGSettings
 	GENERATED_BODY()
 
 public:
-	UPCGMatchAndSetAttributesSettings();
-
 	// ~Begin UPCGSettings interface
 #if WITH_EDITOR
 	virtual FName GetDefaultNodeName() const override;
@@ -56,6 +54,7 @@ public:
 	virtual bool HasDynamicPins() const { return true; }
 	virtual void ApplyDeprecationBeforeUpdatePins(UPCGNode* InOutNode, TArray<TObjectPtr<UPCGPin>>& InputPins, TArray<TObjectPtr<UPCGPin>>& OutputPins) override;
 #endif // WITH_EDITOR
+	virtual bool UseSeed() const override;
 
 protected:
 	virtual TArray<FPCGPinProperties> InputPinProperties() const override;
diff --git a/Engine/Plugins/PCG/Source/PCG/Public/Elements/PCGMutateSeed.h b/Engine/Plugins/PCG/Source/PCG/Public/Elements/PCGMutateSeed.h
--- a/Engine/Plugins/PCG/Source/PCG/Public/Elements/PCGMutateSeed.h
+++ b/Engine/Plugins/PCG/Source/PCG/Public/Elements/PCGMutateSeed.h
@@ -16,8 +16,6 @@ class UPCGMutateSeedSettings : public UPCGSettings
 	GENERATED_BODY()
 
 public:
-	UPCGMutateSeedSettings();
-
 	//~Begin UPCGSettings interface
 #if WITH_EDITOR
 	virtual FName GetDefaultNodeName() const override { return FName(TEXT("MutateSeed")); }
@@ -25,7 +23,7 @@ public:
 	virtual FText GetNodeTooltipText() const override { return NSLOCTEXT("PCGMutateSeedElement", "NodeTooltip", "Applies a new random seed from point input."); }
 	virtual EPCGSettingsType GetType() const override { return EPCGSettingsType::Spatial; }
 #endif
-	
+	virtual bool UseSeed() const override { return true; }
 
 protected:
 	virtual TArray<FPCGPinProperties> InputPinProperties() const override { return Super::DefaultPointInputPinProperties(); }
@@ -38,4 +36,4 @@ class FPCGMutateSeedElement : public FPCGPointOperationElementBase
 {
 protected:
 	virtual bool ExecuteInternal(FPCGContext* Context) const override;
-};
\ No newline at end of file
+};
diff --git a/Engine/Plugins/PCG/Source/PCG/Public/Elements/PCGPointMatchAndSet.h b/Engine/Plugins/PCG/Source/PCG/Public/Elements/PCGPointMatchAndSet.h
--- a/Engine/Plugins/PCG/Source/PCG/Public/Elements/PCGPointMatchAndSet.h
+++ b/Engine/Plugins/PCG/Source/PCG/Public/Elements/PCGPointMatchAndSet.h
@@ -30,7 +30,7 @@ public:
 	virtual EPCGSettingsType GetType() const override { return EPCGSettingsType::Metadata; }
 	virtual void ApplyDeprecation(UPCGNode* InOutNode) override;
 #endif // WITH_EDITOR
-	
+	virtual bool UseSeed() const override;
 
 protected:
 	virtual TArray<FPCGPinProperties> InputPinProperties() const override;
diff --git a/Engine/Plugins/PCG/Source/PCG/Public/Elements/PCGRandomChoice.h b/Engine/Plugins/PCG/Source/PCG/Public/Elements/PCGRandomChoice.h
--- a/Engine/Plugins/PCG/Source/PCG/Public/Elements/PCGRandomChoice.h
+++ b/Engine/Plugins/PCG/Source/PCG/Public/Elements/PCGRandomChoice.h
@@ -22,8 +22,6 @@ class UPCGRandomChoiceSettings : public UPCGSettings
 	GENERATED_BODY()
 
 public:
-	UPCGRandomChoiceSettings();
-
 	//~Begin UPCGSettings interface
 #if WITH_EDITOR
 	virtual FName GetDefaultNodeName() const override;
@@ -31,7 +29,7 @@ public:
 	virtual FText GetNodeTooltipText() const override;
 	virtual EPCGSettingsType GetType() const override { return EPCGSettingsType::Filter; }
 #endif
-
+	virtual bool UseSeed() const override { return true; }
 	virtual bool HasDynamicPins() const override { return true; }
 
 protected:
diff --git a/Engine/Plugins/PCG/Source/PCG/Public/Elements/PCGSelectPoints.h b/Engine/Plugins/PCG/Source/PCG/Public/Elements/PCGSelectPoints.h
--- a/Engine/Plugins/PCG/Source/PCG/Public/Elements/PCGSelectPoints.h
+++ b/Engine/Plugins/PCG/Source/PCG/Public/Elements/PCGSelectPoints.h
@@ -12,8 +12,6 @@ class UPCGSelectPointsSettings : public UPCGSettings
 	GENERATED_BODY()
 
 public:
-	UPCGSelectPointsSettings();
-
 	//~Begin UPCGSettings interface
 #if WITH_EDITOR
 	virtual FName GetDefaultNodeName() const override { return FName(TEXT("SelectPoints")); }
@@ -21,7 +19,7 @@ public:
 	virtual FText GetNodeTooltipText() const override;
 	virtual EPCGSettingsType GetType() const override { return EPCGSettingsType::Sampler; }
 #endif
-	
+	virtual bool UseSeed() const override { return true; }
 
 protected:
 	virtual TArray<FPCGPinProperties> InputPinProperties() const override { return Super::DefaultPointInputPinProperties(); }
diff --git a/Engine/Plugins/PCG/Source/PCG/Public/Elements/PCGSpatialNoise.h b/Engine/Plugins/PCG/Source/PCG/Public/Elements/PCGSpatialNoise.h
--- a/Engine/Plugins/PCG/Source/PCG/Public/Elements/PCGSpatialNoise.h
+++ b/Engine/Plugins/PCG/Source/PCG/Public/Elements/PCGSpatialNoise.h
@@ -62,10 +62,6 @@ class UPCGSpatialNoiseSettings : public UPCGSettings
 
 public:
 	UPCGSpatialNoiseSettings();
-	//~Begin UObject interface
-	virtual void PostLoad() override;
-	virtual void PostEditImport() override;
-	//~End UObject interface
 
 	//~Begin UPCGSettings interface
 #if WITH_EDITOR
@@ -74,6 +70,7 @@ public:
 	virtual EPCGSettingsType GetType() const override { return EPCGSettingsType::Spatial; }
 	virtual void ApplyDeprecation(UPCGNode* InOutNode) override;
 #endif
+	virtual bool UseSeed() const override;
 
 protected:
 	virtual TArray<FPCGPinProperties> InputPinProperties() const override;
@@ -83,7 +80,6 @@ protected:
 	//~End UPCGSettings interface
 
 public:
-
 	// The noise method used
 	UPROPERTY(BlueprintReadWrite, EditAnywhere, Category = Settings, meta = (PCG_Overridable))
 	PCGSpatialNoiseMode Mode = PCGSpatialNoiseMode::Perlin2D;
diff --git a/Engine/Plugins/PCG/Source/PCG/Public/Elements/PCGSplineSampler.h b/Engine/Plugins/PCG/Source/PCG/Public/Elements/PCGSplineSampler.h
--- a/Engine/Plugins/PCG/Source/PCG/Public/Elements/PCGSplineSampler.h
+++ b/Engine/Plugins/PCG/Source/PCG/Public/Elements/PCGSplineSampler.h
@@ -254,8 +254,6 @@ class UPCGSplineSamplerSettings : public UPCGSettings
 	GENERATED_BODY()
 
 public:
-	UPCGSplineSamplerSettings();
-
 	// ~Begin UPCGSettings interface
 #if WITH_EDITOR
 	virtual FName GetDefaultNodeName() const override { return FName(TEXT("SplineSampler")); }
@@ -265,6 +263,8 @@ public:
 	virtual void ApplyDeprecation(UPCGNode* InOutNode) override;
 	virtual void ApplyDeprecationBeforeUpdatePins(UPCGNode* InOutNode, TArray<TObjectPtr<UPCGPin>>& InputPins, TArray<TObjectPtr<UPCGPin>>& OutputPins) override;
 #endif
+	virtual bool UseSeed() const override { return true; }
+
 
 protected:
 	virtual TArray<FPCGPinProperties> InputPinProperties() const override;
diff --git a/Engine/Plugins/PCG/Source/PCG/Public/Elements/PCGStaticMeshSpawner.h b/Engine/Plugins/PCG/Source/PCG/Public/Elements/PCGStaticMeshSpawner.h
--- a/Engine/Plugins/PCG/Source/PCG/Public/Elements/PCGStaticMeshSpawner.h
+++ b/Engine/Plugins/PCG/Source/PCG/Public/Elements/PCGStaticMeshSpawner.h
@@ -35,7 +35,7 @@ public:
 	virtual EPCGSettingsType GetType() const override { return EPCGSettingsType::Spawner; }
 	virtual void ApplyDeprecation(UPCGNode* InOutNode) override;
 #endif
-	
+	virtual bool UseSeed() const override { return true; }
 
 protected:
 	virtual TArray<FPCGPinProperties> InputPinProperties() const override { return Super::DefaultPointInputPinProperties(); }
diff --git a/Engine/Plugins/PCG/Source/PCG/Public/Elements/PCGSurfaceSampler.h b/Engine/Plugins/PCG/Source/PCG/Public/Elements/PCGSurfaceSampler.h
--- a/Engine/Plugins/PCG/Source/PCG/Public/Elements/PCGSurfaceSampler.h
+++ b/Engine/Plugins/PCG/Source/PCG/Public/Elements/PCGSurfaceSampler.h
@@ -94,8 +94,6 @@ class UPCGSurfaceSamplerSettings : public UPCGSettings
 	GENERATED_BODY()
 
 public:
-	UPCGSurfaceSamplerSettings();
-
 	//~Begin UObject interface
 	virtual void PostLoad() override;
 	//~End UObject interface
@@ -110,6 +108,7 @@ public:
 	virtual void ApplyDeprecationBeforeUpdatePins(UPCGNode* InOutNode, TArray<TObjectPtr<UPCGPin>>& InputPins, TArray<TObjectPtr<UPCGPin>>& OutputPins) override;
 	virtual void ApplyDeprecation(UPCGNode* InOutNode) override;
 #endif
+	virtual bool UseSeed() const override { return true; }
 
 protected:
 	virtual TArray<FPCGPinProperties> InputPinProperties() const override;
diff --git a/Engine/Plugins/PCG/Source/PCG/Public/Elements/PCGTransformPoints.h b/Engine/Plugins/PCG/Source/PCG/Public/Elements/PCGTransformPoints.h
--- a/Engine/Plugins/PCG/Source/PCG/Public/Elements/PCGTransformPoints.h
+++ b/Engine/Plugins/PCG/Source/PCG/Public/Elements/PCGTransformPoints.h
@@ -12,15 +12,13 @@ class UPCGTransformPointsSettings : public UPCGSettings
 	GENERATED_BODY()
 
 public:
-	UPCGTransformPointsSettings();
-
 	//~Begin UPCGSettings interface
 #if WITH_EDITOR
 	virtual FName GetDefaultNodeName() const override { return FName(TEXT("TransformPoints")); }
 	virtual FText GetDefaultNodeTitle() const override { return NSLOCTEXT("PCGTransformPointsSettings", "NodeTitle", "Transform Points"); }
 	virtual EPCGSettingsType GetType() const override { return EPCGSettingsType::PointOps; }
 #endif
-	
+	virtual bool UseSeed() const override { return true; }
 
 protected:
 	virtual TArray<FPCGPinProperties> InputPinProperties() const override { return Super::DefaultPointInputPinProperties(); }
diff --git a/Engine/Plugins/PCG/Source/PCG/Public/PCGSettings.h b/Engine/Plugins/PCG/Source/PCG/Public/PCGSettings.h
--- a/Engine/Plugins/PCG/Source/PCG/Public/PCGSettings.h
+++ b/Engine/Plugins/PCG/Source/PCG/Public/PCGSettings.h
@@ -257,13 +257,17 @@ public:
 	EPCGDataType GetTypeUnionOfIncidentEdges(const FName& PinLabel) const;
 
 	// Internal functions, should not be used by any user.
-	// Return a different subset for for input/output pin properties, in case of a default object.
+	// Return a different subset for input/output pin properties, in case of a default object.
 	virtual TArray<FPCGPinProperties> DefaultInputPinProperties() const;
 	virtual TArray<FPCGPinProperties> DefaultOutputPinProperties() const;
 
 	bool operator==(const UPCGSettings& Other) const;
-	
-	bool UseSeed() const { return bUseSeed; }
+
+	PRAGMA_DISABLE_DEPRECATION_WARNINGS
+	/** By default, settings do not use a seed. Override this in the settings subclass to enable usage of the seed. UFUNCTION to be used by EditCondition. */
+	UFUNCTION()
+	virtual bool UseSeed() const { return bUseSeed; }
+	PRAGMA_ENABLE_DEPRECATION_WARNINGS
 
 	// Get the seed, combined with optional PCGComponent seed
 	int GetSeed(const UPCGComponent* InSourceComponent = nullptr) const;
@@ -375,7 +379,7 @@ public:
 	*/
 	virtual EPCGDataType GetCurrentPinTypes(const UPCGPin* InPin) const;
 
-	UPROPERTY(EditAnywhere, BlueprintReadWrite, Category = Settings, meta=(EditCondition=bUseSeed, EditConditionHides, PCG_Overridable))
+	UPROPERTY(EditAnywhere, BlueprintReadWrite, Category = Settings, meta=(EditCondition="UseSeed()", EditConditionHides, PCG_Overridable))
 	int Seed = 0xC35A9631; // Default seed is a random prime number, but will be overriden for new settings based on the class type name hash, making each settings class have a different default seed.
 
 #if WITH_EDITORONLY_DATA
@@ -446,8 +450,8 @@ protected:
 	virtual bool CanEditChange(const FProperty* InProperty) const override;
 #endif
 
-	// By default, settings won't use a seed. Set this bool to true in the child ctor to allow edition and use it.
-	UPROPERTY(VisibleAnywhere, Transient, Category = Settings, meta = (EditCondition = false, EditConditionHides))
+	UE_DEPRECATED(5.5, "Implement the PCGSettings virtual UseSeed() override.")
+	UPROPERTY(Transient, meta = (DeprecatedProperty, DeprecationMessage = "Implement the PCGSettings virtual UseSeed() override."))
 	bool bUseSeed = false;
 
 	/** Methods to remove boilerplate code across settings */
@@ -475,7 +479,7 @@ private:
 
 	// Overridable param section
 public:
-	/** List of all the overridable params available for this settings. */
+	/** List of all the overridable params available for these settings. */
 	virtual const TArray<FPCGSettingsOverridableParam>& OverridableParams() const { return CachedOverridableParams; }
 
 	/** Check if we have some override. Can be overriden to force params pin for example */
@@ -629,4 +633,4 @@ class FPCGTrivialElement : public IPCGElement
 protected:
 	virtual bool ExecuteInternal(FPCGContext* Context) const override;
 	virtual bool IsPassthrough(const UPCGSettings* InSettings) const override { return true; }
-};
\ No newline at end of file
+};
diff --git a/Engine/Plugins/PCGInterops/PCGGeometryScriptInterop/Source/PCGGeometryScriptInterop/Private/Elements/PCGMeshSampler.cpp b/Engine/Plugins/PCGInterops/PCGGeometryScriptInterop/Source/PCGGeometryScriptInterop/Private/Elements/PCGMeshSampler.cpp
--- a/Engine/Plugins/PCGInterops/PCGGeometryScriptInterop/Source/PCGGeometryScriptInterop/Private/Elements/PCGMeshSampler.cpp
+++ b/Engine/Plugins/PCGInterops/PCGGeometryScriptInterop/Source/PCGGeometryScriptInterop/Private/Elements/PCGMeshSampler.cpp
@@ -291,12 +291,6 @@ void FPCGMeshSamplerContext::SetPointColorAndDensity(SetPointDensityFunc SetPoin
 	}
 }
 
-UPCGMeshSamplerSettings::UPCGMeshSamplerSettings()
-	: UPCGSettings()
-{
-	bUseSeed = true;
-}
-
 void UPCGMeshSamplerSettings::PostLoad()
 {
 	Super::PostLoad();
diff --git a/Engine/Plugins/PCGInterops/PCGGeometryScriptInterop/Source/PCGGeometryScriptInterop/Public/Elements/PCGMeshSampler.h b/Engine/Plugins/PCGInterops/PCGGeometryScriptInterop/Source/PCGGeometryScriptInterop/Public/Elements/PCGMeshSampler.h
--- a/Engine/Plugins/PCGInterops/PCGGeometryScriptInterop/Source/PCGGeometryScriptInterop/Public/Elements/PCGMeshSampler.h
+++ b/Engine/Plugins/PCGInterops/PCGGeometryScriptInterop/Source/PCGGeometryScriptInterop/Public/Elements/PCGMeshSampler.h
@@ -57,8 +57,6 @@ class PCGGEOMETRYSCRIPTINTEROP_API UPCGMeshSamplerSettings : public UPCGSettings
 	GENERATED_BODY()
 
 public:
-	UPCGMeshSamplerSettings();
-
 	//~Begin UObject interface
 	virtual void PostLoad() override;
 	//~End UObject interface
@@ -70,7 +68,7 @@ public:
 	virtual FText GetNodeTooltipText() const override;
 	virtual EPCGSettingsType GetType() const override { return EPCGSettingsType::Sampler; }
 #endif
-
+	virtual bool UseSeed() const override { return true; }
 	virtual TArray<FPCGPinProperties> InputPinProperties() const override;
 	virtual TArray<FPCGPinProperties> OutputPinProperties() const override;
 
@@ -239,4 +237,4 @@ protected:
 	virtual FPCGContext* CreateContext() override;
 	virtual bool PrepareDataInternal(FPCGContext* InContext) const override;
 	virtual bool ExecuteInternal(FPCGContext* InContext) const override;
-};
\ No newline at end of file
+};
