[PCG] Add appropriate suffix to PCGLogErrors.h and sort by severity

NOTE: Cosmetic only change, but important since we're still establishing the format of the relatively new workflow, and it came up as an issue in a review.

#rnx

diff --git a/Engine/Plugins/PCG/Source/PCG/Private/Elements/Grammar/PCGSegmentSlicer.cpp b/Engine/Plugins/PCG/Source/PCG/Private/Elements/Grammar/PCGSegmentSlicer.cpp
--- a/Engine/Plugins/PCG/Source/PCG/Private/Elements/Grammar/PCGSegmentSlicer.cpp
+++ b/Engine/Plugins/PCG/Source/PCG/Private/Elements/Grammar/PCGSegmentSlicer.cpp
@@ -293,7 +293,7 @@ bool FPCGSegmentSlicerElement::ExecuteInternal(FPCGContext* InContext) const
 				OutAttribute = OutputMetadata->FindOrCreateAttribute<T>(AttributeName, DefaultValue, false, true);
 				if (!OutAttribute)
 				{
-					PCGLog::Metadata::LogFailToCreateAttribute<T>(AttributeName, InContext);
+					PCGLog::Metadata::LogFailToCreateAttributeError<T>(AttributeName, InContext);
 					return false;
 				}
 			}
diff --git a/Engine/Plugins/PCG/Source/PCG/Private/Elements/Grammar/PCGSelectGrammar.cpp b/Engine/Plugins/PCG/Source/PCG/Private/Elements/Grammar/PCGSelectGrammar.cpp
--- a/Engine/Plugins/PCG/Source/PCG/Private/Elements/Grammar/PCGSelectGrammar.cpp
+++ b/Engine/Plugins/PCG/Source/PCG/Private/Elements/Grammar/PCGSelectGrammar.cpp
@@ -253,14 +253,14 @@ bool FPCGSelectGrammarElement::ExecuteInternal(FPCGContext* InContext) const
 		TUniquePtr<const IPCGAttributeAccessorKeys> InputKeys = PCGAttributeAccessorHelpers::CreateConstKeys(InputPointData, KeySelector);
 		if (!InputKeyAccessor || !InputKeys)
 		{
-			PCGLog::Metadata::LogFailToCreateAccessor(KeySelector);
+			PCGLog::Metadata::LogFailToCreateAccessorError(KeySelector);
 			continue;
 		}
 
 		TUniquePtr<const IPCGAttributeAccessor> ComparedValuesAccessor = PCGAttributeAccessorHelpers::CreateConstAccessor(InputPointData, ComparedValueSelector);
 		if (!ComparedValuesAccessor)
 		{
-			PCGLog::Metadata::LogFailToCreateAccessor(ComparedValueSelector);
+			PCGLog::Metadata::LogFailToCreateAccessorError(ComparedValueSelector);
 			continue;
 		}
 
@@ -271,7 +271,7 @@ bool FPCGSelectGrammarElement::ExecuteInternal(FPCGContext* InContext) const
 		TUniquePtr<IPCGAttributeAccessorKeys> OutputKeys = PCGAttributeAccessorHelpers::CreateKeys(OutputPointData, OutputGrammarSelector);
 		if (!WriteGrammarAccessor || !OutputKeys)
 		{
-			PCGLog::Metadata::LogFailToCreateAccessor(OutputGrammarSelector);
+			PCGLog::Metadata::LogFailToCreateAccessorError(OutputGrammarSelector);
 			continue;
 		}
 
diff --git a/Engine/Plugins/PCG/Source/PCG/Private/Elements/Grammar/PCGSlicingBase.cpp b/Engine/Plugins/PCG/Source/PCG/Private/Elements/Grammar/PCGSlicingBase.cpp
--- a/Engine/Plugins/PCG/Source/PCG/Private/Elements/Grammar/PCGSlicingBase.cpp
+++ b/Engine/Plugins/PCG/Source/PCG/Private/Elements/Grammar/PCGSlicingBase.cpp
@@ -128,13 +128,13 @@ PCGGrammar::FTokenizedGrammar FPCGSlicingBaseElement::GetTokenizedGrammar(FPCGCo
 		const TUniquePtr<const IPCGAttributeAccessor> Accessor = PCGAttributeAccessorHelpers::CreateConstAccessor(InputData, Selector);
 		if (!Accessor)
 		{
-			PCGLog::Metadata::LogFailToCreateAccessor(Selector, InContext);
+			PCGLog::Metadata::LogFailToCreateAccessorError(Selector, InContext);
 			return {};
 		}
 
 		if (!Accessor->Get(Grammar, FPCGAttributeAccessorKeysEntries(PCGInvalidEntryKey), EPCGAttributeAccessorFlags::AllowBroadcastAndConstructible))
 		{
-			PCGLog::Metadata::LogFailToGetAttribute<FString>(Selector, Accessor.Get(), InContext);
+			PCGLog::Metadata::LogFailToGetAttributeError<FString>(Selector, Accessor.Get(), InContext);
 			return {};
 		}
 	}
@@ -215,7 +215,7 @@ TMap<FString, PCGGrammar::FTokenizedGrammar> FPCGSlicingBaseElement::GetTokenize
 		const TUniquePtr<const IPCGAttributeAccessorKeys> Keys = PCGAttributeAccessorHelpers::CreateConstKeys(InputData, Selector);
 		if (!Accessor || !Keys)
 		{
-			PCGLog::Metadata::LogFailToCreateAccessor(Selector, InContext);
+			PCGLog::Metadata::LogFailToCreateAccessorError(Selector, InContext);
 			return Result;
 		}
 
@@ -229,7 +229,7 @@ TMap<FString, PCGGrammar::FTokenizedGrammar> FPCGSlicingBaseElement::GetTokenize
 
 		if (!bSuccess)
 		{
-			PCGLog::Metadata::LogFailToGetAttribute<FString>(Selector, Accessor.Get(), InContext);
+			PCGLog::Metadata::LogFailToGetAttributeError<FString>(Selector, Accessor.Get(), InContext);
 			return Result;
 		}
 	}
diff --git a/Engine/Plugins/PCG/Source/PCG/Private/Elements/Grammar/PCGSplineSlicer.cpp b/Engine/Plugins/PCG/Source/PCG/Private/Elements/Grammar/PCGSplineSlicer.cpp
--- a/Engine/Plugins/PCG/Source/PCG/Private/Elements/Grammar/PCGSplineSlicer.cpp
+++ b/Engine/Plugins/PCG/Source/PCG/Private/Elements/Grammar/PCGSplineSlicer.cpp
@@ -270,13 +270,13 @@ bool FPCGSplineSlicerElement::ExecuteInternal(FPCGContext* InContext) const
 
 			if (!HeightAccessor)
 			{
-				PCGLog::Metadata::LogFailToCreateAccessor(Selector, InContext);
+				PCGLog::Metadata::LogFailToCreateAccessorError(Selector, InContext);
 				continue;
 			}
 
 			if (!HeightAccessor->Get<double>(ModuleHeight, FPCGAttributeAccessorKeysEntries(PCGInvalidEntryKey), EPCGAttributeAccessorFlags::AllowBroadcastAndConstructible))
 			{
-				PCGLog::Metadata::LogFailToGetAttribute(Selector, InContext);
+				PCGLog::Metadata::LogFailToGetAttributeError(Selector, InContext);
 				continue;
 			}
 		}
@@ -299,7 +299,7 @@ bool FPCGSplineSlicerElement::ExecuteInternal(FPCGContext* InContext) const
 				OutAttribute = OutputMetadata->FindOrCreateAttribute<T>(AttributeName, DefaultValue, false, true);
 				if (!OutAttribute)
 				{
-					PCGLog::Metadata::LogFailToCreateAttribute<T>(AttributeName, InContext);
+					PCGLog::Metadata::LogFailToCreateAttributeError<T>(AttributeName, InContext);
 					return false;
 				}
 			}
@@ -323,7 +323,7 @@ bool FPCGSplineSlicerElement::ExecuteInternal(FPCGContext* InContext) const
 			const FPCGAttributeAccessorKeysEntries Keys(PCGInvalidEntryKey);
 			if (!GrammarAccessor)
 			{
-				PCGLog::Metadata::LogFailToCreateAccessor(Selector, InContext);
+				PCGLog::Metadata::LogFailToCreateAccessorError(Selector, InContext);
 				continue;
 			}
 
diff --git a/Engine/Plugins/PCG/Source/PCG/Private/Elements/Grammar/PCGVolumeSlicer.cpp b/Engine/Plugins/PCG/Source/PCG/Private/Elements/Grammar/PCGVolumeSlicer.cpp
--- a/Engine/Plugins/PCG/Source/PCG/Private/Elements/Grammar/PCGVolumeSlicer.cpp
+++ b/Engine/Plugins/PCG/Source/PCG/Private/Elements/Grammar/PCGVolumeSlicer.cpp
@@ -93,13 +93,13 @@ bool FPCGVolumeSlicerElement::ExecuteInternal(FPCGContext* InContext) const
 			const TUniquePtr<const IPCGAttributeAccessor> Accessor = PCGAttributeAccessorHelpers::CreateConstAccessor(InputSplineData, Selector);
 			if (!Accessor)
 			{
-				PCGLog::Metadata::LogFailToCreateAccessor(Selector, InContext);
+				PCGLog::Metadata::LogFailToCreateAccessorError(Selector, InContext);
 				return false;
 			}
 
 			if (!Accessor->Get(OutValue, FPCGAttributeAccessorKeysEntries(PCGInvalidEntryKey), EPCGAttributeAccessorFlags::AllowBroadcastAndConstructible))
 			{
-				PCGLog::Metadata::LogFailToGetAttribute<T>(Selector, Accessor.Get(), InContext);
+				PCGLog::Metadata::LogFailToGetAttributeError<T>(Selector, Accessor.Get(), InContext);
 				return false;
 			}
 
diff --git a/Engine/Plugins/PCG/Source/PCG/Private/Elements/Metadata/PCGMetadataPartition.cpp b/Engine/Plugins/PCG/Source/PCG/Private/Elements/Metadata/PCGMetadataPartition.cpp
--- a/Engine/Plugins/PCG/Source/PCG/Private/Elements/Metadata/PCGMetadataPartition.cpp
+++ b/Engine/Plugins/PCG/Source/PCG/Private/Elements/Metadata/PCGMetadataPartition.cpp
@@ -185,7 +185,7 @@ bool FPCGMetadataPartitionElement::ExecuteInternal(FPCGContext* Context) const
 			FPCGMetadataAttributeBase* AttributePartitionIndex = Metadata->CreateAttribute<int32>(Settings->PartitionIndexAttributeName, 0, /*bAllowInterpolation=*/false, /*bOverrideParent=*/true);
 			if (!AttributePartitionIndex)
 			{
-				PCGLog::Metadata::LogFailToCreateAttribute<int32>(Settings->PartitionIndexAttributeName, Context);
+				PCGLog::Metadata::LogFailToCreateAttributeError<int32>(Settings->PartitionIndexAttributeName, Context);
 				continue;
 			}
 
diff --git a/Engine/Plugins/PCG/Source/PCG/Private/Elements/PCGClusterElement.cpp b/Engine/Plugins/PCG/Source/PCG/Private/Elements/PCGClusterElement.cpp
--- a/Engine/Plugins/PCG/Source/PCG/Private/Elements/PCGClusterElement.cpp
+++ b/Engine/Plugins/PCG/Source/PCG/Private/Elements/PCGClusterElement.cpp
@@ -494,7 +494,7 @@ bool FPCGClusterElement::PrepareDataInternal(FPCGContext* InContext) const
 
 		if (!OutPointData->Metadata->FindOrCreateAttribute<int32>(Settings->ClusterAttribute, 0, /*bAllowsInterpolation=*/false))
 		{
-			PCGLog::Metadata::LogFailToCreateAttribute<int32>(Settings->ClusterAttribute, Context);
+			PCGLog::Metadata::LogFailToCreateAttributeError<int32>(Settings->ClusterAttribute, Context);
 			return EPCGTimeSliceInitResult::AbortExecution;
 		}
 
@@ -506,7 +506,7 @@ bool FPCGClusterElement::PrepareDataInternal(FPCGContext* InContext) const
 		OutState.Keys = PCGAttributeAccessorHelpers::CreateKeys(OutPointData, Selector);
 		if (!OutState.Accessor || !OutState.Keys)
 		{
-			PCGLog::Metadata::LogFailToCreateAccessor(Selector, Context);
+			PCGLog::Metadata::LogFailToCreateAccessorError(Selector, Context);
 			return EPCGTimeSliceInitResult::AbortExecution;
 		}
 
diff --git a/Engine/Plugins/PCG/Source/PCG/Private/InstanceDataPackers/PCGInstanceDataPackerByAttribute.cpp b/Engine/Plugins/PCG/Source/PCG/Private/InstanceDataPackers/PCGInstanceDataPackerByAttribute.cpp
--- a/Engine/Plugins/PCG/Source/PCG/Private/InstanceDataPackers/PCGInstanceDataPackerByAttribute.cpp
+++ b/Engine/Plugins/PCG/Source/PCG/Private/InstanceDataPackers/PCGInstanceDataPackerByAttribute.cpp
@@ -52,7 +52,7 @@ void UPCGInstanceDataPackerByAttribute::PackInstances_Implementation(FPCGContext
 
 		if (!Accessor.IsValid() || !Keys.IsValid())
 		{
-			PCGLog::Metadata::LogFailToCreateAccessor(Selector, &Context);
+			PCGLog::Metadata::LogFailToCreateAccessorError(Selector, &Context);
 			continue;
 		}
 
diff --git a/Engine/Plugins/PCG/Source/PCG/Private/Utils/PCGLogErrors.cpp b/Engine/Plugins/PCG/Source/PCG/Private/Utils/PCGLogErrors.cpp
--- a/Engine/Plugins/PCG/Source/PCG/Private/Utils/PCGLogErrors.cpp
+++ b/Engine/Plugins/PCG/Source/PCG/Private/Utils/PCGLogErrors.cpp
@@ -6,60 +6,60 @@
 
 namespace PCGLog::InputOutput
 {
-	namespace ErrorFormat
+	namespace Format
 	{
-		const FTextFormat TypedInputNotFoundWarning = LOCTEXT("DataInputNotFound", "Data of type {0} not found on pin '{1}'.");
-		const FTextFormat FirstInputOnlyWarning = LOCTEXT("FirstInputOnly", "Multiple inputs found on single-input pin '{0}'. Only the first will be selected.");
-		const FText InvalidInputDataError = LOCTEXT("InvalidInputData", "Invalid input data.");
+		const FTextFormat TypedInputNotFound = LOCTEXT("TypedInputNotFound", "Data of type {0} not found on pin '{1}'.");
+		const FTextFormat FirstInputOnly = LOCTEXT("FirstInputOnly", "Multiple inputs found on single-input pin '{0}'. Only the first will be selected.");
+		const FText InvalidInputData = LOCTEXT("InvalidInputData", "Invalid input data.");
 	}
 
 	void LogTypedDataNotFoundWarning(EPCGDataType DataType, const FName PinLabel, const FPCGContext* InContext)
 	{
 		const UEnum* PCGDataTypeEnum = StaticEnum<EPCGDataType>();
 		const FText TypeText = PCGDataTypeEnum ? PCGDataTypeEnum->GetDisplayNameTextByValue(static_cast<int64>(DataType)) : LOCTEXT("UnknownDataType", "Unknown");
-		LogWarningOnGraph(FText::Format(ErrorFormat::TypedInputNotFoundWarning, FText::FromName(PinLabel), TypeText), InContext);
+		LogWarningOnGraph(FText::Format(Format::TypedInputNotFound, FText::FromName(PinLabel), TypeText), InContext);
 	}
 
 	void LogFirstInputOnlyWarning(const FName PinLabel, const FPCGContext* InContext)
 	{
-		LogWarningOnGraph(FText::Format(ErrorFormat::FirstInputOnlyWarning, FText::FromName(PinLabel)), InContext);
+		LogWarningOnGraph(FText::Format(Format::FirstInputOnly, FText::FromName(PinLabel)), InContext);
 	}
 
 	void LogInvalidInputDataError(const FPCGContext* InContext)
 	{
-		LogErrorOnGraph(ErrorFormat::InvalidInputDataError, InContext);
+		LogErrorOnGraph(Format::InvalidInputData, InContext);
 	}
 }
 
 namespace PCGLog::Metadata
 {
-	namespace ErrorFormat
+	namespace Format
 	{
 		const FTextFormat CreateAccessorFailure = LOCTEXT("CreateAccessorFailure", "Couldn't create accessor. Attribute '{0}' was not found.");
-		const FTextFormat CreateAttributeFailure = LOCTEXT("CreateTypedAttributeFailure", "Couldn't create attribute '{0}' of type '{1}'.");
+		const FTextFormat CreateAttributeFailure = LOCTEXT("CreateAttributeFailure", "Couldn't create attribute '{0}' of type '{1}'.");
 		const FTextFormat GetAttributeFailure = LOCTEXT("GetAttributeFailure", "Couldn't retrieve attribute '{0}' value. Expected type: {1}, Actual Type: {2}.");
-		const FTextFormat GetTypedAttributeFailure = LOCTEXT("GetAttributeFailure", "Couldn't retrieve attribute '{0}' value. Expected type: {1}, Actual Type: {2}.");
-		const FTextFormat GetTypedAttributeFailureNoAccessor = LOCTEXT("GetAttributeFailureNoAccessor", "Couldn't retrieve attribute '{0}' value of type: '{1}.");
+		const FTextFormat GetTypedAttributeFailure = LOCTEXT("GetTypedAttributeFailure", "Couldn't retrieve attribute '{0}' value. Expected type: {1}, Actual Type: {2}.");
+		const FTextFormat GetTypedAttributeFailureNoAccessor = LOCTEXT("GetTypedAttributeFailureNoAccessor", "Couldn't retrieve attribute '{0}' value of type: '{1}.");
 	}
 
-	void LogFailToCreateAccessor(const FPCGAttributePropertySelector& Selector, const FPCGContext* InContext)
+	void LogFailToCreateAccessorError(const FPCGAttributePropertySelector& Selector, const FPCGContext* InContext)
 	{
-		LogErrorOnGraph(FText::Format(ErrorFormat::CreateAccessorFailure, Selector.GetDisplayText()), InContext);
+		LogErrorOnGraph(FText::Format(Format::CreateAccessorFailure, Selector.GetDisplayText()), InContext);
 	}
 
-	void LogFailToGetAttribute(FText AttributeName, const FPCGContext* InContext)
+	void LogFailToGetAttributeError(FText AttributeName, const FPCGContext* InContext)
 	{
-		LogErrorOnGraph(FText::Format(ErrorFormat::GetAttributeFailure, std::move(AttributeName)), InContext);
+		LogErrorOnGraph(FText::Format(Format::GetAttributeFailure, std::move(AttributeName)), InContext);
 	}
 
-	void LogFailToGetAttribute(FName AttributeName, const FPCGContext* InContext)
+	void LogFailToGetAttributeError(FName AttributeName, const FPCGContext* InContext)
 	{
-		LogFailToGetAttribute(FText::FromName(AttributeName), InContext);
+		LogFailToGetAttributeError(FText::FromName(AttributeName), InContext);
 	}
 
-	void LogFailToGetAttribute(const FPCGAttributePropertySelector& Selector, const FPCGContext* InContext)
+	void LogFailToGetAttributeError(const FPCGAttributePropertySelector& Selector, const FPCGContext* InContext)
 	{
-		LogErrorOnGraph(FText::Format(ErrorFormat::GetAttributeFailure, Selector.GetDisplayText()), InContext);
+		LogErrorOnGraph(FText::Format(Format::GetAttributeFailure, Selector.GetDisplayText()), InContext);
 	}
 }
 
diff --git a/Engine/Plugins/PCG/Source/PCG/Public/Helpers/PCGPropertyHelpers.h b/Engine/Plugins/PCG/Source/PCG/Public/Helpers/PCGPropertyHelpers.h
--- a/Engine/Plugins/PCG/Source/PCG/Public/Helpers/PCGPropertyHelpers.h
+++ b/Engine/Plugins/PCG/Source/PCG/Public/Helpers/PCGPropertyHelpers.h
@@ -219,7 +219,7 @@ TArray<T> PCGPropertyHelpers::ExtractAttributeSetAsArrayOfStructs(const UPCGPara
 		{
 			if (!bCanBeDefaulted)
 			{
-				PCGLog::Metadata::LogFailToCreateAccessor(Selector, OptionalContext);
+				PCGLog::Metadata::LogFailToCreateAccessorError(Selector, OptionalContext);
 				return {};
 			}
 			else
@@ -258,7 +258,7 @@ TArray<T> PCGPropertyHelpers::ExtractAttributeSetAsArrayOfStructs(const UPCGPara
 
 			if (!Accessor->GetRange<PropertyType>(Temp, 0, *Keys, EPCGAttributeAccessorFlags::AllowBroadcastAndConstructible))
 			{
-				PCGLog::Metadata::LogFailToGetAttribute<PropertyType>(FText::FromName(PropertyName), Accessor.Get(), OptionalContext);
+				PCGLog::Metadata::LogFailToGetAttributeError<PropertyType>(FText::FromName(PropertyName), Accessor.Get(), OptionalContext);
 				return false;
 			}
 
diff --git a/Engine/Plugins/PCG/Source/PCG/Public/Utils/PCGLogErrors.h b/Engine/Plugins/PCG/Source/PCG/Public/Utils/PCGLogErrors.h
--- a/Engine/Plugins/PCG/Source/PCG/Public/Utils/PCGLogErrors.h
+++ b/Engine/Plugins/PCG/Source/PCG/Public/Utils/PCGLogErrors.h
@@ -19,68 +19,72 @@ namespace PCGLog
 
 	namespace InputOutput
 	{
-		namespace ErrorFormat
+		namespace Format
 		{
-			extern const FTextFormat TypedInputNotFoundWarning;
-			extern const FTextFormat FirstInputOnlyWarning;
-			extern const FText InvalidInputDataError;
+			extern const FTextFormat TypedInputNotFound;
+			extern const FTextFormat FirstInputOnly;
+			extern const FText InvalidInputData;
 		}
 
+		// Warnings
 		void LogTypedDataNotFoundWarning(EPCGDataType DataType, const FName PinLabel, const FPCGContext* InContext = nullptr);
 		void LogFirstInputOnlyWarning(const FName PinLabel, const FPCGContext* InContext = nullptr);
+
+		// Errors
 		void LogInvalidInputDataError(const FPCGContext* InContext = nullptr);
 	}
 
 	namespace Metadata
 	{
-		namespace ErrorFormat
+		namespace Format
 		{
 			extern const FTextFormat CreateAttributeFailure;
 			extern const FTextFormat GetTypedAttributeFailure;
 			extern const FTextFormat GetTypedAttributeFailureNoAccessor;
 		}
 
-		void LogFailToCreateAccessor(const FPCGAttributePropertySelector& Selector, const FPCGContext* InContext = nullptr);
+		// Errors
+		void LogFailToCreateAccessorError(const FPCGAttributePropertySelector& Selector, const FPCGContext* InContext = nullptr);
 
 		template <typename T>
-		void LogFailToCreateAttribute(FText AttributeName, const FPCGContext* InContext = nullptr)
+		void LogFailToCreateAttributeError(FText AttributeName, const FPCGContext* InContext = nullptr)
 		{
-			PCGLog::LogErrorOnGraph(FText::Format(ErrorFormat::CreateAttributeFailure, AttributeName, PCG::Private::GetTypeNameText<T>()), InContext);
+			PCGLog::LogErrorOnGraph(FText::Format(Format::CreateAttributeFailure, AttributeName, PCG::Private::GetTypeNameText<T>()), InContext);
 		}
 
 		template <typename T>
-		void LogFailToCreateAttribute(FName AttributeName, const FPCGContext* InContext = nullptr)
+		void LogFailToCreateAttributeError(FName AttributeName, const FPCGContext* InContext = nullptr)
 		{
-			LogFailToCreateAttribute<T>(FText::FromName(AttributeName), InContext);
+			LogFailToCreateAttributeError<T>(FText::FromName(AttributeName), InContext);
 		}
 
-		void LogFailToGetAttribute(FText AttributeName, const FPCGContext* InContext = nullptr);
-		void LogFailToGetAttribute(FName AttributeName, const FPCGContext* InContext = nullptr);
-		void LogFailToGetAttribute(const FPCGAttributePropertySelector& Selector, const FPCGContext* InContext = nullptr);
+		void LogFailToGetAttributeError(FText AttributeName, const FPCGContext* InContext = nullptr);
+		void LogFailToGetAttributeError(FName AttributeName, const FPCGContext* InContext = nullptr);
+		void LogFailToGetAttributeError(const FPCGAttributePropertySelector& Selector, const FPCGContext* InContext = nullptr);
 
 		template <typename T>
-		void LogFailToGetAttribute(FText AttributeName, const IPCGAttributeAccessor* Accessor, const FPCGContext* InContext = nullptr)
+		void LogFailToGetAttributeError(FText AttributeName, const IPCGAttributeAccessor* Accessor, const FPCGContext* InContext = nullptr)
 		{
 			if (Accessor)
 			{
-				PCGLog::LogErrorOnGraph(FText::Format(ErrorFormat::GetTypedAttributeFailure, AttributeName, PCG::Private::GetTypeNameText<T>(), PCG::Private::GetTypeNameText(Accessor->GetUnderlyingType())), InContext);
+				PCGLog::LogErrorOnGraph(FText::Format(Format::GetTypedAttributeFailure, AttributeName, PCG::Private::GetTypeNameText<T>(), PCG::Private::GetTypeNameText(Accessor->GetUnderlyingType())), InContext);
 			}
 			else
 			{
-				PCGLog::LogErrorOnGraph(FText::Format(ErrorFormat::GetTypedAttributeFailureNoAccessor, AttributeName, PCG::Private::GetTypeNameText<T>()), InContext);
+				PCGLog::LogErrorOnGraph(FText::Format(Format::GetTypedAttributeFailureNoAccessor, AttributeName, PCG::Private::GetTypeNameText<T>()), InContext);
 			}
 		}
 
 		template <typename T>
-		void LogFailToGetAttribute(FName AttributeName, const IPCGAttributeAccessor* Accessor, const FPCGContext* InContext = nullptr)
+		void LogFailToGetAttributeError(FName AttributeName, const IPCGAttributeAccessor* Accessor, const FPCGContext* InContext = nullptr)
 		{
-			LogFailToGetAttribute<T>(FText::FromName(AttributeName), Accessor, InContext);
+			LogFailToGetAttributeError<T>(FText::FromName(AttributeName), Accessor, InContext);
 		}
 
 		template <typename T>
-		void LogFailToGetAttribute(const FPCGAttributePropertySelector& Selector, const IPCGAttributeAccessor* Accessor, const FPCGContext* InContext = nullptr)
+		void LogFailToGetAttributeError(const FPCGAttributePropertySelector& Selector, const IPCGAttributeAccessor* Accessor, const FPCGContext* InContext = nullptr)
 		{
-			return LogFailToGetAttribute<T>(Selector.GetDisplayText(), Accessor, InContext);
+			return LogFailToGetAttributeError<T>(Selector.GetDisplayText(), Accessor, InContext);
 		}
 	}
 }
