[PCG - Sandbox] Remove Experiments PCG Grammar

* Remove all content and code except one node
* PrintGrammar is moved to the main plugin
* Remove the dependency from PCGSample

#rnx

diff --git a/Engine/Plugins/PCG/Source/PCG/Private/Elements/Grammar/PCGPrintGrammar.cpp b/Engine/Plugins/PCG/Source/PCG/Private/Elements/Grammar/PCGPrintGrammar.cpp
new file mode 100644
--- /dev/null
+++ b/Engine/Plugins/PCG/Source/PCG/Private/Elements/Grammar/PCGPrintGrammar.cpp
@@ -0,0 +1,114 @@
+// Copyright Epic Games, Inc. All Rights Reserved.
+
+#include "Elements/Grammar/PCGPrintGrammar.h"
+
+#include "PCGContext.h"
+#include "PCGModule.h"
+#include "PCGParamData.h"
+#include "Grammar/PCGGrammarParser.h"
+
+#define LOCTEXT_NAMESPACE "PCGPrintGrammarElement"
+
+namespace PCGPrintGrammarConstants
+{
+	const FName OutputLabel(TEXT("Description"));
+}
+
+TArray<FPCGPinProperties> UPCGPrintGrammarSettings::OutputPinProperties() const
+{
+	TArray<FPCGPinProperties> PinProperties;
+	PinProperties.Emplace(PCGPrintGrammarConstants::OutputLabel, EPCGDataType::Param);
+	return PinProperties;
+}
+
+FPCGElementPtr UPCGPrintGrammarSettings::CreateElement() const
+{
+	return MakeShared<FPCGPrintGrammarElement>();
+}
+
+bool FPCGPrintGrammarElement::ExecuteInternal(FPCGContext* Context) const
+{
+	TRACE_CPUPROFILER_EVENT_SCOPE(FPCGPrintGrammarElement::Execute);
+
+	const UPCGPrintGrammarSettings* Settings = Context->GetInputSettings<UPCGPrintGrammarSettings>();
+
+	FPCGGrammarResult Result = PCGGrammar::Parse(Settings->Grammar);
+
+	FStringBuilderBase Description;
+
+	for (const FPCGGrammarResult::FLog& Log : Result.GetLogs())
+	{
+		switch (Log.Verbosity)
+		{
+			case FPCGGrammarResult::ELogType::Log:
+				UE_LOG(LogPCG, Log, TEXT("%s"), *Log.Message.ToString());
+				break;
+			case FPCGGrammarResult::ELogType::Warning:
+				PCGLog::LogWarningOnGraph(Log.Message, Context);
+				break;
+			case FPCGGrammarResult::ELogType::Error:
+				PCGLog::LogErrorOnGraph(Log.Message, Context);
+				break;
+			default:
+				checkNoEntry();
+				break;
+		}
+	}
+
+	if (!Result.bSuccess)
+	{
+		return true;
+	}
+
+	for (const PCGGrammar::FModuleDescriptor& Descriptor : Result.Modules)
+	{
+		Description.Append("Module[");
+		switch (Descriptor.Type)
+		{
+			case PCGGrammar::EModuleType::Base:
+				Description.Append("Base");
+				break;
+			case PCGGrammar::EModuleType::Stochastic:
+				Description.Append("Stochastic");
+				break;
+			case PCGGrammar::EModuleType::Priority:
+				Description.Append("Priority");
+				break;
+		}
+
+		Description.Append("] - Repetitions: ");
+		Description.Append(Descriptor.Repetitions < 0 ? "infinite" : FString::FromInt(Descriptor.Repetitions));
+		Description.Append("\n");
+
+		for (const PCGGrammar::FModuleDescriptor::FSubmodule& Submodule : Descriptor.Submodules)
+		{
+			Description.Append("  Submodule: ");
+			Description.Append(Submodule.ID.ToString());
+			Description.Append(" Weight: ");
+			Description.Append(FString::FromInt(Submodule.Weight));
+			Description.Append("\n");
+		}
+	}
+
+	FPCGTaggedData& OutData = Context->OutputData.TaggedData.Emplace_GetRef();
+
+	const UPCGParamData* OutParamData = NewObject<UPCGParamData>();
+	UPCGMetadata* Metadata = OutParamData->Metadata;
+	check(OutParamData && Metadata);
+
+	OutData.Data = OutParamData;
+
+	if (FPCGMetadataAttribute<FString>* Attribute = Metadata->CreateAttribute(PCGPrintGrammarConstants::OutputLabel, FString(), false, true))
+	{
+		Attribute->SetValue(0, Description);
+	}
+	else
+	{
+		PCGLog::LogErrorOnGraph(LOCTEXT("CreateAttributeFailed", "Couldn't create new attribute for the description."), Context);
+		return true;
+	}
+
+	return true;
+}
+
+#undef LOCTEXT_NAMESPACE
diff --git a/Engine/Plugins/PCG/Source/PCG/Public/Elements/Grammar/PCGPrintGrammar.h b/Engine/Plugins/PCG/Source/PCG/Public/Elements/Grammar/PCGPrintGrammar.h
new file mode 100644
--- /dev/null
+++ b/Engine/Plugins/PCG/Source/PCG/Public/Elements/Grammar/PCGPrintGrammar.h
@@ -0,0 +1,42 @@
+// Copyright Epic Games, Inc. All Rights Reserved.
+
+#pragma once
+
+#include "PCGSettings.h"
+
+#include "PCGPrintGrammar.generated.h"
+
+/**
+ * Prints the result of an interpreted grammar.
+ */
+UCLASS(MinimalAPI, BlueprintType, ClassGroup = (Procedural))
+class UPCGPrintGrammarSettings : public UPCGSettings
+{
+	GENERATED_BODY()
+
+public:
+	//~Begin UPCGSettings interface
+#if WITH_EDITOR
+	virtual FName GetDefaultNodeName() const override { return FName(TEXT("PrintGrammar")); }
+	virtual FText GetDefaultNodeTitle() const override { return NSLOCTEXT("PCGPrintGrammarElement", "NodeTitle", "Print Grammar"); }
+	virtual FText GetNodeTooltipText() const override { return NSLOCTEXT("PCGPrintGrammarElement", "NodeTooltip", "Prints the result of an interpreted grammar."); }
+	virtual EPCGSettingsType GetType() const override { return EPCGSettingsType::Debug; }
+#endif
+
+protected:
+	virtual TArray<FPCGPinProperties> InputPinProperties() const override { return {}; }
+	virtual TArray<FPCGPinProperties> OutputPinProperties() const override;
+	virtual FPCGElementPtr CreateElement() const override;
+	//~End UPCGSettings interface
+
+public:
+	/** The grammar to interpret. */
+	UPROPERTY(EditAnywhere, BlueprintReadWrite, Category = Settings, meta = (PCG_Overridable))
+	FString Grammar;
+};
+
+class FPCGPrintGrammarElement : public IPCGElement
+{
+protected:
+	virtual bool ExecuteInternal(FPCGContext* Context) const override;
+};
