[PCG] Force single single connection for GPU Copy Points input pins

#rnx

diff --git a/Engine/Plugins/PCG/Source/PCG/Private/Elements/PCGCopyPoints.cpp b/Engine/Plugins/PCG/Source/PCG/Private/Elements/PCGCopyPoints.cpp
--- a/Engine/Plugins/PCG/Source/PCG/Private/Elements/PCGCopyPoints.cpp
+++ b/Engine/Plugins/PCG/Source/PCG/Private/Elements/PCGCopyPoints.cpp
@@ -120,11 +120,14 @@ void UPCGCopyPointsSettings::CreateAdditionalInputDataInterfaces(TArray<TObjectP
 
 TArray<FPCGPinProperties> UPCGCopyPointsSettings::InputPinProperties() const
 {
+	// Note: If executing on the GPU, we need to prevent multiple connections on inputs, since it is not supported at this time.
+	const bool bAllowMultipleConnections = !ShouldExecuteOnGPU();
+
 	TArray<FPCGPinProperties> PinProperties;
-	FPCGPinProperties& SourcePinProperty = PinProperties.Emplace_GetRef(PCGCopyPointsConstants::SourcePointsLabel, EPCGDataType::Point, /*bAllowMultipleConnections=*/true);
+	FPCGPinProperties& SourcePinProperty = PinProperties.Emplace_GetRef(PCGCopyPointsConstants::SourcePointsLabel, EPCGDataType::Point, bAllowMultipleConnections);
 	SourcePinProperty.SetRequiredPin();
 
-	FPCGPinProperties& TargetPinProperty = PinProperties.Emplace_GetRef(PCGCopyPointsConstants::TargetPointsLabel, EPCGDataType::Point, /*bAllowMultipleConnections=*/true);
+	FPCGPinProperties& TargetPinProperty = PinProperties.Emplace_GetRef(PCGCopyPointsConstants::TargetPointsLabel, EPCGDataType::Point, bAllowMultipleConnections);
 	TargetPinProperty.SetRequiredPin();
 
 	return PinProperties;
