Grammar parser refactoring to allow more complex setups (stochastic and priority) and submodules. Also added '+'.

diff --git a/Engine/Plugins/PCG/Source/PCG/Private/Elements/Grammar/PCGPrintGrammar.cpp b/Engine/Plugins/PCG/Source/PCG/Private/Elements/Grammar/PCGPrintGrammar.cpp
--- a/Engine/Plugins/PCG/Source/PCG/Private/Elements/Grammar/PCGPrintGrammar.cpp
+++ b/Engine/Plugins/PCG/Source/PCG/Private/Elements/Grammar/PCGPrintGrammar.cpp
@@ -60,13 +60,22 @@ bool FPCGPrintGrammarElement::ExecuteInternal(FPCGContext* Context) const
 		return true;
 	}
 
-	for (const PCGGrammar::FModuleDescriptor& Descriptor : Result.Modules)
-	{
-		Description.Append("Module[");
-		switch (Descriptor.Type)
+	auto PrintDescriptor = [&Description](const PCGGrammar::FModuleDescriptor& Descriptor, auto PrintRecurse) -> void
 		{
-			case PCGGrammar::EModuleType::Base:
-				Description.Append("Base");
+			Description.Append("Module[");
+			switch (Descriptor.Type)
+			{
+			case PCGGrammar::EModuleType::Root:
+				Description.Append("Root");
+				break;
+			case PCGGrammar::EModuleType::Literal:
+				Description.Append("Literal: ");
+				Description.Append(Descriptor.Symbol.ToString());
+				Description.Append(" Weight: ");
+				Description.Append(FString::FromInt(Descriptor.Weight));
+				break;
+			case PCGGrammar::EModuleType::Sequence:
+				Description.Append("Sequence");
 				break;
 			case PCGGrammar::EModuleType::Stochastic:
 				Description.Append("Stochastic");
@@ -74,21 +83,30 @@ bool FPCGPrintGrammarElement::ExecuteInternal(FPCGContext* Context) const
 			case PCGGrammar::EModuleType::Priority:
 				Description.Append("Priority");
 				break;
-		}
+			}
+
+			Description.Append("] - Repetitions: ");
+			if (Descriptor.Repetitions == PCGGrammar::InfiniteRepetition)
+			{
+				Description.Append("0 or more");
+			}
+			else if(Descriptor.Repetitions == PCGGrammar::AtLeastOneRepetition)
+			{
+				Description.Append("1 or more");
+			}
+			else
+			{
+				Description.Append(FString::FromInt(Descriptor.Repetitions));
+			}
+			Description.Append("\n");
 
-		Description.Append("] - Repetitions: ");
-		Description.Append(Descriptor.Repetitions < 0 ? "infinite" : FString::FromInt(Descriptor.Repetitions));
-		Description.Append("\n");
+			for (const PCGGrammar::FModuleDescriptor& Submodule : Descriptor.Submodules)
+			{
+				PrintRecurse(Submodule, PrintRecurse);
+			}
+		};
 
-		for (const PCGGrammar::FModuleDescriptor::FSubmodule& Submodule : Descriptor.Submodules)
-		{
-			Description.Append("  Submodule: ");
-			Description.Append(Submodule.ID.ToString());
-			Description.Append(" Weight: ");
-			Description.Append(FString::FromInt(Submodule.Weight));
-			Description.Append("\n");
-		}
-	}
+	PrintDescriptor(Result.Root, PrintDescriptor);
 
 	FPCGTaggedData& OutData = Context->OutputData.TaggedData.Emplace_GetRef();
 
diff --git a/Engine/Plugins/PCG/Source/PCG/Private/Elements/Grammar/PCGSegmentSlicer.cpp b/Engine/Plugins/PCG/Source/PCG/Private/Elements/Grammar/PCGSegmentSlicer.cpp
--- a/Engine/Plugins/PCG/Source/PCG/Private/Elements/Grammar/PCGSegmentSlicer.cpp
+++ b/Engine/Plugins/PCG/Source/PCG/Private/Elements/Grammar/PCGSegmentSlicer.cpp
@@ -49,7 +49,7 @@ public:
 
 		const PCGGrammar::FTokenizedGrammar& CurrentTokenizedGrammar = InOutParameters.CachedModules[InGrammar];
 
-		if (CurrentTokenizedGrammar.IsEmpty())
+		if(!CurrentTokenizedGrammar.IsValid())
 		{
 			return;
 		}
@@ -78,7 +78,7 @@ public:
 
 		TArray<PCGSlicingBase::TPCGSubDivModuleInstance<PCGGrammar::FTokenizedModule>> ModulesInstances;
 		double RemainingSubdivide;
-		const bool bSubdivideSuccess = PCGSlicingBase::Subdivide(CurrentTokenizedGrammar, Size, ModulesInstances, RemainingSubdivide, InOutParameters.Context);
+		const bool bSubdivideSuccess = PCGSlicingBase::Subdivide(*CurrentTokenizedGrammar.ModuleGrammar, Size, ModulesInstances, RemainingSubdivide, InOutParameters.Context);
 
 		if (!bSubdivideSuccess)
 		{
@@ -99,53 +99,48 @@ public:
 		for (int32 ModuleInstanceIndex = 0; ModuleInstanceIndex < ModulesInstances.Num(); ModuleInstanceIndex++)
 		{
 			const PCGSlicingBase::TPCGSubDivModuleInstance<PCGGrammar::FTokenizedModule>& ModuleInstance = ModulesInstances[ModuleInstanceIndex];
-			for (int32 j = 0; j < ModuleInstance.NumRepeat; ++j)
+
+			const FName Symbol = ModuleInstance.Module->Descriptor->Symbol;
+			const FVector Scale = FVector::OneVector + (InOutParameters.SlicingDirection * ModuleInstance.ExtraScale);
+			const FPCGSlicingSubmodule& SlicingSubmodule = InOutParameters.ModulesInfo[Symbol];
+			const bool bIsFirstModule = (ModuleInstanceIndex == 0);
+			const bool bIsFinalModule = (ModuleInstanceIndex == ModulesInstances.Num() - 1);
+
+			const double HalfDisplacement = SlicingSubmodule.Size * 0.5;
+			const double HalfScaledDisplacement = Scale.Dot(InOutParameters.SlicingDirection) * HalfDisplacement;
+
+			const FVector LocalBoundsExtents = InOutParameters.SlicingDirection * HalfDisplacement + HalfExtents2D;
+			const FVector HalfStep = HalfScaledDisplacement * Direction;
+			const FVector Position = CurrentPos + HalfStep;
+			CurrentPos = Position + HalfStep;
+
+			FPCGPoint& OutPoint = InOutParameters.OutPoints->Emplace_GetRef(Point);
+			OutPoint.Transform = FTransform(Point.Transform.GetRotation(), Position + OtherDirection, Scale);
+			OutPoint.SetLocalBounds(FBox(-LocalBoundsExtents, LocalBoundsExtents));
+			InOutParameters.OutputMetadata->InitializeOnSet(OutPoint.MetadataEntry);
+			if (InOutParameters.SymbolAttribute)
 			{
-				for (int32 SymbolIndex = 0; SymbolIndex < ModuleInstance.Module->Symbols.Num(); ++SymbolIndex)
-				{
-					const FName Symbol = ModuleInstance.Module->Symbols[SymbolIndex];
-					const FVector Scale = FVector::OneVector + (InOutParameters.SlicingDirection * ModuleInstance.ExtraScales[SymbolIndex]);
-					const FPCGSlicingSubmodule& SlicingSubmodule = InOutParameters.ModulesInfo[Symbol];
-
-					const bool bIsFirstModule = (ModuleInstanceIndex == 0) && (j == 0) && (SymbolIndex == 0);
-					const bool bIsFinalModule = (ModuleInstanceIndex == ModulesInstances.Num() - 1) && (j == ModuleInstance.NumRepeat - 1) && (SymbolIndex == ModuleInstance.Module->Symbols.Num() - 1);
-					const double HalfDisplacement = SlicingSubmodule.Size * 0.5;
-					const double HalfScaledDisplacement = Scale.Dot(InOutParameters.SlicingDirection) * HalfDisplacement;
-
-					const FVector LocalBoundsExtents = InOutParameters.SlicingDirection * HalfDisplacement + HalfExtents2D;
-					const FVector HalfStep = HalfScaledDisplacement * Direction;
-					const FVector Position = CurrentPos + HalfStep;
-					CurrentPos = Position + HalfStep;
-
-					FPCGPoint& OutPoint = InOutParameters.OutPoints->Emplace_GetRef(Point);
-					OutPoint.Transform = FTransform(Point.Transform.GetRotation(), Position + OtherDirection, Scale);
-					OutPoint.SetLocalBounds(FBox(-LocalBoundsExtents, LocalBoundsExtents));
-					InOutParameters.OutputMetadata->InitializeOnSet(OutPoint.MetadataEntry);
-					if (InOutParameters.SymbolAttribute)
-					{
-						InOutParameters.SymbolAttribute->SetValue(OutPoint.MetadataEntry, Symbol);
-					}
-
-					if (InOutParameters.DebugColorAttribute)
-					{
-						InOutParameters.DebugColorAttribute->SetValue(OutPoint.MetadataEntry, FVector4(SlicingSubmodule.DebugColor, 1.0));
-					}
-
-					if (InOutParameters.ModuleIndexAttribute)
-					{
-						InOutParameters.ModuleIndexAttribute->SetValue(OutPoint.MetadataEntry, ModuleIndex++);
-					}
-
-					if (bIsFirstModule && InOutParameters.IsFirstPointAttribute)
-					{
-						InOutParameters.IsFirstPointAttribute->SetValue(OutPoint.MetadataEntry, true);
-					}
-
-					if (bIsFinalModule && InOutParameters.IsFinalPointAttribute)
-					{
-						InOutParameters.IsFinalPointAttribute->SetValue(OutPoint.MetadataEntry, true);
-					}
-				}
+				InOutParameters.SymbolAttribute->SetValue(OutPoint.MetadataEntry, Symbol);
+			}
+
+			if (InOutParameters.DebugColorAttribute)
+			{
+				InOutParameters.DebugColorAttribute->SetValue(OutPoint.MetadataEntry, FVector4(SlicingSubmodule.DebugColor, 1.0));
+			}
+
+			if (InOutParameters.ModuleIndexAttribute)
+			{
+				InOutParameters.ModuleIndexAttribute->SetValue(OutPoint.MetadataEntry, ModuleIndex++);
+			}
+
+			if (bIsFirstModule && InOutParameters.IsFirstPointAttribute)
+			{
+				InOutParameters.IsFirstPointAttribute->SetValue(OutPoint.MetadataEntry, true);
+			}
+
+			if (bIsFinalModule && InOutParameters.IsFinalPointAttribute)
+			{
+				InOutParameters.IsFinalPointAttribute->SetValue(OutPoint.MetadataEntry, true);
 			}
 		}
 
diff --git a/Engine/Plugins/PCG/Source/PCG/Private/Elements/Grammar/PCGSlicingBase.cpp b/Engine/Plugins/PCG/Source/PCG/Private/Elements/Grammar/PCGSlicingBase.cpp
--- a/Engine/Plugins/PCG/Source/PCG/Private/Elements/Grammar/PCGSlicingBase.cpp
+++ b/Engine/Plugins/PCG/Source/PCG/Private/Elements/Grammar/PCGSlicingBase.cpp
@@ -144,11 +144,14 @@ PCGGrammar::FTokenizedGrammar FPCGSlicingBaseElement::GetTokenizedGrammar(FPCGCo
 
 PCGGrammar::FTokenizedGrammar PCGSlicingBase::GetTokenizedGrammar(FPCGContext* InContext, const FString& InGrammar, const FPCGModulesInfoMap& InModulesInfo, double& OutMinSize)
 {
-	const FPCGGrammarResult Result = PCGGrammar::Parse(InGrammar);
+	FPCGGrammarResult Result = PCGGrammar::Parse(InGrammar);
 
-	if (!Result.bSuccess)
+	// TODO: Add quiet mode
 	{
-		PCGLog::LogErrorOnGraph(LOCTEXT("GrammarParseFail", "Problem while parsing grammar:"), InContext);
+		if (!Result.bSuccess)
+		{
+			PCGLog::LogErrorOnGraph(LOCTEXT("GrammarParseFail", "Problem while parsing grammar:"), InContext);
+		}
 
 		for (const FPCGGrammarResult::FLog& Log : Result.GetLogs())
 		{
@@ -165,42 +168,107 @@ PCGGrammar::FTokenizedGrammar PCGSlicingBase::GetTokenizedGrammar(FPCGContext* I
 				break;
 			}
 		}
+	}
 
+	if (!Result.bSuccess)
+	{
 		return {};
 	}
 
-	PCGGrammar::FTokenizedGrammar TokenizedGrammar;
-	OutMinSize = 0.0;
-
-	for (const PCGGrammar::FModuleDescriptor& ModuleDescriptor : Result.Modules)
+	// Build equivalent grammar tree with the size information.
+	TSet<FName> UnmatchedTokens;
+	auto BuildNode = [&UnmatchedTokens, &InModulesInfo](PCGGrammar::FTokenizedModule& Module, const PCGGrammar::FModuleDescriptor& Descriptor, auto& RecursiveFn) -> void
 	{
-		PCGGrammar::FTokenizedModule& CurrentModule = TokenizedGrammar.Emplace_GetRef();
-		CurrentModule.NumRepeat = ModuleDescriptor.Repetitions;
-
-		for (const PCGGrammar::FModuleDescriptor::FSubmodule& SubmoduleDescriptor : ModuleDescriptor.Submodules)
+		if(Descriptor.Type == PCGGrammar::EModuleType::Literal)
 		{
-			if (const FPCGSlicingSubmodule* It = InModulesInfo.Find(SubmoduleDescriptor.ID))
+			if (const FPCGSlicingSubmodule* It = InModulesInfo.Find(Descriptor.Symbol))
 			{
-				CurrentModule.Symbols.Add(SubmoduleDescriptor.ID);
-				CurrentModule.Size += It->Size;
-				CurrentModule.bScalable |= It->bScalable;
-				CurrentModule.AreSymbolsScalable.Add(It->bScalable);
-				CurrentModule.SymbolSizes.Add(It->Size);
+				Module.UnitSize = It->Size;
+				Module.ConcreteUnitSize = It->Size;
+				Module.bScalable = It->bScalable;
+				Module.bIsValid = true;
+			}
+			else
+			{
+				UnmatchedTokens.Add(Descriptor.Symbol);
 			}
 		}
-
-		if (CurrentModule.Symbols.IsEmpty())
+		else
 		{
-			// If we have no symbol, we skip.
-			continue;
-		}
+			// Go through submodules
+			Module.Submodules.Reserve(Descriptor.Submodules.Num());
 
-		if (CurrentModule.NumRepeat > 0)
-		{
-			OutMinSize += CurrentModule.Size * CurrentModule.NumRepeat;
+			// In any of the modules that can be expanded, the minimum concrete size is going to be either the MinSize if it's greater than 0, otherwise it will be the minimum of the minimum concrete sizes of all submodules
+			double MinSubmoduleConcreteSize = 0.0;
+
+			for (const PCGGrammar::FModuleDescriptor& SubDescriptor : Descriptor.Submodules)
+			{
+				PCGGrammar::FTokenizedModule& Submodule = Module.Submodules.Emplace_GetRef(&SubDescriptor);
+				RecursiveFn(Submodule, SubDescriptor, RecursiveFn);
+
+				if (Submodule.bIsValid)
+				{
+					// Update size/scalable, with Sequence/Root: sum and Stochastic/Priority : min
+					if (Descriptor.Type == PCGGrammar::EModuleType::Root || Descriptor.Type == PCGGrammar::EModuleType::Sequence)
+					{
+						Module.UnitSize += Submodule.GetMinSize();
+					}
+					else if (Descriptor.Type == PCGGrammar::EModuleType::Stochastic || Descriptor.Type == PCGGrammar::EModuleType::Priority)
+					{
+						if (Module.bIsValid)
+						{
+							Module.UnitSize = FMath::Min(Module.UnitSize, Submodule.GetMinSize());
+						}
+						else
+						{
+							Module.UnitSize = Submodule.GetMinSize();
+						}
+					}
+
+					if (Module.bIsValid)
+					{
+						MinSubmoduleConcreteSize = FMath::Min(MinSubmoduleConcreteSize, Submodule.GetMinConcreteSize());
+					}
+					else
+					{
+						MinSubmoduleConcreteSize = Submodule.GetMinConcreteSize();
+					}
+
+					Module.bScalable |= Submodule.bScalable;
+					Module.bIsValid = true;
+				}
+			}
+
+			// If all expansions of the module lead to nothing, we might as well remove it.
+			if (Module.bIsValid && MinSubmoduleConcreteSize > 0)
+			{
+				Module.ConcreteUnitSize = FMath::Max(Module.UnitSize, MinSubmoduleConcreteSize);
+			}
+			else
+			{
+				Module.bIsValid = false;
+			}
 		}
+	};
+
+	// TODO Add quiet mode
+	for (FName UnmatchedToken : UnmatchedTokens)
+	{
+		FText WarningMessage = FText::Format(LOCTEXT("UnmatchedTokensInGrammar", "Unmatched token found in grammar: {0}."), FText::FromName(UnmatchedToken));
+		PCGLog::LogWarningOnGraph(WarningMessage, InContext);
+	}
+
+	if (Result.Root.Submodules.IsEmpty())
+	{
+		return {};
 	}
 
+	PCGGrammar::FTokenizedGrammar TokenizedGrammar;
+	TokenizedGrammar.ParsedGrammar = MakeShared<PCGGrammar::FModuleDescriptor>(MoveTemp(Result.Root));
+	TokenizedGrammar.ModuleGrammar = MakeShared<PCGGrammar::FTokenizedModule>(TokenizedGrammar.ParsedGrammar.Get());
+	BuildNode(*TokenizedGrammar.ModuleGrammar, *TokenizedGrammar.ParsedGrammar, BuildNode);
+	OutMinSize = TokenizedGrammar.ModuleGrammar->GetMinSize();
+
 	return TokenizedGrammar;
 }
 
diff --git a/Engine/Plugins/PCG/Source/PCG/Private/Elements/Grammar/PCGSplineSlicer.cpp b/Engine/Plugins/PCG/Source/PCG/Private/Elements/Grammar/PCGSplineSlicer.cpp
--- a/Engine/Plugins/PCG/Source/PCG/Private/Elements/Grammar/PCGSplineSlicer.cpp
+++ b/Engine/Plugins/PCG/Source/PCG/Private/Elements/Grammar/PCGSplineSlicer.cpp
@@ -93,7 +93,7 @@ namespace PCGSplineSlicerHelpers
 
 		const PCGGrammar::FTokenizedGrammar& CurrentTokenizedGrammar = InOutParameters.CachedModules[InGrammar];
 
-		if (CurrentTokenizedGrammar.IsEmpty())
+		if (!CurrentTokenizedGrammar.IsValid())
 		{
 			return;
 		}
@@ -105,7 +105,7 @@ namespace PCGSplineSlicerHelpers
 		 */
 		TArray<PCGSlicingBase::TPCGSubDivModuleInstance<PCGGrammar::FTokenizedModule>> ModulesInstances;
 		double RemainingSubdivide;
-		if (!Subdivide(CurrentTokenizedGrammar, SplineLength, ModulesInstances, RemainingSubdivide, InOutParameters.Context))
+		if (!Subdivide(*CurrentTokenizedGrammar.ModuleGrammar, SplineLength, ModulesInstances, RemainingSubdivide, InOutParameters.Context))
 		{
 			return;
 		}
@@ -123,85 +123,80 @@ namespace PCGSplineSlicerHelpers
 		for (int32 ModuleInstanceIndex = 0; ModuleInstanceIndex < ModulesInstances.Num(); ModuleInstanceIndex++)
 		{
 			const PCGSlicingBase::TPCGSubDivModuleInstance<PCGGrammar::FTokenizedModule>& ModuleInstance = ModulesInstances[ModuleInstanceIndex];
-			for (int32 SubmoduleInstanceIndex = 0; SubmoduleInstanceIndex < ModuleInstance.NumRepeat; ++SubmoduleInstanceIndex)
-			{
-				for (int32 SymbolIndex = 0; SymbolIndex < ModuleInstance.Module->Symbols.Num(); ++SymbolIndex)
-				{
-					const FName Symbol = ModuleInstance.Module->Symbols[SymbolIndex];
-					const FPCGSlicingSubmodule& SlicingSubmodule = InOutParameters.ModulesInfo[Symbol];
-					const double SubmoduleSize = SlicingSubmodule.Size;
+			const FName Symbol = ModuleInstance.Module->Descriptor->Symbol;
+
+			const FPCGSlicingSubmodule& SlicingSubmodule = InOutParameters.ModulesInfo[Symbol];
+			const double SubmoduleSize = SlicingSubmodule.Size;
 
-					// Move to the next segment of the spline
-					const FVector SegmentStartPoint = InOutParameters.PolyLineData->GetLocationAtAlpha(SplineAlpha);
-					const double PreviousAlpha = SplineAlpha;
+			// Move to the next segment of the spline
+			const FVector SegmentStartPoint = InOutParameters.PolyLineData->GetLocationAtAlpha(SplineAlpha);
+			const double PreviousAlpha = SplineAlpha;
 
-					// Use a numerical method to find the root of where the module lands on the spline
-					SplineAlpha = FindRootAtLinearDistance_Bisection(InOutParameters.PolyLineData, SubmoduleSize, SplineAlpha);
+			// Use a numerical method to find the root of where the module lands on the spline
+			SplineAlpha = FindRootAtLinearDistance_Bisection(InOutParameters.PolyLineData, SubmoduleSize, SplineAlpha);
 
-					const FVector SegmentEndPoint = InOutParameters.PolyLineData->GetLocationAtAlpha(SplineAlpha);
-					const FVector SliceVector = SegmentEndPoint - SegmentStartPoint;
-					const FVector SliceDirection = SliceVector.GetSafeNormal();
+			const FVector SegmentEndPoint = InOutParameters.PolyLineData->GetLocationAtAlpha(SplineAlpha);
+			const FVector SliceVector = SegmentEndPoint - SegmentStartPoint;
+			const FVector SliceDirection = SliceVector.GetSafeNormal();
 
-					// At the end of the spline, truncate an unfinished submodule and end the process
-					if (FMath::IsNearlyEqual(SplineAlpha, 1))
+			// At the end of the spline, truncate an unfinished submodule and end the process
+			if (FMath::IsNearlyEqual(SplineAlpha, 1))
+			{
+				if (SliceVector.Length() < SubmoduleSize)
+				{
+					if (InOutParameters.IsFinalPointAttribute && !InOutParameters.OutPoints->IsEmpty())
 					{
-						if (SliceVector.Length() < SubmoduleSize)
-						{
-							if (InOutParameters.IsFinalPointAttribute && !InOutParameters.OutPoints->IsEmpty())
-							{
-								InOutParameters.IsFinalPointAttribute->SetValue(InOutParameters.OutPoints->Last().MetadataEntry, true);
-							}
-
-							return;
-						}
+						InOutParameters.IsFinalPointAttribute->SetValue(InOutParameters.OutPoints->Last().MetadataEntry, true);
 					}
 
-					// Since its discretized, we won't take the transform's position, but we'll use the up vector--to create the module rotation--and the scale
-					FTransform CenterPointTransform = InOutParameters.PolyLineData->GetTransformAtAlpha((SplineAlpha + PreviousAlpha) * 0.5);
+					return;
+				}
+			}
 
-					const FVector Position = SegmentStartPoint + (SliceVector * 0.5) + FVector(0, 0, InOutParameters.ModuleHeight * 0.5);
-					const FRotator Rotation = FRotationMatrix::MakeFromXZ(SliceDirection, CenterPointTransform.GetRotation().GetUpVector()).Rotator();
-					const FVector Scale = FVector::OneVector + (SliceDirection * ModuleInstance.ExtraScales[SymbolIndex]);
+			// Since its discretized, we won't take the transform's position, but we'll use the up vector--to create the module rotation--and the scale
+			FTransform CenterPointTransform = InOutParameters.PolyLineData->GetTransformAtAlpha((SplineAlpha + PreviousAlpha) * 0.5);
 
-					FPCGPoint& OutPoint = InOutParameters.OutPoints->Emplace_GetRef(FTransform(Rotation, Position, Scale), /*InDensity=*/1, PCGHelpers::ComputeSeedFromPosition(Position));
+			const FVector Position = SegmentStartPoint + (SliceVector * 0.5) + FVector(0, 0, InOutParameters.ModuleHeight * 0.5);
+			const FRotator Rotation = FRotationMatrix::MakeFromXZ(SliceDirection, CenterPointTransform.GetRotation().GetUpVector()).Rotator();
+			const FVector Scale = FVector::OneVector + (SliceDirection * ModuleInstance.ExtraScale);
 
-					const double HalfSubmoduleSize = SlicingSubmodule.Size * 0.5;
-					OutPoint.SetLocalBounds(FBox(FVector(-HalfSubmoduleSize, 0, 0), FVector(HalfSubmoduleSize, 1, InOutParameters.ModuleHeight)));
+			FPCGPoint& OutPoint = InOutParameters.OutPoints->Emplace_GetRef(FTransform(Rotation, Position, Scale), /*InDensity=*/1, PCGHelpers::ComputeSeedFromPosition(Position));
 
-					// Now, handle the metadata attributes
-					if (!bHasMetadata)
-					{
-						continue;
-					}
+			const double HalfSubmoduleSize = SlicingSubmodule.Size * 0.5;
+			OutPoint.SetLocalBounds(FBox(FVector(-HalfSubmoduleSize, 0, 0), FVector(HalfSubmoduleSize, 1, InOutParameters.ModuleHeight)));
 
-					InOutParameters.OutputMetadata->InitializeOnSet(OutPoint.MetadataEntry);
-					if (InOutParameters.SymbolAttribute)
-					{
-						InOutParameters.SymbolAttribute->SetValue(OutPoint.MetadataEntry, Symbol);
-					}
+			// Now, handle the metadata attributes
+			if (!bHasMetadata)
+			{
+				continue;
+			}
 
-					if (InOutParameters.DebugColorAttribute)
-					{
-						InOutParameters.DebugColorAttribute->SetValue(OutPoint.MetadataEntry, FVector4(SlicingSubmodule.DebugColor, 1.0));
-					}
+			InOutParameters.OutputMetadata->InitializeOnSet(OutPoint.MetadataEntry);
+			if (InOutParameters.SymbolAttribute)
+			{
+				InOutParameters.SymbolAttribute->SetValue(OutPoint.MetadataEntry, Symbol);
+			}
 
-					if (InOutParameters.ModuleIndexAttribute)
-					{
-						InOutParameters.ModuleIndexAttribute->SetValue(OutPoint.MetadataEntry, ModuleIndexCounter++);
-					}
+			if (InOutParameters.DebugColorAttribute)
+			{
+				InOutParameters.DebugColorAttribute->SetValue(OutPoint.MetadataEntry, FVector4(SlicingSubmodule.DebugColor, 1.0));
+			}
 
-					const bool bIsFirstModule = (ModuleInstanceIndex == 0) && (SubmoduleInstanceIndex == 0) && (SymbolIndex == 0);
-					if (bIsFirstModule && InOutParameters.IsFirstPointAttribute)
-					{
-						InOutParameters.IsFirstPointAttribute->SetValue(OutPoint.MetadataEntry, true);
-					}
+			if (InOutParameters.ModuleIndexAttribute)
+			{
+				InOutParameters.ModuleIndexAttribute->SetValue(OutPoint.MetadataEntry, ModuleIndexCounter++);
+			}
 
-					const bool bIsFinalModule = (ModuleInstanceIndex == ModulesInstances.Num() - 1) && (SubmoduleInstanceIndex == ModuleInstance.NumRepeat - 1) && (SymbolIndex == ModuleInstance.Module->Symbols.Num() - 1);
-					if (bIsFinalModule && InOutParameters.IsFinalPointAttribute)
-					{
-						InOutParameters.IsFinalPointAttribute->SetValue(OutPoint.MetadataEntry, true);
-					}
-				}
+			const bool bIsFirstModule = (ModuleInstanceIndex == 0);
+			if (bIsFirstModule && InOutParameters.IsFirstPointAttribute)
+			{
+				InOutParameters.IsFirstPointAttribute->SetValue(OutPoint.MetadataEntry, true);
+			}
+
+			const bool bIsFinalModule = (ModuleInstanceIndex == ModulesInstances.Num() - 1);
+			if (bIsFinalModule && InOutParameters.IsFinalPointAttribute)
+			{
+				InOutParameters.IsFinalPointAttribute->SetValue(OutPoint.MetadataEntry, true);
 			}
 		}
 	}
diff --git a/Engine/Plugins/PCG/Source/PCG/Private/Elements/Grammar/PCGVolumeSlicer.cpp b/Engine/Plugins/PCG/Source/PCG/Private/Elements/Grammar/PCGVolumeSlicer.cpp
--- a/Engine/Plugins/PCG/Source/PCG/Private/Elements/Grammar/PCGVolumeSlicer.cpp
+++ b/Engine/Plugins/PCG/Source/PCG/Private/Elements/Grammar/PCGVolumeSlicer.cpp
@@ -122,14 +122,14 @@ bool FPCGVolumeSlicerElement::ExecuteInternal(FPCGContext* InContext) const
 		double MinSize = 0.0;
 		PCGGrammar::FTokenizedGrammar TokenizedGrammar = GetTokenizedGrammar(InContext, InputSplineData, Settings, ModulesInfo, MinSize);
 
-		if (TokenizedGrammar.IsEmpty())
+		if (!TokenizedGrammar.IsValid())
 		{
 			continue;
 		}
 
 		TArray<PCGSlicingBase::TPCGSubDivModuleInstance<PCGGrammar::FTokenizedModule>> Instances;
 		double RemainingLength = 0.0;
-		const bool bHeightSubdivideSuccess = PCGSlicingBase::Subdivide(TokenizedGrammar, ExtrudeLength, Instances, RemainingLength);
+		const bool bHeightSubdivideSuccess = PCGSlicingBase::Subdivide(*TokenizedGrammar.ModuleGrammar, ExtrudeLength, Instances, RemainingLength);
 
 		if (!bHeightSubdivideSuccess)
 		{
@@ -142,49 +142,44 @@ bool FPCGVolumeSlicerElement::ExecuteInternal(FPCGContext* InContext) const
 
 		for (const PCGSlicingBase::TPCGSubDivModuleInstance<PCGGrammar::FTokenizedModule>& Instance : Instances)
 		{
-			for (int i = 0; i < Instance.NumRepeat; ++i)
+			const FName Symbol = Instance.Module->Descriptor->Symbol;
+			const FPCGSlicingSubmodule& CurrentBlock = ModulesInfo[Symbol];
+			const FVector Size = ExtrudeDirection * CurrentBlock.Size * (FVector::OneVector + Instance.ExtraScale);
+			
+			FPCGSplineStruct NewSpline = InputSplineData->SplineStruct;
+			for (FInterpCurvePointVector& ControlPoint : NewSpline.SplineCurves.Position.Points)
 			{
-				for (int32 SymbolIndex = 0; SymbolIndex < Instance.Module->Symbols.Num(); ++SymbolIndex)
-				{
-					const FName Symbol = Instance.Module->Symbols[SymbolIndex];
-					const FVector Size = ExtrudeDirection * Instance.Module->SymbolSizes[SymbolIndex] * (FVector::OneVector + Instance.ExtraScales[SymbolIndex]);
-					const FPCGSlicingSubmodule& CurrentBlock = ModulesInfo[Symbol];
-					FPCGSplineStruct NewSpline = InputSplineData->SplineStruct;
-					for (FInterpCurvePointVector& ControlPoint : NewSpline.SplineCurves.Position.Points)
-					{
-						ControlPoint.OutVal += CurrentDisplacement;
-					}
-
-					UPCGSplineData* NewSplineData = FPCGContext::NewObject_AnyThread<UPCGSplineData>(InContext);
-					NewSplineData->Initialize(NewSpline);
-					NewSplineData->InitializeFromData(InputSplineData);
-
-					NewSplineData->Metadata->FindOrCreateAttribute<FName>(Settings->SymbolAttributeName, Symbol, false, false);
-
-					if (!bMatchAndSetAttributes && Settings->bOutputDebugColorAttribute)
-					{
-						NewSplineData->Metadata->FindOrCreateAttribute<FVector4>(Settings->DebugColorAttributeName, FVector4(CurrentBlock.DebugColor, 1.0), false, false);
-					}
-
-					if (Settings->bOutputSizeAttribute)
-					{
-						NewSplineData->Metadata->FindOrCreateAttribute<FVector>(Settings->SizeAttributeName, Size, false, false);
-					}
-
-					if (!bMatchAndSetAttributes && Settings->bOutputScalableAttribute)
-					{
-						NewSplineData->Metadata->FindOrCreateAttribute<bool>(Settings->ScalableAttributeName, CurrentBlock.bScalable, false, false);
-					}
-
-					if (Settings->bOutputSplineIndexAttribute)
-					{
-						NewSplineData->Metadata->FindOrCreateAttribute<int32>(Settings->SplineIndexAttributeName, SplineIndex++, false, false);
-					}
-
-					Outputs.Add_GetRef(Input).Data = NewSplineData;
-					CurrentDisplacement += Size;
-				}
+				ControlPoint.OutVal += CurrentDisplacement;
 			}
+
+			UPCGSplineData* NewSplineData = FPCGContext::NewObject_AnyThread<UPCGSplineData>(InContext);
+			NewSplineData->Initialize(NewSpline);
+			NewSplineData->InitializeFromData(InputSplineData);
+
+			NewSplineData->Metadata->FindOrCreateAttribute<FName>(Settings->SymbolAttributeName, Symbol, false, false);
+
+			if (!bMatchAndSetAttributes && Settings->bOutputDebugColorAttribute)
+			{
+				NewSplineData->Metadata->FindOrCreateAttribute<FVector4>(Settings->DebugColorAttributeName, FVector4(CurrentBlock.DebugColor, 1.0), false, false);
+			}
+
+			if (Settings->bOutputSizeAttribute)
+			{
+				NewSplineData->Metadata->FindOrCreateAttribute<FVector>(Settings->SizeAttributeName, Size, false, false);
+			}
+
+			if (!bMatchAndSetAttributes && Settings->bOutputScalableAttribute)
+			{
+				NewSplineData->Metadata->FindOrCreateAttribute<bool>(Settings->ScalableAttributeName, CurrentBlock.bScalable, false, false);
+			}
+
+			if (Settings->bOutputSplineIndexAttribute)
+			{
+				NewSplineData->Metadata->FindOrCreateAttribute<int32>(Settings->SplineIndexAttributeName, SplineIndex++, false, false);
+			}
+
+			Outputs.Add_GetRef(Input).Data = NewSplineData;
+			CurrentDisplacement += Size;
 		}
 	}
 
diff --git a/Engine/Plugins/PCG/Source/PCG/Private/Grammar/PCGGrammarParser.cpp b/Engine/Plugins/PCG/Source/PCG/Private/Grammar/PCGGrammarParser.cpp
--- a/Engine/Plugins/PCG/Source/PCG/Private/Grammar/PCGGrammarParser.cpp
+++ b/Engine/Plugins/PCG/Source/PCG/Private/Grammar/PCGGrammarParser.cpp
@@ -4,7 +4,7 @@
 
 #include "Algo/Transform.h"
 #include "Internationalization/Internationalization.h"
-#include "Internationalization/Regex.h"
+#include "String/ParseTokens.h"
 
 #define LOCTEXT_NAMESPACE "PCGGrammar"
 
@@ -19,167 +19,166 @@ void FPCGGrammarResult::AddLog(FText Message, ELogType Verbosity)
 
 namespace PCGGrammar
 {
-	namespace Patterns
+	bool BuildModules(const FString& InGrammar, FPCGGrammarResult& Result)
 	{
-		static constexpr const TCHAR* Base = TEXT("(\\[([^\\]]+)\\](?:\\s*(\\d+|\\*))?)");
-		static constexpr const TCHAR* Stochastic = TEXT("(\\{([^\\}]+)\\}(?:\\s*(\\d++|\\*))?)");
-		static constexpr const TCHAR* Priority = TEXT("(<([^>]+)>(?:\\s*(\\d++|\\*))?)");
-	}
-
-	namespace Tokens
-	{
-		static constexpr const TCHAR* SubmoduleDelimiter = TEXT(",");
-		static constexpr const TCHAR* WeightDelimiter = TEXT(":");
-		static constexpr const TCHAR* InfiniteRepetition = TEXT("*");
-	}
-
-	enum ECaptureGroups : int
-	{
-		// Capture Group 0 is the matched expression for this iteration
-		MatchedExpression = 0,
-
-		// Capture group 1 is the whole group of submodules, including [ ]
-		WholeGroupOfSubmodules = 1,
-
-		// Capture group 2 is the submodule expression. Ie. [A,B], {A,B}, or <A,B> -> A,B
-		SubmoduleExpression = 2,
-
-		// Capture group 3 is the number of repetitions - [A,B]3 -> 3
-		NumberOfRepetitionsOrAsterisk = 3
-	};
-
-	bool FindModules(
-		const FString& Grammar,
-		const EModuleType ModuleType,
-		TArray<TPair<int32, int32>, TInlineAllocator<16>>& OutMatchedIndices,
-		FPCGGrammarResult& OutGrammarResult)
-	{
-		const TCHAR* Pattern = nullptr;
-		switch (ModuleType)
+		static constexpr TCHAR Delimiter[] = TEXT(",");
+		static constexpr TCHAR WeightDelimiter[] = TEXT(":");
+		static constexpr TCHAR SequenceStart[] = TEXT("[");
+		static constexpr TCHAR SequenceEnd[] = TEXT("]");
+		static constexpr TCHAR PriorityStart[] = TEXT("<");
+		static constexpr TCHAR PriorityEnd[] = TEXT(">");
+		static constexpr TCHAR StochasticStart[] = TEXT("{");
+		static constexpr TCHAR StochasticEnd[] = TEXT("}");
+		static constexpr TCHAR InfiniteRepetitionStr[] = TEXT("*");
+		static constexpr TCHAR AtLeastOneRepetitionStr[] = TEXT("+");
+
+		// TODO: parse tokens in a more efficient way; this relies on token separation from what should essentially be token delimiters
+		// This would also allow better information to the end user when the grammar is invalid (as we wouldn't need to modify the string).
+		FString Grammar = InGrammar;
+		Grammar.ReplaceInline(Delimiter, TEXT(" , "));
+		Grammar.ReplaceInline(WeightDelimiter, TEXT(" : "));
+		Grammar.ReplaceInline(SequenceStart, TEXT(" [ "));
+		Grammar.ReplaceInline(SequenceEnd, TEXT(" ] "));
+		Grammar.ReplaceInline(PriorityStart, TEXT(" < "));
+		Grammar.ReplaceInline(PriorityEnd, TEXT(" > "));
+		Grammar.ReplaceInline(StochasticStart, TEXT(" { "));
+		Grammar.ReplaceInline(StochasticEnd, TEXT(" } "));
+		Grammar.ReplaceInline(InfiniteRepetitionStr, TEXT(" * "));
+		Grammar.ReplaceInline(AtLeastOneRepetitionStr, TEXT(" + "));
+
+		TArray<FStringView> Tokens;
+		UE::String::ParseTokens(Grammar, TEXT(" "), Tokens, UE::String::EParseTokensOptions::SkipEmpty);
+
+		PCGGrammar::FModuleDescriptor* CurrentModule = &Result.Root;
+		PCGGrammar::FModuleDescriptor* LastModule = nullptr;
+		bool bParseOk = true;
+		bool bHasWeightDelimiter = false;
+
+		auto CreateSubmodule = [&CurrentModule, &LastModule](EModuleType Type)
 		{
-			case EModuleType::Base:
-				Pattern = Patterns::Base;
-				break;
-			case EModuleType::Stochastic:
-				Pattern = Patterns::Stochastic;
-				break;
-			case EModuleType::Priority:
-				Pattern = Patterns::Priority;
-				break;
-			default:
-				checkNoEntry();
-				return false;
-		}
-
-		check(Pattern);
+			PCGGrammar::FModuleDescriptor& Submodule = CurrentModule->Submodules.Emplace_GetRef(Type);
+			Submodule.Parent = CurrentModule;
+			CurrentModule = &Submodule;
+			LastModule = nullptr;
+		};
 
-		bool bFoundMatch = false;
-
-		TArray<FName> SubmoduleNames;
-		TArray<int32> SubmoduleWeights;
-
-		FRegexMatcher ModuleMatcher(FRegexPattern(Pattern), Grammar);
-		while (ModuleMatcher.FindNext())
+		auto CloseSubmodule = [&CurrentModule, &LastModule, &Result, &bParseOk](EModuleType Type)
 		{
-			bFoundMatch = true;
-
-			const FString SubmoduleCapture = ModuleMatcher.GetCaptureGroup(ECaptureGroups::SubmoduleExpression);
-
-			if (SubmoduleCapture.IsEmpty())
+			if(CurrentModule && CurrentModule->Type == Type)
 			{
-				FText Message = FText::Format(LOCTEXT("NoModuleNameFound", "Unable to find module name match within the grammar declaration: {0}"), FText::FromString(ModuleMatcher.GetCaptureGroup(ECaptureGroups::MatchedExpression)));
-				OutGrammarResult.AddLog(std::move(Message), FPCGGrammarResult::ELogType::Warning);
-				continue;
+				LastModule = CurrentModule;
+				CurrentModule = CurrentModule->Parent;
 			}
-
-			TArray<FString> SubmoduleStrings;
-			SubmoduleCapture.ParseIntoArray(SubmoduleStrings, Tokens::SubmoduleDelimiter);
-
-			SubmoduleNames.Reset(SubmoduleStrings.Num());
-			SubmoduleWeights.Reset(SubmoduleStrings.Num());
-
-			for (const FString& SubmoduleString : SubmoduleStrings)
+			else
 			{
-				TArray<FString> SubmoduleIDString;
-				SubmoduleString.ParseIntoArray(SubmoduleIDString, Tokens::WeightDelimiter);
+				Result.AddLog(LOCTEXT("InvalidCloseSubmodule", "Mismatched module end."), FPCGGrammarResult::ELogType::Error);
+				bParseOk = false;
+			}
+		};
 
-				if (SubmoduleIDString.IsEmpty())
+		for (const FStringView& Token : Tokens)
+		{
+			if (Token == SequenceStart)
+			{
+				CreateSubmodule(EModuleType::Sequence);
+			}
+			else if (Token == PriorityStart)
+			{
+				CreateSubmodule(EModuleType::Priority);
+			}
+			else if (Token == StochasticStart)
+			{
+				CreateSubmodule(EModuleType::Stochastic);
+			}
+			else if (Token == SequenceEnd)
+			{
+				CloseSubmodule(EModuleType::Sequence);
+			}
+			else if (Token == PriorityEnd)
+			{
+				CloseSubmodule(EModuleType::Priority);
+			}
+			else if (Token == StochasticEnd)
+			{
+				CloseSubmodule(EModuleType::Stochastic);
+			}
+			else if(Token == Delimiter)
+			{
+				LastModule = nullptr;
+			}
+			else if (Token == WeightDelimiter)
+			{
+				// TODO: verify module supports weights, add warning if so
+				if(!LastModule || bHasWeightDelimiter)
 				{
-					FText Message = FText::Format(LOCTEXT("InvalidSubmoduleString", "Parsed submodule string is invalid."), FText::FromString(ModuleMatcher.GetCaptureGroup(ECaptureGroups::MatchedExpression)));
-					OutGrammarResult.AddLog(std::move(Message), FPCGGrammarResult::ELogType::Warning);
-					continue;
+					Result.AddLog(LOCTEXT("InvalidWeightDelimiter", "Misplaced or mismatched weight (':') delimiter found."), FPCGGrammarResult::ELogType::Warning);
 				}
-
-				SubmoduleIDString[0].TrimStartAndEndInline();
-				if (SubmoduleIDString[0].IsEmpty())
+				
+				if (LastModule)
 				{
-					FText Message = FText::Format(LOCTEXT("EmptyModuleID", "Module ID must not be empty."), FText::FromString(ModuleMatcher.GetCaptureGroup(ECaptureGroups::MatchedExpression)));
-					OutGrammarResult.AddLog(std::move(Message), FPCGGrammarResult::ELogType::Warning);
-					continue;
+					bHasWeightDelimiter = true;
 				}
-
-				SubmoduleNames.Emplace(SubmoduleIDString[0]);
-				SubmoduleWeights.Emplace(1);
-
-				if (SubmoduleIDString.Num() > 1)
+			}
+			else
+			{
+				// Some Three possibilities here:
+				// A literal (last module is null)
+				if (LastModule == nullptr)
 				{
-					if (ModuleType != EModuleType::Stochastic)
+					// Technically, we could assume this is part of the same token then (so we could have spaces in the tokens, but it's not clean)
+					PCGGrammar::FModuleDescriptor& Literal = CurrentModule->Submodules.Emplace_GetRef(FName(Token));
+					LastModule = &Literal;
+				}
+				// Otherwise : The repetition count or the weight
+				else if(bHasWeightDelimiter)
+				{
+					FString TokenString(Token);
+					if (TokenString.IsNumeric())
 					{
-						FText Message = FText::Format(LOCTEXT("WeightOnNonStochasticType", "Weight added to non-stochastic module type."), FText::FromString(ModuleMatcher.GetCaptureGroup(ECaptureGroups::MatchedExpression)));
-						OutGrammarResult.AddLog(std::move(Message), FPCGGrammarResult::ELogType::Warning);
-						continue;
+						LastModule->Weight = FCString::Atoi(*TokenString);
 					}
-
-					if (SubmoduleIDString.Num() > 2)
+					else
 					{
-						FText Message = LOCTEXT("MultiCharacterWeightDelimiter", "Multi-character delimiter for weight. Weight ignored.");
-						OutGrammarResult.AddLog(std::move(Message), FPCGGrammarResult::ELogType::Warning);
-						continue;
+						FText WarningMessage = FText::Format(LOCTEXT("InvalidWeightValue", "Invalid weight value: {0}, will default to 1."), FText::FromStringView(Token));
+						Result.AddLog(std::move(WarningMessage), FPCGGrammarResult::ELogType::Warning);
 					}
 
-					SubmoduleIDString[1].TrimStartAndEndInline();
-					if (!SubmoduleIDString[1].IsNumeric())
+					// Consumed the weight specifier
+					bHasWeightDelimiter = false;
+				}
+				else
+				{
+					FString TokenString(Token);
+					if (Token == InfiniteRepetitionStr)
 					{
-						FText Message = FText::Format(LOCTEXT("InvalidWeightCharacter", "Invalid weight character '{0}' Weight ignored."), FText::FromString(SubmoduleIDString[1]));
-						OutGrammarResult.AddLog(std::move(Message), FPCGGrammarResult::ELogType::Warning);
-						continue;
+						LastModule->Repetitions = PCGGrammar::InfiniteRepetition;
+					}
+					else if (Token == AtLeastOneRepetitionStr)
+					{
+						LastModule->Repetitions = PCGGrammar::AtLeastOneRepetition;
+					}
+					else if (TokenString.IsNumeric())
+					{
+						LastModule->Repetitions = FCString::Atoi(*TokenString);
+					}
+					else
+					{
+						FText WarningMessage = FText::Format(LOCTEXT("InvalidRepetitionValue", "Invalid repetition value: {0}, will default to 1."), FText::FromStringView(Token));
+						Result.AddLog(std::move(WarningMessage), FPCGGrammarResult::ELogType::Warning);
 					}
-
-					SubmoduleWeights.Last() = FCString::Atoi(*SubmoduleIDString[1]);
 				}
 			}
 
-			// Capture repetition count
-			int32 Repetitions = 1;
-			const FString RepetitionString = ModuleMatcher.GetCaptureGroup(ECaptureGroups::NumberOfRepetitionsOrAsterisk);
-			if (RepetitionString.IsNumeric())
-			{
-				Repetitions = FCString::Atoi(*RepetitionString);
-			}
-			else if (!RepetitionString.IsEmpty())
-			{
-				// It should only be the infinite repetition character, to match the
-				ensure(RepetitionString == Tokens::InfiniteRepetition);
-				Repetitions = -1;
-			}
-
-			FModuleDescriptor& Module = OutGrammarResult.Modules.Emplace_GetRef();
-			Module.GrammarStartEndIndices = {ModuleMatcher.GetMatchBeginning(), ModuleMatcher.GetMatchEnding()};
-			Module.Type = ModuleType;
-			Module.Repetitions = Repetitions;
-
-			for (int Index = 0; Index < SubmoduleNames.Num(); ++Index)
+			if (!bParseOk)
 			{
-				Module.Submodules.Emplace(SubmoduleNames[Index], SubmoduleWeights[Index]);
+				break;
 			}
-
-			OutMatchedIndices.Emplace(ModuleMatcher.GetMatchBeginning(), ModuleMatcher.GetMatchEnding());
 		}
 
-		return bFoundMatch;
+		return bParseOk;
 	}
 
-	PCG_API FPCGGrammarResult Parse(const FString& Grammar, bool bValidateGrammar)
+	PCG_API FPCGGrammarResult Parse(const FString& Grammar)
 	{
 		FPCGGrammarResult Result;
 
@@ -189,47 +188,7 @@ namespace PCGGrammar
 			return Result;
 		}
 
-		bool bFoundMatch = true;
-		TArray<TPair<int32, int32>, TInlineAllocator<16>> MatchedIndices;
-
-		bFoundMatch |= FindModules(Grammar, EModuleType::Base, MatchedIndices, Result);
-		bFoundMatch |= FindModules(Grammar, EModuleType::Stochastic, MatchedIndices, Result);
-		bFoundMatch |= FindModules(Grammar, EModuleType::Priority, MatchedIndices, Result);
-
-		if (!bFoundMatch)
-		{
-			Result.AddLog(LOCTEXT("NoModuleMatch", "Unable to find module match in grammar."), FPCGGrammarResult::ELogType::Warning);
-			return Result;
-		}
-
-		// Sort by start index - where they appeared in the grammar.
-		Algo::Sort(Result.Modules, [](const FModuleDescriptor& LHS, const FModuleDescriptor& RHS)
-		{
-			return LHS.GrammarStartEndIndices.Key < RHS.GrammarStartEndIndices.Key;
-		});
-
-		if (bValidateGrammar)
-		{
-			FString RemainingCharacters = Grammar;
-
-			MatchedIndices.Sort([](const TPair<int32, int32>& Pair1, const TPair<int32, int32>& Pair2)
-			{
-				return Pair1.Value < Pair2.Value;
-			});
-
-			// Must remove in inverse order to avoid conflict
-			for (int I = MatchedIndices.Num() - 1; I >= 0; --I)
-			{
-				const TPair<int32, int32>& Pair = MatchedIndices[I];
-				RemainingCharacters.RemoveAt(Pair.Key, Pair.Value - Pair.Key, EAllowShrinking::No);
-			}
-
-			RemainingCharacters.TrimStartAndEndInline();
-			if (!RemainingCharacters.IsEmpty())
-			{
-				Result.AddLog(FText::Format(LOCTEXT("ExtraCharactersInGrammar", "Extraneous characters in grammar {0}"), FText::FromString(RemainingCharacters)), FPCGGrammarResult::ELogType::Warning);
-			}
-		}
+		Result.bSuccess = BuildModules(Grammar, Result);
 
 		return Result;
 	}
diff --git a/Engine/Plugins/PCG/Source/PCG/Public/Elements/Grammar/PCGSlicingBase.h b/Engine/Plugins/PCG/Source/PCG/Public/Elements/Grammar/PCGSlicingBase.h
--- a/Engine/Plugins/PCG/Source/PCG/Public/Elements/Grammar/PCGSlicingBase.h
+++ b/Engine/Plugins/PCG/Source/PCG/Public/Elements/Grammar/PCGSlicingBase.h
@@ -2,6 +2,7 @@
 
 #pragma once
 
+#include "PCGContext.h"
 #include "PCGSettings.h"
 #include "Grammar/PCGGrammar.h"
 #include "Utils/PCGLogErrors.h"
@@ -81,6 +82,8 @@ public:
 #if WITH_EDITOR
 	virtual EPCGSettingsType GetType() const override { return EPCGSettingsType::Spatial; }
 #endif
+
+	virtual bool UseSeed() const override { return true; }
 	//~End UPCGSettings interface
 
 	virtual void PostLoad() override;
@@ -151,151 +154,307 @@ namespace PCGSlicingBase
 {
 	using FPCGModulesInfoMap = TMap<FName, FPCGSlicingSubmodule>;
 
-#define PCG_SLICING_BASE_USES_CONCEPTS 0
-
-#if PCG_SLICING_BASE_USES_CONCEPTS
-	// To validate that the type has the right interface
-	// To be enabled when this is supported on all our platforms.
-	template <typename T>
-	concept IsValidPCGSubDivModuleInstance =
-		requires(T t) {
-			{ t.IsValid() } -> std::same_as<bool>;
-			{ t.GetSize() } -> std::same_as<double>;
-			{ t.GetNumRepeat() } -> std::same_as<int>;
-			{ t.IsScalable() } -> std::same_as<bool>;
-			{ t.GetSubmodulesCount() } -> std::same_as<int32>;
-			{ t.AreSubmodulesScalable() } -> std::same_as<TArrayView<const bool>>;
-			{ t.SubmoduleSizes() } -> std::same_as<TArrayView<const double>>;
-	};
-#endif // PCG_SLICING_BASE_USES_CONCEPTS
-
-	template <typename T>
+	// Materialized modules created from a tokenized grammar.
+	template<typename T>
 	struct TPCGSubDivModuleInstance
 	{
+		TPCGSubDivModuleInstance() = default;
+		explicit TPCGSubDivModuleInstance(const T* InModule) : Module(InModule) {}
+
 		const T* Module = nullptr;
-		int32 NumRepeat = 0;
-		// If the module is made of multiple submodules, we'll have 1 extra scale per submodule.
-		// Extra scale to add to the module initial scale on the slicing direction.
-		TArray<double> ExtraScales;
+		double ExtraScale = 0.0;
+		bool bIsValid = true;
+		bool bSkipExpansion = false;
 	};
 
 	PCGGrammar::FTokenizedGrammar GetTokenizedGrammar(FPCGContext* InContext, const FString& InGrammar, const FPCGModulesInfoMap& InModulesInfo, double& OutMinSize);
 
-	template <typename T>
-	bool Subdivide(const TArray<T>& Modules, double Length, TArray<TPCGSubDivModuleInstance<T>>& OutModuleInstances, double& RemainingLength, FPCGContext* InOptionalContext = nullptr)
+	template<typename T>
+	bool Subdivide(const T& Root, double Length, TArray<TPCGSubDivModuleInstance<T>>& OutModuleInstances, double& RemainingLength, FPCGContext* InOptionalContext = nullptr, int32 InOptionalSeed = 42)
 	{
-		OutModuleInstances.Empty(Modules.Num());
+		OutModuleInstances.Reset();
 		RemainingLength = Length;
 
-		if (Modules.IsEmpty() || FMath::IsNearlyZero(Length))
+		if (!Root.IsValid() || FMath::IsNearlyZero(Length))
 		{
 			return true;
 		}
 
-		for (const T& Module : Modules)
-		{
-			if (!Module.IsValid())
-			{
-				continue;
-			}
+		FRandomStream RandomStream(InOptionalContext ? InOptionalContext->GetSeed() : InOptionalSeed);
+		TArray<TPCGSubDivModuleInstance<T>> CurrentModules;
 
-			TPCGSubDivModuleInstance<T>& ModuleInstance = OutModuleInstances.Emplace_GetRef();
-			ModuleInstance.Module = &Module;
-			ModuleInstance.ExtraScales.SetNumZeroed(Module.GetSubmodulesCount());
-			if (Module.GetNumRepeat() > 0)
-			{
-				ModuleInstance.NumRepeat = Module.GetNumRepeat();
-				RemainingLength -= Module.GetSize() * ModuleInstance.NumRepeat;
-				if (RemainingLength < 0)
-				{
-					PCGLog::LogErrorOnGraph(NSLOCTEXT("PCGSlicingBase", "SegmentCutFail", "Grammar doesn't fit for this segment."), InOptionalContext);
-					return false;
-				}
-			}
-		}
+		// Start with root
+		CurrentModules.Emplace(&Root);
+		RemainingLength -= Root.GetMinSize(); // Here we use the min size because we'll consume concrete size only during expansion
 
-		if (OutModuleInstances.IsEmpty())
+		if (RemainingLength < 0)
 		{
+			PCGLog::LogErrorOnGraph(NSLOCTEXT("PCGSlicingBase", "SegmentCutFail", "Grammar doesn't fit for this segment."), InOptionalContext);
 			return false;
 		}
 
-		// When we are done and we still have some segemnt left, place the repeatable
-		int32 CurrentModuleIndex = 0;
-		bool bHasModifiedSomething = false;
-		while (RemainingLength >= 0)
+		// Working data sets
+		TArray<TPCGSubDivModuleInstance<T>*> PriorityModules;
+		TArray<TPCGSubDivModuleInstance<T>*> StochasticModules;
+		TArray<TPCGSubDivModuleInstance<T>> ExpandedModules;
+		TArray<TPCGSubDivModuleInstance<T>> PreviousModules;
+		
+		bool bNotDone = true;
+		while (bNotDone)
 		{
-			TPCGSubDivModuleInstance<T>& ModuleInstance = OutModuleInstances[CurrentModuleIndex];
-			if (ModuleInstance.Module->GetNumRepeat() <= 0)
+			bNotDone = false;
+			// 1. Expand/Visit sequences
+			// Implementation note; since we've already consumed the min size, there's no need to update anything when expanding nodes, only when replacing with another choice.
+			// We'll keep track of whether we have certain types of nodes at this point so we can short-circuit some passes
+			ExpandedModules.Reset();
+			ExpandedModules.Reserve(CurrentModules.Num());
+			bool bHasPriorityOrStochasticModules = false;
+
+			for (TPCGSubDivModuleInstance<T>& CurrentModule : CurrentModules)
 			{
-				if (RemainingLength >= ModuleInstance.Module->GetSize())
+				if (!CurrentModule.bIsValid)
 				{
-					++ModuleInstance.NumRepeat;
-					RemainingLength -= ModuleInstance.Module->GetSize();
-					bHasModifiedSomething = true;
+					continue; // discard this module (reclaim?)
 				}
-			}
 
-			if (++CurrentModuleIndex == OutModuleInstances.Num())
-			{
-				if (!bHasModifiedSomething)
+				if ((CurrentModule.Module->GetType() == PCGGrammar::EModuleType::Root || CurrentModule.Module->GetType() == PCGGrammar::EModuleType::Sequence))
 				{
-					// Nothing left
-					break;
+					// When we duplicate sequences, we can't expand them directly otherwise they won't be repeatable
+					if (CurrentModule.bSkipExpansion)
+					{
+						// Copy module as-is, but remove the expansion limitation so we break it down next iteration
+						TPCGSubDivModuleInstance<T>& ExpandedModule = ExpandedModules.Add_GetRef(CurrentModule);
+						ExpandedModule.bSkipExpansion = false;
+					}
+					else
+					{
+						for (const T& Submodule : CurrentModule.Module->Submodules)
+						{
+							if (Submodule.IsValid())
+							{
+								// Expand module if it has a positive number of iterations (>= 1) or if it the unit size would fit.
+								// This is where we do a difference also between the * (0+) and + cases; the + will have a force one repetition here, while the * will not.
+								int NumRepeats = Submodule.GetNumRepeat();
+								if (Submodule.GetNumRepeat() == PCGGrammar::AtLeastOneRepetition)
+								{
+									NumRepeats = 1;
+								}
+								else if(NumRepeats == PCGGrammar::InfiniteRepetition && RemainingLength >= Submodule.GetMinConcreteSize())
+								{
+									NumRepeats = 1;
+									RemainingLength -= Submodule.GetMinConcreteSize();
+								}
+
+								for (int R = 0; R < NumRepeats; ++R)
+								{
+									TPCGSubDivModuleInstance<T>& ExpandedModule = ExpandedModules.Emplace_GetRef(&Submodule);
+								}
+							
+								if (NumRepeats > 0)
+								{
+									bNotDone = true;
+									bHasPriorityOrStochasticModules |= (Submodule.GetType() == PCGGrammar::EModuleType::Priority || Submodule.GetType() == PCGGrammar::EModuleType::Stochastic);
+								}
+							}
+						}
+					}
 				}
+				else if (CurrentModule.Module->GetType() == PCGGrammar::EModuleType::Priority || CurrentModule.Module->GetType() == PCGGrammar::EModuleType::Stochastic)
+				{
+					bHasPriorityOrStochasticModules = true;
 
-				bHasModifiedSomething = false;
-				CurrentModuleIndex = 0;
+					// Copy module as-is, but remove the expansion limitation so we break it down next iteration
+					TPCGSubDivModuleInstance<T>& ExpandedModule = ExpandedModules.Add_GetRef(CurrentModule);
+					ExpandedModule.bSkipExpansion = false;
+				}
+				else
+				{
+					// Copy module as-is, but do NOT remove the expansion limitation, otherwise we will have issues with literals with repetitions.
+					TPCGSubDivModuleInstance<T>& ExpandedModule = ExpandedModules.Add_GetRef(CurrentModule);
+				}
 			}
-		}
 
-		// Finally, try to stretch the scalable to get a complete match.
-		if (!FMath::IsNearlyZero(RemainingLength))
-		{
-			int32 NumScalableSubmodules = 0;
+			CurrentModules = MoveTemp(ExpandedModules);
+
+			// 2- Make a copy of the current set of modules, because we'll use this for repeatable modules
+			PreviousModules = CurrentModules;
 
-			TArray<TPCGSubDivModuleInstance<T>*> ScalableInstances;
-			for (TPCGSubDivModuleInstance<T>& ModuleInstance : OutModuleInstances)
+			// 3- Perform choices
+			if (bHasPriorityOrStochasticModules)
 			{
-				if (ModuleInstance.Module->IsScalable() && ModuleInstance.NumRepeat > 0)
+				// 3.1 - Gather priority & stochastic modules (pointers here)
+				PriorityModules.Reset();
+				StochasticModules.Reset();
+
+				for (TPCGSubDivModuleInstance<T>& CurrentModule : CurrentModules)
 				{
-					ScalableInstances.Add(&ModuleInstance);
-					for (const bool bIsScalable : ModuleInstance.Module->AreSubmodulesScalable())
+					check(CurrentModule.Module);
+					if (CurrentModule.Module->GetType() == PCGGrammar::EModuleType::Priority)
+					{
+						PriorityModules.Add(&CurrentModule);
+					}
+					else if (CurrentModule.Module->GetType() == PCGGrammar::EModuleType::Stochastic)
 					{
-						if (bIsScalable)
+						StochasticModules.Add(&CurrentModule);
+					}
+				}
+
+				bNotDone |= (!PriorityModules.IsEmpty() || !StochasticModules.IsEmpty());
+
+				// 3.2 - Apply priority choices first (from left to right)
+				for (TPCGSubDivModuleInstance<T>* PriorityModule : PriorityModules)
+				{
+					bool bWasReplaced = false;
+					// Replace current module by the first of its childen that has a min size that fits in with the remaining length
+					for (const T& PrioritySubmodule : PriorityModule->Module->Submodules)
+					{
+						if (!PrioritySubmodule.IsValid())
+						{
+							continue;
+						}
+
+						const double DeltaMinSize = (PrioritySubmodule.GetMinConcreteSize() - PriorityModule->Module->GetConcreteUnitSize());
+						check(DeltaMinSize >= 0);
+						if (RemainingLength >= DeltaMinSize)
 						{
-							NumScalableSubmodules += ModuleInstance.NumRepeat;
+							PriorityModule->Module = &PrioritySubmodule;
+							RemainingLength -= DeltaMinSize;
+							bWasReplaced = true;
+							break;
 						}
 					}
+
+					if (!bWasReplaced)
+					{
+						// TODO Kill this module + reclaim size(?)
+						PriorityModule->bIsValid = false;
+					}
 				}
-			}
 
+				// 3.3 Process stochastic modules
+				// TODO random pick order of stochastic modules
+				for (TPCGSubDivModuleInstance<T>* StochasticModule : StochasticModules)
+				{
+					// Replace current module by a random pick according to total weights of "valid" choices we can still make
+					int TotalWeight = 0;
+					for (const T& StochasticSubmodule : StochasticModule->Module->Submodules)
+					{
+						if (!StochasticSubmodule.IsValid())
+						{
+							continue;
+						}
+
+						const double DeltaMinSize = (StochasticSubmodule.GetMinConcreteSize() - StochasticModule->Module->GetConcreteUnitSize());
+						check(DeltaMinSize >= 0);
+						if (DeltaMinSize <= RemainingLength)
+						{
+							TotalWeight += StochasticSubmodule.GetWeight();
+						}
+					}
+
+					if (TotalWeight == 0)
+					{
+						// Kill module + reclaim size(?)
+						StochasticModule->bIsValid = false;
+						continue;
+					}
 
-			if (NumScalableSubmodules > 0 && !ScalableInstances.IsEmpty())
-			{
-				const double ExtraLengthPerSubmodule = RemainingLength / NumScalableSubmodules;
+					int WeightPick = RandomStream.RandRange(0, TotalWeight - 1);
+					bool bWasReplaced = false;
 
-				for (TPCGSubDivModuleInstance<T>* ModuleInstance : ScalableInstances)
-				{
-					const TArrayView<const bool> AreSubmodulesScalable = ModuleInstance->Module->AreSubmodulesScalable();
-					const TArrayView<const double> SubmoduleSizes = ModuleInstance->Module->SubmoduleSizes();
-					check(AreSubmodulesScalable.Num() == SubmoduleSizes.Num());
-					for (int32 i = 0; i < SubmoduleSizes.Num(); ++i)
+					for (const T& StochasticSubmodule : StochasticModule->Module->Submodules)
 					{
-						if (AreSubmodulesScalable[i])
+						if (!StochasticSubmodule.IsValid())
 						{
-							const double SubmoduleSize = SubmoduleSizes[i];
-							if (!FMath::IsNearlyZero(SubmoduleSize))
+							continue;
+						}
+
+						const double DeltaMinSize = (StochasticSubmodule.GetMinConcreteSize() - StochasticModule->Module->GetConcreteUnitSize());
+						check(DeltaMinSize >= 0);
+						if(DeltaMinSize <= RemainingLength)
+						{
+							if (StochasticSubmodule.GetWeight() > WeightPick)
 							{
-								ModuleInstance->ExtraScales[i] = ExtraLengthPerSubmodule / SubmoduleSize;
+								bWasReplaced = true;
+								RemainingLength -= DeltaMinSize;
+								StochasticModule->Module = &StochasticSubmodule;
+								break;
+							}
+							else
+							{
+								WeightPick -= StochasticSubmodule.GetWeight();
 							}
 						}
 					}
+
+					check(bWasReplaced);
 				}
+			}
+
+			// 4. Add arbitrary repetitions, requires to keep data from previous iteration
+			// Implementation note: since we add one module only in each iteration, it will definitely be inefficient when there is a long chain of things to duplicate.
+			check(PreviousModules.Num() == CurrentModules.Num());
+			int NumAddedRepetitions = 0;
 
-				RemainingLength = 0.0;
+			for (int ModuleIndex = 0; ModuleIndex < PreviousModules.Num(); ++ModuleIndex)
+			{
+				TPCGSubDivModuleInstance<T>& PreviousModule = PreviousModules[ModuleIndex];
+				if (!PreviousModule.bSkipExpansion && PreviousModule.Module->GetNumRepeat() < 0 && PreviousModule.Module->GetMinConcreteSize() <= RemainingLength)
+				{
+					// Mark the module to skip expansion, otherwise we can't chain for repeats
+					PreviousModule.bSkipExpansion = true;
+					CurrentModules.Insert(PreviousModule, ModuleIndex + NumAddedRepetitions + 1);
+					++NumAddedRepetitions;
+					RemainingLength -= PreviousModule.Module->GetMinConcreteSize();
+				}
 			}
+
+			bNotDone |= (NumAddedRepetitions > 0);
 		}
 
+#if WITH_EDITOR
+		// Perform some early validation and see if there's a mismatch on the reported size and the one actually placed
+		double CountedLength = RemainingLength;
+		for (TPCGSubDivModuleInstance<T>& CurrentModule : CurrentModules)
+		{
+			CountedLength += CurrentModule.Module->GetUnitSize();
+		}
+
+		ensure(FMath::Abs(CountedLength-Length) < 1.0);
+#endif
+
+		// 7. Finally, apply adjusted scales to modules that support it.
+		check(RemainingLength >= 0);
+		if (!FMath::IsNearlyZero(RemainingLength))
+		{
+			int32 NumScalableModules = 0;
+			double ScalableLength = 0;
+			for (TPCGSubDivModuleInstance<T>& CurrentModule : CurrentModules)
+			{
+				if (CurrentModule.Module->IsScalable())
+				{
+					++NumScalableModules;
+					// implementation note: at this point we have only unit-literals, so we need to ignore repetitions, if any here, hence using the unit size
+					ScalableLength += CurrentModule.Module->GetUnitSize();
+				}
+			}
+
+			if (ScalableLength > 0)
+			{
+				for (TPCGSubDivModuleInstance<T>& CurrentModule : CurrentModules)
+				{
+					if (CurrentModule.Module->IsScalable())
+					{
+						CurrentModule.ExtraScale = (RemainingLength / ScalableLength);
+					}
+				}
+
+				RemainingLength = 0;
+			}
+		}
+
+		OutModuleInstances = MoveTemp(CurrentModules);
+
 		return true;
 	}
 }
diff --git a/Engine/Plugins/PCG/Source/PCG/Public/Grammar/PCGGrammar.h b/Engine/Plugins/PCG/Source/PCG/Public/Grammar/PCGGrammar.h
--- a/Engine/Plugins/PCG/Source/PCG/Public/Grammar/PCGGrammar.h
+++ b/Engine/Plugins/PCG/Source/PCG/Public/Grammar/PCGGrammar.h
@@ -2,31 +2,47 @@
 
 #pragma once
 
+#include "Grammar/PCGGrammarParser.h"
 #include "Metadata/PCGAttributePropertySelector.h"
 
 #include "PCGGrammar.generated.h"
 
 namespace PCGGrammar
 {
+	// Coupled concrete module information with a grammar module descriptor (node)
 	struct FTokenizedModule
 	{
-		TArray<FName> Symbols;
-		TArray<bool> AreSymbolsScalable;
-		TArray<double> SymbolSizes;
-		int NumRepeat = 1;
-		bool bScalable = false;
-		double Size = 0.0;
+		FTokenizedModule() = default;
+		FTokenizedModule(const PCGGrammar::FModuleDescriptor* InDescriptor)
+			: Descriptor(InDescriptor)
+		{
+		}
 
-		bool IsValid() const { return Size > 0 && !Symbols.IsEmpty() && Symbols.Num() == AreSymbolsScalable.Num() && Symbols.Num() == SymbolSizes.Num(); }
-		double GetSize() const { return Size; }
-		int GetNumRepeat() const { return NumRepeat; }
+		bool IsValid() const { return bIsValid; }
+		double GetMinSize() const { return UnitSize * (Descriptor ? Descriptor->GetMinNumberOfRepetitions() : 0); }
+		double GetUnitSize() const { return UnitSize; }
+		double GetMinConcreteSize() const { return ConcreteUnitSize * (Descriptor ? Descriptor->GetMinConcreteNumberOfRepetitions() : 0); }
+		double GetConcreteUnitSize() const { return ConcreteUnitSize; }
+		int GetNumRepeat() const { return Descriptor ? Descriptor->Repetitions : 0; }
+		int8 GetWeight() const { return Descriptor ? Descriptor->Weight : 0; }
 		bool IsScalable() const { return bScalable; }
-		int32 GetSubmodulesCount() const { return Symbols.Num(); }
-		TArrayView<const bool> AreSubmodulesScalable() const { return AreSymbolsScalable; }
-		TArrayView<const double> SubmoduleSizes() const { return SymbolSizes; }
+		EModuleType GetType() const { return Descriptor ? Descriptor->Type : EModuleType::Unknown; }
+
+		TArray<FTokenizedModule> Submodules;
+		const PCGGrammar::FModuleDescriptor* Descriptor = nullptr;
+		double UnitSize = 0.0; // minimum size not including repetitions
+		double ConcreteUnitSize = 0.0; // minimum size of the expansion of a unit of this module
+		bool bScalable = false;
+		bool bIsValid = false;
 	};
 
-	using FTokenizedGrammar = TArray<FTokenizedModule>;
+	struct FTokenizedGrammar
+	{
+		TSharedPtr<FModuleDescriptor> ParsedGrammar;
+		TSharedPtr<FTokenizedModule> ModuleGrammar;
+
+		bool IsValid() const { return ModuleGrammar.IsValid() && ModuleGrammar.Get()->IsValid(); }
+	};
 }
 
 USTRUCT(Blueprintable, BlueprintType)
diff --git a/Engine/Plugins/PCG/Source/PCG/Public/Grammar/PCGGrammarParser.h b/Engine/Plugins/PCG/Source/PCG/Public/Grammar/PCGGrammarParser.h
--- a/Engine/Plugins/PCG/Source/PCG/Public/Grammar/PCGGrammarParser.h
+++ b/Engine/Plugins/PCG/Source/PCG/Public/Grammar/PCGGrammarParser.h
@@ -10,29 +10,40 @@ namespace PCGGrammar
 {
 	enum class EModuleType
 	{
-		Base,
+		Root,
+		Literal,
+		Sequence,
 		Stochastic,
-		Priority
+		Priority,
+		Unknown // used when there are no matching descriptors
 	};
 
+	static constexpr int InfiniteRepetition = -1;
+	static constexpr int AtLeastOneRepetition = -2;
+
 	struct FModuleDescriptor
 	{
-		struct FSubmodule
+		FModuleDescriptor() = default;
+		FModuleDescriptor(FName InLiteral)
+		{
+			Type = EModuleType::Literal;
+			Symbol = InLiteral;
+		}
+
+		explicit FModuleDescriptor(EModuleType InType)
 		{
-			FSubmodule(const FName InID, const int8 InWeight)
-				: ID(InID)
-				, Weight(InWeight)
-			{}
+			Type = InType;
+		}
 
-			FName ID;
-			int8 Weight;
-		};
+		int32 GetMinNumberOfRepetitions() const { return (Repetitions >= 0 ? Repetitions : (Repetitions == AtLeastOneRepetition ? 1 : 0)); }
+		int32 GetMinConcreteNumberOfRepetitions() const { return ((Repetitions >= 0) ? Repetitions : 1); }
 
-		EModuleType Type = EModuleType::Base;
+		TArray<FModuleDescriptor> Submodules;
+		FModuleDescriptor* Parent = nullptr;
+		FName Symbol;
 		int32 Repetitions = 1;
-		// The matched beginning and ending indices of the grammar string for this module
-		TPair<int32, int32> GrammarStartEndIndices;
-		TArray<FSubmodule> Submodules;
+		EModuleType Type = EModuleType::Literal;
+		int8 Weight = 1; // TODO: review the need for this based on struct size
 	};
 }
 
@@ -63,7 +74,7 @@ struct PCG_API FPCGGrammarResult
 	const TArray<FLog>& GetLogs() const{ return Logs; }
 
 	bool bSuccess = true;
-	TArray<PCGGrammar::FModuleDescriptor> Modules;
+	PCGGrammar::FModuleDescriptor Root = PCGGrammar::FModuleDescriptor(PCGGrammar::EModuleType::Root);
 
 private:
 	TArray<FLog> Logs;
@@ -71,5 +82,5 @@ private:
 
 namespace PCGGrammar
 {
-	PCG_API FPCGGrammarResult Parse(const FString& Grammar, bool bValidateGrammar = true);
+	PCG_API FPCGGrammarResult Parse(const FString& Grammar);
 }
