[PCG] UX pass on the Clustering Node

- Removed the option to use input/manual/random centroids in favor of an input or random if empty
+ Added a toggle to output the final centroids/gaussians as an additional output pin

#rnx

diff --git a/Engine/Plugins/PCG/Source/PCG/Private/Elements/PCGClusterElement.cpp b/Engine/Plugins/PCG/Source/PCG/Private/Elements/PCGClusterElement.cpp
--- a/Engine/Plugins/PCG/Source/PCG/Private/Elements/PCGClusterElement.cpp
+++ b/Engine/Plugins/PCG/Source/PCG/Private/Elements/PCGClusterElement.cpp
@@ -16,6 +16,7 @@ namespace PCGClusterElement
 	namespace Constants
 	{
 		const FName InputCentroidsLabel = TEXT("Initial Centroids");
+		const FName FinalCentroidsLabel = TEXT("Final Centroids");
 	}
 
 	namespace Helpers
@@ -373,7 +374,7 @@ FText UPCGClusterSettings::GetNodeTooltipText() const
 
 bool UPCGClusterSettings::UseSeed() const
 {
-	return InitialCentroidSelection == EPCGInitialClusteringCentroidSelection::RandomPoints;
+	return true;
 }
 
 TArray<FPCGPinProperties> UPCGClusterSettings::InputPinProperties() const
@@ -381,14 +382,26 @@ TArray<FPCGPinProperties> UPCGClusterSettings::InputPinProperties() const
 	TArray<FPCGPinProperties> Properties;
 	Properties.Emplace_GetRef(PCGPinConstants::DefaultInputLabel, EPCGDataType::Point).SetRequiredPin();
 
-	if (InitialCentroidSelection == EPCGInitialClusteringCentroidSelection::Input)
+	// TODO: For now, currently only support one input for N:1, but N:N would be useful too.
+	FPCGPinProperties& InputCentroidsPin = Properties.Emplace_GetRef(PCGClusterElement::Constants::InputCentroidsLabel, EPCGDataType::Point, /*bInAllowMultipleConnections=*/false, /*bAllowMultipleData=*/false);
+#if WITH_EDITOR
+	InputCentroidsPin.Tooltip = LOCTEXT("InputCentroidsTooltip", "(Optional) Points whose locations will be used as both the quantity and location of the initial centroids for the clustering algorithm. If none provided, random points will be used as centroids.");
+#endif // WITH_EDITOR
+
+	return Properties;
+}
+
+TArray<FPCGPinProperties> UPCGClusterSettings::OutputPinProperties() const
+{
+	TArray<FPCGPinProperties> Properties;
+	Properties.Emplace(PCGPinConstants::DefaultOutputLabel, EPCGDataType::Point);
+
+	if (bOutputFinalCentroids)
 	{
-		// TODO: For now, currently only support one input for N:1, but N:N would be useful too.
-		FPCGPinProperties& InputCentroidsPin = Properties.Emplace_GetRef(PCGClusterElement::Constants::InputCentroidsLabel, EPCGDataType::Point, /*bInAllowMultipleConnections=*/false, /*bAllowMultipleData=*/false);
+		FPCGPinProperties& FinalCentroids = Properties.Emplace_GetRef(PCGClusterElement::Constants::FinalCentroidsLabel, EPCGDataType::Point);
 #if WITH_EDITOR
-		InputCentroidsPin.Tooltip = LOCTEXT("InputCentroidsTooltip", "Points whose locations will be used as both the quantity and location of the initial centroids for the clustering algorithm.");
+		FinalCentroids.Tooltip = LOCTEXT("FinalCentroidsTooltip", "The final locations of the centroids/gaussians.");
 #endif // WITH_EDITOR
-		return Properties;
 	}
 
 	return Properties;
@@ -407,15 +420,18 @@ bool FPCGClusterElement::PrepareDataInternal(FPCGContext* InContext) const
 	check(Context);
 	const UPCGClusterSettings* Settings = Context->GetInputSettings<UPCGClusterSettings>();
 	check(Settings);
-	const TArray<FPCGTaggedData> Inputs = Context->InputData.GetInputsByPin(PCGPinConstants::DefaultInputLabel);
+
+	const TArray<FPCGTaggedData> PointInputs = Context->InputData.GetInputsByPin(PCGPinConstants::DefaultInputLabel);
 
 	// Early out if no source to modify
-	if (Inputs.IsEmpty())
+	if (PointInputs.IsEmpty())
 	{
 		return true;
 	}
 
-	Context->InitializePerExecutionState([Settings](const ContextType* Context, ExecStateType& OutState)
+	const TArray<FPCGTaggedData> CentroidInputs = Context->InputData.GetInputsByPin(PCGClusterElement::Constants::InputCentroidsLabel);
+
+	Context->InitializePerExecutionState([&CentroidInputs, Settings](const ContextType* Context, ExecStateType& OutState)
 	{
 		TRACE_CPUPROFILER_EVENT_SCOPE(FPCGClusterElement::PrepareData::InitializePerExecutionState);
 
@@ -433,56 +449,41 @@ bool FPCGClusterElement::PrepareDataInternal(FPCGContext* InContext) const
 		}
 		check(OutState.ClusteringFunction);
 
-		switch (Settings->InitialCentroidSelection)
+		if (CentroidInputs.IsEmpty())
 		{
-			case EPCGInitialClusteringCentroidSelection::RandomPoints:
+			OutState.bSelectRandomCentroids = true;
+			// Centroids will be initialized on a per-iteration input basis, but we can allocate them here.
+			OutState.InitialCentroids.SetNum(Settings->NumClusters);
+		}
+		else
+		{
+			OutState.bSelectRandomCentroids = false;
+			if (CentroidInputs.Num() > 1)
 			{
-				// Centroids will be initialized on a per-iteration input basis, but we can allocate them here.
-				OutState.InitialCentroids.SetNum(Settings->NumClusters);
-				break;
+				PCGLog::InputOutput::LogFirstInputOnlyWarning(PCGClusterElement::Constants::InputCentroidsLabel, Context);
 			}
 
-			case EPCGInitialClusteringCentroidSelection::Input:
+			const UPCGPointData* CentroidPoints = Cast<UPCGPointData>(CentroidInputs[0].Data);
+			if (!CentroidPoints || CentroidPoints->IsEmpty())
 			{
-				const TArray<FPCGTaggedData> CentroidInputs = Context->InputData.GetInputsByPin(PCGClusterElement::Constants::InputCentroidsLabel);
-				if (CentroidInputs.IsEmpty())
-				{
-					return EPCGTimeSliceInitResult::NoOperation;
-				}
-
-				if (CentroidInputs.Num() > 1)
-				{
-					PCGLog::InputOutput::LogFirstInputOnlyWarning(PCGClusterElement::Constants::InputCentroidsLabel, Context);
-				}
-
-				const UPCGPointData* CentroidPoints = Cast<UPCGPointData>(CentroidInputs[0].Data);
-				if (!CentroidPoints || CentroidPoints->IsEmpty())
-				{
-					return EPCGTimeSliceInitResult::NoOperation;
-				}
-
-				OutState.InitialCentroids.Reserve(CentroidPoints->GetPoints().Num());
-				for (const FPCGPoint& Point : CentroidPoints->GetPoints())
-				{
-					OutState.InitialCentroids.Emplace(Point.Transform.GetLocation());
-				}
-				break;
+				return EPCGTimeSliceInitResult::NoOperation;
 			}
 
-			default:
-				checkNoEntry();
-				return EPCGTimeSliceInitResult::AbortExecution;
+			OutState.InitialCentroids.Reserve(CentroidPoints->GetPoints().Num());
+			for (const FPCGPoint& Point : CentroidPoints->GetPoints())
+			{
+				OutState.InitialCentroids.Emplace(Point.Transform.GetLocation());
+			}
 		}
 
 		return EPCGTimeSliceInitResult::Success;
 	});
 
-	auto InitializeIteration = [&Inputs, Settings, Context](IterStateType& OutState, const ExecStateType& ExecState, const uint32 IterationIndex)
+	auto InitializeIteration = [&PointInputs, Settings, Context](IterStateType& OutState, const ExecStateType& ExecState, const uint32 IterationIndex)
 	{
 		TRACE_CPUPROFILER_EVENT_SCOPE(FPCGClusterElement::PrepareData::InitializeIteration);
 
-		const UPCGPointData* InputPointData = CastChecked<UPCGPointData>(Inputs[IterationIndex].Data);
-
+		const UPCGPointData* InputPointData = CastChecked<UPCGPointData>(PointInputs[IterationIndex].Data);
 		if (!InputPointData || InputPointData->IsEmpty())
 		{
 			return EPCGTimeSliceInitResult::NoOperation;
@@ -540,7 +541,7 @@ bool FPCGClusterElement::PrepareDataInternal(FPCGContext* InContext) const
 
 		auto InitializeCentroids = [Settings, NumClusters, NumPoints, &OutState, &ExecState]
 		{
-			if (Settings->InitialCentroidSelection == EPCGInitialClusteringCentroidSelection::RandomPoints)
+			if (ExecState.bSelectRandomCentroids)
 			{
 				const TArray<FPCGPoint>& Points = OutState.OutPointData->GetPoints();
 
@@ -589,24 +590,24 @@ bool FPCGClusterElement::PrepareDataInternal(FPCGContext* InContext) const
 					MaxSquaredLength = FMath::Max(MaxSquaredLength, Point.Transform.GetLocation().SquaredLength());
 				}
 
-				double ScalingFactor = 1.0;
 				// Only normalize the system if needed.
+				OutState.EMData.ScalingFactor = 1.0;
 				if (MaxSquaredLength > 1.0)
 				{
 					const double Length = FMath::Sqrt(MaxSquaredLength);
-					ScalingFactor = UE_DOUBLE_EULERS_NUMBER / Length;
+					OutState.EMData.ScalingFactor = UE_DOUBLE_EULERS_NUMBER / Length;
 				}
 
 				OutState.EMData.Gaussians.Reserve(NumClusters);
 				for (int i = 0; i < NumClusters; ++i)
 				{
-					OutState.EMData.Gaussians.Emplace(OutState.Centroids[i] * ScalingFactor, FMatrix::Identity, 1.0 / NumClusters);
+					OutState.EMData.Gaussians.Emplace(OutState.Centroids[i] * OutState.EMData.ScalingFactor, FMatrix::Identity, 1.0 / NumClusters);
 				}
 
 				OutState.EMData.Probabilities.Initialize(Points.Num(), NumClusters);
 
 				OutState.EMData.NormalizedPointLocations.Reserve(Points.Num());
-				Algo::Transform(Points, OutState.EMData.NormalizedPointLocations, [ScalingFactor](const FPCGPoint& Point)
+				Algo::Transform(Points, OutState.EMData.NormalizedPointLocations, [ScalingFactor = OutState.EMData.ScalingFactor](const FPCGPoint& Point)
 				{
 					return Point.Transform.GetLocation() * ScalingFactor;
 				});
@@ -621,7 +622,7 @@ bool FPCGClusterElement::PrepareDataInternal(FPCGContext* InContext) const
 		return EPCGTimeSliceInitResult::Success;
 	};
 
-	Context->InitializePerIterationStates(Inputs.Num(), InitializeIteration);
+	Context->InitializePerIterationStates(PointInputs.Num(), InitializeIteration);
 
 	return true;
 }
@@ -647,25 +648,60 @@ bool FPCGClusterElement::ExecuteInternal(FPCGContext* InContext) const
 	const UPCGClusterSettings* Settings = TimeSlicedContext->GetInputSettings<UPCGClusterSettings>();
 	check(Settings);
 
-	return ExecuteSlice(TimeSlicedContext, [MaxIterations = Settings->MaxIterations](ContextType* Context, const ExecStateType& ExecState, IterStateType& ClusteringData, const uint32 IterIndex)
+	return ExecuteSlice(TimeSlicedContext, [Settings](ContextType* Context, const ExecStateType& ExecState, IterStateType& ClusteringData, const uint32 IterIndex)
 	{
 		if (Context->GetIterationStateResult(IterIndex) == EPCGTimeSliceInitResult::NoOperation)
 		{
 			return true;
 		}
 
-		if (!ExecState.ClusteringFunction(Context, ClusteringData, MaxIterations))
+		if (!ExecState.ClusteringFunction(Context, ClusteringData, Settings->MaxIterations))
 		{
 			return false;
 		}
 
 		// Once finished, set the attribute values.
 		check(ClusteringData.OutPointData);
-		Context->OutputData.TaggedData.Emplace_GetRef().Data = ClusteringData.OutPointData;
+		FPCGTaggedData& OutputData = Context->OutputData.TaggedData.Emplace_GetRef();
+		OutputData.Pin = PCGPinConstants::DefaultOutputLabel;
+		OutputData.Data = ClusteringData.OutPointData;
 
-		// Finally, update the attribute
+		// Update the cluster assignment attribute.
 		ClusteringData.Accessor->SetRange(MakeConstArrayView(ClusteringData.PointToClusterAssignments), 0, *ClusteringData.Keys);
 
+		// If optionally outputting the final centroids, unscale the final gaussians for EM or use the centroids directly in K-Means.
+		if (Settings->bOutputFinalCentroids)
+		{
+			check(!ClusteringData.Centroids.IsEmpty());
+
+			TArray<FVector> OutCentroids;
+			if (Settings->Algorithm == EPCGClusterAlgorithm::EM)
+			{
+				OutCentroids.Reserve(ClusteringData.EMData.Gaussians.Num());
+				using EM = PCGClusterElement::FClusteringData::FExpectationMaximizationData;
+				for (const EM::FGaussian& Gaussian : ClusteringData.EMData.Gaussians)
+				{
+					if (!Gaussian.bDead)
+					{
+						OutCentroids.Emplace(Gaussian.Mean * 1.0 / ClusteringData.EMData.ScalingFactor);
+					}
+				}
+			}
+			else
+			{
+				OutCentroids = std::move(ClusteringData.Centroids);
+			}
+
+			UPCGPointData* OutputCentroidPointData = FPCGContext::NewObject_AnyThread<UPCGPointData>(Context);
+			TArray<FPCGPoint>& OutputCentroidPoints = OutputCentroidPointData->GetMutablePoints();
+			OutputCentroidPoints.Reserve(OutCentroids.Num());
+			Algo::Transform(OutCentroids, OutputCentroidPoints, [](const FVector& Location) { return FPCGPoint(FTransform(Location), 1.0, PCGHelpers::ComputeSeedFromPosition(Location)); });
+
+			FPCGTaggedData& OutputCentroidData = Context->OutputData.TaggedData.Emplace_GetRef();
+			OutputCentroidData.Pin = PCGClusterElement::Constants::FinalCentroidsLabel;
+			OutputCentroidData.Data = OutputCentroidPointData;
+		}
+
 		return true;
 	});
 }
diff --git a/Engine/Plugins/PCG/Source/PCG/Public/Elements/PCGClusterElement.h b/Engine/Plugins/PCG/Source/PCG/Public/Elements/PCGClusterElement.h
--- a/Engine/Plugins/PCG/Source/PCG/Public/Elements/PCGClusterElement.h
+++ b/Engine/Plugins/PCG/Source/PCG/Public/Elements/PCGClusterElement.h
@@ -18,13 +18,6 @@ enum class EPCGClusterAlgorithm
 	EM UMETA(Tooltip = "Expectation-Maximization - Categorizes points by using a Gaussian Mixture Model.")
 };
 
-UENUM()
-enum class EPCGInitialClusteringCentroidSelection
-{
-	RandomPoints UMETA(Tooltip = "Random points from the input will be selected as the starting centroids."),
-	Input UMETA(Tooltip = "Point locations from the secondary pin will be used as the starting centroid locations."),
-};
-
 /** Given a desired number of clusters (categories), find the best fit cluster for each point by distance, using one of various clustering algorithms. */
 UCLASS(MinimalAPI, BlueprintType, ClassGroup = (Procedural))
 class UPCGClusterSettings : public UPCGSettings
@@ -43,7 +36,7 @@ public:
 
 protected:
 	virtual TArray<FPCGPinProperties> InputPinProperties() const override;
-	virtual TArray<FPCGPinProperties> OutputPinProperties() const override { return Super::DefaultPointOutputPinProperties(); }
+	virtual TArray<FPCGPinProperties> OutputPinProperties() const override;
 	virtual FPCGElementPtr CreateElement() const override;
 	//~End UPCGSettings interface
 
@@ -52,12 +45,8 @@ public:
 	UPROPERTY(EditAnywhere, BlueprintReadWrite, Category = Settings, meta = (PCG_Overridable))
 	EPCGClusterAlgorithm Algorithm = EPCGClusterAlgorithm::KMeans;
 
-	/** Method for selecting the initial average location of the clusters. A good estimate of these locations can rapidly speed up the process in some cases. */
-	UPROPERTY(EditAnywhere, BlueprintReadWrite, Category = Settings)
-	EPCGInitialClusteringCentroidSelection InitialCentroidSelection = EPCGInitialClusteringCentroidSelection::RandomPoints;
-
 	/** Number of clusters (segments) to group the points into. Each point will be assigned a cluster at the end. */
-	UPROPERTY(EditAnywhere, BlueprintReadWrite, Category = Settings, DisplayName = "Clusters", meta = (ClampMin = "1", EditCondition = "InitialCentroidSelection == EPCGInitialClusteringCentroidSelection::RandomPoints", EditConditionHides, PCG_Overridable))
+	UPROPERTY(EditAnywhere, BlueprintReadWrite, Category = Settings, DisplayName = "Clusters", meta = (ClampMin = "1", PCG_Overridable))
 	int32 NumClusters = 3;
 
 	/** Cluster IDs will be written to this attribute on the output. */
@@ -73,6 +62,10 @@ public:
 	 */
 	UPROPERTY(EditAnywhere, BlueprintReadWrite, Category = "Settings|Advanced", meta = (ClampMin = "0.0001", EditCondition = "Algorithm == EPCGClusterAlgorithm::EM", EditConditionHides, PCG_Overridable))
 	double Tolerance = UE_DOUBLE_KINDA_SMALL_NUMBER;
+
+	/** Output the final location of the centroids or gaussians. */
+	UPROPERTY(EditAnywhere, BlueprintReadWrite, Category = "Settings|Advanced")
+	bool bOutputFinalCentroids = false;
 };
 
 namespace PCGClusterElement
@@ -136,6 +129,7 @@ namespace PCGClusterElement
 			FProbabilityArray Probabilities;
 
 			TArray<FVector> NormalizedPointLocations;
+			double ScalingFactor = 1.0;
 			double Tolerance = UE_DOUBLE_KINDA_SMALL_NUMBER;
 			// A value from -inf to some positive finite value to represent how well the model fits the data. It will be used to recognize convergence.
 			std::atomic<double> LogLikelihoodSum = 0;
@@ -149,6 +143,8 @@ namespace PCGClusterElement
 	{
 		using AlgoSignature = bool(*)(FPCGContext* Context, FClusteringData& ClusteringData, const int MaxIterations);
 		AlgoSignature ClusteringFunction = nullptr;
+
+		bool bSelectRandomCentroids = true;
 		TArray<FVector> InitialCentroids;
 	};
 }
