Fixed a few outstanding bugs in the save data element
- Asset description and color were not properly on the resulting asset
- Overrides (such as name, path..) would be passed down to the asset erroneously
- Improved handling of invalid paths so it's much less likely to crash

Light refactoring in the asset exporter utils to facilitate error reporting when used in a graph.

diff --git a/Engine/Plugins/PCG/Source/PCG/Private/Elements/IO/PCGSaveAssetElement.cpp b/Engine/Plugins/PCG/Source/PCG/Private/Elements/IO/PCGSaveAssetElement.cpp
--- a/Engine/Plugins/PCG/Source/PCG/Private/Elements/IO/PCGSaveAssetElement.cpp
+++ b/Engine/Plugins/PCG/Source/PCG/Private/Elements/IO/PCGSaveAssetElement.cpp
@@ -11,6 +11,10 @@ bool UPCGDataCollectionExporter::ExportAsset(const FString& PackageName, UPCGDat
 {
 	// Relies on default behavior to duplicate if needed
 	Asset->Data = Data;
+#if WITH_EDITOR
+	Asset->Description = FText::FromString(AssetDescription);
+	Asset->Color = AssetColor;
+#endif
 	return true;
 }
 
@@ -48,10 +52,25 @@ bool FPCGSaveDataAssetElement::ExecuteInternal(FPCGContext* Context) const
 		Exporter = NewObject<UPCGDataCollectionExporter>();
 	}
 
+	TArray<FPCGPinProperties> InputPins = Settings->InputPinProperties();
+
 	check(Exporter);
-	Exporter->Data = Context->InputData;
+	// Implementation note: we can't simply copy the input wholesale because this will also gather overrides if any.
+	for (const FPCGTaggedData& TaggedData : Context->InputData.TaggedData)
+	{
+		// Accept data if its pin matches any of the actual input pins
+		if (InputPins.FindByPredicate([&TaggedData](const FPCGPinProperties& Pin) { return TaggedData.Pin == Pin.Label; }))
+		{
+			Exporter->Data.TaggedData.Add(TaggedData);
+		}
+	}
+
+#if WITH_EDITOR
+	Exporter->AssetDescription = Settings->AssetDescription;
+	Exporter->AssetColor = Settings->AssetColor;
+#endif
 
-	UPCGAssetExporterUtils::CreateAsset(Exporter, Settings->Params);
+	UPCGAssetExporterUtils::CreateAsset(Exporter, Settings->Params, Context);
 
 	return true;
 }
\ No newline at end of file
diff --git a/Engine/Plugins/PCG/Source/PCG/Private/PCGAssetExporterUtils.cpp b/Engine/Plugins/PCG/Source/PCG/Private/PCGAssetExporterUtils.cpp
--- a/Engine/Plugins/PCG/Source/PCG/Private/PCGAssetExporterUtils.cpp
+++ b/Engine/Plugins/PCG/Source/PCG/Private/PCGAssetExporterUtils.cpp
@@ -2,6 +2,7 @@
 
 #include "PCGAssetExporterUtils.h"
 
+#include "PCGContext.h"
 #include "PCGModule.h"
 
 #if WITH_EDITOR
@@ -11,12 +12,19 @@
 #include "AssetRegistry/AssetRegistryModule.h"
 #endif // WITH_EDITOR
 
+#define LOCTEXT_NAMESPACE "PCGAssetExporterUtils"
+
 UPackage* UPCGAssetExporterUtils::CreateAsset(UPCGAssetExporter* Exporter, FPCGAssetExporterParameters Parameters)
+{
+	return CreateAsset(Exporter, Parameters, nullptr);
+}
+
+UPackage* UPCGAssetExporterUtils::CreateAsset(UPCGAssetExporter* Exporter, const FPCGAssetExporterParameters& Parameters, FPCGContext* Context)
 {
 #if WITH_EDITOR
 	if (!Exporter || !Exporter->GetAssetType())
 	{
-		UE_LOG(LogPCG, Error, TEXT("Unable to create asset without an exporter, or exporter is not setup properly."));
+		PCGLog::LogErrorOnGraph(LOCTEXT("MissingExporter", "Unable to create asset without an exporter, or exporter is not setup properly."), Context);
 		return nullptr;
 	}
 
@@ -60,6 +68,16 @@ UPackage* UPCGAssetExporterUtils::CreateAsset(UPCGAssetExporter* Exporter, FPCGA
 			return nullptr;
 		}
 	}
+	else
+	{
+		// Perform some validation on the package name, so we can prevent crashes downstream when trying to create or save the package.
+		FText Reason;
+		if (!FPackageName::IsValidObjectPath(PackageName, &Reason))
+		{
+			PCGLog::LogErrorOnGraph(FText::Format(LOCTEXT("InvalidPackagePath", "Invalid package path '{0}': {1}."), { FText::FromString(PackageName), Reason }), Context);
+			return nullptr;
+		}
+	}
 
 	UPackage* Package = FPackageName::DoesPackageExist(PackageName) ? LoadPackage(nullptr, *PackageName, LOAD_None) : nullptr;
 
@@ -83,7 +101,16 @@ UPackage* UPCGAssetExporterUtils::CreateAsset(UPCGAssetExporter* Exporter, FPCGA
 	else
 	{
 		Package = CreatePackage(*PackageName);
-		NewAssetCreated = true;
+
+		if (Package)
+		{
+			NewAssetCreated = true;
+		}
+		else
+		{
+			PCGLog::LogErrorOnGraph(FText::Format(LOCTEXT("UnableToCreatePackage", "Unable to create the package with name '{0}'."), FText::FromString(PackageName)), Context);
+			return nullptr;
+		}
 	}
 
 	if (!Asset)
@@ -122,12 +149,17 @@ UPackage* UPCGAssetExporterUtils::CreateAsset(UPCGAssetExporter* Exporter, FPCGA
 
 	return Package;
 #else
-	UE_LOG(LogPCG, Error, TEXT("PCG Asset Exporter Utils cannot be used in non-editor builds."));
+	PCGLog::LogErrorOnGraph(LOCTEXT("CannotExportInNonEditor", "PCG Asset Exporter Utils cannot be used in non-editor builds."), Context);
 	return nullptr;
 #endif
 }
 
 void UPCGAssetExporterUtils::UpdateAssets(const TArray<FAssetData>& PCGAssets, FPCGAssetExporterParameters InParameters)
+{
+	UpdateAssets(PCGAssets, InParameters, nullptr);
+}
+
+void UPCGAssetExporterUtils::UpdateAssets(const TArray<FAssetData>& PCGAssets, const FPCGAssetExporterParameters& InParameters, FPCGContext* Context)
 {
 #if WITH_EDITOR
 	TArray<UPackage*> PackagesToSave;
@@ -148,7 +180,7 @@ void UPCGAssetExporterUtils::UpdateAssets(const TArray<FAssetData>& PCGAssets, F
 		}
 		else
 		{
-			UE_LOG(LogPCG, Error, TEXT("Unable to update asset '%s' because exporter isn't valid."), *PCGAsset.AssetName.ToString());
+			PCGLog::LogErrorOnGraph(FText::Format(LOCTEXT("InvalidExporterOnAsset", "Unable to update asset '{0}' because exporter isn't valid."), FText::FromName(PCGAsset.AssetName)), Context);
 			continue;
 		}
 
@@ -156,7 +188,7 @@ void UPCGAssetExporterUtils::UpdateAssets(const TArray<FAssetData>& PCGAssets, F
 
 		if (!Exporter)
 		{
-			UE_LOG(LogPCG, Error, TEXT("Unable to create exporter for asset '%s' during update process."), *PCGAsset.AssetName.ToString());
+			PCGLog::LogErrorOnGraph(FText::Format(LOCTEXT("UnableToCreateExporter", "Unable to create exporter for asset '{0}' during update process."), FText::FromName(PCGAsset.AssetName)), Context);
 			continue;
 		}
 
@@ -177,6 +209,8 @@ void UPCGAssetExporterUtils::UpdateAssets(const TArray<FAssetData>& PCGAssets, F
 		FEditorFileUtils::PromptForCheckoutAndSave(PackagesToSave, /*bCheckDirty=*/false, /*bPromptToSave=*/false);
 	}
 #else
-	UE_LOG(LogPCG, Error, TEXT("PCG Asset Exporter Utils cannot be used in non-editor builds."));
+	PCGLog::LogErrorOnGraph(LOCTEXT("CannotExportInNonEditor", "PCG Asset Exporter Utils cannot be used in non-editor builds."), Context);
 #endif // WITH_EDITOR
-}
\ No newline at end of file
+}
+
+#undef LOCTEXT_NAMESPACE
\ No newline at end of file
diff --git a/Engine/Plugins/PCG/Source/PCG/Public/Elements/IO/PCGSaveAssetElement.h b/Engine/Plugins/PCG/Source/PCG/Public/Elements/IO/PCGSaveAssetElement.h
--- a/Engine/Plugins/PCG/Source/PCG/Public/Elements/IO/PCGSaveAssetElement.h
+++ b/Engine/Plugins/PCG/Source/PCG/Public/Elements/IO/PCGSaveAssetElement.h
@@ -26,6 +26,14 @@ protected:
 public:
 	UPROPERTY(BlueprintReadWrite, VisibleAnywhere, Category = Data)
 	FPCGDataCollection Data;
+
+#if WITH_EDITORONLY_DATA
+	UPROPERTY(BlueprintReadWrite, VisibleAnywhere, Category = Data)
+	FString AssetDescription;
+
+	UPROPERTY(BlueprintReadWrite, VisibleAnywhere, Category = Data)
+	FLinearColor AssetColor = FLinearColor::White;
+#endif
 };
 
 /** Node that will save input data to a PCG data asset */
@@ -45,11 +53,11 @@ public:
 	virtual EPCGSettingsType GetType() const override { return EPCGSettingsType::InputOutput; }
 #endif // WITH_EDITOR
 
-protected:
-	virtual FPCGElementPtr CreateElement() const override;
-	
 	virtual TArray<FPCGPinProperties> InputPinProperties() const override { return Pins; }
 	virtual TArray<FPCGPinProperties> OutputPinProperties() const override { return TArray<FPCGPinProperties>(); }
+
+protected:
+	virtual FPCGElementPtr CreateElement() const override;
 	// ~End UPCGSettings interface
 
 public:
diff --git a/Engine/Plugins/PCG/Source/PCG/Public/PCGAssetExporterUtils.h b/Engine/Plugins/PCG/Source/PCG/Public/PCGAssetExporterUtils.h
--- a/Engine/Plugins/PCG/Source/PCG/Public/PCGAssetExporterUtils.h
+++ b/Engine/Plugins/PCG/Source/PCG/Public/PCGAssetExporterUtils.h
@@ -8,6 +8,8 @@
 
 #include "PCGAssetExporterUtils.generated.h"
 
+struct FPCGContext;
+
 /**
 * Asset export utils - will work only in editor builds. 
 */
@@ -17,9 +19,15 @@ class PCG_API UPCGAssetExporterUtils : public UBlueprintFunctionLibrary
 	GENERATED_BODY()
 
 public:
+	/** Exports an asset from the given exporter. When calling this function, the exporter should be able to create its own data and as such needs prior information (world, etc.) before calling this method. */
 	UFUNCTION(BlueprintCallable, Category = "PCG|IO")
 	static UPackage* CreateAsset(UPCGAssetExporter* Exporter, FPCGAssetExporterParameters Parameters = FPCGAssetExporterParameters());
 
+	/** Updates assets based on their embedded exporter & metadata. */
 	UFUNCTION(BlueprintCallable, Category = "PCG|IO")
 	static void UpdateAssets(const TArray<FAssetData>& PCGAssets, FPCGAssetExporterParameters Parameters = FPCGAssetExporterParameters());
+
+	// Native API equivalent
+	static UPackage* CreateAsset(UPCGAssetExporter* Exporter, const FPCGAssetExporterParameters& Parameters, FPCGContext* InOptionalContext);
+	static void UpdateAssets(const TArray<FAssetData>& PCGAssets, const FPCGAssetExporterParameters& Parameters, FPCGContext* InOptionalContext);
 };
\ No newline at end of file
