Revisited changes for the editor perf improvements to deprecate old subsystem delegates instead of replacing them.

#rnx

diff --git a/Engine/Plugins/PCG/Source/PCG/Private/PCGComponent.cpp b/Engine/Plugins/PCG/Source/PCG/Private/PCGComponent.cpp
--- a/Engine/Plugins/PCG/Source/PCG/Private/PCGComponent.cpp
+++ b/Engine/Plugins/PCG/Source/PCG/Private/PCGComponent.cpp
@@ -650,7 +650,7 @@ void UPCGComponent::PostProcessGraph(const FBox& InNewBounds, bool bInGenerated,
 	StopGenerationInProgress();
 
 	UPCGSubsystem* Subsystem = GetSubsystem();
-	Subsystem->OnComponentGenerationCompleteOrCancelled.Broadcast(Subsystem, IsValid(this) ? this : nullptr);
+	Subsystem->OnPCGComponentGenerationDone.Broadcast(Subsystem, IsValid(this) ? this : nullptr, EPCGGenerationStatus::Completed);
 #endif
 }
 
@@ -758,7 +758,7 @@ void UPCGComponent::OnProcessGraphAborted(bool bQuiet)
 	OnPCGGraphCancelledDelegate.Broadcast(this);
 
 	UPCGSubsystem* Subsystem = GetSubsystem();
-	Subsystem->OnComponentGenerationCompleteOrCancelled.Broadcast(Subsystem, this);
+	Subsystem->OnPCGComponentGenerationDone.Broadcast(Subsystem, this, EPCGGenerationStatus::Aborted);
 #endif
 
 	PCGComponent::BroadcastDynamicDelegate(OnPCGGraphCancelledExternal, this);
@@ -1555,7 +1555,7 @@ void UPCGComponent::OnUnregister()
 			Subsystem->CancelGeneration(this);
 		}
 
-		Subsystem->OnComponentUnregistered.Broadcast(IsValid(this) ? this : nullptr);
+		Subsystem->OnPCGComponentUnregistered.Broadcast(IsValid(this) ? this : nullptr);
 	}
 #endif // WITH_EDITOR
 
diff --git a/Engine/Plugins/PCG/Source/PCG/Public/PCGCommon.h b/Engine/Plugins/PCG/Source/PCG/Public/PCGCommon.h
--- a/Engine/Plugins/PCG/Source/PCG/Public/PCGCommon.h
+++ b/Engine/Plugins/PCG/Source/PCG/Public/PCGCommon.h
@@ -407,4 +407,11 @@ enum class EPCGDensityMergeOperation : uint8
 	Multiply,
 	/** D = A / B */
 	Divide
+};
+
+UENUM(BlueprintType)
+enum class EPCGGenerationStatus : uint8
+{
+	Completed,
+	Aborted
 };
\ No newline at end of file
diff --git a/Engine/Plugins/PCG/Source/PCG/Public/PCGSubsystem.h b/Engine/Plugins/PCG/Source/PCG/Public/PCGSubsystem.h
--- a/Engine/Plugins/PCG/Source/PCG/Public/PCGSubsystem.h
+++ b/Engine/Plugins/PCG/Source/PCG/Public/PCGSubsystem.h
@@ -36,8 +36,13 @@ typedef TSharedPtr<IPCGElement, ESPMode::ThreadSafe> FPCGElementPtr;
 class UWorld;
 
 #if WITH_EDITOR
-DECLARE_MULTICAST_DELEGATE_OneParam(FPCGOnComponentUnregistered, UPCGComponent*);
-DECLARE_MULTICAST_DELEGATE_TwoParams(FPCGOnComponentGenerationCompleteOrCancelled, UPCGSubsystem*, UPCGComponent*);
+/** Deprecated - use FPCGOnPCGComponentUnregistered */
+DECLARE_MULTICAST_DELEGATE(FPCGOnComponentUnregistered);
+/** Deprecated - use FPCGOnPCGComponentGenerationDone */
+DECLARE_MULTICAST_DELEGATE_OneParam(FPCGOnComponentGenerationCompleteOrCancelled, UPCGSubsystem*);
+
+DECLARE_MULTICAST_DELEGATE_OneParam(FPCGOnPCGComponentUnregistered, UPCGComponent*);
+DECLARE_MULTICAST_DELEGATE_ThreeParams(FPCGOnPCGComponentGenerationDone, UPCGSubsystem*, UPCGComponent*, EPCGGenerationStatus);
 #endif // WITH_EDITOR
 
 /**
@@ -327,9 +332,15 @@ public:
 	void ClearExecutedStacks(FPCGStack BeginningWithStack);
 	void ClearExecutedStacks(const UPCGGraph* InContainingGraph);
 
+	UE_DEPRECATED(5.5, "Deprecated in favor of OnPCGComponentUnregistered, will not be notified anymore")
 	FPCGOnComponentUnregistered OnComponentUnregistered;
+
+	UE_DEPRECATED(5.5, "Deprecated in favor of OnPCGComponentGenerationDone, will not be notified anymore")
 	FPCGOnComponentGenerationCompleteOrCancelled OnComponentGenerationCompleteOrCancelled;
 
+	FPCGOnPCGComponentUnregistered OnPCGComponentUnregistered;
+	FPCGOnPCGComponentGenerationDone OnPCGComponentGenerationDone;
+
 	void CreateMissingPartitionActors();
 private:
 	enum class EOperation : uint32
diff --git a/Engine/Plugins/PCG/Source/PCGEditor/Private/PCGEditor.cpp b/Engine/Plugins/PCG/Source/PCGEditor/Private/PCGEditor.cpp
--- a/Engine/Plugins/PCG/Source/PCGEditor/Private/PCGEditor.cpp
+++ b/Engine/Plugins/PCG/Source/PCGEditor/Private/PCGEditor.cpp
@@ -3071,7 +3071,7 @@ void FPCGEditor::OnComponentUnregistered(UPCGComponent* Component)
 	}
 }
 
-void FPCGEditor::OnComponentGenerationCompleteOrCancelled(UPCGSubsystem* Subsystem, UPCGComponent* Component)
+void FPCGEditor::OnComponentGenerationDone(UPCGSubsystem* Subsystem, UPCGComponent* Component, EPCGGenerationStatus Status)
 {
 	if(Component && Component->GetGraph() != PCGGraphBeingEdited)
 	{
@@ -3114,8 +3114,8 @@ void FPCGEditor::RegisterDelegatesForWorld(UWorld* World)
 
 	if (UPCGSubsystem* Subsystem = UPCGSubsystem::GetInstance(World))
 	{
-		Subsystem->OnComponentUnregistered.AddRaw(this, &FPCGEditor::OnComponentUnregistered);
-		Subsystem->OnComponentGenerationCompleteOrCancelled.AddRaw(this, &FPCGEditor::OnComponentGenerationCompleteOrCancelled);
+		Subsystem->OnPCGComponentUnregistered.AddRaw(this, &FPCGEditor::OnComponentUnregistered);
+		Subsystem->OnPCGComponentGenerationDone.AddRaw(this, &FPCGEditor::OnComponentGenerationDone);
 	}
 }
 
@@ -3123,8 +3123,8 @@ void FPCGEditor::UnregisterDelegatesForWorld(UWorld* World)
 {
 	if (UPCGSubsystem* Subsystem = UPCGSubsystem::GetInstance(World))
 	{
-		Subsystem->OnComponentUnregistered.RemoveAll(this);
-		Subsystem->OnComponentGenerationCompleteOrCancelled.RemoveAll(this);
+		Subsystem->OnPCGComponentUnregistered.RemoveAll(this);
+		Subsystem->OnPCGComponentGenerationDone.RemoveAll(this);
 	}
 }
 
diff --git a/Engine/Plugins/PCG/Source/PCGEditor/Private/PCGEditor.h b/Engine/Plugins/PCG/Source/PCGEditor/Private/PCGEditor.h
--- a/Engine/Plugins/PCG/Source/PCGEditor/Private/PCGEditor.h
+++ b/Engine/Plugins/PCG/Source/PCGEditor/Private/PCGEditor.h
@@ -337,7 +337,7 @@ private:
 	void OnComponentUnregistered(UPCGComponent* Component);
 
 	/** Called when a component finishes executing. Useful for updating debugging tools/UIs. */
-	void OnComponentGenerationCompleteOrCancelled(UPCGSubsystem* Subsystem, UPCGComponent* Component);
+	void OnComponentGenerationDone(UPCGSubsystem* Subsystem, UPCGComponent* Component, EPCGGenerationStatus Status);
 
 	/** Trigger any generation required to ensure debug display is up to date. */
 	void UpdateDebugAfterComponentSelection(UPCGComponent* InOldComponent, UPCGComponent* InNewComponent, bool bNewComponentStartedInspecting);
