Added support to load multiple assets in the Load PCG Data asset node.
Fixed graph compilation/culling issue when switching from/to FromInput in the get actor prop/get actor data nodes.

diff --git a/Engine/Plugins/PCG/Source/PCG/Private/Elements/IO/PCGLoadAssetElement.cpp b/Engine/Plugins/PCG/Source/PCG/Private/Elements/IO/PCGLoadAssetElement.cpp
--- a/Engine/Plugins/PCG/Source/PCG/Private/Elements/IO/PCGLoadAssetElement.cpp
+++ b/Engine/Plugins/PCG/Source/PCG/Private/Elements/IO/PCGLoadAssetElement.cpp
@@ -35,6 +35,18 @@ void UPCGLoadDataAssetSettings::PostEditChangeProperty(FPropertyChangedEvent& Pr
 
 	Super::PostEditChangeProperty(PropertyChangedEvent);
 }
+
+EPCGChangeType UPCGLoadDataAssetSettings::GetChangeTypeForProperty(const FName& InPropertyName) const
+{
+	EPCGChangeType ChangeType = Super::GetChangeTypeForProperty(InPropertyName) | EPCGChangeType::Cosmetic;
+
+	if (InPropertyName == GET_MEMBER_NAME_CHECKED(UPCGLoadDataAssetSettings, bLoadFromInput))
+	{
+		ChangeType |= EPCGChangeType::Structural;
+	}
+
+	return ChangeType;
+}
 #endif // WITH_EDITOR
 
 FPCGElementPtr UPCGLoadDataAssetSettings::CreateElement() const
@@ -42,9 +54,40 @@ FPCGElementPtr UPCGLoadDataAssetSettings::CreateElement() const
 	return MakeShared<FPCGLoadDataAssetElement>();
 }
 
+TArray<FPCGPinProperties> UPCGLoadDataAssetSettings::InputPinProperties() const
+{
+	TArray<FPCGPinProperties> PinProperties;
+
+	if (bLoadFromInput)
+	{
+		FPCGPinProperties& InputPin = PinProperties.Emplace_GetRef(PCGPinConstants::DefaultInputLabel, EPCGDataType::Param);
+		InputPin.SetRequiredPin();
+	}
+	
+	return PinProperties;
+}
+
 FString UPCGLoadDataAssetSettings::GetAdditionalTitleInformation() const
 {
-	return AssetName.IsEmpty() ? Asset.ToSoftObjectPath().GetAssetName() : AssetName;
+#if WITH_EDITOR
+	if (bLoadFromInput || IsPropertyOverriddenByPin(GET_MEMBER_NAME_CHECKED(UPCGLoadDataAssetSettings, Asset)))
+	{
+		// If loading data from a specified input or from an overriden value, we shouldn't show the template asset name.
+		return FString();
+	}
+	else
+#endif // WITH_EDITOR
+	{
+		return AssetName.IsEmpty() ? Asset.ToSoftObjectPath().GetAssetName() : AssetName;
+	}
+}
+
+EPCGDataType UPCGLoadDataAssetSettings::GetCurrentPinTypes(const UPCGPin* InPin) const
+{
+	// Implementation notes: the output pin types don't depend on the input pin types,
+	// but they can change based on the asset selected, hence why they are dynamic, but we need to return the pin type as-is.
+	check(InPin);
+	return InPin->Properties.AllowedTypes;
 }
 
 void UPCGLoadDataAssetSettings::SetFromAsset(const FAssetData& InAsset)
@@ -123,57 +166,71 @@ bool FPCGLoadDataAssetElement::PrepareDataInternal(FPCGContext* InContext) const
 	const UPCGLoadDataAssetSettings* Settings = Context->GetInputSettings<UPCGLoadDataAssetSettings>();
 	check(Settings);
 
-	if (Settings->Asset.IsNull())
-	{
-		return true;
-	}
-
-	// Request load, return false if we need to wait, otherwise continue
-	if (!Context->WasLoadRequested())
-	{
-		return !Context->RequestResourceLoad(Context, { Settings->Asset.ToSoftObjectPath() }, !Settings->bSynchronousLoad);
-	}
-	else
-	{
-		return true;
-	}
+	return Context->InitializeAndRequestLoad(PCGPinConstants::DefaultInputLabel,
+		Settings->AssetReferenceSelector,
+		{ Settings->Asset.ToSoftObjectPath() },
+		/*bPersistAllData=*/false,
+		/*bSilenceErrorOnEmptyObjectPath*/!Settings->bWarnIfNoAsset,
+		/*bSynchronousLoad=*/Settings->bSynchronousLoad);
 }
 
-bool FPCGLoadDataAssetElement::ExecuteInternal(FPCGContext* Context) const
+bool FPCGLoadDataAssetElement::ExecuteInternal(FPCGContext* InContext) const
 {
 	TRACE_CPUPROFILER_EVENT_SCOPE(FPCGDataAssetElement::Execute);
 
-	check(Context);
+	check(InContext);
+	FPCGLoadDataAssetContext* Context = static_cast<FPCGLoadDataAssetContext*>(InContext);
 	const UPCGLoadDataAssetSettings* Settings = Context->GetInputSettings<UPCGLoadDataAssetSettings>();
 	check(Settings);
 
 #if WITH_EDITOR
-	if (Context->IsValueOverriden(GET_MEMBER_NAME_CHECKED(UPCGLoadDataAssetSettings, Asset)))
+	if (Context->IsValueOverriden(GET_MEMBER_NAME_CHECKED(UPCGLoadDataAssetSettings, Asset)) || (Settings->bLoadFromInput && !Context->PathsToObjectsAndDataIndex.IsEmpty()))
 	{
-		FPCGDynamicTrackingHelper::AddSingleDynamicTrackingKey(Context, FPCGSelectionKey::CreateFromPath(Settings->Asset.ToSoftObjectPath()), /*bIsCulled=*/false);
+		FPCGDynamicTrackingHelper DynamicTracking;
+		DynamicTracking.EnableAndInitialize(Context, Context->PathsToObjectsAndDataIndex.Num());
+		for(const TTuple<FSoftObjectPath, int32, int32>& Path : Context->PathsToObjectsAndDataIndex)
+		{
+			DynamicTracking.AddToTracking(FPCGSelectionKey::CreateFromPath(Path.Get<0>()), /*bIsCulled=*/false);
+		}
+
+		DynamicTracking.Finalize(Context);
 	}
 #endif
 
 	// At this point, the data should already be loaded
-	if (UPCGDataAsset* AssetData = Settings->Asset.LoadSynchronous())
+	for(const TTuple<FSoftObjectPath, int32, int32>& AssetPath : Context->PathsToObjectsAndDataIndex)
 	{
-		Context->OutputData = AssetData->Data;
-
-		if (Settings->bTagOutputsBasedOnOutputPins)
+		TSoftObjectPtr<UPCGDataAsset> Asset(AssetPath.Get<0>());
+		UPCGDataAsset* AssetData = Asset.LoadSynchronous();
+		if (AssetData)
 		{
-			for (FPCGTaggedData& TaggedData : Context->OutputData.TaggedData)
+			const int TaggedDataOffset = Context->OutputData.TaggedData.Num();
+			Context->OutputData.TaggedData.Append(AssetData->Data.TaggedData);
+
+			if (Settings->bTagOutputsBasedOnOutputPins || Settings->InputIndexTag != NAME_None || Settings->DataIndexTag != NAME_None)
 			{
-				if (TaggedData.Pin != NAME_None)
+				for (int TaggedDataIndex = TaggedDataOffset; TaggedDataIndex < Context->OutputData.TaggedData.Num(); ++TaggedDataIndex)
 				{
-					TaggedData.Tags.Add(TaggedData.Pin.ToString());
+					FPCGTaggedData& TaggedData = Context->OutputData.TaggedData[TaggedDataIndex];
+
+					if (Settings->bTagOutputsBasedOnOutputPins && TaggedData.Pin != NAME_None)
+					{
+						TaggedData.Tags.Add(TaggedData.Pin.ToString());
+					}
+
+					if (Settings->InputIndexTag != NAME_None)
+					{
+						TaggedData.Tags.Add(FString::Format(TEXT("{0}:{1}"), { Settings->InputIndexTag.ToString(), AssetPath.Get<1>() }));
+					}
+
+					if (Settings->DataIndexTag != NAME_None)
+					{
+						TaggedData.Tags.Add(FString::Format(TEXT("{0}:{1}"), { Settings->DataIndexTag.ToString(), AssetPath.Get<2>() }));
+					}
 				}
 			}
 		}
 	}
-	else if (!Settings->Asset.IsNull() && Settings->bWarnIfNoAsset)
-	{
-		PCGE_LOG(Warning, GraphAndLog, FText::Format(NSLOCTEXT("PCGLoadDataAssetSettings", "UnableToLoadAsset", "Unable to load PCG asset '{0}'"), FText::FromString(Settings->Asset.ToString())));
-	}
 
 	return true;
 }
\ No newline at end of file
diff --git a/Engine/Plugins/PCG/Source/PCG/Private/Elements/PCGDataFromActor.cpp b/Engine/Plugins/PCG/Source/PCG/Private/Elements/PCGDataFromActor.cpp
--- a/Engine/Plugins/PCG/Source/PCG/Private/Elements/PCGDataFromActor.cpp
+++ b/Engine/Plugins/PCG/Source/PCG/Private/Elements/PCGDataFromActor.cpp
@@ -177,7 +177,22 @@ void UPCGDataFromActorSettings::PostEditChangeProperty(FPropertyChangedEvent& Pr
 	}
 }
 
-#endif
+EPCGChangeType UPCGDataFromActorSettings::GetChangeTypeForProperty(FPropertyChangedEvent& PropertyChangedEvent) const
+{
+	EPCGChangeType ChangeType = Super::GetChangeTypeForProperty(PropertyChangedEvent) | EPCGChangeType::Cosmetic;
+
+	if (PropertyChangedEvent.GetMemberPropertyName() == GET_MEMBER_NAME_CHECKED(UPCGDataFromActorSettings, ActorSelector))
+	{
+		// If we change from/to FromInput, this needs to trigger a graph recompilation especially for culling.
+		if (PropertyChangedEvent.GetPropertyName() == GET_MEMBER_NAME_CHECKED(FPCGActorSelectorSettings, ActorFilter))
+		{
+			ChangeType |= EPCGChangeType::Structural;
+		}
+	}
+
+	return ChangeType;
+}
+#endif // WITH_EDITOR
 
 TSubclassOf<AActor> UPCGDataFromActorSettings::GetDefaultActorSelectorClass() const
 {
diff --git a/Engine/Plugins/PCG/Source/PCG/Private/Elements/PCGGetActorProperty.cpp b/Engine/Plugins/PCG/Source/PCG/Private/Elements/PCGGetActorProperty.cpp
--- a/Engine/Plugins/PCG/Source/PCG/Private/Elements/PCGGetActorProperty.cpp
+++ b/Engine/Plugins/PCG/Source/PCG/Private/Elements/PCGGetActorProperty.cpp
@@ -37,6 +37,23 @@ void UPCGGetActorPropertySettings::PostEditChangeProperty(FPropertyChangedEvent&
 		}
 	}
 }
+
+EPCGChangeType UPCGGetActorPropertySettings::GetChangeTypeForProperty(FPropertyChangedEvent& PropertyChangedEvent) const
+{
+	EPCGChangeType ChangeType = Super::GetChangeTypeForProperty(PropertyChangedEvent) | EPCGChangeType::Cosmetic;
+
+	if (PropertyChangedEvent.GetMemberPropertyName() == GET_MEMBER_NAME_CHECKED(UPCGGetActorPropertySettings, ActorSelector))
+	{
+		// If we change from/to FromInput, this needs to trigger a graph recompilation especially for culling.
+		if (PropertyChangedEvent.GetPropertyName() == GET_MEMBER_NAME_CHECKED(FPCGActorSelectorSettings, ActorFilter))
+		{
+			ChangeType |= EPCGChangeType::Structural;
+		}
+	}
+
+	return ChangeType;
+}
+
 #endif // WITH_EDITOR
 
 void UPCGGetActorPropertySettings::PostLoad()
diff --git a/Engine/Plugins/PCG/Source/PCG/Public/Elements/IO/PCGLoadAssetElement.h b/Engine/Plugins/PCG/Source/PCG/Public/Elements/IO/PCGLoadAssetElement.h
--- a/Engine/Plugins/PCG/Source/PCG/Public/Elements/IO/PCGLoadAssetElement.h
+++ b/Engine/Plugins/PCG/Source/PCG/Public/Elements/IO/PCGLoadAssetElement.h
@@ -5,7 +5,7 @@
 #include "PCGContext.h"
 #include "PCGDataAsset.h"
 #include "PCGSettings.h"
-#include "Async/PCGAsyncLoadingContext.h"
+#include "Elements/PCGLoadObjectsContext.h"
 
 struct FAssetData;
 
@@ -40,12 +40,16 @@ public:
 	virtual bool HasDynamicPins() const override { return true; }
 	virtual bool HasFlippedTitleLines() const override { return true; }
 	virtual FString GetAdditionalTitleInformation() const override;
+	virtual EPCGDataType GetCurrentPinTypes(const UPCGPin* InPin) const;
 
 protected:
 	virtual FPCGElementPtr CreateElement() const override;
+#if WITH_EDITOR
+	virtual EPCGChangeType GetChangeTypeForProperty(const FName& InPropertyName) const override;
+#endif
 	// TODO tracking if data is stored in external asset
 
-	virtual TArray<FPCGPinProperties> InputPinProperties() const override { return TArray<FPCGPinProperties>(); }
+	virtual TArray<FPCGPinProperties> InputPinProperties() const override;
 	virtual TArray<FPCGPinProperties> OutputPinProperties() const override { return Pins; }
 	// ~End UPCGSettings interface
 
@@ -56,36 +60,48 @@ public:
 	UPROPERTY(BlueprintReadWrite, EditAnywhere, Category = Data, meta = (PCG_Overridable))
 	TSoftObjectPtr<UPCGDataAsset> Asset;
 
-	UPROPERTY(BlueprintReadOnly, VisibleAnywhere, Category = Data, meta = (NoResetToDefault))
+	UPROPERTY(BlueprintReadOnly, VisibleAnywhere, Category = "Data|Asset Info", meta = (NoResetToDefault))
 	TArray<FPCGPinProperties> Pins;
 
 	// Cached from the data when loaded
-	UPROPERTY(BlueprintReadOnly, VisibleAnywhere, Category = Data)
+	UPROPERTY(BlueprintReadOnly, VisibleAnywhere, Category = "Data|Asset Info")
 	FString AssetName;
 
 #if WITH_EDITORONLY_DATA
 	// Cached from the data when loaded
-	UPROPERTY(BlueprintReadOnly, VisibleAnywhere, Category = Data)
+	UPROPERTY(BlueprintReadOnly, VisibleAnywhere, Category = "Data|Asset Info")
 	FText AssetDescription = FText::GetEmpty();
 
 	// Cached from the data when loaded
-	UPROPERTY(BlueprintReadOnly, VisibleAnywhere, Category = Data)
+	UPROPERTY(BlueprintReadOnly, VisibleAnywhere, Category = "Data|Asset Info")
 	FLinearColor AssetColor = FLinearColor::White;
 #endif // WITH_EDITORONLY_DATA
 
+	UPROPERTY(BlueprintReadWrite, EditAnywhere, Category = Settings)
+	bool bLoadFromInput = false;
+
+	UPROPERTY(BlueprintReadWrite, EditAnywhere, Category = Settings, meta = (EditCondition = "bLoadFromInput", EditConditionHides, PCG_DiscardPropertySelection))
+	FPCGAttributePropertyInputSelector AssetReferenceSelector;
+
+	UPROPERTY(BlueprintReadWrite, EditAnywhere, Category = Settings, meta = (EditCondition = "bLoadFromInput", EditConditionHides))
+	FName InputIndexTag = NAME_None;
+
+	UPROPERTY(BlueprintReadWrite, EditAnywhere, Category = Settings, meta = (EditCondition = "bLoadFromInput", EditConditionHides))
+	FName DataIndexTag = NAME_None;
+
 	UPROPERTY(BlueprintReadWrite, EditAnywhere, Category = Settings, meta = (Tooltip="Warns if asset is null or couldn't be loaded"))
 	bool bWarnIfNoAsset = true;
 
-	/** By default, data table loading is asynchronous, can force it synchronous if needed. */
-	UPROPERTY(BlueprintReadWrite, EditAnywhere, Category = "Settings|Debug")
-	bool bSynchronousLoad = false;
-
 	/** Controls whether the data output from the loaded asset will be passed to the default pin with tags or on the proper pins. */
 	UPROPERTY(BlueprintReadOnly, VisibleAnywhere, Category = "Settings", meta = (NoResetToDefault))
 	bool bTagOutputsBasedOnOutputPins = true;
+
+	/** By default, data table loading is asynchronous, can force it synchronous if needed. */
+	UPROPERTY(BlueprintReadWrite, EditAnywhere, Category = "Settings|Debug")
+	bool bSynchronousLoad = false;
 };
 
-struct FPCGLoadDataAssetContext : public FPCGContext, public IPCGAsyncLoadingContext {};
+struct FPCGLoadDataAssetContext : public FPCGLoadObjectsFromPathContext {};
 
 class PCG_API FPCGLoadDataAssetElement : public IPCGElementWithCustomContext<FPCGLoadDataAssetContext>
 {
diff --git a/Engine/Plugins/PCG/Source/PCG/Public/Elements/PCGDataFromActor.h b/Engine/Plugins/PCG/Source/PCG/Public/Elements/PCGDataFromActor.h
--- a/Engine/Plugins/PCG/Source/PCG/Public/Elements/PCGDataFromActor.h
+++ b/Engine/Plugins/PCG/Source/PCG/Public/Elements/PCGDataFromActor.h
@@ -47,7 +47,7 @@ public:
 
 protected:
 #if WITH_EDITOR
-	virtual EPCGChangeType GetChangeTypeForProperty(const FName& InPropertyName) const override { return Super::GetChangeTypeForProperty(InPropertyName) | EPCGChangeType::Cosmetic; }
+	virtual EPCGChangeType GetChangeTypeForProperty(FPropertyChangedEvent& PropertyChangedEvent) const override;
 #endif
 	virtual TArray<FPCGPinProperties> InputPinProperties() const override;
 	virtual TArray<FPCGPinProperties> OutputPinProperties() const override;
diff --git a/Engine/Plugins/PCG/Source/PCG/Public/Elements/PCGGetActorProperty.h b/Engine/Plugins/PCG/Source/PCG/Public/Elements/PCGGetActorProperty.h
--- a/Engine/Plugins/PCG/Source/PCG/Public/Elements/PCGGetActorProperty.h
+++ b/Engine/Plugins/PCG/Source/PCG/Public/Elements/PCGGetActorProperty.h
@@ -44,7 +44,7 @@ public:
 
 protected:
 #if WITH_EDITOR
-	virtual EPCGChangeType GetChangeTypeForProperty(const FName& InPropertyName) const override { return Super::GetChangeTypeForProperty(InPropertyName) | EPCGChangeType::Cosmetic; }
+	virtual EPCGChangeType GetChangeTypeForProperty(FPropertyChangedEvent& PropertyChangedEvent) const override;
 #endif
 	virtual FPCGElementPtr CreateElement() const override;
 	//~End UPCGSettings interface
