[PCG] Add option to not ignore PCG generated component in Get Actor Data

#rnx

diff --git a/Engine/Plugins/PCG/Source/PCG/Private/Elements/PCGDataFromActor.cpp b/Engine/Plugins/PCG/Source/PCG/Private/Elements/PCGDataFromActor.cpp
--- a/Engine/Plugins/PCG/Source/PCG/Private/Elements/PCGDataFromActor.cpp
+++ b/Engine/Plugins/PCG/Source/PCG/Private/Elements/PCGDataFromActor.cpp
@@ -828,6 +828,7 @@ void FPCGDataFromActorElement::ProcessActor(FPCGContext* InContext, const UPCGDa
 		GDFRParams.ComponentSelector = &Context->ComponentSelector;
 		GDFRParams.bParseActor = bParseActor;
 		GDFRParams.DataTypeFilter = DataType;
+		GDFRParams.bIgnorePCGGeneratedComponents = Settings->bIgnorePCGGeneratedComponents;
 
 		FPCGGetDataFunctionRegistryOutput GDFROutput;
 		FPCGModule::ConstGetDataFunctionRegistry().GetDataFromActor(Context, GDFRParams, Actor, GDFROutput);
diff --git a/Engine/Plugins/PCG/Source/PCG/Public/Elements/PCGDataFromActor.h b/Engine/Plugins/PCG/Source/PCG/Public/Elements/PCGDataFromActor.h
--- a/Engine/Plugins/PCG/Source/PCG/Public/Elements/PCGDataFromActor.h
+++ b/Engine/Plugins/PCG/Source/PCG/Public/Elements/PCGDataFromActor.h
@@ -89,6 +89,10 @@ public:
 	UPROPERTY(BlueprintReadWrite, EditAnywhere, Category = "Data Retrieval Settings", meta = (EditCondition = "DisplayModeSettings()", EditConditionHides, HideEditConditionToggle))
 	EPCGGetDataFromActorMode Mode = EPCGGetDataFromActorMode::ParseActorComponents;
 
+	/** Ignores any component that was spawned by PCG. */
+	UPROPERTY(BlueprintReadWrite, EditAnywhere, Category = "Data Retrieval Settings", meta = (EditCondition = "Mode == EPCGGetDataFromActorMode::ParseActorComponents || Mode == EPCGGetDataFromActorMode::GetComponentsReference", EditConditionHides))
+	bool bIgnorePCGGeneratedComponents = true;
+
 	/** Also produces a single point data at the actor location. */
 	UPROPERTY(BlueprintReadWrite, EditAnywhere, Category = "Data Retrieval Settings", meta = (EditCondition = "Mode == EPCGGetDataFromActorMode::GetDataFromPCGComponent || Mode == EPCGGetDataFromActorMode::GetDataFromPCGComponentOrParseComponents", EditConditionHides))
 	bool bAlsoOutputSinglePointData = false;
