[PCG] Refactor of Create Spline Mesh to Spawn Spline Mesh

* With its redirector

#rnx

diff --git a/Engine/Plugins/PCG/Config/BasePCG.ini b/Engine/Plugins/PCG/Config/BasePCG.ini
--- a/Engine/Plugins/PCG/Config/BasePCG.ini
+++ b/Engine/Plugins/PCG/Config/BasePCG.ini
@@ -86,6 +86,7 @@
 +ClassRedirects=(OldName="PCGPropertyToParamDataSettings",NewName="/Script/PCG.PCGGetActorPropertySettings")
 +ClassRedirects=(OldName="PCGSortPointsSettings",NewName="/Script/PCG.PCGSortAttributesSettings")
 +ClassRedirects=(OldName="PCGMetadataOperationSettings",NewName="/Script/PCG.PCGCopyAttributesSettings")
++ClassRedirects=(OldName="PCGCreateSplineMeshSettings",NewName="/Script/PCG.PCGSpawnSplineMeshSettings")
 
 +StructRedirects=(OldName="PCGPointFilterThresholdSettings",NewName="/Script/PCG.PCGAttributeFilterThresholdSettings")
 +StructRedirects=(OldName="PCGActorPropertyOverrideDescription",NewName="/Script/PCG.PCGObjectPropertyOverrideDescription")
diff --git a/Engine/Plugins/PCG/Source/PCG/Private/Elements/PCGCreateSplineMesh.cpp b/Engine/Plugins/PCG/Source/PCG/Private/Elements/PCGCreateSplineMesh.cpp
deleted file mode 100644
diff --git a/Engine/Plugins/PCG/Source/PCG/Private/Elements/PCGSpawnSplineMesh.cpp b/Engine/Plugins/PCG/Source/PCG/Private/Elements/PCGSpawnSplineMesh.cpp
new file mode 100644
--- /dev/null
+++ b/Engine/Plugins/PCG/Source/PCG/Private/Elements/PCGSpawnSplineMesh.cpp
@@ -0,0 +1,235 @@
+// Copyright Epic Games, Inc. All Rights Reserved.
+
+#include "Elements/PCGSpawnSplineMesh.h"
+
+#include "PCGComponent.h"
+#include "PCGContext.h"
+#include "PCGManagedResource.h"
+#include "PCGPin.h"
+#include "PCGSettings.h"
+#include "PCGSubsystem.h"
+#include "Data/PCGLandscapeSplineData.h"
+#include "Data/PCGSplineData.h"
+#include "Helpers/PCGActorHelpers.h"
+#include "Helpers/PCGHelpers.h"
+
+#include "DrawDebugHelpers.h"
+#include "Components/SplineMeshComponent.h"
+#include "Engine/StaticMesh.h"
+#include "Engine/World.h"
+#include "Materials/MaterialInterface.h"
+#include "VT/RuntimeVirtualTexture.h"
+
+#define LOCTEXT_NAMESPACE "PCGCreateSplineMeshElement"
+
+#if WITH_EDITOR
+FText UPCGSpawnSplineMeshSettings::GetNodeTooltipText() const
+{
+	return LOCTEXT("CreateSplineMeshTooltip", "Create a USplineMeshComponent for each segment along a given spline.");
+}
+#endif
+
+TArray<FPCGPinProperties> UPCGSpawnSplineMeshSettings::InputPinProperties() const
+{
+	TArray<FPCGPinProperties> PinProperties;
+	PinProperties.Emplace(PCGPinConstants::DefaultInputLabel, EPCGDataType::PolyLine);
+
+	return PinProperties;
+}
+
+FPCGElementPtr UPCGSpawnSplineMeshSettings::CreateElement() const
+{
+	return MakeShared<FPCGSpawnSplineMeshElement>();
+}
+
+bool FPCGSpawnSplineMeshElement::PrepareDataInternal(FPCGContext* InContext) const
+{
+	TRACE_CPUPROFILER_EVENT_SCOPE(FPCGCreateSplineMeshElement::PrepareDataInternal);
+
+	FPCGSpawnSplineMeshContext* Context = static_cast<FPCGSpawnSplineMeshContext*>(InContext);
+	check(Context);
+
+	const UPCGSpawnSplineMeshSettings* Settings = Context->GetInputSettings<UPCGSpawnSplineMeshSettings>();
+	check(Settings);
+
+	const FSoftSplineMeshComponentDescriptor& Descriptor = Settings->SplineMeshDescriptor;
+
+	if (Descriptor.StaticMesh.IsNull())
+	{
+		return true;
+	}
+
+	if (!Context->WasLoadRequested())
+	{
+		TArray<FSoftObjectPath> ObjectsToLoad = { Descriptor.StaticMesh.ToSoftObjectPath() };
+
+		if (!Descriptor.OverlayMaterial.IsNull())
+		{
+			ObjectsToLoad.Emplace(Descriptor.OverlayMaterial.ToSoftObjectPath());
+		}
+
+		for (int MaterialIndex = 0; MaterialIndex < Descriptor.OverrideMaterials.Num(); ++MaterialIndex)
+		{
+			if (!Descriptor.OverrideMaterials[MaterialIndex].IsNull())
+			{
+				ObjectsToLoad.Emplace(Descriptor.OverrideMaterials[MaterialIndex].ToSoftObjectPath());
+			}
+		}
+
+		for (int RVTIndex = 0; RVTIndex < Descriptor.RuntimeVirtualTextures.Num(); ++RVTIndex)
+		{
+			if (!Descriptor.RuntimeVirtualTextures[RVTIndex].IsNull())
+			{
+				ObjectsToLoad.Emplace(Descriptor.RuntimeVirtualTextures[RVTIndex].ToSoftObjectPath());
+			}
+		}
+
+		return Context->RequestResourceLoad(Context, std::move(ObjectsToLoad), !Settings->bSynchronousLoad);
+	}
+
+	return true;
+}
+
+bool FPCGSpawnSplineMeshElement::ExecuteInternal(FPCGContext* Context) const
+{
+	TRACE_CPUPROFILER_EVENT_SCOPE(FPCGCreateSplineMeshElement::Execute);
+
+	const UPCGSpawnSplineMeshSettings* Settings = Context->GetInputSettings<UPCGSpawnSplineMeshSettings>();
+	check(Settings);
+
+	FSplineMeshComponentDescriptor Descriptor(Settings->SplineMeshDescriptor);
+
+	FPCGSplineMeshComponentBuilderParameters SMCBuilderParams;
+	SMCBuilderParams.SplineMeshParams = Settings->SplineMeshParams;
+
+	TArray<FPCGTaggedData> Inputs = Context->InputData.GetInputs();
+
+	for (const FPCGTaggedData& Input : Inputs)
+	{
+		const UPCGPolyLineData* SplineData = Cast<UPCGPolyLineData>(Input.Data);
+		const UPCGLandscapeSplineData* LandscapeSplineData = Cast<UPCGLandscapeSplineData>(SplineData);
+
+		if (!SplineData)
+		{
+			PCGE_LOG(Error, GraphAndLog, LOCTEXT("InvalidInputData", "Invalid input data, must be of type PolyLine."));
+			continue;
+		}
+
+		AActor* TargetActor = Settings->TargetActor.Get() ? Settings->TargetActor.Get() : Context->GetTargetActor(nullptr);
+
+		if (!TargetActor)
+		{
+			PCGE_LOG(Error, GraphAndLog, LOCTEXT("InvalidTargetActor", "Invalid target actor."));
+			continue;
+		}
+
+		const int NumSegments = SplineData->GetNumSegments();
+		const bool bIsClosed = SplineData->IsClosed();
+
+		// Copy the descriptor, for the overrides (and reset the descriptor if it was overriden by previous inputs)
+		SMCBuilderParams.Descriptor = Descriptor;
+
+		// TODO: Support overrides on control points when we have better support for those.
+		FPCGObjectOverrides DescriptorOverride(&SMCBuilderParams.Descriptor);
+		DescriptorOverride.Initialize(Settings->SplineMeshOverrideDescriptions, &SMCBuilderParams.Descriptor, SplineData, Context);
+		if (DescriptorOverride.IsValid() && !DescriptorOverride.Apply(0))
+		{
+			const int32 Index = static_cast<int32>(&Input - Inputs.GetData());
+			PCGLog::LogWarningOnGraph(FText::Format(LOCTEXT("FailOverride", "Failed to override descriptor for input {0}"), Index));
+		}
+
+		UStaticMesh* StaticMesh = SMCBuilderParams.Descriptor.StaticMesh.Get();
+
+		if (!StaticMesh)
+		{
+			continue;
+		}
+
+		const FVector MeshExtents = StaticMesh->GetBounds().BoxExtent;
+
+		for (int Index = 0; Index < NumSegments; ++Index)
+		{
+			const int NextIndex = (bIsClosed && Index == NumSegments - 1) ? 0 : (Index + 1);
+
+			FBox Bounds, NextBounds;
+			FTransform Transform = SplineData->GetTransformAtDistance(Index, /*Distance=*/0.0, /*bWorldSpace=*/true, &Bounds);
+			FTransform NextTransform = SplineData->GetTransformAtDistance(Index, SplineData->GetSegmentLength(Index), /*bWorldSpace=*/true, &NextBounds);
+
+			if (LandscapeSplineData && !SMCBuilderParams.SplineMeshParams.bScaleMeshToLandscapeSplineFullWidth)
+			{
+				Bounds = FBox::BuildAABB(FVector::ZeroVector, FVector::OneVector);
+				NextBounds = FBox::BuildAABB(FVector::ZeroVector, FVector::OneVector);
+			}
+
+			FVector LeaveTangent, ArriveTangent, Dummy;
+			SplineData->GetTangentsAtSegmentStart(Index, Dummy, LeaveTangent);
+			SplineData->GetTangentsAtSegmentStart(NextIndex, ArriveTangent, Dummy);
+
+			// Set Position
+			SMCBuilderParams.SplineMeshParams.StartPosition = Transform.GetLocation();
+			SMCBuilderParams.SplineMeshParams.StartTangent = LeaveTangent;
+			SMCBuilderParams.SplineMeshParams.EndPosition = NextTransform.GetLocation();
+			SMCBuilderParams.SplineMeshParams.EndTangent = ArriveTangent;
+
+			// Set Roll (Rotation)
+			const FQuat StartRotation = Transform.GetRotation();
+			const FQuat EndRotation = NextTransform.GetRotation();
+
+			SMCBuilderParams.SplineMeshParams.StartRollDegrees = StartRotation.Rotator().Roll;
+			SMCBuilderParams.SplineMeshParams.EndRollDegrees = EndRotation.Rotator().Roll;
+
+			// Set Scale
+			FVector StartScale = Transform.GetScale3D() * Bounds.GetExtent();
+			FVector EndScale = NextTransform.GetScale3D() * NextBounds.GetExtent();
+
+			if (SMCBuilderParams.SplineMeshParams.bScaleMeshToBounds)
+			{
+				FVector::FReal ScaleFactorY = 1.0f;
+				FVector::FReal ScaleFactorZ = 1.0f;
+
+				// We only scale in two dimensions since we are extruding along one of the axes. Scale on the two axes we are not extruding along.
+				if (SMCBuilderParams.SplineMeshParams.ForwardAxis == EPCGSplineMeshForwardAxis::X)
+				{
+					ScaleFactorY = MeshExtents.Y;
+					ScaleFactorZ = MeshExtents.Z;
+				}
+				else if (SMCBuilderParams.SplineMeshParams.ForwardAxis == EPCGSplineMeshForwardAxis::Y)
+				{
+					ScaleFactorY = MeshExtents.X;
+					ScaleFactorZ = MeshExtents.Z;
+				}
+				else if (SMCBuilderParams.SplineMeshParams.ForwardAxis == EPCGSplineMeshForwardAxis::Z)
+				{
+					ScaleFactorY = MeshExtents.X;
+					ScaleFactorZ = MeshExtents.Y;
+				}
+
+				StartScale.Y /= ScaleFactorY;
+				StartScale.Z /= ScaleFactorZ;
+				EndScale.Y /= ScaleFactorY;
+				EndScale.Z /= ScaleFactorZ;
+			}
+
+			SMCBuilderParams.SplineMeshParams.StartScale = FVector2D(StartScale.Y, StartScale.Z);
+			SMCBuilderParams.SplineMeshParams.EndScale = FVector2D(EndScale.Y, EndScale.Z);
+
+			USplineMeshComponent* SplineMeshComponent = UPCGActorHelpers::GetOrCreateSplineMeshComponent(TargetActor, Context->SourceComponent.Get(), Settings->UID, SMCBuilderParams);
+
+			// TODO: Write out the geometry to a dynamic mesh type.
+			//SplineMeshComponent->BodySetup->TriMeshGeometries
+		}
+
+		// Execute PostProcess Functions
+		if (TargetActor)
+		{
+			for (UFunction* Function : PCGHelpers::FindUserFunctions(TargetActor->GetClass(), Settings->PostProcessFunctionNames, { UPCGFunctionPrototypes::GetPrototypeWithNoParams() }, Context))
+			{
+				TargetActor->ProcessEvent(Function, nullptr);
+			}
+		}
+	}
+
+	return true;
+}
+
+#undef LOCTEXT_NAMESPACE
diff --git a/Engine/Plugins/PCG/Source/PCG/Public/Elements/PCGCreateSplineMesh.h b/Engine/Plugins/PCG/Source/PCG/Public/Elements/PCGCreateSplineMesh.h
deleted file mode 100644
diff --git a/Engine/Plugins/PCG/Source/PCG/Public/Elements/PCGSpawnSplineMesh.h b/Engine/Plugins/PCG/Source/PCG/Public/Elements/PCGSpawnSplineMesh.h
new file mode 100644
--- /dev/null
+++ b/Engine/Plugins/PCG/Source/PCG/Public/Elements/PCGSpawnSplineMesh.h
@@ -0,0 +1,74 @@
+// Copyright Epic Games, Inc. All Rights Reserved.
+
+#pragma once
+
+#include "PCGCommon.h"
+#include "PCGContext.h"
+#include "PCGElement.h"
+#include "PCGSettings.h"
+#include "PCGSplineMeshParams.h"
+#include "Async/PCGAsyncLoadingContext.h"
+#include "Metadata/PCGObjectPropertyOverride.h"
+
+#include "Engine/SplineMeshComponentDescriptor.h"
+
+#include "PCGSpawnSplineMesh.generated.h"
+
+struct FPCGContext;
+
+/** Create a USplineMeshComponent for each segment along a given spline. */
+UCLASS(MinimalAPI, BlueprintType, ClassGroup = (Procedural))
+class UPCGSpawnSplineMeshSettings : public UPCGSettings
+{
+	GENERATED_BODY()
+
+public:
+	//~Begin UPCGSettings interface
+#if WITH_EDITOR
+	virtual FName GetDefaultNodeName() const override { return FName(TEXT("SpawnSplineMesh")); }
+	virtual FText GetDefaultNodeTitle() const override { return NSLOCTEXT("PCGSpawnSplineMeshElement", "NodeTitle", "Spawn Spline Mesh"); }
+	virtual FText GetNodeTooltipText() const override;
+	virtual EPCGSettingsType GetType() const override { return EPCGSettingsType::Spawner; }
+#endif
+
+protected:
+	virtual TArray<FPCGPinProperties> InputPinProperties() const override;
+	virtual TArray<FPCGPinProperties> OutputPinProperties() const override { return TArray<FPCGPinProperties>(); }
+	virtual FPCGElementPtr CreateElement() const override;
+	//~End UPCGSettings interface
+
+public:
+	UPROPERTY(EditAnywhere, Category = Settings)
+	FSoftSplineMeshComponentDescriptor SplineMeshDescriptor;
+
+	UPROPERTY(EditAnywhere, Category = Settings, meta = (PCG_Overridable))
+	FPCGSplineMeshParams SplineMeshParams;
+
+	UPROPERTY(meta = (PCG_Overridable))
+	TSoftObjectPtr<AActor> TargetActor;
+
+	/** Specify a list of functions to be called on the target actor after spline mesh creation. Functions need to be parameter-less and with "CallInEditor" flag enabled. */
+	UPROPERTY(BlueprintReadWrite, EditAnywhere, Category = Settings)
+	TArray<FName> PostProcessFunctionNames;
+
+	/** Force meshes/materials to load synchronously. */
+	UPROPERTY(BlueprintReadWrite, EditAnywhere, Category = "Settings|Debug")
+	bool bSynchronousLoad = false;
+
+	/** Overrides for spline mesh descriptor. For now it is only support data wide attributes. Does not work for per-control point overrides for now. */
+	UPROPERTY(EditAnywhere, BlueprintReadWrite, Category = Settings)
+	TArray<FPCGObjectPropertyOverrideDescription> SplineMeshOverrideDescriptions;
+};
+
+struct FPCGSpawnSplineMeshContext : public FPCGContext, public IPCGAsyncLoadingContext {};
+
+class FPCGSpawnSplineMeshElement : public IPCGElementWithCustomContext<FPCGSpawnSplineMeshContext>
+{
+public:
+	virtual bool CanExecuteOnlyOnMainThread(FPCGContext* Context) const override { return true; }
+	virtual bool IsCacheable(const UPCGSettings* InSettings) const override { return false; }
+
+protected:
+	virtual bool PrepareDataInternal(FPCGContext* InContext) const override;
+	virtual bool ExecuteInternal(FPCGContext* Context) const override;
+};
