[PCG] Graph Executor V2: Fix handling of Caching results with Debugged tasks

#rnx

diff --git a/Engine/Plugins/PCG/Source/PCG/Private/Graph/PCGGraphExecutor.cpp b/Engine/Plugins/PCG/Source/PCG/Private/Graph/PCGGraphExecutor.cpp
--- a/Engine/Plugins/PCG/Source/PCG/Private/Graph/PCGGraphExecutor.cpp
+++ b/Engine/Plugins/PCG/Source/PCG/Private/Graph/PCGGraphExecutor.cpp
@@ -1149,8 +1149,7 @@ void FPCGGraphExecutor::OnTaskInputsReady(FPCGGraphTask& Task, TArray<FCachedRes
 			check(CachedResult->TaskId != InvalidPCGTaskId);
 			OutCachedResults.Add(CachedResult);
 		}
-
-		if(!CachedResult || CachedResult->bIsBypassed)
+		else 
 		{
 			ReadyTasks.Emplace(MoveTemp(Task));
 		}
@@ -1303,12 +1302,14 @@ void FPCGGraphExecutor::PrepareForExecute(FPCGGraphTask& Task, FCachedResult*& O
 	{
 		check(OutCachedResult->TaskId != InvalidPCGTaskId);
 		Task.bIsBypassed = true;
-		OutCachedResult->bIsBypassed = true;
-		Task.Context->OutputData = OutCachedResult->Output;
+
+		// Directly transfer Cache result for debugged tasks to Context output
+		PCGGraphExecutor::TScopeLock ScopeLock(CachingResultsLock);
+		Task.Context->OutputData = MoveTemp(OutCachedResult->Output);
+		OutCachedResult = nullptr;
+		CollectGCCachingResults.Remove(Task.NodeId);
 	}
 #endif
-
-	return;
 }
 
 void FPCGGraphExecutor::Execute()
@@ -1414,8 +1415,11 @@ bool FPCGGraphExecutor::ExecuteScheduling(double EndTime, TSharedPtr<FPCGGraphAc
 					PrepareForExecute(Task, CachedResult, /*bLiveTasksLockAlreadyLocked=*/true);
 					check(Task.bHasDonePrepareForExecute);
 
-					if (CachedResult && !CachedResult->bIsBypassed)
+					if (CachedResult)
 					{
+#if WITH_EDITOR
+						check(!Task.bIsBypassed);
+#endif
 						check(CachedResult->TaskId != InvalidPCGTaskId);
 						CachedResults.Add(MoveTemp(CachedResult));
 						ReadyTasks.RemoveAtSwap(ReadyTaskIndex, EAllowShrinking::No);
@@ -2012,14 +2016,6 @@ TArray<FPCGTaskId> FPCGGraphExecutor::ProcessCachedResultsInternal(TArray<FCache
 
 	for (FCachedResult* CachedResult : CachedResults)
 	{
-#if WITH_EDITOR
-		// Cache was used to bypass actual execution but we let the task go through the normal execution path
-		if (CachedResult->bIsBypassed)
-		{
-			continue;
-		}
-#endif
-
 		check(CachedResult);
 		if (CachedResult->bDoDynamicTaskCulling)
 		{
diff --git a/Engine/Plugins/PCG/Source/PCG/Private/Graph/PCGGraphExecutor.h b/Engine/Plugins/PCG/Source/PCG/Private/Graph/PCGGraphExecutor.h
--- a/Engine/Plugins/PCG/Source/PCG/Private/Graph/PCGGraphExecutor.h
+++ b/Engine/Plugins/PCG/Source/PCG/Private/Graph/PCGGraphExecutor.h
@@ -268,7 +268,6 @@ private:
 		const UPCGNode* Node = nullptr;
 		bool bDoDynamicTaskCulling = false;
 		bool bIsPostGraphTask = false;		
-		bool bIsBypassed = false;
 	};
 		
 	void PostTaskExecute(TSharedPtr<FPCGGraphActiveTask> ActiveTask, bool bIsDone);
