[PCG] Support multiple downstream input pins for a single data interface

#rnx

diff --git a/Engine/Plugins/PCG/Source/PCG/Private/Graph/PCGGraphCompilerGPU.cpp b/Engine/Plugins/PCG/Source/PCG/Private/Graph/PCGGraphCompilerGPU.cpp
--- a/Engine/Plugins/PCG/Source/PCG/Private/Graph/PCGGraphCompilerGPU.cpp
+++ b/Engine/Plugins/PCG/Source/PCG/Private/Graph/PCGGraphCompilerGPU.cpp
@@ -688,13 +688,12 @@ void FPCGGraphCompilerGPU::BuildGPUGraphTask(
 		const UPCGSettings* Settings = Node ? Node->GetSettings() : nullptr;
 		check(Settings && Settings->bEnabled && Settings->ShouldExecuteOnGPU());
 
-		// For every usage of a DI, get the original (non-aliased) pin label.
-		TMap<TPair<FPCGTaskId, UComputeDataInterface*>, FName> DataInterfaceUsageToPinLabel;
-
-		TArray<int> InputDataInterfaceIndices;
-		TArray<int> OutputDataInterfaceIndices;
-		InputDataInterfaceIndices.Reserve(Settings->InputPinProperties().Num());
-		OutputDataInterfaceIndices.Reserve(Settings->OutputPinProperties().Num());
+		// One data interface may be connected to multiple downstream pins. Because each entry in these arrays produces its own bindings and graph edges,
+		// we also need to specify the pin label to uniquely describe each edge.
+		TArray<TPair</*DataInterfaceIndex=*/int, /*PinLabel=*/FName>> InputDataInterfaceIndexAndPin;
+		TArray<TPair</*DataInterfaceIndex=*/int, /*PinLabel=*/FName>> OutputDataInterfaceIndexAndPin;
+		InputDataInterfaceIndexAndPin.Reserve(Settings->InputPinProperties().Num());
+		OutputDataInterfaceIndexAndPin.Reserve(Settings->OutputPinProperties().Num());
 
 		// Add DIs (PCG -> CF transcoding).
 
@@ -726,11 +725,8 @@ void FPCGGraphCompilerGPU::BuildGPUGraphTask(
 				continue;
 			}
 
-			InputDataInterfaceIndices.Add(Index);
-
-			FName DownstreamInputPinLabel = Input.DownstreamPin->Label;
-
-			DataInterfaceUsageToPinLabel.Add({ TaskId, UpstreamDI }, DownstreamInputPinLabel);
+			const FName DownstreamInputPinLabel = Input.DownstreamPin->Label;
+			InputDataInterfaceIndexAndPin.Emplace(Index, DownstreamInputPinLabel);
 
 			const bool bIsInputPin = true;
 			UpstreamDI->AddDownstreamInputPin(DownstreamInputPinLabel, InOriginalToVirtualPin.Find({ TaskId, DownstreamInputPinLabel, bIsInputPin }));
@@ -752,8 +748,7 @@ void FPCGGraphCompilerGPU::BuildGPUGraphTask(
 				continue;
 			}
 
-			OutputDataInterfaceIndices.Add(Index);
-			DataInterfaceUsageToPinLabel.Add({ TaskId, *FoundDI }, OutputPinProperties.Label);
+			OutputDataInterfaceIndexAndPin.Emplace(Index, OutputPinProperties.Label);
 		}
 
 		// Make sure every downstream input pin is registered with the upstream DI.
@@ -794,7 +789,7 @@ void FPCGGraphCompilerGPU::BuildGPUGraphTask(
 
 		const int32 KernelDIIndex = ComputeGraph->DataInterfaces.Num();
 		ComputeGraph->DataInterfaces.Add(KernelDI);
-		InputDataInterfaceIndices.Add(KernelDIIndex);
+		InputDataInterfaceIndexAndPin.Emplace(KernelDIIndex, NAME_None);
 
 		TArray<TObjectPtr<UComputeDataInterface>> AdditionalInputDIs, AdditionalOutputDIs;
 		Settings->CreateAdditionalInputDataInterfaces(AdditionalInputDIs);
@@ -822,7 +817,7 @@ void FPCGGraphCompilerGPU::BuildGPUGraphTask(
 
 			if (DataInterfaceIndex != INDEX_NONE)
 			{
-				InputDataInterfaceIndices.Add(DataInterfaceIndex);
+				InputDataInterfaceIndexAndPin.Emplace(DataInterfaceIndex, NAME_None);
 			}
 		}
 
@@ -832,7 +827,7 @@ void FPCGGraphCompilerGPU::BuildGPUGraphTask(
 
 			if (DataInterfaceIndex != INDEX_NONE)
 			{
-				OutputDataInterfaceIndices.Add(DataInterfaceIndex);
+				OutputDataInterfaceIndexAndPin.Emplace(DataInterfaceIndex, NAME_None);
 			}
 		}
 
@@ -848,7 +843,7 @@ void FPCGGraphCompilerGPU::BuildGPUGraphTask(
 
 			const int32 DebugDIIndex = ComputeGraph->DataInterfaces.Num();
 			ComputeGraph->DataInterfaces.Add(DebugDI);
-			OutputDataInterfaceIndices.Add_GetRef(DebugDIIndex);
+			OutputDataInterfaceIndexAndPin.Emplace(DebugDIIndex, NAME_None);
 		}
 
 		// Now that all data interfaces added, create the (trivial) binding mapping. All map to primary binding, index 0.
@@ -856,15 +851,15 @@ void FPCGGraphCompilerGPU::BuildGPUGraphTask(
 
 		struct FInterfaceBinding
 		{
-			const UComputeDataInterface* DataInterface;
-			int32 DataInterfaceBindingIndex;
-			FString BindingFunctionName;
-			FString BindingFunctionNamespace;
+			const UComputeDataInterface* DataInterface = nullptr;
+			int32 DataInterfaceBindingIndex = INDEX_NONE;
+			FName PinLabel = NAME_None; // Required to uniquely identify this binding if this data interface is connected to multiple pins.
+			FString BindingFunctionName = TEXT("");
 		};
 
 		struct FKernelWithDataBindings
 		{
-			UComputeKernel* Kernel;
+			UComputeKernel* Kernel = nullptr;
 			TArray<FInterfaceBinding> InputDataBindings;
 			TArray<FInterfaceBinding> OutputDataBindings;
 		};
@@ -876,9 +871,9 @@ void FPCGGraphCompilerGPU::BuildGPUGraphTask(
 		ComputeGraph->KernelInvocations.Add(KernelWithBindings.Kernel);
 		ComputeGraph->KernelToNode.Add(Node);
 
-		auto SetupAllInputBindings = [&KernelWithBindings, ComputeGraph](int InDataInterfaceIndex)
+		auto SetupAllInputBindings = [&KernelWithBindings, ComputeGraph](const TPair<int, FName>& InDataInterfaceIndex)
 		{
-			const UComputeDataInterface* DataInterface = ComputeGraph->DataInterfaces[InDataInterfaceIndex];
+			const UComputeDataInterface* DataInterface = ComputeGraph->DataInterfaces[InDataInterfaceIndex.Key];
 			TArray<FShaderFunctionDefinition> Functions;
 			DataInterface->GetSupportedInputs(Functions);
 
@@ -886,15 +881,15 @@ void FPCGGraphCompilerGPU::BuildGPUGraphTask(
 			{
 				FInterfaceBinding& Binding = KernelWithBindings.InputDataBindings.Emplace_GetRef();
 				Binding.DataInterface = DataInterface;
+				Binding.PinLabel = InDataInterfaceIndex.Value;
 				Binding.BindingFunctionName = Functions[FuncIndex].Name;
-				Binding.BindingFunctionNamespace = TEXT("");
 				Binding.DataInterfaceBindingIndex = FuncIndex;
 			}
 		};
 
-		auto SetupAllOutputBindings = [&KernelWithBindings, ComputeGraph](int InDataInterfaceIndex)
+		auto SetupAllOutputBindings = [&KernelWithBindings, ComputeGraph](const TPair<int, FName>& InDataInterfaceIndex)
 		{
-			const UComputeDataInterface* DataInterface = ComputeGraph->DataInterfaces[InDataInterfaceIndex];
+			const UComputeDataInterface* DataInterface = ComputeGraph->DataInterfaces[InDataInterfaceIndex.Key];
 			TArray<FShaderFunctionDefinition> Functions;
 			DataInterface->GetSupportedOutputs(Functions);
 
@@ -902,19 +897,19 @@ void FPCGGraphCompilerGPU::BuildGPUGraphTask(
 			{
 				FInterfaceBinding& Binding = KernelWithBindings.OutputDataBindings.Emplace_GetRef();
 				Binding.DataInterface = DataInterface;
+				Binding.PinLabel = InDataInterfaceIndex.Value;
 				Binding.BindingFunctionName = Functions[FuncIndex].Name;
-				Binding.BindingFunctionNamespace = TEXT("");
 				Binding.DataInterfaceBindingIndex = FuncIndex;
 			}
 		};
 
 		// Bind data interfaces.
-		for (int InputDataInterfaceIndex : InputDataInterfaceIndices)
+		for (const TPair<int, FName>& InputDataInterfaceIndex : InputDataInterfaceIndexAndPin)
 		{
 			SetupAllInputBindings(InputDataInterfaceIndex);
 		}
 
-		for (int OutputDataInterfaceIndex : OutputDataInterfaceIndices)
+		for (const TPair<int, FName>& OutputDataInterfaceIndex : OutputDataInterfaceIndexAndPin)
 		{
 			SetupAllOutputBindings(OutputDataInterfaceIndex);
 		}
@@ -981,7 +976,7 @@ void FPCGGraphCompilerGPU::BuildGPUGraphTask(
 			}
 		}
 
-		auto AddAllEdgesForKernel = [&KernelWithBindings, ComputeGraph, TaskId, &DataInterfaceUsageToPinLabel](int32 InKernelIndex, bool bInEdgesAreInputs)
+		auto AddAllEdgesForKernel = [&KernelWithBindings, ComputeGraph](int32 InKernelIndex, bool bInEdgesAreInputs)
 		{
 			TArray<FInterfaceBinding>& Bindings = bInEdgesAreInputs ? KernelWithBindings.InputDataBindings : KernelWithBindings.OutputDataBindings;
 
@@ -1002,7 +997,7 @@ void FPCGGraphCompilerGPU::BuildGPUGraphTask(
 				UComputeDataInterface* DataInterface = ComputeGraph->DataInterfaces[Edge.DataInterfaceIndex];
 				check(DataInterface);
 
-				if (FName* PinLabel = DataInterfaceUsageToPinLabel.Find({ TaskId, DataInterface }))
+				if (Binding.PinLabel != NAME_None)
 				{
 					TArray<FShaderFunctionDefinition> DataInterfaceFunctions;
 					if (bInEdgesAreInputs)
@@ -1016,7 +1011,7 @@ void FPCGGraphCompilerGPU::BuildGPUGraphTask(
 
 					Edge.BindingFunctionNameOverride = FString::Format(
 						TEXT("{0}_{1}"),
-						{ PinLabel->ToString(), DataInterfaceFunctions[Edge.DataInterfaceBindingIndex].Name }
+						{ Binding.PinLabel.ToString(), DataInterfaceFunctions[Edge.DataInterfaceBindingIndex].Name }
 					);
 				}
 			}
