Exposed static mesh spawner classes to make it extensible

#rnx

diff --git a/Engine/Plugins/PCG/Source/PCG/Public/Elements/PCGStaticMeshSpawner.h b/Engine/Plugins/PCG/Source/PCG/Public/Elements/PCGStaticMeshSpawner.h
--- a/Engine/Plugins/PCG/Source/PCG/Public/Elements/PCGStaticMeshSpawner.h
+++ b/Engine/Plugins/PCG/Source/PCG/Public/Elements/PCGStaticMeshSpawner.h
@@ -20,8 +20,8 @@ struct FPCGStaticMeshSpawnerContext;
 
 class UStaticMesh;
 
-UCLASS(MinimalAPI, BlueprintType, ClassGroup = (Procedural), meta = (PrioritizeCategories = "Settings"))
-class UPCGStaticMeshSpawnerSettings : public UPCGSettings
+UCLASS(BlueprintType, ClassGroup = (Procedural), meta = (PrioritizeCategories = "Settings"))
+class PCG_API UPCGStaticMeshSpawnerSettings : public UPCGSettings
 {
 	GENERATED_BODY()
 
@@ -53,10 +53,10 @@ public:
 #endif
 
 	UFUNCTION(BlueprintCallable, Category = Settings)
-	PCG_API void SetMeshSelectorType(TSubclassOf<UPCGMeshSelectorBase> InMeshSelectorType);
+	void SetMeshSelectorType(TSubclassOf<UPCGMeshSelectorBase> InMeshSelectorType);
 
 	UFUNCTION(BlueprintCallable, Category = Settings)
-	PCG_API void SetInstancePackerType(TSubclassOf<UPCGInstanceDataPackerBase> InInstancePackerType);
+	void SetInstancePackerType(TSubclassOf<UPCGInstanceDataPackerBase> InInstancePackerType);
 
 public:
 	/** Defines the method of mesh selection per input data */
@@ -118,7 +118,7 @@ protected:
 	void RefreshInstancePacker();
 };
 
-class FPCGStaticMeshSpawnerElement : public IPCGElement
+class PCG_API FPCGStaticMeshSpawnerElement : public IPCGElement
 {
 public:
 	virtual bool CanExecuteOnlyOnMainThread(FPCGContext* Context) const override;
diff --git a/Engine/Plugins/PCG/Source/PCG/Public/Elements/PCGStaticMeshSpawnerContext.h b/Engine/Plugins/PCG/Source/PCG/Public/Elements/PCGStaticMeshSpawnerContext.h
--- a/Engine/Plugins/PCG/Source/PCG/Public/Elements/PCGStaticMeshSpawnerContext.h
+++ b/Engine/Plugins/PCG/Source/PCG/Public/Elements/PCGStaticMeshSpawnerContext.h
@@ -20,11 +20,11 @@ struct FPCGInstancesAndWeights
 };
 
 USTRUCT(BlueprintType)
-struct FPCGStaticMeshSpawnerContext : public FPCGContext, public IPCGAsyncLoadingContext
+struct PCG_API FPCGStaticMeshSpawnerContext : public FPCGContext, public IPCGAsyncLoadingContext
 {
 	GENERATED_BODY()
 
-	struct FPackedInstanceListData
+	struct PCG_API FPackedInstanceListData
 	{
 		FPackedInstanceListData();
 		~FPackedInstanceListData();
diff --git a/Engine/Plugins/PCG/Source/PCG/Public/MeshSelectors/PCGMeshSelectorBase.h b/Engine/Plugins/PCG/Source/PCG/Public/MeshSelectors/PCGMeshSelectorBase.h
--- a/Engine/Plugins/PCG/Source/PCG/Public/MeshSelectors/PCGMeshSelectorBase.h
+++ b/Engine/Plugins/PCG/Source/PCG/Public/MeshSelectors/PCGMeshSelectorBase.h
@@ -53,7 +53,7 @@ enum class EPCGMeshSelectorMaterialOverrideMode : uint8
 };
 
 /** Struct used to efficiently gather overrides and cache them during instance packing */
-struct FPCGMeshMaterialOverrideHelper
+struct PCG_API FPCGMeshMaterialOverrideHelper
 {
 	FPCGMeshMaterialOverrideHelper() = default;
 
