[PCG] Add a Spawn Spline Component node

* Mark some functions const in spline data.

#rnx

diff --git a/Engine/Plugins/PCG/Source/PCG/Private/Data/PCGSplineData.cpp b/Engine/Plugins/PCG/Source/PCG/Private/Data/PCGSplineData.cpp
--- a/Engine/Plugins/PCG/Source/PCG/Private/Data/PCGSplineData.cpp
+++ b/Engine/Plugins/PCG/Source/PCG/Private/Data/PCGSplineData.cpp
@@ -71,7 +71,7 @@ void UPCGSplineData::Initialize(const FPCGSplineStruct& InSplineStruct)
 	CachedBounds = CachedBounds.TransformBy(SplineStruct.Transform);
 }
 
-void UPCGSplineData::ApplyTo(USplineComponent* InSplineComponent)
+void UPCGSplineData::ApplyTo(USplineComponent* InSplineComponent) const
 {
 	SplineStruct.ApplyTo(InSplineComponent);
 }
diff --git a/Engine/Plugins/PCG/Source/PCG/Private/Data/PCGSplineStruct.cpp b/Engine/Plugins/PCG/Source/PCG/Private/Data/PCGSplineStruct.cpp
--- a/Engine/Plugins/PCG/Source/PCG/Private/Data/PCGSplineStruct.cpp
+++ b/Engine/Plugins/PCG/Source/PCG/Private/Data/PCGSplineStruct.cpp
@@ -102,7 +102,7 @@ void FPCGSplineStruct::Initialize(const TArray<FSplinePoint>& InSplinePoints, bo
 	LocalBounds = PCGSplineStruct::CalcBounds(SplineCurves, bClosedLoop, FTransform::Identity);
 }
 
-void FPCGSplineStruct::ApplyTo(USplineComponent* InSplineComponent)
+void FPCGSplineStruct::ApplyTo(USplineComponent* InSplineComponent) const
 {
 	check(InSplineComponent);
 
diff --git a/Engine/Plugins/PCG/Source/PCG/Private/Elements/PCGSpawnSpline.cpp b/Engine/Plugins/PCG/Source/PCG/Private/Elements/PCGSpawnSpline.cpp
new file mode 100644
--- /dev/null
+++ b/Engine/Plugins/PCG/Source/PCG/Private/Elements/PCGSpawnSpline.cpp
@@ -0,0 +1,166 @@
+// Copyright Epic Games, Inc. All Rights Reserved.
+
+#include "Elements/PCGSpawnSpline.h"
+
+#include "PCGComponent.h"
+#include "PCGContext.h"
+#include "PCGManagedResource.h"
+#include "Data/PCGSplineData.h"
+#include "Elements/Metadata/PCGMetadataElementCommon.h"
+#include "Helpers/PCGHelpers.h"
+
+#define LOCTEXT_NAMESPACE "PCGSpawnSplineElement"
+
+#if WITH_EDITOR
+FName UPCGSpawnSplineSettings::GetDefaultNodeName() const
+{
+	return FName(TEXT("SpawnSplineComponent"));
+}
+
+FText UPCGSpawnSplineSettings::GetDefaultNodeTitle() const
+{
+	return LOCTEXT("NodeTitle", "Spawn Spline Component");
+}
+#endif // WITH_EDITOR
+
+FPCGElementPtr UPCGSpawnSplineSettings::CreateElement() const
+{
+	return MakeShared<FPCGSpawnSplineElement>();
+}
+
+TArray<FPCGPinProperties> UPCGSpawnSplineSettings::InputPinProperties() const
+{
+	TArray<FPCGPinProperties> Properties;
+	Properties.Emplace_GetRef(PCGPinConstants::DefaultInputLabel, EPCGDataType::Spline).SetRequiredPin();
+	return Properties;
+}
+
+TArray<FPCGPinProperties> UPCGSpawnSplineSettings::OutputPinProperties() const
+{ 
+	TArray<FPCGPinProperties> Properties;
+	Properties.Emplace(PCGPinConstants::DefaultOutputLabel, EPCGDataType::Spline);
+	return Properties;
+}
+
+bool FPCGSpawnSplineElement::ExecuteInternal(FPCGContext* InContext) const
+{
+	TRACE_CPUPROFILER_EVENT_SCOPE(FPCGSpawnSplineElement::Execute);
+
+	check(InContext);
+
+	const UPCGSpawnSplineSettings* Settings = InContext->GetInputSettings<UPCGSpawnSplineSettings>();
+	check(Settings);
+
+	if (!*Settings->SplineComponent || !InContext->SourceComponent.IsValid())
+	{
+		return true;
+	}
+
+	AActor* TargetActor = Settings->TargetActor.Get() ? Settings->TargetActor.Get() : InContext->GetTargetActor(nullptr);
+	if (!TargetActor)
+	{
+		PCGLog::LogErrorOnGraph(LOCTEXT("InvalidTargetActor", "Invalid target actor. Ensure TargetActor member is initialized when creating SpatialData."), InContext);
+		return true;
+	}
+
+	const TArray<FPCGTaggedData> Inputs = InContext->InputData.GetInputsByPin(PCGPinConstants::DefaultInputLabel);
+	
+	for (const FPCGTaggedData& InputData : Inputs)
+	{
+		const UPCGSplineData* InSplineData = Cast<const UPCGSplineData>(InputData.Data);
+		if (!InSplineData)
+		{
+			continue;
+		}
+
+		TSubclassOf<USplineComponent> SplineComponentToSpawn = Settings->SplineComponent;
+		if (Settings->bSpawnComponentFromAttribute)
+		{
+			const FPCGAttributePropertyInputSelector InputSelector = Settings->SpawnComponentFromAttributeName.CopyAndFixLast(InSplineData);
+			TUniquePtr<const IPCGAttributeAccessor> Accessor = PCGAttributeAccessorHelpers::CreateConstAccessor(InSplineData, InputSelector);
+			if (!Accessor)
+			{
+				PCGLog::Metadata::LogFailToCreateAccessorError(InputSelector, InContext);
+				continue;
+			}
+
+			FPCGAttributeAccessorKeysEntries DefaultEntry(PCGInvalidEntryKey);
+			FSoftObjectPath SplineComponentClassPath;
+			if (!Accessor->Get<FSoftObjectPath>(SplineComponentClassPath, DefaultEntry, EPCGAttributeAccessorFlags::AllowBroadcastAndConstructible))
+			{
+				PCGLog::Metadata::LogFailToGetAttributeError<FSoftObjectPath>(InputSelector, Accessor.Get(), InContext);
+				continue;
+			}
+
+			SplineComponentToSpawn = Cast<UClass>(SplineComponentClassPath.TryLoad());
+		}
+
+		if (!*SplineComponentToSpawn)
+		{
+			PCGLog::LogErrorOnGraph(LOCTEXT("InvalidSplineComponent", "Invalid spline component subclass."), InContext);
+			continue;
+		}
+
+		check(IsInGameThread());
+		USplineComponent* SplineComponent = NewObject<USplineComponent>(TargetActor, *SplineComponentToSpawn);
+		SplineComponent->ComponentTags.Add(InContext->SourceComponent.Get()->GetFName());
+		SplineComponent->ComponentTags.Add(PCGHelpers::DefaultPCGTag);
+
+		// Also propagate the tags on the tagged data to the component.
+		Algo::Transform(InputData.Tags, SplineComponent->ComponentTags, [](const FString& Tag) { return FName(Tag); });
+
+		InSplineData->ApplyTo(SplineComponent);
+
+		SplineComponent->RegisterComponent();
+		TargetActor->AddInstanceComponent(SplineComponent);
+		// Spline Data apply to Spline Component is done in world coordinate, so keep world location.
+		SplineComponent->AttachToComponent(TargetActor->GetRootComponent(), FAttachmentTransformRules(EAttachmentRule::KeepWorld, EAttachmentRule::KeepWorld, EAttachmentRule::KeepWorld, false));
+
+		UPCGManagedComponent* ManagedComponent = NewObject<UPCGManagedComponent>(InContext->SourceComponent.Get());
+		ManagedComponent->GeneratedComponent = SplineComponent;
+		InContext->SourceComponent->AddToManagedResources(ManagedComponent);
+
+		if (!Settings->PropertyOverrideDescriptions.IsEmpty())
+		{
+			FPCGObjectOverrides DescriptorOverride(SplineComponent);
+			DescriptorOverride.Initialize(Settings->PropertyOverrideDescriptions, SplineComponent, InSplineData, InContext);
+			if (DescriptorOverride.IsValid() && !DescriptorOverride.Apply(0))
+			{
+				const int32 Index = static_cast<int32>(&InputData - Inputs.GetData());
+				PCGLog::LogWarningOnGraph(FText::Format(LOCTEXT("FailOverride", "Failed to override descriptor for input {0}"), Index));
+			}
+
+			// Apply 0, will be the default values for the spline
+			DescriptorOverride.Apply(0);
+		}
+
+		// Execute PostProcess Functions
+		if (TargetActor)
+		{
+			for (UFunction* Function : PCGHelpers::FindUserFunctions(TargetActor->GetClass(), Settings->PostProcessFunctionNames, { UPCGFunctionPrototypes::GetPrototypeWithNoParams() }, InContext))
+			{
+				TargetActor->ProcessEvent(Function, nullptr);
+			}
+		}
+		
+		if (Settings->bOutputSplineComponentReference)
+		{
+			UPCGSplineData* OutputSplineData = CastChecked<UPCGSplineData>(InSplineData->DuplicateData(InContext));
+
+			if (!PCGMetadataElementCommon::ClearOrCreateAttribute<FSoftObjectPath>(OutputSplineData->Metadata, Settings->ComponentReferenceAttributeName, FSoftObjectPath(SplineComponent)))
+			{
+				PCGLog::Metadata::LogFailToCreateAttributeError<FSoftObjectPath>(Settings->ComponentReferenceAttributeName, InContext);
+			}
+			
+			InContext->OutputData.TaggedData.Emplace_GetRef(InputData).Data = OutputSplineData;
+		}
+		else
+		{
+			InContext->OutputData.TaggedData.Emplace(InputData);
+		}
+	}
+
+	return true;
+}
+
+#undef LOCTEXT_NAMESPACE
diff --git a/Engine/Plugins/PCG/Source/PCG/Public/Data/PCGSplineData.h b/Engine/Plugins/PCG/Source/PCG/Public/Data/PCGSplineData.h
--- a/Engine/Plugins/PCG/Source/PCG/Public/Data/PCGSplineData.h
+++ b/Engine/Plugins/PCG/Source/PCG/Public/Data/PCGSplineData.h
@@ -24,7 +24,7 @@ public:
 	void Initialize(const USplineComponent* InSpline);
 	void Initialize(const TArray<FSplinePoint>& InSplinePoints, bool bInClosedLoop, const FTransform& InTransform);
 	void Initialize(const FPCGSplineStruct& InSplineStruct);
-	void ApplyTo(USplineComponent* InSpline);
+	void ApplyTo(USplineComponent* InSpline) const;
 
 	// ~Begin UPCGData interface
 	virtual EPCGDataType GetDataType() const override { return EPCGDataType::Spline; }
diff --git a/Engine/Plugins/PCG/Source/PCG/Public/Data/PCGSplineStruct.h b/Engine/Plugins/PCG/Source/PCG/Public/Data/PCGSplineStruct.h
--- a/Engine/Plugins/PCG/Source/PCG/Public/Data/PCGSplineStruct.h
+++ b/Engine/Plugins/PCG/Source/PCG/Public/Data/PCGSplineStruct.h
@@ -15,7 +15,7 @@ struct PCG_API FPCGSplineStruct
 
 	void Initialize(const USplineComponent* InSplineComponent);
 	void Initialize(const TArray<FSplinePoint>& InSplinePoints, bool bIsClosedLoop, const FTransform& InTransform);
-	void ApplyTo(USplineComponent* InSplineComponent);
+	void ApplyTo(USplineComponent* InSplineComponent) const;
 
 	FTransform GetTransform() const { return Transform; }
 
diff --git a/Engine/Plugins/PCG/Source/PCG/Public/Elements/PCGSpawnSpline.h b/Engine/Plugins/PCG/Source/PCG/Public/Elements/PCGSpawnSpline.h
new file mode 100644
--- /dev/null
+++ b/Engine/Plugins/PCG/Source/PCG/Public/Elements/PCGSpawnSpline.h
@@ -0,0 +1,74 @@
+// Copyright Epic Games, Inc. All Rights Reserved.
+
+#pragma once
+
+#include "PCGSettings.h"
+
+#include "Components/SplineComponent.h"
+#include "Elements/PCGAddComponent.h"
+#include "Metadata/PCGObjectPropertyOverride.h"
+
+#include "PCGSpawnSpline.generated.h"
+
+/**
+* Spawn a spline component from a spline data.
+*/
+UCLASS(MinimalAPI, BlueprintType, ClassGroup = (Procedural))
+class UPCGSpawnSplineSettings : public UPCGSettings
+{
+	GENERATED_BODY()
+
+public:
+	//~Begin UPCGSettings interface
+#if WITH_EDITOR
+	virtual FName GetDefaultNodeName() const override;
+	virtual FText GetDefaultNodeTitle() const override;
+	virtual EPCGSettingsType GetType() const override { return EPCGSettingsType::Spawner; }
+#endif
+
+protected:
+	virtual TArray<FPCGPinProperties> InputPinProperties() const override;
+	virtual TArray<FPCGPinProperties> OutputPinProperties() const override;
+	virtual FPCGElementPtr CreateElement() const override;
+	//~End UPCGSettings interface
+
+public:
+	/** Class of the component to spawn, must be a subclass of Spline Component. */
+	UPROPERTY(BlueprintReadWrite, EditAnywhere, Category="Settings", meta = (OnlyPlaceable, PCG_Overridable))
+	TSubclassOf<USplineComponent> SplineComponent = USplineComponent::StaticClass();
+
+	UPROPERTY(BlueprintReadWrite, EditAnywhere, Category="Settings", meta = (PCG_Overridable, InlineEditConditionToggle))
+	bool bSpawnComponentFromAttribute = false;
+
+	/** If the class of the component to spawn is coming from an attribute. */
+	UPROPERTY(BlueprintReadWrite, EditAnywhere, Category="Settings", meta = (PCG_Overridable, EditCondition = bSpawnComponentFromAttribute, PCG_DiscardPropertySelection, PCG_DiscardExtraSelection))
+	FPCGAttributePropertyInputSelector SpawnComponentFromAttributeName;
+
+	UPROPERTY(meta = (PCG_Overridable))
+	TSoftObjectPtr<AActor> TargetActor;
+
+	/** Specify a list of functions to be called on the target actor after spline creation. Functions need to be parameter-less and with "CallInEditor" flag enabled. */
+	UPROPERTY(BlueprintReadWrite, EditAnywhere, Category = Settings)
+	TArray<FName> PostProcessFunctionNames;
+
+	/** Overrides to apply on the spawned component. */
+	UPROPERTY(BlueprintReadWrite, EditAnywhere, Category="Settings")
+	TArray<FPCGObjectPropertyOverrideDescription> PropertyOverrideDescriptions;
+
+	UPROPERTY(BlueprintReadWrite, EditAnywhere, Category="Settings", meta = (PCG_Overridable, InlineEditConditionToggle))
+	bool bOutputSplineComponentReference = true;
+
+	/** Can output the spawned component reference in an attribute. */
+	UPROPERTY(BlueprintReadWrite, EditAnywhere, Category="Settings", meta = (PCG_Overridable, EditCondition = bOutputSplineComponentReference))
+	FName ComponentReferenceAttributeName = PCGAddComponentConstants::ComponentReferenceAttribute;
+};
+
+class FPCGSpawnSplineElement : public IPCGElement
+{
+public:
+	virtual bool CanExecuteOnlyOnMainThread(FPCGContext*) const override { return true; };
+	virtual bool IsCacheable(const UPCGSettings*) const override { return false; };
+protected:
+	virtual bool ExecuteInternal(FPCGContext* InContext) const override;
+};
+
