[PCG] Builder: New Iterative mode (-IterativeCellLoading, IterativeCellSize=Value)

- New mode will create missing Partition Actors, Add/Remove PCG Components to existing Partition Actors based on Original Component volume intersection
- Generates based on same parameters as non iterative mode for now
- Doesn't support Unbounded Partitioned components (next task)

diff --git a/Engine/Plugins/PCG/Source/PCG/Private/Grid/PCGPartitionActor.cpp b/Engine/Plugins/PCG/Source/PCG/Private/Grid/PCGPartitionActor.cpp
--- a/Engine/Plugins/PCG/Source/PCG/Private/Grid/PCGPartitionActor.cpp
+++ b/Engine/Plugins/PCG/Source/PCG/Private/Grid/PCGPartitionActor.cpp
@@ -186,6 +186,17 @@ void APCGPartitionActor::PostRegisterAllComponents()
 	UpdateBoundsComponentExtents();
 #endif // WITH_EDITOR
 
+	RebuildOriginalToLocal();
+
+	// Make the Partition actor register itself to the PCG Subsystem. RuntimeGen components should only register in the PostCreation path.
+	if (PCGGridSize != InvalidPCGGridSizeValue && !IsRuntimeGenerated())
+	{
+		RegisterPCG();
+	}
+}
+
+void APCGPartitionActor::RebuildOriginalToLocal()
+{
 	// Reset the OriginalToLocal and build it from the local map
 	OriginalToLocal.Reset();
 	OriginalToLocal.Reserve(LocalToOriginal.Num());
@@ -193,12 +204,6 @@ void APCGPartitionActor::PostRegisterAllComponents()
 	{
 		OriginalToLocal.Add(It.Value.Get(), It.Key);
 	}
-
-	// Make the Partition actor register itself to the PCG Subsystem. RuntimeGen components should only register in the PostCreation path.
-	if (PCGGridSize != InvalidPCGGridSizeValue && !IsRuntimeGenerated())
-	{
-		RegisterPCG();
-	}
 }
 
 void APCGPartitionActor::PostUnregisterAllComponents()
@@ -288,6 +293,12 @@ void APCGPartitionActor::SetInvalidForPCG()
 		SetActorLabel(TEXT("TO_DELETE_") + GetActorLabel());
 	}
 }
+
+TSoftObjectPtr<UPCGComponent> APCGPartitionActor::GetOriginalComponentSoftObjectPtr(UPCGComponent* LocalComponent) const
+{
+	const TSoftObjectPtr<UPCGComponent>* OriginalComponent = LocalToOriginal.Find(LocalComponent);
+	return OriginalComponent ? *OriginalComponent : TSoftObjectPtr<UPCGComponent>();
+}
 #endif
 
 FBox APCGPartitionActor::GetFixedBounds() const
@@ -383,8 +394,22 @@ void APCGPartitionActor::GetActorBounds(bool bOnlyCollidingComponents, FVector&
 
 UPCGComponent* APCGPartitionActor::GetLocalComponent(const UPCGComponent* OriginalComponent) const
 {
-	const TObjectPtr<UPCGComponent>* LocalComponent = OriginalToLocal.Find(OriginalComponent);
-	return LocalComponent ? *LocalComponent : nullptr;
+	return GetLocalComponent(OriginalComponent, /*bRebuildMappingOnNullEntries=*/true);
+}
+
+UPCGComponent* APCGPartitionActor::GetLocalComponent(const UPCGComponent* OriginalComponent, bool bRebuildMappingOnNullEntries) const
+{
+	if (const TObjectPtr<UPCGComponent>* LocalComponent = OriginalToLocal.Find(OriginalComponent))
+	{
+		return *LocalComponent;
+	}
+	else if (bRebuildMappingOnNullEntries && OriginalToLocal.Contains(nullptr))
+	{
+		const_cast<APCGPartitionActor*>(this)->RebuildOriginalToLocal();
+		return GetLocalComponent(OriginalComponent, /*bRebuildMappingOnNullEntries=*/false);
+	}
+
+	return nullptr;
 }
 
 UPCGComponent* APCGPartitionActor::GetOriginalComponent(const UPCGComponent* LocalComponent) const
diff --git a/Engine/Plugins/PCG/Source/PCG/Private/PCGActorAndComponentMapping.cpp b/Engine/Plugins/PCG/Source/PCG/Private/PCGActorAndComponentMapping.cpp
--- a/Engine/Plugins/PCG/Source/PCG/Private/PCGActorAndComponentMapping.cpp
+++ b/Engine/Plugins/PCG/Source/PCG/Private/PCGActorAndComponentMapping.cpp
@@ -856,6 +856,16 @@ void FPCGActorAndComponentMapping::UpdateMappingPCGComponentPartitionActor(UPCGC
 	}
 }
 
+TSet<TObjectPtr<APCGPartitionActor>> FPCGActorAndComponentMapping::GetPCGComponentPartitionActorMappings(UPCGComponent* InComponent) const
+{
+	if (const TSet<TObjectPtr<APCGPartitionActor>>* PartitionActorsPtr = ComponentToPartitionActorsMap.Find(InComponent))
+	{
+		return *PartitionActorsPtr;
+	}
+
+	return TSet<TObjectPtr<APCGPartitionActor>>();
+}
+
 void FPCGActorAndComponentMapping::DeleteMappingPCGComponentPartitionActor(UPCGComponent* InComponent)
 {
 	check(InComponent);
diff --git a/Engine/Plugins/PCG/Source/PCG/Private/PCGSubsystem.cpp b/Engine/Plugins/PCG/Source/PCG/Private/PCGSubsystem.cpp
--- a/Engine/Plugins/PCG/Source/PCG/Private/PCGSubsystem.cpp
+++ b/Engine/Plugins/PCG/Source/PCG/Private/PCGSubsystem.cpp
@@ -128,6 +128,10 @@ namespace PCGSubsystemConsole
 #endif
 }
 
+#if WITH_EDITOR
+TSet<UWorld*> UPCGSubsystem::DisablePartitionActorCreationForWorld;
+#endif
+
 UPCGSubsystem::UPCGSubsystem()
 	: Super()
 	, ActorAndComponentMapping(this)
@@ -1801,11 +1805,11 @@ void UPCGSubsystem::CreateMissingPartitionActors()
 
 void UPCGSubsystem::CreatePartitionActorsWithinBounds(UPCGComponent* InComponent, const FBox& InBounds, const PCGHiGenGrid::FSizeArray& InGridSizes)
 {
-	if (!PCGHelpers::IsRuntimeOrPIE())
+	UWorld* World = GetWorld();
+	if (!PCGHelpers::IsRuntimeOrPIE() && !IsPartitionActorCreationDisabledForWorld(World))
 	{
 		// We can't spawn actors if we are running constructions scripts, asserting when we try to get the actor with the WP API.
 		// We should never enter this if we are in a construction script. If the ensure is hit, we need to fix it.
-		UWorld* World = GetWorld();
 		if (ensure(World && !World->bIsRunningConstructionScript))
 		{
 			ForAllOverlappingCells(InComponent, InBounds, InGridSizes, true, {}, [](APCGPartitionActor*, const FBox&) { return InvalidPCGTaskId; });
@@ -1813,6 +1817,16 @@ void UPCGSubsystem::CreatePartitionActorsWithinBounds(UPCGComponent* InComponent
 	}
 }
 
+void UPCGSubsystem::UpdateMappingPCGComponentPartitionActor(UPCGComponent* InComponent) 
+{ 
+	ActorAndComponentMapping.UpdateMappingPCGComponentPartitionActor(InComponent); 
+}
+
+TSet<TObjectPtr<APCGPartitionActor>> UPCGSubsystem::GetPCGComponentPartitionActorMappings(UPCGComponent* InComponent) const 
+{ 
+	return ActorAndComponentMapping.GetPCGComponentPartitionActorMappings(InComponent); 
+}
+
 #endif // WITH_EDITOR
 
 void UPCGSubsystem::ExecuteBeginTickActions()
diff --git a/Engine/Plugins/PCG/Source/PCG/Public/Grid/PCGPartitionActor.h b/Engine/Plugins/PCG/Source/PCG/Public/Grid/PCGPartitionActor.h
--- a/Engine/Plugins/PCG/Source/PCG/Public/Grid/PCGPartitionActor.h
+++ b/Engine/Plugins/PCG/Source/PCG/Public/Grid/PCGPartitionActor.h
@@ -126,10 +126,16 @@ private:
 #if WITH_EDITOR
 	void UpdateBoundsComponentExtents();
 	void SetInvalidForPCG();
+
+	friend struct FPCGWorldPartitionBuilder;
+	TSoftObjectPtr<UPCGComponent> GetOriginalComponentSoftObjectPtr(UPCGComponent* LocalComponent) const;
 #endif // WITH_EDITOR
 
+	void RebuildOriginalToLocal();
+	UPCGComponent* GetLocalComponent(const UPCGComponent* OriginalComponent, bool bRebuildMappingOnNullEntries) const;
+
 	// Note: this map is not a property and not serialized since we will rebuild it from the LocalToOriginal
-	TMap<TObjectPtr<UPCGComponent>, TObjectPtr<UPCGComponent>> OriginalToLocal;
+	mutable TMap<TObjectPtr<UPCGComponent>, TObjectPtr<UPCGComponent>> OriginalToLocal;
 
 	UPROPERTY(NonTransactional)
 	TMap<TObjectPtr<UPCGComponent>, TSoftObjectPtr<UPCGComponent>> LocalToOriginal;
diff --git a/Engine/Plugins/PCG/Source/PCG/Public/PCGActorAndComponentMapping.h b/Engine/Plugins/PCG/Source/PCG/Public/PCGActorAndComponentMapping.h
--- a/Engine/Plugins/PCG/Source/PCG/Public/PCGActorAndComponentMapping.h
+++ b/Engine/Plugins/PCG/Source/PCG/Public/PCGActorAndComponentMapping.h
@@ -151,6 +151,9 @@ private:
 	/** Update the current mapping between a PCG component and its PCG Partition actors */
 	void UpdateMappingPCGComponentPartitionActor(UPCGComponent* InComponent);
 
+	/** Returns the current mapping between a PCG component and its PCG Partition actors */
+	TSet<TObjectPtr<APCGPartitionActor>> GetPCGComponentPartitionActorMappings(UPCGComponent* InComponent) const;
+	
 	/** Delete the current mapping between a PCG component and its PCG Partition actors */
 	void DeleteMappingPCGComponentPartitionActor(UPCGComponent* InComponent);
 
diff --git a/Engine/Plugins/PCG/Source/PCG/Public/PCGComponent.h b/Engine/Plugins/PCG/Source/PCG/Public/PCGComponent.h
--- a/Engine/Plugins/PCG/Source/PCG/Public/PCGComponent.h
+++ b/Engine/Plugins/PCG/Source/PCG/Public/PCGComponent.h
@@ -87,6 +87,7 @@ class PCG_API UPCGComponent : public UActorComponent
 	friend class UPCGManagedActors;
 	friend class UPCGSubsystem;
 	friend class FPCGActorAndComponentMapping;
+	friend struct FPCGWorldPartitionBuilder;
 
 public:
 	/** ~Begin UObject interface */
diff --git a/Engine/Plugins/PCG/Source/PCG/Public/PCGSubsystem.h b/Engine/Plugins/PCG/Source/PCG/Public/PCGSubsystem.h
--- a/Engine/Plugins/PCG/Source/PCG/Public/PCGSubsystem.h
+++ b/Engine/Plugins/PCG/Source/PCG/Public/PCGSubsystem.h
@@ -55,6 +55,7 @@ class PCG_API UPCGSubsystem : public UTickableWorldSubsystem
 
 public:
 	friend FPCGActorAndComponentMapping;
+	friend struct FPCGWorldPartitionBuilder;
 
 	UPCGSubsystem();
 
@@ -351,6 +352,8 @@ private:
 	};
 
 	void CreatePartitionActorsWithinBounds(UPCGComponent* InComponent, const FBox& InBounds, const PCGHiGenGrid::FSizeArray& InGridSizes);
+	void UpdateMappingPCGComponentPartitionActor(UPCGComponent* InComponent);
+	TSet<TObjectPtr<APCGPartitionActor>> GetPCGComponentPartitionActorMappings(UPCGComponent* InComponent) const;
 
 	FPCGNodeVisualLogs NodeVisualLogs;
 #endif // WITH_EDITOR
@@ -376,6 +379,26 @@ private:
 		
 	using FConstructionScriptSourceComponents = TMap<FName, TObjectKey<UPCGComponent>>;
 	TMap<TObjectKey<AActor>, FConstructionScriptSourceComponents> PerActorConstructionScriptSourceComponents;
+
+	static TSet<UWorld*> DisablePartitionActorCreationForWorld;
+
+	// Used by UPCGWorldPartitonBuilder to disable PA creation while outside of a certain scope
+	static void SetDisablePartitionActorCreationForWorld(UWorld* InWorld, bool bDisable) 
+	{ 
+		if (bDisable)
+		{
+			DisablePartitionActorCreationForWorld.Add(InWorld);
+		}
+		else
+		{
+			DisablePartitionActorCreationForWorld.Remove(InWorld);
+		}
+	}
+
+	static bool IsPartitionActorCreationDisabledForWorld(UWorld* InWorld)
+	{
+		return DisablePartitionActorCreationForWorld.Contains(InWorld);
+	}
 #endif
 };
 
diff --git a/Engine/Plugins/PCG/Source/PCGEditor/Private/WorldPartitionBuilder/PCGBuilderSettings.h b/Engine/Plugins/PCG/Source/PCGEditor/Private/WorldPartitionBuilder/PCGBuilderSettings.h
--- a/Engine/Plugins/PCG/Source/PCGEditor/Private/WorldPartitionBuilder/PCGBuilderSettings.h
+++ b/Engine/Plugins/PCG/Source/PCGEditor/Private/WorldPartitionBuilder/PCGBuilderSettings.h
@@ -33,6 +33,14 @@ public:
 	UPROPERTY(EditAnywhere, Category = Build, meta = (DisplayName = "One Component at a Time"))
 	bool bOneComponentAtATime = false;
 
+	/** Run builder iteratively loading cell regions instead of the whole map at once. Useful for large maps. */
+	UPROPERTY(EditAnywhere, Category = Advanced)
+	bool bIterativeCellLoading = false;
+
+	/** Cell size when using iterative cell loading. */
+	UPROPERTY(EditAnywhere, Category = Advanced, meta = (EditCondition=bIterativeCellLoading, EditConditionHides, ClampMin=12800, UIMin=12800))
+	int32 IterativeCellSize = 25600;
+
 	/** Submit dirty files even if errors occurred during generation. */
 	UPROPERTY(EditAnywhere, Category = Advanced)
 	bool bIgnoreGenerationErrors = false;
diff --git a/Engine/Plugins/PCG/Source/PCGEditor/Private/WorldPartitionBuilder/PCGWorldPartitionBuilder.cpp b/Engine/Plugins/PCG/Source/PCGEditor/Private/WorldPartitionBuilder/PCGWorldPartitionBuilder.cpp
--- a/Engine/Plugins/PCG/Source/PCGEditor/Private/WorldPartitionBuilder/PCGWorldPartitionBuilder.cpp
+++ b/Engine/Plugins/PCG/Source/PCGEditor/Private/WorldPartitionBuilder/PCGWorldPartitionBuilder.cpp
@@ -6,48 +6,83 @@
 #include "PCGComponent.h"
 #include "PCGEditorSettings.h"
 #include "PCGGraph.h"
+#include "PCGSubsystem.h"
 #include "SPCGBuilderDialog.h"
 
+#include "Grid/PCGPartitionActor.h"
+#include "Helpers/PCGHelpers.h"
+
 #include "AssetCompilingManager.h"
 #include "Editor.h"
+#include "EngineUtils.h"
 #include "FileHelpers.h"
+#include "ISourceControlModule.h"
+#include "SourceControlHelpers.h"
+#include "SourceControlOperations.h"
 
+#include "Algo/Copy.h"
+#include "Algo/Transform.h"
 #include "Framework/Application/SlateApplication.h"
 #include "Interfaces/IMainFrameModule.h"
 #include "Misc/OutputDevice.h"
 #include "Modules/ModuleManager.h"
 #include "UObject/Linker.h"
 #include "WorldPartition/IWorldPartitionEditorModule.h"
+#include "WorldPartition/WorldPartition.h"
 #include "WorldPartition/WorldPartitionHelpers.h"
+#include "WorldPartition/LoaderAdapter/LoaderAdapterShape.h"
 
 DEFINE_LOG_CATEGORY_STATIC(LogPCGWorldPartitionBuilder, All, All);
 
 #define LOCTEXT_NAMESPACE "PCGWorldPartitionBulder"
 
-namespace PCGWorldPartitionBuilder
+struct FPCGWorldPartitionBuilder
 {
-	/** Grab all original PCG components in world, applying given filter. */
-	void CollectComponentsToGenerate(UWorld* InWorld, TFunctionRef<bool(const UPCGComponent*)> ComponentFilter, TArray<TWeakObjectPtr<UPCGComponent>>& OutComponents);
-	
 	/** Generate the given components. Optionally generate only one component at a time with a wait on async processes after each. Optionally apply given filter to select components. */
-	bool GenerateComponents(TArray<TWeakObjectPtr<UPCGComponent>>& Components, UWorld* InWorld, bool bOneComponentAtATime, TArray<TObjectPtr<UPackage>>& InOutDeletedActorPackages, bool& bOutGenerationErrors);
-	bool GenerateComponents(
-		TArray<TWeakObjectPtr<UPCGComponent>>& Components,
-		UWorld* InWorld,
-		bool bOneComponentAtATime,
-		TFunctionRef<bool(const UPCGComponent*)> ComponentFilter,
-		TArray<TObjectPtr<UPackage>>& InOutDeletedActorPackages,
-		bool& bOutGenerationErrors);
+	static bool GenerateComponents(TArray<TWeakObjectPtr<UPCGComponent>>& Components, UWorld* InWorld, const FPCGWorldPartitionBuilderArgs& InArgs, TSet<FString>& InOutDeletedActorPackages, bool& bOutGenerationErrors);
 
 	/** Generate a component. Applies correct editing mode if necessary. */
-	void GenerateComponent(UPCGComponent* InComponent, UWorld* InWorld);
+	static void GenerateComponent(UPCGComponent* InComponent, UWorld* InWorld);
+
+	/** Groups components into buckets to generate in order */
+	static TArray<TArray<TWeakObjectPtr<UPCGComponent>>> GroupComponentsToGenerate(const TArray<TWeakObjectPtr<UPCGComponent>>& InComponentsToGenerate, const FPCGWorldPartitionBuilderArgs& InArgs);
 
 	/** Waits for background activity to be quiet. */
-	void WaitForAllAsyncEditorProcesses(UWorld* InWorld);
+	static void WaitForAllAsyncEditorProcesses(UWorld* InWorld);
 
 	/** Runs builder on current editor world. */
-	bool Build(UWorld* InWorld, const TArray<FString>& Args);
+	static bool Build(UWorld* InWorld, const TArray<FString>& Args);
+
+	/** Load Actors this Partition Actor depends on : Original Component Actor and Generated Managed Actors */
+	static void LoadPartitionedActorDependencies(UWorld* InWorld, TSet<FWorldPartitionReference>& OutLoadedActors);
+
+	/** Creates missing Partitioned Actors with the original components */
+	static bool CreatePartitionedActors(UWorld* InWorld, const FPCGWorldPartitionBuilderArgs& InArgs, TMap<TSoftObjectPtr<UPCGComponent>, FWorldPartitionReference>& InOutPartitionedComponents);
+
+	/** Update mappings to loaded Partitioned Actors */
+	static bool UpdatePartitionedActors(UWorld* InWorld, const TMap<TSoftObjectPtr<UPCGComponent>, FWorldPartitionReference>& InPartitionedComponents, TSet<FString>& InOutDeletedActorPackages);
+
+	/** Grab all original PCG components in world, applying given filter. */
+	static void CollectComponentsToGenerate(UWorld* InWorld, const FPCGWorldPartitionBuilderArgs& InArgs, TFunctionRef<bool(const UPCGComponent*)> ComponentFilter, TArray<TWeakObjectPtr<UPCGComponent>>& OutComponents);
+
+	/** Reset loaded package dirty flags */
+	static void ResetPackageDirtyFlags();
+
+	/** Extract an Actor Object Path from a PCG Component Soft Object Ptr */
+	static FSoftObjectPath GetActorPathFromComponentPtr(TSoftObjectPtr<UPCGComponent> InComponentPtr);
+
+	/** Get the list of Dirty packages */
+	static TArray<UPackage*> GetDirtyPackages();
+
+	/** Set if Partition Actor creation should be disabled for world */
+	static void SetDisablePartitionActorCreationForWorld(UWorld* InWorld, bool bDisable);
+
+	/** Expand Loaded Bounds if needed */
+	static TUniquePtr<IWorldPartitionActorLoaderInterface::ILoaderAdapter> ExpandLoadedBoundsIfNeeded(UWorld* InWorld, const FBox& InCurrentBounds, const FBox& InNewBounds, const FString& InLoaderName);
+};
 
+namespace PCGWorldPartitionBuilder
+{
 	static FAutoConsoleCommand CommandBuildComponents(
 		TEXT("pcg.BuildComponents"),
 		TEXT("Runs PCG world builder on PCG components in current world. Arguments (multiple values separated with ';'), Flags can optionally have a specified '=true or false', true if none provided:\n"
@@ -58,8 +93,72 @@ namespace PCGWorldPartitionBuilder
 			"\t[-GenerateComponentEditingModePreview] (default: false)\n"
 			"\t[-IgnoreGenerationErrors] (default: false)\n"
 			"\t[-IncludeActorIDs=MyActor1_UAID1234678;MyActor2_UAID1234678] (default: all actors)\n"
-			"\t[-OneComponentAtATime] (default: false)\n"),
-		FConsoleCommandWithArgsDelegate::CreateLambda([](const TArray<FString>& Args) { PCGWorldPartitionBuilder::Build(GEditor ? GEditor->GetEditorWorldContext().World() : nullptr, Args); }));
+			"\t[-OneComponentAtATime] (default: false)\n"
+			"\t[-IterativeCellLoading] (default: false)\n"
+			"\t[-IterativeCellSize=25600] (default: 25600)\n"),
+		FConsoleCommandWithArgsDelegate::CreateLambda([](const TArray<FString>& Args) { FPCGWorldPartitionBuilder::Build(GEditor ? GEditor->GetEditorWorldContext().World() : nullptr, Args); }));
+};
+
+/** Scoped object class that tracks if any package was dirtied during its lifetime */
+class FPCGDetectDirtyPackageInScope
+{
+public:
+	FPCGDetectDirtyPackageInScope()
+	{
+		UPackage::PackageDirtyStateChangedEvent.AddRaw(this, &FPCGDetectDirtyPackageInScope::HandleDirtyPackage);
+	}
+
+	~FPCGDetectDirtyPackageInScope()
+	{
+		UPackage::PackageDirtyStateChangedEvent.RemoveAll(this);
+	}
+
+	bool AnyDirtyPackage() { return bAnyDirtyPackage; }
+
+private:
+	void HandleDirtyPackage(UPackage* InPackage)
+	{
+		bAnyDirtyPackage = true;
+	}
+
+	bool bAnyDirtyPackage = false;
+};
+
+/** Scoped object class that tracks deleted actors during its lifetime */
+class FPCGDetectDeletedActorInScope
+{
+public:
+	FPCGDetectDeletedActorInScope(UWorld* InWorld, TSet<FString>& InDeletedActorPackages)
+		: World(InWorld), DeletedActorPackages(InDeletedActorPackages)
+	{
+		GEngine->OnLevelActorDeleted().AddRaw(this, &FPCGDetectDeletedActorInScope::HandleDeletedActor);
+	}
+
+	~FPCGDetectDeletedActorInScope()
+	{
+		GEngine->OnLevelActorDeleted().RemoveAll(this);
+	}
+
+private:
+	void HandleDeletedActor(AActor* InActor)
+	{
+		if (InActor && InActor->IsPackageExternal())
+		{
+			if (UPackage* ActorPackage = InActor->GetPackage())
+			{
+				UE_LOG(LogPCGWorldPartitionBuilder, Display, TEXT("Actor '%s' deleted, package '%s' added to delete list."),
+					*InActor->GetName(),
+					*ActorPackage->GetName());
+
+				// Notify World Partition of Deleted Package right away in case Actor gets Garbage Collected before saving
+				World->GetWorldPartition()->OnPackageDeleted(ActorPackage);
+				DeletedActorPackages.Add(ActorPackage->GetName());
+			}
+		}
+	}
+
+	UWorld* World = nullptr;
+	TSet<FString>& DeletedActorPackages;
 };
 
 /** Output device to capture presence of errors during generation. */
@@ -120,6 +219,8 @@ FPCGWorldPartitionBuilderArgs FPCGWorldPartitionBuilderArgs::InitializeFrom(cons
 			BuilderArgs.bOneComponentAtATime = BuilderSettings->bOneComponentAtATime;
 			BuilderArgs.IncludeActorIDs = BuilderSettings->FilterByActorNames;
 			BuilderArgs.IncludeGraphs = BuilderSettings->Graphs;
+			BuilderArgs.bIterativeCellLoading = BuilderSettings->bIterativeCellLoading;
+			BuilderArgs.IterativeCellSize = BuilderSettings->IterativeCellSize;
 
 			UE_LOG(LogPCGWorldPartitionBuilder, Display, TEXT("Loaded PCGBuilderSettings '%s'"), *SettingsPath.ToString());
 			bUsingBuilderSettings = true;
@@ -226,6 +327,28 @@ FPCGWorldPartitionBuilderArgs FPCGWorldPartitionBuilderArgs::InitializeFrom(cons
 		GraphsToString.IsEmpty() ? TEXT("None") : *FString::Join(GraphsToString, TEXT(",")),
 		FromDefault);
 
+	// Init 'bIterativeCellLoading'
+	bFromCommandline = CommandlineArgs.bIterativeCellLoading.IsSet();
+	if (bFromCommandline)
+	{
+		BuilderArgs.bIterativeCellLoading = CommandlineArgs.bIterativeCellLoading.GetValue();
+	}
+
+	UE_LOG(LogPCGWorldPartitionBuilder, Display, TEXT("bIterativeCellLoading : '%d' (%s)"),
+		BuilderArgs.bIterativeCellLoading ? 1 : 0,
+		bFromCommandline ? FromCommandline : FromDefault);
+
+	// Init 'IterativeCellSize'
+	bFromCommandline = CommandlineArgs.IterativeCellSize.IsSet();
+	if (bFromCommandline)
+	{
+		BuilderArgs.IterativeCellSize = CommandlineArgs.IterativeCellSize.GetValue();
+	}
+
+	UE_LOG(LogPCGWorldPartitionBuilder, Display, TEXT("IterativeCellSize : '%d' (%s)"),
+		BuilderArgs.IterativeCellSize,
+		bFromCommandline ? FromCommandline : FromDefault);
+
 	return BuilderArgs;
 }
 
@@ -253,6 +376,7 @@ UPCGWorldPartitionBuilder::UPCGWorldPartitionBuilder(const FObjectInitializer& O
 		GetBoolParam("GenerateComponentEditingModePreview", CommandlineArgs.bGenerateEditingModePreviewComponents);
 		GetBoolParam("OneComponentAtATime", CommandlineArgs.bOneComponentAtATime);
 		GetBoolParam("IgnoreGenerationErrors", CommandlineArgs.bIgnoreGenerationErrors);
+		GetBoolParam("IterativeCellLoading", CommandlineArgs.bIterativeCellLoading);
 	
 		FString IncludeGraphNamesValue;
 		if (GetParamValue("IncludeGraphNames=", IncludeGraphNamesValue))
@@ -281,7 +405,35 @@ UPCGWorldPartitionBuilder::UPCGWorldPartitionBuilder(const FObjectInitializer& O
 		{
 			CommandlineArgs.SettingAsset = TSoftObjectPtr<UPCGBuilderSettings>(FSoftObjectPath(SettingAsset));
 		}
+
+		int32 ParamIterativeCellSize = 25600;
+		if (GetParamValue("IterativeCellSize=", ParamIterativeCellSize))
+		{
+			const int32 MinIterativeCellSize = 12800;
+			const int32 MaxIterativeCellSize = FMath::Max(MinIterativeCellSize, ParamIterativeCellSize);
+			if (MaxIterativeCellSize == ParamIterativeCellSize)
+			{
+				CommandlineArgs.IterativeCellSize = ParamIterativeCellSize;
+			}
+			else
+			{
+				UE_LOG(LogPCGWorldPartitionBuilder, Warning, TEXT("Minimum value for 'IterativeCellSize' is %d"), MinIterativeCellSize);
+			}
+		}
+	}
+}
+
+bool UPCGWorldPartitionBuilder::PreWorldInitialization(UWorld* World, FPackageSourceControlHelper& PackageHelper)
+{
+	if (!Super::PreWorldInitialization(World, PackageHelper))
+	{
+		return false;
 	}
+
+	// Prevent Actor Creation outside of run scope
+	FPCGWorldPartitionBuilder::SetDisablePartitionActorCreationForWorld(World, /*bDisable=*/true);
+
+	return true;
 }
 
 bool UPCGWorldPartitionBuilder::PreRun(UWorld* World, FPackageSourceControlHelper& PackageHelper)
@@ -292,7 +444,23 @@ bool UPCGWorldPartitionBuilder::PreRun(UWorld* World, FPackageSourceControlHelpe
 	}
 
 	Args = FPCGWorldPartitionBuilderArgs::InitializeFrom(CommandlineArgs);
-		
+
+	// Set base class member
+	IterativeCellSize = Args.IterativeCellSize;
+
+	// If World doesn't have a World Partition do not use iterative cell loading
+	Args.bIterativeCellLoading &= World->IsPartitionedWorld();
+
+	if (GetLoadingMode() == ELoadingMode::IterativeCells2D)
+	{
+		// Disable Source Control until Post Run
+		PreviousProviderName = ISourceControlModule::Get().GetProvider().GetName();
+		if (PreviousProviderName != NAME_None)
+		{
+			ISourceControlModule::Get().SetProvider(NAME_None);
+		}
+	}
+
 	return true;
 }
 
@@ -304,125 +472,74 @@ bool UPCGWorldPartitionBuilder::RunInternal(UWorld* World, const FCellInfo& InCe
 	}
 
 	// Try to eliminate as much background noise as possible before getting started.
-	PCGWorldPartitionBuilder::WaitForAllAsyncEditorProcesses(World);
+	FPCGWorldPartitionBuilder::WaitForAllAsyncEditorProcesses(World);
+
+	if (GetLoadingMode() == ELoadingMode::EntireWorld)
+	{
+		return RunEntireWorld(World, PackageHelper);
+	}
+	else
+	{
+		return RunPerCell(World, InCellInfo, PackageHelper);
+	}
+}
 
+bool UPCGWorldPartitionBuilder::RunEntireWorld(UWorld* World, FPackageSourceControlHelper& PackageHelper)
+{
 	// 1. Collect components of interest once upfront.
 
 	TArray<TWeakObjectPtr<UPCGComponent>> ComponentsToGenerate;
 	{
-		auto ComponentFilter = [this](const UPCGComponent* InComponent)
+		auto ComponentFilter = [](const UPCGComponent* InComponent)
 		{
-			if (!InComponent || !InComponent->GetOwner())
-			{
-				return false;
-			}
-
-			// Check actor in inclusion list if provided.
-			if (!Args.IncludeActorIDs.IsEmpty() && !Args.IncludeActorIDs.Contains(InComponent->GetOwner()->GetName()))
-			{
-				return false;
-			}
-
-			// Accept based on editing mode.
-			return (Args.bGenerateEditingModeLoadAsPreviewComponents && InComponent->GetSerializedEditingMode() == EPCGEditorDirtyMode::LoadAsPreview)
-				|| (Args.bGenerateEditingModeNormalComponents && InComponent->GetSerializedEditingMode() == EPCGEditorDirtyMode::Normal)
-				|| (Args.bGenerateEditingModePreviewComponents && InComponent->GetSerializedEditingMode() == EPCGEditorDirtyMode::Preview);
+			// Skip LCs, rely on original component to generate everything.
+			return !InComponent->IsLocalComponent();
 		};
 
-		PCGWorldPartitionBuilder::CollectComponentsToGenerate(World, ComponentFilter, ComponentsToGenerate);
+		FPCGWorldPartitionBuilder::CollectComponentsToGenerate(World, Args, ComponentFilter, ComponentsToGenerate);
 	}
 
 	// 2. Clear the dirty flag on any packages that dirtied themselves while loading. We're not interested in saving these.
-
-	{
-		TArray<UPackage*> DirtyPackages;
-		FEditorFileUtils::GetDirtyWorldPackages(DirtyPackages);
-		DirtyPackages.RemoveSwap(nullptr);
-		for (UPackage* DirtyPackage : DirtyPackages)
-		{
-			// Make clean as it's not a change we have made.
-			if (!PendingDirtyPackages.Contains(DirtyPackage->GetPersistentGuid()))
-			{
-				DirtyPackage->SetDirtyFlag(false);
-			}
-		}
-		DirtyPackages.Reset();
-	}
-
-	const int32 NumPendingDirtyPackagesBefore = PendingDirtyPackages.Num();
-
+	FPCGWorldPartitionBuilder::ResetPackageDirtyFlags();
+		
+	// Create missing Partition Actors
+	FPCGWorldPartitionBuilder::SetDisablePartitionActorCreationForWorld(World, /*bDisable=*/false);
+	FPCGWorldPartitionBuilder::CreatePartitionedActors(World, Args, PartitionedComponents);
+	
+	// Update mappings
+	TSet<FString> DeletedActorPackages;
+	FPCGWorldPartitionBuilder::UpdatePartitionedActors(World, PartitionedComponents, DeletedActorPackages);
+	
 	// 3. Generate components.
-
 	bool bGeneratedAnyComponent = false;
-
-	if (!Args.IncludeGraphNames.IsEmpty())
-	{
-		// Generate all components with graph of each name, in specified order.
-		for (const FString& IncludeGraphName : Args.IncludeGraphNames)
-		{
-			auto FilterOnGraphName = [IncludeGraphName](const UPCGComponent* InComponent)
-			{
-				const UPCGGraph* Graph = InComponent ? InComponent->GetGraph() : nullptr;
-				return Graph && Graph->GetName() == IncludeGraphName;
-			};
-
-			bool bErrorsOccurred = false;
-			bGeneratedAnyComponent |= PCGWorldPartitionBuilder::GenerateComponents(ComponentsToGenerate, World, Args.bOneComponentAtATime, FilterOnGraphName, DeletedActorPackages, bErrorsOccurred);
-
-			bErrorOccurredWhileGenerating |= bErrorsOccurred;
-		}
-	}
-	else if (!Args.IncludeGraphs.IsEmpty())
-	{
-		// Generate all components with graph of each name, in specified order.
-		for (const TSoftObjectPtr<UPCGGraphInterface>& IncludeGraphPtr : Args.IncludeGraphs)
-		{
-			auto FilterOnGraphPtr = [&IncludeGraphPtr](const UPCGComponent* InComponent)
-				{
-					TSoftObjectPtr<UPCGGraphInterface> ComponentGraphPtr(InComponent ? InComponent->GetGraph() : nullptr);
-					return ComponentGraphPtr.IsValid() && ComponentGraphPtr == IncludeGraphPtr;
-				};
-
-			bool bErrorsOccurred = false;
-			bGeneratedAnyComponent |= PCGWorldPartitionBuilder::GenerateComponents(ComponentsToGenerate, World, Args.bOneComponentAtATime, FilterOnGraphPtr, DeletedActorPackages, bErrorsOccurred);
-
-			bErrorOccurredWhileGenerating |= bErrorsOccurred;
-		}
-	}
-	else
+		
+	TArray<TArray<TWeakObjectPtr<UPCGComponent>>> GroupsToGenerate = FPCGWorldPartitionBuilder::GroupComponentsToGenerate(ComponentsToGenerate, Args);
+	for (TArray<TWeakObjectPtr<UPCGComponent>>& GroupToGenerate : GroupsToGenerate)
 	{
 		bool bErrorsOccurred = false;
-		bGeneratedAnyComponent |= PCGWorldPartitionBuilder::GenerateComponents(ComponentsToGenerate, World, Args.bOneComponentAtATime, DeletedActorPackages, bErrorsOccurred);
-
+		bGeneratedAnyComponent |= FPCGWorldPartitionBuilder::GenerateComponents(GroupToGenerate, World, Args, DeletedActorPackages, bErrorsOccurred);
 		bErrorOccurredWhileGenerating |= bErrorsOccurred;
 	}
 
 	// 4. Get packages that were dirtied during generation and record them for saving later.
+	TArray<UPackage*> DirtyPackages;
 
 	if (bGeneratedAnyComponent)
 	{
-		TArray<UPackage*> DirtyPackages;
-		FEditorFileUtils::GetDirtyWorldPackages(DirtyPackages);
-		
-		DirtyPackages.RemoveSwap(nullptr);
-		for (UPackage* DirtyPackage : DirtyPackages)
-		{
-			PendingDirtyPackages.Add(DirtyPackage->GetPersistentGuid(), DirtyPackage);
-		}
+		DirtyPackages = FPCGWorldPartitionBuilder::GetDirtyPackages();
 	}
-	
-	const int NumPackagesDirtied = PendingDirtyPackages.Num() - NumPendingDirtyPackagesBefore;
+
 	if (!bGeneratedAnyComponent)
 	{
 		UE_LOG(LogPCGWorldPartitionBuilder, Display, TEXT("No components found to generate. No packages will be saved."));
 	}
-	else if (NumPackagesDirtied <= 0)
+	else if (DirtyPackages.IsEmpty())
 	{
 		UE_LOG(LogPCGWorldPartitionBuilder, Display, TEXT("At least one component was generated but no additional packages were dirtied."));
 	}
 	else
 	{
-		UE_LOG(LogPCGWorldPartitionBuilder, Display, TEXT("Generation complete, %d packages dirtied."), NumPackagesDirtied);
+		UE_LOG(LogPCGWorldPartitionBuilder, Display, TEXT("Generation complete, %d packages dirtied."), DirtyPackages.Num());
 	}
 
 	if (!bGeneratedAnyComponent)
@@ -431,11 +548,10 @@ bool UPCGWorldPartitionBuilder::RunInternal(UWorld* World, const FCellInfo& InCe
 		return !bErrorOccurredWhileGenerating;
 	}
 
-	// TODO: Review the save flow when we have iterative loading.
-	return SaveDirtyPackages(World, PackageHelper);
+	return SaveEntireWorldPackages(World, PackageHelper, MoveTemp(DirtyPackages), MoveTemp(DeletedActorPackages));
 }
 
-bool UPCGWorldPartitionBuilder::SaveDirtyPackages(UWorld* World, FPackageSourceControlHelper& PackageHelper)
+bool UPCGWorldPartitionBuilder::SaveEntireWorldPackages(UWorld* World, FPackageSourceControlHelper& PackageHelper, TArray<UPackage*> DirtyPackages, TSet<FString> DeletedActorPackages)
 {
 	// Check whether an error was thrown while generating components and fail the builder if the ignore argument is not provided.
 	if (bErrorOccurredWhileGenerating)
@@ -454,221 +570,355 @@ bool UPCGWorldPartitionBuilder::SaveDirtyPackages(UWorld* World, FPackageSourceC
 
 	// Save/delete pending packages.
 
-	TArray<TObjectPtr<UPackage>> DirtyPackages;
-	PendingDirtyPackages.GenerateValueArray(DirtyPackages);
-
 	// Empty packages should be deleted - mirrors logic in InternalPromptForCheckoutAndSave()
 	TArray<UPackage*> PackagesToDelete;
+	TArray<UPackage*> PackagesToSave;
+
 	for (int i = DirtyPackages.Num() - 1; i >= 0; --i)
 	{
-		if (UPackage::IsEmptyPackage(ToRawPtr(DirtyPackages[i])))
+		if (UPackage* DirtyPackage = DirtyPackages[i])
 		{
-			PackagesToDelete.Add(DirtyPackages[i]);
-			DirtyPackages.RemoveAtSwap(i);
-		}
-	}
-
-	UE_LOG(LogPCGWorldPartitionBuilder, Display, TEXT("PostRun: %d packages modified, %d actor packages deleted, %d empty packages marked for delete."),
-		DirtyPackages.Num(),
-		DeletedActorPackages.Num(),
-		PackagesToDelete.Num());
+			if (UPackage::IsEmptyPackage(DirtyPackage))
+			{
+				PackagesToDelete.Add(DirtyPackage);
+			}
+			else
+			{
+				PackagesToSave.Add(DirtyPackage);
+			}
 
-	// Combine the empty-deleted packages with the actor-deleted packages.
-	for (UPackage* DeletedPackage : DeletedActorPackages)
-	{
-		PackagesToDelete.AddUnique(DeletedPackage);
+			// Remove from here if package was still part of the loaded DirtyPackages
+			DeletedActorPackages.Remove(DirtyPackage->GetName());
+		}
 	}
 
 	// Log final changes after combining packages (there may have been duplicates in the two deleted package lists).
-	UE_LOG(LogPCGWorldPartitionBuilder, Display, TEXT("PostRun: Final package changes: %d modified, %d deleted."), DirtyPackages.Num(), PackagesToDelete.Num());
+	UE_LOG(LogPCGWorldPartitionBuilder, Display, TEXT("PostRun: Final package changes: %d modified, %d deleted."), DirtyPackages.Num(), DeletedActorPackages.Num() + PackagesToDelete.Num());
 
-	if (!SavePackages(DirtyPackages, PackageHelper))
+	if (!SavePackages(PackagesToSave, PackageHelper))
 	{
 		return false;
 	}
-	
+
+	// Loaded Packages to delete
 	if (!PackagesToDelete.IsEmpty())
 	{
-		// Release any file handles and prepare for delete.
-		for (UPackage* PackageToDelete : PackagesToDelete)
-		{
-			ResetLoaders(PackageToDelete);
-		}
-
-		if (!DeletePackages(PackagesToDelete, PackageHelper))
+		if (!PackageHelper.Delete(PackagesToDelete))
 		{
 			return false;
 		}
 	}
 
+	// Unloaded Packages to delete
+	if (!DeletedActorPackages.IsEmpty())
+	{
+		PackageHelper.Delete(DeletedActorPackages.Array());
+	}
+
 	TArray<FString> FilesToSubmit;
 	FilesToSubmit.Append(SourceControlHelpers::PackageFilenames(DirtyPackages));
 	FilesToSubmit.Append(SourceControlHelpers::PackageFilenames(PackagesToDelete));
+	FilesToSubmit.Append(SourceControlHelpers::PackageFilenames(DeletedActorPackages.Array()));
 
 	const FString ChangeDescription = FString::Printf(TEXT("Generated PCG components for world '%s'"), *World->GetName());
 	return OnFilesModified(FilesToSubmit, ChangeDescription);
 }
 
-void PCGWorldPartitionBuilder::CollectComponentsToGenerate(
-	UWorld* InWorld,
-	TFunctionRef<bool(const UPCGComponent*)> ComponentFilter,
-	TArray<TWeakObjectPtr<UPCGComponent>>& OutComponents)
+bool UPCGWorldPartitionBuilder::RunPerCell(UWorld* InWorld, const FCellInfo& InCellInfo, FPackageSourceControlHelper& PackageHelper)
 {
-	check(InWorld);
+	ON_SCOPE_EXIT
+	{
+		// Since we are potentially releasing actor references when exiting the RunPerCell call it is important to run a GC so that actors can be properly unloaded (and may be reloaded in another Cell)
+		FWorldPartitionHelpers::DoCollectGarbage();
 
-	TArray<UObject*> AllComponents;
-	GetObjectsOfClass(UPCGComponent::StaticClass(), AllComponents, /*bIncludeDerivedClasses=*/true);
-	for (UObject* ComponentObject : AllComponents)
+		FPCGWorldPartitionBuilder::SetDisablePartitionActorCreationForWorld(InWorld, /*bDisable=*/true);
+	};
+
+	// Load missing dependencies (SoftObjectPath references)	
+	TSet<FWorldPartitionReference> LoadedActorDependencies;
+	FPCGWorldPartitionBuilder::LoadPartitionedActorDependencies(InWorld, LoadedActorDependencies);
+
+	// Reset Dirty flags so we can save only packages that will get modified after this line
+	FPCGWorldPartitionBuilder::ResetPackageDirtyFlags();
+
+	// Create missing PAs and Update mappings
+	FPCGWorldPartitionBuilder::SetDisablePartitionActorCreationForWorld(InWorld, /*bDisable=*/false);
+
+	// Create Missing PAs
+	bool bCellHasChanges = FPCGWorldPartitionBuilder::CreatePartitionedActors(InWorld, Args, PartitionedComponents);
+
+	// Update Mappings
+	TSet<FString> CellDeletedPackageNames;
+	bCellHasChanges |= FPCGWorldPartitionBuilder::UpdatePartitionedActors(InWorld, PartitionedComponents, CellDeletedPackageNames);
+
+	// Find Components to generate and make sure we load their bounds
+	TArray<TWeakObjectPtr<UPCGComponent>> ComponentsToGenerate;
+	FBox ExpandedBounds = InCellInfo.Bounds;
+
+	auto GenerateComponentFilter = [&InCellInfo, &ExpandedBounds, this](const UPCGComponent* InComponent)
 	{
-		if (!IsValid(ComponentObject))
+		if (InComponent->IsPartitioned())
 		{
-			continue;
+			return false;
 		}
 
-		UPCGComponent* Component = Cast<UPCGComponent>(ComponentObject);
-		if (!Component || !Component->GetOwner() || Component->GetWorld() != InWorld)
+		if (GeneratedComponents.Contains(InComponent))
 		{
-			continue;
+			return false;
 		}
 
-		if (Component->IsLocalComponent())
+		// Generate Component if its bound center is inside the currently processing cell bounds (to avoid expanding bounds too much)
+		const FBox ActorBounds = PCGHelpers::GetActorBounds(InComponent->GetOwner());
+		if (!InCellInfo.Bounds.IsInsideOrOnXY(ActorBounds.GetCenter()))
 		{
-			// Skip LCs, rely on original component to generate everything.
-			continue;
+			return false;
 		}
 
-		if (!ComponentFilter(Component))
+		// Expand the cell bounds to include the whole actor bounds (keep Z same as original Cell bounds)
+		FBox ActorXYBounds(FVector(ActorBounds.Min.X, ActorBounds.Min.Y, ExpandedBounds.Min.Z), FVector(ActorBounds.Max.X, ActorBounds.Max.Y, ExpandedBounds.Max.Z));
+		ExpandedBounds += ActorXYBounds;
+
+		// Consider this component generated
+		GeneratedComponents.Add(TSoftObjectPtr<const UPCGComponent>(InComponent));
+
+		return true;
+	};
+
+	// Collect Components to Generate a first time, this will allow us to call FPCGWorldPartitionBuilder::UpdatePartitionedActors on ParttionActors with all their overlapp
+	FPCGWorldPartitionBuilder::CollectComponentsToGenerate(InWorld, Args, GenerateComponentFilter, ComponentsToGenerate);
+		
+	// Do a Bounds expansion to make sure we are generating PCG Components on areas covering their owning actor bounds (which might be a bigger area then the currently loaded cell)
+	TUniquePtr<IWorldPartitionActorLoaderInterface::ILoaderAdapter> LoaderAdapter = FPCGWorldPartitionBuilder::ExpandLoadedBoundsIfNeeded(InWorld, InCellInfo.Bounds, ExpandedBounds, TEXT("Expanded Bounds"));
+	
+	bool bErrorsOccurred = false;
+	// Do Generation and Collect Packages to save
+	if (ComponentsToGenerate.Num() > 0)
+	{
+		TArray<TArray<TWeakObjectPtr<UPCGComponent>>> GroupsToGenerate = FPCGWorldPartitionBuilder::GroupComponentsToGenerate(ComponentsToGenerate, Args);
+		for (TArray<TWeakObjectPtr<UPCGComponent>>& GroupToGenerate : GroupsToGenerate)
 		{
-			continue;
+			bCellHasChanges |= FPCGWorldPartitionBuilder::GenerateComponents(GroupToGenerate, InWorld, Args, CellDeletedPackageNames, bErrorsOccurred);
+			bErrorOccurredWhileGenerating |= bErrorsOccurred;
 		}
+	}
 
-		OutComponents.Add(Component);
+	if (bErrorsOccurred && !Args.bIgnoreGenerationErrors)
+	{
+		if (bCellHasChanges)
+		{
+			UE_LOG(LogPCGWorldPartitionBuilder, Display, TEXT("Save skipped due to errors during generation."));
+		}
+		return false;
 	}
-}
 
-bool PCGWorldPartitionBuilder::GenerateComponents(TArray<TWeakObjectPtr<UPCGComponent>>& Components, UWorld* InWorld, bool bOneComponentAtATime, TArray<TObjectPtr<UPackage>>& InOutDeletedActorPackages, bool& bOutGenerationErrors)
-{
-	return PCGWorldPartitionBuilder::GenerateComponents(Components, InWorld, bOneComponentAtATime, [](const UPCGComponent*) { return true; }, InOutDeletedActorPackages, bOutGenerationErrors);
+	if (bCellHasChanges)
+	{
+		return SavePerCellPackages(InWorld, PackageHelper, CellDeletedPackageNames);
+	}
+
+	return true;
 }
 
-bool PCGWorldPartitionBuilder::GenerateComponents(
-	TArray<TWeakObjectPtr<UPCGComponent>>& Components,
-	UWorld* InWorld,
-	bool bOneComponentAtATime,
-	TFunctionRef<bool(const UPCGComponent*)> ComponentFilter,
-	TArray<TObjectPtr<UPackage>>& InOutDeletedActorPackages,
-	bool& bOutGenerationErrors)
+bool UPCGWorldPartitionBuilder::PostRun(UWorld* World, FPackageSourceControlHelper& PackageHelper, const bool bInRunSuccess)
 {
-	if (!bOneComponentAtATime)
+	// Log final changes after combining packages (there may have been duplicates in the two deleted package lists).
+	UE_LOG(LogPCGWorldPartitionBuilder, Display, TEXT("PostRun: Final package changes: %d new, %d modified, %d deleted."), NewPackageNames.Num(), ModifiedPackageNames.Num(), DeletedPackageNames.Num());
+
+	FPCGWorldPartitionBuilder::SetDisablePartitionActorCreationForWorld(World, /*bDisable=*/false);
+
+	if (GetLoadingMode() == ELoadingMode::EntireWorld)
 	{
-		PCGWorldPartitionBuilder::WaitForAllAsyncEditorProcesses(InWorld);
+		PartitionedComponents.Empty();
+		return true;
 	}
 
-	const FPCGDetectErrorsInScope DetectErrors;
+	// Builder was launched with no source control (nothing more to do)
+	if (PreviousProviderName.IsNone())
+	{
+		return true;
+	}
 
-	auto WaitForComponentGeneration = [InWorld](const UPCGComponent* InComponent)
+	// Re-establish SCC Provider
+	ISourceControlModule::Get().SetProvider(PreviousProviderName);
+	ISourceControlProvider::FInitResult InitResult = ISourceControlModule::Get().GetProvider().Init(ISourceControlProvider::EInitFlags::AttemptConnection);
+	if (!InitResult.bIsAvailable)
 	{
-		UE_LOG(LogPCGWorldPartitionBuilder, Display, TEXT("Completing generation on PCG component on actor '%s' label '%s', graph '%s'"),
-			*InComponent->GetOwner()->GetName(),
-			*InComponent->GetOwner()->GetActorNameOrLabel(),
-			*InComponent->GetGraph()->GetName());
+		UE_LOG(LogPCGWorldPartitionBuilder, Error, TEXT("Failed to initialize source control provider."));
+		return false;
+	}
 
-		while (InComponent->IsGenerating())
+	// Mark files for delete
+	if (DeletedPackageNames.Num() > 0)
+	{
+		TArray<FString> FilesToDelete;
+		FilesToDelete.Reserve(DeletedPackageNames.Num());
+		Algo::Transform(DeletedPackageNames, FilesToDelete, [](const FString& PackageName) { return SourceControlHelpers::PackageFilename(PackageName); });
+		if (!SourceControlHelpers::MarkFilesForDelete(FilesToDelete))
 		{
-			FWorldPartitionHelpers::FakeEngineTick(InWorld);
+			return false;
 		}
+	}
 
-		// Can be useful to let some things flush/update after generation.
-		FWorldPartitionHelpers::FakeEngineTick(InWorld);
-	};
+	TArray<FString> PackagesToCheckout;
+	TArray<FString> PackagesToAdd;
+	TArray<FString> PackageNames;
+	PackageNames.Reserve(NewPackageNames.Num() + ModifiedPackageNames.Num());
+	PackageNames.Append(NewPackageNames.Array());
+	PackageNames.Append(ModifiedPackageNames.Array());
 
-	// Hook actor deleted events and track any corresponding deleted packages, as the packages can be GC'd if PCG triggers a GC
-	// before generation.
-	FDelegateHandle ActorDeletedHandle = GEngine->OnLevelActorDeleted().AddLambda([&InOutDeletedActorPackages](AActor* InActor)
+	if (!PackageHelper.GetDesiredStatesForModification(PackageNames, PackagesToCheckout, PackagesToAdd))
 	{
-		if (InActor && InActor->IsPackageExternal())
-		{
-			if (UPackage* ActorPackage = InActor->GetPackage())
-			{
-				UE_LOG(LogPCGWorldPartitionBuilder, Display, TEXT("Actor '%s' deleted, package '%s' added to delete list."),
-					*InActor->GetName(),
-					*ActorPackage->GetName());
+		return false;
+	}
 
-				InOutDeletedActorPackages.AddUnique(ActorPackage);
-			}
+	// Checkout modified files
+	if (PackagesToCheckout.Num() > 0)
+	{
+		if (!PackageHelper.Checkout(PackagesToCheckout))
+		{
+			return false;
 		}
-	});
+	}
 
-	ON_SCOPE_EXIT
+	// Add new files to source control
+	if (PackagesToAdd.Num() > 0)
 	{
-		GEngine->OnLevelActorDeleted().Remove(ActorDeletedHandle);
-	};
+		if (!PackageHelper.AddToSourceControl(PackagesToAdd))
+		{
+			return false;
+		}
+	}
 
-	TSet<TObjectKey<UPCGComponent>> GeneratedComponents;
+	// Release references
+	PartitionedComponents.Empty();
 
-	for (TWeakObjectPtr<UPCGComponent> ComponentWeakPtr : Components)
+	TArray<FString> FilesToSubmit;
+	FilesToSubmit.Append(SourceControlHelpers::PackageFilenames(DeletedPackageNames.Array()));
+	FilesToSubmit.Append(SourceControlHelpers::PackageFilenames(PackagesToCheckout));
+	FilesToSubmit.Append(SourceControlHelpers::PackageFilenames(PackagesToAdd));
+
+	const FString ChangeDescription = FString::Printf(TEXT("Generated PCG components for world '%s'"), *World->GetName());
+	return OnFilesModified(FilesToSubmit, ChangeDescription);
+}
+
+bool UPCGWorldPartitionBuilder::SavePerCellPackages(UWorld* InWorld, FPackageSourceControlHelper& PackageHelper, TSet<FString> CellDeletedPackages)
+{
+	TArray<UPackage*> DirtyPackages = FPCGWorldPartitionBuilder::GetDirtyPackages();
+	TArray<UPackage*> PackagesToDelete;
+	UWorldPartition* WorldPartition = InWorld->GetWorldPartition();
+
+	uint32 CellNewPackageCount = 0;
+	uint32 CellModifiedPackageCount = 0;
+	uint32 CellDeletedPackageCount = 0;
+
+	for (int i = DirtyPackages.Num() - 1; i >= 0; --i)
 	{
-		UPCGComponent* Component = ComponentWeakPtr.Get();
-		if (!Component)
+		if (UPackage* DirtyPackage = DirtyPackages[i])
 		{
-			UE_LOG(LogPCGWorldPartitionBuilder, Warning, TEXT("Lost a PCG component weak pointer, component will not be generated."));
-			continue;
-		}
+			const FString PackageName = DirtyPackage->GetName();
+			if (UPackage::IsEmptyPackage(DirtyPackage))
+			{
+				DeletedPackageNames.Add(PackageName);
 
-		// Validate this before running the filter as the filtering can check the graph name etc.
-		const UPCGGraph* Graph = Component->GetGraph();
-		if (!Graph)
-		{
-			UE_LOG(LogPCGWorldPartitionBuilder, Warning, TEXT("PCG component on actor '%s' label '%s' has no graph assigned, skipping."),
-				*Component->GetOwner()->GetName(),
-				*Component->GetOwner()->GetActorNameOrLabel());
-			continue;
-		}
+				PackagesToDelete.Add(DirtyPackage);
 
-		if (!ComponentFilter(Component))
-		{
-			continue;
-		}
-		
-		// Last minute validations, done here just prior to generation (after component has passed all previous filters) to minimize spam.
-		if (!Component->bActivated)
-		{
-			UE_LOG(LogPCGWorldPartitionBuilder, Display, TEXT("'Activated' toggle was set false on PCG component on actor '%s' label '%s' graph '%s'. Component skipped."),
-				*Component->GetOwner()->GetName(),
-				*Component->GetOwner()->GetActorNameOrLabel(),
-				*Component->GetGraph()->GetName());
-			continue;
+				DirtyPackages.RemoveAtSwap(i);
+
+				CellDeletedPackageCount++;
+			}
+			else if (DirtyPackage->HasAnyPackageFlags(PKG_NewlyCreated))
+			{
+				NewPackageNames.Add(PackageName);
+
+				CellNewPackageCount++;
+			}
+			else if (!NewPackageNames.Contains(PackageName))
+			{
+				ModifiedPackageNames.Add(PackageName);
+
+				CellModifiedPackageCount++;
+			}
+
+			CellDeletedPackages.Remove(*PackageName);
 		}
+	}
 
-		if (Component->IsManagedByRuntimeGenSystem())
+	// Delete loaded packages
+	if (PackagesToDelete.Num() > 0 && !PackageHelper.Delete(PackagesToDelete))
+	{
+		return false;
+	}
+
+	// Delete unloaded packages
+	if (CellDeletedPackages.Num() > 0 && !PackageHelper.Delete(CellDeletedPackages.Array()))
+	{
+		return false;
+	}
+
+	// Save Dirty/New packages
+	if (DirtyPackages.Num() > 0 && !SavePackages(DirtyPackages, PackageHelper))
+	{
+		return false;
+	}
+
+	UE_LOG(LogPCGWorldPartitionBuilder, Display, TEXT("Cell package changes: %d new, %d modified, %d deleted."), CellNewPackageCount, CellModifiedPackageCount, CellDeletedPackageCount);
+	return true;
+}
+
+bool FPCGWorldPartitionBuilder::GenerateComponents(
+	TArray<TWeakObjectPtr<UPCGComponent>>& Components,
+	UWorld* InWorld,
+	const FPCGWorldPartitionBuilderArgs& InArgs,
+	TSet<FString>& InOutDeletedActorPackages,
+	bool& bOutGenerationErrors)
+{
+	if (!InArgs.bOneComponentAtATime)
+	{
+		FPCGWorldPartitionBuilder::WaitForAllAsyncEditorProcesses(InWorld);
+	}
+
+	const FPCGDetectErrorsInScope DetectErrors;
+
+	auto WaitForComponentGeneration = [InWorld](const UPCGComponent* InComponent)
+	{
+		UE_LOG(LogPCGWorldPartitionBuilder, Display, TEXT("Completing generation on PCG component on actor '%s' label '%s', graph '%s'"),
+			*InComponent->GetOwner()->GetName(),
+			*InComponent->GetOwner()->GetActorNameOrLabel(),
+			*InComponent->GetGraph()->GetName());
+
+		while (InComponent->IsGenerating())
 		{
-			UE_LOG(LogPCGWorldPartitionBuilder, Display, TEXT("PCG component generation trigger is set to run-time generation on actor '%s' label '%s' graph '%s'. Component skipped."),
-				*Component->GetOwner()->GetName(),
-				*Component->GetOwner()->GetActorNameOrLabel(),
-				*Component->GetGraph()->GetName());
-			continue;
+			FWorldPartitionHelpers::FakeEngineTick(InWorld);
 		}
 
-		if (bOneComponentAtATime)
+		// Can be useful to let some things flush/update after generation.
+		FWorldPartitionHelpers::FakeEngineTick(InWorld);
+	};
+
+	// Will track any deleted actors in this scope
+	FPCGDetectDeletedActorInScope DeletedActorScope(InWorld, InOutDeletedActorPackages);
+
+	TSet<TObjectKey<UPCGComponent>> GeneratedComponents;
+
+	for (TWeakObjectPtr<UPCGComponent> ComponentWeakPtr : Components)
+	{
+		UPCGComponent* Component = ComponentWeakPtr.Get();
+		if (!Component)
 		{
-			PCGWorldPartitionBuilder::WaitForAllAsyncEditorProcesses(InWorld);
+			UE_LOG(LogPCGWorldPartitionBuilder, Warning, TEXT("Lost a PCG component weak pointer, component will not be generated."));
+			continue;
 		}
 
-		UE_LOG(LogPCGWorldPartitionBuilder, Display, TEXT("Generate PCG component on actor '%s' label '%s', graph '%s'"),
-			*Component->GetOwner()->GetName(),
-			*Component->GetOwner()->GetActorNameOrLabel(),
-			*Component->GetGraph()->GetName());
-
-		PCGWorldPartitionBuilder::GenerateComponent(Component, InWorld);
+		FPCGWorldPartitionBuilder::GenerateComponent(Component, InWorld);
 		GeneratedComponents.Add(Component);
 
-		if (bOneComponentAtATime)
+		if (InArgs.bOneComponentAtATime)
 		{
 			WaitForComponentGeneration(Component);
 		}
 	}
 
-	if (!bOneComponentAtATime)
+	if (!InArgs.bOneComponentAtATime)
 	{
 		for (TObjectKey<UPCGComponent>& GeneratedComponent : GeneratedComponents)
 		{
@@ -688,7 +938,55 @@ bool PCGWorldPartitionBuilder::GenerateComponents(
 	return !GeneratedComponents.IsEmpty();
 }
 
-void PCGWorldPartitionBuilder::GenerateComponent(UPCGComponent* InComponent, UWorld* InWorld)
+TArray<TArray<TWeakObjectPtr<UPCGComponent>>> FPCGWorldPartitionBuilder::GroupComponentsToGenerate(const TArray<TWeakObjectPtr<UPCGComponent>>& InComponentsToGenerate, const FPCGWorldPartitionBuilderArgs& InArgs)
+{
+	TArray<TArray<TWeakObjectPtr<UPCGComponent>>> GroupsToGenerate;
+
+	if (!InArgs.IncludeGraphNames.IsEmpty())
+	{
+		for (const FString& IncludeGraphName : InArgs.IncludeGraphNames)
+		{
+			TArray<TWeakObjectPtr<UPCGComponent>> ComponentsForGraph;
+			Algo::CopyIf(InComponentsToGenerate, ComponentsForGraph, [IncludeGraphName](TWeakObjectPtr<UPCGComponent> InComponent)
+			{
+				if (UPCGComponent* PCGComponent = InComponent.Get())
+				{
+					return PCGComponent->GetGraph() && PCGComponent->GetGraph()->GetName() == IncludeGraphName;
+				}
+
+				return false;
+			});
+
+			GroupsToGenerate.Add(MoveTemp(ComponentsForGraph));
+		}
+	}
+	else if (!InArgs.IncludeGraphs.IsEmpty())
+	{
+		for (const TSoftObjectPtr<UPCGGraphInterface>& IncludeGraph : InArgs.IncludeGraphs)
+		{
+			TArray<TWeakObjectPtr<UPCGComponent>> ComponentsForGraph;
+			Algo::CopyIf(InComponentsToGenerate, ComponentsForGraph, [&IncludeGraph](TWeakObjectPtr<UPCGComponent> InComponent)
+			{
+				if (UPCGComponent* PCGComponent = InComponent.Get())
+				{
+					return TSoftObjectPtr<UPCGGraphInterface>(PCGComponent->GetGraph()) == IncludeGraph;
+				}
+
+				return false;
+			});
+
+			GroupsToGenerate.Add(MoveTemp(ComponentsForGraph));
+		}
+	}
+	else
+	{
+		GroupsToGenerate.Add(InComponentsToGenerate);
+	}
+
+	return GroupsToGenerate;
+}
+
+void FPCGWorldPartitionBuilder::GenerateComponent(UPCGComponent* InComponent, UWorld* InWorld)
 {
 	// Separate ensures for maximum debug information.
 	if (!ensure(InComponent) || !ensure(InComponent->GetGraph()))
@@ -731,7 +1029,7 @@ void PCGWorldPartitionBuilder::GenerateComponent(UPCGComponent* InComponent, UWo
 	}
 }
 
-void PCGWorldPartitionBuilder::WaitForAllAsyncEditorProcesses(UWorld* InWorld)
+void FPCGWorldPartitionBuilder::WaitForAllAsyncEditorProcesses(UWorld* InWorld)
 {
 	InWorld->BlockTillLevelStreamingCompleted();
 
@@ -760,7 +1058,48 @@ void PCGWorldPartitionBuilder::WaitForAllAsyncEditorProcesses(UWorld* InWorld)
 	FAssetCompilingManager::Get().ProcessAsyncTasks();
 }
 
-bool PCGWorldPartitionBuilder::Build(UWorld* InWorld, const TArray<FString>& Args)
+void FPCGWorldPartitionBuilder::SetDisablePartitionActorCreationForWorld(UWorld* InWorld, bool bDisable)
+{
+	UPCGSubsystem::SetDisablePartitionActorCreationForWorld(InWorld, bDisable);
+}
+
+TUniquePtr<IWorldPartitionActorLoaderInterface::ILoaderAdapter> FPCGWorldPartitionBuilder::ExpandLoadedBoundsIfNeeded(UWorld* InWorld, const FBox& InCurrentBounds, const FBox& InNewBounds, const FString& InLoaderName)
+{
+	TUniquePtr<IWorldPartitionActorLoaderInterface::ILoaderAdapter> LoaderAdapter;
+	if (!InNewBounds.Equals(InCurrentBounds))
+	{
+		LoaderAdapter = MakeUnique<FLoaderAdapterShape>(InWorld, InNewBounds, InLoaderName);
+		LoaderAdapter->Load();
+
+		FWorldPartitionHelpers::FakeEngineTick(InWorld);
+
+		// Try to eliminate as much background noise as possible before getting started.
+		FPCGWorldPartitionBuilder::WaitForAllAsyncEditorProcesses(InWorld);
+
+		UE_LOG(LogPCGWorldPartitionBuilder, Verbose, TEXT("Expanding Cell Bounds: Min %s, Max %s"), *InNewBounds.Min.ToString(), *InNewBounds.Max.ToString());
+	}
+
+	return LoaderAdapter;
+}
+
+void FPCGWorldPartitionBuilder::ResetPackageDirtyFlags()
+{
+	TArray<UPackage*> DirtyPackages = FPCGWorldPartitionBuilder::GetDirtyPackages();
+	for (UPackage* DirtyPackage : DirtyPackages)
+	{
+		DirtyPackage->SetDirtyFlag(false);
+	}
+}
+
+TArray<UPackage*> FPCGWorldPartitionBuilder::GetDirtyPackages()
+{
+	TArray<UPackage*> DirtyPackages;
+	FEditorFileUtils::GetDirtyWorldPackages(DirtyPackages);
+	DirtyPackages.Remove(nullptr);
+	return DirtyPackages;
+}
+
+bool FPCGWorldPartitionBuilder::Build(UWorld* InWorld, const TArray<FString>& Args)
 {
 	if (InWorld)
 	{
@@ -784,7 +1123,7 @@ bool PCGWorldPartitionBuilder::Build(UWorld* InWorld, const TArray<FString>& Arg
 
 bool UPCGWorldPartitionBuilder::CanBuild(const UWorld* InWorld, FName InBuildOption)
 {
-	return InWorld && InWorld->IsPartitionedWorld();
+	return InWorld != nullptr;
 }
 
 EEditorBuildResult UPCGWorldPartitionBuilder::Build(UWorld* InWorld, FName InBuildOption)
@@ -815,10 +1154,325 @@ EEditorBuildResult UPCGWorldPartitionBuilder::Build(UWorld* InWorld, FName InBui
 		{
 			Args.Add("-PCGBuilderSettings=" + BuilderSettingStr);
 		}
-		return PCGWorldPartitionBuilder::Build(InWorld, Args) ? EEditorBuildResult::Success : EEditorBuildResult::Skipped;
+		return FPCGWorldPartitionBuilder::Build(InWorld, Args) ? EEditorBuildResult::Success : EEditorBuildResult::Skipped;
 	}
 
 	return EEditorBuildResult::Skipped;
 }
 
+FSoftObjectPath FPCGWorldPartitionBuilder::GetActorPathFromComponentPtr(TSoftObjectPtr<UPCGComponent> InComponentPtr)
+{
+	const FString PersistentLevelString(TEXT("PersistentLevel"));
+
+	FSoftObjectPath ComponentPath = InComponentPtr.ToSoftObjectPath();
+	FString PersistentLevel, SubObjectPath;
+	if (ComponentPath.GetSubPathString().Split(TEXT("."), &PersistentLevel, &SubObjectPath))
+	{
+		if (PersistentLevel == PersistentLevelString)
+		{
+			FString ActorName;
+			FString ComponentName;
+			if (SubObjectPath.Split(TEXT("."), &ActorName, &ComponentName))
+			{
+				return FSoftObjectPath(ComponentPath.GetAssetPath(), PersistentLevelString + "." + ActorName);
+			}
+		}
+	}
+
+	return FSoftObjectPath();
+}
+
+void FPCGWorldPartitionBuilder::LoadPartitionedActorDependencies(UWorld* InWorld, TSet<FWorldPartitionReference>& OutLoadedActorDependencies)
+{
+	UWorldPartition* WorldPartition = InWorld->GetWorldPartition();
+
+	// Loop on every loaded APCGPartitionActor (PAs):
+	// 
+	// Load missing Original Component Actors: this can happen when the cell getting loaded intersects with the PA but not the Original Component Volume.
+	// When this happens we need to call RebuildOriginalToLocal so that the PAs internal Original to Local mappings are up to date
+	//
+	// Load missing Managed Actors: this will allow proper cleanup of managed resources later when processing the PAs
+	for (TActorIterator<APCGPartitionActor> It(InWorld); It; ++It)
+	{
+		APCGPartitionActor* PartitionActor = *It;
+
+		if (PartitionActor->GetLevel() == InWorld->PersistentLevel)
+		{
+			bool bRebuildOriginalToLocal = false;
+
+			for (TObjectPtr<UPCGComponent> PCGComponent : PartitionActor->GetAllLocalPCGComponents())
+			{
+				if (PCGComponent)
+				{
+					// If the Original Component for this Local Component isn't loaded, try and load it (this is to avoid removing a dead instance because original isn't loaded yet)
+					TSoftObjectPtr<UPCGComponent> OriginalPCGComponentPtr = PartitionActor->GetOriginalComponentSoftObjectPtr(PCGComponent);
+					if (!OriginalPCGComponentPtr.IsNull() && !OriginalPCGComponentPtr.IsValid())
+					{
+						// The Original Component Path will be something like /Game/SomePath/SomeWorld.SomeWorld:PersistentLevel.ActorName.PCGComponentName
+						// and we want to get /Game/SomePath/SomeWorld.SomeWorld:PersistentLevel.ActorName to load it
+						FSoftObjectPath ActorPath = FPCGWorldPartitionBuilder::GetActorPathFromComponentPtr(OriginalPCGComponentPtr);
+						if (ActorPath.IsValid())
+						{
+							if (const FWorldPartitionActorDescInstance* ActorDescInstance = WorldPartition->GetActorDescInstanceByPath(ActorPath))
+							{
+								FWorldPartitionReference ActorReference(WorldPartition, ActorDescInstance->GetGuid());
+								if (AActor* LoadedActor = ActorReference.GetActor())
+								{
+									OutLoadedActorDependencies.Add(ActorReference);
+									UE_LOG(LogPCGWorldPartitionBuilder, Verbose, TEXT("Loaded Original PCG Component Actor '%s' label '%s'."),
+										*LoadedActor->GetName(),
+										*LoadedActor->GetActorNameOrLabel());
+									bRebuildOriginalToLocal = true;
+								}
+							}
+						}
+					}
+				}
+			}
+
+			// An Original Component was loaded so we need to update the PAs internal mappings 
+			if (bRebuildOriginalToLocal)
+			{
+				PartitionActor->RebuildOriginalToLocal();
+			}
+
+			// If this Partition Actor has some Managed Actors that aren't loaded, try and load them (this is to make sure we can delete them properly)
+			TArray<TSoftObjectPtr<AActor>> ManagedActorPaths = UPCGComponent::GetManagedActorPaths(PartitionActor);
+			for (const TSoftObjectPtr<AActor>& ManagedActorPath : ManagedActorPaths)
+			{
+				if (!ManagedActorPath.IsNull() && !ManagedActorPath.IsValid())
+				{
+					if (const FWorldPartitionActorDescInstance* ActorDescInstance = WorldPartition->GetActorDescInstanceByPath(ManagedActorPath.ToSoftObjectPath()))
+					{
+						FWorldPartitionReference ActorReference(WorldPartition, ActorDescInstance->GetGuid());
+						if (AActor* LoadedActor = ActorReference.GetActor())
+						{
+							OutLoadedActorDependencies.Add(ActorReference);
+							UE_LOG(LogPCGWorldPartitionBuilder, Verbose, TEXT("Loaded Managed Actor '%s' label '%s'."),
+								*LoadedActor->GetName(),
+								*LoadedActor->GetActorNameOrLabel());
+						}
+					}
+				}
+			}
+		}
+	}
+}
+
+
+bool FPCGWorldPartitionBuilder::CreatePartitionedActors(UWorld* InWorld, const FPCGWorldPartitionBuilderArgs& InArgs, TMap<TSoftObjectPtr<UPCGComponent>, FWorldPartitionReference>& InOutPartitionedComponents)
+{
+	// Scope object tracks if packages were dirtied 
+	FPCGDetectDirtyPackageInScope DirtyPackageInScope;
+
+	UPCGSubsystem* PCGSubsystem = UWorld::GetSubsystem<UPCGSubsystem>(InWorld);
+	check(PCGSubsystem);
+
+	auto ComponentFilter = [InWorld, &InOutPartitionedComponents](const UPCGComponent* InComponent)
+	{
+		// Only generate Partitioned Components part of the persistent level
+		if (InComponent->GetComponentLevel() != InWorld->PersistentLevel)
+		{
+			return false;
+		}
+
+		// Only process each component once
+		if (InOutPartitionedComponents.Contains(InComponent))
+		{
+			return false;
+		}
+
+		return InComponent->IsPartitioned();
+	};
+
+	TArray<TWeakObjectPtr<UPCGComponent>> PartitionedComponentsToProcess;
+	FPCGWorldPartitionBuilder::CollectComponentsToGenerate(InWorld, InArgs, ComponentFilter, PartitionedComponentsToProcess);
+
+	UWorldPartition* WorldPartition = InWorld->GetWorldPartition();
+
+	for (TWeakObjectPtr<UPCGComponent> ComponentPtr : PartitionedComponentsToProcess)
+	{
+		if (UPCGComponent* PCGComponent = ComponentPtr.Get())
+		{
+			check(!InOutPartitionedComponents.Contains(PCGComponent));
+			InOutPartitionedComponents.Add(PCGComponent, FWorldPartitionReference(WorldPartition, PCGComponent->GetOwner()->GetActorGuid()));
+
+			bool bHasUnbounded = false;
+			PCGHiGenGrid::FSizeArray GridSizes;
+			ensure(PCGHelpers::GetGenerationGridSizes(PCGComponent->GetGraph(), PCGSubsystem->GetPCGWorldActor(), GridSizes, bHasUnbounded));
+
+			if (!PCGComponent->bGenerated)
+			{
+				PCGComponent->Modify();
+				PCGComponent->bGenerated = true;
+				PCGComponent->LastGeneratedBounds = PCGComponent->GetGridBounds();
+			}
+
+			if (!GridSizes.IsEmpty())
+			{
+				PCGSubsystem->CreatePartitionActorsWithinBounds(PCGComponent, PCGComponent->GetGridBounds(), GridSizes);
+			}
+		}
+	}
+
+	return DirtyPackageInScope.AnyDirtyPackage();
+}
+
+bool FPCGWorldPartitionBuilder::UpdatePartitionedActors(UWorld* InWorld, const TMap<TSoftObjectPtr<UPCGComponent>, FWorldPartitionReference>& InPartitionedComponents, TSet<FString>& InOutDeletedActorPackages)
+{
+	// Scope object tracks if packages were dirtied 
+	FPCGDetectDirtyPackageInScope DirtyPackageInScope;
+
+	// Will track any deleted actors in this scope
+	FPCGDetectDeletedActorInScope DeletedActorScope(InWorld, InOutDeletedActorPackages);
+
+	UPCGSubsystem* PCGSubsystem = UWorld::GetSubsystem<UPCGSubsystem>(InWorld);
+	check(PCGSubsystem);
+
+	TArray<TKeyValuePair<UPCGComponent*, TSet<TObjectPtr<APCGPartitionActor>>>> UpdatedMappings;
+
+	// When loading existing PAs, we need to Update the mappings to add missing PCG Components
+	for (auto KeyValuePair : InPartitionedComponents)
+	{
+		if (UPCGComponent* PartitionedComponent = KeyValuePair.Key.Get())
+		{
+			// Those updated mappings will be partial (only include loaded PAs) but that is ok because here we want to make sure loaded PAs get the missing graph instances and in the next loop remove the graph
+			// instances from the loaded PAs no longer in the mappings
+			PCGSubsystem->UpdateMappingPCGComponentPartitionActor(PartitionedComponent);
+
+			UpdatedMappings.Add({ PartitionedComponent, PCGSubsystem->GetPCGComponentPartitionActorMappings(PartitionedComponent) });
+		}
+	}
+
+	// If PartitionActor is loaded and it has invalid Original Components it means that they need to be cleaned up
+	for (TActorIterator<APCGPartitionActor> It(InWorld); It; ++It)
+	{
+		// Cleanup Graph Instances to invalid Original Components
+		APCGPartitionActor* PartitionActor = *It;
+
+		if (PartitionActor->GetLevel() == InWorld->PersistentLevel)
+		{
+			// Here we cleanup dead graphs since we tried to load the Original components before, if they can't resolve it means they don't exist anymore.
+			PartitionActor->CleanupDeadGraphInstances();
+
+			// Remove Graph Instances from Partition Actors that are no longer part of the component mappings (bounds change)
+			for (const TKeyValuePair<UPCGComponent*, TSet<TObjectPtr<APCGPartitionActor>>>& Mapping : UpdatedMappings)
+			{
+				if (!Mapping.Value.Contains(PartitionActor))
+				{
+					PartitionActor->RemoveGraphInstance(Mapping.Key);
+				}
+			}
+		}
+	}
+
+	return DirtyPackageInScope.AnyDirtyPackage();
+}
+
+void FPCGWorldPartitionBuilder::CollectComponentsToGenerate(UWorld* InWorld, const FPCGWorldPartitionBuilderArgs& InArgs, TFunctionRef<bool(const UPCGComponent*)> ComponentFilter, TArray<TWeakObjectPtr<UPCGComponent>>& OutComponents)
+{
+	check(InWorld);
+
+	TArray<UObject*> AllComponents;
+	GetObjectsOfClass(UPCGComponent::StaticClass(), AllComponents, /*bIncludeDerivedClasses=*/true);
+	for (UObject* ComponentObject : AllComponents)
+	{
+		if (!IsValid(ComponentObject))
+		{
+			continue;
+		}
+
+		UPCGComponent* Component = Cast<UPCGComponent>(ComponentObject);
+		if (!Component || !Component->GetOwner() || Component->GetWorld() != InWorld)
+		{
+			continue;
+		}
+
+		// Check actor in inclusion list if provided.
+		if (!InArgs.IncludeActorIDs.IsEmpty() && !InArgs.IncludeActorIDs.Contains(Component->GetOwner()->GetName()))
+		{
+			continue;
+		}
+
+		// Accept based on editing mode.
+		const bool bValidEditingMode = (InArgs.bGenerateEditingModeLoadAsPreviewComponents && Component->GetSerializedEditingMode() == EPCGEditorDirtyMode::LoadAsPreview)
+			|| (InArgs.bGenerateEditingModeNormalComponents && Component->GetSerializedEditingMode() == EPCGEditorDirtyMode::Normal)
+			|| (InArgs.bGenerateEditingModePreviewComponents && Component->GetSerializedEditingMode() == EPCGEditorDirtyMode::Preview);
+
+		if (!bValidEditingMode)
+		{
+			continue;
+		}
+
+		if (!ComponentFilter(Component))
+		{
+			continue;
+		}
+
+		UPCGGraph* Graph = Component->GetGraph();
+		if (!Graph)
+		{
+			UE_LOG(LogPCGWorldPartitionBuilder, Warning, TEXT("PCG component on actor '%s' label '%s' has no graph assigned, skipping."),
+				*Component->GetOwner()->GetName(),
+				*Component->GetOwner()->GetActorNameOrLabel());
+			continue;
+		}
+
+		if (!InArgs.IncludeGraphNames.IsEmpty())
+		{
+			bool bFound = false;
+			for (const FString& IncludeGraphName : InArgs.IncludeGraphNames)
+			{
+				if (Graph->GetName() == IncludeGraphName)
+				{
+					bFound = true;
+					break;
+				}
+			}
+
+			if (!bFound)
+			{
+				continue;
+			}
+		}
+		else if (!InArgs.IncludeGraphs.IsEmpty())
+		{
+			bool bFound = false;
+			for (const TSoftObjectPtr<UPCGGraphInterface>& IncludeGraph : InArgs.IncludeGraphs)
+			{
+				if (TSoftObjectPtr<UPCGGraphInterface>(Graph) == IncludeGraph)
+				{
+					bFound = true;
+					break;
+				}
+			}
+
+			if (!bFound)
+			{
+				continue;
+			}
+		}
+
+		if (!Component->bActivated)
+		{
+			UE_LOG(LogPCGWorldPartitionBuilder, Display, TEXT("'Activated' toggle was set false on PCG component on actor '%s' label '%s' graph '%s'. Component skipped."),
+				*Component->GetOwner()->GetName(),
+				*Component->GetOwner()->GetActorNameOrLabel(),
+				*Component->GetGraph()->GetName());
+			continue;
+		}
+
+		if (Component->IsManagedByRuntimeGenSystem())
+		{
+			UE_LOG(LogPCGWorldPartitionBuilder, Display, TEXT("PCG component generation trigger is set to run-time generation on actor '%s' label '%s' graph '%s'. Component skipped."),
+				*Component->GetOwner()->GetName(),
+				*Component->GetOwner()->GetActorNameOrLabel(),
+				*Component->GetGraph()->GetName());
+			continue;
+		}
+
+		OutComponents.Add(Component);
+	}
+}
+
 #undef LOCTEXT_NAMESPACE
diff --git a/Engine/Plugins/PCG/Source/PCGEditor/Private/WorldPartitionBuilder/PCGWorldPartitionBuilder.h b/Engine/Plugins/PCG/Source/PCGEditor/Private/WorldPartitionBuilder/PCGWorldPartitionBuilder.h
--- a/Engine/Plugins/PCG/Source/PCGEditor/Private/WorldPartitionBuilder/PCGWorldPartitionBuilder.h
+++ b/Engine/Plugins/PCG/Source/PCGEditor/Private/WorldPartitionBuilder/PCGWorldPartitionBuilder.h
@@ -3,12 +3,14 @@
 #pragma once
 
 #include "WorldPartition/WorldPartitionBuilder.h"
+#include "WorldPartition/WorldPartitionHandle.h"
 
 #include "EditorBuildUtils.h"
 
 #include "PCGWorldPartitionBuilder.generated.h"
 
 class UWorld;
+class UPCGComponent;
 class UPCGBuilderSettings;
 class UPCGGraphInterface;
 
@@ -37,6 +39,12 @@ struct FPCGWorldPartitionCommandlineArgs
 		
 	/** Path to builder settings asset */
 	TOptional<TSoftObjectPtr<UPCGBuilderSettings>> SettingAsset;
+
+	/** Flag to generate using an iterative cell loading. Useful for large maps to limit memory consumption. */
+	TOptional<bool> bIterativeCellLoading;
+
+	/** Cell size for iterative cell loading */
+	TOptional<int32> IterativeCellSize;
 };
 
 struct FPCGWorldPartitionBuilderArgs
@@ -65,6 +73,12 @@ struct FPCGWorldPartitionBuilderArgs
 	/** Submit dirty files even if errors occurred during generation. */
 	bool bIgnoreGenerationErrors = false;
 
+	/** Flag to generate using an iterative cell loading. Useful for large maps to limit memory consumption. */
+	bool bIterativeCellLoading = false;
+
+	/** Cell size for iterative cell loading */
+	int32 IterativeCellSize = 25600;
+
 	static FPCGWorldPartitionBuilderArgs InitializeFrom(const FPCGWorldPartitionCommandlineArgs& CommandlineArgs);
 };
 
@@ -72,7 +86,20 @@ struct FPCGWorldPartitionBuilderArgs
 * Builder that triggers generation on PCG components.
 * 
 * Example command line:
-*   ProjectName MapName -Unattended -AllowCommandletRendering -run=WorldPartitionBuilderCommandlet -Builder=PCGWorldPartitionBuilder  -IncludeGraphNames=PCG_GraphA;PCG_GraphB
+*   ProjectName MapName -Unattended -AllowCommandletRendering -run=WorldPartitionBuilderCommandlet -Builder=PCGWorldPartitionBuilder -SCCProvider=[ProviderName] -AllowCommandletRendering -AllowSoftwareRendering -AssetGatherAll=true
+* 
+* Optional switches:
+* 
+*  -PCGBuilderSettings=/Game/Path/BuilderSettingsAssetName (default: none)
+*  -IncludeGraphNames=PCG_GraphA;PCG_GraphB (default: all graphs)
+*  -GenerateComponentEditingModeLoadAsPreview (default: true)
+*  -GenerateComponentEditingModeNormal (default: false)
+*  -GenerateComponentEditingModePreview (default: false)
+*  -IgnoreGenerationErrors (default: false)
+*  -IncludeActorIDs=MyActor1_UAID1234678;MyActor2_UAID1234678 (default: all actors)
+*  -OneComponentAtATime (default: false)
+*  -IterativeCellLoading (default: false)
+*  -IterativeCellSize=25600 (default: 25600)
 */
 UCLASS()
 class UPCGWorldPartitionBuilder : public UWorldPartitionBuilder
@@ -87,27 +114,39 @@ public:
 	virtual bool RequiresCommandletRendering() const override { return true; }
 
 	// In the future we may load a world in stages, but for now load everything to make sure we have the complete level.
-	virtual ELoadingMode GetLoadingMode() const override { return ELoadingMode::EntireWorld; }
+	virtual ELoadingMode GetLoadingMode() const override { return Args.bIterativeCellLoading ? ELoadingMode::IterativeCells2D : ELoadingMode::EntireWorld; }
 
 protected:
+	virtual bool PreWorldInitialization(UWorld* World, FPackageSourceControlHelper& PackageHelper) override;
 	virtual bool PreRun(UWorld* World, FPackageSourceControlHelper& PackageHelper) override;
 	virtual bool RunInternal(UWorld* World, const FCellInfo& InCellInfo, FPackageSourceControlHelper& PackageHelper) override;
+	virtual bool PostRun(UWorld* World, FPackageSourceControlHelper& PackageHelper, const bool bInRunSuccess) override;
+	
+	bool RunEntireWorld(UWorld* World, FPackageSourceControlHelper& PackageHelper);
+	bool RunPerCell(UWorld* World, const FCellInfo& InCellInfo, FPackageSourceControlHelper& PackageHelper);
+
 	virtual bool CanProcessNonPartitionedWorlds() const override { return true; }
 
 	/** Save all the pending dirty and deleted packages. */
-	virtual bool SaveDirtyPackages(UWorld* World, FPackageSourceControlHelper& PackageHelper);
+	bool SaveEntireWorldPackages(UWorld* World, FPackageSourceControlHelper& PackageHelper, TArray<UPackage*> DirtyPackages, TSet<FString> DeletedPackages);
+
+	/** Save all the cell packages */
+	bool SavePerCellPackages(UWorld* World, FPackageSourceControlHelper& PackageHelper, TSet<FString> CellDeletedPackages);
 
 	friend class FPCGEditorModule;
 	static bool CanBuild(const UWorld* InWorld, FName InBuildOption);
 	static EEditorBuildResult Build(UWorld* InWorld, FName InBuildOption);
 private:
-	/** The packages dirtied while generating all components. */
-	UPROPERTY(Transient)
-	TMap<FGuid, TObjectPtr<UPackage>> PendingDirtyPackages;
-
-	/** Packages that were logged as deleted through the OnActorDeleted event. UPROP to prevent GC. */
-	UPROPERTY(Transient)
-	TArray<TObjectPtr<UPackage>> DeletedActorPackages;
+	/** Track file states to build changelist in PostRun */
+	FName PreviousProviderName;
+	TSet<FString> NewPackageNames;
+	TSet<FString> ModifiedPackageNames;
+	TSet<FString> DeletedPackageNames;
+
+	/** Track the generated components to avoid multiple generations when cells overlap */
+	TSet<TSoftObjectPtr<const UPCGComponent>> GeneratedComponents;
+	/** Processed Original Partitioned Components + Reference to owning actor to avoid unloading Original Components and causing destructive changes to Partition Actors, we keep them loaded until PostRun */
+	TMap<TSoftObjectPtr<UPCGComponent>, FWorldPartitionReference> PartitionedComponents;
 
 	FPCGWorldPartitionCommandlineArgs CommandlineArgs;
 	FPCGWorldPartitionBuilderArgs Args;
