Fixed issues when there's an unclosed module in the grammar parser.
Added missing context parameters in error logging in select grammar, and missing context usage in the volume slicer.
Added option to have an additional seed argument per input in the slicer classes.
Fixed bounding issue with the splines in the volume slicer and made a minor optimization.

diff --git a/Engine/Plugins/PCG/Source/PCG/Private/Elements/Grammar/PCGSegmentSlicer.cpp b/Engine/Plugins/PCG/Source/PCG/Private/Elements/Grammar/PCGSegmentSlicer.cpp
--- a/Engine/Plugins/PCG/Source/PCG/Private/Elements/Grammar/PCGSegmentSlicer.cpp
+++ b/Engine/Plugins/PCG/Source/PCG/Private/Elements/Grammar/PCGSegmentSlicer.cpp
@@ -31,6 +31,7 @@ public:
 
 		FVector SlicingDirection;
 		FVector PerpendicularSlicingDirection;
+		int32 AdditionalSeed = 0;
 
 		const UPCGSegmentSlicerSettings* Settings = nullptr;
 		FPCGContext* Context = nullptr;
@@ -78,7 +79,7 @@ public:
 
 		TArray<PCGSlicingBase::TPCGSubDivModuleInstance<PCGGrammar::FTokenizedModule>> ModulesInstances;
 		double RemainingSubdivide;
-		const bool bSubdivideSuccess = PCGSlicingBase::Subdivide(*CurrentTokenizedGrammar.ModuleGrammar, Size, ModulesInstances, RemainingSubdivide, InOutParameters.Context);
+		const bool bSubdivideSuccess = PCGSlicingBase::Subdivide(*CurrentTokenizedGrammar.ModuleGrammar, Size, ModulesInstances, RemainingSubdivide, InOutParameters.Context, InOutParameters.AdditionalSeed);
 
 		if (!bSubdivideSuccess)
 		{
@@ -251,7 +252,7 @@ bool FPCGSegmentSlicerElement::ExecuteInternal(FPCGContext* InContext) const
 			Keys = PCGAttributeAccessorHelpers::CreateConstKeys(InputPointData, Selector);
 			if (!GrammarAccessor || !Keys)
 			{
-				PCGLog::LogErrorOnGraph(FText::Format(LOCTEXT("GrammarAccessor", "Attribute {0} was not found for the grammar."), Selector.GetDisplayText()), InContext);
+				PCGLog::Metadata::LogFailToCreateAccessorError(Selector, InContext);
 				continue;
 			}
 		}
@@ -267,7 +268,7 @@ bool FPCGSegmentSlicerElement::ExecuteInternal(FPCGContext* InContext) const
 
 			if (!FlipAxisAccessor || !Keys)
 			{
-				PCGLog::LogErrorOnGraph(FText::Format(LOCTEXT("FlipAxisAccessor", "Attribute {0} was not found for the flip axis attribute."), Selector.GetDisplayText()), InContext);
+				PCGLog::Metadata::LogFailToCreateAccessorError(Selector, InContext);
 				continue;
 			}
 		}
@@ -306,6 +307,27 @@ bool FPCGSegmentSlicerElement::ExecuteInternal(FPCGContext* InContext) const
 			continue;
 		}
 
+		// Set seed if required
+		Parameters.AdditionalSeed = 0;
+		if (Settings->bUseSeedAttribute)
+		{
+			const FPCGAttributePropertyInputSelector Selector = Settings->SeedAttribute.CopyAndFixLast(InputPointData);
+			TUniquePtr<const IPCGAttributeAccessor> SeedAccessor = PCGAttributeAccessorHelpers::CreateConstAccessor(InputPointData, Selector);
+			if (!Keys)
+			{
+				Keys = PCGAttributeAccessorHelpers::CreateConstKeys(InputPointData, Selector);
+			}
+
+			if (!SeedAccessor || !Keys)
+			{
+				PCGLog::Metadata::LogFailToCreateAccessorError(Selector, InContext);
+			}
+			// Otherwise, get the value, if it fails, the attribute wasn't compatible
+			else if (!SeedAccessor->Get(Parameters.AdditionalSeed, FPCGAttributeAccessorKeysEntries(PCGInvalidEntryKey), EPCGAttributeAccessorFlags::AllowBroadcastAndConstructible))
+			{
+				PCGLog::Metadata::LogFailToGetAttributeError<int32>(Selector, SeedAccessor.Get(), InContext);
+			}
+		}
 
 		if (Settings->GrammarSelection.bGrammarAsAttribute && Settings->bFlipAxisAsAttribute)
 		{
diff --git a/Engine/Plugins/PCG/Source/PCG/Private/Elements/Grammar/PCGSelectGrammar.cpp b/Engine/Plugins/PCG/Source/PCG/Private/Elements/Grammar/PCGSelectGrammar.cpp
--- a/Engine/Plugins/PCG/Source/PCG/Private/Elements/Grammar/PCGSelectGrammar.cpp
+++ b/Engine/Plugins/PCG/Source/PCG/Private/Elements/Grammar/PCGSelectGrammar.cpp
@@ -209,7 +209,7 @@ bool FPCGSelectGrammarElement::ExecuteInternal(FPCGContext* InContext) const
 				{
 					if (Criterion.FirstValue > Criterion.SecondValue)
 					{
-						PCGLog::LogErrorOnGraph(FText::Format(LOCTEXT("InvalidCriteriaRange", "Criteria range is inverted for criterion with Key->Grammar: '{0}'->'{1}'"), FText::FromName(Criterion.Key), FText::FromString(Criterion.Grammar)));
+						PCGLog::LogErrorOnGraph(FText::Format(LOCTEXT("InvalidCriteriaRange", "Criteria range is inverted for criterion with Key->Grammar: '{0}'->'{1}'"), FText::FromName(Criterion.Key), FText::FromString(Criterion.Grammar)), InContext);
 						return true;
 					}
 				}
@@ -253,14 +253,14 @@ bool FPCGSelectGrammarElement::ExecuteInternal(FPCGContext* InContext) const
 		TUniquePtr<const IPCGAttributeAccessorKeys> InputKeys = PCGAttributeAccessorHelpers::CreateConstKeys(InputPointData, KeySelector);
 		if (!InputKeyAccessor || !InputKeys)
 		{
-			PCGLog::Metadata::LogFailToCreateAccessorError(KeySelector);
+			PCGLog::Metadata::LogFailToCreateAccessorError(KeySelector, InContext);
 			continue;
 		}
 
 		TUniquePtr<const IPCGAttributeAccessor> ComparedValuesAccessor = PCGAttributeAccessorHelpers::CreateConstAccessor(InputPointData, ComparedValueSelector);
 		if (!ComparedValuesAccessor)
 		{
-			PCGLog::Metadata::LogFailToCreateAccessorError(ComparedValueSelector);
+			PCGLog::Metadata::LogFailToCreateAccessorError(ComparedValueSelector, InContext);
 			continue;
 		}
 
@@ -271,7 +271,7 @@ bool FPCGSelectGrammarElement::ExecuteInternal(FPCGContext* InContext) const
 		TUniquePtr<IPCGAttributeAccessorKeys> OutputKeys = PCGAttributeAccessorHelpers::CreateKeys(OutputPointData, OutputGrammarSelector);
 		if (!WriteGrammarAccessor || !OutputKeys)
 		{
-			PCGLog::Metadata::LogFailToCreateAccessorError(OutputGrammarSelector);
+			PCGLog::Metadata::LogFailToCreateAccessorError(OutputGrammarSelector, InContext);
 			continue;
 		}
 
diff --git a/Engine/Plugins/PCG/Source/PCG/Private/Elements/Grammar/PCGSplineSlicer.cpp b/Engine/Plugins/PCG/Source/PCG/Private/Elements/Grammar/PCGSplineSlicer.cpp
--- a/Engine/Plugins/PCG/Source/PCG/Private/Elements/Grammar/PCGSplineSlicer.cpp
+++ b/Engine/Plugins/PCG/Source/PCG/Private/Elements/Grammar/PCGSplineSlicer.cpp
@@ -22,6 +22,7 @@ namespace PCGSplineSlicerHelpers
 		TMap<FString, PCGGrammar::FTokenizedGrammar> CachedModules;
 		bool bAcceptIncompleteSlicing = false;
 		double ModuleHeight = 0.0;
+		int32 AdditionalSeed = 0;
 
 		const UPCGPolyLineData* PolyLineData = nullptr;
 		TArray<FPCGPoint>* OutPoints = nullptr;
@@ -105,7 +106,7 @@ namespace PCGSplineSlicerHelpers
 		 */
 		TArray<PCGSlicingBase::TPCGSubDivModuleInstance<PCGGrammar::FTokenizedModule>> ModulesInstances;
 		double RemainingSubdivide;
-		if (!Subdivide(*CurrentTokenizedGrammar.ModuleGrammar, SplineLength, ModulesInstances, RemainingSubdivide, InOutParameters.Context))
+		if (!Subdivide(*CurrentTokenizedGrammar.ModuleGrammar, SplineLength, ModulesInstances, RemainingSubdivide, InOutParameters.Context, InOutParameters.AdditionalSeed))
 		{
 			return;
 		}
@@ -311,6 +312,24 @@ bool FPCGSplineSlicerElement::ExecuteInternal(FPCGContext* InContext) const
 			continue;
 		}
 
+		// Set seed if required
+		Parameters.AdditionalSeed = 0;
+		if (Settings->bUseSeedAttribute)
+		{
+			const FPCGAttributePropertyInputSelector Selector = Settings->SeedAttribute.CopyAndFixLast(InputPolyLineData);
+			TUniquePtr<const IPCGAttributeAccessor> SeedAccessor = PCGAttributeAccessorHelpers::CreateConstAccessor(InputPolyLineData, Selector);
+
+			if (!SeedAccessor)
+			{
+				PCGLog::Metadata::LogFailToCreateAccessorError(Selector, InContext);
+			}
+			// Otherwise, get the value, if it fails, the attribute wasn't compatible
+			else if (!SeedAccessor->Get(Parameters.AdditionalSeed, FPCGAttributeAccessorKeysEntries(PCGInvalidEntryKey), EPCGAttributeAccessorFlags::AllowBroadcastAndConstructible))
+			{
+				PCGLog::Metadata::LogFailToGetAttributeError<int32>(Selector, SeedAccessor.Get(), InContext);
+			}
+		}
+
 		if (Settings->GrammarSelection.bGrammarAsAttribute)
 		{
 			const FPCGAttributePropertyInputSelector Selector = Settings->GrammarSelection.GrammarAttribute.CopyAndFixLast(InputPolyLineData);
diff --git a/Engine/Plugins/PCG/Source/PCG/Private/Elements/Grammar/PCGVolumeSlicer.cpp b/Engine/Plugins/PCG/Source/PCG/Private/Elements/Grammar/PCGVolumeSlicer.cpp
--- a/Engine/Plugins/PCG/Source/PCG/Private/Elements/Grammar/PCGVolumeSlicer.cpp
+++ b/Engine/Plugins/PCG/Source/PCG/Private/Elements/Grammar/PCGVolumeSlicer.cpp
@@ -117,6 +117,24 @@ bool FPCGVolumeSlicerElement::ExecuteInternal(FPCGContext* InContext) const
 			continue;
 		}
 
+		// Set seed if required
+		int32 AdditionalSeed = 0;
+		if (Settings->bUseSeedAttribute)
+		{
+			const FPCGAttributePropertyInputSelector Selector = Settings->SeedAttribute.CopyAndFixLast(InputSplineData);
+			TUniquePtr<const IPCGAttributeAccessor> SeedAccessor = PCGAttributeAccessorHelpers::CreateConstAccessor(InputSplineData, Selector);
+
+			if (!SeedAccessor)
+			{
+				PCGLog::Metadata::LogFailToCreateAccessorError(Selector, InContext);
+			}
+			// Otherwise, get the value, if it fails, the attribute wasn't compatible
+			else if (!SeedAccessor->Get(AdditionalSeed, FPCGAttributeAccessorKeysEntries(PCGInvalidEntryKey), EPCGAttributeAccessorFlags::AllowBroadcastAndConstructible))
+			{
+				PCGLog::Metadata::LogFailToGetAttributeError<int32>(Selector, SeedAccessor.Get(), InContext);
+			}
+		}
+
 		const FVector ExtrudeDirection = ExtrudeVector / ExtrudeLength;
 
 		double MinSize = 0.0;
@@ -129,7 +147,7 @@ bool FPCGVolumeSlicerElement::ExecuteInternal(FPCGContext* InContext) const
 
 		TArray<PCGSlicingBase::TPCGSubDivModuleInstance<PCGGrammar::FTokenizedModule>> Instances;
 		double RemainingLength = 0.0;
-		const bool bHeightSubdivideSuccess = PCGSlicingBase::Subdivide(*TokenizedGrammar.ModuleGrammar, ExtrudeLength, Instances, RemainingLength);
+		const bool bHeightSubdivideSuccess = PCGSlicingBase::Subdivide(*TokenizedGrammar.ModuleGrammar, ExtrudeLength, Instances, RemainingLength, InContext, AdditionalSeed);
 
 		if (!bHeightSubdivideSuccess)
 		{
@@ -147,10 +165,7 @@ bool FPCGVolumeSlicerElement::ExecuteInternal(FPCGContext* InContext) const
 			const FVector Size = ExtrudeDirection * CurrentBlock.Size * (FVector::OneVector + Instance.ExtraScale);
 			
 			FPCGSplineStruct NewSpline = InputSplineData->SplineStruct;
-			for (FInterpCurvePointVector& ControlPoint : NewSpline.SplineCurves.Position.Points)
-			{
-				ControlPoint.OutVal += CurrentDisplacement;
-			}
+			NewSpline.Transform.AddToTranslation(CurrentDisplacement);
 
 			UPCGSplineData* NewSplineData = FPCGContext::NewObject_AnyThread<UPCGSplineData>(InContext);
 			NewSplineData->Initialize(NewSpline);
diff --git a/Engine/Plugins/PCG/Source/PCG/Private/Grammar/PCGGrammarParser.cpp b/Engine/Plugins/PCG/Source/PCG/Private/Grammar/PCGGrammarParser.cpp
--- a/Engine/Plugins/PCG/Source/PCG/Private/Grammar/PCGGrammarParser.cpp
+++ b/Engine/Plugins/PCG/Source/PCG/Private/Grammar/PCGGrammarParser.cpp
@@ -175,6 +175,13 @@ namespace PCGGrammar
 			}
 		}
 
+		// Finally, validate that we're back on the root, otherwise we have an mismatched module (e.g. missing its end)
+		if (CurrentModule != &Result.Root)
+		{
+			Result.AddLog(LOCTEXT("UnclosedSubmodule", "Missing module end."), FPCGGrammarResult::ELogType::Error);
+			bParseOk = false;
+		}
+
 		return bParseOk;
 	}
 
diff --git a/Engine/Plugins/PCG/Source/PCG/Public/Elements/Grammar/PCGSlicingBase.h b/Engine/Plugins/PCG/Source/PCG/Public/Elements/Grammar/PCGSlicingBase.h
--- a/Engine/Plugins/PCG/Source/PCG/Public/Elements/Grammar/PCGSlicingBase.h
+++ b/Engine/Plugins/PCG/Source/PCG/Public/Elements/Grammar/PCGSlicingBase.h
@@ -109,6 +109,14 @@ public:
 	UPROPERTY(BlueprintReadWrite, EditAnywhere, Category = Settings, meta = (EditCondition = "bGrammarAsAttribute", EditConditionHides, DisplayAfter = bGrammarAsAttribute, PCG_Overridable))
 	FPCGAttributePropertyInputSelector GrammarAttribute;
 
+	/** Controls whether we'll use an attribute to drive random seeding for stochastic processes in the subdivision. */
+	UPROPERTY(BlueprintReadWrite, EditAnywhere, Category = Settings, meta = (PCG_Overridable, InlineEditConditionToggle))
+	bool bUseSeedAttribute = false;
+
+	/** Attribute to use to drive seed selection. It should be convertible to an integer. */
+	UPROPERTY(BlueprintReadWrite, EditAnywhere, Category = Settings, meta = (EditCondition = "bUseSeedAttribute", PCG_Overridable))
+	FPCGAttributePropertyInputSelector SeedAttribute;
+
 	/** Do a match and set with the incoming modules info, only if the modules info is passed as input. */
 	UPROPERTY(BlueprintReadWrite, EditAnywhere, Category = "Settings|Extra Output Attributes", meta = (EditCondition = "bModuleInfoAsInput", EditConditionHides, PCG_Overridable))
 	bool bForwardAttributesFromModulesInfo = false;
@@ -170,7 +178,7 @@ namespace PCGSlicingBase
 	PCGGrammar::FTokenizedGrammar GetTokenizedGrammar(FPCGContext* InContext, const FString& InGrammar, const FPCGModulesInfoMap& InModulesInfo, double& OutMinSize);
 
 	template<typename T>
-	bool Subdivide(const T& Root, double Length, TArray<TPCGSubDivModuleInstance<T>>& OutModuleInstances, double& RemainingLength, FPCGContext* InOptionalContext = nullptr, int32 InOptionalSeed = 42)
+	bool Subdivide(const T& Root, double Length, TArray<TPCGSubDivModuleInstance<T>>& OutModuleInstances, double& RemainingLength, FPCGContext* InOptionalContext = nullptr, int32 InOptionalAdditionalSeed = 0)
 	{
 		OutModuleInstances.Reset();
 		RemainingLength = Length;
@@ -180,7 +188,7 @@ namespace PCGSlicingBase
 			return true;
 		}
 
-		FRandomStream RandomStream(InOptionalContext ? InOptionalContext->GetSeed() : InOptionalSeed);
+		FRandomStream RandomStream((InOptionalContext ? InOptionalContext->GetSeed() : 42) + InOptionalAdditionalSeed);
 		TArray<TPCGSubDivModuleInstance<T>> CurrentModules;
 
 		// Start with root
