Fixed issue with scaling modules in the spline slicer (scaling axis was wrong) and added a convergence pass when the modules overshoot the linear length

#rnx

diff --git a/Engine/Plugins/PCG/Source/PCG/Private/Elements/Grammar/PCGSplineSlicer.cpp b/Engine/Plugins/PCG/Source/PCG/Private/Elements/Grammar/PCGSplineSlicer.cpp
--- a/Engine/Plugins/PCG/Source/PCG/Private/Elements/Grammar/PCGSplineSlicer.cpp
+++ b/Engine/Plugins/PCG/Source/PCG/Private/Elements/Grammar/PCGSplineSlicer.cpp
@@ -117,49 +117,134 @@ namespace PCGSplineSlicerHelpers
 			return;
 		}
 
-		const bool bHasMetadata = InOutParameters.SymbolAttribute || InOutParameters.DebugColorAttribute || InOutParameters.ModuleIndexAttribute || InOutParameters.IsFirstPointAttribute || InOutParameters.IsFinalPointAttribute;
-
-		int32 ModuleIndexCounter = 0; // Keep track for the output attribute
-		double SplineAlpha = 0.0; // Will be incremented as we progress through the spline
-		for (int32 ModuleInstanceIndex = 0; ModuleInstanceIndex < ModulesInstances.Num(); ModuleInstanceIndex++)
+		// Since we're essentially mapping spline-space modules onto linear-space modules (they aren't deformed),
+		// Then we need to refit potentially, knowing that having scaled up some modules might make the total length too large.
+		// We'll do a few passes to converge to a better fit.
+		TArray<double> ModuleAlphas;
+		int LengthConvergencePass = 0;
+		const int MaxLengthConvergencePass = 10;
+
+		constexpr double OvershootToleranceFactor = 0.02; // 2%
+		constexpr double OvershootToleranceMax = 5.0; // 5 cm
+		double MinimumModuleSize = std::numeric_limits<double>::max();
+		for (const PCGSlicingBase::TPCGSubDivModuleInstance<PCGGrammar::FTokenizedModule>& ModuleInstance : ModulesInstances)
 		{
-			const PCGSlicingBase::TPCGSubDivModuleInstance<PCGGrammar::FTokenizedModule>& ModuleInstance = ModulesInstances[ModuleInstanceIndex];
 			const FName Symbol = ModuleInstance.Module->Descriptor->Symbol;
-
 			const FPCGSlicingSubmodule& SlicingSubmodule = InOutParameters.ModulesInfo[Symbol];
-			const double SubmoduleSize = SlicingSubmodule.Size;
+			MinimumModuleSize = FMath::Min(MinimumModuleSize, SlicingSubmodule.Size);
+		}
 
-			// Move to the next segment of the spline
-			const FVector SegmentStartPoint = InOutParameters.PolyLineData->GetLocationAtAlpha(SplineAlpha);
-			const double PreviousAlpha = SplineAlpha;
+		const double OvershootTolerance = FMath::Min(MinimumModuleSize * OvershootToleranceFactor, OvershootToleranceMax);
 
-			// Use a numerical method to find the root of where the module lands on the spline
-			SplineAlpha = FindRootAtLinearDistance_Bisection(InOutParameters.PolyLineData, SubmoduleSize, SplineAlpha);
+		for(;LengthConvergencePass < MaxLengthConvergencePass; ++LengthConvergencePass)
+		{
+			ModuleAlphas.Reset(ModulesInstances.Num() + 1);
+			ModuleAlphas.Emplace(0.0);
 
-			const FVector SegmentEndPoint = InOutParameters.PolyLineData->GetLocationAtAlpha(SplineAlpha);
-			const FVector SliceVector = SegmentEndPoint - SegmentStartPoint;
-			const FVector SliceDirection = SliceVector.GetSafeNormal();
+			double ScaledLength = 0.0;
+			double OvershootDistance = 0.0;
+			bool bAtSplineEnd = false;
 
-			// At the end of the spline, truncate an unfinished submodule and end the process
-			if (FMath::IsNearlyEqual(SplineAlpha, 1))
+			for (const PCGSlicingBase::TPCGSubDivModuleInstance<PCGGrammar::FTokenizedModule>& ModuleInstance : ModulesInstances)
 			{
-				if (SliceVector.Length() < SubmoduleSize)
+				const FName Symbol = ModuleInstance.Module->Descriptor->Symbol;
+				const FPCGSlicingSubmodule& SlicingSubmodule = InOutParameters.ModulesInfo[Symbol];
+				const double SubmoduleSize = SlicingSubmodule.Size * (1.0 + ModuleInstance.ExtraScale);
+
+				// TODO: modules that can be deformed can also have extra scale, but the move in spline space (so their real length needs to be measured)
+				if (ModuleInstance.ExtraScale > 0)
 				{
-					if (InOutParameters.IsFinalPointAttribute && !InOutParameters.OutPoints->IsEmpty())
+					ScaledLength += SlicingSubmodule.Size * ModuleInstance.ExtraScale;
+				}
+
+				if (!bAtSplineEnd)
+				{
+					const double PreviousAlpha = ModuleAlphas.Last();
+					const double CurrentAlpha = FindRootAtLinearDistance_Bisection(InOutParameters.PolyLineData, SubmoduleSize, PreviousAlpha);
+
+					ModuleAlphas.Add(CurrentAlpha);
+
+					if (FMath::IsNearlyEqual(CurrentAlpha, 1.0))
 					{
-						InOutParameters.IsFinalPointAttribute->SetValue(InOutParameters.OutPoints->Last().MetadataEntry, true);
+						bAtSplineEnd = true;
+
+						// In this case, we know that the end of the module might be really after the end of the spline, so we need to compare the distance and add it to the overshoot.
+						const FVector SubmoduleStartPoint = InOutParameters.PolyLineData->GetLocationAtAlpha(ModuleAlphas.Last(1));
+						const FVector SplineEndPoint = InOutParameters.PolyLineData->GetLocationAtAlpha(ModuleAlphas.Last());
+
+						const double LastSegmentOvershoot = SubmoduleSize - (SplineEndPoint - SubmoduleStartPoint).Length();
+						if (LastSegmentOvershoot > 0)
+						{
+							OvershootDistance += LastSegmentOvershoot;
+						}
 					}
+				}
+				else
+				{
+					OvershootDistance += SubmoduleSize;
+				}
+			}
+
+			// Evaluate if we need another pass
+			if (LengthConvergencePass < MaxLengthConvergencePass && OvershootDistance > OvershootTolerance && ScaledLength > 0)
+			{
+				double UpdateExtraScaleFactor = 1.0;
+
+				// If the leeway we currently have (scaled length) is smaller than the overshoot distance, we are unlikely to converge, and in this instance it's probably best to set the scale to zero across the board.
+				if (OvershootDistance > ScaledLength)
+				{
+					UpdateExtraScaleFactor = 0.0;
+				}
+				// Otherwise, we want to move towards a better solution, but not too quickly lest we undershoot (which we don't detect here).
+				else
+				{
+					UpdateExtraScaleFactor = 0.25 * (3.0 + ((ScaledLength - OvershootDistance) / ScaledLength));
+					check(UpdateExtraScaleFactor >= 0.0 && UpdateExtraScaleFactor <= 1.0);
+				}
 
-					return;
+				for (PCGSlicingBase::TPCGSubDivModuleInstance<PCGGrammar::FTokenizedModule>& ModuleInstance : ModulesInstances)
+				{
+					if (ModuleInstance.ExtraScale > 0)
+					{
+						ModuleInstance.ExtraScale *= UpdateExtraScaleFactor;
+					}
 				}
 			}
+			else
+			{
+				// TODO: possibly log an error here, but it's subtle to get to a place where it's not too much
+				break;
+			}
+		}
+
+		const bool bHasMetadata = InOutParameters.SymbolAttribute || InOutParameters.DebugColorAttribute || InOutParameters.ModuleIndexAttribute || InOutParameters.IsFirstPointAttribute || InOutParameters.IsFinalPointAttribute;
+
+		FVector PreviousSegmentEndPoint = (ModuleAlphas.IsEmpty() ? FVector::Zero() : InOutParameters.PolyLineData->GetLocationAtAlpha(ModuleAlphas[0]));
+		const int32 NumIterations = FMath::Min(ModulesInstances.Num(), ModuleAlphas.Num() - 1);
+		for (int32 ModuleInstanceIndex = 0; ModuleInstanceIndex < NumIterations; ModuleInstanceIndex++)
+		{
+			const PCGSlicingBase::TPCGSubDivModuleInstance<PCGGrammar::FTokenizedModule>& ModuleInstance = ModulesInstances[ModuleInstanceIndex];
+			const FName Symbol = ModuleInstance.Module->Descriptor->Symbol;
+
+			const FPCGSlicingSubmodule& SlicingSubmodule = InOutParameters.ModulesInfo[Symbol];
+			const double SubmoduleSize = SlicingSubmodule.Size * (1.0 + ModuleInstance.ExtraScale);
+
+			const double& PreviousAlpha = ModuleAlphas[ModuleInstanceIndex];
+			const double& SplineAlpha = ModuleAlphas[ModuleInstanceIndex + 1];
+
+			// Move to the next segment of the spline
+			const FVector SegmentStartPoint = PreviousSegmentEndPoint;
+			const FVector SegmentEndPoint = InOutParameters.PolyLineData->GetLocationAtAlpha(SplineAlpha);
+			PreviousSegmentEndPoint = SegmentEndPoint;
+			const FVector SliceVector = SegmentEndPoint - SegmentStartPoint;
+			const FVector SliceDirection = SliceVector.GetSafeNormal();
 
 			// Since its discretized, we won't take the transform's position, but we'll use the up vector--to create the module rotation--and the scale
 			FTransform CenterPointTransform = InOutParameters.PolyLineData->GetTransformAtAlpha((SplineAlpha + PreviousAlpha) * 0.5);
 
 			const FVector Position = SegmentStartPoint + (SliceVector * 0.5) + FVector(0, 0, InOutParameters.ModuleHeight * 0.5);
 			const FRotator Rotation = FRotationMatrix::MakeFromXZ(SliceDirection, CenterPointTransform.GetRotation().GetUpVector()).Rotator();
-			const FVector Scale = FVector::OneVector + (SliceDirection * ModuleInstance.ExtraScale);
+			const FVector Scale = FVector(1.0 + ModuleInstance.ExtraScale, 1.0, 1.0);
 
 			FPCGPoint& OutPoint = InOutParameters.OutPoints->Emplace_GetRef(FTransform(Rotation, Position, Scale), /*InDensity=*/1, PCGHelpers::ComputeSeedFromPosition(Position));
 
@@ -185,7 +270,7 @@ namespace PCGSplineSlicerHelpers
 
 			if (InOutParameters.ModuleIndexAttribute)
 			{
-				InOutParameters.ModuleIndexAttribute->SetValue(OutPoint.MetadataEntry, ModuleIndexCounter++);
+				InOutParameters.ModuleIndexAttribute->SetValue(OutPoint.MetadataEntry, ModuleInstanceIndex);
 			}
 
 			const bool bIsFirstModule = (ModuleInstanceIndex == 0);
@@ -194,7 +279,7 @@ namespace PCGSplineSlicerHelpers
 				InOutParameters.IsFirstPointAttribute->SetValue(OutPoint.MetadataEntry, true);
 			}
 
-			const bool bIsFinalModule = (ModuleInstanceIndex == ModulesInstances.Num() - 1);
+			const bool bIsFinalModule = (ModuleInstanceIndex == NumIterations - 1);
 			if (bIsFinalModule && InOutParameters.IsFinalPointAttribute)
 			{
 				InOutParameters.IsFinalPointAttribute->SetValue(OutPoint.MetadataEntry, true);
