[PCG] Rework Custom HLSL node UI
* Pins now use an Initialization Mode, either FromPins or Custom
  * FromInputPins allows initializing Data Count, Element Count, and Attributes from existing input data
  * Custom requires the user to provide fixed Data & Element Count, and must add their own Attributes
* Rename IsValid() to IsPointRemoved()
* Fix up pin label references when user edits input pin labels
* Use the most significant bit of the GPU data collection NumData to flag a kernel as having executed.
* Always dispatch kernels for at least one thread, even if there are no elements present. This is required to set the executed flag when there are no data/elements.

#rnx

diff --git a/Engine/Plugins/PCG/Shaders/Private/Elements/PCGCopyPoints.usf b/Engine/Plugins/PCG/Shaders/Private/Elements/PCGCopyPoints.usf
--- a/Engine/Plugins/PCG/Shaders/Private/Elements/PCGCopyPoints.usf
+++ b/Engine/Plugins/PCG/Shaders/Private/Elements/PCGCopyPoints.usf
@@ -109,6 +109,14 @@ void CopyAttributesFromTarget(uint TargetDataIndex, uint TargetElemIndex, uint O
 [numthreads(64, 1, 1)]
 void Main(uint3 GroupId : SV_GroupID, uint GroupIndex : SV_GroupIndex)
 {
+	// Mark the kernel as having executed. Must run before we early via thread index, because the kernel is still 'executed' even if the number
+	// of threads to iterate on is zero. Even if GetNumThreads() returns 0, the kernel will still have been dispatched on a single thread to set
+	// this flag.
+	if (all(GroupId == 0) && GroupIndex == 0)
+	{
+		Out_SetAsExecutedInternal();
+	}
+
 	const uint ThreadIndex = GetUnWrappedDispatchThreadId(GroupId, GroupIndex, 64);
 	if (ThreadIndex >= GetNumThreads().x) return;
 
@@ -191,17 +199,14 @@ void Main(uint3 GroupId : SV_GroupID, uint GroupIndex : SV_GroupIndex)
 		OutNumData = max(SourceNumData, TargetNumData);
 	}
 
-	Out_SetNumDataInternal(OutNumData);
-
 	uint OutDataIndex, OutElemIndex;
 	if (!Out_GetThreadData(ThreadIndex, OutDataIndex, OutElemIndex))
 	{
 		return;
 	}
 
-	// Propagate the 'Removed' status to the output point if either source or target point is invalid. If we don't do this, then
-	// the output point will be considered valid and not culled even if the source or target point were culled upstream.
-	if (!Source_IsValid(SourceDataIndex, SourcePointIndexInData) || !Target_IsValid(TargetDataIndex, TargetPointIndexInData))
+	// Propagate the 'Removed' status to the output point if either source or target point is removed. Otherwise, this point will not be culled.
+	if (!Source_IsPointRemoved(SourceDataIndex, SourcePointIndexInData) || !Target_IsPointRemoved(TargetDataIndex, TargetPointIndexInData))
 	{
 		Out_RemovePoint(OutDataIndex, OutElemIndex);
 		return;
diff --git a/Engine/Plugins/PCG/Shaders/Private/PCGDataCollectionDataInterface.ush b/Engine/Plugins/PCG/Shaders/Private/PCGDataCollectionDataInterface.ush
--- a/Engine/Plugins/PCG/Shaders/Private/PCGDataCollectionDataInterface.ush
+++ b/Engine/Plugins/PCG/Shaders/Private/PCGDataCollectionDataInterface.ush
@@ -34,7 +34,8 @@ void StoreBufferInternal_{DataInterfaceName}(uint Address, uint Value)
 
 uint GetNumData_{DataInterfaceName}()
 {
-	return {DataInterfaceName}_DataCollectionBuffer.Load(0);
+	// We use the most significant bit of NumData to store the KernelExecuted flag
+	return {DataInterfaceName}_DataCollectionBuffer.Load(0) & ~(1 << 31); // TODO: Code-gen to PCGComputeConstants::KernelExecutedFlag
 }
 
 uint GetDataAddress_{DataInterfaceName}(uint DataIndex)
@@ -162,9 +163,11 @@ uint GetAttributeAddress_{DataInterfaceName}(uint InDataAddress, uint InAttribut
 
 // #################### HEADER WRITERS ##########################
 
-void SetNumDataInternal_{DataInterfaceName}(uint NumData)
+void SetAsExecutedInternal_{DataInterfaceName}()
 {
-	{DataInterfaceName}_DataCollectionBuffer.Store(0, NumData);
+	// We use the most significant bit of NumData to store the KernelExecuted flag
+	const uint NumData = GetNumData_{DataInterfaceName}();
+	{DataInterfaceName}_DataCollectionBuffer.Store(0, NumData | (1 << 31)); // TODO: Code-gen to PCGComputeConstants::KernelExecutedFlag
 }
 
 void SetDataAddressInternal_{DataInterfaceName}(uint DataIndex, uint DataAddress)
@@ -388,7 +391,7 @@ float4x4 GetPointTransform_{DataInterfaceName}(uint InDataIndex, uint InElementI
 		(float3)0.0, 1.0);
 }
 
-bool IsValid_{DataInterfaceName}(uint InDataIndex, uint InElementIndex)
+bool IsPointRemoved_{DataInterfaceName}(uint InDataIndex, uint InElementIndex)
 {
 	return GetDensity_{DataInterfaceName}(InDataIndex, InElementIndex) != INVALID_DENSITY;
 }
diff --git a/Engine/Plugins/PCG/Source/PCG/Private/Compute/DataInterfaces/PCGCustomKernelDataInterface.cpp b/Engine/Plugins/PCG/Source/PCG/Private/Compute/DataInterfaces/PCGCustomKernelDataInterface.cpp
--- a/Engine/Plugins/PCG/Source/PCG/Private/Compute/DataInterfaces/PCGCustomKernelDataInterface.cpp
+++ b/Engine/Plugins/PCG/Source/PCG/Private/Compute/DataInterfaces/PCGCustomKernelDataInterface.cpp
@@ -84,25 +84,7 @@ UComputeDataProvider* UPCGCustomKernelDataInterface::CreateDataProvider(TObjectP
 
 FComputeDataProviderRenderProxy* UPCGCustomComputeKernelDataProvider::GetRenderProxy()
 {
-	TArray<int32> InvocationCounts;
-	int32 TotalThreadCount = 0;
-	
-	if (!GetInvocationThreadCounts(InvocationCounts, TotalThreadCount))
-	{
-		InvocationCounts.Reset();
-	}
-
-	return new FPCGCustomComputeKernelDataProviderProxy(MoveTemp(InvocationCounts), TotalThreadCount, Seed, SourceComponentBounds);
-}
-
-bool UPCGCustomComputeKernelDataProvider::GetInvocationThreadCounts(TArray<int32>& OutInvocationThreadCount, int32& OutTotalThreadCount) const
-{
-	OutInvocationThreadCount.Reset(1);
-	OutInvocationThreadCount.Add(ThreadCount);
-
-	OutTotalThreadCount = ThreadCount;
-	
-	return true;
+	return new FPCGCustomComputeKernelDataProviderProxy(ThreadCount, Seed, SourceComponentBounds);
 }
 
 bool FPCGCustomComputeKernelDataProviderProxy::IsValid(FValidationData const& InValidationData) const
@@ -113,22 +95,13 @@ bool FPCGCustomComputeKernelDataProviderProxy::IsValid(FValidationData const& In
 		return false;
 	}
 
-	if (InvocationThreadCounts.Num() == 0)
-	{
-		UE_LOG(LogPCG, Error, TEXT("Proxy invalid due to empty InvocationThreadCounts."));
-		return false;
-	}
-
 	return true;
 }
 
 int32 FPCGCustomComputeKernelDataProviderProxy::GetDispatchThreadCount(TArray<FIntVector>& InOutThreadCounts) const
 {
-	InOutThreadCounts.Reset(InvocationThreadCounts.Num());
-	for (const int32 Count : InvocationThreadCounts)
-	{
-		InOutThreadCounts.Add({Count, 1, 1});
-	}
+	// Always dispatch at least one thread. This is necessary in order to flag the kernel as executed.
+	InOutThreadCounts.Emplace(FMath::Max(1, ThreadCount), 1, 1);
 	return InOutThreadCounts.Num();
 }
 
@@ -140,7 +113,9 @@ void FPCGCustomComputeKernelDataProviderProxy::GatherDispatchData(FDispatchData
 		FParameters& Parameters = ParameterArray[InvocationIndex];
 
 		// Thread count
-		Parameters.NumThreads.X = InDispatchData.bUnifiedDispatch ? TotalThreadCount : InvocationThreadCounts[InvocationIndex];
+		// Note: If this ThreadCount is zero, the kernel will still execute one thread, but terminates early when comparing thread index against this value.
+		ensure(InDispatchData.bUnifiedDispatch || InDispatchData.NumInvocations == 1);
+		Parameters.NumThreads.X = ThreadCount;
 		Parameters.NumThreads.Y = Parameters.NumThreads.Z = 1;
 
 		// Seed for the node
diff --git a/Engine/Plugins/PCG/Source/PCG/Private/Compute/DataInterfaces/PCGCustomKernelDataInterface.h b/Engine/Plugins/PCG/Source/PCG/Private/Compute/DataInterfaces/PCGCustomKernelDataInterface.h
--- a/Engine/Plugins/PCG/Source/PCG/Private/Compute/DataInterfaces/PCGCustomKernelDataInterface.h
+++ b/Engine/Plugins/PCG/Source/PCG/Private/Compute/DataInterfaces/PCGCustomKernelDataInterface.h
@@ -52,17 +52,13 @@ public:
 	uint32 Seed = 42;
 
 	FBox SourceComponentBounds;
-
-protected:
-	bool GetInvocationThreadCounts(TArray<int32>& OutInvocationThreadCount, int32& OutTotalThreadCount) const;
 };
 
 class FPCGCustomComputeKernelDataProviderProxy : public FComputeDataProviderRenderProxy
 {
 public:
-	FPCGCustomComputeKernelDataProviderProxy(TArray<int32>&& InInvocationThreadCounts, int32 InTotalThreadCount, int32 InSeed, const FBox& InSourceComponentBounds)
-		: InvocationThreadCounts(MoveTemp(InInvocationThreadCounts))
-		, TotalThreadCount(InTotalThreadCount)
+	FPCGCustomComputeKernelDataProviderProxy(int32 InThreadCount, int32 InSeed, const FBox& InSourceComponentBounds)
+		: ThreadCount(InThreadCount)
 		, Seed(InSeed)
 		, SourceComponentBounds(InSourceComponentBounds)
 	{}
@@ -76,10 +72,7 @@ public:
 protected:
 	using FParameters = FPCGCustomKernelDataInterfaceParameters;
 
-	TArray<int32> InvocationThreadCounts;
-	int32 TotalThreadCount;
-
+	int32 ThreadCount;
 	uint32 Seed = 42;
-
 	FBox SourceComponentBounds;
 };
diff --git a/Engine/Plugins/PCG/Source/PCG/Private/Compute/DataInterfaces/PCGDataCollectionDataInterface.cpp b/Engine/Plugins/PCG/Source/PCG/Private/Compute/DataInterfaces/PCGDataCollectionDataInterface.cpp
--- a/Engine/Plugins/PCG/Source/PCG/Private/Compute/DataInterfaces/PCGDataCollectionDataInterface.cpp
+++ b/Engine/Plugins/PCG/Source/PCG/Private/Compute/DataInterfaces/PCGDataCollectionDataInterface.cpp
@@ -122,8 +122,7 @@ void UPCGDataCollectionDataInterface::GetSupportedInputs(TArray<FShaderFunctionD
 	// Header Writers
 	{
 		OutFunctions.AddDefaulted_GetRef()
-			.SetName(TEXT("SetNumDataInternal"))
-			.AddParam(EShaderFundamentalType::Uint);
+			.SetName(TEXT("SetAsExecutedInternal"));
 
 		OutFunctions.AddDefaulted_GetRef()
 			.SetName(TEXT("SetDataAddressInternal"))
@@ -353,7 +352,7 @@ void UPCGDataCollectionDataInterface::GetSupportedInputs(TArray<FShaderFunctionD
 			.AddParam(EShaderFundamentalType::Uint);
 
 		OutFunctions.AddDefaulted_GetRef()
-			.SetName(TEXT("IsValid"))
+			.SetName(TEXT("IsPointRemoved"))
 			.AddReturnType(EShaderFundamentalType::Bool)
 			.AddParam(EShaderFundamentalType::Uint) // DataIndex
 			.AddParam(EShaderFundamentalType::Uint); // ElementIndex
diff --git a/Engine/Plugins/PCG/Source/PCG/Private/Compute/Elements/PCGCustomHLSL.cpp b/Engine/Plugins/PCG/Source/PCG/Private/Compute/Elements/PCGCustomHLSL.cpp
--- a/Engine/Plugins/PCG/Source/PCG/Private/Compute/Elements/PCGCustomHLSL.cpp
+++ b/Engine/Plugins/PCG/Source/PCG/Private/Compute/Elements/PCGCustomHLSL.cpp
@@ -66,21 +66,6 @@ namespace PCGHLSLElement
 	}
 }
 
-#if WITH_EDITOR
-bool FPCGPinPropertiesGPU::CanEditChange(const FEditPropertyChain& PropertyChain) const
-{
-	if (FProperty* Property = PropertyChain.GetActiveNode()->GetValue())
-	{
-		if (Property->GetFName() == GET_MEMBER_NAME_CHECKED(FPCGPinProperties, bAllowMultipleData))
-		{
-			return bAllowEditMultipleData;
-		}
-	}
-
-	return true;
-}
-#endif // WITH_EDITOR
-
 UPCGCustomHLSLSettings::UPCGCustomHLSLSettings()
 {
 	bExecuteOnGPU = true;
@@ -115,11 +100,72 @@ TArray<FPCGPinProperties> UPCGCustomHLSLSettings::OutputPinProperties() const
 }
 
 #if WITH_EDITOR
+void UPCGCustomHLSLSettings::PreEditChange(FProperty* PropertyAboutToChange)
+{
+	Super::PreEditChange(PropertyAboutToChange);
+
+	// If a pin label is about to change, cache all input label names to diff against in PostEditChangeProperty. We'll use this to fix-up pin label references.
+	if (PropertyAboutToChange && PropertyAboutToChange->GetFName() == GET_MEMBER_NAME_CHECKED(FPCGPinProperties, Label))
+	{
+		InputPinLabelsPreEditChange.Reset();
+
+		for (const FPCGPinProperties& PinProps : InputPins)
+		{
+			InputPinLabelsPreEditChange.Add(PinProps.Label);
+		}
+	}
+}
+
 void UPCGCustomHLSLSettings::PostEditChangeProperty(FPropertyChangedEvent& PropertyChangedEvent)
 {
 	// Apply any pin setup before refreshing the node.
 	UpdatePinSettings();
 
+	const FName MemberProperty = PropertyChangedEvent.MemberProperty ? PropertyChangedEvent.MemberProperty->GetFName() : NAME_None;
+	const FName Property = PropertyChangedEvent.Property ? PropertyChangedEvent.Property->GetFName() : NAME_None;
+
+	if (MemberProperty == GET_MEMBER_NAME_CHECKED(UPCGCustomHLSLSettings, OutputPins)
+		&& Property == GET_MEMBER_NAME_CHECKED(UPCGCustomHLSLSettings, OutputPins)
+		&& PropertyChangedEvent.ChangeType == EPropertyChangeType::ArrayAdd)
+	{
+		// Whenever a new output pin is created, we should default initialize 'PinsToInitializeFrom' with the first input pin label (if it exists).
+		if (const UPCGPin* FirstInputPin = GetFirstInputPin())
+		{
+			check(!OutputPins.IsEmpty());
+
+			FPCGPinPropertiesGPU& PinProps = OutputPins.Last();
+			PinProps.PropertiesGPU.PinsToInititalizeFrom.Add(FirstInputPin->Properties.Label);
+		}
+	}
+	else if (MemberProperty == GET_MEMBER_NAME_CHECKED(UPCGCustomHLSLSettings, InputPins)
+		&& Property == GET_MEMBER_NAME_CHECKED(FPCGPinProperties, Label))
+	{
+		check(InputPinLabelsPreEditChange.Num() == InputPins.Num());
+
+		// Fix-up pin input pin label references if an input pin label changed.
+		for (int Index = 0; Index < InputPinLabelsPreEditChange.Num(); ++Index)
+		{
+			const FName InputLabelBeforeChange = InputPinLabelsPreEditChange[Index];
+			const FName InputLabelAfterChange = InputPins[Index].Label;
+
+			if (InputLabelBeforeChange != InputLabelAfterChange)
+			{
+				for (FPCGPinPropertiesGPU& OutPinProps : OutputPins)
+				{
+					for (FName& InitPinLabel : OutPinProps.PropertiesGPU.PinsToInititalizeFrom)
+					{
+						if (InitPinLabel == InputLabelBeforeChange)
+						{
+							InitPinLabel = InputLabelAfterChange;
+						}
+					}
+				}
+
+				// TODO: Could also find/replace to fix-up the kernel source
+			}
+		}
+	}
+
 	Super::PostEditChangeProperty(PropertyChangedEvent);
 
 	UpdateDeclarations();
@@ -199,31 +245,10 @@ int UPCGCustomHLSLSettings::ComputeKernelThreadCount(const UPCGDataBinding* Bind
 	{
 		if (DispatchThreadCount == EPCGDispatchThreadCount::FromFirstOutputPin)
 		{
-			const UPCGPin* OutputPin = GetFirstOutputPin();
-			const FPCGPinPropertiesGPU* PropertiesGPU = OutputPin ? GetOutputPinPropertiesGPU(OutputPin->Properties.Label) : nullptr;
-			if (PropertiesGPU)
+			if (const UPCGPin* OutputPin = GetFirstOutputPin())
 			{
-				if (PropertiesGPU->BufferSizeMode == EPCGPinBufferSizeMode::FixedElementCount)
-				{
-					ThreadCount = PropertiesGPU->FixedBufferElementCount;
-				}
-				else if (PropertiesGPU->BufferSizeMode == EPCGPinBufferSizeMode::FromFirstPin)
-				{
-					if (const UPCGPin* InputPin = GetFirstInputPin())
-					{
-						ThreadCount = GetProcessingElemCountForInputPin(InputPin, Binding);
-					}
-				}
-				else if (PropertiesGPU->BufferSizeMode == EPCGPinBufferSizeMode::FromProductOfInputPins)
-				{
-					for (const FName& PinLabel : PropertiesGPU->BufferSizeInputPinLabels)
-					{
-						if (const UPCGPin* InputPin = GetInputPin(PinLabel))
-						{
-							ThreadCount = FMath::Max(ThreadCount, 1) * GetProcessingElemCountForInputPin(InputPin, Binding);
-						}
-					}
-				}
+				const FPCGDataCollectionDesc Desc = ComputeOutputPinDataDesc(OutputPin, Binding);
+				ThreadCount = Desc.ComputeDataElementCount(EPCGDataType::Any);
 			}
 		}
 		else if (DispatchThreadCount == EPCGDispatchThreadCount::FromProductOfInputPins)
@@ -232,7 +257,7 @@ int UPCGCustomHLSLSettings::ComputeKernelThreadCount(const UPCGDataBinding* Bind
 			{
 				if (const UPCGPin* InputPin = GetInputPin(PinLabel))
 				{
-					ThreadCount = FMath::Max(ThreadCount, 1) * GetProcessingElemCountForInputPin(InputPin, Binding);
+					ThreadCount = GetProcessingElemCountForInputPin(InputPin, Binding);
 				}
 			}
 		}
@@ -259,7 +284,7 @@ FPCGDataCollectionDesc UPCGCustomHLSLSettings::ComputeOutputPinDataDesc(const UP
 	check(OutputPin);
 	FPCGDataCollectionDesc PinDesc;
 
-	const FPCGPinPropertiesGPU* PropertiesGPU = GetOutputPinPropertiesGPU(OutputPin->Properties.Label);
+	const FPCGPinPropertiesGPU* Properties = GetOutputPinPropertiesGPU(OutputPin->Properties.Label);
 	const UPCGPin* FirstOutputPin = GetFirstOutputPin();
 
 	// The primary output pin follows any rules prescribed by kernel type.
@@ -276,54 +301,9 @@ FPCGDataCollectionDesc UPCGCustomHLSLSettings::ComputeOutputPinDataDesc(const UP
 		// Generators always produce a single point data with known point count.
 		PinDesc.DataDescs.Emplace(EPCGDataType::Point, PointCount);
 	}
-	else if (ensure(PropertiesGPU))
+	else if (ensure(Properties))
 	{
-		// No size set by kernel, fall back to pin settings.
-		if (PropertiesGPU->BufferSizeMode == EPCGPinBufferSizeMode::FromFirstPin)
-		{
-			if (const UPCGPin* InputPin = GetFirstInputPin())
-			{
-				PinDesc = ComputeInputPinDataDesc(InputPin, Binding);
-			}
-		}
-		else if (PropertiesGPU->BufferSizeMode == EPCGPinBufferSizeMode::FromProductOfInputPins)
-		{
-			int TotalElementCount = 0;
-
-			for (const FName& PinLabel : PropertiesGPU->BufferSizeInputPinLabels)
-			{
-				if (const UPCGPin* InputPin = GetInputPin(PinLabel))
-				{
-					const int ElementCount = ComputeInputPinDataDesc(InputPin, Binding).ComputeDataElementCount(PropertiesGPU->AllowedTypes);
-					TotalElementCount = FMath::Max(TotalElementCount, 1) * ElementCount;
-				}
-			}
-
-			if (TotalElementCount > 0)
-			{
-				PinDesc.DataDescs.Emplace(PropertiesGPU->AllowedTypes, TotalElementCount);
-			}
-		}
-		else if (PropertiesGPU->BufferSizeMode == EPCGPinBufferSizeMode::FixedElementCount)
-		{
-			if (ensure(PropertiesGPU->FixedBufferElementCount > 0))
-			{
-				const UPCGPin* InitializeFromPin = (PropertiesGPU->AllowedTypes == EPCGDataType::Param) ? GetInputPin(PropertiesGPU->InitializeFromPin) : nullptr;
-
-				if (InitializeFromPin)
-				{
-					PinDesc = ComputeInputPinDataDesc(InitializeFromPin, Binding);
-				}
-				else
-				{
-					PinDesc.DataDescs.Emplace(PropertiesGPU->AllowedTypes, PropertiesGPU->FixedBufferElementCount);
-				}
-			}
-		}
-		else
-		{
-			checkNoEntry();
-		}
+		PCGComputeHelpers::ComputeOutputPinDataDesc(*Properties, this, Binding, PinDesc);
 	}
 
 	const TMap<FPCGKernelAttributeKey, int32>& GlobalAttributeLookupTable = ensure(Binding && Binding->Graph) ? Binding->Graph->GetAttributeLookupTable() : TMap<FPCGKernelAttributeKey, int32>();
@@ -383,6 +363,24 @@ void UPCGCustomHLSLSettings::ApplyPreconfiguredSettings(const FPCGPreConfiguredS
 		if (EnumPtr->IsValidEnumValue(PreconfiguredInfo.PreconfiguredIndex))
 		{
 			KernelType = EPCGKernelType(PreconfiguredInfo.PreconfiguredIndex);
+
+#if WITH_EDITOR
+			UpdatePinSettings();
+#endif
+
+			// Default to initializing the first output pin's from the first input pin's data.
+			if (const UPCGPin* FirstInputPin = GetFirstInputPin())
+			{
+				if (!OutputPins.IsEmpty())
+				{
+					FPCGPinPropertiesGPU& PinProps = OutputPins.Last();
+					PinProps.PropertiesGPU.PinsToInititalizeFrom.Add(FirstInputPin->Properties.Label);
+				}
+			}
+
+#if WITH_EDITOR
+			UpdateDeclarations();
+#endif
 		}
 	}
 }
@@ -557,7 +555,7 @@ void UPCGCustomHLSLSettings::UpdateInputDeclarations()
 			"int {0}_GetSeed(uint DataIndex, uint ElementIndex);\n"
 			"float {0}_GetSteepness(uint DataIndex, uint ElementIndex);\n"
 			"float4x4 {0}_GetPointTransform(uint DataIndex, uint ElementIndex);\n"
-			"bool {0}_IsValid(uint DataIndex, uint ElementIndex);\n"),
+			"bool {0}_IsPointRemoved(uint DataIndex, uint ElementIndex);\n"),
 			{ bMultiPin ? PCGHLSLElement::PinDeclTemplateStr : PointDataPins[0] });
 
 		InputDeclarations += TEXT("\n");
@@ -875,10 +873,10 @@ void UPCGCustomHLSLSettings::UpdatePinSettings()
 			Properties.AllowedTypes = EPCGDataType::Point;
 		}
 
-		// Primary pin settings driven by kernel (if not custom kernel type).
-		const bool bPinCanBeSized = PinIndex > 0 || KernelType == EPCGKernelType::Custom;
-		const bool bDataCanBeSized = Properties.AllowedTypes == EPCGDataType::Point;
-		Properties.bDisplayBufferSizeSettings = bPinCanBeSized && bDataCanBeSized;
+		// Only allow editing the initialization mode if it's not driven by the kernel type.
+		const bool bInitModeDrivenByKernel = PinIndex == 0 && (KernelType == EPCGKernelType::PointProcessor || KernelType == EPCGKernelType::PointGenerator);
+		Properties.PropertiesGPU.bAllowEditInitMode = !bInitModeDrivenByKernel;
+		Properties.PropertiesGPU.bMultipleInitPins = Properties.PropertiesGPU.PinsToInititalizeFrom.Num() > 1;
 
 		// Output pins should always allow multiple connections.
 		// TODO this could be hoisted up somewhere in the future.
@@ -888,13 +886,14 @@ void UPCGCustomHLSLSettings::UpdatePinSettings()
 		{
 			Properties.bAllowMultipleData = false;
 			Properties.bAllowEditMultipleData = false;
-			Properties.bAllowEditInitializationPin = true;
-			Properties.BufferSizeMode = EPCGPinBufferSizeMode::FixedElementCount;
+			Properties.PropertiesGPU.bAllowEditDataCount = false;
+			Properties.PropertiesGPU.DataCountMode = EPCGDataCountMode::Fixed;
+			Properties.PropertiesGPU.DataCount = 1;
 		}
 		else
 		{
 			Properties.bAllowEditMultipleData = true;
-			Properties.bAllowEditInitializationPin = false;
+			Properties.PropertiesGPU.bAllowEditDataCount = true;
 		}
 	}
 }
@@ -955,7 +954,7 @@ void UPCGCustomHLSLSettings::UpdateAttributeKeys()
 	// Process each output pin for any new attributes they want to create.
 	for (const FPCGPinPropertiesGPU& OutputPin : OutputPins)
 	{
-		for (const FPCGKernelAttributeKey& Key : OutputPin.CreatedKernelAttributeKeys)
+		for (const FPCGKernelAttributeKey& Key : OutputPin.PropertiesGPU.CreatedKernelAttributeKeys)
 		{
 			KernelAttributeKeys.AddUnique(Key);
 
@@ -1118,80 +1117,79 @@ bool UPCGCustomHLSLSettings::IsKernelValid(FPCGContext* InContext, bool bQuiet)
 
 		if (!bPinIsDefinedByKernel)
 		{
-			if (Properties.BufferSizeMode == EPCGPinBufferSizeMode::FixedElementCount)
+			const FPCGPinPropertiesGPUStruct& Props = Properties.PropertiesGPU;
+
+			if (Props.InitializationMode == EPCGPinInitMode::FromInputPins)
 			{
-				if (Properties.FixedBufferElementCount <= 0)
+				if (Props.PinsToInititalizeFrom.IsEmpty())
 				{
 					PCG_KERNEL_VALIDATION_ERR(InContext, this, bQuiet, FText::Format(
-						LOCTEXT("InvalidFixedBufferSize", "Fixed GPU buffer size on '{0}' was invalid (%d)."),
-						FText::FromName(Properties.Label),
-						Properties.FixedBufferElementCount));
+						LOCTEXT("InitFromEmptyPins", "Output pin '{0}' tried to initialize from input pins, but no pins were specified."),
+						FText::FromName(Properties.Label)));
 
 					return false;
 				}
 
-				if (Properties.AllowedTypes == EPCGDataType::Param && Properties.InitializeFromPin != NAME_None && !GetInputPin(Properties.InitializeFromPin))
+				for (const FName InitPinName : Props.PinsToInititalizeFrom)
 				{
-					PCG_KERNEL_VALIDATION_ERR(InContext, this, bQuiet, FText::Format(
-						LOCTEXT("InvalidInitFromPin", "Tried to initialize attribute set pin '{0}' from non-existent pin '{1}'. Must reference a valid input pin or be 'None'."),
-						FText::FromName(Properties.Label),
-						FText::FromName(Properties.InitializeFromPin)));
+					const FPCGPinProperties* InitPinProps = InputPins.FindByPredicate([InitPinName](const FPCGPinProperties& InPinProps)
+					{
+						return InPinProps.Label == InitPinName;
+					});
 
-					return false;
-				}
-			}
-			else if (Properties.BufferSizeMode == EPCGPinBufferSizeMode::FromFirstPin)
-			{
-				if (InputPins.IsEmpty())
-				{
-					PCG_KERNEL_VALIDATION_ERR(InContext, this, bQuiet, FText::Format(
-						LOCTEXT("InvalidBufferSizeNoInputPin", "GPU buffer size for pin '{0}' could not be computed as there are no input pins."),
-						FText::FromName(Properties.Label)));
+					if (!InitPinProps)
+					{
+						PCG_KERNEL_VALIDATION_ERR(InContext, this, bQuiet, FText::Format(
+							LOCTEXT("InitFromNonExistentPin", "Output pin '{0}' tried to initialize from non-existent input pin '{1}'."),
+							FText::FromName(Properties.Label),
+							FText::FromName(InitPinName)));
 
-					return false;
-				}
+						return false;
+					}
 
-				if (!GetFirstInputPin())
-				{
-					PCG_KERNEL_VALIDATION_ERR(InContext, this, bQuiet, FText::Format(
-						LOCTEXT("MissingPrimaryInputPin", "GPU buffer size for pin '{0}' could not be computed, because it refers to the primary input pin, which does not exist."),
-						FText::FromName(Properties.Label)));
+					if (!PCGComputeHelpers::IsTypeAllowedInDataCollection(InitPinProps->AllowedTypes))
+					{
+						PCG_KERNEL_VALIDATION_ERR(InContext, this, bQuiet, FText::Format(
+							LOCTEXT("InitFromInvalidPinType", "Output pin '{0}' tried to initialize from input pin '{1}', but pin '{1}' has an invalid type."),
+							FText::FromName(Properties.Label),
+							FText::FromName(InitPinName)));
 
-					return false;
+						return false;
+					}
 				}
+
+				// TODO: Could do validation on data multiplicity for Pairwise, checking that data counts are 1 or N, but maybe that should be a runtime error instead.
 			}
-			else if (Properties.BufferSizeMode == EPCGPinBufferSizeMode::FromProductOfInputPins)
+
+			const bool bUsingFixedDataCount = Props.InitializationMode == EPCGPinInitMode::Custom || Props.DataCountMode == EPCGDataCountMode::Fixed;
+
+			if (bUsingFixedDataCount)
 			{
-				if (InputPins.IsEmpty())
+				if (Props.DataCount < 1)
 				{
 					PCG_KERNEL_VALIDATION_ERR(InContext, this, bQuiet, FText::Format(
-						LOCTEXT("InvalidBufferSizeNoInputPins", "GPU buffer size for pin '{0}' could not be computed as there are no input pins on this node."),
+						LOCTEXT("InvalidDataCount", "Invalid fixed data count {0} on output pin '{1}'. Must be greater than 0."),
+						FText::AsNumber(Props.DataCount),
 						FText::FromName(Properties.Label)));
 
 					return false;
 				}
 
-				if (Properties.BufferSizeInputPinLabels.IsEmpty())
+			}
+
+			const bool bUsingFixedElemCount = Props.InitializationMode == EPCGPinInitMode::Custom || Props.ElementCountMode == EPCGElementCountMode::Fixed;
+
+			if (bUsingFixedElemCount)
+			{
+				if (Props.ElementCount < 1)
 				{
 					PCG_KERNEL_VALIDATION_ERR(InContext, this, bQuiet, FText::Format(
-						LOCTEXT("InvalidBufferSizeNoBufferPins", "GPU buffer size for pin '{0}' could not be computed as input pins are specified in the pin settings."),
+						LOCTEXT("InvalidElementCount", "Invalid fixed element count {0} on output pin '{1}'. Must be greater than 0."),
+						FText::AsNumber(Props.ElementCount),
 						FText::FromName(Properties.Label)));
 
 					return false;
 				}
-
-				for (const FName& Label : Properties.BufferSizeInputPinLabels)
-				{
-					if (!GetInputPin(Label))
-					{
-						PCG_KERNEL_VALIDATION_ERR(InContext, this, bQuiet, FText::Format(
-							LOCTEXT("MissingBufferSizePin", "GPU buffer size for pin '{0}' could not be computed. Invalid pin specified in Input Pins array: '{1}'."),
-							FText::FromName(Properties.Label),
-							FText::FromName(Label)));
-
-						return false;
-					}
-				}
 			}
 		}
 
@@ -1443,7 +1441,7 @@ FString UPCGCustomHLSLSettings::GetCookedKernelSource(const TMap<FPCGKernelAttri
 	FString ShaderPathName = GetPathName();
 	PCGHLSLElement::ConvertObjectPathToShaderFilePath(ShaderPathName);
 
-	const bool bHasKernelKeyword = Source.Contains(TEXT("KERNEL"), ESearchCase::CaseSensitive);
+	//const bool bHasKernelKeyword = Source.Contains(TEXT("KERNEL"), ESearchCase::CaseSensitive);
 
 	FString Includes;
 	{
@@ -1466,108 +1464,12 @@ FString UPCGCustomHLSLSettings::GetCookedKernelSource(const TMap<FPCGKernelAttri
 		GroupSize.X * GroupSize.Y * GroupSize.Z
 	);
 
-	// Header writers initialize PCG data collection format headers in output buffers.
-	FString HeaderWriters;
-
-	auto EmitHeaderWriterSingleData = [&HeaderWriters](const FPCGPinProperties& InOutputPinProps)
-	{
-		HeaderWriters += FString::Format(TEXT(
-			"    // Signal kernel executed by setting data count from first thread. Rest of header was already set up by the CPU.\n"
-			"    if (GroupIndex == 0) {0}_SetNumDataInternal(1);\n"
-			"    AllMemoryBarrier();\n"),
-			{ InOutputPinProps.Label.ToString() });
-	};
-
-	auto EmitHeaderWriterFromInputPin = [&HeaderWriters](const FPCGPinProperties& InOutputPinProps, const UPCGPin* InFromPin)
-	{
-		check(InFromPin);
-		HeaderWriters += FString::Format(TEXT(
-			"    // Signal kernel executed by copying data count from pin {0} to pin {1} from first thread. Rest of header was already set up by the CPU.\n"
-			"    if (GroupIndex == 0) {1}_SetNumDataInternal({0}_GetNumData());\n"
-			"    AllMemoryBarrier();\n"),
-			{ InFromPin->Properties.Label.ToString(), InOutputPinProps.Label.ToString() });
-	};
-
-	if (KernelType == EPCGKernelType::PointProcessor || KernelType == EPCGKernelType::Custom)
-	{
-		const UPCGPin* FirstPin = (KernelType == EPCGKernelType::PointProcessor)
-			? GetPointProcessingInputPin()
-			: CastChecked<UPCGNode>(GetOuter())->GetPassThroughInputPin();
-
-		// Initialize all output headers.
-		for (const FPCGPinPropertiesGPU& PinProps : OutputPins)
-		{
-			if (PinProps.BufferSizeMode == EPCGPinBufferSizeMode::FromFirstPin && FirstPin)
-			{
-				EmitHeaderWriterFromInputPin(PinProps, FirstPin);
-			}
-			else if (PinProps.BufferSizeMode == EPCGPinBufferSizeMode::FixedElementCount)
-			{
-				const UPCGPin* InitFromPin = (PinProps.AllowedTypes == EPCGDataType::Param) ? GetInputPin(PinProps.InitializeFromPin) : nullptr;
-
-				if (InitFromPin)
-				{
-					EmitHeaderWriterFromInputPin(PinProps, InitFromPin);
-				}
-				else
-				{
-					EmitHeaderWriterSingleData(PinProps);
-				}
-			}
-			else if (PinProps.BufferSizeMode == EPCGPinBufferSizeMode::FromProductOfInputPins)
-			{
-				// TODO: FromProductOfInputPins always produces a single point data for now, make it more flexible?
-				EmitHeaderWriterSingleData(PinProps);
-			}
-		}
-	}
-	else if (KernelType == EPCGKernelType::PointGenerator)
-	{
-		if (const UPCGNode* Node = Cast<UPCGNode>(GetOuter()))
-		{
-			const UPCGPin* FirstPin = Node->GetPassThroughInputPin();
-			const UPCGPin* PrimaryOutputPin = GetFirstPointOutputPin();
-
-			for (const UPCGPin* OutputPin : Node->GetOutputPins())
-			{
-				if (!OutputPin || !PCGComputeHelpers::IsTypeAllowedInDataCollection(OutputPin->Properties.AllowedTypes))
-				{
-					continue;
-				}
-
-				const FPCGPinPropertiesGPU* PinPropsGPU = GetOutputPinPropertiesGPU(OutputPin->Properties.Label);
-				if (!ensure(PinPropsGPU))
-				{
-					continue;
-				}
+	// Used to signal that a kernel has executed. Set the most significant bit in NumData.
+	FString SetAsExecuted = TEXT("    // Signal kernel executed by setting the most significant bit of NumData.\n");
 
-				const FPCGPinProperties& PinProps = OutputPin->Properties;
-
-				if (OutputPin == PrimaryOutputPin)
-				{
-					EmitHeaderWriterSingleData(PinProps);
-				}
-				else if (PinPropsGPU->BufferSizeMode == EPCGPinBufferSizeMode::FromFirstPin && FirstPin)
-				{
-					EmitHeaderWriterFromInputPin(PinProps, FirstPin);
-				}
-				else if (PinPropsGPU->AllowedTypes == EPCGDataType::Param)
-				{
-					if (const UPCGPin* InitFromPin = GetInputPin(PinPropsGPU->InitializeFromPin))
-					{
-						EmitHeaderWriterFromInputPin(PinProps, InitFromPin);
-					}
-					else
-					{
-						EmitHeaderWriterSingleData(PinProps);
-					}
-				}
-			}
-		}
-	}
-	else
+	for (const FPCGPinPropertiesGPU& PinProps : OutputPins)
 	{
-		checkNoEntry();
+		SetAsExecuted += FString::Format(TEXT("    if (all(GroupId == 0) && GroupIndex == 0) {0}_SetAsExecutedInternal();\n"), { PinProps.Label.ToString() });
 	}
 
 	// Per-kernel-type preamble. Set up shader inputs and initialize output data.
@@ -1596,7 +1498,7 @@ FString UPCGCustomHLSLSettings::GetCookedKernelSource(const TMap<FPCGKernelAttri
 
 			// If input point is invalid, mark output point as invalid and abort.
 			KernelSpecificPreamble += FString::Format(TEXT(
-				"    if (!{0}_IsValid({0}_DataIndex, ElementIndex))\n"
+				"    if (!{0}_IsPointRemoved({0}_DataIndex, ElementIndex))\n"
 				"    {\n"
 				"        {1}_RemovePoint({1}_DataIndex, ElementIndex);\n"
 				"        return;\n"
@@ -1655,7 +1557,9 @@ FString UPCGCustomHLSLSettings::GetCookedKernelSource(const TMap<FPCGKernelAttri
 
 	FString Result;
 
-	if (bHasKernelKeyword)
+	// TODO: Support KERNEL keyword in shader source. Could be handy for external source assets and breaking kernels into sections to
+	// support pin/attribute declarations, etc.
+	/*if (bHasKernelKeyword)
 	{
 		Source.ReplaceInline(TEXT("KERNEL"), TEXT("void __kernel_func(uint ThreadIndex)"), ESearchCase::CaseSensitive);
 
@@ -1667,21 +1571,21 @@ FString UPCGCustomHLSLSettings::GetCookedKernelSource(const TMap<FPCGKernelAttri
 			"%s { __kernel_func(%s); }\n"), // KernelFunc, UnWrappedDispatchThreadId
 			*ShaderPathName, *Includes, *Functions, *Source, *KernelFunc, *UnWrappedDispatchThreadId);
 	}
-	else
+	else*/
 	{
 		Result = FString::Printf(TEXT(
 			"%s\n\n" // Includes
 			"%s\n\n" // Functions
 			"%s\n" // KernelFunc
 			"{\n"
+			"%s\n" // SetAsExecuted
 			"	const uint ThreadIndex = %s;\n" // UnWrappedDispatchThreadId
 			"	if (ThreadIndex >= GetNumThreads().x) return;\n"
-			"%s\n" // HeaderWriters
 			"%s\n" // KernelSpecificPreamble
 			"#line 0 \"%s\"\n" // ShaderPathName
 			"%s\n" // Source
 			"}\n"),
-			*Includes, *Functions, *KernelFunc, *UnWrappedDispatchThreadId, *HeaderWriters, *KernelSpecificPreamble, *ShaderPathName, *Source);
+			*Includes, *Functions, *KernelFunc, *SetAsExecuted, *UnWrappedDispatchThreadId, *KernelSpecificPreamble, *ShaderPathName, *Source);
 	}
 
 	return Result;
diff --git a/Engine/Plugins/PCG/Source/PCG/Private/Compute/PCGComputeCommon.cpp b/Engine/Plugins/PCG/Source/PCG/Private/Compute/PCGComputeCommon.cpp
--- a/Engine/Plugins/PCG/Source/PCG/Private/Compute/PCGComputeCommon.cpp
+++ b/Engine/Plugins/PCG/Source/PCG/Private/Compute/PCGComputeCommon.cpp
@@ -6,6 +6,7 @@
 #include "PCGContext.h"
 #include "PCGParamData.h"
 #include "PCGSubsystem.h"
+#include "Compute/PCGPinPropertiesGPU.h"
 #include "Data/PCGPointData.h"
 
 namespace PCGComputeHelpers
@@ -77,4 +78,214 @@ namespace PCGComputeHelpers
 		PCGE_LOG_C(Error, LogOnly, Context, InText);
 	}
 #endif
+
+	void ComputeOutputPinDataDesc(const FPCGPinPropertiesGPU& PinProperties, const UPCGSettings* Settings, const UPCGDataBinding* Binding, FPCGDataCollectionDesc& OutPinDesc)
+	{
+		check(Settings);
+		check(Binding);
+		
+		const UPCGNode* Node = CastChecked<UPCGNode>(Settings->GetOuter());
+
+		// No size set by kernel, fall back to pin settings.
+		const FPCGPinPropertiesGPUStruct& Props = PinProperties.PropertiesGPU;
+
+		if (Props.InitializationMode == EPCGPinInitMode::FromInputPins)
+		{
+			TArray<FPCGDataCollectionDesc> InputDescs;
+
+			TArray<FName> InitPins;
+
+			for (const FName PinToInitFrom : Props.PinsToInititalizeFrom)
+			{
+				if (const UPCGPin* InitPin = Node->GetInputPin(PinToInitFrom))
+				{
+					InitPins.Emplace(PinToInitFrom);
+					InputDescs.Emplace(Settings->ComputeInputPinDataDesc(InitPin, Binding));
+				}
+			}
+
+			const int NumInitPins = InitPins.Num();
+			check(NumInitPins == InputDescs.Num());
+
+			// Copies unique (non-reserved) attribute descriptions from 'InDataDesc' to 'OutDataDesc'.
+			auto AddAttributesFromData = [](const FPCGDataDesc& InDataDesc, FPCGDataDesc& OutDataDesc)
+			{
+				for (const FPCGKernelAttributeDesc& InAttrDesc : InDataDesc.AttributeDescs)
+				{
+					if (InAttrDesc.Index >= PCGComputeConstants::NUM_RESERVED_ATTRS
+						&& !OutDataDesc.AttributeDescs.ContainsByPredicate([InAttrDesc](const FPCGKernelAttributeDesc& OutAttrDesc)
+							{
+								// Note: We shouldn't need to check for uniqueness of attr Index, since attributes should all have unique index via
+								// the GlobalAttributeLookupTable
+								return InAttrDesc.Name == OutAttrDesc.Name;
+							}))
+					{
+						OutDataDesc.AttributeDescs.Emplace(InAttrDesc);
+					}
+				}
+			};
+
+			// Combines the data for index i of each pin into one data. Creates exactly 'MaxDataCount' datas.
+			auto AddDataPairwise = [&](int MaxDataCount)
+			{
+				for (int DataIndex = 0; DataIndex < MaxDataCount; ++DataIndex)
+				{
+					// Set element count to 0 for now, but we will overwrite it.
+					FPCGDataDesc& DataDesc = OutPinDesc.DataDescs.Emplace_GetRef(PinProperties.AllowedTypes, /*SetNumElements=*/0);
+					int TotalNumElements = 0; // Total number of elements computed for this data index.
+
+					// For each data index, loop over all the pins and create the uber-data.
+					for (int InputPinIndex = 0; InputPinIndex < NumInitPins; ++InputPinIndex)
+					{
+						const FPCGDataCollectionDesc& InputDesc = InputDescs[InputPinIndex];
+
+						int ClampedDataIndex = DataIndex;
+
+						// If this pin does not have the same number of data, clamp it to the first data.
+						if (DataIndex != InputDesc.DataDescs.Num())
+						{
+							ClampedDataIndex = 0;
+						}
+
+						const FPCGDataDesc& InputDataDesc = InputDesc.DataDescs.IsValidIndex(ClampedDataIndex) ? InputDesc.DataDescs[DataIndex] : FPCGDataDesc(EPCGDataType::Any, 0);
+
+						if (Props.ElementCountMode == EPCGElementCountMode::FromInputData)
+						{
+							if (Props.ElementMultiplicity == EPCGElementMultiplicity::Product)
+							{
+								TotalNumElements = FMath::Max(TotalNumElements, 1) * InputDataDesc.ElementCount;
+							}
+							else if (Props.ElementMultiplicity == EPCGElementMultiplicity::Sum)
+							{
+								TotalNumElements += InputDataDesc.ElementCount;
+							}
+							else
+							{
+								checkNoEntry();
+							}
+						}
+						else if (Props.ElementCountMode == EPCGElementCountMode::Fixed)
+						{
+							TotalNumElements += Props.ElementCount;
+						}
+						else
+						{
+							checkNoEntry();
+						}
+
+						if (Props.AttributeInheritanceMode == EPCGAttributeInheritanceMode::CopyAttributeSetup)
+						{
+							AddAttributesFromData(InputDataDesc, DataDesc);
+						}
+
+						DataDesc.ElementCount = TotalNumElements;
+					}
+				}
+			};
+
+			if (Props.DataCountMode == EPCGDataCountMode::FromInputData)
+			{
+				// If this is the only input pin, we can just copy it.
+				if (NumInitPins == 1)
+				{
+					for (const FPCGDataDesc& InputDataDesc : InputDescs[0].DataDescs)
+					{
+						FPCGDataDesc& DataDesc = OutPinDesc.DataDescs.Emplace_GetRef(PinProperties.AllowedTypes, InputDataDesc.ElementCount);
+
+						if (Props.AttributeInheritanceMode == EPCGAttributeInheritanceMode::CopyAttributeSetup)
+						{
+							AddAttributesFromData(InputDataDesc, DataDesc);
+						}
+					}
+				}
+				// Take pairs of datas, where the pairs are given by each data of each pin to each data of every other pin.
+				else if (Props.DataMultiplicity == EPCGDataMultiplicity::CartesianProduct)
+				{
+					for (int InputPinIndex = 0; InputPinIndex < NumInitPins; ++InputPinIndex)
+					{
+						const FPCGDataCollectionDesc& InputDesc = InputDescs[InputPinIndex];
+
+						for (int OtherInputPinIndex = InputPinIndex + 1; OtherInputPinIndex < NumInitPins; ++OtherInputPinIndex)
+						{
+							const FPCGDataCollectionDesc& OtherInputDesc = InputDescs[OtherInputPinIndex];
+
+							for (const FPCGDataDesc& InputDataDesc : InputDesc.DataDescs)
+							{
+								for (const FPCGDataDesc& OtherInputDataDesc : OtherInputDesc.DataDescs)
+								{
+									FPCGDataDesc* DataDesc = nullptr;
+
+									if (Props.ElementCountMode == EPCGElementCountMode::FromInputData)
+									{
+										if (Props.ElementMultiplicity == EPCGElementMultiplicity::Product)
+										{
+											DataDesc = &OutPinDesc.DataDescs.Emplace_GetRef(PinProperties.AllowedTypes, InputDataDesc.ElementCount * OtherInputDataDesc.ElementCount);
+										}
+										else if (Props.ElementMultiplicity == EPCGElementMultiplicity::Sum)
+										{
+											DataDesc = &OutPinDesc.DataDescs.Emplace_GetRef(PinProperties.AllowedTypes, InputDataDesc.ElementCount + OtherInputDataDesc.ElementCount);
+										}
+										else
+										{
+											checkNoEntry();
+										}
+									}
+									else if (Props.ElementCountMode == EPCGElementCountMode::Fixed)
+									{
+										DataDesc = &OutPinDesc.DataDescs.Emplace_GetRef(PinProperties.AllowedTypes, Props.ElementCount);
+									}
+									else
+									{
+										checkNoEntry();
+									}
+
+									if (ensure(DataDesc) && Props.AttributeInheritanceMode == EPCGAttributeInheritanceMode::CopyAttributeSetup)
+									{
+										AddAttributesFromData(InputDataDesc, *DataDesc);
+										AddAttributesFromData(OtherInputDataDesc, *DataDesc);
+									}
+								}
+							}
+						}
+					}
+				}
+				// Combine elements for each set of datas, where the sets are given by the Nth datas on each pin (or the first data if there is only one data).
+				else if (Props.DataMultiplicity == EPCGDataMultiplicity::Pairwise)
+				{
+					int MaxDataCount = 0;
+
+					// Find the maximum number of data among the init pins. Note, they should all be the same number of data, or only one data.
+					for (int I = 0; I < NumInitPins; ++I)
+					{
+						MaxDataCount = FMath::Max(MaxDataCount, InputDescs[I].DataDescs.Num());
+					}
+
+					AddDataPairwise(MaxDataCount);
+				}
+				else
+				{
+					checkNoEntry();
+				}
+			}
+			else if (Props.DataCountMode == EPCGDataCountMode::Fixed)
+			{
+				AddDataPairwise(Props.DataCount);
+			}
+			else
+			{
+				checkNoEntry();
+			}
+		}
+		else if (Props.InitializationMode == EPCGPinInitMode::Custom)
+		{
+			for (int I = 0; I < Props.DataCount; ++I)
+			{
+				OutPinDesc.DataDescs.Emplace(PinProperties.AllowedTypes, Props.ElementCount);
+			}
+		}
+		else
+		{
+			checkNoEntry();
+		}
+	}
 }
diff --git a/Engine/Plugins/PCG/Source/PCG/Private/Compute/PCGComputeCommon.h b/Engine/Plugins/PCG/Source/PCG/Private/Compute/PCGComputeCommon.h
--- a/Engine/Plugins/PCG/Source/PCG/Private/Compute/PCGComputeCommon.h
+++ b/Engine/Plugins/PCG/Source/PCG/Private/Compute/PCGComputeCommon.h
@@ -5,8 +5,11 @@
 #include "PCGCommon.h"
 
 class UPCGData;
+class UPCGDataBinding;
 class UPCGSettings;
 struct FPCGContext;
+struct FPCGDataCollectionDesc;
+struct FPCGPinPropertiesGPU;
 
 #define PCG_KERNEL_LOGGING_ENABLED (!(UE_BUILD_SHIPPING || UE_BUILD_TEST) || USE_LOGGING_IN_SHIPPING)
 
@@ -53,6 +56,8 @@ namespace PCGComputeConstants
 
 	/** PCG data types supported in GPU data collections. */
 	constexpr EPCGDataType AllowedDataCollectionTypes = EPCGDataType::Point | EPCGDataType::Param;
+
+	constexpr uint32 KernelExecutedFlag = 1 << 31;
 }
 
 namespace PCGComputeHelpers
@@ -70,7 +75,13 @@ namespace PCGComputeHelpers
 	bool IsTypeAllowedInDataCollection(EPCGDataType Type);
 
 #if PCG_KERNEL_LOGGING_ENABLED
+	/** Logs a warning on a GPU node in the graph and console. */
 	void LogKernelWarning(const FPCGContext* Context, const UPCGSettings* Settings, const FText& InText);
+
+	/** Logs an error on a GPU node in the graph and console. */
 	void LogKernelError(const FPCGContext* Context, const UPCGSettings* Settings, const FText& InText);
 #endif
+
+	/** Computes a pin data description based on the given GPU pin properties. */
+	void ComputeOutputPinDataDesc(const FPCGPinPropertiesGPU& PinProperties, const UPCGSettings* Settings, const UPCGDataBinding* Binding, FPCGDataCollectionDesc& OutPinDesc);
 }
diff --git a/Engine/Plugins/PCG/Source/PCG/Private/Compute/PCGDataForGPU.cpp b/Engine/Plugins/PCG/Source/PCG/Private/Compute/PCGDataForGPU.cpp
--- a/Engine/Plugins/PCG/Source/PCG/Private/Compute/PCGDataForGPU.cpp
+++ b/Engine/Plugins/PCG/Source/PCG/Private/Compute/PCGDataForGPU.cpp
@@ -778,24 +778,7 @@ void FPCGDataCollectionDesc::PrepareBufferForKernelOutput(TArray<uint32>& OutPac
 	const uint32 PackedDataCollectionSizeBytes = ComputePackedSize(&DataAddresses);
 
 	OutPackedDataCollection.SetNumZeroed(PackedDataCollectionSizeBytes / sizeof(uint32));
-	
-	const uint32 TotalNumElements = ComputeDataElementCount(EPCGDataType::Any);
-
-	if (TotalNumElements == 0)
-	{
-		// The kernel won't run with 0 element count, so we have to set the number of data in advance.
-		OutPackedDataCollection[0] = NumData;
-
-		// TODO: This assumption will break dynamic control flow in the future. If we are using 'NumData == 0' to flag whether a kernel ran or not,
-		// we can't just write NumData anyways whenever there are zero elements.
-		// In the future we should dispatch a single thread to execute even when there are zero elements to allow for book keeping, such as setting
-		// the data count.
-	}
-	else
-	{
-		// Num data - set to zero if writing kernel executes. If kernel doesn't execute, 0 means data collection is empty.
-		OutPackedDataCollection[0] = 0;
-	}
+	OutPackedDataCollection[0] = NumData;
 
 	for (uint32 DataIndex = 0; DataIndex < NumData; ++DataIndex)
 	{
@@ -880,7 +863,10 @@ EPCGUnpackDataCollectionResult FPCGDataCollectionDesc::UnpackDataCollection(cons
 	const int* DataAsInt = static_cast<const int*>(PackedData);
 
 	const uint32 NumPackedFloats = InPackedData.Num() / 4;
-	const uint32 NumData = DataAsUint[0];
+
+	// Most significant bit of NumData is reserved to flag whether or not the kernel executed.
+	ensureAlwaysMsgf(DataAsUint[0] & PCGComputeConstants::KernelExecutedFlag, TEXT("Tried to unpack a GPU data collection, but the compute shader did not execute."));
+	const uint32 NumData = DataAsUint[0] & ~PCGComputeConstants::KernelExecutedFlag;
 
 	if (NumData != DataDescs.Num())
 	{
diff --git a/Engine/Plugins/PCG/Source/PCG/Private/Compute/PCGPinPropertiesGPU.cpp b/Engine/Plugins/PCG/Source/PCG/Private/Compute/PCGPinPropertiesGPU.cpp
new file mode 100644
--- /dev/null
+++ b/Engine/Plugins/PCG/Source/PCG/Private/Compute/PCGPinPropertiesGPU.cpp
@@ -0,0 +1,20 @@
+// Copyright Epic Games, Inc. All Rights Reserved.
+
+#include "Compute/PCGPinPropertiesGPU.h"
+
+#include UE_INLINE_GENERATED_CPP_BY_NAME(PCGPinPropertiesGPU)
+
+#if WITH_EDITOR
+bool FPCGPinPropertiesGPU::CanEditChange(const FEditPropertyChain& PropertyChain) const
+{
+	if (FProperty* Property = PropertyChain.GetActiveNode()->GetValue())
+	{
+		if (Property->GetFName() == GET_MEMBER_NAME_CHECKED(FPCGPinProperties, bAllowMultipleData))
+		{
+			return bAllowEditMultipleData;
+		}
+	}
+
+	return true;
+}
+#endif // WITH_EDITOR
diff --git a/Engine/Plugins/PCG/Source/PCG/Public/Compute/Elements/PCGCustomHLSL.h b/Engine/Plugins/PCG/Source/PCG/Public/Compute/Elements/PCGCustomHLSL.h
--- a/Engine/Plugins/PCG/Source/PCG/Public/Compute/Elements/PCGCustomHLSL.h
+++ b/Engine/Plugins/PCG/Source/PCG/Public/Compute/Elements/PCGCustomHLSL.h
@@ -3,75 +3,12 @@
 #pragma once
 
 #include "PCGSettings.h"
+#include "Compute/PCGPinPropertiesGPU.h"
 
 #include "PCGCustomHLSL.generated.h"
 
 class UPCGPin;
 
-/** Method for computing the size of a pin on a GPU node. */
-UENUM()
-enum class EPCGPinBufferSizeMode : uint8
-{
-	FromFirstPin UMETA(DisplayName = "Match First Input Pin"),
-	FromProductOfInputPins UMETA(Tooltip = "Dispatches a thread per element in the product of one or more pins. So if there are 4 data elements in pin A and 6 data elements in pin B, 24 threads will be dispatched."),
-	FixedElementCount,
-};
-
-/** An extension of the pin properties that adds hints for GPU thread count / buffer size calculations. */
-USTRUCT(BlueprintType)
-struct PCG_API FPCGPinPropertiesGPU : public FPCGPinProperties
-{
-	GENERATED_BODY()
-
-public:
-	FPCGPinPropertiesGPU() = default;
-
-	explicit FPCGPinPropertiesGPU(const FName& InLabel, EPCGDataType InAllowedTypes)
-		: FPCGPinProperties(InLabel, InAllowedTypes)
-	{
-	}
-
-#if WITH_EDITOR
-	bool CanEditChange(const FEditPropertyChain& PropertyChain) const;
-#endif
-
-public:
-	/** Compute graphs use this to calculate the buffer size of output pins. */
-	UPROPERTY(BlueprintReadWrite, EditAnywhere, Category = Settings, meta = (DisplayAfter = "Tooltip", EditCondition = "bDisplayBufferSizeSettings", EditConditionHides, HideEditConditionToggle))
-	EPCGPinBufferSizeMode BufferSizeMode = EPCGPinBufferSizeMode::FromFirstPin;
-
-	UPROPERTY(BlueprintReadWrite, EditAnywhere, Category = Settings, meta = (DisplayAfter = "BufferSizeMode", EditCondition = "(bDisplayBufferSizeSettings && BufferSizeMode == EPCGPinBufferSizeMode::FixedElementCount) || AllowedTypes == EPCGDataType::Param", EditConditionHides))
-	int FixedBufferElementCount = 4;
-
-	UPROPERTY(BlueprintReadWrite, EditAnywhere, DisplayName = "Buffer Size Input Pins", Category = Settings, meta = (DisplayAfter = "BufferSizeMode", EditCondition = "bDisplayBufferSizeSettings && BufferSizeMode == EPCGPinBufferSizeMode::FromProductOfInputPins", EditConditionHides, GetOptions = "GetInputPinNames"))
-	TArray<FName> BufferSizeInputPinLabels;
-
-	/** Select an input pin to copy attributes from. If left as 'None', this will be ignored. Note, this will copy attribute names only, not their values. */
-	UPROPERTY(BlueprintReadWrite, EditAnywhere, Category = Settings, meta = (DisplayAfter = "BufferSizeMode", EditCondition = "bAllowEditInitializationPin", EditConditionHides, HideEditConditionToggle, GetOptions = "GetInputPinNamesAndNone"))
-	FName InitializeFromPin = NAME_None;
-
-	/** Add entries to create new attributes on data emitted by this pin. */
-	UPROPERTY(EditAnywhere, DisplayName = "Attributes to Create", Category = Settings, meta = (DisplayAfter = "BufferSizeMode"))
-	TArray<FPCGKernelAttributeKey> CreatedKernelAttributeKeys;
-
-#if WITH_EDITORONLY_DATA
-	UPROPERTY(Transient)
-	bool bDisplayBufferSizeSettings = true;
-
-	UPROPERTY(Transient)
-	bool bAllowEditInitializationPin = false;
-#endif // WITH_EDITORONLY_DATA
-};
-
-template<>
-struct TStructOpsTypeTraits<FPCGPinPropertiesGPU> : public TStructOpsTypeTraitsBase2<FPCGPinPropertiesGPU>
-{
-	enum
-	{
-		WithCanEditChange = true,
-	};
-};
-
 /** Type of kernel allows us to make decisions about execution automatically, streamlining authoring. */
 UENUM()
 enum class EPCGKernelType : uint8
@@ -132,6 +69,7 @@ public:
 
 protected:
 #if WITH_EDITOR
+	virtual void PreEditChange(FProperty* PropertyAboutToChange) override;
 	virtual void PostEditChangeProperty(FPropertyChangedEvent& PropertyChangedEvent) override;
 	virtual EPCGChangeType GetChangeTypeForProperty(const FName& InPropertyName) const override;
 #endif
@@ -196,6 +134,11 @@ public:
 	UPROPERTY(BlueprintReadWrite, EditAnywhere, Category = "Settings")
 	TArray<FPCGPinPropertiesGPU> OutputPins = { FPCGPinPropertiesGPU(PCGPinConstants::DefaultOutputLabel, EPCGDataType::Point) };
 
+#if WITH_EDITOR
+	/** Holds input pin labels from PreEditChange, used in PostEditPropertyChange to update any references in output pin setup. */
+	TArray<FName> InputPinLabelsPreEditChange;
+#endif
+
 protected:
 	UPROPERTY(BlueprintReadWrite, EditAnywhere, Category = "Source", meta = (MultiLine = true, Tooltip = ""))
 	FString ShaderFunctions;
diff --git a/Engine/Plugins/PCG/Source/PCG/Public/Compute/PCGPinPropertiesGPU.h b/Engine/Plugins/PCG/Source/PCG/Public/Compute/PCGPinPropertiesGPU.h
new file mode 100644
--- /dev/null
+++ b/Engine/Plugins/PCG/Source/PCG/Public/Compute/PCGPinPropertiesGPU.h
@@ -0,0 +1,152 @@
+// Copyright Epic Games, Inc. All Rights Reserved.
+
+#pragma once
+
+#include "PCGPin.h"
+#include "Compute/PCGDataForGPU.h"
+
+#include "PCGPinPropertiesGPU.generated.h"
+
+/** Method for initializing an output pin. */
+UENUM()
+enum class EPCGPinInitMode : uint8
+{
+	FromInputPins UMETA(Tooltip = "Initialize from one or more input pins."),
+	Custom UMETA(Tooltip = "Initialize from fixed values."),
+};
+
+/** Method for computing data count. */
+UENUM()
+enum class EPCGDataCountMode : uint8
+{
+	FromInputData,
+	Fixed,
+};
+
+/** Method for computing element count. */
+UENUM()
+enum class EPCGElementCountMode : uint8
+{
+	FromInputData,
+	Fixed,
+};
+
+/** Method for combining two or more data counts. */
+UENUM()
+enum class EPCGDataMultiplicity : uint8
+{
+	Pairwise UMETA(Tooltip = "A data item will be produced for each pair/tuple/etc of input data items across the input pins. Requires all pins to have the same number of data."),
+	CartesianProduct UMETA(Tooltip = "If there are two input pins with N and M data items respectively, the output will have NxM data items."),
+};
+
+/** Method for combining two or more element counts. */
+UENUM()
+enum class EPCGElementMultiplicity : uint8
+{
+	Product UMETA(Tooltip = "Element count will be the product of the elements in each pair/tuple/etc of input data."),
+	Sum UMETA(Tooltip = "Element count will be the sum of the elements in each pair/tuple/etc of input data."),
+};
+
+/** Method for inheriting attribute data from input pins. */
+UENUM()
+enum class EPCGAttributeInheritanceMode : uint8
+{
+	None UMETA(Tooltip = "No attributes inherited from initialization pins."),
+	CopyAttributeSetup UMETA(Tooltip = "Take attribute names and types from initialization pins. Pins declared first will take priority during conflicts. Does not copy values."),
+	// TODO: Mode to automatically copy values. Requires code-gen to loop on batches of elements for pins based on the total number of elements and total number of threads.
+	//CopyAttributeSetupAndValues UMETA(Tooltip = "Take attribute names, types, and initial values from initialization pins. Pins declared first will take priority during conflicts."),
+};
+
+/** Helper struct to nest GPU pin properties inside a UI category. */
+USTRUCT(BlueprintType)
+struct FPCGPinPropertiesGPUStruct
+{
+	GENERATED_BODY()
+
+	/** How the output data for this pin will be initialized. */
+	UPROPERTY(BlueprintReadWrite, EditAnywhere, Category = Settings, meta = (EditCondition = "bAllowEditInitMode", EditConditionHides, HideEditConditionToggle))
+	EPCGPinInitMode InitializationMode = EPCGPinInitMode::FromInputPins;
+
+	/** Input pins to initialize this pin's data from. */
+	UPROPERTY(BlueprintReadWrite, EditAnywhere, Category = Settings, meta = (GetOptions = "GetInputPinNames", EditCondition = "bAllowEditInitMode && InitializationMode == EPCGPinInitMode::FromInputPins", EditConditionHides))
+	TArray<FName> PinsToInititalizeFrom;
+
+	/** How the number of data is determined. */
+	UPROPERTY(BlueprintReadWrite, EditAnywhere, Category = Settings, meta = (EditCondition = "bAllowEditInitMode && bAllowEditDataCount && InitializationMode == EPCGPinInitMode::FromInputPins", EditConditionHides))
+	EPCGDataCountMode DataCountMode = EPCGDataCountMode::FromInputData;
+
+	/** How to combine data counts. */
+	UPROPERTY(BlueprintReadWrite, EditAnywhere, Category = Settings, meta = (EditCondition = "bAllowEditInitMode && bAllowEditDataCount && InitializationMode == EPCGPinInitMode::FromInputPins && DataCountMode == EPCGDataCountMode::FromInputData && bMultipleInitPins", EditConditionHides))
+	EPCGDataMultiplicity DataMultiplicity = EPCGDataMultiplicity::Pairwise;
+
+	/** Number of data to create. */
+	UPROPERTY(BlueprintReadWrite, EditAnywhere, Category = Settings, meta = (EditCondition = "bAllowEditInitMode && (InitializationMode == EPCGPinInitMode::Custom || DataCountMode == EPCGDataCountMode::Fixed)", EditConditionHides, ClampMin = 1))
+	int DataCount = 1;
+
+	/** How the number of elements is determined. */
+	UPROPERTY(BlueprintReadWrite, EditAnywhere, Category = Settings, meta = (EditCondition = "bAllowEditInitMode && InitializationMode == EPCGPinInitMode::FromInputPins", EditConditionHides))
+	EPCGElementCountMode ElementCountMode = EPCGElementCountMode::FromInputData;
+
+	/** How to combine element counts. */
+	UPROPERTY(BlueprintReadWrite, EditAnywhere, Category = Settings, meta = (EditCondition = "bAllowEditInitMode && InitializationMode == EPCGPinInitMode::FromInputPins && ElementCountMode == EPCGElementCountMode::FromInputData && bMultipleInitPins", EditConditionHides))
+	EPCGElementMultiplicity ElementMultiplicity = EPCGElementMultiplicity::Product;
+
+	/** Fixed number of elements to create in each output data. */
+	UPROPERTY(BlueprintReadWrite, EditAnywhere, Category = Settings, meta = (EditCondition = "bAllowEditInitMode && (InitializationMode == EPCGPinInitMode::Custom || ElementCountMode == EPCGElementCountMode::Fixed)", EditConditionHides, ClampMin = 1))
+	int ElementCount = 1;
+
+	/** How to inherit attribute names, types, and values. */
+	UPROPERTY(BlueprintReadWrite, EditAnywhere, Category = Settings, meta = (EditCondition = "bAllowEditInitMode && InitializationMode == EPCGPinInitMode::FromInputPins", EditConditionHides))
+	EPCGAttributeInheritanceMode AttributeInheritanceMode = EPCGAttributeInheritanceMode::CopyAttributeSetup;
+
+	/** Add entries to create new attributes on data emitted by this pin. */
+	UPROPERTY(EditAnywhere, DisplayName = "Attributes to Create", Category = Settings)
+	TArray<FPCGKernelAttributeKey> CreatedKernelAttributeKeys;
+
+#if WITH_EDITORONLY_DATA
+	UPROPERTY(Transient)
+	bool bAllowEditInitMode = false;
+
+	/** Some types force single data, such as Attribute Sets. We should not allow data count to be edited for such types. */
+	UPROPERTY(Transient)
+	bool bAllowEditDataCount = false;
+
+	UPROPERTY(Transient)
+	bool bMultipleInitPins = false;
+#endif // WITH_EDITORONLY_DATA
+};
+
+/** An extension of the pin properties that adds hints for GPU specific properties, such as buffer size and data layout. */
+USTRUCT(BlueprintType)
+struct PCG_API FPCGPinPropertiesGPU : public FPCGPinProperties
+{
+	GENERATED_BODY()
+
+public:
+	FPCGPinPropertiesGPU() = default;
+
+	explicit FPCGPinPropertiesGPU(const FName& InLabel, EPCGDataType InAllowedTypes)
+		: FPCGPinProperties(InLabel, InAllowedTypes)
+	{}
+
+	explicit FPCGPinPropertiesGPU(const FPCGPinProperties& InPinProps)
+		: FPCGPinProperties(InPinProps.Label, InPinProps.AllowedTypes)
+	{}
+
+#if WITH_EDITOR
+	bool CanEditChange(const FEditPropertyChain& PropertyChain) const;
+#endif
+
+public:
+	UPROPERTY(BlueprintReadWrite, EditAnywhere, Category = Settings, meta = (DisplayName = "GPU Properties", DisplayAfter = "Tooltip"))
+	FPCGPinPropertiesGPUStruct PropertiesGPU;
+};
+
+template<>
+struct TStructOpsTypeTraits<FPCGPinPropertiesGPU> : public TStructOpsTypeTraitsBase2<FPCGPinPropertiesGPU>
+{
+	enum
+	{
+		WithCanEditChange = true,
+	};
+};
