Added tags to PCG ism descriptor.
Replaced usage of FISMComponentDescriptor in the builder params to have a more streamlined code regarding tags.

diff --git a/Engine/Plugins/PCG/Source/PCG/Private/Elements/PCGSpawnActor.cpp b/Engine/Plugins/PCG/Source/PCG/Private/Elements/PCGSpawnActor.cpp
--- a/Engine/Plugins/PCG/Source/PCG/Private/Elements/PCGSpawnActor.cpp
+++ b/Engine/Plugins/PCG/Source/PCG/Private/Elements/PCGSpawnActor.cpp
@@ -673,11 +673,11 @@ void FPCGSpawnActorElement::CollapseIntoTargetActor(FPCGSubgraphContext* Context
 	const UPCGSpawnActorSettings* Settings = Context->GetInputSettings<UPCGSpawnActorSettings>();
 	check(Settings);
 
-	TMap<FPCGISMCBuilderParameters, TArray<FTransform>> MeshDescriptorTransforms;
+	TMap<FPCGISMComponentBuilderParams, TArray<FTransform>> MeshDescriptorTransforms;
 
 	AActor::ForEachComponentOfActorClassDefault<UStaticMeshComponent>(TemplateActorClass, [&MeshDescriptorTransforms](const UStaticMeshComponent* StaticMeshComponent)
 	{
-		FPCGISMCBuilderParameters Params;
+		FPCGISMComponentBuilderParams Params;
 		Params.Descriptor.InitFrom(StaticMeshComponent);
 		// TODO: No custom data float support?
 
@@ -705,9 +705,9 @@ void FPCGSpawnActorElement::CollapseIntoTargetActor(FPCGSubgraphContext* Context
 		return true;
 	});
 
-	for (const TPair<FPCGISMCBuilderParameters, TArray<FTransform>>& ISMCBuilderTransforms : MeshDescriptorTransforms)
+	for (const TPair<FPCGISMComponentBuilderParams, TArray<FTransform>>& ISMCBuilderTransforms : MeshDescriptorTransforms)
 	{
-		const FPCGISMCBuilderParameters& ISMCParams = ISMCBuilderTransforms.Key;
+		const FPCGISMComponentBuilderParams& ISMCParams = ISMCBuilderTransforms.Key;
 
 		UPCGManagedISMComponent* MISMC = UPCGActorHelpers::GetOrCreateManagedISMC(TargetActor, Context->SourceComponent.Get(), Settings->UID, ISMCParams);
 		if (!MISMC)
diff --git a/Engine/Plugins/PCG/Source/PCG/Private/Elements/PCGStaticMeshSpawner.cpp b/Engine/Plugins/PCG/Source/PCG/Private/Elements/PCGStaticMeshSpawner.cpp
--- a/Engine/Plugins/PCG/Source/PCG/Private/Elements/PCGStaticMeshSpawner.cpp
+++ b/Engine/Plugins/PCG/Source/PCG/Private/Elements/PCGStaticMeshSpawner.cpp
@@ -419,8 +419,8 @@ void FPCGStaticMeshSpawnerElement::SpawnStaticMeshInstances(FPCGStaticMeshSpawne
 
 	const FPCGPackedCustomData& PackedCustomData = *InPackedCustomData;
 
-	FPCGISMCBuilderParameters Params;
-	Params.Descriptor = FISMComponentDescriptor(InstanceList.Descriptor);
+	FPCGISMComponentBuilderParams Params;
+	Params.Descriptor = InstanceList.Descriptor;
 	Params.NumCustomDataFloats = PackedCustomData.NumCustomDataFloats;
 
 	// If the root actor we're binding to is movable, then the ISMC should be movable by default
diff --git a/Engine/Plugins/PCG/Source/PCG/Private/Helpers/PCGActorHelpers.cpp b/Engine/Plugins/PCG/Source/PCG/Private/Helpers/PCGActorHelpers.cpp
--- a/Engine/Plugins/PCG/Source/PCG/Private/Helpers/PCGActorHelpers.cpp
+++ b/Engine/Plugins/PCG/Source/PCG/Private/Helpers/PCGActorHelpers.cpp
@@ -41,7 +41,19 @@ namespace UE::PCGActorHelpers::Local
 
 #endif
 
+PRAGMA_DISABLE_DEPRECATION_WARNINGS
 UInstancedStaticMeshComponent* UPCGActorHelpers::GetOrCreateISMC(AActor* InTargetActor, UPCGComponent* InSourceComponent, uint64 SettingsUID, const FPCGISMCBuilderParameters& InParams)
+{
+	return GetOrCreateISMC(InTargetActor, InSourceComponent, SettingsUID, FPCGISMComponentBuilderParams(InParams));
+}
+
+UPCGManagedISMComponent* UPCGActorHelpers::GetOrCreateManagedISMC(AActor* InTargetActor, UPCGComponent* InSourceComponent, uint64 SettingsUID, const FPCGISMCBuilderParameters& InParams)
+{
+	return GetOrCreateManagedISMC(InTargetActor, InSourceComponent, SettingsUID, FPCGISMComponentBuilderParams(InParams));
+}
+PRAGMA_ENABLE_DEPRECATION_WARNINGS
+
+UInstancedStaticMeshComponent* UPCGActorHelpers::GetOrCreateISMC(AActor* InTargetActor, UPCGComponent* InSourceComponent, uint64 SettingsUID, const FPCGISMComponentBuilderParams& InParams)
 {
 	UPCGManagedISMComponent* MISMC = GetOrCreateManagedISMC(InTargetActor, InSourceComponent, SettingsUID, InParams);
 	if (MISMC)
@@ -54,19 +66,33 @@ UInstancedStaticMeshComponent* UPCGActorHelpers::GetOrCreateISMC(AActor* InTarge
 	}
 }
 
-UPCGManagedISMComponent* UPCGActorHelpers::GetOrCreateManagedISMC(AActor* InTargetActor, UPCGComponent* InSourceComponent, uint64 SettingsUID, const FPCGISMCBuilderParameters& InParams)
+UPCGManagedISMComponent* UPCGActorHelpers::GetOrCreateManagedISMC(AActor* InTargetActor, UPCGComponent* InSourceComponent, uint64 SettingsUID, const FPCGISMComponentBuilderParams& InParams)
 {
 	check(InTargetActor && InSourceComponent);
 
-	const UStaticMesh* StaticMesh = InParams.Descriptor.StaticMesh;
+	TRACE_CPUPROFILER_EVENT_SCOPE(UPCGActorHelpers::GetOrCreateManagedISMC);
+
+	FISMComponentDescriptor Descriptor(InParams.Descriptor);
+
+	const UStaticMesh* StaticMesh = Descriptor.StaticMesh;
 	if (!StaticMesh)
 	{
 		return nullptr;
 	}
 
-	TRACE_CPUPROFILER_EVENT_SCOPE(UPCGActorHelpers::GetOrCreateManagedISMC);
-
-	FISMComponentDescriptor Descriptor = InParams.Descriptor;
+	auto AddTagsToComponent = [InSourceComponent, &InParams](UInstancedStaticMeshComponent* ISMC)
+	{
+		if(ISMC)
+		{
+			ISMC->ComponentTags.AddUnique(PCGHelpers::DefaultPCGTag);
+			ISMC->ComponentTags.AddUnique(InSourceComponent->GetFName());
+			
+			for (FName ComponentTag : InParams.Descriptor.ComponentTags)
+			{
+				ISMC->ComponentTags.AddUnique(ComponentTag);
+			}
+		}
+	};
 
 	// If the component class is invalid, default to HISM.
 	// TODO: should this be part of the descriptor changes?
@@ -128,6 +154,7 @@ UPCGManagedISMComponent* UPCGActorHelpers::GetOrCreateManagedISMC(AActor* InTarg
 		{
 			TRACE_CPUPROFILER_EVENT_SCOPE(UPCGActorHelpers::GetOrCreateManagedISMC::MarkAsUsed);
 			MatchingResource->MarkAsUsed();
+			AddTagsToComponent(Cast<UInstancedStaticMeshComponent>(MatchingResource->GeneratedComponent.Get()));
 
 			return MatchingResource;
 		}
@@ -160,8 +187,8 @@ UPCGManagedISMComponent* UPCGActorHelpers::GetOrCreateManagedISMC(AActor* InTarg
 	InTargetActor->AddInstanceComponent(ISMC);
 
 	ISMC->AttachToComponent(InTargetActor->GetRootComponent(), FAttachmentTransformRules(EAttachmentRule::KeepRelative, EAttachmentRule::KeepWorld, EAttachmentRule::KeepWorld, false));
-	ISMC->ComponentTags.Add(InSourceComponent->GetFName());
-	ISMC->ComponentTags.Add(PCGHelpers::DefaultPCGTag);
+	// Implementation note: Because we've used the FISMComponentDescriptor here (Descriptor vs InParams.Descriptor) which takes care of the loading, we still need to apply tags manually.
+	AddTagsToComponent(ISMC);
 
 	// Create managed resource on source component
 	UPCGManagedISMComponent* Resource = NewObject<UPCGManagedISMComponent>(InSourceComponent);
diff --git a/Engine/Plugins/PCG/Source/PCG/Private/MeshSelectors/PCGISMDescriptor.cpp b/Engine/Plugins/PCG/Source/PCG/Private/MeshSelectors/PCGISMDescriptor.cpp
--- a/Engine/Plugins/PCG/Source/PCG/Private/MeshSelectors/PCGISMDescriptor.cpp
+++ b/Engine/Plugins/PCG/Source/PCG/Private/MeshSelectors/PCGISMDescriptor.cpp
@@ -16,6 +16,22 @@ FPCGSoftISMComponentDescriptor::FPCGSoftISMComponentDescriptor()
 	BodyInstance.SetCollisionProfileName(UCollisionProfile::NoCollision_ProfileName);
 }
 
+FPCGSoftISMComponentDescriptor::FPCGSoftISMComponentDescriptor(const FISMComponentDescriptor& Other)
+	: Super(Other)
+{}
+
+void FPCGSoftISMComponentDescriptor::InitFrom(const UStaticMeshComponent* Component, bool bInitBodyInstance)
+{
+	ComponentTags = Component->ComponentTags;
+	Super::InitFrom(Component, bInitBodyInstance);
+}
+
+void FPCGSoftISMComponentDescriptor::InitComponent(UInstancedStaticMeshComponent* ISMComponent) const
+{
+	ISMComponent->ComponentTags = ComponentTags;
+	Super::InitComponent(ISMComponent);
+}
+
 bool FPCGSoftISMComponentDescriptor::SerializeFromMismatchedTag(const FPropertyTag& Tag, FStructuredArchive::FSlot Slot)
 {
 	if (Tag.GetType().IsStruct(FSoftISMComponentDescriptor::StaticStruct()->GetFName()))
diff --git a/Engine/Plugins/PCG/Source/PCG/Private/PCGManagedResource.cpp b/Engine/Plugins/PCG/Source/PCG/Private/PCGManagedResource.cpp
--- a/Engine/Plugins/PCG/Source/PCG/Private/PCGManagedResource.cpp
+++ b/Engine/Plugins/PCG/Source/PCG/Private/PCGManagedResource.cpp
@@ -602,7 +602,9 @@ void UPCGManagedComponentBase::MarkAsUsed()
 	{
 		if (GeneratedComponent.Get())
 		{
-			GeneratedComponent->ComponentTags.Remove(PCGHelpers::MarkedForCleanupPCGTag);
+			// Remove all non-default tags, including the "marked for cleanup" tag
+			GeneratedComponent->ComponentTags.Reset();
+			GeneratedComponent->ComponentTags.Add(PCGHelpers::DefaultPCGTag);
 		}
 	}
 }
diff --git a/Engine/Plugins/PCG/Source/PCG/Public/Helpers/PCGActorHelpers.h b/Engine/Plugins/PCG/Source/PCG/Public/Helpers/PCGActorHelpers.h
--- a/Engine/Plugins/PCG/Source/PCG/Public/Helpers/PCGActorHelpers.h
+++ b/Engine/Plugins/PCG/Source/PCG/Public/Helpers/PCGActorHelpers.h
@@ -3,6 +3,7 @@
 #pragma once
 
 #include "Elements/PCGSplineMeshParams.h"
+#include "MeshSelectors/PCGISMDescriptor.h"
 
 #include "Engine/EngineTypes.h"
 #include "Engine/SplineMeshComponentDescriptor.h"
@@ -26,7 +27,7 @@ class UPCGManagedSplineMeshComponent;
 class UStaticMesh;
 class UWorld;
 
-struct FPCGISMCBuilderParameters
+struct UE_DEPRECATED(5.5, "Use FPCGISMComponentBuilderParams instead.") FPCGISMCBuilderParameters
 {
 	FISMComponentDescriptor Descriptor;
 	int32 NumCustomDataFloats = 0;
@@ -40,6 +41,29 @@ struct FPCGISMCBuilderParameters
 	inline bool operator==(const FPCGISMCBuilderParameters& Other) const { return Descriptor == Other.Descriptor && NumCustomDataFloats == Other.NumCustomDataFloats && bAllowDescriptorChanges == Other.bAllowDescriptorChanges; }
 };
 
+struct FPCGISMComponentBuilderParams
+{
+	FPCGISMComponentBuilderParams() = default;
+
+	FPCGSoftISMComponentDescriptor Descriptor;
+	int32 NumCustomDataFloats = 0;
+	bool bAllowDescriptorChanges = true;
+
+PRAGMA_DISABLE_DEPRECATION_WARNINGS
+	explicit FPCGISMComponentBuilderParams(const FPCGISMCBuilderParameters& Params)
+	: Descriptor(Params.Descriptor), NumCustomDataFloats(Params.NumCustomDataFloats), bAllowDescriptorChanges(Params.bAllowDescriptorChanges)
+	{
+	}
+PRAGMA_ENABLE_DEPRECATION_WARNINGS
+
+	friend inline uint32 GetTypeHash(const FPCGISMComponentBuilderParams& Key)
+	{
+		return HashCombine(HashCombine(GetTypeHash(Key.Descriptor), 1 + Key.NumCustomDataFloats), (Key.bAllowDescriptorChanges ? 2 : 1));
+	}
+
+	inline bool operator==(const FPCGISMComponentBuilderParams& Other) const { return Descriptor == Other.Descriptor && NumCustomDataFloats == Other.NumCustomDataFloats && bAllowDescriptorChanges == Other.bAllowDescriptorChanges; }
+};
+
 struct FPCGSplineMeshComponentBuilderParameters
 {
 	FSplineMeshComponentDescriptor Descriptor;
@@ -59,8 +83,13 @@ class PCG_API UPCGActorHelpers : public UBlueprintFunctionLibrary
 	GENERATED_BODY()
 
 public:
+PRAGMA_DISABLE_DEPRECATION_WARNINGS
 	static UInstancedStaticMeshComponent* GetOrCreateISMC(AActor* InTargetActor, UPCGComponent* SourceComponent, uint64 SettingsUID, const FPCGISMCBuilderParameters& Params);
 	static UPCGManagedISMComponent* GetOrCreateManagedISMC(AActor* InTargetActor, UPCGComponent* SourceComponent, uint64 SettingsUID, const FPCGISMCBuilderParameters& Params);
+PRAGMA_ENABLE_DEPRECATION_WARNINGS
+
+	static UInstancedStaticMeshComponent* GetOrCreateISMC(AActor* InTargetActor, UPCGComponent* SourceComponent, uint64 SettingsUID, const FPCGISMComponentBuilderParams& Params);
+	static UPCGManagedISMComponent* GetOrCreateManagedISMC(AActor* InTargetActor, UPCGComponent* SourceComponent, uint64 SettingsUID, const FPCGISMComponentBuilderParams& Params);
 	static USplineMeshComponent* GetOrCreateSplineMeshComponent(AActor* InTargetActor, UPCGComponent* SourceComponent, uint64 SettingsUID, const FPCGSplineMeshComponentBuilderParameters& Params);
 	static UPCGManagedSplineMeshComponent* GetOrCreateManagedSplineMeshComponent(AActor* InTargetActor, UPCGComponent* SourceComponent, uint64 SettingsUID, const FPCGSplineMeshComponentBuilderParameters& Params);
 	static bool DeleteActors(UWorld* World, const TArray<TSoftObjectPtr<AActor>>& ActorsToDelete);
diff --git a/Engine/Plugins/PCG/Source/PCG/Public/MeshSelectors/PCGISMDescriptor.h b/Engine/Plugins/PCG/Source/PCG/Public/MeshSelectors/PCGISMDescriptor.h
--- a/Engine/Plugins/PCG/Source/PCG/Public/MeshSelectors/PCGISMDescriptor.h
+++ b/Engine/Plugins/PCG/Source/PCG/Public/MeshSelectors/PCGISMDescriptor.h
@@ -8,14 +8,21 @@
 #include "PCGISMDescriptor.generated.h"
 
 /** Convenience PCG-side component descriptor so we can adjust defaults to the most common use cases. */
+// Implementation note: the tags don't really need to contribute to the hash, so we will retain the base class !=, == and ComputeHash implementations
 USTRUCT()
 struct FPCGSoftISMComponentDescriptor : public FSoftISMComponentDescriptor
 {
 	GENERATED_BODY()
 
 	PCG_API FPCGSoftISMComponentDescriptor();
+	PCG_API explicit FPCGSoftISMComponentDescriptor(const FISMComponentDescriptor& Other);
+	PCG_API virtual void InitFrom(const UStaticMeshComponent* Component, bool bInitBodyInstance = true) override;
+	PCG_API virtual void InitComponent(UInstancedStaticMeshComponent* ISMComponent) const override;
 
 public:
+	UPROPERTY(EditAnywhere, Category="Component Settings|Tags", meta = (DisplayAfter = "ComponentClass"))
+	TArray<FName> ComponentTags;
+
 	PCG_API bool SerializeFromMismatchedTag(const FPropertyTag& Tag, FStructuredArchive::FSlot Slot);
 };
 
diff --git a/Engine/Plugins/PCG/Source/PCGEditor/Private/DataVisualizations/PCGSpatialDataVisualization.cpp b/Engine/Plugins/PCG/Source/PCGEditor/Private/DataVisualizations/PCGSpatialDataVisualization.cpp
--- a/Engine/Plugins/PCG/Source/PCGEditor/Private/DataVisualizations/PCGSpatialDataVisualization.cpp
+++ b/Engine/Plugins/PCG/Source/PCGEditor/Private/DataVisualizations/PCGSpatialDataVisualization.cpp
@@ -100,13 +100,8 @@ void IPCGSpatialDataVisualization::ExecuteDebugDisplay(FPCGContext* Context, con
 		return;
 	}
 
-	UMaterialInterface* Material = DebugSettings.GetMaterial().LoadSynchronous();
-
-	TArray<UMaterialInterface*> Materials;
-	if (Material)
-	{
-		Materials.Add(Material);
-	}
+	TArray<TSoftObjectPtr<UMaterialInterface>> Materials;
+	Materials.Add(DebugSettings.GetMaterial());
 
 	const UPCGPointData* PointData = CollapseToDebugPointData(Context, Data);
 
@@ -160,7 +155,7 @@ void IPCGSpatialDataVisualization::ExecuteDebugDisplay(FPCGContext* Context, con
 		}
 	}
 
-	FPCGISMCBuilderParameters Params[2];
+	FPCGISMComponentBuilderParams Params[2];
 	Params[0].NumCustomDataFloats = NumCustomData;
 	Params[0].Descriptor.StaticMesh = Mesh;
 	Params[0].Descriptor.OverrideMaterials = Materials;
