[PCG] Create a Fuse node.

+ Fuses two point data or same point data by distance
+ Fuse priorities include: target, min attribute, max attribute, weight, weighted average
+ Fuse targets include: transform, attribute
+ Can remove fused points

Future TODO:
* Various input combinations, 1:N, N:1, N:N
* Fuse N points or Fuse all but N points, etc.

diff --git a/Engine/Plugins/PCG/Source/PCG/Private/Elements/PCGFuseElement.cpp b/Engine/Plugins/PCG/Source/PCG/Private/Elements/PCGFuseElement.cpp
new file mode 100644
--- /dev/null
+++ b/Engine/Plugins/PCG/Source/PCG/Private/Elements/PCGFuseElement.cpp
@@ -0,0 +1,693 @@
+// Copyright Epic Games, Inc. All Rights Reserved.
+
+#include "Elements/PCGFuseElement.h"
+
+#include "PCGContext.h"
+#include "Data/PCGPointData.h"
+#include "Helpers/PCGAsync.h"
+#include "Helpers/PCGHelpers.h"
+#include "Metadata/Accessors/PCGAttributeAccessorHelpers.h"
+#include "SpatialAlgo/PCGOctreeQueries.h"
+
+#include "Algo/MaxElement.h"
+
+#define LOCTEXT_NAMESPACE "PCGFuseElement"
+
+namespace PCGFuseElement
+{
+	namespace Constants
+	{
+		const FName InputSourceLabel = TEXT("Source");
+		const FName InputTargetLabel = TEXT("Target");
+	}
+
+	namespace Helpers
+	{
+		static void FusePoints(FFuseState& OutFuseData, const int32 PointIndex, const int32 TargetIndex, const UPCGFuseSettings& Settings)
+		{
+			TRACE_CPUPROFILER_EVENT_SCOPE(FPCGFuseElement::FusePoints);
+
+			// If we're just going to remove the points later, no need to do any work.
+			if (Settings.bRemoveFusedPoints)
+			{
+				return;
+			}
+
+			const TArray<FPCGPoint>& TargetPoints = OutFuseData.TargetPointData->GetPoints();
+			TArray<FPCGPoint>& OutPoints = OutFuseData.OutPointData->GetMutablePoints();
+
+			double Alpha = 0.0;
+			switch (Settings.Priority)
+			{
+				case EPCGFusePriority::Target:
+				{
+					Alpha = 1.0;
+					break;
+				}
+
+				case EPCGFusePriority::MinAttribute: // Fall-through
+				case EPCGFusePriority::MaxAttribute:
+				{
+					auto Callback = [PointIndex, TargetIndex, &OutFuseData]<typename T>(T) -> bool
+					{
+						if constexpr (PCG::Private::MetadataTraits<T>::CanCompare)
+						{
+							bool bSuccess = true;
+							T SourceValue{};
+							bSuccess &= OutFuseData.SourcePriorityAccessor->Get(SourceValue, PointIndex, *OutFuseData.SourcePriorityKeys, EPCGAttributeAccessorFlags::AllowBroadcastAndConstructible);
+							T TargetValue{};
+							bSuccess &= OutFuseData.TargetPriorityAccessor->Get(TargetValue, TargetIndex, *OutFuseData.TargetPriorityKeys, EPCGAttributeAccessorFlags::AllowBroadcastAndConstructible);
+							return bSuccess && PCG::Private::MetadataTraits<T>::LessOrEqual(SourceValue, TargetValue);
+						}
+						else
+						{
+							return true;
+						}
+					};
+					bool bSourceIsLessThanOrEqual = PCGMetadataAttribute::CallbackWithRightType(OutFuseData.SourcePriorityAccessor->GetUnderlyingType(), Callback);
+
+					const bool bSelectSource = (Settings.Priority == EPCGFusePriority::MinAttribute) == bSourceIsLessThanOrEqual;
+					Alpha = bSelectSource ? 0.0 : 1.0;
+					break;
+				}
+
+				case EPCGFusePriority::Weighted:
+				{
+					// Should have already been checked in PrepareData
+					check(!FMath::IsNearlyZero(OutFuseData.WeightMax));
+					Alpha = Settings.bWeightAsAttribute ? (OutFuseData.SourceWeights[PointIndex] / OutFuseData.WeightMax) : Settings.Weight;
+					break;
+				}
+
+				case EPCGFusePriority::WeightedAverage:
+				{
+					// Should have already been checked in PrepareData
+					check(!FMath::IsNearlyZero(OutFuseData.WeightMax));
+					const double SourceWeight = Settings.bWeightAsAttribute ? OutFuseData.SourceWeights[PointIndex] : Settings.Weight;
+					const double TargetWeight = OutFuseData.TargetWeights[TargetIndex];
+					Alpha = (SourceWeight + TargetWeight) / OutFuseData.WeightMax * 0.5;
+					break;
+				}
+
+				default:
+					checkNoEntry();
+					return;
+			}
+
+			if (Settings.bFuseAttribute)
+			{
+				auto Callback = [PointIndex, TargetIndex, Alpha, &OutFuseData]<typename T>(T)
+				{
+					// Only relevant if the attribute can be interpolated.
+					if constexpr (PCG::Private::MetadataTraits<T>::CanInterpolate)
+					{
+						T TargetValue{}, Value{};
+						OutFuseData.SourceFuseAttributeAccessor->Get(Value, PointIndex, *OutFuseData.SourceFuseAttributeKeys, EPCGAttributeAccessorFlags::AllowBroadcastAndConstructible);
+						OutFuseData.TargetFuseAttributeAccessor->Get(TargetValue, TargetIndex, *OutFuseData.TargetFuseAttributeKeys, EPCGAttributeAccessorFlags::AllowBroadcastAndConstructible);
+
+						// TODO: Can instead use the metadata traits' WeightedSum if the Transform version is added
+						// Transform has an explicit interpolation mode.
+						if constexpr (std::is_same_v<T, FTransform>)
+						{
+							Value.BlendWith(TargetValue, Alpha);
+						}
+						else
+						{
+							Value = FMath::Lerp(Value, TargetValue, Alpha);
+						}
+
+						OutFuseData.OutputFuseAttributeAccessor->Set(Value, PointIndex, *OutFuseData.OutputFuseAttributeKeys);
+					}
+				};
+				PCGMetadataAttribute::CallbackWithRightType(OutFuseData.SourceFuseAttributeAccessor->GetUnderlyingType(), Callback);
+			}
+
+			// Full transform shortcut.
+			if (Settings.bFuseLocation && Settings.bFuseRotation && Settings.bFuseScale)
+			{
+				OutPoints[PointIndex].Transform.BlendWith(TargetPoints[TargetIndex].Transform, Alpha);
+				return;
+			}
+
+			if (Settings.bFuseLocation)
+			{
+				OutPoints[PointIndex].Transform.SetLocation(FMath::Lerp(OutPoints[PointIndex].Transform.GetLocation(), TargetPoints[TargetIndex].Transform.GetLocation(), Alpha));
+			}
+
+			if (Settings.bFuseRotation)
+			{
+				OutPoints[PointIndex].Transform.SetRotation(FMath::Lerp(OutPoints[PointIndex].Transform.GetRotation(), TargetPoints[TargetIndex].Transform.GetRotation(), Alpha));
+			}
+
+			if (Settings.bFuseScale)
+			{
+				OutPoints[PointIndex].Transform.SetScale3D(FMath::Lerp(OutPoints[PointIndex].Transform.GetScale3D(), TargetPoints[TargetIndex].Transform.GetScale3D(), Alpha));
+			}
+		}
+
+		/* Since the flags are bit arrays, traverse in reverse order. RemoveAtSwap will change the order of the points, 
+		 * so the bit arrays are only valid for one iteration, before needing to be reset.
+		 */
+		static void RemovedFlaggedPoints(FFuseState& OutFuseData)
+		{
+			TRACE_CPUPROFILER_EVENT_SCOPE(FPCGFuseElement::RemovedFlaggedPoints);
+
+			TArray<FPCGPoint>& OutPoints = OutFuseData.OutPointData->GetMutablePoints();
+			int32 Index = OutFuseData.RemovedIndices.FindLastFrom(true, OutPoints.Num() - 1);
+			while (OutPoints.IsValidIndex(Index))
+			{
+				// Swap the points and their removal flag.
+				OutFuseData.RemovedIndices[Index] = false;
+				OutFuseData.RemovedIndices[OutPoints.Num() - 1] = true;
+				OutPoints.RemoveAtSwap(Index, EAllowShrinking::No);
+				Index = OutFuseData.RemovedIndices.FindLastFrom(true, Index - 1);
+			}
+		}
+
+		// Runs a single iteration on a point to check for fuse candidates.
+		static bool ExecuteOnClosestPoint(const int32 Index, FFuseState& OutFuseData, const UPCGFuseSettings& FuseSettings)
+		{
+			TRACE_CPUPROFILER_EVENT_SCOPE(FPCGFuseElement::ExecuteOnClosestPoint);
+
+			check(!OutFuseData.CheckedIndices[Index]);
+
+			if (OutFuseData.RemovedIndices[Index])
+			{
+				return false;
+			}
+
+			const TArray<FPCGPoint>& SourcePoints = OutFuseData.SourcePointData->GetPoints();
+
+			// If source and target are not the same, then just find the nearest point and fuse it.
+			if (const FPCGPoint* ClosestPoint = UPCGOctreeQueries::GetClosestPoint(
+					OutFuseData.TargetPointData,
+					SourcePoints[Index].Transform.GetLocation(),
+					/*bInDiscardCenter=*/false,
+					FuseSettings.Distance))
+			{
+				const int32 TargetIndex = ClosestPoint - OutFuseData.TargetPointData->GetPoints().GetData();
+				FusePoints(OutFuseData, Index, TargetIndex, FuseSettings);
+				OutFuseData.RemovedIndices[Index] = true;
+
+				return true;
+			}
+
+			return false;
+		}
+
+		// Runs a single iteration on a point to check for fuse candidates. Assumes the target and source are the same.
+		static bool ExecuteOnClosestPointSameSource(const int32 Index, FFuseState& OutFuseData, const UPCGFuseSettings& FuseSettings)
+		{
+			TRACE_CPUPROFILER_EVENT_SCOPE(FPCGFuseElement::ExecuteOnClosestPointSameSource);
+
+			check(!OutFuseData.CheckedIndices[Index]);
+
+			if (OutFuseData.RemovedIndices[Index])
+			{
+				return false;
+			}
+
+			const TArray<FPCGPoint>& OutPoints = OutFuseData.OutPointData->GetPoints();
+
+			double MinDistanceSquared = std::numeric_limits<double>::max();
+			int32 ClosestTargetIndex = INDEX_NONE;
+			bool bFoundColocatedPoint = false;
+
+			// Cycle through co-located points, skipping if it is the source point, or if the found points were already checked.
+			UPCGOctreeQueries::ForEachPointInsideSphere(
+					OutFuseData.TargetPointData,
+					OutPoints[Index].Transform.GetLocation(),
+					FuseSettings.Distance,
+					[Index, &OutFuseData, &MinDistanceSquared, &ClosestTargetIndex, &bFoundColocatedPoint](const FPCGPointRef& PointRef, const double DistanceSquared)
+					{
+						// Found a co-located target already during this iteration.
+						if (bFoundColocatedPoint)
+						{
+							return;
+						}
+
+						const int32 TargetIndex = PointRef.Point - OutFuseData.TargetPointData->GetPoints().GetData();
+						if (OutFuseData.RemovedIndices[TargetIndex])
+						{
+							return;
+						}
+
+						// If the points are co-located, the current point will be fused.
+						if (FMath::IsNearlyZero(DistanceSquared))
+						{
+							if (TargetIndex != Index && !OutFuseData.CheckedIndices[TargetIndex])
+							{
+								ClosestTargetIndex = TargetIndex;
+								bFoundColocatedPoint = true;
+							}
+
+							return;
+						}
+
+						// No co-located points that haven't already been checked. Distance check to find the closest point.
+						{
+							if (DistanceSquared < MinDistanceSquared)
+							{
+								ClosestTargetIndex = TargetIndex;
+								MinDistanceSquared = DistanceSquared;
+							}
+						}
+					});
+
+			if (ClosestTargetIndex != INDEX_NONE)
+			{
+				FusePoints(OutFuseData, Index, ClosestTargetIndex, FuseSettings);
+
+				OutFuseData.RemovedIndices[ClosestTargetIndex] = true;
+
+				return true;
+			}
+
+			return false;
+		}
+
+		// Runs a single iteration of the fuse algorithm. Returns true when no points were fused.
+		static bool RunFuseIteration(const int32 Index, FFuseState& OutFuseData, const UPCGFuseSettings& FuseSettings, bool bFinalIteration)
+		{
+			// If the target == source, then there a different algorithm needed.
+			const auto& Function = OutFuseData.bTargetIsSource ? ExecuteOnClosestPointSameSource : ExecuteOnClosestPoint;
+			OutFuseData.bPointWasFusedThisCycle |= Function(Index, OutFuseData, FuseSettings);
+
+			// Annotate that we've checked this one, in case it comes up in the future.
+			OutFuseData.CheckedIndices[Index] = true;
+
+			// On the final iteration, if the target is the source and points were removed, the algorithm will need to restart.
+			if (bFinalIteration)
+			{
+				if (FuseSettings.bRemoveFusedPoints)
+				{
+					RemovedFlaggedPoints(OutFuseData);
+				}
+
+				// Nothing was fused or if target != source, then algorithm is complete.
+				if (!OutFuseData.bTargetIsSource || !OutFuseData.bPointWasFusedThisCycle)
+				{
+					return true;
+				}
+
+				// Adjust the tracking bit arrays to the new count and reset them.
+				OutFuseData.RemovedIndices.SetRange(0, OutFuseData.RemovedIndices.Num(), false);
+				OutFuseData.CheckedIndices.SetRange(0, OutFuseData.CheckedIndices.Num(), false);
+
+				// Reset the counters.
+				OutFuseData.IterationIndex = 0;
+				OutFuseData.bPointWasFusedThisCycle = false;
+			}
+
+			return false;
+		}
+	} // namespace Helpers
+
+	namespace Algorithm
+	{
+		static bool Sequential(const FPCGContext* InContext, FFuseState& OutFuseData, const UPCGFuseSettings& FuseSettings)
+		{
+			TRACE_CPUPROFILER_EVENT_SCOPE(FPCGFuseElement::Sequential);
+
+			check(InContext);
+
+			bool bProcessComplete = false;
+			do
+			{
+				const bool bIsFinalIteration = OutFuseData.IterationIndex == (OutFuseData.OutPointData->GetNumPoints() - 1);
+				bProcessComplete = Helpers::RunFuseIteration(OutFuseData.IterationIndex, OutFuseData, FuseSettings, bIsFinalIteration);
+				++OutFuseData.IterationIndex;
+
+				if (InContext->ShouldStop())
+				{
+					return bProcessComplete;
+				}
+			}
+			while (!bProcessComplete);
+
+			return true;
+		}
+
+		static bool Random(const FPCGContext* InContext, FFuseState& OutFuseData, const UPCGFuseSettings& FuseSettings)
+		{
+			TRACE_CPUPROFILER_EVENT_SCOPE(FPCGFuseElement::Random);
+
+			// Caller responsibility to initialize the random indices.
+			check(InContext && !OutFuseData.RandomIndices.IsEmpty());
+
+			bool bProcessComplete = false;
+			do
+			{
+				const int32 RandomIndex = OutFuseData.RandomIndices[OutFuseData.IterationIndex];
+				const bool bFinalIteration = OutFuseData.IterationIndex == (OutFuseData.OutPointData->GetNumPoints() - 1);
+				bProcessComplete = Helpers::RunFuseIteration(RandomIndex, OutFuseData, FuseSettings, bFinalIteration);
+				++OutFuseData.IterationIndex;
+
+				if (bFinalIteration)
+				{
+					OutFuseData.RandomIndices.SetNum(OutFuseData.TargetPointData->GetNumPoints());
+					for (int i = 0; i < OutFuseData.TargetPointData->GetNumPoints(); ++i)
+					{
+						OutFuseData.RandomIndices[i] = i;
+					}
+
+					PCGHelpers::ShuffleArray(OutFuseData.RandomStream, OutFuseData.RandomIndices);
+				}
+
+				if (InContext->ShouldStop())
+				{
+					return bProcessComplete;
+				}
+			}
+			while (!bProcessComplete);
+
+			return true;
+		}
+	} // namespace Algorithm
+} // namespace PCGFuseElement
+
+#if WITH_EDITOR
+FText UPCGFuseSettings::GetNodeTooltipText() const
+{
+	return LOCTEXT("NodeTooltip", "Fuse transforms or attributes of two 'co-located' points--based on a minimum distance and other criteria in world space.");
+}
+#endif // WITH_EDITOR
+
+bool UPCGFuseSettings::UseSeed() const
+{
+	return Algorithm == EPCGFuseAlgorithm::Random;
+}
+
+TArray<FPCGPinProperties> UPCGFuseSettings::InputPinProperties() const
+{
+	TArray<FPCGPinProperties> Properties;
+
+	FPCGPinProperties SourcePinProperty = Properties.Emplace_GetRef(PCGFuseElement::Constants::InputSourceLabel, EPCGDataType::Point);
+	SourcePinProperty.SetRequiredPin();
+
+	// TODO: For now, only allow a single target input. Only N:0 and N:1 are currently supported.
+	FPCGPinProperties TargetPinProperty = Properties.Emplace_GetRef(PCGFuseElement::Constants::InputTargetLabel, EPCGDataType::Point, /*bInAllowMultipleConnections=*/false, /*bAllowMultipleData=*/false);
+#if WITH_EDITOR
+	SourcePinProperty.Tooltip = LOCTEXT("SourcePinTooltip", "The source points will be compared by distance to the target points.");
+	TargetPinProperty.Tooltip = LOCTEXT("TargetPinTooltip", "[Optional] The target points will be used as a target for the fuse. If empty, the source points will be used as target points.");
+#endif // WITH_EDITOR
+
+	return Properties;
+}
+
+FPCGElementPtr UPCGFuseSettings::CreateElement() const
+{
+	return MakeShared<FPCGFuseElement>();
+}
+
+bool FPCGFuseElement::PrepareDataInternal(FPCGContext* InContext) const
+{
+	TRACE_CPUPROFILER_EVENT_SCOPE(FPCGFuseElement::PrepareData);
+
+	const UPCGFuseSettings* Settings = InContext->GetInputSettings<UPCGFuseSettings>();
+	check(Settings);
+
+	ContextType* Context = static_cast<ContextType*>(InContext);
+	check(Context);
+
+	const TArray<FPCGTaggedData> SourceInputs = Context->InputData.GetInputsByPin(PCGFuseElement::Constants::InputSourceLabel);
+	const TArray<FPCGTaggedData> TargetInputs = Context->InputData.GetInputsByPin(PCGFuseElement::Constants::InputTargetLabel);
+
+	// Early out if no source to modify
+	if (SourceInputs.IsEmpty())
+	{
+		return true;
+	}
+
+	// Only N:0 and N:1 are currently supported.
+	if (TargetInputs.Num() > 1 && TargetInputs.Num() != SourceInputs.Num())
+	{
+		PCGLog::InputOutput::LogInvalidCardinalityError(PCGFuseElement::Constants::InputSourceLabel, PCGFuseElement::Constants::InputTargetLabel, InContext);
+		return true;
+	}
+
+	Context->InitializePerExecutionState([Settings](const ContextType*, ExecStateType& OutState)
+	{
+		OutState.FuseFunction = nullptr;
+
+		// Algorithm must be selected on a per-input basis, as point fusing behavior differs for target vs self.
+		switch (Settings->Algorithm)
+		{
+			case EPCGFuseAlgorithm::Sequential:
+				OutState.FuseFunction = PCGFuseElement::Algorithm::Sequential;
+				break;
+			case EPCGFuseAlgorithm::Random:
+				OutState.FuseFunction = PCGFuseElement::Algorithm::Random;
+				break;
+			default:
+				checkNoEntry();
+				return EPCGTimeSliceInitResult::AbortExecution;
+		}
+
+		check(OutState.FuseFunction);
+
+		return EPCGTimeSliceInitResult::Success;
+	});
+
+	Context->InitializePerIterationStates(SourceInputs.Num(), [&SourceInputs, &TargetInputs, Settings, Context](IterStateType& OutState, const ExecStateType&, const uint32 IterationIndex)
+	{
+		const UPCGPointData* SourcePointData = CastChecked<UPCGPointData>(SourceInputs[IterationIndex].Data);
+
+		if (!SourcePointData || SourcePointData->IsEmpty())
+		{
+			return EPCGTimeSliceInitResult::NoOperation;
+		}
+
+		const UPCGPointData* TargetPointData = nullptr;
+		// Should have confirmed previously that we're N:0, N:1, or N:N.
+		if (TargetInputs.Num() != 1)
+		{
+			// No input target is fine. Just use the source as the target.
+			TargetPointData = TargetInputs.IsEmpty() ? SourcePointData : CastChecked<UPCGPointData>(TargetInputs[IterationIndex].Data);
+		}
+		else
+		{
+			TargetPointData = CastChecked<UPCGPointData>(TargetInputs[0].Data);
+		}
+
+		UPCGPointData* OutPointData = FPCGContext::NewObject_AnyThread<UPCGPointData>(Context);
+		OutPointData->InitializeFromData(SourcePointData);
+		// Create output points in advance as a copy of the input source points.
+		OutPointData->GetMutablePoints() = SourcePointData->GetPointsCopy();
+		Context->OutputData.TaggedData.Emplace_GetRef().Data = OutPointData;
+
+		bool bTargetIsSource = (TargetPointData == SourcePointData);
+		if (bTargetIsSource)
+		{
+			TargetPointData = OutPointData;
+		}
+
+		OutState = PCGFuseElement::FFuseState
+		{
+			.TargetPointData = TargetPointData,
+			.SourcePointData = SourcePointData,
+			.OutPointData = OutPointData,
+
+			.IterationIndex = 0,
+			.bTargetIsSource = bTargetIsSource,
+			.RandomStream = FRandomStream(Context->GetSeed()),
+		};
+
+		OutState.CheckedIndices.SetNum(OutPointData->GetNumPoints(), false);
+		OutState.RemovedIndices.SetNum(OutPointData->GetNumPoints(), false);
+
+		// Will randomly select indices to compare for fusion. Pre-generate an array of random indices.
+		if (Settings->Algorithm == EPCGFuseAlgorithm::Random)
+		{
+			OutState.RandomIndices.SetNumUninitialized(OutPointData->GetNumPoints());
+			for (int i = 0; i < OutState.RandomIndices.Num(); ++i)
+			{
+				OutState.RandomIndices[i] = i;
+			}
+
+			PCGHelpers::ShuffleArray(OutState.RandomStream, OutState.RandomIndices);
+		}
+
+		// Prioritize selected fuse target by comparing attributes for the min or max value.
+		if (Settings->Priority == EPCGFusePriority::MinAttribute || Settings->Priority == EPCGFusePriority::MaxAttribute)
+		{
+			const FPCGAttributePropertyInputSelector SourcePrioritySelector = Settings->SourcePriorityAttribute.CopyAndFixLast(SourcePointData);
+			OutState.SourcePriorityAccessor = PCGAttributeAccessorHelpers::CreateConstAccessor(SourcePointData, SourcePrioritySelector);
+			OutState.SourcePriorityKeys = PCGAttributeAccessorHelpers::CreateConstKeys(SourcePointData, SourcePrioritySelector);
+			if (!OutState.SourcePriorityAccessor)
+			{
+				PCGLog::Metadata::LogFailToCreateAccessorError(SourcePrioritySelector, Context);
+				return EPCGTimeSliceInitResult::NoOperation;
+			}
+
+			const FPCGAttributePropertyInputSelector TargetPrioritySelector = Settings->TargetPriorityAttribute.CopyAndFixLast(TargetPointData);
+			OutState.TargetPriorityAccessor = PCGAttributeAccessorHelpers::CreateConstAccessor(TargetPointData, TargetPrioritySelector);
+			OutState.TargetPriorityKeys = PCGAttributeAccessorHelpers::CreateConstKeys(TargetPointData, TargetPrioritySelector);
+			if (!OutState.TargetPriorityAccessor)
+			{
+				PCGLog::Metadata::LogFailToCreateAccessorError(TargetPrioritySelector, Context);
+				return EPCGTimeSliceInitResult::NoOperation;
+			}
+
+			// Guarantee the two attributes will be comparable.
+			if (OutState.SourcePriorityAccessor->GetUnderlyingType() != OutState.TargetPriorityAccessor->GetUnderlyingType() &&
+				!PCG::Private::IsBroadcastableOrConstructible(OutState.SourcePriorityAccessor->GetUnderlyingType(), OutState.TargetPriorityAccessor->GetUnderlyingType()))
+			{
+				PCGLog::Metadata::LogIncomparableAttributesError(SourcePrioritySelector, TargetPrioritySelector, Context);
+				return EPCGTimeSliceInitResult::NoOperation;
+			}
+		}
+
+		// Give priority for the final fuse result to an arbitrary weight or from attribute(s).
+		if (Settings->Priority == EPCGFusePriority::Weighted || Settings->Priority == EPCGFusePriority::WeightedAverage)
+		{
+			ensure(OutState.WeightMax == 1.0);
+
+			if (Settings->bWeightAsAttribute)
+			{
+				// The fusion will happen on a normalized [0..1] linear interpolation between the two points by either a single weight attribute...
+				const FPCGAttributePropertyInputSelector SourceWeightSelector = Settings->SourceWeightAttribute.CopyAndFixLast(SourcePointData);
+				const TUniquePtr<const IPCGAttributeAccessor> SourceWeightAccessor = PCGAttributeAccessorHelpers::CreateConstAccessor(SourcePointData, SourceWeightSelector);
+				const TUniquePtr<const IPCGAttributeAccessorKeys> SourceWeightKeys = PCGAttributeAccessorHelpers::CreateConstKeys(SourcePointData, SourceWeightSelector);
+				if (!SourceWeightAccessor)
+				{
+					PCGLog::Metadata::LogFailToCreateAccessorError(SourceWeightSelector, Context);
+					return EPCGTimeSliceInitResult::NoOperation;
+				}
+
+				OutState.SourceWeights.SetNumUninitialized(SourcePointData->GetNumPoints());
+
+				if (!SourceWeightAccessor->GetRange(MakeArrayView(OutState.SourceWeights), 0, *SourceWeightKeys))
+				{
+					PCGLog::Metadata::LogFailToGetAttributeError(SourceWeightSelector, Context);
+					return EPCGTimeSliceInitResult::NoOperation;
+				}
+
+				OutState.WeightMax = *Algo::MaxElement(OutState.SourceWeights);
+			}
+
+			// Or in the case of a weighted average, it will normalize the average of the two attributes.
+			if (Settings->Priority == EPCGFusePriority::WeightedAverage)
+			{
+				const FPCGAttributePropertyInputSelector TargetWeightSelector = Settings->TargetWeightAttribute.CopyAndFixLast(TargetPointData);
+				const TUniquePtr<const IPCGAttributeAccessor> TargetWeightAccessor = PCGAttributeAccessorHelpers::CreateConstAccessor(TargetPointData, TargetWeightSelector);
+				const TUniquePtr<const IPCGAttributeAccessorKeys> TargetWeightKeys = PCGAttributeAccessorHelpers::CreateConstKeys(TargetPointData, TargetWeightSelector);
+				if (!TargetWeightAccessor)
+				{
+					PCGLog::Metadata::LogFailToCreateAccessorError(TargetWeightSelector, Context);
+					return EPCGTimeSliceInitResult::NoOperation;
+				}
+
+				OutState.TargetWeights.SetNumUninitialized(TargetPointData->GetNumPoints());
+
+				if (!TargetWeightAccessor->GetRange(MakeArrayView(OutState.TargetWeights), 0, *TargetWeightKeys))
+				{
+					PCGLog::Metadata::LogFailToGetAttributeError(TargetWeightSelector, Context);
+					return EPCGTimeSliceInitResult::NoOperation;
+				}
+
+				// Update the max to include the target weight as well.
+				OutState.WeightMax = FMath::Max(OutState.WeightMax, *Algo::MaxElement(OutState.TargetWeights));
+			}
+
+			OutState.WeightMax = FMath::Max(OutState.WeightMax, UE_DOUBLE_SMALL_NUMBER);
+		}
+
+		// The output will be a fusion between an attribute on the source and target points.
+		if (Settings->bFuseAttribute)
+		{
+			const FPCGAttributePropertyInputSelector SourceFuseAttributeSelector = Settings->SourceFuseAttribute.CopyAndFixLast(SourcePointData);
+			OutState.SourceFuseAttributeAccessor = PCGAttributeAccessorHelpers::CreateConstAccessor(SourcePointData, SourceFuseAttributeSelector);
+			OutState.SourceFuseAttributeKeys = PCGAttributeAccessorHelpers::CreateConstKeys(SourcePointData, SourceFuseAttributeSelector);
+			if (!OutState.SourceFuseAttributeAccessor || !OutState.SourceFuseAttributeKeys)
+			{
+				PCGLog::Metadata::LogFailToCreateAccessorError(SourceFuseAttributeSelector, Context);
+				return EPCGTimeSliceInitResult::NoOperation;
+			}
+
+			const FPCGAttributePropertyInputSelector TargetFuseAttributeSelector = Settings->TargetFuseAttribute.CopyAndFixLast(TargetPointData);
+			OutState.TargetFuseAttributeAccessor = PCGAttributeAccessorHelpers::CreateConstAccessor(TargetPointData, TargetFuseAttributeSelector);
+			OutState.TargetFuseAttributeKeys = PCGAttributeAccessorHelpers::CreateConstKeys(TargetPointData, TargetFuseAttributeSelector);
+			if (!OutState.TargetFuseAttributeAccessor || !OutState.TargetFuseAttributeKeys)
+			{
+				PCGLog::Metadata::LogFailToCreateAccessorError(TargetFuseAttributeSelector, Context);
+				return EPCGTimeSliceInitResult::NoOperation;
+			}
+
+			// Guarantee the two attributes will be comparable, as they will be interpolated between the two.
+			if (OutState.SourceFuseAttributeAccessor->GetUnderlyingType() != OutState.TargetFuseAttributeAccessor->GetUnderlyingType() &&
+				!PCG::Private::IsBroadcastableOrConstructible(OutState.SourceFuseAttributeAccessor->GetUnderlyingType(), OutState.TargetFuseAttributeAccessor->GetUnderlyingType()))
+			{
+				PCGLog::Metadata::LogIncomparableAttributesError(SourceFuseAttributeSelector, TargetFuseAttributeSelector, Context);
+				return EPCGTimeSliceInitResult::NoOperation;
+			}
+
+			const FPCGAttributePropertyOutputSelector OutputFuseAttributeSelector = Settings->OutputFuseAttribute.CopyAndFixSource(&SourceFuseAttributeSelector);
+			OutState.OutputFuseAttributeKeys = PCGAttributeAccessorHelpers::CreateKeys(OutPointData, OutputFuseAttributeSelector);
+			if (OutputFuseAttributeSelector.GetSelection() == EPCGAttributePropertySelection::Attribute)
+			{
+				PCGMetadataAttribute::CallbackWithRightType(OutState.SourceFuseAttributeAccessor->GetUnderlyingType(), [&OutState, &OutputFuseAttributeSelector]<typename T>(T) -> void
+				{
+					UPCGMetadata* Metadata = OutState.OutPointData->Metadata;
+					FPCGMetadataAttribute<T>* OutputFuseAttribute = Metadata->FindOrCreateAttribute(OutputFuseAttributeSelector.GetAttributeName(), T{});
+					OutState.OutputFuseAttributeAccessor = PCGAttributeAccessorHelpers::CreateAccessor(OutputFuseAttribute, Metadata);
+				});
+			}
+			else
+			{
+				OutState.OutputFuseAttributeAccessor = PCGAttributeAccessorHelpers::CreateAccessor(OutState.OutPointData, OutputFuseAttributeSelector);
+			}
+
+			if (!OutState.OutputFuseAttributeAccessor || !OutState.OutputFuseAttributeKeys)
+			{
+				PCGLog::Metadata::LogFailToCreateAccessorError(OutputFuseAttributeSelector, Context);
+				return EPCGTimeSliceInitResult::NoOperation;
+			}
+		}
+
+		return EPCGTimeSliceInitResult::Success;
+	});
+
+	return true;
+}
+
+bool FPCGFuseElement::ExecuteInternal(FPCGContext* InContext) const
+{
+	TRACE_CPUPROFILER_EVENT_SCOPE(FPCGFuseElement::Execute);
+
+	ContextType* TimeSlicedContext = static_cast<ContextType*>(InContext);
+	check(TimeSlicedContext);
+
+	if (!TimeSlicedContext->DataIsPreparedForExecution())
+	{
+		return true;
+	}
+
+	if (TimeSlicedContext->GetExecutionStateResult() == EPCGTimeSliceInitResult::NoOperation)
+	{
+		TimeSlicedContext->OutputData = TimeSlicedContext->InputData;
+		return true;
+	}
+
+	const UPCGFuseSettings* Settings = TimeSlicedContext->GetInputSettings<UPCGFuseSettings>();
+	check(Settings);
+
+	return ExecuteSlice(TimeSlicedContext, [Settings](const ContextType* Context, const ExecStateType& ExecState, IterStateType& FuseData, const uint32 IterIndex)
+	{
+		if (Context->GetIterationStateResult(IterIndex) == EPCGTimeSliceInitResult::NoOperation)
+		{
+			return true;
+		}
+
+		if (!ExecState.FuseFunction(Context, FuseData, *Settings))
+		{
+			return false;
+		}
+
+		FuseData.OutPointData->GetMutablePoints().Shrink();
+
+		return true;
+	});
+}
+
+#undef LOCTEXT_NAMESPACE
diff --git a/Engine/Plugins/PCG/Source/PCG/Private/Utils/PCGLogErrors.cpp b/Engine/Plugins/PCG/Source/PCG/Private/Utils/PCGLogErrors.cpp
--- a/Engine/Plugins/PCG/Source/PCG/Private/Utils/PCGLogErrors.cpp
+++ b/Engine/Plugins/PCG/Source/PCG/Private/Utils/PCGLogErrors.cpp
@@ -8,11 +8,14 @@ namespace PCGLog::InputOutput
 {
 	namespace Format
 	{
+		const FText InvalidInputData = LOCTEXT("InvalidInputData", "Invalid input data.");
+
 		const FTextFormat TypedInputNotFound = LOCTEXT("TypedInputNotFound", "Data of type {0} not found on pin '{1}'.");
 		const FTextFormat FirstInputOnly = LOCTEXT("FirstInputOnly", "Multiple inputs found on single-input pin '{0}'. Only the first will be selected.");
-		const FText InvalidInputData = LOCTEXT("InvalidInputData", "Invalid input data.");
+		const FTextFormat InvalidCardinality = LOCTEXT("InvalidCardinality", "Invalid cardinality among pins '{0}' and '{1}'. They must match 1:1, N:1, or N:N.");
 	}
 
+	// Warnings
 	void LogTypedDataNotFoundWarning(EPCGDataType DataType, const FName PinLabel, const FPCGContext* InContext)
 	{
 		const UEnum* PCGDataTypeEnum = StaticEnum<EPCGDataType>();
@@ -25,10 +28,16 @@ namespace PCGLog::InputOutput
 		LogWarningOnGraph(FText::Format(Format::FirstInputOnly, FText::FromName(PinLabel)), InContext);
 	}
 
+	// Errors
 	void LogInvalidInputDataError(const FPCGContext* InContext)
 	{
 		LogErrorOnGraph(Format::InvalidInputData, InContext);
 	}
+
+	void LogInvalidCardinalityError(const FName SourcePinLabel, const FName TargetPinLabel, const FPCGContext* InContext)
+	{
+		LogErrorOnGraph(FText::Format(Format::InvalidCardinality, FText::FromName(SourcePinLabel), FText::FromName(TargetPinLabel)), InContext);
+	}
 }
 
 namespace PCGLog::Metadata
@@ -40,6 +49,7 @@ namespace PCGLog::Metadata
 		const FTextFormat GetAttributeFailure = LOCTEXT("GetAttributeFailure", "Couldn't retrieve attribute '{0}' value. Expected type: {1}, Actual Type: {2}.");
 		const FTextFormat GetTypedAttributeFailure = LOCTEXT("GetTypedAttributeFailure", "Couldn't retrieve attribute '{0}' value. Expected type: {1}, Actual Type: {2}.");
 		const FTextFormat GetTypedAttributeFailureNoAccessor = LOCTEXT("GetTypedAttributeFailureNoAccessor", "Couldn't retrieve attribute '{0}' value of type: '{1}.");
+		const FTextFormat IncomparableTypesFailure = LOCTEXT("IncomparableTypesFailure", "Attributes '{0}' and '{1}' are incomparable. Ensure they are either of the same or compatible types.");
 	}
 
 	void LogFailToCreateAccessorError(const FPCGAttributePropertySelector& Selector, const FPCGContext* InContext)
@@ -61,6 +71,11 @@ namespace PCGLog::Metadata
 	{
 		LogErrorOnGraph(FText::Format(Format::GetAttributeFailure, Selector.GetDisplayText()), InContext);
 	}
+
+	void LogIncomparableAttributesError(const FPCGAttributePropertySelector& FirstSelector, const FPCGAttributePropertySelector& SecondSelector, const FPCGContext* InContext)
+	{
+		LogErrorOnGraph(FText::Format(Format::IncomparableTypesFailure, FirstSelector.GetDisplayText(), SecondSelector.GetDisplayText()), InContext);
+	}
 }
 
 #undef LOCTEXT_NAMESPACE
diff --git a/Engine/Plugins/PCG/Source/PCG/Public/Elements/PCGFuseElement.h b/Engine/Plugins/PCG/Source/PCG/Public/Elements/PCGFuseElement.h
new file mode 100644
--- /dev/null
+++ b/Engine/Plugins/PCG/Source/PCG/Public/Elements/PCGFuseElement.h
@@ -0,0 +1,175 @@
+// Copyright Epic Games, Inc. All Rights Reserved.
+
+#pragma once
+
+#include "PCGSettings.h"
+
+#include "Elements/PCGTimeSlicedElementBase.h"
+#include "Metadata/PCGMetadataAttributeTpl.h"
+#include "Metadata/Accessors/PCGPropertyAccessor.h"
+
+#include "PCGFuseElement.generated.h"
+
+class UPCGPointData;
+
+UENUM(BlueprintType)
+enum class EPCGFuseAlgorithm : uint8
+{
+	Sequential UMETA(Tooltip = "Evaluate the points sequentially in index order."),
+	Random UMETA(Tooltip = "Evaluate the points in a random order.")
+};
+
+UENUM(BlueprintType)
+enum class EPCGFusePriority : uint8
+{
+	Target UMETA(Tooltip = "Fuse source to target directly."),
+	MinAttribute UMETA(Tooltip = "Fuse source or target, depending on which has the lower value for a selected attribute."),
+	MaxAttribute UMETA(Tooltip = "Fuse source or target, depending on which has the higher value for a selected attribute."),
+	Weighted UMETA(Tooltip = "The source point will be fused between target and source based on a directly selected weight or weight attribute."),
+	WeightedAverage UMETA(Tooltip = "The source point will be fused between target and source based on a normalized and weighted attribute on each of the two."),
+};
+
+/** Fuse transforms or attributes of two 'co-located' points--based on a minimum distance and other criteria in world space. */
+UCLASS(MinimalAPI, BlueprintType, ClassGroup = (Procedural))
+class UPCGFuseSettings : public UPCGSettings
+{
+	GENERATED_BODY()
+
+public:
+	//~Begin UPCGSettings interface
+#if WITH_EDITOR
+	virtual FName GetDefaultNodeName() const override { return FName(TEXT("FuseElement")); }
+	virtual FText GetDefaultNodeTitle() const override { return NSLOCTEXT("PCGFuseElement", "NodeTitle", "Fuse"); }
+	virtual EPCGSettingsType GetType() const override { return EPCGSettingsType::PointOps; }
+	virtual FText GetNodeTooltipText() const override;
+#endif
+	virtual bool UseSeed() const override;
+
+protected:
+	virtual TArray<FPCGPinProperties> InputPinProperties() const override;
+	virtual TArray<FPCGPinProperties> OutputPinProperties() const override { return Super::DefaultPointOutputPinProperties(); }
+	virtual FPCGElementPtr CreateElement() const override;
+	//~End UPCGSettings interface
+
+public:
+	/** Will be used to determine which points to fuse. */
+	UPROPERTY(EditAnywhere, BlueprintReadWrite, Category = Settings, meta = (ClampMin = "0.01", PCG_Overridable))
+	double Distance = 100.0;
+
+	/** Dictates the process and order of points as they are considered for fusion. They can differ in their results and efficiency. */
+	UPROPERTY(EditAnywhere, BlueprintReadWrite, Category = Settings, meta = (PCG_Overridable))
+	EPCGFuseAlgorithm Algorithm = EPCGFuseAlgorithm::Sequential;
+
+	/** Further control of the fusion allows for interpolation by giving 'priority' between the source and target point. Includes attribute-wise comparison or weighting. */
+	UPROPERTY(EditAnywhere, BlueprintReadWrite, Category = Settings, meta = (PCG_Overridable))
+	EPCGFusePriority Priority = EPCGFusePriority::Target;
+
+	/** The source attribute used to determine fusion priority. I.e. The fusion can choose between the point with the minimum or maximum value in compared attributes. */
+	UPROPERTY(EditAnywhere, BlueprintReadWrite, Category = Settings, meta = (EditCondition = "Priority == EPCGFusePriority::MinAttribute || Priority == EPCGFusePriority::MaxAttribute", EditConditionHides, PCG_Overridable))
+	FPCGAttributePropertyInputSelector SourcePriorityAttribute;
+
+	/** The target attribute used to determine fusion priority. I.e. The fusion can choose between the point with the minimum or maximum value in compared attributes. */
+	UPROPERTY(EditAnywhere, BlueprintReadWrite, Category = Settings, meta = (EditCondition = "Priority == EPCGFusePriority::MinAttribute || Priority == EPCGFusePriority::MaxAttribute", EditConditionHides, PCG_Overridable))
+	FPCGAttributePropertyInputSelector TargetPriorityAttribute;
+
+	/** Use an attribute on the point directly as the weight. */
+	UPROPERTY(EditAnywhere, BlueprintReadWrite, Category = Settings, meta = (EditCondition = "Priority == EPCGFusePriority::Weighted || Priority == EPCGFusePriority::WeightedAverage", EditConditionHides, PCG_Overridable))
+	bool bWeightAsAttribute = false;
+
+	/** The normalized [0..1] weight when determining the interpolated fusion result between the two points--0.0 equates to the source point and 1.0 equates to the target point. */
+	UPROPERTY(EditAnywhere, BlueprintReadWrite, Category = Settings, meta = (ClampMin = "0.0", ClampMax = "1.0", EditCondition = "(Priority == EPCGFusePriority::Weighted || Priority == EPCGFusePriority::WeightedAverage) && !bWeightAsAttribute", EditConditionHides, PCG_Overridable))
+	double Weight = 0.5;
+
+	/** This attribute will determine the weight of the fusion result for the source point. It will be normalized to the range of [0..1]. */
+	UPROPERTY(EditAnywhere, BlueprintReadWrite, Category = Settings, meta = (EditCondition = "(Priority == EPCGFusePriority::Weighted || Priority == EPCGFusePriority::WeightedAverage) && bWeightAsAttribute", EditConditionHides, PCG_Overridable))
+	FPCGAttributePropertyInputSelector SourceWeightAttribute;
+
+	/** This attribute will determine the weight of the fusion result for the target point. It will be normalized to the range of [0..1]. */
+	UPROPERTY(EditAnywhere, BlueprintReadWrite, Category = Settings, meta = (EditCondition = "Priority == EPCGFusePriority::WeightedAverage", EditConditionHides, PCG_Overridable))
+	FPCGAttributePropertyInputSelector TargetWeightAttribute;
+
+	/** Will remove the points to be fused, rather than fusing them. */
+	UPROPERTY(EditAnywhere, BlueprintReadWrite, Category = Settings, meta = (PCG_Overridable))
+	bool bRemoveFusedPoints = true;
+
+	/** Fuse the two points' location. */
+	UPROPERTY(EditAnywhere, BlueprintReadWrite, Category = Settings, meta = (EditCondition = "!bRemoveFusedPoints", EditConditionHides, PCG_Overridable))
+	bool bFuseLocation = true;
+
+	/** Fuse the two points' rotation. */
+	UPROPERTY(EditAnywhere, BlueprintReadWrite, Category = Settings, meta = (EditCondition = "!bRemoveFusedPoints", EditConditionHides, PCG_Overridable))
+	bool bFuseRotation = true;
+
+	/** Fuse the two points' scale. */
+	UPROPERTY(EditAnywhere, BlueprintReadWrite, Category = Settings, meta = (EditCondition = "!bRemoveFusedPoints", EditConditionHides, PCG_Overridable))
+	bool bFuseScale = true;
+
+	/** Fuse a pair of specified attributes on both the target and source points. */
+	UPROPERTY(EditAnywhere, BlueprintReadWrite, Category = Settings, meta = (EditCondition = "!bRemoveFusedPoints", EditConditionHides, PCG_Overridable))
+	bool bFuseAttribute = false;
+
+	/** The source attribute to be compared during the fusion. */
+	UPROPERTY(EditAnywhere, BlueprintReadWrite, Category = Settings, meta = (EditCondition = "!bRemoveFusedPoints && bFuseAttribute", EditConditionHides, PCG_Overridable))
+	FPCGAttributePropertyInputSelector SourceFuseAttribute;
+
+	/** The target attribute to be compared during the fusion. */
+	UPROPERTY(EditAnywhere, BlueprintReadWrite, Category = Settings, meta = (EditCondition = "!bRemoveFusedPoints && bFuseAttribute", EditConditionHides, PCG_Overridable))
+	FPCGAttributePropertyInputSelector TargetFuseAttribute;
+
+	/** The output attribute to store the interpolated result of the fusion. */
+	UPROPERTY(EditAnywhere, BlueprintReadWrite, Category = Settings, meta = (EditCondition = "!bRemoveFusedPoints && bFuseAttribute", EditConditionHides, PCG_Overridable))
+	FPCGAttributePropertyOutputSelector OutputFuseAttribute;
+};
+
+namespace PCGFuseElement
+{
+	struct FFuseState
+	{
+		const UPCGPointData* TargetPointData = nullptr;
+		const UPCGPointData* SourcePointData = nullptr;
+		UPCGPointData* OutPointData = nullptr;
+
+		int32 IterationIndex = INDEX_NONE;
+		bool bTargetIsSource = false;
+
+		TBitArray<TInlineAllocator<512>> CheckedIndices;
+		TBitArray<TInlineAllocator<512>> RemovedIndices;
+
+		TUniquePtr<const IPCGAttributeAccessor> SourcePriorityAccessor;
+		TUniquePtr<const IPCGAttributeAccessorKeys> SourcePriorityKeys;
+		TUniquePtr<const IPCGAttributeAccessor> TargetPriorityAccessor;
+		TUniquePtr<const IPCGAttributeAccessorKeys> TargetPriorityKeys;
+
+		TUniquePtr<const IPCGAttributeAccessor> SourceFuseAttributeAccessor;
+		TUniquePtr<const IPCGAttributeAccessorKeys> SourceFuseAttributeKeys;
+		TUniquePtr<const IPCGAttributeAccessor> TargetFuseAttributeAccessor;
+		TUniquePtr<const IPCGAttributeAccessorKeys> TargetFuseAttributeKeys;
+
+		TUniquePtr<IPCGAttributeAccessor> OutputFuseAttributeAccessor;
+		TUniquePtr<IPCGAttributeAccessorKeys> OutputFuseAttributeKeys;
+
+		TArray<double> SourceWeights;
+		TArray<double> TargetWeights;
+
+		// For weighted and normalized weighted selections. This stores the maximum weight of the point data to be used in normalization.
+		double WeightMax = 1.0;
+
+		FRandomStream RandomStream;
+		TArray<int32> RandomIndices;
+		// For an iterative approach to fusing.
+		bool bPointWasFusedThisCycle = false;
+	};
+
+	struct FExecState
+	{
+		using FuseSignature = bool(*)(const FPCGContext* InContext, FFuseState& FuseData, const UPCGFuseSettings& FuseSettings);
+		FuseSignature FuseFunction = nullptr;
+	};
+}
+
+class FPCGFuseElement : public TPCGTimeSlicedElementBase<PCGFuseElement::FExecState, PCGFuseElement::FFuseState>
+{
+protected:
+	virtual bool PrepareDataInternal(FPCGContext* InContext) const override;
+	virtual bool ExecuteInternal(FPCGContext* InContext) const override;
+};
diff --git a/Engine/Plugins/PCG/Source/PCG/Public/Elements/PCGTimeSlicedElementBase.h b/Engine/Plugins/PCG/Source/PCG/Public/Elements/PCGTimeSlicedElementBase.h
--- a/Engine/Plugins/PCG/Source/PCG/Public/Elements/PCGTimeSlicedElementBase.h
+++ b/Engine/Plugins/PCG/Source/PCG/Public/Elements/PCGTimeSlicedElementBase.h
@@ -66,7 +66,7 @@ struct TPCGTimeSlicedContext : public FPCGContext
 	EPCGTimeSliceInitResult GetExecutionStateResult() const { return ExecutionStateResult; }
 
 	/** Will return the result of the attempt to initialize a specific iteration, by Index */
-	EPCGTimeSliceInitResult GetIterationStateResult(const int32 Index)
+	EPCGTimeSliceInitResult GetIterationStateResult(const int32 Index) const
 	{
 		check(Index >= 0 && Index < PerIterationStateResultArray.Num());
 
diff --git a/Engine/Plugins/PCG/Source/PCG/Public/Utils/PCGLogErrors.h b/Engine/Plugins/PCG/Source/PCG/Public/Utils/PCGLogErrors.h
--- a/Engine/Plugins/PCG/Source/PCG/Public/Utils/PCGLogErrors.h
+++ b/Engine/Plugins/PCG/Source/PCG/Public/Utils/PCGLogErrors.h
@@ -32,6 +32,7 @@ namespace PCGLog
 
 		// Errors
 		void LogInvalidInputDataError(const FPCGContext* InContext = nullptr);
+		void LogInvalidCardinalityError(const FName SourcePinLabel, const FName TargetPinLabel, const FPCGContext* InContext = nullptr);
 	}
 
 	namespace Metadata
@@ -61,6 +62,7 @@ namespace PCGLog
 		void LogFailToGetAttributeError(FText AttributeName, const FPCGContext* InContext = nullptr);
 		void LogFailToGetAttributeError(FName AttributeName, const FPCGContext* InContext = nullptr);
 		void LogFailToGetAttributeError(const FPCGAttributePropertySelector& Selector, const FPCGContext* InContext = nullptr);
+		void LogIncomparableAttributesError(const FPCGAttributePropertySelector& FirstSelector, const FPCGAttributePropertySelector& SecondSelector, const FPCGContext* InContext = nullptr);
 
 		template <typename T>
 		void LogFailToGetAttributeError(FText AttributeName, const IPCGAttributeAccessor* Accessor, const FPCGContext* InContext = nullptr)
