[PCG] Add node to extract elements (point or param) from index provided by an attribute on points or attribute sets

diff --git a/Engine/Plugins/PCG/Source/PCG/Private/Elements/PCGFilterByIndex.cpp b/Engine/Plugins/PCG/Source/PCG/Private/Elements/PCGFilterByIndex.cpp
--- a/Engine/Plugins/PCG/Source/PCG/Private/Elements/PCGFilterByIndex.cpp
+++ b/Engine/Plugins/PCG/Source/PCG/Private/Elements/PCGFilterByIndex.cpp
@@ -36,8 +36,7 @@ FText UPCGFilterByIndexSettings::GetNodeTooltipText() const
 
 FString UPCGFilterByIndexSettings::GetAdditionalTitleInformation() const
 {
-	FString NodeName = TEXT("Index: ");
-
+	FString NodeName = TEXT("Indices: ");
 	NodeName += SelectedIndices.Len() <= PCGFilterByIndexConstants::IndexExpressionTruncation
 		? SelectedIndices
 		: SelectedIndices.Left(PCGFilterByIndexConstants::IndexExpressionTruncation - 3) + TEXT("...");
@@ -70,16 +69,16 @@ bool FPCGFilterByIndexElement::ExecuteInternal(FPCGContext* Context) const
 	{
 		// Error cases are caught for early out.
 		case PCGParser::EPCGParserResult::InvalidCharacter:
-			PCGE_LOG_C(Error, GraphAndLog, Context, LOCTEXT("ErrorInvalidCharacter", "Invalid character in index selection string."));
+			PCGLog::Parsing::LogInvalidCharacterInParsedStringError(Settings->SelectedIndices, Context);
 			return true;
 
 		case PCGParser::EPCGParserResult::InvalidExpression:
-			PCGE_LOG_C(Error, GraphAndLog, Context, LOCTEXT("ErrorInvalidExpression", "Invalid expression in index selection string."));
+			PCGLog::Parsing::LogInvalidExpressionInParsedStringError(Settings->SelectedIndices, Context);
 			return true;
 
 		// If the expression is empty, treat it as though there are simply no selected indices. But, also log.
 		case PCGParser::EPCGParserResult::EmptyExpression:
-			PCGE_LOG_C(Log, LogOnly, Context, LOCTEXT("WarningEmptyExpression", "Empty expression in index selection string."));
+			PCGLog::Parsing::LogEmptyExpressionWarning(Context);
 			break;
 
 		case PCGParser::EPCGParserResult::Success:
@@ -112,4 +111,4 @@ bool FPCGFilterByIndexElement::ExecuteInternal(FPCGContext* Context) const
 	return true;
 }
 
-#undef LOCTEXT_NAMESPACE
\ No newline at end of file
+#undef LOCTEXT_NAMESPACE
diff --git a/Engine/Plugins/PCG/Source/PCG/Private/Elements/PCGFilterElementsByIndex.cpp b/Engine/Plugins/PCG/Source/PCG/Private/Elements/PCGFilterElementsByIndex.cpp
new file mode 100644
--- /dev/null
+++ b/Engine/Plugins/PCG/Source/PCG/Private/Elements/PCGFilterElementsByIndex.cpp
@@ -0,0 +1,376 @@
+// Copyright Epic Games, Inc. All Rights Reserved.
+
+#include "Elements/PCGFilterElementsByIndex.h"
+
+#include "PCGContext.h"
+#include "PCGParamData.h"
+#include "Data/PCGPointData.h"
+#include "Helpers/Parsing/PCGParsing.h"
+#include "Metadata/Accessors/PCGAttributeAccessorHelpers.h"
+
+#define LOCTEXT_NAMESPACE "PCGFilterElementsByIndexElement"
+
+namespace PCGFilterElementsByIndex
+{
+	namespace Constants
+	{
+		const FName FilterInputsPinLabel = TEXT("Filter");
+		const FName InFilterOutputPinLabel = TEXT("In Filter");
+		const FName OutFilterOutputPinLabel = TEXT("Out Filter");
+
+		// After so many characters on the node, it will truncate with an ellipsis ...
+		static constexpr int32 IndexExpressionTruncation = 12;
+	}
+
+	namespace Helpers
+	{
+		int32 GetDataElementCount(const UPCGData* Data)
+		{
+			if (const UPCGPointData* PointData = Cast<UPCGPointData>(Data))
+			{
+				return PointData->GetNumPoints();
+			}
+			else if (const UPCGParamData* ParamData = Cast<UPCGParamData>(Data))
+			{
+				check(ParamData->Metadata);
+				return ParamData->Metadata->GetItemCountForChild();
+			}
+			else
+			{
+				return 0;
+			}
+		}
+
+		PCGIndexing::FPCGIndexCollection StringToIndexCollection(const int32 ElementCount, const FStringView& Expression, const FPCGContext* InContext)
+		{
+			check(InContext);
+
+			PCGIndexing::FPCGIndexCollection FilteredIndices(ElementCount);
+
+			// Parse the indices and switch through possible issues.
+			switch (PCGParser::ParseIndexRanges(FilteredIndices, Expression))
+			{
+				// Error cases are caught for early out.
+				case PCGParser::EPCGParserResult::InvalidCharacter:
+					PCGLog::Parsing::LogInvalidCharacterInParsedStringError(Expression, InContext);
+					return PCGIndexing::FPCGIndexCollection::Invalid();
+
+				case PCGParser::EPCGParserResult::InvalidExpression:
+					PCGLog::Parsing::LogInvalidExpressionInParsedStringError(Expression, InContext);
+					return PCGIndexing::FPCGIndexCollection::Invalid();
+
+				// If the expression is empty, treat it as though there are simply no selected indices. But, also log.
+				case PCGParser::EPCGParserResult::EmptyExpression:
+					PCGLog::Parsing::LogEmptyExpressionWarning(InContext);
+					break;
+
+				case PCGParser::EPCGParserResult::Success:
+					break;
+
+				default:
+					checkNoEntry();
+					return PCGIndexing::FPCGIndexCollection::Invalid();
+			}
+
+			return FilteredIndices;
+		}
+
+		void PassThroughInputData(FPCGContext* InContext)
+		{
+			const TArray<FPCGTaggedData> PrimaryInputs = InContext->InputData.GetInputsByPin(PCGPinConstants::DefaultInputLabel);
+			InContext->OutputData.TaggedData.Reserve(PrimaryInputs.Num());
+
+			for (const FPCGTaggedData& Input : PrimaryInputs)
+			{
+				InContext->OutputData.TaggedData.Emplace_GetRef(Input).Pin = Constants::InFilterOutputPinLabel;
+			}
+		}
+	} // namespace Helpers
+}
+
+FPCGElementPtr UPCGFilterElementsByIndexSettings::CreateElement() const
+{
+	return MakeShared<FPCGFilterElementsByIndexElement>();
+}
+
+FString UPCGFilterElementsByIndexSettings::GetAdditionalTitleInformation() const
+{
+	if (bSelectIndicesByInput)
+	{
+		return Super::GetAdditionalTitleInformation();
+	}
+
+	FString NodeName = TEXT("Indices: ");
+	NodeName += SelectedIndices.Len() <= PCGFilterElementsByIndex::Constants::IndexExpressionTruncation
+					? SelectedIndices
+					: SelectedIndices.Left(PCGFilterElementsByIndex::Constants::IndexExpressionTruncation - 3) + TEXT("...");
+
+	return NodeName;
+}
+
+TArray<FPCGPinProperties> UPCGFilterElementsByIndexSettings::InputPinProperties() const
+{
+	TArray<FPCGPinProperties> PinProperties;
+	FPCGPinProperties& PrimaryPinProperties = PinProperties.Emplace_GetRef(PCGPinConstants::DefaultInputLabel, EPCGDataType::PointOrParam);
+	PrimaryPinProperties.SetRequiredPin();
+#if WITH_EDITOR
+	PrimaryPinProperties.Tooltip = LOCTEXT("PrimaryInputPinTooltip", "To be filtered by selected indices from a secondary input or a selection string.");
+#endif // WITH_EDITOR
+
+	if (bSelectIndicesByInput)
+	{
+		FPCGPinProperties& FilterPinProperties = PinProperties.Emplace_GetRef(PCGFilterElementsByIndex::Constants::FilterInputsPinLabel, EPCGDataType::PointOrParam);
+#if WITH_EDITOR
+		FilterPinProperties.Tooltip = LOCTEXT("FilterInputPinTooltip", "An attribute on elements from this input will be used to determine which elements from the first input will be filtered.");
+#endif // WITH_EDITOR
+	}
+
+	return PinProperties;
+}
+
+TArray<FPCGPinProperties> UPCGFilterElementsByIndexSettings::OutputPinProperties() const
+{
+	TArray<FPCGPinProperties> PinProperties;
+	FPCGPinProperties& InFilterPinProperties = PinProperties.Emplace_GetRef(PCGFilterElementsByIndex::Constants::InFilterOutputPinLabel, EPCGDataType::PointOrParam);
+#if WITH_EDITOR
+	InFilterPinProperties.Tooltip = LOCTEXT("InFilterOutputPinTooltip", "The elements that remain included after the filtering process.");
+#endif // WITH_EDITOR
+
+	if (bOutputDiscardedElements)
+	{
+		FPCGPinProperties& OutFilterPinProperties = PinProperties.Emplace_GetRef(PCGFilterElementsByIndex::Constants::OutFilterOutputPinLabel, EPCGDataType::PointOrParam);
+#if WITH_EDITOR
+		OutFilterPinProperties.Tooltip = LOCTEXT("OutFilterOutputPinTooltip", "The elements that were excluded during the filtering process.");
+#endif // WITH_EDITOR
+	}
+
+	return PinProperties;
+}
+
+bool FPCGFilterElementsByIndexElement::ExecuteInternal(FPCGContext* InContext) const
+{
+	TRACE_CPUPROFILER_EVENT_SCOPE(FPCGFilterElementsByIndexElement::Execute);
+	using namespace PCGFilterElementsByIndex;
+
+	check(InContext);
+
+	const TArray<FPCGTaggedData> PrimaryInputs = InContext->InputData.GetInputsByPin(PCGPinConstants::DefaultInputLabel);
+	if (PrimaryInputs.IsEmpty())
+	{
+		return true;
+	}
+
+	const UPCGFilterElementsByIndexSettings* Settings = InContext->GetInputSettings<UPCGFilterElementsByIndexSettings>();
+	check(Settings);
+
+	PCGIndexing::FPCGIndexCollection AllFilteredIndices = PCGIndexing::FPCGIndexCollection::Invalid();
+
+	// Getting the filtered indices from the secondary input pin.
+	if (Settings->bSelectIndicesByInput)
+	{
+		TArray<FPCGTaggedData> FilterInputs = InContext->InputData.GetInputsByPin(Constants::FilterInputsPinLabel);
+
+		// Forward inputs if no filtering data.
+		if (FilterInputs.IsEmpty())
+		{
+			Helpers::PassThroughInputData(InContext);
+			return true;
+		}
+
+		// ... or if not 1:N, N:1, or N:N.
+		if (PrimaryInputs.Num() != FilterInputs.Num() && (PrimaryInputs.Num() != 1 && FilterInputs.Num() != 1))
+		{
+			PCGLog::InputOutput::LogInvalidCardinalityError(PCGPinConstants::DefaultInputLabel, Constants::FilterInputsPinLabel, InContext);
+			Helpers::PassThroughInputData(InContext);
+			return true;
+		}
+
+		// Collect all the filter information first.
+		for (const FPCGTaggedData& FilterInput : FilterInputs)
+		{
+			const UPCGData* FilterInputData = FilterInput.Data;
+			const UPCGPointData* FilterPointData = Cast<UPCGPointData>(FilterInputData);
+			const UPCGParamData* FilterParamData = Cast<UPCGParamData>(FilterInputData);
+			if (!(FilterPointData || FilterParamData))
+			{
+				PCGLog::InputOutput::LogTypedDataNotFoundWarning(EPCGDataType::PointOrParam, Constants::FilterInputsPinLabel, InContext);
+				continue;
+			}
+
+			FPCGAttributePropertyInputSelector IndexSelector = Settings->IndexSelectionAttribute.CopyAndFixLast(FilterInputData);
+			const TUniquePtr<const IPCGAttributeAccessor> Accessor = PCGAttributeAccessorHelpers::CreateConstAccessor(FilterInputData, IndexSelector);
+			const TUniquePtr<const IPCGAttributeAccessorKeys> Keys = PCGAttributeAccessorHelpers::CreateConstKeys(FilterInputData, IndexSelector);
+
+			if (!Accessor || !Keys)
+			{
+				PCGLog::Metadata::LogFailToCreateAccessorError(IndexSelector, InContext);
+				continue;
+			}
+
+			if (Keys->GetNum() < 1)
+			{
+				continue;
+			}
+
+			// Get the indices from the secondary input.
+			TArray<int32> IndexArray;
+			IndexArray.SetNumUninitialized(Keys->GetNum());
+			if (!Accessor->GetRange<int32>(IndexArray, 0, *Keys, EPCGAttributeAccessorFlags::AllowBroadcastAndConstructible))
+			{
+				PCGLog::Metadata::LogFailToGetAttributeError(IndexSelector, InContext);
+				continue;
+			}
+			check(!IndexArray.IsEmpty());
+
+			IndexArray.Sort();
+			// Convert the highest index found to a local array size.
+			PCGIndexing::FPCGIndexCollection FilteredIndices(IndexArray.Last() + 1);
+			check(FilteredIndices.IsValid());
+
+			// Now sorted, add by range rather than adding every single index as a range.
+			int32 RangeStartIndex = IndexArray[0];
+			int32 RangeEndIndex = IndexArray[0];
+			for (int i = 1; i < IndexArray.Num(); ++i)
+			{
+				if (IndexArray[i] != RangeEndIndex + 1)
+				{
+					FilteredIndices.AddRange(RangeStartIndex, RangeEndIndex);
+					RangeStartIndex = IndexArray[i];
+					RangeEndIndex = RangeStartIndex;
+				}
+				else
+				{
+					RangeEndIndex = IndexArray[i];
+				}
+			}
+
+			// Add the final range.
+			FilteredIndices.AddRange(RangeStartIndex, RangeEndIndex);
+
+			// Update the base list.
+			AllFilteredIndices += FilteredIndices;
+		}
+
+		// No filtered indices. Pass-through.
+		if (AllFilteredIndices.GetTotalIndexCount() < 1)
+		{
+			Helpers::PassThroughInputData(InContext);
+			return true;
+		}
+	}
+
+	TArray<FPCGTaggedData>& Outputs = InContext->OutputData.TaggedData;
+
+	for (const FPCGTaggedData& PrimaryInput : PrimaryInputs)
+	{
+		const UPCGData* PrimaryInputData = PrimaryInput.Data;
+
+		// Get the filtered indices directly from the user string.
+		if (!Settings->bSelectIndicesByInput)
+		{
+			const int32 NumElements = Helpers::GetDataElementCount(PrimaryInputData);
+			AllFilteredIndices = Helpers::StringToIndexCollection(NumElements, Settings->SelectedIndices, InContext);
+		}
+
+		auto AddNewOutput = [&Outputs, &Tags = PrimaryInput.Tags](UPCGData* Data, FName Pin)
+		{
+			Outputs.Emplace(FPCGTaggedData
+			{
+				.Data = Data,
+				.Tags = Tags,
+				.Pin = Pin
+			});
+		};
+
+		if (const UPCGPointData* InputPointData = Cast<UPCGPointData>(PrimaryInputData))
+		{
+			TArray<FPCGPoint, TInlineAllocator<256>> InFilterPoints;
+			TArray<FPCGPoint, TInlineAllocator<256>> OutFilterPoints;
+			InFilterPoints.Reserve(InputPointData->GetNumPoints());
+			if (Settings->bOutputDiscardedElements)
+			{
+				OutFilterPoints.Reserve(InputPointData->GetNumPoints());
+			}
+
+			const TArray<FPCGPoint>& InputPoints = InputPointData->GetPoints();
+			for (int PointIndex = 0; PointIndex < InputPoints.Num(); ++PointIndex)
+			{
+				bool bIndexIsFiltered = AllFilteredIndices.ContainsIndex(PointIndex);
+				bIndexIsFiltered = Settings->bInvertFilter ? !bIndexIsFiltered : bIndexIsFiltered;
+				if (bIndexIsFiltered)
+				{
+					InFilterPoints.Emplace_GetRef(InputPoints[PointIndex]);
+				}
+				else if (Settings->bOutputDiscardedElements)
+				{
+					OutFilterPoints.Emplace_GetRef(InputPoints[PointIndex]);
+				}
+			}
+
+			UPCGPointData* InFilterPointData = FPCGContext::NewObject_AnyThread<UPCGPointData>(InContext);
+			InFilterPointData->InitializeFromData(InputPointData);
+			InFilterPointData->GetMutablePoints() = std::move(InFilterPoints);
+			AddNewOutput(InFilterPointData, Constants::InFilterOutputPinLabel);
+
+			if (Settings->bOutputDiscardedElements)
+			{
+				UPCGPointData* OutFilterPointData = FPCGContext::NewObject_AnyThread<UPCGPointData>(InContext);
+				OutFilterPointData->GetMutablePoints() = std::move(OutFilterPoints);
+				OutFilterPoints.Reserve(InputPointData->GetNumPoints());
+				AddNewOutput(OutFilterPointData, Constants::OutFilterOutputPinLabel);
+			}
+		}
+		else if (const UPCGParamData* InputParamData = Cast<UPCGParamData>(PrimaryInputData))
+		{
+			const UPCGMetadata* InputMetadata = InputParamData->Metadata;
+			check(InputMetadata);
+
+			const int32 NumEntries = InputMetadata->GetItemCountForChild();
+			TArray<PCGMetadataEntryKey, TInlineAllocator<256>> InEntryKeys;
+			TArray<PCGMetadataEntryKey, TInlineAllocator<256>> OutEntryKeys;
+			InEntryKeys.Reserve(NumEntries);
+			if (Settings->bOutputDiscardedElements)
+			{
+				OutEntryKeys.Reserve(NumEntries);
+			}
+
+			for (PCGMetadataEntryKey i = 0; i < NumEntries; ++i)
+			{
+				bool bIndexIsFiltered = AllFilteredIndices.ContainsIndex(i);
+				bIndexIsFiltered = Settings->bInvertFilter ? !bIndexIsFiltered : bIndexIsFiltered;
+				if (bIndexIsFiltered)
+				{
+					InEntryKeys.Emplace(i);
+				}
+				else if (Settings->bOutputDiscardedElements)
+				{
+					OutEntryKeys.Emplace(i);
+				}
+			}
+
+			UPCGParamData* InFilterParamData = FPCGContext::NewObject_AnyThread<UPCGParamData>(InContext);
+			AddNewOutput(InFilterParamData, Constants::InFilterOutputPinLabel);
+
+			UPCGMetadata* InFilterMetadata = InFilterParamData->Metadata;
+			check(InFilterMetadata);
+			InFilterMetadata->SetAttributes(InEntryKeys, InputMetadata);
+
+			if (Settings->bOutputDiscardedElements)
+			{
+				UPCGParamData* OutFilterParamData = FPCGContext::NewObject_AnyThread<UPCGParamData>(InContext);
+				UPCGMetadata* OutFilterMetadata = OutFilterParamData->Metadata;
+				check(OutFilterMetadata);
+				AddNewOutput(OutFilterParamData, Constants::OutFilterOutputPinLabel);
+				OutFilterMetadata->SetAttributes(OutEntryKeys, InputMetadata);
+			}
+		}
+		else
+		{
+			PCGLog::InputOutput::LogTypedDataNotFoundWarning(EPCGDataType::PointOrParam, PCGPinConstants::DefaultInputLabel, InContext);
+		}
+	}
+
+	return true;
+}
+
+#undef LOCTEXT_NAMESPACE
diff --git a/Engine/Plugins/PCG/Source/PCG/Private/Helpers/Parsing/PCGIndexing.cpp b/Engine/Plugins/PCG/Source/PCG/Private/Helpers/Parsing/PCGIndexing.cpp
--- a/Engine/Plugins/PCG/Source/PCG/Private/Helpers/Parsing/PCGIndexing.cpp
+++ b/Engine/Plugins/PCG/Source/PCG/Private/Helpers/Parsing/PCGIndexing.cpp
@@ -31,6 +31,11 @@ namespace PCGIndexing
 
 	bool FPCGIndexCollection::AddRange(const int32 StartIndex, const int32 EndIndex)
 	{
+		if (!ensure(IsValid()))
+		{
+			return false;
+		}
+
 		const FPCGIndexRange NewRange = AdjustIndicesAndCreateRange(StartIndex, EndIndex);
 
 		if (!RangeIsValid(NewRange))
@@ -81,7 +86,12 @@ namespace PCGIndexing
 
 	bool FPCGIndexCollection::RangeIsValid(const FPCGIndexRange& Range) const
 	{
-		return (Range.StartIndex >= 0 && Range.StartIndex <= ArraySize && Range.EndIndex >= 0 && Range.EndIndex <= ArraySize && Range.StartIndex < Range.EndIndex);
+		return IsValid()
+				&& Range.StartIndex >= 0
+				&& Range.StartIndex <= ArraySize
+				&& Range.EndIndex >= 0
+				&& Range.EndIndex <= ArraySize
+				&& Range.StartIndex < Range.EndIndex;
 	}
 
 	bool FPCGIndexCollection::ContainsIndex(const int32 Index) const
@@ -120,11 +130,32 @@ namespace PCGIndexing
 		return TotalIndexCount;
 	}
 
+	bool FPCGIndexCollection::IsValid() const
+	{
+		return ArraySize > 0;
+	}
+
+	bool FPCGIndexCollection::IsEmpty() const
+	{
+		return IndexRanges.IsEmpty();
+	}
+
 	bool FPCGIndexCollection::operator==(const FPCGIndexCollection& Other) const
 	{
 		return (ArraySize == Other.ArraySize) && (IndexRanges == Other.IndexRanges);
 	}
 
+	FPCGIndexCollection& FPCGIndexCollection::operator+=(const FPCGIndexCollection& Other)
+	{
+		ArraySize = FMath::Max(ArraySize, Other.ArraySize);
+		for (const FPCGIndexRange& Range : Other.IndexRanges)
+		{
+			AddRange(Range);
+		}
+
+		return *this;
+	}
+
 	FPCGIndexRange FPCGIndexCollection::AdjustIndicesAndCreateRange(int32 StartIndex, int32 EndIndex) const
 	{
 		// Convert both indices to a positive range
diff --git a/Engine/Plugins/PCG/Source/PCG/Private/Helpers/Parsing/PCGParsing.cpp b/Engine/Plugins/PCG/Source/PCG/Private/Helpers/Parsing/PCGParsing.cpp
--- a/Engine/Plugins/PCG/Source/PCG/Private/Helpers/Parsing/PCGParsing.cpp
+++ b/Engine/Plugins/PCG/Source/PCG/Private/Helpers/Parsing/PCGParsing.cpp
@@ -15,7 +15,7 @@ namespace PCGParser
 	}
 
 	// TODO: Consider passing error messages as well for more precise feedback
-	EPCGParserResult ParseIndexRanges(PCGIndexing::FPCGIndexCollection& OutIndexCollection, const FString& InputString)
+	EPCGParserResult ParseIndexRanges(PCGIndexing::FPCGIndexCollection& OutIndexCollection, const FStringView& InputString)
 	{
 		FString PurgedInputString(InputString);
 		PurgedInputString.RemoveSpacesInline();
diff --git a/Engine/Plugins/PCG/Source/PCG/Private/Metadata/PCGMetadata.cpp b/Engine/Plugins/PCG/Source/PCG/Private/Metadata/PCGMetadata.cpp
--- a/Engine/Plugins/PCG/Source/PCG/Private/Metadata/PCGMetadata.cpp
+++ b/Engine/Plugins/PCG/Source/PCG/Private/Metadata/PCGMetadata.cpp
@@ -1309,7 +1309,7 @@ void UPCGMetadata::SetPointAttributes(const TArrayView<const FPCGPoint>& InPoint
 	}
 }
 
-void UPCGMetadata::SetAttributes(const TArrayView<const PCGMetadataEntryKey>& InOriginalKeys, const UPCGMetadata* InMetadata, const TArrayView<PCGMetadataEntryKey>& OutOriginalKeys, FPCGContext* OptionalContext)
+void UPCGMetadata::SetAttributes(const TArrayView<const PCGMetadataEntryKey>& InOriginalKeys, const UPCGMetadata* InMetadata, const TArrayView<PCGMetadataEntryKey>* InOutOptionalKeys, FPCGContext* OptionalContext)
 {
 	if (!InMetadata || InMetadata->GetAttributeCount() == 0 || GetAttributeCount() == 0)
 	{
@@ -1318,13 +1318,14 @@ void UPCGMetadata::SetAttributes(const TArrayView<const PCGMetadataEntryKey>& In
 
 	TRACE_CPUPROFILER_EVENT_SCOPE(UPCGMetadata::SetAttributes);
 
-	check(InOriginalKeys.Num() == OutOriginalKeys.Num());
+	check(!InOutOptionalKeys || (InOriginalKeys.Num() == InOutOptionalKeys->Num()));
 
 	// There are a few things we can do to optimize here -
 	// basically, we don't need to set attributes more than once for a given <in, out> pair
 	TArray<PCGMetadataEntryKey, TInlineAllocator<256>> InKeys;
 	TArray<PCGMetadataEntryKey, TInlineAllocator<256>> OutKeys;
 
+	if (InOutOptionalKeys)
 	{
 		TRACE_CPUPROFILER_EVENT_SCOPE(UPCGMetadata::SetAttributes::CreateDeduplicatedKeys);
 		TMap<TPair<PCGMetadataEntryKey, PCGMetadataEntryKey>, int> PairMapping;
@@ -1332,7 +1333,7 @@ void UPCGMetadata::SetAttributes(const TArrayView<const PCGMetadataEntryKey>& In
 		for (int KeyIndex = 0; KeyIndex < InOriginalKeys.Num(); ++KeyIndex)
 		{
 			PCGMetadataEntryKey InKey = InOriginalKeys[KeyIndex];
-			PCGMetadataEntryKey& OutKey = OutOriginalKeys[KeyIndex];
+			PCGMetadataEntryKey& OutKey = (*InOutOptionalKeys)[KeyIndex];
 
 			if (int* MatchingPairIndex = PairMapping.Find(TPair<PCGMetadataEntryKey, PCGMetadataEntryKey>(InKey, OutKey)))
 			{
@@ -1348,6 +1349,10 @@ void UPCGMetadata::SetAttributes(const TArrayView<const PCGMetadataEntryKey>& In
 			}
 		}
 	}
+	else
+	{
+		OutKeys.Init(PCGInvalidEntryKey, InOriginalKeys.Num());
+	}
 
 	{
 		TRACE_CPUPROFILER_EVENT_SCOPE(UPCGMetadata::SetAttributes::InitializeOnSet);
@@ -1428,13 +1433,21 @@ void UPCGMetadata::SetAttributes(const TArrayView<const PCGMetadataEntryKey>& In
 	}
 	AttributeLock.ReadUnlock();
 
-	// Finally, copy back the actual out keys to the original out keys
-	for (PCGMetadataEntryKey& OutKey : OutOriginalKeys)
+	if (InOutOptionalKeys)
 	{
-		OutKey = OutKeys[OutKey];
+		// Finally, copy back the actual out keys to the original out keys
+		for (PCGMetadataEntryKey& OutKey : *InOutOptionalKeys)
+		{
+			OutKey = OutKeys[OutKey];
+		}
 	}
 }
 
+void UPCGMetadata::SetAttributes(const TArrayView<const PCGMetadataEntryKey>& InKeys, const UPCGMetadata* InMetadata, const TArrayView<PCGMetadataEntryKey>& OutKeys, FPCGContext* OptionalContext)
+{
+	SetAttributes(InKeys, InMetadata, &OutKeys, OptionalContext);
+}
+
 void UPCGMetadata::MergeAttributesByKey(int64 KeyA, const UPCGMetadata* MetadataA, int64 KeyB, const UPCGMetadata* MetadataB, int64 TargetKey, EPCGMetadataOp Op, int64& OutKey)
 {
 	OutKey = TargetKey;
@@ -1487,4 +1500,4 @@ void UPCGMetadata::SetLastCachedSelectorOnOwner(FName AttributeName)
 		Selector.SetAttributeName(AttributeName);
 		OwnerData->SetLastSelector(Selector);
 	}
-}
\ No newline at end of file
+}
diff --git a/Engine/Plugins/PCG/Source/PCG/Private/Tests/Elements/PCGFilterByIndexTest.cpp b/Engine/Plugins/PCG/Source/PCG/Private/Tests/Elements/PCGFilterByIndexTest.cpp
--- a/Engine/Plugins/PCG/Source/PCG/Private/Tests/Elements/PCGFilterByIndexTest.cpp
+++ b/Engine/Plugins/PCG/Source/PCG/Private/Tests/Elements/PCGFilterByIndexTest.cpp
@@ -192,7 +192,7 @@ bool FPCGFilterByIndexTest_InvalidSelection::RunTest(const FString& Parameters)
 	PCGTestsCommon::FTestData TestData;
 	UPCGFilterByIndexSettings* Settings = PCGTestsCommon::GenerateSettings<UPCGFilterByIndexSettings>(TestData);
 
-	AddExpectedError("Invalid expression in index selection string", EAutomationExpectedMessageFlags::Contains, 2);
+	AddExpectedError("Invalid expression in parsed string:", EAutomationExpectedMessageFlags::Contains, 2);
 
 	// Test inverted range
 	{
@@ -210,7 +210,7 @@ bool FPCGFilterByIndexTest_InvalidSelection::RunTest(const FString& Parameters)
 		UTEST_TRUE("Output count", Outputs.IsEmpty());
 	}
 
-	AddExpectedError("Invalid character in index selection string");
+	AddExpectedError("Invalid character in parsed string:");
 
 	// Test invalid character
 	{
diff --git a/Engine/Plugins/PCG/Source/PCG/Private/Utils/PCGLogErrors.cpp b/Engine/Plugins/PCG/Source/PCG/Private/Utils/PCGLogErrors.cpp
--- a/Engine/Plugins/PCG/Source/PCG/Private/Utils/PCGLogErrors.cpp
+++ b/Engine/Plugins/PCG/Source/PCG/Private/Utils/PCGLogErrors.cpp
@@ -78,4 +78,32 @@ namespace PCGLog::Metadata
 	}
 }
 
+namespace PCGLog::Parsing
+{
+	namespace Format
+	{
+		const FText EmptyExpression = LOCTEXT("EmptyExpression", "Empty expression in parsed string.");
+
+		const FTextFormat InvalidCharacter = LOCTEXT("InvalidCharacter", "Invalid character in parsed string: '{0}'.");
+		const FTextFormat InvalidExpression = LOCTEXT("InvalidExpression", "Invalid expression in parsed string: '{0}'.");
+	}
+
+	// Warnings
+	void LogEmptyExpressionWarning(const FPCGContext* InContext)
+	{
+		LogWarningOnGraph(Format::EmptyExpression, InContext);
+	}
+
+	// Errors
+	void LogInvalidCharacterInParsedStringError(const FStringView& ParsedString, const FPCGContext* InContext)
+	{
+		LogErrorOnGraph(FText::Format(Format::InvalidCharacter, FText::FromStringView(ParsedString)), InContext);
+	}
+
+	void LogInvalidExpressionInParsedStringError(const FStringView& ParsedString, const FPCGContext* InContext)
+	{
+		LogErrorOnGraph(FText::Format(Format::InvalidExpression, FText::FromStringView(ParsedString)), InContext);
+	}
+}
+
 #undef LOCTEXT_NAMESPACE
diff --git a/Engine/Plugins/PCG/Source/PCG/Public/Elements/PCGFilterElementsByIndex.h b/Engine/Plugins/PCG/Source/PCG/Public/Elements/PCGFilterElementsByIndex.h
new file mode 100644
--- /dev/null
+++ b/Engine/Plugins/PCG/Source/PCG/Public/Elements/PCGFilterElementsByIndex.h
@@ -0,0 +1,59 @@
+// Copyright Epic Games, Inc. All Rights Reserved.
+
+#pragma once
+
+#include "PCGSettings.h"
+
+#include "PCGFilterElementsByIndex.generated.h"
+
+/** Filters points or the elements of an attribute set based on a second input of points, attribute sets, or a user-defined index range expression. */
+UCLASS(MinimalAPI, BlueprintType, ClassGroup = (Procedural))
+class UPCGFilterElementsByIndexSettings : public UPCGSettings
+{
+	GENERATED_BODY()
+
+public:
+	//~Begin UPCGSettings interface
+#if WITH_EDITOR
+	virtual FName GetDefaultNodeName() const override { return FName(TEXT("FilterElementsByIndex")); }
+	virtual FText GetDefaultNodeTitle() const override { return NSLOCTEXT("PCGFilterElementsByIndexElement", "NodeTitle", "Filter Elements By Index"); }
+	virtual FText GetNodeTooltipText() const override { return NSLOCTEXT("PCGFilterElementsByIndexElement", "NodeTooltip", "Filters points or the elements of an attribute set based on a second input of points, attribute sets, or a user-defined index range expression."); }
+	virtual EPCGSettingsType GetType() const override { return EPCGSettingsType::Filter; }
+#endif // WITH_EDITOR
+	virtual FString GetAdditionalTitleInformation() const override;
+
+protected:
+	virtual TArray<FPCGPinProperties> InputPinProperties() const override;
+	virtual TArray<FPCGPinProperties> OutputPinProperties() const override;
+	virtual FPCGElementPtr CreateElement() const override;
+	//~End UPCGSettings interface
+
+public:
+	/** A second input will define the indices to filter. */
+	UPROPERTY(BlueprintReadWrite, EditAnywhere, Category = Settings)
+	bool bSelectIndicesByInput = true;
+
+	/** The attribute which will define the indices to filter. */
+	UPROPERTY(BlueprintReadWrite, EditAnywhere, Category = Settings, meta = (EditCondition = "bSelectIndicesByInput", EditConditionHides, PCG_Overridable))
+	FPCGAttributePropertyInputSelector IndexSelectionAttribute;
+
+	/** Selected individual indices or index ranges to include or exclude. Negative end indices allowed.
+	* For example, on an array of size 10: '0,2,4:5,7:-1' will include indices: 0,2,4,7,8.
+	*/
+	UPROPERTY(BlueprintReadWrite, EditAnywhere, Category = Settings, meta = (EditCondition = "!bSelectIndicesByInput", EditConditionHides, PCG_Overridable))
+	FString SelectedIndices = TEXT(":");
+
+	/** An additional output for discarded elements. */
+	UPROPERTY(BlueprintReadWrite, EditAnywhere, Category = Settings)
+	bool bOutputDiscardedElements = false;
+
+	/** Will invert which indices will be included and excluded. */
+	UPROPERTY(BlueprintReadWrite, EditAnywhere, Category = Settings, meta = (PCG_Overridable))
+	bool bInvertFilter = false;
+};
+
+class FPCGFilterElementsByIndexElement : public IPCGElement
+{
+protected:
+	virtual bool ExecuteInternal(FPCGContext* InContext) const override;
+};
diff --git a/Engine/Plugins/PCG/Source/PCG/Public/Helpers/Parsing/PCGIndexing.h b/Engine/Plugins/PCG/Source/PCG/Public/Helpers/Parsing/PCGIndexing.h
--- a/Engine/Plugins/PCG/Source/PCG/Public/Helpers/Parsing/PCGIndexing.h
+++ b/Engine/Plugins/PCG/Source/PCG/Public/Helpers/Parsing/PCGIndexing.h
@@ -10,9 +10,7 @@ namespace PCGIndexing
 	/** A simple data structure to represent a range of indices [X,Y]. Validating ranges are client responsibility. */
 	struct FPCGIndexRange
 	{
-		FPCGIndexRange(const int32 InStartIndex, const int32 InEndIndex) : StartIndex(InStartIndex), EndIndex(InEndIndex)
-		{
-		}
+		FPCGIndexRange(const int32 InStartIndex, const int32 InEndIndex) : StartIndex(InStartIndex), EndIndex(InEndIndex) {}
 
 		/** Returns true if the index can be found in this range. */
 		bool ContainsIndex(int32 Index) const;
@@ -32,11 +30,11 @@ namespace PCGIndexing
 	class FPCGIndexCollection
 	{
 	public:
+		static FPCGIndexCollection Invalid() { return {}; }
+
 		/** The constructor must accept the size of the array to support negative terminating indices. */
-		explicit FPCGIndexCollection(const int32 InArraySize) : ArraySize(InArraySize)
-		{
-			check(InArraySize > 0);
-		}
+		explicit FPCGIndexCollection(const int32 InArraySize) : ArraySize(FMath::Max(InArraySize, static_cast<int32>(INDEX_NONE))) {}
+
 
 		/** Add a new index range to the collection directly via start and end indices. */
 		bool AddRange(int32 StartIndex, int32 EndIndex);
@@ -59,9 +57,18 @@ namespace PCGIndexing
 		/** Gets the total number of concrete indices within the collection. */
 		int32 GetTotalIndexCount() const;
 
+		/** The index collection has a valid array size. */
+		bool IsValid() const;
+
+		/** The index collection has no ranges. */
+		bool IsEmpty() const;
+
 		bool operator==(const FPCGIndexCollection& Other) const;
+		FPCGIndexCollection& operator+=(const FPCGIndexCollection& Other);
 
 	private:
+		FPCGIndexCollection() = default;
+
 		/** Adjusts indices into a positive range, including offsetting the EndIndex if they are equal. */
 		FPCGIndexRange AdjustIndicesAndCreateRange(int32 StartIndex, int32 EndIndex) const;
 
@@ -72,9 +79,9 @@ namespace PCGIndexing
 		FPCGIndexRange MergeRanges(const FPCGIndexRange& FirstRange, const FPCGIndexRange& SecondRange) const;
 
 		/** The size of the representative array associated with this collection. Needed for negative terminating indices. */
-		int32 ArraySize = 0;
+		int32 ArraySize = INDEX_NONE;
 
 		/** A collection of the abstract index ranges in the collection. */
-		TArray<FPCGIndexRange> IndexRanges;
+		TArray<FPCGIndexRange, TInlineAllocator<64>> IndexRanges;
 	};
 }
diff --git a/Engine/Plugins/PCG/Source/PCG/Public/Helpers/Parsing/PCGParsing.h b/Engine/Plugins/PCG/Source/PCG/Public/Helpers/Parsing/PCGParsing.h
--- a/Engine/Plugins/PCG/Source/PCG/Public/Helpers/Parsing/PCGParsing.h
+++ b/Engine/Plugins/PCG/Source/PCG/Public/Helpers/Parsing/PCGParsing.h
@@ -23,5 +23,5 @@ namespace PCGParser
 	/** Indices should be included individually or with ranges using a delimiter. Negative terminating ranges are accepted.
 	* For example, on an array of size 10: '0,2,4:5,7:-1' will include indices: 0,2,4,5,7,8
 	*/
-	PCG_API EPCGParserResult ParseIndexRanges(PCGIndexing::FPCGIndexCollection& OutIndexCollection, const FString& InputString);
+	PCG_API EPCGParserResult ParseIndexRanges(PCGIndexing::FPCGIndexCollection& OutIndexCollection, const FStringView& InputString);
 }
diff --git a/Engine/Plugins/PCG/Source/PCG/Public/Metadata/PCGMetadata.h b/Engine/Plugins/PCG/Source/PCG/Public/Metadata/PCGMetadata.h
--- a/Engine/Plugins/PCG/Source/PCG/Public/Metadata/PCGMetadata.h
+++ b/Engine/Plugins/PCG/Source/PCG/Public/Metadata/PCGMetadata.h
@@ -267,6 +267,7 @@ public:
 	void AccumulateWeightedAttributes(PCGMetadataEntryKey InKey, const UPCGMetadata* InMetadata, float Weight, bool bSetNonInterpolableAttributes, PCGMetadataEntryKey& OutKey);
 
 	void SetAttributes(PCGMetadataEntryKey InKey, const UPCGMetadata* InMetadata, PCGMetadataEntryKey& OutKey);
+	void SetAttributes(const TArrayView<const PCGMetadataEntryKey>& InOriginalKeys, const UPCGMetadata* InMetadata, const TArrayView<PCGMetadataEntryKey>* InOutOptionalKeys = nullptr, FPCGContext* OptionalContext = nullptr);
 	void SetAttributes(const TArrayView<const PCGMetadataEntryKey>& InKeys, const UPCGMetadata* InMetadata, const TArrayView<PCGMetadataEntryKey>& OutKeys, FPCGContext* OptionalContext = nullptr);
 
 	/** Attributes operations - shorthand for points */
diff --git a/Engine/Plugins/PCG/Source/PCG/Public/Utils/PCGLogErrors.h b/Engine/Plugins/PCG/Source/PCG/Public/Utils/PCGLogErrors.h
--- a/Engine/Plugins/PCG/Source/PCG/Public/Utils/PCGLogErrors.h
+++ b/Engine/Plugins/PCG/Source/PCG/Public/Utils/PCGLogErrors.h
@@ -6,8 +6,6 @@
 #include "Metadata/PCGAttributePropertySelector.h"
 #include "Metadata/Accessors/IPCGAttributeAccessor.h"
 
-#include "Internationalization/Internationalization.h"
-
 struct FPCGContext;
 
 namespace PCGLog
@@ -89,4 +87,14 @@ namespace PCGLog
 			return LogFailToGetAttributeError<T>(Selector.GetDisplayText(), Accessor, InContext);
 		}
 	}
+
+	namespace Parsing
+	{
+		// Warnings
+		void LogEmptyExpressionWarning(const FPCGContext* InContext = nullptr);
+
+		// Errors
+		void LogInvalidCharacterInParsedStringError(const FStringView& ParsedString, const FPCGContext* InContext = nullptr);
+		void LogInvalidExpressionInParsedStringError(const FStringView& ParsedString, const FPCGContext* InContext = nullptr);
+	}
 }
