Fixed roots steepness to 1 in the level to asset code

#rnx

diff --git a/Engine/Plugins/PCG/Source/PCGEditor/Private/PCGLevelToAsset.cpp b/Engine/Plugins/PCG/Source/PCGEditor/Private/PCGLevelToAsset.cpp
--- a/Engine/Plugins/PCG/Source/PCGEditor/Private/PCGLevelToAsset.cpp
+++ b/Engine/Plugins/PCG/Source/PCGEditor/Private/PCGLevelToAsset.cpp
@@ -358,6 +358,7 @@ bool UPCGLevelToAsset::BP_ExportWorld_Implementation(UWorld* World, const FStrin
 		FPCGPoint& RootPoint = Roots.Emplace_GetRef(FTransform::Identity, 1.0f, 0);
 		RootPoint.BoundsMin = AllActorBounds.Min;
 		RootPoint.BoundsMax = AllActorBounds.Max;
+		RootPoint.Steepness = 1.0f;
 	}
 
 	return true;
