[PCG] Forcing bGenerated to false in PostLoad if we have no graph and making sure we early out in registering Partitions actors if we have no graph

#rnx

diff --git a/Engine/Plugins/PCG/Source/PCG/Private/PCGActorAndComponentMapping.cpp b/Engine/Plugins/PCG/Source/PCG/Private/PCGActorAndComponentMapping.cpp
--- a/Engine/Plugins/PCG/Source/PCG/Private/PCGActorAndComponentMapping.cpp
+++ b/Engine/Plugins/PCG/Source/PCG/Private/PCGActorAndComponentMapping.cpp
@@ -785,6 +785,11 @@ void FPCGActorAndComponentMapping::UpdateMappingPCGComponentPartitionActor(UPCGC
 
 	check(InComponent);
 
+	if (!InComponent->GetGraph())
+	{
+		return;
+	}
+
 	// Get the bounds
 	FBox Bounds = PartitionedOctree.GetBounds(InComponent);
 
diff --git a/Engine/Plugins/PCG/Source/PCG/Private/PCGComponent.cpp b/Engine/Plugins/PCG/Source/PCG/Private/PCGComponent.cpp
--- a/Engine/Plugins/PCG/Source/PCG/Private/PCGComponent.cpp
+++ b/Engine/Plugins/PCG/Source/PCG/Private/PCGComponent.cpp
@@ -1661,6 +1661,12 @@ void UPCGComponent::PostLoad()
 	// Force dirty to be false on load. We should never refresh on load.
 	bDirtyGenerated = false;
 
+	// We can never be generated if we have no graph
+	if (!GetGraph())
+	{
+		bGenerated = false;
+	}
+
 	// If we have both default value (bIsComponentPartitioned = false and bIsPartitioned = true)
 	// we will follow the value of bIsPartitioned.
 	// bIsPartitioned will be set to false to new objects
