[PCG] Fix nested NewObject_AnyThread ensure

#rnx

diff --git a/Engine/Plugins/PCG/Source/PCG/Private/Graph/PCGGraphExecutor.cpp b/Engine/Plugins/PCG/Source/PCG/Private/Graph/PCGGraphExecutor.cpp
--- a/Engine/Plugins/PCG/Source/PCG/Private/Graph/PCGGraphExecutor.cpp
+++ b/Engine/Plugins/PCG/Source/PCG/Private/Graph/PCGGraphExecutor.cpp
@@ -112,7 +112,15 @@ namespace PCGGraphExecutor
 			if (ensure(AsyncObject->HasAnyInternalFlags(EInternalObjectFlags::Async)))
 			{
 				AsyncObject->ClearInternalFlags(EInternalObjectFlags::Async);
-				ForEachObjectWithOuter(AsyncObject, [](UObject* SubObject) { SubObject->ClearInternalFlags(EInternalObjectFlags::Async); }, true);
+				ForEachObjectWithOuter(AsyncObject, [&AsyncObjects](UObject* SubObject)
+				{
+					if (AsyncObjects.Contains(SubObject))
+					{
+						return;
+					}
+					
+					SubObject->ClearInternalFlags(EInternalObjectFlags::Async);
+				}, true);
 			}
 		}
 	}
@@ -1977,9 +1985,9 @@ TArray<FPCGGraphExecutor::FCachedResult*> FPCGGraphExecutor::QueueNextTasksInter
 						{
 							OnTaskInputsReady(SuccessorTask, CachedResults, bIsInGameThread);
 							Tasks.Remove(Successor);
-						}
 					}
 				}
+				}
 			}
 		}
 	}
