[PCG] Fix static analysis warning in PCGFilterElementsByIndex

#rnx

diff --git a/Engine/Plugins/PCG/Source/PCG/Private/Elements/PCGFilterElementsByIndex.cpp b/Engine/Plugins/PCG/Source/PCG/Private/Elements/PCGFilterElementsByIndex.cpp
--- a/Engine/Plugins/PCG/Source/PCG/Private/Elements/PCGFilterElementsByIndex.cpp
+++ b/Engine/Plugins/PCG/Source/PCG/Private/Elements/PCGFilterElementsByIndex.cpp
@@ -316,7 +316,6 @@ bool FPCGFilterElementsByIndexElement::ExecuteInternal(FPCGContext* InContext) c
 			{
 				UPCGPointData* OutFilterPointData = FPCGContext::NewObject_AnyThread<UPCGPointData>(InContext);
 				OutFilterPointData->GetMutablePoints() = std::move(OutFilterPoints);
-				OutFilterPoints.Reserve(InputPointData->GetNumPoints());
 				AddNewOutput(OutFilterPointData, Constants::OutFilterOutputPinLabel);
 			}
 		}
