[PCG] Iterative builder now supports Unbounded and linking GridSizes dependencies

#rnx

diff --git a/Engine/Plugins/PCG/Source/PCG/Private/PCGComponent.cpp b/Engine/Plugins/PCG/Source/PCG/Private/PCGComponent.cpp
--- a/Engine/Plugins/PCG/Source/PCG/Private/PCGComponent.cpp
+++ b/Engine/Plugins/PCG/Source/PCG/Private/PCGComponent.cpp
@@ -432,6 +432,11 @@ FPCGTaskId UPCGComponent::GenerateLocalGetTaskId(EPCGComponentGenerationTrigger
 	return GenerateInternal(bForce, Grid, RequestedGenerationTrigger, {});
 }
 
+FPCGTaskId UPCGComponent::GenerateLocalGetTaskId(EPCGComponentGenerationTrigger RequestedGenerationTrigger, bool bForce, EPCGHiGenGrid Grid, const TArray<FPCGTaskId>& Dependencies)
+{
+	return GenerateInternal(bForce, Grid, RequestedGenerationTrigger, Dependencies);
+}
+
 FPCGTaskId UPCGComponent::GenerateInternal(bool bForce, EPCGHiGenGrid Grid, EPCGComponentGenerationTrigger RequestedGenerationTrigger, const TArray<FPCGTaskId>& Dependencies)
 {
 	if (IsGenerating() || !GetSubsystem() || !ShouldGenerate(bForce, RequestedGenerationTrigger))
diff --git a/Engine/Plugins/PCG/Source/PCG/Public/PCGComponent.h b/Engine/Plugins/PCG/Source/PCG/Public/PCGComponent.h
--- a/Engine/Plugins/PCG/Source/PCG/Public/PCGComponent.h
+++ b/Engine/Plugins/PCG/Source/PCG/Public/PCGComponent.h
@@ -166,6 +166,7 @@ public:
 
 	FPCGTaskId GenerateLocalGetTaskId(bool bForce);
 	FPCGTaskId GenerateLocalGetTaskId(EPCGComponentGenerationTrigger RequestedGenerationTrigger, bool bForce, EPCGHiGenGrid Grid = EPCGHiGenGrid::Uninitialized);
+	FPCGTaskId GenerateLocalGetTaskId(EPCGComponentGenerationTrigger RequestedGenerationTrigger, bool bForce, EPCGHiGenGrid Grid, const TArray<FPCGTaskId>& Dependencies);
 
 	/** Cleans up the generation from a local (vs. remote) standpoint. Will not be replicated. Will be delayed. */
 	UFUNCTION(BlueprintCallable, Category = PCG)
diff --git a/Engine/Plugins/PCG/Source/PCGEditor/Private/WorldPartitionBuilder/PCGWorldPartitionBuilder.cpp b/Engine/Plugins/PCG/Source/PCGEditor/Private/WorldPartitionBuilder/PCGWorldPartitionBuilder.cpp
--- a/Engine/Plugins/PCG/Source/PCGEditor/Private/WorldPartitionBuilder/PCGWorldPartitionBuilder.cpp
+++ b/Engine/Plugins/PCG/Source/PCGEditor/Private/WorldPartitionBuilder/PCGWorldPartitionBuilder.cpp
@@ -39,10 +39,10 @@ DEFINE_LOG_CATEGORY_STATIC(LogPCGWorldPartitionBuilder, All, All);
 struct FPCGWorldPartitionBuilder
 {
 	/** Generate the given components. Optionally generate only one component at a time with a wait on async processes after each. Optionally apply given filter to select components. */
-	static bool GenerateComponents(TArray<TWeakObjectPtr<UPCGComponent>>& Components, UWorld* InWorld, const FPCGWorldPartitionBuilderArgs& InArgs, TSet<FString>& InOutDeletedActorPackages, bool& bOutGenerationErrors);
+	static bool GenerateComponents(TArray<TWeakObjectPtr<UPCGComponent>>& Components, UWorld* InWorld, const FPCGWorldPartitionBuilderArgs& InArgs, const TMap<UPCGComponent*, UPCGComponent*>& InTaskDependencies, TSet<FString>& InOutDeletedActorPackages, bool& bOutGenerationErrors);
 
 	/** Generate a component. Applies correct editing mode if necessary. */
-	static void GenerateComponent(UPCGComponent* InComponent, UWorld* InWorld);
+	static FPCGTaskId GenerateComponent(UPCGComponent* InComponent, UWorld* InWorld, const FPCGWorldPartitionBuilderArgs& InArgs, FPCGTaskId InTaskDependency);
 
 	/** Groups components into buckets to generate in order */
 	static TArray<TArray<TWeakObjectPtr<UPCGComponent>>> GroupComponentsToGenerate(const TArray<TWeakObjectPtr<UPCGComponent>>& InComponentsToGenerate, const FPCGWorldPartitionBuilderArgs& InArgs);
@@ -65,6 +65,9 @@ struct FPCGWorldPartitionBuilder
 	/** Grab all original PCG components in world, applying given filter. */
 	static void CollectComponentsToGenerate(UWorld* InWorld, const FPCGWorldPartitionBuilderArgs& InArgs, TFunctionRef<bool(const UPCGComponent*)> ComponentFilter, TArray<TWeakObjectPtr<UPCGComponent>>& OutComponents);
 
+	/** Grab all sub grid PCG components to generate based on the current collected list */
+	static void CollectSubGridComponentsToGenerate(UWorld* InWorld, const FPCGWorldPartitionBuilderArgs& InArgs, const TArray<TWeakObjectPtr<UPCGComponent>>& InComponentsToGenerate, TArray<UPCGComponent*>& OutAdditionalComponentsToGenerate, TMap<UPCGComponent*, UPCGComponent*>& OutDependencies);
+
 	/** Reset loaded package dirty flags */
 	static void ResetPackageDirtyFlags();
 
@@ -79,8 +82,25 @@ struct FPCGWorldPartitionBuilder
 
 	/** Expand Loaded Bounds if needed */
 	static TUniquePtr<IWorldPartitionActorLoaderInterface::ILoaderAdapter> ExpandLoadedBoundsIfNeeded(UWorld* InWorld, const FBox& InCurrentBounds, const FBox& InNewBounds, const FString& InLoaderName);
+		
+	/** Get Grid Size information for a given graph and cache it for future calls */
+	static void GetGenerationGridSizes(const UPCGGraph* InGraph, const APCGWorldActor* InWorldActor, PCGHiGenGrid::FSizeArray& OutGridSizes, bool& bOutHasUnbounded);
+
+private:
+	static void CollectSubGridComponentsToGenerateInternal(UPCGSubsystem* InPCGSubsystem, int32 InGridSizeIndex, const PCGHiGenGrid::FSizeArray& InGridSizes, UPCGComponent* InOriginalComponent, APCGPartitionActor* InPartitionActor, TArray<UPCGComponent*>& OutAdditionalComponentsToGenerate, TMap<UPCGComponent*, UPCGComponent*>& OutDependencies);
+
+	/** Cached information about Partitioned Component Grid sizes */
+	struct FGridSizeInfo
+	{
+		bool bHasUnbounded = false;
+		PCGHiGenGrid::FSizeArray GridSizes;
+	};
+
+	static TMap<TSoftObjectPtr<const UPCGGraph>, FGridSizeInfo> GraphToGridSizeInfo;
 };
 
+TMap<TSoftObjectPtr<const UPCGGraph>, FPCGWorldPartitionBuilder::FGridSizeInfo> FPCGWorldPartitionBuilder::GraphToGridSizeInfo;
+
 namespace PCGWorldPartitionBuilder
 {
 	static FAutoConsoleCommand CommandBuildComponents(
@@ -517,7 +537,7 @@ bool UPCGWorldPartitionBuilder::RunEntireWorld(UWorld* World, FPackageSourceCont
 	for (TArray<TWeakObjectPtr<UPCGComponent>>& GroupToGenerate : GroupsToGenerate)
 	{
 		bool bErrorsOccurred = false;
-		bGeneratedAnyComponent |= FPCGWorldPartitionBuilder::GenerateComponents(GroupToGenerate, World, Args, DeletedActorPackages, bErrorsOccurred);
+		bGeneratedAnyComponent |= FPCGWorldPartitionBuilder::GenerateComponents(GroupToGenerate, World, Args, TMap<UPCGComponent*, UPCGComponent*>(), DeletedActorPackages, bErrorsOccurred);
 		bErrorOccurredWhileGenerating |= bErrorsOccurred;
 	}
 
@@ -655,25 +675,58 @@ bool UPCGWorldPartitionBuilder::RunPerCell(UWorld* InWorld, const FCellInfo& InC
 	TArray<TWeakObjectPtr<UPCGComponent>> ComponentsToGenerate;
 	FBox ExpandedBounds = InCellInfo.Bounds;
 
-	auto GenerateComponentFilter = [&InCellInfo, &ExpandedBounds, this](const UPCGComponent* InComponent)
+	UPCGSubsystem* PCGSubsystem = UWorld::GetSubsystem<UPCGSubsystem>(InWorld);
+	check(PCGSubsystem);
+
+	auto GenerateComponentFilter = [&InCellInfo, &ExpandedBounds, PCGSubsystem, this](const UPCGComponent* InComponent)
 	{
-		if (InComponent->IsPartitioned())
+		if (GeneratedComponents.Contains(InComponent))
 		{
 			return false;
 		}
+			
+		bool bCheckBoundsCenter = true;
+		if (InComponent->IsPartitioned())
+		{
+			bool bHasUnbounded = false;
+			PCGHiGenGrid::FSizeArray GridSizes;
+			FPCGWorldPartitionBuilder::GetGenerationGridSizes(InComponent->GetGraph(), PCGSubsystem->GetPCGWorldActor(), GridSizes, bHasUnbounded);
 
-		if (GeneratedComponents.Contains(InComponent))
+			if (!bHasUnbounded)
+			{
+				return false;
+			}
+
+			// If an Unbounded Partitioned Component exists, it needs to be generated before all its PAs, this can be wasteful on the loading side of things because we are going to expand 
+			// the bounds of loading to fit this unbounded volume, if it happens to be a very large volume it can make the IterativeCell mode disadvantageous
+			bCheckBoundsCenter = false;
+		}
+				
+		// When processing a PA make sure to only collect the top level grid actors, we are then going to find all child grid actors in a next step
+		if (APCGPartitionActor* PartitionActor = Cast<APCGPartitionActor>(InComponent->GetOwner()))
 		{
-			return false;
+			UPCGComponent* OriginalComponent = PartitionActor->GetOriginalComponent(InComponent);
+			check(OriginalComponent);
+			
+			// Avoid calling this multiple times
+			bool bHasUnbounded = false;
+			PCGHiGenGrid::FSizeArray GridSizes;
+			FPCGWorldPartitionBuilder::GetGenerationGridSizes(OriginalComponent->GetGraph(), PCGSubsystem->GetPCGWorldActor(), GridSizes, bHasUnbounded);
+
+			// Only generate if this PA represents the top level grid of the Original Component
+			if (GridSizes.IsEmpty() || PartitionActor->GetPCGGridSize() != GridSizes[0])
+			{
+				return false;
+			}
 		}
 
 		// Generate Component if its bound center is inside the currently processing cell bounds (to avoid expanding bounds too much)
 		const FBox ActorBounds = PCGHelpers::GetActorBounds(InComponent->GetOwner());
-		if (!InCellInfo.Bounds.IsInsideOrOnXY(ActorBounds.GetCenter()))
+		if (bCheckBoundsCenter && !InCellInfo.Bounds.IsInsideOrOnXY(ActorBounds.GetCenter()))
 		{
 			return false;
 		}
-
+		
 		// Expand the cell bounds to include the whole actor bounds (keep Z same as original Cell bounds)
 		FBox ActorXYBounds(FVector(ActorBounds.Min.X, ActorBounds.Min.Y, ExpandedBounds.Min.Z), FVector(ActorBounds.Max.X, ActorBounds.Max.Y, ExpandedBounds.Max.Z));
 		ExpandedBounds += ActorXYBounds;
@@ -687,9 +740,42 @@ bool UPCGWorldPartitionBuilder::RunPerCell(UWorld* InWorld, const FCellInfo& InC
 	// Collect Components to Generate a first time, this will allow us to call FPCGWorldPartitionBuilder::UpdatePartitionedActors on ParttionActors with all their overlapp
 	FPCGWorldPartitionBuilder::CollectComponentsToGenerate(InWorld, Args, GenerateComponentFilter, ComponentsToGenerate);
 		
+	// Make sure Partitioned (Unbounded) components are Generated first
+	ComponentsToGenerate.Sort([](const TWeakObjectPtr<UPCGComponent>& WeakComponentA, const TWeakObjectPtr<UPCGComponent>& WeakComponentB)
+	{
+		UPCGComponent* ComponentA = WeakComponentA.Get();
+		check(ComponentA);
+		UPCGComponent* ComponentB = WeakComponentB.Get();
+		check(ComponentB);
+		
+		if (ComponentA->IsPartitioned() == ComponentB->IsPartitioned())
+		{
+			// Default to ptr hash (we don't really care at that point as we really just want the Partitioned volumes first)
+			return WeakComponentA.GetWeakPtrTypeHash() < WeakComponentB.GetWeakPtrTypeHash();
+		}
+
+		return ComponentA->IsPartitioned();
+	});
+		
 	// Do a Bounds expansion to make sure we are generating PCG Components on areas covering their owning actor bounds (which might be a bigger area then the currently loaded cell)
 	TUniquePtr<IWorldPartitionActorLoaderInterface::ILoaderAdapter> LoaderAdapter = FPCGWorldPartitionBuilder::ExpandLoadedBoundsIfNeeded(InWorld, InCellInfo.Bounds, ExpandedBounds, TEXT("Expanded Bounds"));
 	
+	// Add Components for sub grids
+	TArray<UPCGComponent*> SubGridComponentsToGenerate;
+	TMap<UPCGComponent*, UPCGComponent*> TaskDependencies;
+	FPCGWorldPartitionBuilder::CollectSubGridComponentsToGenerate(InWorld, Args, ComponentsToGenerate, SubGridComponentsToGenerate, TaskDependencies);
+
+	for (UPCGComponent* SubGridPCGComponent : SubGridComponentsToGenerate)
+	{
+		if (!GeneratedComponents.Contains(SubGridPCGComponent))
+		{
+			ComponentsToGenerate.Add(SubGridPCGComponent);
+
+			// Consider this component generated
+			GeneratedComponents.Add(SubGridPCGComponent);
+		}
+	}
+
 	bool bErrorsOccurred = false;
 	// Do Generation and Collect Packages to save
 	if (ComponentsToGenerate.Num() > 0)
@@ -697,7 +783,7 @@ bool UPCGWorldPartitionBuilder::RunPerCell(UWorld* InWorld, const FCellInfo& InC
 		TArray<TArray<TWeakObjectPtr<UPCGComponent>>> GroupsToGenerate = FPCGWorldPartitionBuilder::GroupComponentsToGenerate(ComponentsToGenerate, Args);
 		for (TArray<TWeakObjectPtr<UPCGComponent>>& GroupToGenerate : GroupsToGenerate)
 		{
-			bCellHasChanges |= FPCGWorldPartitionBuilder::GenerateComponents(GroupToGenerate, InWorld, Args, CellDeletedPackageNames, bErrorsOccurred);
+			bCellHasChanges |= FPCGWorldPartitionBuilder::GenerateComponents(GroupToGenerate, InWorld, Args, TaskDependencies, CellDeletedPackageNames, bErrorsOccurred);
 			bErrorOccurredWhileGenerating |= bErrorsOccurred;
 		}
 	}
@@ -869,6 +955,7 @@ bool FPCGWorldPartitionBuilder::GenerateComponents(
 	TArray<TWeakObjectPtr<UPCGComponent>>& Components,
 	UWorld* InWorld,
 	const FPCGWorldPartitionBuilderArgs& InArgs,
+	const TMap<UPCGComponent*, UPCGComponent*>& InTaskDependencies,
 	TSet<FString>& InOutDeletedActorPackages,
 	bool& bOutGenerationErrors)
 {
@@ -898,7 +985,7 @@ bool FPCGWorldPartitionBuilder::GenerateComponents(
 	// Will track any deleted actors in this scope
 	FPCGDetectDeletedActorInScope DeletedActorScope(InWorld, InOutDeletedActorPackages);
 
-	TSet<TObjectKey<UPCGComponent>> GeneratedComponents;
+	TMap<TObjectKey<UPCGComponent>, FPCGTaskId> GeneratedComponents;
 
 	for (TWeakObjectPtr<UPCGComponent> ComponentWeakPtr : Components)
 	{
@@ -908,9 +995,21 @@ bool FPCGWorldPartitionBuilder::GenerateComponents(
 			UE_LOG(LogPCGWorldPartitionBuilder, Warning, TEXT("Lost a PCG component weak pointer, component will not be generated."));
 			continue;
 		}
+		
+		FPCGTaskId TaskIdDependency = InvalidPCGTaskId;
+		if (const UPCGComponent* const* TaskDependency = InTaskDependencies.Find(Component))
+		{
+			if (FPCGTaskId* TaskIdDependencyPtr = GeneratedComponents.Find(*TaskDependency))
+			{
+				TaskIdDependency = *TaskIdDependencyPtr;
+			}
+		}
 
-		FPCGWorldPartitionBuilder::GenerateComponent(Component, InWorld);
-		GeneratedComponents.Add(Component);
+		FPCGTaskId GenerateTaskId = FPCGWorldPartitionBuilder::GenerateComponent(Component, InWorld, InArgs, TaskIdDependency);
+		if (GenerateTaskId != InvalidPCGTaskId)
+		{
+			GeneratedComponents.Add(Component, GenerateTaskId);
+		}
 
 		if (InArgs.bOneComponentAtATime)
 		{
@@ -920,9 +1019,9 @@ bool FPCGWorldPartitionBuilder::GenerateComponents(
 
 	if (!InArgs.bOneComponentAtATime)
 	{
-		for (TObjectKey<UPCGComponent>& GeneratedComponent : GeneratedComponents)
+		for (const auto& KeyValuePair : GeneratedComponents)
 		{
-			if (const UPCGComponent* Component = GeneratedComponent.ResolveObjectPtr())
+			if (const UPCGComponent* Component = KeyValuePair.Key.ResolveObjectPtr())
 			{
 				WaitForComponentGeneration(Component);
 			}
@@ -986,12 +1085,12 @@ TArray<TArray<TWeakObjectPtr<UPCGComponent>>> FPCGWorldPartitionBuilder::GroupCo
 	return GroupsToGenerate;
 }
 
-void FPCGWorldPartitionBuilder::GenerateComponent(UPCGComponent* InComponent, UWorld* InWorld)
+FPCGTaskId FPCGWorldPartitionBuilder::GenerateComponent(UPCGComponent* InComponent, UWorld* InWorld, const FPCGWorldPartitionBuilderArgs& InArgs, FPCGTaskId InTaskDependency)
 {
 	// Separate ensures for maximum debug information.
 	if (!ensure(InComponent) || !ensure(InComponent->GetGraph()))
 	{
-		return;
+		return InvalidPCGTaskId;
 	}
 
 	ensure(InComponent->bActivated);
@@ -1009,7 +1108,21 @@ void FPCGWorldPartitionBuilder::GenerateComponent(UPCGComponent* InComponent, UW
 	}
 
 	// Force generate as components that are already generated may decline the request.
-	const FPCGTaskId GenerateTask = InComponent->GenerateLocalGetTaskId(/*bForce=*/true);
+	FPCGTaskId GenerateTask = InvalidPCGTaskId;
+
+	// The only Partitioned components we do generate in iterative mode are generated only for their Unbounded grid as the other grid sizes will be generated through the PAs directly
+	if (InArgs.bIterativeCellLoading && InComponent->IsPartitioned())
+	{
+		GenerateTask = InComponent->GenerateLocalGetTaskId(EPCGComponentGenerationTrigger::GenerateOnDemand, /*bForce=*/true, EPCGHiGenGrid::Unbounded);
+	}
+	else if(InTaskDependency != InvalidPCGTaskId)
+	{
+		GenerateTask = InComponent->GenerateLocalGetTaskId(EPCGComponentGenerationTrigger::GenerateOnDemand, /*bForce=*/true, EPCGHiGenGrid::Uninitialized, { InTaskDependency });
+	}
+	else
+	{
+		GenerateTask = InComponent->GenerateLocalGetTaskId(/*bForce=*/true);
+	}
 
 	if (GenerateTask == InvalidPCGTaskId)
 	{
@@ -1027,6 +1140,8 @@ void FPCGWorldPartitionBuilder::GenerateComponent(UPCGComponent* InComponent, UW
 				*InComponent->GetGraph()->GetName());
 		}
 	}
+
+	return GenerateTask;
 }
 
 void FPCGWorldPartitionBuilder::WaitForAllAsyncEditorProcesses(UWorld* InWorld)
@@ -1082,6 +1197,22 @@ TUniquePtr<IWorldPartitionActorLoaderInterface::ILoaderAdapter> FPCGWorldPartiti
 	return LoaderAdapter;
 }
 
+void FPCGWorldPartitionBuilder::GetGenerationGridSizes(const UPCGGraph* InGraph, const APCGWorldActor* InWorldActor, PCGHiGenGrid::FSizeArray& OutGridSizes, bool& bOutHasUnbounded)
+{
+	if (FGridSizeInfo* FoundGridSizeInfo = GraphToGridSizeInfo.Find(InGraph))
+	{
+		OutGridSizes = FoundGridSizeInfo->GridSizes;
+		bOutHasUnbounded = FoundGridSizeInfo->bHasUnbounded;
+	}
+		
+	check(PCGHelpers::GetGenerationGridSizes(InGraph, InWorldActor, OutGridSizes, bOutHasUnbounded));
+
+	// Store
+	FGridSizeInfo& NewGridSizeInfo = GraphToGridSizeInfo.Add(InGraph);
+	NewGridSizeInfo.bHasUnbounded = bOutHasUnbounded;
+	NewGridSizeInfo.GridSizes = OutGridSizes;
+}
+
 void FPCGWorldPartitionBuilder::ResetPackageDirtyFlags()
 {
 	TArray<UPackage*> DirtyPackages = FPCGWorldPartitionBuilder::GetDirtyPackages();
@@ -1299,9 +1430,9 @@ bool FPCGWorldPartitionBuilder::CreatePartitionedActors(UWorld* InWorld, const F
 
 			bool bHasUnbounded = false;
 			PCGHiGenGrid::FSizeArray GridSizes;
-			ensure(PCGHelpers::GetGenerationGridSizes(PCGComponent->GetGraph(), PCGSubsystem->GetPCGWorldActor(), GridSizes, bHasUnbounded));
+			FPCGWorldPartitionBuilder::GetGenerationGridSizes(PCGComponent->GetGraph(), PCGSubsystem->GetPCGWorldActor(), GridSizes, bHasUnbounded);
 
-			if (!PCGComponent->bGenerated)
+			if (!PCGComponent->bGenerated && !bHasUnbounded)
 			{
 				PCGComponent->Modify();
 				PCGComponent->bGenerated = true;
@@ -1369,6 +1500,77 @@ bool FPCGWorldPartitionBuilder::UpdatePartitionedActors(UWorld* InWorld, const T
 	return DirtyPackageInScope.AnyDirtyPackage();
 }
 
+void FPCGWorldPartitionBuilder::CollectSubGridComponentsToGenerate(UWorld* InWorld, const FPCGWorldPartitionBuilderArgs& InArgs, const TArray<TWeakObjectPtr<UPCGComponent>>& InComponentsToGenerate, TArray<UPCGComponent*>& OutAdditionalComponentsToGenerate, TMap<UPCGComponent*, UPCGComponent*>& OutDependencies)
+{
+	UPCGSubsystem* PCGSubsystem = UWorld::GetSubsystem<UPCGSubsystem>(InWorld);
+	check(PCGSubsystem);
+
+	for (TWeakObjectPtr<UPCGComponent> WeakComponentToGenerate : InComponentsToGenerate)
+	{
+		if (UPCGComponent* ComponentToGenerate = WeakComponentToGenerate.Get())
+		{
+			if (APCGPartitionActor* PartitionActor = Cast<APCGPartitionActor>(ComponentToGenerate->GetOwner()))
+			{
+				if (UPCGComponent* OriginalComponent = PartitionActor->GetOriginalComponent(ComponentToGenerate))
+				{
+					// Add dependency on Unbounded generation if it is part of the components to generate
+					if (InComponentsToGenerate.Contains(OriginalComponent))
+					{
+						OutDependencies.Add(ComponentToGenerate, OriginalComponent);
+					}
+					
+					bool bHasUnbounded = false;
+					PCGHiGenGrid::FSizeArray GridSizes;
+					FPCGWorldPartitionBuilder::GetGenerationGridSizes(OriginalComponent->GetGraph(), PCGSubsystem->GetPCGWorldActor(), GridSizes, bHasUnbounded);
+
+					check(GridSizes[0] == PartitionActor->GetPCGGridSize());
+					const int32 GridSizeIndex = 0;
+					CollectSubGridComponentsToGenerateInternal(PCGSubsystem, GridSizeIndex + 1, GridSizes, OriginalComponent, PartitionActor, OutAdditionalComponentsToGenerate, OutDependencies);
+				}
+			}
+		}
+	}
+
+	// Sort by GridSize
+	OutAdditionalComponentsToGenerate.Sort([](const UPCGComponent& A, const UPCGComponent& B)
+	{
+		return Cast<APCGPartitionActor>(A.GetOwner())->GetPCGGridSize() > Cast<APCGPartitionActor>(B.GetOwner())->GetPCGGridSize();
+	});
+}
+
+void FPCGWorldPartitionBuilder::CollectSubGridComponentsToGenerateInternal(UPCGSubsystem* InPCGSubsystem, int32 InGridSizeIndex, const PCGHiGenGrid::FSizeArray& InGridSizes, UPCGComponent* InOriginalComponent, APCGPartitionActor* InPartitionActor, TArray<UPCGComponent*>& OutAdditionalComponentsToGenerate, TMap<UPCGComponent*, UPCGComponent*>& OutDependencies)
+{
+	if (InGridSizeIndex >= InGridSizes.Num())
+	{
+		return;
+	}
+
+	// Get the next Grid size overlapping components
+	TArray<UPCGComponent*> LocalAdditionalComponentsToGenerate;
+	InPCGSubsystem->ForAllOverlappingCells(InOriginalComponent, InPartitionActor->GetFixedBounds(), { InGridSizes[InGridSizeIndex] }, false, {},
+		[&InGridSizes, InGridSizeIndex, InPartitionActor, InOriginalComponent, &LocalAdditionalComponentsToGenerate, &OutDependencies](APCGPartitionActor* SubPartitionActor, const FBox& InBounds)
+		{
+			check(SubPartitionActor->GetPCGGridSize() == InGridSizes[InGridSizeIndex]);
+			if (InPartitionActor->GetFixedBounds().IsInsideXY(SubPartitionActor->GetFixedBounds().GetCenter()))
+			{
+				if (UPCGComponent* LocalComponent = SubPartitionActor->GetLocalComponent(InOriginalComponent))
+				{
+					LocalAdditionalComponentsToGenerate.Add(LocalComponent);
+					OutDependencies.Add(LocalComponent, InPartitionActor->GetLocalComponent(InOriginalComponent));
+				}
+			}
+			return InvalidPCGTaskId;
+		});
+
+	OutAdditionalComponentsToGenerate.Append(LocalAdditionalComponentsToGenerate);
+
+	for (UPCGComponent* AdditionalComponentToGenerate : LocalAdditionalComponentsToGenerate)
+	{
+		APCGPartitionActor* SubPartitionActor = CastChecked<APCGPartitionActor>(AdditionalComponentToGenerate->GetOwner());
+		CollectSubGridComponentsToGenerateInternal(InPCGSubsystem, InGridSizeIndex + 1, InGridSizes, InOriginalComponent, SubPartitionActor, OutAdditionalComponentsToGenerate, OutDependencies);
+	}
+}
+
 void FPCGWorldPartitionBuilder::CollectComponentsToGenerate(UWorld* InWorld, const FPCGWorldPartitionBuilderArgs& InArgs, TFunctionRef<bool(const UPCGComponent*)> ComponentFilter, TArray<TWeakObjectPtr<UPCGComponent>>& OutComponents)
 {
 	check(InWorld);
