[PCG] Add multi-line strings to the Print String node

diff --git a/Engine/Plugins/PCG/Source/PCG/Public/Elements/PCGPrintElement.h b/Engine/Plugins/PCG/Source/PCG/Public/Elements/PCGPrintElement.h
--- a/Engine/Plugins/PCG/Source/PCG/Public/Elements/PCGPrintElement.h
+++ b/Engine/Plugins/PCG/Source/PCG/Public/Elements/PCGPrintElement.h
@@ -65,7 +65,7 @@ protected:
 
 public:
 	/** The core message to print to the logger, graph, and/or screen. */
-	UPROPERTY(BlueprintReadWrite, EditAnywhere, Category = Settings, meta = (PCG_Overridable))
+	UPROPERTY(BlueprintReadWrite, EditAnywhere, Category = Settings, meta = (MultiLine, PCG_Overridable))
 	FString PrintString;
 
 	/** The verbosity level of the printed message. */
@@ -128,4 +128,4 @@ public:
 
 protected:
 	virtual bool ExecuteInternal(FPCGContext* Context) const override;
-};
\ No newline at end of file
+};
