[PCG] Create a Pathfinding node to find a path through Point Data using A* (V1)

Still TODO:
Start and Goal points as PointOrParam data.

diff --git a/Engine/Plugins/PCG/Source/PCG/Private/Elements/PCGPathfindingElement.cpp b/Engine/Plugins/PCG/Source/PCG/Private/Elements/PCGPathfindingElement.cpp
new file mode 100644
--- /dev/null
+++ b/Engine/Plugins/PCG/Source/PCG/Private/Elements/PCGPathfindingElement.cpp
@@ -0,0 +1,221 @@
+// Copyright Epic Games, Inc. All Rights Reserved.
+
+#include "Elements/PCGPathfindingElement.h"
+
+#include "PCGContext.h"
+#include "Data/PCGPointData.h"
+#include "Data/PCGSplineData.h"
+#include "Helpers/PCGHelpers.h"
+#include "SpatialAlgo/PCGAStar.h"
+
+#include "Components/SplineComponent.h"
+
+#define LOCTEXT_NAMESPACE "PCGPathfindingElement"
+
+namespace PCGPathfindingElement::Helpers
+{
+	TArray<FSplinePoint> ConvertPathToSplinePoints(const TArrayView<const FPCGPoint> Path, const EPCGPathfindingSplineMode SplineMode)
+	{
+		ESplinePointType::Type SplineCurveMode = ESplinePointType::Type::Constant;
+		switch (SplineMode)
+		{
+			case EPCGPathfindingSplineMode::Curve:
+				SplineCurveMode = ESplinePointType::Curve;
+				break;
+			case EPCGPathfindingSplineMode::Linear:
+				SplineCurveMode = ESplinePointType::Linear;
+			default:
+				break;
+		}
+
+		TArray<FSplinePoint> SplinePoints;
+		SplinePoints.Reserve(Path.Num());
+
+		int Index = 0;
+		Algo::Transform(Path, SplinePoints, [&Index, SplineCurveMode](const FPCGPoint& Point)
+		{
+			return FSplinePoint(
+					Index++, // Spline points must be indexed in ascending order
+					Point.Transform.GetLocation(),
+					FVector::ZeroVector,
+					FVector::ZeroVector,
+					FRotator::ZeroRotator,
+					FVector::OneVector,
+					SplineCurveMode);
+		});
+
+		return SplinePoints;
+	}
+} // namespace PCGPathfindingElement::Helpers
+
+TArray<FPCGPinProperties> UPCGPathfindingSettings::InputPinProperties() const
+{
+	TArray<FPCGPinProperties> Properties;
+	Properties.Emplace_GetRef(PCGPinConstants::DefaultInputLabel, EPCGDataType::Point).SetRequiredPin();
+	// TODO: PointOrParam input for start/goal.
+	return Properties;
+}
+
+TArray<FPCGPinProperties> UPCGPathfindingSettings::OutputPinProperties() const
+{
+	TArray<FPCGPinProperties> Properties;
+	if (bOutputAsSpline)
+	{
+		Properties.Emplace(PCGPinConstants::DefaultOutputLabel, EPCGDataType::Spline);
+	}
+	else
+	{
+		Properties.Emplace(PCGPinConstants::DefaultOutputLabel, EPCGDataType::Point);
+	}
+
+	return Properties;
+}
+
+FPCGElementPtr UPCGPathfindingSettings::CreateElement() const
+{
+	return MakeShared<FPCGPathfindingElement>();
+}
+
+#if WITH_EDITOR
+EPCGChangeType UPCGPathfindingSettings::GetChangeTypeForProperty(const FName& InPropertyName) const
+{
+	EPCGChangeType LocalChangeType = EPCGChangeType::Cosmetic;
+	if (InPropertyName == GET_MEMBER_NAME_CHECKED(UPCGPathfindingSettings, bOutputAsSpline))
+	{
+		LocalChangeType |= EPCGChangeType::Structural;
+	}
+
+	return Super::GetChangeTypeForProperty(InPropertyName) | LocalChangeType;
+}
+#endif // WITH_EDITOR
+
+bool FPCGPathfindingElement::PrepareDataInternal(FPCGContext* InContext) const
+{
+	TRACE_CPUPROFILER_EVENT_SCOPE(FPCGPathfindingElement::PrepareData);
+
+	ContextType* Context = static_cast<ContextType*>(InContext);
+	check(Context);
+
+	const UPCGPathfindingSettings* Settings = Context->GetInputSettings<UPCGPathfindingSettings>();
+	check(Settings);
+
+	const TArray<FPCGTaggedData> PointInputs = Context->InputData.GetInputsByPin(PCGPinConstants::DefaultInputLabel);
+
+	if (PointInputs.IsEmpty())
+	{
+		return true;
+	}
+
+	Context->InitializePerExecutionState([&PointInputs, Settings](ContextType*, ExecStateType& OutState)
+	{
+		TRACE_CPUPROFILER_EVENT_SCOPE(FPCGClusterElement::PrepareData::InitializePerExecutionState);
+
+		if (PointInputs.IsEmpty())
+		{
+			return EPCGTimeSliceInitResult::NoOperation;
+		}
+
+		// TODO: The start and goal points will eventually come from PointOrParam.
+		OutState = PCGSpatialAlgo::AStar::FSearchSettings
+		{
+			.StartPoint = FPCGPoint(FTransform(Settings->Start), /*InDensity=*/1, PCGHelpers::ComputeSeedFromPosition(Settings->Start)),
+			.SearchDistance = Settings->SearchDistance,
+			.Start = Settings->Start,
+			.Goal = Settings->Goal,
+			.HeuristicWeight = Settings->HeuristicWeight,
+			.bAcceptPartialPath = Settings->bAcceptPartialPath,
+			.bCopyOriginatingPoints = Settings->bCopyOriginatingPoints
+		};
+
+		return EPCGTimeSliceInitResult::Success;
+	});
+
+	Context->InitializePerIterationStates(PointInputs.Num(), [&PointInputs](IterStateType& OutSearchState, const ExecStateType& SearchSettings, const uint32 IterationIndex)
+	{
+		TRACE_CPUPROFILER_EVENT_SCOPE(FPCGClusterElement::PrepareData::InitializePerIterationStates);
+
+		OutSearchState.OriginatingPointData = Cast<UPCGPointData>(PointInputs[IterationIndex].Data);
+		if (OutSearchState.OriginatingPointData->IsEmpty())
+		{
+			// Already confirmed we can't make it from start->goal in the execution state check.
+			return EPCGTimeSliceInitResult::NoOperation;
+		}
+
+		PCGSpatialAlgo::AStar::Initialize(OutSearchState.OriginatingPointData, SearchSettings, OutSearchState);
+
+		return EPCGTimeSliceInitResult::Success;
+	});
+
+	return true;
+}
+
+bool FPCGPathfindingElement::ExecuteInternal(FPCGContext* InContext) const
+{
+	TRACE_CPUPROFILER_EVENT_SCOPE(FPCGPathfindingElement::Execute);
+
+	ContextType* TimeSlicedContext = static_cast<ContextType*>(InContext);
+	check(TimeSlicedContext);
+
+	if (!TimeSlicedContext->DataIsPreparedForExecution() || TimeSlicedContext->GetExecutionStateResult() == EPCGTimeSliceInitResult::NoOperation)
+	{
+		return true;
+	}
+
+	const UPCGPathfindingSettings* Settings = TimeSlicedContext->GetInputSettings<UPCGPathfindingSettings>();
+	check(Settings);
+
+	return ExecuteSlice(TimeSlicedContext, [Settings](ContextType* Context, const ExecStateType& SearchSettings, IterStateType& PathData, const uint32 IterIndex)
+	{
+		if (Context->GetIterationStateResult(IterIndex) == EPCGTimeSliceInitResult::NoOperation)
+		{
+			return true;
+		}
+
+		TArray<FPCGPoint> FinalPath;
+		bool bFinished = false;
+		do
+		{
+			bFinished = PCGSpatialAlgo::AStar::ExecuteSearchIteration(SearchSettings, PathData, FinalPath);
+
+			if (!bFinished && Context->ShouldStop())
+			{
+				return false;
+			}
+		}
+		while (!bFinished);
+
+		// No path was found and partial paths were not enabled.
+		if (FinalPath.IsEmpty())
+		{
+			check(!SearchSettings.bAcceptPartialPath);
+			return true;
+		}
+
+		// Finally, output the path as either a spline or points
+		if (Settings->bOutputAsSpline)
+		{
+			const TArray<FSplinePoint> SplinePoints = PCGPathfindingElement::Helpers::ConvertPathToSplinePoints(FinalPath, Settings->SplineMode);
+			UPCGSplineData* SplineData = FPCGContext::NewObject_AnyThread<UPCGSplineData>(Context);
+			SplineData->Initialize(SplinePoints, /*bInClosedLoop=*/false, FTransform::Identity);
+
+			FPCGTaggedData& OutputData = Context->OutputData.TaggedData.Emplace_GetRef();
+			OutputData.Data = SplineData;
+		}
+		else
+		{
+			UPCGPointData* OutputPointData = FPCGContext::NewObject_AnyThread<UPCGPointData>(Context);
+			if (SearchSettings.bCopyOriginatingPoints)
+			{
+				OutputPointData->InitializeFromData(PathData.OriginatingPointData);
+			}
+			OutputPointData->GetMutablePoints() = std::move(FinalPath);
+
+			FPCGTaggedData& OutputData = Context->OutputData.TaggedData.Emplace_GetRef();
+			OutputData.Data = OutputPointData;
+		}
+
+		return true;
+	});
+}
+
+#undef LOCTEXT_NAMESPACE
diff --git a/Engine/Plugins/PCG/Source/PCG/Private/SpatialAlgo/PCGAStar.cpp b/Engine/Plugins/PCG/Source/PCG/Private/SpatialAlgo/PCGAStar.cpp
new file mode 100644
--- /dev/null
+++ b/Engine/Plugins/PCG/Source/PCG/Private/SpatialAlgo/PCGAStar.cpp
@@ -0,0 +1,171 @@
+// Copyright Epic Games, Inc. All Rights Reserved.
+
+#include "SpatialAlgo/PCGAStar.h"
+
+#include "Helpers/PCGHelpers.h"
+#include "SpatialAlgo/PCGOctreeQueries.h"
+
+namespace PCGSpatialAlgo::AStar
+{
+namespace Helpers
+{
+	bool CompareNodes(const FNode& Node1, const FNode& Node2)
+	{
+		return Node1.EstimatedGoalCost < Node2.EstimatedGoalCost;
+	}
+
+	void BuildFinalPath(const FSearchState& SearchState, const FNode& FinalNode, const bool bCopyOriginatingPoints, TArray<FPCGPoint>& OutPath)
+	{
+		check(&FinalNode);
+
+		OutPath.Reset();
+
+		const FNode* Node = &FinalNode;
+		do // Build path from the goal backwards.
+		{
+			check(Node->PCGPoint);
+			if (bCopyOriginatingPoints)
+			{
+				OutPath.Emplace(*Node->PCGPoint);
+			}
+			else
+			{
+				OutPath.Emplace(Node->PCGPoint->Transform, /*InDensity=*/1, PCGHelpers::ComputeSeedFromPosition(Node->PCGPoint->Transform.GetLocation()));
+			}
+
+			Node = Node->PreviousNodeIndex != INDEX_NONE ? &SearchState.NodeList[Node->PreviousNodeIndex] : nullptr;
+		}
+		while (Node);
+
+		// Reverse to get it in the correct travel order.
+		Algo::Reverse(OutPath);
+	}
+}
+
+namespace Cost
+{
+	double CalculateCost_EuclideanDistance(const double PreviousNodeCost, const double DistanceToPreviousNodeSquared)
+	{
+		return PreviousNodeCost + FMath::Sqrt(DistanceToPreviousNodeSquared);
+	}
+}
+
+/** Note: In order for the path to be optimal, the heuristic cost must always be less than or equal to the actual cost. */
+namespace Heuristic
+{
+	double CalculateHeuristic_EuclideanDistance(const FVector& CurrentLocation, const FVector& GoalLocation)
+	{
+		return FVector::Dist(CurrentLocation, GoalLocation);
+	}
+}
+
+/** Initialize the search state for AStar. Must be called before ExecuteSearchIteration. */
+void Initialize(const UPCGPointData* const PointData, const FSearchSettings& Settings, FSearchState& OutSearchState)
+{
+	check(PointData && !PointData->IsEmpty());
+
+	// Emplace the starting node on the open list. The costs will be 0 at the starting point.
+	OutSearchState.NodeList.Reset();
+	OutSearchState.NodeList.Emplace(&Settings.StartPoint, /*InParent=*/INDEX_NONE, /*InCost=*/0.0, /*HeuristicCost=*/0.0);
+	OutSearchState.OpenIndexList.Reset();
+	// Okay to add the first index this way, since it will be "heapified" by default with one element.
+	OutSearchState.OpenIndexList.Add(0);
+	OutSearchState.ClosedIndexList.Reset();
+}
+
+/** Runs a single iteration of the A* algorithm. Intended to be called multiple times, whereupon it will return true when the algorithm is finished. */
+bool ExecuteSearchIteration(const FSearchSettings& SearchSettings, FSearchState& SearchState, TArray<FPCGPoint>& OutPath)
+{
+	TRACE_CPUPROFILER_EVENT_SCOPE(FPCGPathfindingElement::Algorithm::AStar);
+	check(SearchState.OriginatingPointData);
+
+	// The starting node should be guaranteed by the caller.
+	if (!ensure(!SearchState.OpenIndexList.IsEmpty()))
+	{
+		return true;
+	}
+
+	TArray<FNode, TInlineAllocator<FSearchState::PreAllocNodeCount>>& NodeList = SearchState.NodeList;
+	// Get the lowest cost point on the list, which has been binary heap sorted.
+	int32 CurrentNodeIndex;
+	SearchState.OpenIndexList.HeapPop(CurrentNodeIndex, [&SearchState](int32 Index1, int32 Index2) { return Helpers::CompareNodes(SearchState.NodeList[Index1], SearchState.NodeList[Index2]); }, EAllowShrinking::No);
+	check(NodeList[CurrentNodeIndex].PCGPoint);
+
+	// Arrived at the goal. Add the goal to the path and return it. TODO: Will need to be updated when cost functions consider more than just distance.
+	if (FVector::DistSquared(NodeList[CurrentNodeIndex].PCGPoint->Transform.GetLocation(), SearchSettings.Goal) < SearchSettings.SearchDistance * SearchSettings.SearchDistance)
+	{
+		const FPCGPoint GoalPoint(FTransform(SearchSettings.Goal), /*InDensity=*/1, PCGHelpers::ComputeSeedFromPosition(SearchSettings.Goal));
+		const FNode GoalNode(&GoalPoint, CurrentNodeIndex, NodeList[CurrentNodeIndex].EstimatedGoalCost, /*bInHeuristicCost=*/0);
+		Helpers::BuildFinalPath(SearchState, GoalNode, SearchSettings.bCopyOriginatingPoints, OutPath);
+		return true;
+	}
+
+	// Gather neighbors within the search radius.
+	UPCGOctreeQueries::ForEachPointInsideSphere(
+			SearchState.OriginatingPointData,
+			NodeList[CurrentNodeIndex].PCGPoint->Transform.GetLocation(),
+			SearchSettings.SearchDistance,
+			[&NodeList, CurrentNodeIndex, &SearchSettings, &SearchState](const FPCGPointRef& PointRef, const double DistanceToPointSquared)
+			{
+				const bool bTrackedPoint = SearchState.PointToNodeIndexMap.Contains(PointRef.Point);
+				// Node has already been ruled out.
+				if (bTrackedPoint && SearchState.ClosedIndexList.Contains(SearchState.PointToNodeIndexMap[PointRef.Point]))
+				{
+					return;
+				}
+
+				const double TentativeNewLocalCost = Cost::CalculateCost_EuclideanDistance(NodeList[CurrentNodeIndex].LocalCost, DistanceToPointSquared);
+
+				// Not tracking this point yet--add it to the node list and map it.
+				if (!bTrackedPoint)
+				{
+					const double HeuristicCost = SearchSettings.HeuristicWeight * Heuristic::CalculateHeuristic_EuclideanDistance(PointRef.Point->Transform.GetLocation(), SearchSettings.Goal);
+					NodeList.Emplace(PointRef.Point, CurrentNodeIndex, TentativeNewLocalCost, TentativeNewLocalCost + HeuristicCost);
+					const int32 NewNodeIndex = NodeList.Num() - 1;
+					SearchState.OpenIndexList.HeapPush(NewNodeIndex, [&SearchState](int32 Index1, int32 Index2) { return Helpers::CompareNodes(SearchState.NodeList[Index1], SearchState.NodeList[Index2]); });
+					SearchState.PointToNodeIndexMap.Emplace(PointRef.Point, NewNodeIndex);
+
+					return;
+				}
+
+				// Check if the path to this neighbor is a better path than its current one. If so, update accordingly.
+				const int32 NeighborIndex = SearchState.PointToNodeIndexMap[PointRef.Point];
+				FNode& Neighbor = SearchState.NodeList[NeighborIndex];
+				if (TentativeNewLocalCost < Neighbor.LocalCost)
+				{
+					check(SearchState.OpenIndexList.Contains(NeighborIndex));
+
+					const double HeuristicCost = SearchSettings.HeuristicWeight * Heuristic::CalculateHeuristic_EuclideanDistance(PointRef.Point->Transform.GetLocation(), SearchSettings.Goal);
+					Neighbor.PreviousNodeIndex = CurrentNodeIndex;
+					Neighbor.LocalCost = TentativeNewLocalCost;
+					Neighbor.EstimatedGoalCost = TentativeNewLocalCost + HeuristicCost;
+				}
+			});
+
+	// This node has been completely evaluated.
+	SearchState.ClosedIndexList.Add(CurrentNodeIndex);
+
+	// Final point to check, no path can be found.
+	if (SearchState.OpenIndexList.IsEmpty())
+	{
+		OutPath.Reset();
+
+		if (SearchSettings.bAcceptPartialPath)
+		{
+			auto DistanceSquaredToGoal = [&NodeList = SearchState.NodeList, &Goal = SearchSettings.Goal](const int32 NodeIndex)
+			{
+				return FVector::DistSquared(NodeList[NodeIndex].PCGPoint->Transform.GetLocation(), Goal);
+			};
+
+			if (const int32* MaxNodeIndex = Algo::MinElementBy(SearchState.ClosedIndexList, DistanceSquaredToGoal))
+			{
+				Helpers::BuildFinalPath(SearchState, SearchState.NodeList[*MaxNodeIndex], SearchSettings.bCopyOriginatingPoints, OutPath);
+			}
+		}
+
+		return true;
+	}
+
+	return false;
+}
+} // namespace PCGSpatialAlgo::AStar
diff --git a/Engine/Plugins/PCG/Source/PCG/Public/Elements/PCGPathfindingElement.h b/Engine/Plugins/PCG/Source/PCG/Public/Elements/PCGPathfindingElement.h
new file mode 100644
--- /dev/null
+++ b/Engine/Plugins/PCG/Source/PCG/Public/Elements/PCGPathfindingElement.h
@@ -0,0 +1,86 @@
+// Copyright Epic Games, Inc. All Rights Reserved.
+
+#pragma once
+
+#include "PCGSettings.h"
+
+#include "Elements/PCGTimeSlicedElementBase.h"
+#include "SpatialAlgo/PCGAStar.h"
+
+#include "PCGPathfindingElement.generated.h"
+
+class UPCGPointData;
+
+UENUM(BlueprintType, Blueprintable)
+enum class EPCGPathfindingSplineMode : uint8
+{
+	Curve UMETA(Tooltip = "Interpret the spline as a continuous curve."),
+	Linear UMETA(Tooltip = "Interpret the spline as a conjunction of linear segments."),
+};
+
+/** Finds the optimal path across the points of a given point cloud--should one exist--when provided a start and goal
+ * location, and a maximum jump distance between points. Can return a partial path.
+ */
+UCLASS(MinimalAPI, BlueprintType, ClassGroup = (Procedural))
+class UPCGPathfindingSettings : public UPCGSettings
+{
+	GENERATED_BODY()
+
+public:
+	//~Begin UPCGSettings interface
+#if WITH_EDITOR
+	virtual FName GetDefaultNodeName() const override { return FName(TEXT("PathfindingElement")); }
+	virtual FText GetDefaultNodeTitle() const override { return NSLOCTEXT("PCGPathfindingElement", "NodeTitle", "Pathfinding"); }
+	virtual EPCGSettingsType GetType() const override { return EPCGSettingsType::Generic; }
+	virtual FText GetNodeTooltipText() const override { return NSLOCTEXT("PCGPathfindingElement", "NodeTooltip", "Finds the optimal path across the points of a given point cloud--should one exist--when provided a start and goal location, and a maximum jump distance between points. Can return a partial path."); }
+#endif // WITH_EDITOR
+
+protected:
+	virtual TArray<FPCGPinProperties> InputPinProperties() const override;
+	virtual TArray<FPCGPinProperties> OutputPinProperties() const override;
+	virtual FPCGElementPtr CreateElement() const override;
+#if WITH_EDITOR
+	virtual EPCGChangeType GetChangeTypeForProperty(const FName& InPropertyName) const override;
+#endif
+	//~End UPCGSettings interface
+
+public:
+	/** The max distance from each point to search for the next viable point in the path. */
+	UPROPERTY(EditAnywhere, BlueprintReadWrite, Category = Settings, meta = (PCG_Overridable))
+	double SearchDistance = 1000;
+
+	/** The path's starting location. */
+	UPROPERTY(EditAnywhere, BlueprintReadWrite, Category = Settings, meta = (PCG_Overridable))
+	FVector Start = FVector::ZeroVector;
+
+	/** The location the pathfinding should attempt to reach. */
+	UPROPERTY(EditAnywhere, BlueprintReadWrite, Category = Settings, meta = (PCG_Overridable))
+	FVector Goal = FVector::ZeroVector;
+
+	/** The heuristic estimates a faster path to speed up processing. A lower heuristic weight can be faster, but it may cease being the optimal path. A weight of 0 is essentially flood fill. */
+	UPROPERTY(EditAnywhere, BlueprintReadWrite, Category = Settings, meta = (EditCondition = "Algorithm == EPCGPathfindingAlgorithm::AStar", PCG_Overridable))
+	double HeuristicWeight = 1.0;
+
+	/** Even if the path is not complete, return the most optimal and viable partial path to the goal. */
+	UPROPERTY(EditAnywhere, BlueprintReadWrite, Category = Settings, meta = (PCG_Overridable))
+	bool bAcceptPartialPath = true;
+
+	/** The final path will be a spline. If false, the final path will be an ordered point data. */
+	UPROPERTY(EditAnywhere, BlueprintReadWrite, Category = Settings)
+	bool bOutputAsSpline = true;
+
+	/** Determines how the output spline's curves will be calculated. */
+	UPROPERTY(EditAnywhere, BlueprintReadWrite, Category = Settings, meta = (EditCondition = "bOutputAsSpline", EditConditionHides, PCG_Overridable))
+	EPCGPathfindingSplineMode SplineMode = EPCGPathfindingSplineMode::Curve;
+
+	/** Copy the properties and attributes from the originating point input to the output points. */
+	UPROPERTY(EditAnywhere, BlueprintReadWrite, Category = Settings, meta = (EditCondition = "!bOutputAsSpline", EditConditionHides, PCG_Overridable))
+	bool bCopyOriginatingPoints = false;
+};
+
+class FPCGPathfindingElement : public TPCGTimeSlicedElementBase<PCGSpatialAlgo::AStar::FSearchSettings, PCGSpatialAlgo::AStar::FSearchState>
+{
+protected:
+	virtual bool PrepareDataInternal(FPCGContext* InContext) const override;
+	virtual bool ExecuteInternal(FPCGContext* InContext) const override;
+};
diff --git a/Engine/Plugins/PCG/Source/PCG/Public/SpatialAlgo/PCGAStar.h b/Engine/Plugins/PCG/Source/PCG/Public/SpatialAlgo/PCGAStar.h
new file mode 100644
--- /dev/null
+++ b/Engine/Plugins/PCG/Source/PCG/Public/SpatialAlgo/PCGAStar.h
@@ -0,0 +1,85 @@
+// Copyright Epic Games, Inc. All Rights Reserved.
+
+#pragma once
+
+#include "PCGPoint.h"
+#include "Math/Vector.h"
+
+class UPCGPointData;
+
+namespace PCGSpatialAlgo::AStar
+{
+struct FNode
+{
+	explicit FNode(const FPCGPoint* InPCGPoint, const int32 InPreviousNodeIndex, const double InLocalCost, const double InEstimatedGoalCost)
+		: PCGPoint(InPCGPoint)
+		, PreviousNodeIndex(InPreviousNodeIndex)
+		, LocalCost(InLocalCost)
+		, EstimatedGoalCost(InEstimatedGoalCost)
+	{}
+
+	const FPCGPoint* PCGPoint = nullptr;
+	/** The currently assigned parent node with the least cost. */
+	int32 PreviousNodeIndex = INDEX_NONE;
+	/** Concrete cost from start to this node. */
+	double LocalCost = 0.0;
+	/** Current total cost when heuristic is added to the local cost. */
+	double EstimatedGoalCost = 0.0;
+};
+
+struct FSearchSettings
+{
+	/** Cached as an initial point in an otherwise const point data. */
+	FPCGPoint StartPoint;
+	/** The max distance from each point to search for the next viable point in the path. */
+	double SearchDistance = 1000;
+	/** The path's starting location. */
+	FVector Start = FVector::ZeroVector;
+	/** The location the pathfinding should attempt to reach. */
+	FVector Goal = FVector::ZeroVector;
+	/** The heuristic estimates a faster path to speed up processing. A lower heuristic weight can be faster, but it may cease being the optimal path. A weight of 0 is essentially flood fill. */
+	double HeuristicWeight = 1.0;
+	/** Even if the path is not complete, return the most optimal and viable partial path to the goal. */
+	bool bAcceptPartialPath = true;
+	/** Copy the point data from the originating points. */
+	bool bCopyOriginatingPoints = false;
+};
+
+struct FSearchState
+{
+	static uint32 constexpr PreAllocNodeCount = 1024u;
+	// Keep track of the point data for the octree search.
+	const UPCGPointData* OriginatingPointData = nullptr;
+
+	// Storage of the current known node. Source of truth for the nodes, so it should not be resized after its initial reserve.
+	TArray<FNode, TInlineAllocator<PreAllocNodeCount>> NodeList;
+	// Used as a priority list of nodes, sorted by cost.
+	TArray<int32, TInlineAllocator<PreAllocNodeCount>> OpenIndexList;
+	// Nodes that are no longer viable.
+	TArray<int32, TInlineAllocator<PreAllocNodeCount>> ClosedIndexList;
+	// To track nodes to their corresponding points.
+	TMap<const FPCGPoint*, int32> PointToNodeIndexMap;
+};
+
+namespace Helpers
+{
+	bool CompareNodes(const FNode& Node1, const FNode& Node2);
+	void BuildFinalPath(const FSearchState& SearchState, const FNode& FinalNode, bool bCopyOriginatingPoints, TArray<FPCGPoint>& OutPath);
+}
+
+namespace Cost
+{
+	double CalculateCost_EuclideanDistance(const double PreviousNodeCost, const double DistanceToPreviousNodeSquared);
+} // namespace Cost
+
+namespace Heuristic
+{
+	double CalculateHeuristic_EuclideanDistance(const FVector& CurrentLocation, const FVector& GoalLocation);
+} // namespace Heuristic
+
+/** Initialize the search state for AStar. Must be called before ExecuteSearchIteration. */
+void Initialize(const UPCGPointData* const PointData, const FSearchSettings& Settings, FSearchState& OutSearchState);
+
+/** Runs a single iteration of the A* algorithm. Intended to be called multiple times, whereupon it will return true when the algorithm is finished. */
+bool ExecuteSearchIteration(const FSearchSettings& SearchSettings, FSearchState& SearchState, TArray<FPCGPoint>& OutPath);
+}
