[PCG] Fix breakage to IsPointRemoved introduced by refactor.

Will update test level to make failure clearer.

#rnx

diff --git a/Engine/Plugins/PCG/Shaders/Private/Elements/PCGCopyPoints.usf b/Engine/Plugins/PCG/Shaders/Private/Elements/PCGCopyPoints.usf
--- a/Engine/Plugins/PCG/Shaders/Private/Elements/PCGCopyPoints.usf
+++ b/Engine/Plugins/PCG/Shaders/Private/Elements/PCGCopyPoints.usf
@@ -206,7 +206,7 @@ void Main(uint3 GroupId : SV_GroupID, uint GroupIndex : SV_GroupIndex)
 	}
 
 	// Propagate the 'Removed' status to the output point if either source or target point is removed. Otherwise, this point will not be culled.
-	if (!Source_IsPointRemoved(SourceDataIndex, SourcePointIndexInData) || !Target_IsPointRemoved(TargetDataIndex, TargetPointIndexInData))
+	if (Source_IsPointRemoved(SourceDataIndex, SourcePointIndexInData) || Target_IsPointRemoved(TargetDataIndex, TargetPointIndexInData))
 	{
 		Out_RemovePoint(OutDataIndex, OutElemIndex);
 		return;
diff --git a/Engine/Plugins/PCG/Shaders/Private/PCGDataCollectionDataInterface.ush b/Engine/Plugins/PCG/Shaders/Private/PCGDataCollectionDataInterface.ush
--- a/Engine/Plugins/PCG/Shaders/Private/PCGDataCollectionDataInterface.ush
+++ b/Engine/Plugins/PCG/Shaders/Private/PCGDataCollectionDataInterface.ush
@@ -393,7 +393,7 @@ float4x4 GetPointTransform_{DataInterfaceName}(uint InDataIndex, uint InElementI
 
 bool IsPointRemoved_{DataInterfaceName}(uint InDataIndex, uint InElementIndex)
 {
-	return GetDensity_{DataInterfaceName}(InDataIndex, InElementIndex) != INVALID_DENSITY;
+	return GetDensity_{DataInterfaceName}(InDataIndex, InElementIndex) == INVALID_DENSITY;
 }
 
 // #################### POINT ATTRIBUTE SETTERS ##########################
diff --git a/Engine/Plugins/PCG/Source/PCG/Private/Compute/Elements/PCGCustomHLSL.cpp b/Engine/Plugins/PCG/Source/PCG/Private/Compute/Elements/PCGCustomHLSL.cpp
--- a/Engine/Plugins/PCG/Source/PCG/Private/Compute/Elements/PCGCustomHLSL.cpp
+++ b/Engine/Plugins/PCG/Source/PCG/Private/Compute/Elements/PCGCustomHLSL.cpp
@@ -1498,7 +1498,7 @@ FString UPCGCustomHLSLSettings::GetCookedKernelSource(const TMap<FPCGKernelAttri
 
 			// If input point is invalid, mark output point as invalid and abort.
 			KernelSpecificPreamble += FString::Format(TEXT(
-				"    if (!{0}_IsPointRemoved({0}_DataIndex, ElementIndex))\n"
+				"    if ({0}_IsPointRemoved({0}_DataIndex, ElementIndex))\n"
 				"    {\n"
 				"        {1}_RemovePoint({1}_DataIndex, ElementIndex);\n"
 				"        return;\n"
