[PCG] Expose log + Volume sampling methods

#rnx

diff --git a/Engine/Plugins/PCG/Source/PCG/Public/Elements/PCGVolumeSampler.h b/Engine/Plugins/PCG/Source/PCG/Public/Elements/PCGVolumeSampler.h
--- a/Engine/Plugins/PCG/Source/PCG/Public/Elements/PCGVolumeSampler.h
+++ b/Engine/Plugins/PCG/Source/PCG/Public/Elements/PCGVolumeSampler.h
@@ -45,10 +45,10 @@ namespace PCGVolumeSampler
 	};
 
 	/** Sample a volume and returns the resulting point data. */
-	UPCGPointData* SampleVolume(FPCGContext* Context, const FVolumeSamplerParams& SamplerSettings, const UPCGSpatialData* Volume, const UPCGSpatialData* BoundingShape = nullptr);
+	PCG_API UPCGPointData* SampleVolume(FPCGContext* Context, const FVolumeSamplerParams& SamplerSettings, const UPCGSpatialData* Volume, const UPCGSpatialData* BoundingShape = nullptr);
 
 	/** Sample a volume and write the results in the given point data. Can be timesliced and will return false if the processing is not done, true otherwise. */
-	bool SampleVolume(FPCGContext* Context, const FVolumeSamplerParams& SamplerSettings, const UPCGSpatialData* Volume, const UPCGSpatialData* BoundingShape, UPCGPointData* OutputData, const bool bTimeSlicingIsEnabled = false);
+	PCG_API bool SampleVolume(FPCGContext* Context, const FVolumeSamplerParams& SamplerSettings, const UPCGSpatialData* Volume, const UPCGSpatialData* BoundingShape, UPCGPointData* OutputData, const bool bTimeSlicingIsEnabled = false);
 }
 
 UCLASS(MinimalAPI, BlueprintType, ClassGroup = (Procedural))
diff --git a/Engine/Plugins/PCG/Source/PCG/Public/Utils/PCGLogErrors.h b/Engine/Plugins/PCG/Source/PCG/Public/Utils/PCGLogErrors.h
--- a/Engine/Plugins/PCG/Source/PCG/Public/Utils/PCGLogErrors.h
+++ b/Engine/Plugins/PCG/Source/PCG/Public/Utils/PCGLogErrors.h
@@ -19,31 +19,31 @@ namespace PCGLog
 	{
 		namespace Format
 		{
-			extern const FTextFormat TypedInputNotFound;
-			extern const FTextFormat FirstInputOnly;
-			extern const FText InvalidInputData;
+			PCG_API extern const FTextFormat TypedInputNotFound;
+			PCG_API extern const FTextFormat FirstInputOnly;
+			PCG_API extern const FText InvalidInputData;
 		}
 
 		// Warnings
-		void LogTypedDataNotFoundWarning(EPCGDataType DataType, const FName PinLabel, const FPCGContext* InContext = nullptr);
-		void LogFirstInputOnlyWarning(const FName PinLabel, const FPCGContext* InContext = nullptr);
+		PCG_API void LogTypedDataNotFoundWarning(EPCGDataType DataType, const FName PinLabel, const FPCGContext* InContext = nullptr);
+		PCG_API void LogFirstInputOnlyWarning(const FName PinLabel, const FPCGContext* InContext = nullptr);
 
 		// Errors
-		void LogInvalidInputDataError(const FPCGContext* InContext = nullptr);
-		void LogInvalidCardinalityError(const FName SourcePinLabel, const FName TargetPinLabel, const FPCGContext* InContext = nullptr);
+		PCG_API void LogInvalidInputDataError(const FPCGContext* InContext = nullptr);
+		PCG_API void LogInvalidCardinalityError(const FName SourcePinLabel, const FName TargetPinLabel, const FPCGContext* InContext = nullptr);
 	}
 
 	namespace Metadata
 	{
 		namespace Format
 		{
-			extern const FTextFormat CreateAttributeFailure;
-			extern const FTextFormat GetTypedAttributeFailure;
-			extern const FTextFormat GetTypedAttributeFailureNoAccessor;
+			PCG_API extern const FTextFormat CreateAttributeFailure;
+			PCG_API extern const FTextFormat GetTypedAttributeFailure;
+			PCG_API extern const FTextFormat GetTypedAttributeFailureNoAccessor;
 		}
 
 		// Errors
-		void LogFailToCreateAccessorError(const FPCGAttributePropertySelector& Selector, const FPCGContext* InContext = nullptr);
+		PCG_API void LogFailToCreateAccessorError(const FPCGAttributePropertySelector& Selector, const FPCGContext* InContext = nullptr);
 
 		template <typename T>
 		void LogFailToCreateAttributeError(FText AttributeName, const FPCGContext* InContext = nullptr)
@@ -57,10 +57,10 @@ namespace PCGLog
 			LogFailToCreateAttributeError<T>(FText::FromName(AttributeName), InContext);
 		}
 
-		void LogFailToGetAttributeError(FText AttributeName, const FPCGContext* InContext = nullptr);
-		void LogFailToGetAttributeError(FName AttributeName, const FPCGContext* InContext = nullptr);
-		void LogFailToGetAttributeError(const FPCGAttributePropertySelector& Selector, const FPCGContext* InContext = nullptr);
-		void LogIncomparableAttributesError(const FPCGAttributePropertySelector& FirstSelector, const FPCGAttributePropertySelector& SecondSelector, const FPCGContext* InContext = nullptr);
+		PCG_API void LogFailToGetAttributeError(FText AttributeName, const FPCGContext* InContext = nullptr);
+		PCG_API void LogFailToGetAttributeError(FName AttributeName, const FPCGContext* InContext = nullptr);
+		PCG_API void LogFailToGetAttributeError(const FPCGAttributePropertySelector& Selector, const FPCGContext* InContext = nullptr);
+		PCG_API void LogIncomparableAttributesError(const FPCGAttributePropertySelector& FirstSelector, const FPCGAttributePropertySelector& SecondSelector, const FPCGContext* InContext = nullptr);
 
 		template <typename T>
 		void LogFailToGetAttributeError(FText AttributeName, const IPCGAttributeAccessor* Accessor, const FPCGContext* InContext = nullptr)
@@ -91,10 +91,10 @@ namespace PCGLog
 	namespace Parsing
 	{
 		// Warnings
-		void LogEmptyExpressionWarning(const FPCGContext* InContext = nullptr);
+		PCG_API void LogEmptyExpressionWarning(const FPCGContext* InContext = nullptr);
 
 		// Errors
-		void LogInvalidCharacterInParsedStringError(const FStringView& ParsedString, const FPCGContext* InContext = nullptr);
-		void LogInvalidExpressionInParsedStringError(const FStringView& ParsedString, const FPCGContext* InContext = nullptr);
+		PCG_API void LogInvalidCharacterInParsedStringError(const FStringView& ParsedString, const FPCGContext* InContext = nullptr);
+		PCG_API void LogInvalidExpressionInParsedStringError(const FStringView& ParsedString, const FPCGContext* InContext = nullptr);
 	}
 }
