Fixed off-by-one index validation error

#rnx

diff --git a/Engine/Plugins/PCG/Source/PCG/Private/Elements/Metadata/PCGMetadataOpElementBase.cpp b/Engine/Plugins/PCG/Source/PCG/Private/Elements/Metadata/PCGMetadataOpElementBase.cpp
--- a/Engine/Plugins/PCG/Source/PCG/Private/Elements/Metadata/PCGMetadataOpElementBase.cpp
+++ b/Engine/Plugins/PCG/Source/PCG/Private/Elements/Metadata/PCGMetadataOpElementBase.cpp
@@ -300,7 +300,7 @@ void FPCGMetadataElementBase::PassthroughInput(FPCGContext* Context, TArray<FPCG
 	}
 
 	// Take the index of the iteration, except for the 1:N case, where we just grab the first index
-	const int32 AdjustedIndex = (Index <= InputsToForward.Num()) ? Index : 0;
+	const int32 AdjustedIndex = (Index < InputsToForward.Num()) ? Index : 0;
 
 	// Passthrough this single input to all of the outputs
 	for (uint32 I = 0; I < NumberOfOutputs; ++I)
