[PCG] Update to Select Grammar node allowing for non-existant attributes to act as a match and set

#rnx

diff --git a/Engine/Plugins/PCG/Source/PCG/Private/Elements/Grammar/PCGSelectGrammar.cpp b/Engine/Plugins/PCG/Source/PCG/Private/Elements/Grammar/PCGSelectGrammar.cpp
--- a/Engine/Plugins/PCG/Source/PCG/Private/Elements/Grammar/PCGSelectGrammar.cpp
+++ b/Engine/Plugins/PCG/Source/PCG/Private/Elements/Grammar/PCGSelectGrammar.cpp
@@ -10,63 +10,62 @@
 
 #define LOCTEXT_NAMESPACE "PCGSelectGrammarElement"
 
-namespace PCGSelectGrammarConstants
+namespace PCGSelectGrammar
+{
+namespace Constants
 {
 	const FName CriteriaDataPinLabel = TEXT("Selection Criteria Data");
 	const FName OutputGrammarAttributeName = TEXT("Grammar");
 	// Most likely used value for comparison with grammars
 	const FString ComparedValueAttributeString = TEXT("$ScaledLocalSize.X");
+
+	constexpr uint16 BinaryEnumIndex = static_cast<uint16>(EPCGSelectGrammarComparator::BinaryOps) + 1;
+	constexpr uint16 TernaryEnumIndex = static_cast<uint16>(EPCGSelectGrammarComparator::TernaryOps) + 1;
 }
 
-namespace PCGSelectGrammar
+struct FGrammarSizePair
 {
-	struct FGrammarSizePair
-	{
-		FGrammarSizePair(const FString& InGrammar, double InSize)
-			: Grammar(InGrammar)
-			, Size(InSize)
-		{
-		}
-
-		FString Grammar;
-		double Size;
-	};
+	FGrammarSizePair(const FString& InGrammar, double InSize)
+		: Grammar(InGrammar)
+		, Size(InSize)
+	{}
 
-	constexpr uint16 BinaryEnumIndex = static_cast<uint16>(EPCGSelectGrammarComparator::BinaryOps) + 1;
-	constexpr uint16 TernaryEnumIndex = static_cast<uint16>(EPCGSelectGrammarComparator::TernaryOps) + 1;
+	FString Grammar;
+	double Size;
+};
 
-	// Function indices must match EPCGSelectGrammarComparator indices directly
-	template <typename T>
-	using TBinarySignature = bool(*)(const T&, const T&);
+// Function indices must match EPCGSelectGrammarComparator indices directly
+template <typename T>
+using TBinarySignature = bool(*)(const T&, const T&);
 
-	template <typename T>
-	static constexpr TBinarySignature<T> BinaryCompare[] =
-	{
-		[](const T&, const T&) { return true; }, // Select
-		PCG::Private::MetadataTraits<T>::Less,
-		PCG::Private::MetadataTraits<T>::LessOrEqual,
-		PCG::Private::MetadataTraits<T>::Equal,
-		PCG::Private::MetadataTraits<T>::GreaterOrEqual,
-		PCG::Private::MetadataTraits<T>::Greater,
-	};
-
-	template <typename T>
-	using TTernarySignature = bool(*)(const T&, const T&, const T&);
-
-	template <typename T>
-	static constexpr TTernarySignature<T> TernaryCompare[] =
-	{
-		// Exclusive range
-		[](const T& Value, const T& Start, const T& End) { return PCG::Private::MetadataTraits<T>::Greater(Value, Start) && PCG::Private::MetadataTraits<T>::Less(Value, End); },
-		// Inclusive range
-		[](const T& Value, const T& Start, const T& End) { return PCG::Private::MetadataTraits<T>::GreaterOrEqual(Value, Start) && PCG::Private::MetadataTraits<T>::LessOrEqual(Value, End); }
-	};
+template <typename T>
+static constexpr TBinarySignature<T> BinaryCompare[] =
+{
+	[](const T&, const T&) { return true; }, // Select
+	PCG::Private::MetadataTraits<T>::Less,
+	PCG::Private::MetadataTraits<T>::LessOrEqual,
+	PCG::Private::MetadataTraits<T>::Equal,
+	PCG::Private::MetadataTraits<T>::GreaterOrEqual,
+	PCG::Private::MetadataTraits<T>::Greater,
+};
+
+template <typename T>
+using TTernarySignature = bool(*)(const T&, const T&, const T&);
+
+template <typename T>
+static constexpr TTernarySignature<T> TernaryCompare[] =
+{
+	// Exclusive range
+	[](const T& Value, const T& Start, const T& End) { return PCG::Private::MetadataTraits<T>::Greater(Value, Start) && PCG::Private::MetadataTraits<T>::Less(Value, End); },
+	// Inclusive range
+	[](const T& Value, const T& Start, const T& End) { return PCG::Private::MetadataTraits<T>::GreaterOrEqual(Value, Start) && PCG::Private::MetadataTraits<T>::LessOrEqual(Value, End); }
+};
 }
 
 UPCGSelectGrammarSettings::UPCGSelectGrammarSettings()
 {
-	ComparedValueAttribute.Update(PCGSelectGrammarConstants::ComparedValueAttributeString);
-	OutputGrammarAttribute.SetAttributeName(PCGSelectGrammarConstants::OutputGrammarAttributeName);
+	ComparedValueAttribute.Update(PCGSelectGrammar::Constants::ComparedValueAttributeString);
+	OutputGrammarAttribute.SetAttributeName(PCGSelectGrammar::Constants::OutputGrammarAttributeName);
 }
 
 #if WITH_EDITOR
@@ -125,7 +124,7 @@ TArray<FPCGPinProperties> UPCGSelectGrammarSettings::InputPinProperties() const
 
 	if (bCriteriaAsInput)
 	{
-		FPCGPinProperties& CriteriaPin = PinProperties.Emplace_GetRef(PCGSelectGrammarConstants::CriteriaDataPinLabel, EPCGDataType::Param, false, false);
+		FPCGPinProperties& CriteriaPin = PinProperties.Emplace_GetRef(PCGSelectGrammar::Constants::CriteriaDataPinLabel, EPCGDataType::Param, false, false);
 		CriteriaPin.SetRequiredPin();
 #if WITH_EDITOR
 		CriteriaPin.Tooltip = LOCTEXT("CriteriaPinTooltip", "Criteria to compare against an input value for conditionally selecting a grammar. Will be evaluated in sequential order.");
@@ -151,19 +150,19 @@ bool FPCGSelectGrammarElement::ExecuteInternal(FPCGContext* InContext) const
 	check(Settings);
 
 	const TArray<FPCGTaggedData> PointInputs = InContext->InputData.GetInputsByPin(PCGPinConstants::DefaultInputLabel);
-	const TArray<FPCGTaggedData> CriteriaDataInputs = InContext->InputData.GetInputsByPin(PCGSelectGrammarConstants::CriteriaDataPinLabel);
+	const TArray<FPCGTaggedData> CriteriaDataInputs = InContext->InputData.GetInputsByPin(PCGSelectGrammar::Constants::CriteriaDataPinLabel);
 
-	if (PointInputs.IsEmpty() || (Settings->bCriteriaAsInput && CriteriaDataInputs.IsEmpty()))
+	if (PointInputs.IsEmpty())
 	{
 		return true;
 	}
 
 	TArray<FPCGSelectGrammarCriterion> Criteria = Settings->bCriteriaAsInput ? TArray<FPCGSelectGrammarCriterion>{} : Settings->Criteria;
-	if (Settings->bCriteriaAsInput)
+	if (Settings->bCriteriaAsInput && !CriteriaDataInputs.IsEmpty())
 	{
 		if (CriteriaDataInputs.Num() != 1)
 		{
-			PCGLog::InputOutput::LogFirstInputOnlyWarning(PCGSelectGrammarConstants::CriteriaDataPinLabel, InContext);
+			PCGLog::InputOutput::LogFirstInputOnlyWarning(PCGSelectGrammar::Constants::CriteriaDataPinLabel, InContext);
 		}
 
 		// Should only be one, so take the first
@@ -188,7 +187,7 @@ bool FPCGSelectGrammarElement::ExecuteInternal(FPCGContext* InContext) const
 		}
 		else
 		{
-			PCGLog::InputOutput::LogTypedDataNotFoundWarning(EPCGDataType::Param, PCGSelectGrammarConstants::CriteriaDataPinLabel, InContext);
+			PCGLog::InputOutput::LogTypedDataNotFoundWarning(EPCGDataType::Param, PCGSelectGrammar::Constants::CriteriaDataPinLabel, InContext);
 			return true;
 		}
 	}
@@ -196,29 +195,42 @@ bool FPCGSelectGrammarElement::ExecuteInternal(FPCGContext* InContext) const
 	// Early out with nothing to compare to
 	if (Criteria.IsEmpty())
 	{
+		InContext->OutputData = InContext->InputData;
 		return true;
 	}
 
+	// Sum the number of select criteria to check if they are all select.
+	int32 NumSelectCriteria = 0;
 	// Validate the criteria, to ensure the ranges are in order, etc
 	for (const FPCGSelectGrammarCriterion& Criterion : Criteria)
 	{
 		switch (Criterion.Comparator)
 		{
+			case EPCGSelectGrammarComparator::Select:
+			{
+				++NumSelectCriteria;
+			}
+			break;
+
 			case EPCGSelectGrammarComparator::RangeExclusive: // fall-through
 			case EPCGSelectGrammarComparator::RangeInclusive:
+			{
+				if (Criterion.FirstValue > Criterion.SecondValue)
 				{
-					if (Criterion.FirstValue > Criterion.SecondValue)
-					{
-						PCGLog::LogErrorOnGraph(FText::Format(LOCTEXT("InvalidCriteriaRange", "Criteria range is inverted for criterion with Key->Grammar: '{0}'->'{1}'"), FText::FromName(Criterion.Key), FText::FromString(Criterion.Grammar)), InContext);
-						return true;
-					}
+					PCGLog::LogErrorOnGraph(FText::Format(LOCTEXT("InvalidCriteriaRange", "Criteria range is inverted for criterion with Key->Grammar: '{0}'->'{1}'"), FText::FromName(Criterion.Key), FText::FromString(Criterion.Grammar)), InContext);
+					return true;
 				}
-				break;
+			}
+			break;
+
 			default:
 				break;
 		}
 	}
 
+	// In this case, all the comparators are 'Select', so bypass the criteria checks.
+	const bool bBypassCriteria = NumSelectCriteria == Criteria.Num();
+
 	// Pre-process the criteria into a map for O(1) access later
 	TMap<FName, TArray<FPCGSelectGrammarCriterion, TInlineAllocator<16>>> CriteriaMap;
 	for (FPCGSelectGrammarCriterion& Criterion : Criteria)
@@ -231,6 +243,14 @@ bool FPCGSelectGrammarElement::ExecuteInternal(FPCGContext* InContext) const
 		CriteriaMap[Criterion.Key].Emplace(std::move(Criterion));
 	}
 
+	for (const TTuple<FName, TArray<FPCGSelectGrammarCriterion, TSizedInlineAllocator<16, 32>>>& Tuple : CriteriaMap)
+	{
+		if (Tuple.Get<1>().Num() > 1)
+		{
+			PCGLog::LogWarningOnGraph(FText::Format(LOCTEXT("MultipleSelectSameCriteria", "Multiple 'Select' comparators found on the criteria with same key: {0}"), FText::FromName(Tuple.Get<0>())), InContext);
+		}
+	}
+
 	TArray<FPCGTaggedData>& Outputs = InContext->OutputData.TaggedData;
 	for (const FPCGTaggedData& PointInput : PointInputs)
 	{
@@ -243,12 +263,9 @@ bool FPCGSelectGrammarElement::ExecuteInternal(FPCGContext* InContext) const
 		UPCGData* OutputData = InputPointData->DuplicateData(InContext);
 		UPCGPointData* OutputPointData = CastChecked<UPCGPointData>(OutputData);
 		check(OutputPointData->Metadata)
-
 		Outputs.Emplace_GetRef(PointInput).Data = OutputData;
 
 		FPCGAttributePropertyInputSelector KeySelector = Settings->KeyAttribute.CopyAndFixLast(InputPointData);
-		FPCGAttributePropertyInputSelector ComparedValueSelector = Settings->ComparedValueAttribute.CopyAndFixLast(InputPointData);
-
 		TUniquePtr<const IPCGAttributeAccessor> InputKeyAccessor = PCGAttributeAccessorHelpers::CreateConstAccessor(InputPointData, KeySelector);
 		TUniquePtr<const IPCGAttributeAccessorKeys> InputKeys = PCGAttributeAccessorHelpers::CreateConstKeys(InputPointData, KeySelector);
 		if (!InputKeyAccessor || !InputKeys)
@@ -257,16 +274,8 @@ bool FPCGSelectGrammarElement::ExecuteInternal(FPCGContext* InContext) const
 			continue;
 		}
 
-		TUniquePtr<const IPCGAttributeAccessor> ComparedValuesAccessor = PCGAttributeAccessorHelpers::CreateConstAccessor(InputPointData, ComparedValueSelector);
-		if (!ComparedValuesAccessor)
-		{
-			PCGLog::Metadata::LogFailToCreateAccessorError(ComparedValueSelector, InContext);
-			continue;
-		}
-
 		FPCGAttributePropertyOutputSelector OutputGrammarSelector = Settings->OutputGrammarAttribute.CopyAndFixSource(&KeySelector);
 		FPCGMetadataAttribute<FString>* OutputAttribute = OutputPointData->Metadata->FindOrCreateAttribute(OutputGrammarSelector.GetName(), FString{}, /*bAllowsInterpolation=*/false);
-
 		TUniquePtr<IPCGAttributeAccessor> WriteGrammarAccessor = PCGAttributeAccessorHelpers::CreateAccessor(OutputAttribute, OutputPointData->Metadata);
 		TUniquePtr<IPCGAttributeAccessorKeys> OutputKeys = PCGAttributeAccessorHelpers::CreateKeys(OutputPointData, OutputGrammarSelector);
 		if (!WriteGrammarAccessor || !OutputKeys)
@@ -275,9 +284,27 @@ bool FPCGSelectGrammarElement::ExecuteInternal(FPCGContext* InContext) const
 			continue;
 		}
 
-		TArray Accessors = {InputKeyAccessor.Get(), ComparedValuesAccessor.Get()};
-		// TODO: For now, cast everything to double for numerical comparison, but this could easily expand in the future
-		PCGMetadataElementCommon::ApplyOnMultiAccessorsRange<FName, double>(*InputKeys, Accessors, [&WriteGrammarAccessor, &OutputKeys, &CriteriaMap, bPassThroughKey = Settings->bCopyKeyForUnselectedGrammar](const TArrayView<FName>& KeyValues, const TArrayView<double>& CompareValues, int32 Start, int32 Range)
+		auto ProcessWithoutCriteria = [&WriteGrammarAccessor, &OutputKeys, &CriteriaMap, bPassThroughKey = Settings->bCopyKeyForUnselectedGrammar](const TArrayView<FName>& KeyValues, int32 Start, int32 Range)
+		{
+			TArray<FString, TInlineAllocator<PCGMetadataElementCommon::DefaultChunkSize>> Grammars;
+			Grammars.SetNum(Range);
+
+			// Iterate through this chunk of points and evaluate the selection criteria for each point
+			for (int32 i = 0; i < Range; ++i)
+			{
+				const FName& Key = KeyValues[i];
+				if (!CriteriaMap.Contains(Key) || CriteriaMap[Key].IsEmpty())
+				{
+					continue;
+				}
+
+				Grammars[i] = CriteriaMap[Key][0].Grammar;
+			}
+
+			WriteGrammarAccessor->SetRange<FString>(Grammars, Start, *OutputKeys);
+		};
+
+		auto ProcessWithCriteria = [&WriteGrammarAccessor, &OutputKeys, &CriteriaMap, bPassThroughKey = Settings->bCopyKeyForUnselectedGrammar](const TArrayView<FName>& KeyValues, const TArrayView<double>& CompareValues, int32 Start, int32 Range)
 		{
 			TArray<FString, TInlineAllocator<PCGMetadataElementCommon::DefaultChunkSize>> Grammars;
 			Grammars.SetNum(Range);
@@ -288,7 +315,7 @@ bool FPCGSelectGrammarElement::ExecuteInternal(FPCGContext* InContext) const
 				bool bFoundGrammar = false;
 
 				const FName& Key = KeyValues[i];
-				if (!CriteriaMap.Contains(Key))
+				if (!CriteriaMap.Contains(Key) || CriteriaMap[Key].IsEmpty())
 				{
 					continue;
 				}
@@ -300,10 +327,10 @@ bool FPCGSelectGrammarElement::ExecuteInternal(FPCGContext* InContext) const
 
 					// Find the comparison function based on the enum selection
 					const uint16 EnumIndex = static_cast<uint16>(Criterion.Comparator);
-					if (EnumIndex < PCGSelectGrammar::TernaryEnumIndex) // Binary comparison
+					if (EnumIndex < PCGSelectGrammar::Constants::TernaryEnumIndex) // Binary comparison
 					{
 						// Convert enum index to static function array index
-						const uint16 CompIndex = EnumIndex - PCGSelectGrammar::BinaryEnumIndex;
+						const uint16 CompIndex = EnumIndex - PCGSelectGrammar::Constants::BinaryEnumIndex;
 						check(CompIndex < std::size(PCGSelectGrammar::BinaryCompare<double>))
 						if (PCGSelectGrammar::BinaryCompare<double>[CompIndex](CompareValues[i], Criterion.FirstValue))
 						{
@@ -314,7 +341,7 @@ bool FPCGSelectGrammarElement::ExecuteInternal(FPCGContext* InContext) const
 					}
 					else // Ternary comparison
 					{
-						const uint16 CompIndex = EnumIndex - PCGSelectGrammar::TernaryEnumIndex;
+						const uint16 CompIndex = EnumIndex - PCGSelectGrammar::Constants::TernaryEnumIndex;
 						check(CompIndex < std::size(PCGSelectGrammar::TernaryCompare<double>))
 						if (PCGSelectGrammar::TernaryCompare<double>[CompIndex](CompareValues[i], Criterion.FirstValue, Criterion.SecondValue))
 						{
@@ -332,7 +359,27 @@ bool FPCGSelectGrammarElement::ExecuteInternal(FPCGContext* InContext) const
 			}
 
 			WriteGrammarAccessor->SetRange<FString>(Grammars, Start, *OutputKeys);
-		});
+		};
+
+		if (bBypassCriteria)
+		{
+			PCGMetadataElementCommon::ApplyOnAccessorRange<FName>(*InputKeys, *InputKeyAccessor, ProcessWithoutCriteria);
+		}
+		else
+		{
+			FPCGAttributePropertyInputSelector ComparedValueSelector = Settings->ComparedValueAttribute.CopyAndFixLast(InputPointData);
+			TUniquePtr<const IPCGAttributeAccessor> ComparedValuesAccessor = PCGAttributeAccessorHelpers::CreateConstAccessor(InputPointData, ComparedValueSelector);
+
+			if (!ComparedValuesAccessor)
+			{
+				PCGLog::Metadata::LogFailToCreateAccessorError(ComparedValueSelector, InContext);
+				continue;
+			}
+
+			// TODO: For now, cast everything to double for numerical comparison, but this could easily expand in the future
+			TArray<const IPCGAttributeAccessor*> Accessors = {InputKeyAccessor.Get(), ComparedValuesAccessor.Get()};
+			PCGMetadataElementCommon::ApplyOnMultiAccessorsRange<FName, double>(*InputKeys, Accessors, ProcessWithCriteria);
+		}
 	}
 
 	return true;
