[PCG] Fixed PCGLog erroneous template deduction with LogFailToGetAttributeError

#rnx

diff --git a/Engine/Plugins/PCG/Source/PCG/Private/Utils/PCGLogErrors.cpp b/Engine/Plugins/PCG/Source/PCG/Private/Utils/PCGLogErrors.cpp
--- a/Engine/Plugins/PCG/Source/PCG/Private/Utils/PCGLogErrors.cpp
+++ b/Engine/Plugins/PCG/Source/PCG/Private/Utils/PCGLogErrors.cpp
@@ -46,7 +46,7 @@ namespace PCGLog::Metadata
 	{
 		const FTextFormat CreateAccessorFailure = LOCTEXT("CreateAccessorFailure", "Couldn't create accessor. Attribute '{0}' was not found.");
 		const FTextFormat CreateAttributeFailure = LOCTEXT("CreateAttributeFailure", "Couldn't create attribute '{0}' of type '{1}'.");
-		const FTextFormat GetAttributeFailure = LOCTEXT("GetAttributeFailure", "Couldn't retrieve attribute '{0}' value. Expected type: {1}, Actual Type: {2}.");
+		const FTextFormat GetAttributeFailure = LOCTEXT("GetAttributeFailure", "Couldn't retrieve attribute '{0}' value.");
 		const FTextFormat GetTypedAttributeFailure = LOCTEXT("GetTypedAttributeFailure", "Couldn't retrieve attribute '{0}' value. Expected type: {1}, Actual Type: {2}.");
 		const FTextFormat GetTypedAttributeFailureNoAccessor = LOCTEXT("GetTypedAttributeFailureNoAccessor", "Couldn't retrieve attribute '{0}' value of type: '{1}.");
 		const FTextFormat IncomparableTypesFailure = LOCTEXT("IncomparableTypesFailure", "Attributes '{0}' and '{1}' are incomparable. Ensure they are either of the same or compatible types.");
@@ -57,9 +57,9 @@ namespace PCGLog::Metadata
 		LogErrorOnGraph(FText::Format(Format::CreateAccessorFailure, Selector.GetDisplayText()), InContext);
 	}
 
-	void LogFailToGetAttributeError(FText AttributeName, const FPCGContext* InContext)
+	void LogFailToGetAttributeError(const FText& AttributeName, const FPCGContext* InContext)
 	{
-		LogErrorOnGraph(FText::Format(Format::GetAttributeFailure, std::move(AttributeName)), InContext);
+		LogErrorOnGraph(FText::Format(Format::GetAttributeFailure, AttributeName), InContext);
 	}
 
 	void LogFailToGetAttributeError(FName AttributeName, const FPCGContext* InContext)
diff --git a/Engine/Plugins/PCG/Source/PCG/Public/Utils/PCGLogErrors.h b/Engine/Plugins/PCG/Source/PCG/Public/Utils/PCGLogErrors.h
--- a/Engine/Plugins/PCG/Source/PCG/Public/Utils/PCGLogErrors.h
+++ b/Engine/Plugins/PCG/Source/PCG/Public/Utils/PCGLogErrors.h
@@ -46,7 +46,7 @@ namespace PCGLog
 		PCG_API void LogFailToCreateAccessorError(const FPCGAttributePropertySelector& Selector, const FPCGContext* InContext = nullptr);
 
 		template <typename T>
-		void LogFailToCreateAttributeError(FText AttributeName, const FPCGContext* InContext = nullptr)
+		void LogFailToCreateAttributeError(const FText& AttributeName, const FPCGContext* InContext = nullptr)
 		{
 			PCGLog::LogErrorOnGraph(FText::Format(Format::CreateAttributeFailure, AttributeName, PCG::Private::GetTypeNameText<T>()), InContext);
 		}
@@ -57,13 +57,13 @@ namespace PCGLog
 			LogFailToCreateAttributeError<T>(FText::FromName(AttributeName), InContext);
 		}
 
-		PCG_API void LogFailToGetAttributeError(FText AttributeName, const FPCGContext* InContext = nullptr);
+		PCG_API void LogFailToGetAttributeError(const FText& AttributeName, const FPCGContext* InContext = nullptr);
 		PCG_API void LogFailToGetAttributeError(FName AttributeName, const FPCGContext* InContext = nullptr);
 		PCG_API void LogFailToGetAttributeError(const FPCGAttributePropertySelector& Selector, const FPCGContext* InContext = nullptr);
 		PCG_API void LogIncomparableAttributesError(const FPCGAttributePropertySelector& FirstSelector, const FPCGAttributePropertySelector& SecondSelector, const FPCGContext* InContext = nullptr);
 
 		template <typename T>
-		void LogFailToGetAttributeError(FText AttributeName, const IPCGAttributeAccessor* Accessor, const FPCGContext* InContext = nullptr)
+		void LogFailToGetAttributeError(const FText& AttributeName, const IPCGAttributeAccessor* Accessor, const FPCGContext* InContext = nullptr)
 		{
 			if (Accessor)
 			{
