Added cost/heuristic function control in the pathfinding node.

diff --git a/Engine/Plugins/PCG/Source/PCG/Private/Elements/PCGPathfindingElement.cpp b/Engine/Plugins/PCG/Source/PCG/Private/Elements/PCGPathfindingElement.cpp
--- a/Engine/Plugins/PCG/Source/PCG/Private/Elements/PCGPathfindingElement.cpp
+++ b/Engine/Plugins/PCG/Source/PCG/Private/Elements/PCGPathfindingElement.cpp
@@ -6,6 +6,9 @@
 #include "Data/PCGPointData.h"
 #include "Data/PCGSplineData.h"
 #include "Helpers/PCGHelpers.h"
+#include "Metadata/Accessors/IPCGAttributeAccessor.h"
+#include "Metadata/Accessors/PCGAttributeAccessorHelpers.h"
+#include "Metadata/Accessors/PCGAttributeAccessorKeys.h"
 #include "SpatialAlgo/PCGAStar.h"
 
 #include "Components/SplineComponent.h"
@@ -119,6 +122,7 @@ bool FPCGPathfindingElement::PrepareDataInternal(FPCGContext* InContext) const
 		OutState = PCGSpatialAlgo::AStar::FSearchSettings
 		{
 			.StartPoint = FPCGPoint(FTransform(Settings->Start), /*InDensity=*/1, PCGHelpers::ComputeSeedFromPosition(Settings->Start)),
+			.GoalPoint = FPCGPoint(FTransform(Settings->Goal), /*InDensity=*/1.0, PCGHelpers::ComputeSeedFromPosition(Settings->Goal)),
 			.SearchDistance = Settings->SearchDistance,
 			.Start = Settings->Start,
 			.Goal = Settings->Goal,
@@ -130,7 +134,7 @@ bool FPCGPathfindingElement::PrepareDataInternal(FPCGContext* InContext) const
 		return EPCGTimeSliceInitResult::Success;
 	});
 
-	Context->InitializePerIterationStates(PointInputs.Num(), [&PointInputs](IterStateType& OutSearchState, const ExecStateType& SearchSettings, const uint32 IterationIndex)
+	Context->InitializePerIterationStates(PointInputs.Num(), [&PointInputs, Settings, Context](IterStateType& OutSearchState, const ExecStateType& SearchSettings, const uint32 IterationIndex)
 	{
 		TRACE_CPUPROFILER_EVENT_SCOPE(FPCGClusterElement::PrepareData::InitializePerIterationStates);
 
@@ -141,6 +145,63 @@ bool FPCGPathfindingElement::PrepareDataInternal(FPCGContext* InContext) const
 			return EPCGTimeSliceInitResult::NoOperation;
 		}
 
+		// Build cost attribute accessor if required
+		TSharedPtr<const IPCGAttributeAccessor> CostAccessor = nullptr;
+
+		if (Settings->CostFunctionMode != EPCGPathfindingCostFunctionMode::Distance)
+		{
+			const FPCGAttributePropertyInputSelector Selector = Settings->CostAttribute.CopyAndFixLast(OutSearchState.OriginatingPointData);
+			CostAccessor = MakeShareable(PCGAttributeAccessorHelpers::CreateConstAccessor(OutSearchState.OriginatingPointData, Selector).Release());
+
+			FPCGAttributeAccessorKeysPoints Keys(OutSearchState.OriginatingPointData->GetPoints());
+
+			if (!CostAccessor)
+			{
+				PCGLog::Metadata::LogFailToCreateAccessorError(Selector, Context);
+			}
+			else if(!PCG::Private::IsBroadcastableOrConstructible(CostAccessor->GetUnderlyingType(), PCG::Private::MetadataTypes<double>::Id))
+			{
+				PCGLog::Metadata::LogFailToGetAttributeError<double>(Selector, CostAccessor.Get(), Context);
+				CostAccessor = nullptr;
+			}
+		}
+
+		if (CostAccessor)
+		{
+			if (Settings->CostFunctionMode == EPCGPathfindingCostFunctionMode::FitnessScore)
+			{
+				const double MaxFitnessPenaltyFactor = FMath::Max(Settings->MaximumFitnessPenaltyFactor, 1.0);
+
+				OutSearchState.CostFunction = [FitnessAccessor = CostAccessor, MaxFitnessPenaltyFactor](const double PreviousNodeCost, const FPCGPoint* PreviousNodePoint, const double DistanceToCurrentSquared, const FPCGPoint* CurrentNodePoint)
+				{
+					double FitnessScore = 1.0;
+					FPCGAttributeAccessorKeysPoints Key(*CurrentNodePoint);
+
+					FitnessAccessor->Get(FitnessScore, Key, EPCGAttributeAccessorFlags::AllowBroadcastAndConstructible);
+					FitnessScore = FMath::Clamp(FitnessScore, 0.0, 1.0);
+
+					return PreviousNodeCost + (1.0 - FitnessScore) * MaxFitnessPenaltyFactor * FMath::Sqrt(DistanceToCurrentSquared);
+				};
+			}
+			else if (Settings->CostFunctionMode == EPCGPathfindingCostFunctionMode::CostMultipler)
+			{
+				OutSearchState.CostFunction = [MultiplierAccessor = CostAccessor](const double PreviousNodeCost, const FPCGPoint* PreviousNodePoint, const double DistanceToCurrentSquared, const FPCGPoint* CurrentNodePoint)
+				{
+					double Multiplier = 1.0;
+					FPCGAttributeAccessorKeysPoints Key(*CurrentNodePoint);
+
+					MultiplierAccessor->Get(Multiplier, Key, EPCGAttributeAccessorFlags::AllowBroadcastAndConstructible);
+					Multiplier = FMath::Max(Multiplier, 1.0);
+
+					return PreviousNodeCost + Multiplier * FMath::Sqrt(DistanceToCurrentSquared);
+				};
+			}
+			else
+			{
+				checkNoEntry();
+			}
+		}
+
 		PCGSpatialAlgo::AStar::Initialize(OutSearchState.OriginatingPointData, SearchSettings, OutSearchState);
 
 		return EPCGTimeSliceInitResult::Success;
diff --git a/Engine/Plugins/PCG/Source/PCG/Private/SpatialAlgo/PCGAStar.cpp b/Engine/Plugins/PCG/Source/PCG/Private/SpatialAlgo/PCGAStar.cpp
--- a/Engine/Plugins/PCG/Source/PCG/Private/SpatialAlgo/PCGAStar.cpp
+++ b/Engine/Plugins/PCG/Source/PCG/Private/SpatialAlgo/PCGAStar.cpp
@@ -44,7 +44,7 @@ namespace Helpers
 
 namespace Cost
 {
-	double CalculateCost_EuclideanDistance(const double PreviousNodeCost, const double DistanceToPreviousNodeSquared)
+	double CalculateCost_EuclideanDistance(const double PreviousNodeCost, const FPCGPoint* PreviousNodePoint, const double DistanceToPreviousNodeSquared, const FPCGPoint* CurrentPoint)
 	{
 		return PreviousNodeCost + FMath::Sqrt(DistanceToPreviousNodeSquared);
 	}
@@ -91,56 +91,66 @@ bool ExecuteSearchIteration(const FSearchSettings& SearchSettings, FSearchState&
 	SearchState.OpenIndexList.HeapPop(CurrentNodeIndex, [&SearchState](int32 Index1, int32 Index2) { return Helpers::CompareNodes(SearchState.NodeList[Index1], SearchState.NodeList[Index2]); }, EAllowShrinking::No);
 	check(NodeList[CurrentNodeIndex].PCGPoint);
 
-	// Arrived at the goal. Add the goal to the path and return it. TODO: Will need to be updated when cost functions consider more than just distance.
-	if (FVector::DistSquared(NodeList[CurrentNodeIndex].PCGPoint->Transform.GetLocation(), SearchSettings.Goal) < SearchSettings.SearchDistance * SearchSettings.SearchDistance)
+	const double SquareDistanceToGoal = FVector::DistSquared(NodeList[CurrentNodeIndex].PCGPoint->Transform.GetLocation(), SearchSettings.Goal);
+
+	auto PerPointProcessing = [&NodeList, CurrentNodeIndex, &SearchSettings, &SearchState](const FPCGPoint* Point, const double DistanceToPointSquared)
+	{
+		const bool bTrackedPoint = SearchState.PointToNodeIndexMap.Contains(Point);
+		// Node has already been ruled out.
+		if (bTrackedPoint && SearchState.ClosedIndexList.Contains(SearchState.PointToNodeIndexMap[Point]))
+		{
+			return;
+		}
+
+		const double TentativeNewLocalCost = SearchState.CostFunction(NodeList[CurrentNodeIndex].LocalCost, NodeList[CurrentNodeIndex].PCGPoint, DistanceToPointSquared, Point);
+
+		// Not tracking this point yet--add it to the node list and map it.
+		if (!bTrackedPoint)
+		{
+			const double HeuristicCost = SearchSettings.HeuristicWeight * SearchState.HeuristicFunction(Point->Transform.GetLocation(), SearchSettings.Goal);
+			NodeList.Emplace(Point, CurrentNodeIndex, TentativeNewLocalCost, TentativeNewLocalCost + HeuristicCost);
+			const int32 NewNodeIndex = NodeList.Num() - 1;
+			SearchState.OpenIndexList.HeapPush(NewNodeIndex, [&SearchState](int32 Index1, int32 Index2) { return Helpers::CompareNodes(SearchState.NodeList[Index1], SearchState.NodeList[Index2]); });
+			SearchState.PointToNodeIndexMap.Emplace(Point, NewNodeIndex);
+
+			return;
+		}
+
+		// Check if the path to this neighbor is a better path than its current one. If so, update accordingly.
+		const int32 NeighborIndex = SearchState.PointToNodeIndexMap[Point];
+		FNode& Neighbor = SearchState.NodeList[NeighborIndex];
+		if (TentativeNewLocalCost < Neighbor.LocalCost)
+		{
+			check(SearchState.OpenIndexList.Contains(NeighborIndex));
+
+			const double HeuristicCost = SearchSettings.HeuristicWeight * SearchState.HeuristicFunction(Point->Transform.GetLocation(), SearchSettings.Goal);
+			Neighbor.PreviousNodeIndex = CurrentNodeIndex;
+			Neighbor.LocalCost = TentativeNewLocalCost;
+			Neighbor.EstimatedGoalCost = TentativeNewLocalCost + HeuristicCost;
+		}
+	};
+
+	// Arrived at the goal.
+	if (FMath::IsNearlyZero(SquareDistanceToGoal))
 	{
-		const FPCGPoint GoalPoint(FTransform(SearchSettings.Goal), /*InDensity=*/1, PCGHelpers::ComputeSeedFromPosition(SearchSettings.Goal));
-		const FNode GoalNode(&GoalPoint, CurrentNodeIndex, NodeList[CurrentNodeIndex].EstimatedGoalCost, /*bInHeuristicCost=*/0);
-		Helpers::BuildFinalPath(SearchState, GoalNode, SearchSettings.bCopyOriginatingPoints, OutPath);
+		Helpers::BuildFinalPath(SearchState, SearchState.NodeList[CurrentNodeIndex], SearchSettings.bCopyOriginatingPoints, OutPath);
 		return true;
 	}
+	// Close enough to the goal we need to add it as a possibility.
+	else if (SquareDistanceToGoal <= SearchSettings.SearchDistance * SearchSettings.SearchDistance)
+	{
+		PerPointProcessing(&SearchSettings.GoalPoint, SquareDistanceToGoal);
+	}
 
 	// Gather neighbors within the search radius.
 	UPCGOctreeQueries::ForEachPointInsideSphere(
-			SearchState.OriginatingPointData,
-			NodeList[CurrentNodeIndex].PCGPoint->Transform.GetLocation(),
-			SearchSettings.SearchDistance,
-			[&NodeList, CurrentNodeIndex, &SearchSettings, &SearchState](const FPCGPointRef& PointRef, const double DistanceToPointSquared)
-			{
-				const bool bTrackedPoint = SearchState.PointToNodeIndexMap.Contains(PointRef.Point);
-				// Node has already been ruled out.
-				if (bTrackedPoint && SearchState.ClosedIndexList.Contains(SearchState.PointToNodeIndexMap[PointRef.Point]))
-				{
-					return;
-				}
-
-				const double TentativeNewLocalCost = Cost::CalculateCost_EuclideanDistance(NodeList[CurrentNodeIndex].LocalCost, DistanceToPointSquared);
-
-				// Not tracking this point yet--add it to the node list and map it.
-				if (!bTrackedPoint)
-				{
-					const double HeuristicCost = SearchSettings.HeuristicWeight * Heuristic::CalculateHeuristic_EuclideanDistance(PointRef.Point->Transform.GetLocation(), SearchSettings.Goal);
-					NodeList.Emplace(PointRef.Point, CurrentNodeIndex, TentativeNewLocalCost, TentativeNewLocalCost + HeuristicCost);
-					const int32 NewNodeIndex = NodeList.Num() - 1;
-					SearchState.OpenIndexList.HeapPush(NewNodeIndex, [&SearchState](int32 Index1, int32 Index2) { return Helpers::CompareNodes(SearchState.NodeList[Index1], SearchState.NodeList[Index2]); });
-					SearchState.PointToNodeIndexMap.Emplace(PointRef.Point, NewNodeIndex);
-
-					return;
-				}
-
-				// Check if the path to this neighbor is a better path than its current one. If so, update accordingly.
-				const int32 NeighborIndex = SearchState.PointToNodeIndexMap[PointRef.Point];
-				FNode& Neighbor = SearchState.NodeList[NeighborIndex];
-				if (TentativeNewLocalCost < Neighbor.LocalCost)
-				{
-					check(SearchState.OpenIndexList.Contains(NeighborIndex));
-
-					const double HeuristicCost = SearchSettings.HeuristicWeight * Heuristic::CalculateHeuristic_EuclideanDistance(PointRef.Point->Transform.GetLocation(), SearchSettings.Goal);
-					Neighbor.PreviousNodeIndex = CurrentNodeIndex;
-					Neighbor.LocalCost = TentativeNewLocalCost;
-					Neighbor.EstimatedGoalCost = TentativeNewLocalCost + HeuristicCost;
-				}
-			});
+		SearchState.OriginatingPointData,
+		NodeList[CurrentNodeIndex].PCGPoint->Transform.GetLocation(),
+		SearchSettings.SearchDistance,
+		[&PerPointProcessing](const FPCGPointRef& PointRef, const double DistanceToPointSquared)
+		{
+			PerPointProcessing(PointRef.Point, DistanceToPointSquared);
+		});
 
 	// This node has been completely evaluated.
 	SearchState.ClosedIndexList.Add(CurrentNodeIndex);
diff --git a/Engine/Plugins/PCG/Source/PCG/Public/Elements/PCGPathfindingElement.h b/Engine/Plugins/PCG/Source/PCG/Public/Elements/PCGPathfindingElement.h
--- a/Engine/Plugins/PCG/Source/PCG/Public/Elements/PCGPathfindingElement.h
+++ b/Engine/Plugins/PCG/Source/PCG/Public/Elements/PCGPathfindingElement.h
@@ -11,13 +11,21 @@
 
 class UPCGPointData;
 
-UENUM(BlueprintType, Blueprintable)
+UENUM(BlueprintType)
 enum class EPCGPathfindingSplineMode : uint8
 {
 	Curve UMETA(Tooltip = "Interpret the spline as a continuous curve."),
 	Linear UMETA(Tooltip = "Interpret the spline as a conjunction of linear segments."),
 };
 
+UENUM(BlueprintType)
+enum class EPCGPathfindingCostFunctionMode : uint8
+{
+	Distance UMETA(Tooltip = "Pathfinding cost will be the distance only."),
+	FitnessScore UMETA(Tooltip = "Pathfinding cost will be driven by a fitness score (0-1 range), with a maximum penalty applied at fitness = 0."),
+	CostMultipler UMETA(Tooltip = "Pathfinding cost will be the distance multiplied by the provided factor. Note that multipliers below 1 will be clamped to 1.")
+};
+
 /** Finds the optimal path across the points of a given point cloud--should one exist--when provided a start and goal
  * location, and a maximum jump distance between points. Can return a partial path.
  */
@@ -46,7 +54,7 @@ protected:
 
 public:
 	/** The max distance from each point to search for the next viable point in the path. */
-	UPROPERTY(EditAnywhere, BlueprintReadWrite, Category = Settings, meta = (PCG_Overridable))
+	UPROPERTY(EditAnywhere, BlueprintReadWrite, Category = Settings, meta = (ClampMin = "0.01", PCG_Overridable))
 	double SearchDistance = 1000;
 
 	/** The path's starting location. */
@@ -61,6 +69,18 @@ public:
 	UPROPERTY(EditAnywhere, BlueprintReadWrite, Category = Settings, meta = (EditCondition = "Algorithm == EPCGPathfindingAlgorithm::AStar", PCG_Overridable))
 	double HeuristicWeight = 1.0;
 
+	/** Controls whether the cost function will use a given attribute as a scalar wrt to the distance. */
+	UPROPERTY(EditAnywhere, BlueprintReadWrite, Category = Settings, meta = (PCG_Overridable))
+	EPCGPathfindingCostFunctionMode CostFunctionMode = EPCGPathfindingCostFunctionMode::Distance;
+
+	/** Attribute to use as part of the cost function - it's meaning will depend on the cost function mode (fitness value, scalar multiplier, or else). */
+	UPROPERTY(EditAnywhere, BlueprintReadWrite, Category = Settings, meta = (EditCondition = "CostFunctionMode != EPCGPathfindingCostFunctionMode::Distance", EditConditionHides, PCG_Overridable))
+	FPCGAttributePropertyInputSelector CostAttribute;
+
+	/** Fitness penalty scalar (maximum penalty applied when fitness is zero.) */
+	UPROPERTY(EditAnywhere, BlueprintReadWrite, Category = Settings, meta = (ClampMin = "1.0", EditCondition = "CostFunctionMode == EPCGPathfindingCostFunctionMode::FitnessScore", EditConditionHides, PCG_Overridable))
+	double MaximumFitnessPenaltyFactor = 10.0;
+
 	/** Even if the path is not complete, return the most optimal and viable partial path to the goal. */
 	UPROPERTY(EditAnywhere, BlueprintReadWrite, Category = Settings, meta = (PCG_Overridable))
 	bool bAcceptPartialPath = true;
diff --git a/Engine/Plugins/PCG/Source/PCG/Public/SpatialAlgo/PCGAStar.h b/Engine/Plugins/PCG/Source/PCG/Public/SpatialAlgo/PCGAStar.h
--- a/Engine/Plugins/PCG/Source/PCG/Public/SpatialAlgo/PCGAStar.h
+++ b/Engine/Plugins/PCG/Source/PCG/Public/SpatialAlgo/PCGAStar.h
@@ -27,10 +27,21 @@ struct FNode
 	double EstimatedGoalCost = 0.0;
 };
 
+namespace Cost
+{
+	double CalculateCost_EuclideanDistance(const double PreviousNodeCost, const FPCGPoint* PreviousNodePoint, const double DistanceToPreviousNodeSquared, const FPCGPoint* CurrentPoint);
+}
+
+namespace Heuristic
+{
+	double CalculateHeuristic_EuclideanDistance(const FVector& CurrentLocation, const FVector& GoalLocation);
+}
+
 struct FSearchSettings
 {
-	/** Cached as an initial point in an otherwise const point data. */
+	/** Cached start/goal points as initial points in an otherwise const point data. */
 	FPCGPoint StartPoint;
+	FPCGPoint GoalPoint;
 	/** The max distance from each point to search for the next viable point in the path. */
 	double SearchDistance = 1000;
 	/** The path's starting location. */
@@ -50,6 +61,10 @@ struct FSearchState
 	static uint32 constexpr PreAllocNodeCount = 1024u;
 	// Keep track of the point data for the octree search.
 	const UPCGPointData* OriginatingPointData = nullptr;
+	/** Cost function */
+	TFunction<double(const double /*PreviousNodeCost*/, const FPCGPoint* /*PreviousNodePoint*/, const double /*DistanceToPreviousNodeSquared*/, const FPCGPoint* /*CurrentPoint*/)> CostFunction = Cost::CalculateCost_EuclideanDistance;
+	/** Heuristic function */
+	TFunction<double(const FVector& /*CurrentLocation*/, const FVector& /*GoalLocation*/)> HeuristicFunction = Heuristic::CalculateHeuristic_EuclideanDistance;
 
 	// Storage of the current known node. Source of truth for the nodes, so it should not be resized after its initial reserve.
 	TArray<FNode, TInlineAllocator<PreAllocNodeCount>> NodeList;
@@ -67,16 +82,6 @@ namespace Helpers
 	void BuildFinalPath(const FSearchState& SearchState, const FNode& FinalNode, bool bCopyOriginatingPoints, TArray<FPCGPoint>& OutPath);
 }
 
-namespace Cost
-{
-	double CalculateCost_EuclideanDistance(const double PreviousNodeCost, const double DistanceToPreviousNodeSquared);
-} // namespace Cost
-
-namespace Heuristic
-{
-	double CalculateHeuristic_EuclideanDistance(const FVector& CurrentLocation, const FVector& GoalLocation);
-} // namespace Heuristic
-
 /** Initialize the search state for AStar. Must be called before ExecuteSearchIteration. */
 void Initialize(const UPCGPointData* const PointData, const FSearchSettings& Settings, FSearchState& OutSearchState);
 
