[PCG] Add source editor window for nodes with source text.

#rnx

diff --git a/Engine/Plugins/PCG/Source/PCG/Private/Compute/Elements/PCGCustomHLSL.cpp b/Engine/Plugins/PCG/Source/PCG/Private/Compute/Elements/PCGCustomHLSL.cpp
--- a/Engine/Plugins/PCG/Source/PCG/Private/Compute/Elements/PCGCustomHLSL.cpp
+++ b/Engine/Plugins/PCG/Source/PCG/Private/Compute/Elements/PCGCustomHLSL.cpp
@@ -455,6 +455,43 @@ const UPCGPin* UPCGCustomHLSLSettings::GetFirstPointOutputPin() const
 	return nullptr;
 }
 
+#if WITH_EDITOR
+FString UPCGCustomHLSLSettings::GetDeclarationsText() const
+{
+	return InputDeclarations + TEXT("\n\n") + OutputDeclarations + TEXT("\n\n") + HelperDeclarations;
+}
+
+FString UPCGCustomHLSLSettings::GetShaderFunctionsText() const
+{
+	return ShaderFunctions;
+}
+
+FString UPCGCustomHLSLSettings::GetShaderText() const
+{
+	return ShaderSource;
+}
+
+void UPCGCustomHLSLSettings::SetShaderFunctionsText(const FString& NewFunctionsText)
+{
+	ShaderFunctions = NewFunctionsText;
+}
+
+void UPCGCustomHLSLSettings::SetShaderText(const FString& NewText)
+{
+	ShaderSource = NewText;
+}
+
+bool UPCGCustomHLSLSettings::IsShaderTextReadOnly() const
+{
+	return false;
+}
+
+void UPCGCustomHLSLSettings::ApplySourceChanges()
+{
+	OnSettingsChangedDelegate.Broadcast(this, GetChangeTypeForProperty(GET_MEMBER_NAME_CHECKED(UPCGCustomHLSLSettings, ShaderSource)));
+}
+#endif //WITH_EDITOR
+
 const FPCGPinPropertiesGPU* UPCGCustomHLSLSettings::GetOutputPinPropertiesGPU(const FName& InPinLabel) const
 {
 	return OutputPins.FindByPredicate([InPinLabel](const FPCGPinPropertiesGPU& InProperties)
@@ -476,6 +513,43 @@ void UPCGCustomHLSLSettings::UpdateInputDeclarations()
 {
 	InputDeclarations.Reset();
 
+	// Constants category
+	{
+		if (KernelType == EPCGKernelType::PointGenerator)
+		{
+			InputDeclarations += TEXT("/*** INPUT CONSTANTS ***/\n\n");
+			InputDeclarations += FString::Format(TEXT("const uint PointCount = {0};\n\n"), { PointCount });
+		}
+
+		InputDeclarations += TEXT("/*** INPUT PER-THREAD CONSTANTS ***/\n\n");
+		InputDeclarations += TEXT("const uint ThreadIndex;\n");
+
+		if (KernelType == EPCGKernelType::PointProcessor)
+		{
+			const UPCGPin* PointProcessingInputPin = GetPointProcessingInputPin();
+			const UPCGPin* PointProcessingOutputPin = GetFirstPointOutputPin();
+
+			if (PointProcessingInputPin && PointProcessingOutputPin)
+			{
+				InputDeclarations += FString::Format(TEXT(
+					"const uint {0}_DataIndex;\n"
+					"const uint {1}_DataIndex;\n"),
+					{ PointProcessingInputPin->Properties.Label.ToString(),  PointProcessingOutputPin->Properties.Label.ToString() });
+			}
+		}
+		else if (KernelType == EPCGKernelType::PointGenerator)
+		{
+			if (const UPCGPin* PointProcessingOutputPin = GetFirstPointOutputPin())
+			{
+				InputDeclarations += FString::Format(
+					TEXT("const uint {0}_DataIndex;\n"),
+					{ PointProcessingOutputPin->Properties.Label.ToString() });
+			}
+		}
+
+		InputDeclarations += TEXT("const uint ElementIndex;\n\n");
+	}
+
 	TArray<FString> DataPins;
 	TArray<FString> PointDataPins;
 	TArray<FString> SplineDataPins;
@@ -510,7 +584,7 @@ void UPCGCustomHLSLSettings::UpdateInputDeclarations()
 
 	if (!DataPins.IsEmpty())
 	{
-		InputDeclarations += TEXT("### DATA FUNCTIONS ###\n\n");
+		InputDeclarations += TEXT("/*** INPUT DATA FUNCTIONS ***/\n\n");
 
 		const bool bMultiPin = DataPins.Num() > 1;
 
@@ -534,7 +608,7 @@ void UPCGCustomHLSLSettings::UpdateInputDeclarations()
 
 	if (!PointDataPins.IsEmpty())
 	{
-		InputDeclarations += TEXT("### POINT DATA FUNCTIONS ###\n\n");
+		InputDeclarations += TEXT("/*** INPUT POINT DATA FUNCTIONS ***/\n\n");
 
 		const bool bMultiPin = PointDataPins.Num() > 1;
 
@@ -563,7 +637,7 @@ void UPCGCustomHLSLSettings::UpdateInputDeclarations()
 
 	if (!LandscapeDataPins.IsEmpty())
 	{
-		InputDeclarations += TEXT("### LANDSCAPE DATA FUNCTIONS ###\n\n");
+		InputDeclarations += TEXT("/*** INPUT LANDSCAPE DATA FUNCTIONS ***/\n\n");
 
 		const bool bMultiPin = LandscapeDataPins.Num() > 1;
 
@@ -582,7 +656,7 @@ void UPCGCustomHLSLSettings::UpdateInputDeclarations()
 
 	if (!TextureDataPins.IsEmpty())
 	{
-		InputDeclarations += TEXT("### TEXTURE DATA FUNCTIONS ###\n\n");
+		InputDeclarations += TEXT("/*** INPUT TEXTURE DATA FUNCTIONS ***/\n\n");
 
 		const bool bMultiPin = TextureDataPins.Num() > 1;
 
@@ -601,7 +675,7 @@ void UPCGCustomHLSLSettings::UpdateInputDeclarations()
 
 	if (!RawBufferDataPins.IsEmpty())
 	{
-		InputDeclarations += TEXT("### BYTE ADDRESS BUFFER DATA FUNCTIONS ###\n\n");
+		InputDeclarations += TEXT("/*** INPUT BYTE ADDRESS BUFFER DATA FUNCTIONS ***/\n\n");
 
 		const bool bMultiPin = RawBufferDataPins.Num() > 1;
 
@@ -647,7 +721,7 @@ void UPCGCustomHLSLSettings::UpdateOutputDeclarations()
 
 	if (!DataPins.IsEmpty())
 	{
-		OutputDeclarations += TEXT("### DATA FUNCTIONS ###\n\n");
+		OutputDeclarations += TEXT("/*** OUTPUT DATA FUNCTIONS ***/\n\n");
 
 		const bool bMultiPin = DataPins.Num() > 1;
 
@@ -668,7 +742,7 @@ void UPCGCustomHLSLSettings::UpdateOutputDeclarations()
 
 	if (!PointDataPins.IsEmpty())
 	{
-		OutputDeclarations += TEXT("### POINT DATA FUNCTIONS ###\n\n");
+		OutputDeclarations += TEXT("/*** OUTPUT POINT DATA FUNCTIONS ***/\n\n");
 
 		const bool bMultiPin = PointDataPins.Num() > 1;
 
@@ -696,7 +770,7 @@ void UPCGCustomHLSLSettings::UpdateOutputDeclarations()
 
 	if (!RawBufferDataPins.IsEmpty())
 	{
-		OutputDeclarations += TEXT("### BYTE ADDRESS BUFFER DATA FUNCTIONS ###\n\n");
+		OutputDeclarations += TEXT("/*** OUTPUT BYTE ADDRESS BUFFER DATA FUNCTIONS ***/\n\n");
 
 		const bool bMultiPin = RawBufferDataPins.Num() > 1;
 
@@ -719,48 +793,10 @@ void UPCGCustomHLSLSettings::UpdateHelperDeclarations()
 {
 	HelperDeclarations.Reset();
 
-	// Constants category
-	{
-		if (KernelType == EPCGKernelType::PointGenerator)
-		{
-			HelperDeclarations += TEXT("### CONSTANTS ###\n\n");
-			HelperDeclarations += FString::Format(TEXT("const uint PointCount = {0};\n\n"), { PointCount });
-		}
-
-		HelperDeclarations += TEXT("### PER-THREAD CONSTANTS ###\n\n");
-		HelperDeclarations += TEXT("const uint ThreadIndex;\n");
-
-		if (KernelType == EPCGKernelType::PointProcessor)
-		{
-			const UPCGPin* PointProcessingInputPin = GetPointProcessingInputPin();
-			const UPCGPin* PointProcessingOutputPin = GetFirstPointOutputPin();
-
-			if (PointProcessingInputPin && PointProcessingOutputPin)
-			{
-				HelperDeclarations += FString::Format(TEXT(
-					"const uint {0}_DataIndex;\n"
-					"const uint {1}_DataIndex;\n"),
-					{ PointProcessingInputPin->Properties.Label.ToString(),  PointProcessingOutputPin->Properties.Label.ToString() });
-			}
-		}
-		else if (KernelType == EPCGKernelType::PointGenerator)
-		{
-			if (const UPCGPin* PointProcessingOutputPin = GetFirstPointOutputPin())
-			{
-				HelperDeclarations += FString::Format(
-					TEXT("const uint {0}_DataIndex;\n"),
-					{ PointProcessingOutputPin->Properties.Label.ToString() });
-			}
-		}
-
-		HelperDeclarations += TEXT("const uint ElementIndex;\n");
-		HelperDeclarations += TEXT("\n");
-	}
-
 	// Helper funcs category
 	{
 		HelperDeclarations += TEXT(
-			"### HELPER FUNCTIONS ###\n"
+			"/*** HELPER FUNCTIONS ***/\n"
 			"\n"
 			"int3 GetNumThreads();\n");
 
@@ -1427,8 +1463,6 @@ FString UPCGCustomHLSLSettings::GetCookedKernelSource(const TMap<FPCGKernelAttri
 	Functions.ReplaceInline(TEXT("\r"), TEXT(""));
 #endif
 
-	Source.ReplaceInline(TEXT("\n"), TEXT("\n    ")); // Properly indent kernel source
-
 	for (const TPair<FPCGKernelAttributeKey, int32>& Pair : GlobalAttributeLookupTable)
 	{
 		const FString SourceDefinition = PCGHLSLElement::GetKernelAttributeKeyAsString(Pair.Key);
@@ -1469,7 +1503,7 @@ FString UPCGCustomHLSLSettings::GetCookedKernelSource(const TMap<FPCGKernelAttri
 
 	for (const FPCGPinPropertiesGPU& PinProps : OutputPins)
 	{
-		SetAsExecuted += FString::Format(TEXT("    if (all(GroupId == 0) && GroupIndex == 0) {0}_SetAsExecutedInternal();\n"), { PinProps.Label.ToString() });
+		SetAsExecuted += FString::Format(TEXT("    if (all(GroupId == 0u) && GroupIndex == 0) {0}_SetAsExecutedInternal();\n"), { PinProps.Label.ToString() });
 	}
 
 	// Per-kernel-type preamble. Set up shader inputs and initialize output data.
diff --git a/Engine/Plugins/PCG/Source/PCG/Private/Compute/PCGComputeGraph.cpp b/Engine/Plugins/PCG/Source/PCG/Private/Compute/PCGComputeGraph.cpp
--- a/Engine/Plugins/PCG/Source/PCG/Private/Compute/PCGComputeGraph.cpp
+++ b/Engine/Plugins/PCG/Source/PCG/Private/Compute/PCGComputeGraph.cpp
@@ -2,18 +2,82 @@
 
 #include "PCGComputeGraph.h"
 
+#include "PCGGraph.h"
 #include "PCGModule.h"
 #include "PCGNode.h"
+#include "Compute/PCGCompilerDiagnostic.h"
 
 #include "ComputeFramework/ComputeKernelCompileResult.h"
 
 #include UE_INLINE_GENERATED_CPP_BY_NAME(PCGComputeGraph)
 
+static FPCGCompilerDiagnostic ProcessCompilationMessage(const FComputeKernelCompileMessage& InMessage)
+{
+	FPCGCompilerDiagnostic Diagnostic;
+
+	if (InMessage.Type == FComputeKernelCompileMessage::EMessageType::Error)
+	{
+		Diagnostic.Level = EPCGDiagnosticLevel::Error;
+	}
+	else if (InMessage.Type == FComputeKernelCompileMessage::EMessageType::Warning)
+	{
+		Diagnostic.Level = EPCGDiagnosticLevel::Warning;
+	}
+	else if (InMessage.Type == FComputeKernelCompileMessage::EMessageType::Info)
+	{
+		Diagnostic.Level = EPCGDiagnosticLevel::Info;
+	}
+
+	Diagnostic.Line = InMessage.Line;
+	Diagnostic.ColumnStart = InMessage.ColumnStart;
+	Diagnostic.ColumnEnd = InMessage.ColumnEnd;
+
+	FString Message;
+	if (!InMessage.VirtualFilePath.IsEmpty())
+	{
+		Message = InMessage.VirtualFilePath;
+		if (InMessage.Line != -1)
+		{
+			if (InMessage.ColumnStart == InMessage.ColumnEnd)
+			{
+				Message += FString::Printf(TEXT(" (%d,%d)"), InMessage.Line, InMessage.ColumnStart);
+			}
+			else
+			{
+				Message += FString::Printf(TEXT(" (%d,%d-%d)"), InMessage.Line, InMessage.ColumnStart, InMessage.ColumnEnd);
+			}
+		}
+		Message += TEXT(": ");
+	}
+	Message += InMessage.Text;
+	Diagnostic.Message = FText::FromString(Message);
+
+	return Diagnostic;
+}
+
 void UPCGComputeGraph::OnKernelCompilationComplete(int32 InKernelIndex, FComputeKernelCompileResults const& InCompileResults)
 {
-	if (KernelToNode.IsValidIndex(InKernelIndex) && KernelToNode[InKernelIndex].Get())
+	const UPCGNode* Node = (KernelToNode.IsValidIndex(InKernelIndex) && KernelToNode[InKernelIndex].IsValid()) ? KernelToNode[InKernelIndex].Get() : nullptr;
+
+	if (Node)
 	{
 		KernelToCompileMessages.FindOrAdd(KernelToNode[InKernelIndex].Get()) = InCompileResults.Messages;
+
+		UPCGGraph* Graph = Cast<UPCGGraph>(Node->GetOuter());
+		if (Graph)
+		{
+			FPCGCompilerDiagnostics Diagnostics;
+			Diagnostics.Diagnostics.Reserve(InCompileResults.Messages.Num());
+
+			Algo::Transform(InCompileResults.Messages, Diagnostics.Diagnostics, [](const FComputeKernelCompileMessage& InMessage)
+			{
+				return ProcessCompilationMessage(InMessage);
+			});
+
+#if WITH_EDITOR
+			Graph->OnNodeSourceCompiledDelegate.Broadcast(Node, Diagnostics);
+#endif
+		}
 	}
 	else
 	{
diff --git a/Engine/Plugins/PCG/Source/PCG/Public/Compute/Elements/PCGCustomHLSL.h b/Engine/Plugins/PCG/Source/PCG/Public/Compute/Elements/PCGCustomHLSL.h
--- a/Engine/Plugins/PCG/Source/PCG/Public/Compute/Elements/PCGCustomHLSL.h
+++ b/Engine/Plugins/PCG/Source/PCG/Public/Compute/Elements/PCGCustomHLSL.h
@@ -5,6 +5,8 @@
 #include "PCGSettings.h"
 #include "Compute/PCGPinPropertiesGPU.h"
 
+#include "Compute/IPCGNodeSourceTextProvider.h"
+
 #include "PCGCustomHLSL.generated.h"
 
 class UPCGPin;
@@ -29,7 +31,9 @@ enum class EPCGDispatchThreadCount : uint8
 
 /** Produces a HLSL compute shader which will be executed on the GPU. */
 UCLASS(MinimalAPI, BlueprintType, ClassGroup = (Procedural))
-class UPCGCustomHLSLSettings : public UPCGSettings
+class UPCGCustomHLSLSettings
+	: public UPCGSettings
+	, public IPCGNodeSourceTextProvider
 {
 	GENERATED_BODY()
 
@@ -76,6 +80,18 @@ protected:
 	virtual FPCGElementPtr CreateElement() const override;
 	//~End UPCGSettings interface
 
+#if WITH_EDITOR
+	//~Begin IPCGNodeSourceTextProvider interface
+	FString GetShaderText() const override;
+	FString GetDeclarationsText() const override;
+	FString GetShaderFunctionsText() const override;
+	void SetShaderFunctionsText(const FString& NewFunctionsText) override;
+	void SetShaderText(const FString& NewText) override;
+	bool IsShaderTextReadOnly() const override;
+	void ApplySourceChanges() override;
+	//~End IPCGNodeSourceTextProvider interface
+#endif
+
 	/** Gets the GPU pin properties for the output pin with the given label. */
 	const FPCGPinPropertiesGPU* GetOutputPinPropertiesGPU(const FName& InPinLabel) const;
 
@@ -140,18 +156,23 @@ public:
 #endif
 
 protected:
-	UPROPERTY(BlueprintReadWrite, EditAnywhere, Category = "Source", meta = (MultiLine = true, Tooltip = ""))
+	/** Optional functions that can be called from the source. Intended to be edited using the Node Source Editor window. */
+	UPROPERTY(BlueprintReadWrite, EditAnywhere, Category = "Source", meta = (MultiLine = true))
 	FString ShaderFunctions;
 
-	UPROPERTY(BlueprintReadWrite, EditAnywhere, Category = "Source", meta = (MultiLine = true, Tooltip = ""))
+	/** Shader code that forms the body of the kernel. Intended to be edited using the Node Source Editor window. */
+	UPROPERTY(BlueprintReadWrite, EditAnywhere, Category = "Source", meta = (MultiLine = true))
 	FString ShaderSource;
 
+	/** Inputs data accessors that can be used from the shader code. Intended to be viewed using the Node Source Editor window. */
 	UPROPERTY(Transient, VisibleAnywhere, Category = "Declarations|Inputs", meta = (MultiLine = true))
 	FString InputDeclarations;
 
+	/** Output data accessors that can be used from the shader code. Intended to be viewed using the Node Source Editor window. */
 	UPROPERTY(Transient, VisibleAnywhere, Category = "Declarations|Outputs", meta = (MultiLine = true))
 	FString OutputDeclarations;
 
+	/** Helper data and functions that can be used from the shader code. Intended to be viewed using the Node Source Editor window. */
 	UPROPERTY(Transient, VisibleAnywhere, Category = "Declarations|Helpers", meta = (MultiLine = true))
 	FString HelperDeclarations;
 };
diff --git a/Engine/Plugins/PCG/Source/PCG/Public/Compute/IPCGNodeSourceTextProvider.h b/Engine/Plugins/PCG/Source/PCG/Public/Compute/IPCGNodeSourceTextProvider.h
new file mode 100644
--- /dev/null
+++ b/Engine/Plugins/PCG/Source/PCG/Public/Compute/IPCGNodeSourceTextProvider.h
@@ -0,0 +1,44 @@
+// Copyright Epic Games, Inc. All Rights Reserved.
+
+#pragma once
+
+#include "UObject/Interface.h"
+#include "IPCGNodeSourceTextProvider.generated.h"
+
+UINTERFACE()
+class PCG_API UPCGNodeSourceTextProvider : public UInterface
+{
+	GENERATED_BODY()
+};
+
+/**
+* Interface for node source text editing widgets.
+*/
+class PCG_API IPCGNodeSourceTextProvider
+{
+	GENERATED_BODY()
+
+#if WITH_EDITOR	
+public:
+	/** Get text for declarations pane. */
+	virtual FString GetDeclarationsText() const = 0;
+
+	/** Get text for shader functions pane. */
+	virtual FString GetShaderFunctionsText() const = 0;
+
+	/** Get shader text for editor pane. */
+	virtual FString GetShaderText() const = 0;
+
+	/** Set shader functions text after edit modifications. */
+	virtual void SetShaderFunctionsText(const FString& NewFunctionsText) = 0;
+
+	/** Set shader text after edit modifications. */
+	virtual void SetShaderText(const FString& NewText) = 0;
+
+	/** Return true is the shader text can be edited */
+	virtual bool IsShaderTextReadOnly() const = 0;
+
+	/** Apply source changes (compile and execute). */
+	virtual void ApplySourceChanges() = 0;
+#endif
+};
diff --git a/Engine/Plugins/PCG/Source/PCG/Public/Compute/PCGCompilerDiagnostic.h b/Engine/Plugins/PCG/Source/PCG/Public/Compute/PCGCompilerDiagnostic.h
new file mode 100644
--- /dev/null
+++ b/Engine/Plugins/PCG/Source/PCG/Public/Compute/PCGCompilerDiagnostic.h
@@ -0,0 +1,40 @@
+// Copyright Epic Games, Inc. All Rights Reserved.
+
+#pragma once
+
+#include "CoreMinimal.h"
+
+enum class EPCGDiagnosticLevel : uint8
+{
+	None,
+	Info,
+	Warning,
+	Error
+};
+
+struct FPCGCompilerDiagnostic
+{
+	FPCGCompilerDiagnostic() = default;
+	FPCGCompilerDiagnostic(const FPCGCompilerDiagnostic&) = default;
+	FPCGCompilerDiagnostic& operator=(const FPCGCompilerDiagnostic&) = default;
+
+	// The severity of the issue.
+	EPCGDiagnosticLevel Level = EPCGDiagnosticLevel::None;
+
+	// The actual diagnostic message.
+	FText Message;
+
+	// Line location in source.
+	int32 Line = INDEX_NONE;
+
+	// Starting column (inclusive).
+	int32 ColumnStart = INDEX_NONE;
+
+	// Ending column (inclusive).
+	int32 ColumnEnd = INDEX_NONE;
+};
+
+struct FPCGCompilerDiagnostics
+{
+	TArray<FPCGCompilerDiagnostic> Diagnostics;
+};
diff --git a/Engine/Plugins/PCG/Source/PCG/Public/PCGGraph.h b/Engine/Plugins/PCG/Source/PCG/Public/PCGGraph.h
--- a/Engine/Plugins/PCG/Source/PCG/Public/PCGGraph.h
+++ b/Engine/Plugins/PCG/Source/PCG/Public/PCGGraph.h
@@ -5,6 +5,7 @@
 #include "PCGCommon.h"
 #include "PCGNode.h"
 #include "PCGSettings.h"
+#include "Compute/PCGCompilerDiagnostic.h"
 #include "Graph/PCGStackContext.h"
 #include "Helpers/PCGGraphParameterExtension.h"
 
@@ -39,6 +40,7 @@ enum class EPCGGraphParameterEvent
 DECLARE_MULTICAST_DELEGATE_TwoParams(FOnPCGGraphChanged, UPCGGraphInterface* /*Graph*/, EPCGChangeType /*ChangeType*/);
 DECLARE_MULTICAST_DELEGATE_OneParam(FOnPCGGraphStructureChanged, UPCGGraphInterface* /*Graph*/);
 DECLARE_MULTICAST_DELEGATE_ThreeParams(FOnPCGGraphParametersChanged, UPCGGraphInterface* /*Graph*/, EPCGGraphParameterEvent /*ChangeType*/, FName /*ChangedPropertyName*/);
+DECLARE_MULTICAST_DELEGATE_TwoParams(FOnPCGNodeSourceCompiled, const UPCGNode*, const FPCGCompilerDiagnostics&);
 #endif // WITH_EDITOR
 
 /**
@@ -114,6 +116,7 @@ public:
 #if WITH_EDITOR
 	FOnPCGGraphChanged OnGraphChangedDelegate;
 	FOnPCGGraphParametersChanged OnGraphParametersChangedDelegate;
+	FOnPCGNodeSourceCompiled OnNodeSourceCompiledDelegate;
 
 	virtual TOptional<FText> GetTitleOverride() const { return bOverrideTitle ? Title : TOptional<FText>(); }
 	virtual TOptional<FLinearColor> GetColorOverride() const { return bOverrideColor ? Color : TOptional<FLinearColor>(); }
diff --git a/Engine/Plugins/PCG/Source/PCGEditor/Private/PCGEditor.cpp b/Engine/Plugins/PCG/Source/PCGEditor/Private/PCGEditor.cpp
--- a/Engine/Plugins/PCG/Source/PCGEditor/Private/PCGEditor.cpp
+++ b/Engine/Plugins/PCG/Source/PCGEditor/Private/PCGEditor.cpp
@@ -12,6 +12,7 @@
 #include "PCGInputOutputSettings.h"
 #include "PCGPin.h"
 #include "PCGSubsystem.h"
+#include "Compute/IPCGNodeSourceTextProvider.h"
 #include "Elements/PCGReroute.h"
 #include "Helpers/PCGSubgraphHelpers.h"
 #include "Rendering/SlateRenderer.h"
@@ -39,6 +40,7 @@
 #include "SPCGEditorGraphLogView.h"
 #include "SPCGEditorGraphNodePalette.h"
 #include "SPCGEditorGraphProfilingView.h"
+#include "Widgets/SPCGEditorNodeSource.h"
 
 #include "AssetToolsModule.h"
 #include "EdGraphUtilities.h"
@@ -91,6 +93,7 @@ namespace FPCGEditor_private
 	const FName DeterminismID = FName(TEXT("Determinism"));
 	const FName ProfilingID = FName(TEXT("Profiling"));
 	const FName LogID = FName(TEXT("Log"));
+	const FName NodeSourceID = FName(TEXT("NodeSource"));
 }
 
 UPCGEditorGraph* FPCGEditor::GetPCGEditorGraph(UPCGGraph* InGraph)
@@ -142,6 +145,7 @@ void FPCGEditor::Initialize(const EToolkitMode::Type InMode, const TSharedPtr<cl
 	DeterminismWidget = CreateDeterminismWidget();
 	ProfilingWidget = CreateProfilingWidget();
 	LogWidget = CreateLogWidget();
+	NodeSourceWidget = CreateNodeSourceWidget();
 
 	BindCommands();
 	RegisterToolbar();
@@ -204,6 +208,7 @@ void FPCGEditor::Initialize(const EToolkitMode::Type InMode, const TSharedPtr<cl
 	InitAssetEditor(InMode, InToolkitHost, PCGGraphEditorAppName, StandaloneDefaultLayout, /*bCreateDefaultStandaloneMenu=*/ true, /*bCreateDefaultToolbar=*/ true, InPCGGraph);
 
 	PCGGraphBeingEdited->OnGraphChangedDelegate.AddRaw(this, &FPCGEditor::OnGraphChanged);
+	PCGGraphBeingEdited->OnNodeSourceCompiledDelegate.AddRaw(this, &FPCGEditor::OnNodeSourceCompiled);
 
 	// Hook to map change / delete actor to refresh debug object selection list, to help prevent it going stale.
 	FLevelEditorModule& LevelEditor = FModuleManager::LoadModuleChecked<FLevelEditorModule>("LevelEditor");
@@ -443,6 +448,11 @@ const FPCGStack* FPCGEditor::GetStackBeingInspected() const
 	return StackBeingInspected.GetStackFrames().IsEmpty() ? nullptr : &StackBeingInspected;
 }
 
+void FPCGEditor::SetSourceEditorTargetObject(UObject* InObject)
+{
+	NodeSourceWidget->SetTextProviderObject(InObject);
+}
+
 void FPCGEditor::JumpToNode(const UEdGraphNode* InNode)
 {
 	if (GraphEditorWidget.IsValid())
@@ -549,6 +559,10 @@ void FPCGEditor::RegisterTabSpawners(const TSharedRef<FTabManager>& InTabManager
 	InTabManager->RegisterTabSpawner(FPCGEditor_private::LogID, FOnSpawnTab::CreateSP(this, &FPCGEditor::SpawnTab_Log))
 		.SetDisplayName(LOCTEXT("LogCaptureTab", "Log Capture"))
 		.SetGroup(WorkspaceMenuCategoryRef);
+
+	InTabManager->RegisterTabSpawner(FPCGEditor_private::NodeSourceID, FOnSpawnTab::CreateSP(this, &FPCGEditor::SpawnTab_NodeSource))
+		.SetDisplayName(LOCTEXT("NodeSourceTab", "Node Source Editor"))
+		.SetGroup(WorkspaceMenuCategoryRef);
 }
 
 void FPCGEditor::UnregisterTabSpawners(const TSharedRef<class FTabManager>& InTabManager)
@@ -566,6 +580,9 @@ void FPCGEditor::UnregisterTabSpawners(const TSharedRef<class FTabManager>& InTa
 	InTabManager->UnregisterTabSpawner(FPCGEditor_private::AttributesID[3]);
 	InTabManager->UnregisterTabSpawner(FPCGEditor_private::FindID);
 	InTabManager->UnregisterTabSpawner(FPCGEditor_private::DeterminismID);
+	InTabManager->UnregisterTabSpawner(FPCGEditor_private::ProfilingID);
+	InTabManager->UnregisterTabSpawner(FPCGEditor_private::LogID);
+	InTabManager->UnregisterTabSpawner(FPCGEditor_private::NodeSourceID);
 
 	FAssetEditorToolkit::UnregisterTabSpawners(InTabManager);
 }
@@ -2961,6 +2978,11 @@ TSharedRef<SPCGEditorGraphLogView> FPCGEditor::CreateLogWidget()
 	return SNew(SPCGEditorGraphLogView, SharedThis(this));
 }
 
+TSharedRef<SPCGEditorNodeSource> FPCGEditor::CreateNodeSourceWidget()
+{
+	return SNew(SPCGEditorNodeSource, SharedThis(this));
+}
+
 void FPCGEditor::OnSelectedNodesChanged(const TSet<UObject*>& NewSelection)
 {
 	TArray<TWeakObjectPtr<UObject>> SelectedObjects;
@@ -2984,6 +3006,11 @@ void FPCGEditor::OnSelectedNodesChanged(const TSet<UObject*>& NewSelection)
 	{
 		PropertyDetailsWidget->SetObjects(SelectedObjects, /*bForceRefresh=*/true);
 	}
+
+	// Give a single selected node with valid settings to the source editor, or give it null so it can clear the UI.
+	UPCGEditorGraphNode* SelectedNode = (NewSelection.Num() == 1) ? Cast<UPCGEditorGraphNode>(*NewSelection.CreateConstIterator()) : nullptr;
+	UPCGNode* PCGNode = SelectedNode ? SelectedNode->GetPCGNode() : nullptr;
+	SetSourceEditorTargetObject(PCGNode ? PCGNode->GetSettings() : nullptr);
 }
 
 void FPCGEditor::OnNodeTitleCommitted(const FText& NewText, ETextCommit::Type CommitInfo, UEdGraphNode* NodeBeingChanged)
@@ -3138,6 +3165,21 @@ void FPCGEditor::OnGraphChanged(UPCGGraphInterface* InGraph, EPCGChangeType Chan
 	}
 }
 
+void FPCGEditor::OnNodeSourceCompiled(const UPCGNode* InNode, const FPCGCompilerDiagnostics& InDiagnostics)
+{
+	check(NodeSourceWidget);
+
+	const UPCGSettings* Settings = InNode ? InNode->GetSettings() : nullptr;
+	if (Settings && NodeSourceWidget->GetTextProviderObject() == Settings)
+	{
+		NodeSourceWidget->OnDiagnosticsUpdated(InDiagnostics);
+	}
+	else
+	{
+		NodeSourceWidget->OnDiagnosticsUpdated({});
+	}
+}
+
 void FPCGEditor::OnMapChanged(UWorld* InWorld, EMapChangeType InMapChangedType)
 {
 	if (InMapChangedType != EMapChangeType::SaveMap)
@@ -3272,6 +3314,16 @@ TSharedRef<SDockTab> FPCGEditor::SpawnTab_Log(const FSpawnTabArgs& Args)
 		];
 }
 
+TSharedRef<SDockTab> FPCGEditor::SpawnTab_NodeSource(const FSpawnTabArgs& Args)
+{
+	return SNew(SDockTab)
+		.Label(LOCTEXT("PCGNodeSourceTitle", "Node Source Editor"))
+		.TabColorScale(GetTabColorScale())
+		[
+			NodeSourceWidget.ToSharedRef()
+		];
+}
+
 FText FPCGEditor::GetDetailsTabLabel(int DetailsIndex)
 {
 	if (DetailsIndex == 0)
diff --git a/Engine/Plugins/PCG/Source/PCGEditor/Private/PCGEditor.h b/Engine/Plugins/PCG/Source/PCGEditor/Private/PCGEditor.h
--- a/Engine/Plugins/PCG/Source/PCGEditor/Private/PCGEditor.h
+++ b/Engine/Plugins/PCG/Source/PCGEditor/Private/PCGEditor.h
@@ -22,11 +22,13 @@ class SPCGEditorGraphFind;
 class SPCGEditorGraphLogView;
 class SPCGEditorGraphNodePalette;
 class SPCGEditorGraphProfilingView;
+class SPCGEditorNodeSource;
 class UEdGraphNode;
 class UPCGComponent;
 class UPCGEditorGraph;
 class UPCGEditorGraphNodeBase;
 class UPCGGraph;
+struct FPCGCompilerDiagnostics;
 
 DECLARE_MULTICAST_DELEGATE_OneParam(FOnInspectedStackChanged, const FPCGStack&);
 DECLARE_MULTICAST_DELEGATE_OneParam(FOnInspectedNodeChanged, UPCGEditorGraphNodeBase*);
@@ -58,6 +60,8 @@ public:
 	/** Gets the PCG stack we are inspecting */
 	const FPCGStack* GetStackBeingInspected() const;
 
+	void SetSourceEditorTargetObject(UObject* InObject);
+
 	/** Focus the graph view on a specific node */
 	void JumpToNode(const UEdGraphNode* InNode);
 	void JumpToNode(const UPCGNode* InNode);
@@ -304,9 +308,12 @@ private:
 	/** Create a new profiling tab widget */
 	TSharedRef<SPCGEditorGraphProfilingView> CreateProfilingWidget();
 
-	/** Create a new profiling tab widget */
+	/** Create a new log capture tab widget */
 	TSharedRef<SPCGEditorGraphLogView> CreateLogWidget();
 
+	/** Create a new node source editor tab widget */
+	TSharedRef<SPCGEditorNodeSource> CreateNodeSourceWidget();
+
 	/** Called when the selection changes in the GraphEditor */
 	void OnSelectedNodesChanged(const TSet<UObject*>& NewSelection);
 
@@ -346,6 +353,8 @@ private:
 	void UnregisterDelegatesForWorld(UWorld* World);
 
 	void OnGraphChanged(UPCGGraphInterface* InGraph, EPCGChangeType ChangeType);
+	void OnNodeSourceCompiled(const UPCGNode* InNode, const FPCGCompilerDiagnostics& InDiagnostics);
+
 	void OnMapChanged(UWorld* InWorld, EMapChangeType InMapChangedType);
 	void OnPostPIEStarted(bool bIsSimulating);
 	void OnEndPIE(bool bIsSimulating);
@@ -360,6 +369,7 @@ private:
 	TSharedRef<SDockTab> SpawnTab_Determinism(const FSpawnTabArgs& Args);
 	TSharedRef<SDockTab> SpawnTab_Profiling(const FSpawnTabArgs& Args);
 	TSharedRef<SDockTab> SpawnTab_Log(const FSpawnTabArgs& Args);
+	TSharedRef<SDockTab> SpawnTab_NodeSource(const FSpawnTabArgs& Args);
 
 	FText GetDetailsTabLabel(int DetailsIndex);
 	FText GetDetailsViewObjectName(int DetailsIndex);
@@ -374,6 +384,7 @@ private:
 	TSharedPtr<SPCGEditorGraphDeterminismListView> DeterminismWidget;
 	TSharedPtr<SPCGEditorGraphProfilingView> ProfilingWidget;
 	TSharedPtr<SPCGEditorGraphLogView> LogWidget;
+	TSharedPtr<SPCGEditorNodeSource> NodeSourceWidget;
 
 	TSharedPtr<FUICommandList> GraphEditorCommands;
 
diff --git a/Engine/Plugins/PCG/Source/PCGEditor/Private/PCGEditorModule.cpp b/Engine/Plugins/PCG/Source/PCGEditor/Private/PCGEditorModule.cpp
--- a/Engine/Plugins/PCG/Source/PCGEditor/Private/PCGEditorModule.cpp
+++ b/Engine/Plugins/PCG/Source/PCGEditor/Private/PCGEditorModule.cpp
@@ -3,13 +3,6 @@
 #include "PCGEditorModule.h"
 
 #include "PCGComponent.h"
-#include "PCGEditorCommands.h"
-#include "PCGEditorGraphNodeFactory.h"
-#include "PCGEditorMenuUtils.h"
-#include "PCGEditorProgressNotification.h"
-#include "PCGEditorSettings.h"
-#include "PCGEditorStyle.h"
-#include "PCGEditorUtils.h"
 #include "PCGEngineSettings.h"
 #include "PCGModule.h"
 #include "PCGParamData.h"
@@ -23,6 +16,15 @@
 #include "Grid/PCGPartitionActor.h"
 #include "WorldPartitionBuilder/PCGWorldPartitionBuilder.h"
 
+#include "PCGEditorCommands.h"
+#include "PCGEditorGraphNodeFactory.h"
+#include "PCGEditorMenuUtils.h"
+#include "PCGEditorProgressNotification.h"
+#include "PCGEditorSettings.h"
+#include "PCGEditorStyle.h"
+#include "PCGEditorUtils.h"
+#include "Widgets/SPCGNodeSourceTextBox.h"
+
 #include "ContentBrowserMenuContexts.h"
 #include "ContentBrowserModule.h"
 #include "Editor.h"
@@ -66,7 +68,8 @@ void FPCGEditorModule::StartupModule()
 
 	FPCGEditorCommands::Register();
 	FPCGEditorStyle::Register();
-
+	FPCGNodeSourceEditorTextBoxCommands::Register();
+	
 	GraphNodeFactory = MakeShareable(new FPCGEditorGraphNodeFactory());
 	FEdGraphUtilities::RegisterVisualNodeFactory(GraphNodeFactory);
 
diff --git a/Engine/Plugins/PCG/Source/PCGEditor/Private/PCGEditorStyle.cpp b/Engine/Plugins/PCG/Source/PCGEditor/Private/PCGEditorStyle.cpp
--- a/Engine/Plugins/PCG/Source/PCGEditor/Private/PCGEditorStyle.cpp
+++ b/Engine/Plugins/PCG/Source/PCGEditor/Private/PCGEditorStyle.cpp
@@ -129,6 +129,38 @@ FPCGEditorStyle::FPCGEditorStyle() : FSlateStyleSet("PCGEditorStyle")
 	Set("PCG.Editor.ZoomToSelection", new IMAGE_BRUSH_SVG("Style/PCG_Editor_ZoomToSelection", Icon16x16));
 	Set("PCG.Editor.Pause", new IMAGE_BRUSH_SVG("Icons/PCG_16", Icon16x16));
 	Set("PCG.Editor.AlternatePause", new IMAGE_BRUSH_SVG("Style/PCG_Editor_AlternatePause", Icon16x16));
+
+	// Text editor styles
+	{
+		const FSlateFontInfo Consolas10 = FCoreStyle::GetDefaultFontStyle("Mono", 9);
+
+		const FTextBlockStyle NormalText = FTextBlockStyle()
+			.SetFont(Consolas10)
+			.SetColorAndOpacity(FLinearColor::White)
+			.SetShadowOffset(FVector2D::ZeroVector)
+			.SetShadowColorAndOpacity(FLinearColor::Black)
+			.SetSelectedBackgroundColor(FLinearColor::Blue)
+			.SetHighlightColor(FLinearColor::Yellow)
+			.SetHighlightShape(BOX_BRUSH("Images/TextBlockHighlightShape", FMargin(3.0f / 8.0f)));
+
+		Set("TextEditor.NormalText", NormalText);
+
+		Set("TextEditor.Border", new BOX_BRUSH("Images/TextEditorBorder", FMargin(4.0f / 16.0f), FLinearColor(0.02f, 0.02f, 0.02f, 1)));
+
+		const FEditableTextBoxStyle EditableTextBoxStyle = FEditableTextBoxStyle()
+			.SetTextStyle(NormalText)
+			.SetBackgroundImageNormal(FSlateNoResource())
+			.SetBackgroundImageHovered(FSlateNoResource())
+			.SetBackgroundImageFocused(FSlateNoResource())
+			.SetBackgroundImageReadOnly(FSlateNoResource());
+
+		Set("TextEditor.EditableTextBox", EditableTextBoxStyle);
+
+		FSearchBoxStyle SearchBoxStyle = FCoreStyle::Get().GetWidgetStyle<FSearchBoxStyle>("SearchBox");
+		SearchBoxStyle.SetLeftAlignGlassImageAndClearButton(true);
+		SearchBoxStyle.SetLeftAlignSearchResultButtons(false);
+		Set("TextEditor.SearchBoxStyle", SearchBoxStyle);
+	}
 }
 
 const FPCGEditorStyle& FPCGEditorStyle::Get()
diff --git a/Engine/Plugins/PCG/Source/PCGEditor/Private/PCGHLSLSyntaxHighlighter.cpp b/Engine/Plugins/PCG/Source/PCGEditor/Private/PCGHLSLSyntaxHighlighter.cpp
new file mode 100644
--- /dev/null
+++ b/Engine/Plugins/PCG/Source/PCGEditor/Private/PCGHLSLSyntaxHighlighter.cpp
@@ -0,0 +1,96 @@
+// Copyright Epic Games, Inc. All Rights Reserved.
+
+#include "PCGHLSLSyntaxHighlighter.h"
+#include "PCGEditorStyle.h"
+#include "Styling/SlateTypes.h"
+#include "Framework/Text/IRun.h"
+#include "Framework/Text/TextLayout.h"
+#include "Framework/Text/SlateTextUnderlineLineHighlighter.h"
+
+TSharedRef<FPCGHLSLSyntaxHighlighter> FPCGHLSLSyntaxHighlighter::Create(const FSyntaxTextStyle& InSyntaxTextStyle)
+{
+	return MakeShareable(new FPCGHLSLSyntaxHighlighter(CreateTokenizer(), InSyntaxTextStyle));
+}
+
+void FPCGHLSLSyntaxHighlighter::SetCompilerMessages(const FPCGCompilerDiagnostics& InCompilerMessages)
+{
+	CompilerMessages.Reset();
+
+	for (const FPCGCompilerDiagnostic& Message : InCompilerMessages.Diagnostics)
+	{
+		// Line adjustment required to work correctly in text box.
+		CompilerMessages.Add(Message.Line - 1, Message);
+	}
+
+	MakeDirty();
+}
+
+void FPCGHLSLSyntaxHighlighter::ClearCompilerMessages()
+{
+	if (!CompilerMessages.IsEmpty())
+	{
+		CompilerMessages.Reset();
+		MakeDirty();
+	}
+}
+
+FPCGHLSLSyntaxHighlighter::FPCGHLSLSyntaxHighlighter(TSharedPtr<ISyntaxTokenizer> InTokenizer, const FSyntaxTextStyle& InSyntaxTextStyle) :
+	FHLSLSyntaxHighlighterMarshaller(MoveTemp(InTokenizer), InSyntaxTextStyle)
+{
+}
+
+void FPCGHLSLSyntaxHighlighter::ParseTokens(const FString& SourceString, FTextLayout& TargetTextLayout,	TArray<ISyntaxTokenizer::FTokenizedLine> TokenizedLines)
+{
+	LineHighlightsToAdd.Empty();
+	FHLSLSyntaxHighlighterMarshaller::ParseTokens(SourceString, TargetTextLayout, TokenizedLines);
+
+	TargetTextLayout.SetLineHighlights(LineHighlightsToAdd);
+}
+
+FTextLayout::FNewLineData FPCGHLSLSyntaxHighlighter::ProcessTokenizedLine(const ISyntaxTokenizer::FTokenizedLine& TokenizedLine, const int32& LineNumber, const FString& SourceString, EParseState& CurrentParseState)
+{
+	FTextLayout::FNewLineData LineData = FHLSLSyntaxHighlighterMarshaller::ProcessTokenizedLine(TokenizedLine, LineNumber, SourceString, CurrentParseState);
+
+	FTextBlockStyle ErrorTextStyle = SyntaxTextStyle.ErrorTextStyle;
+	TSharedPtr<FSlateTextUnderlineLineHighlighter> UnderlineLineHighlighter = FSlateTextUnderlineLineHighlighter::Create(ErrorTextStyle.UnderlineBrush, ErrorTextStyle.Font, ErrorTextStyle.ColorAndOpacity, ErrorTextStyle.ShadowOffset, ErrorTextStyle.ShadowColorAndOpacity);
+	
+	TArray<const FPCGCompilerDiagnostic*> Diagnostics;
+	CompilerMessages.MultiFindPointer(LineNumber, Diagnostics);
+		
+	for (const FPCGCompilerDiagnostic* Diagnostic : Diagnostics)
+	{
+		// If no column specified, underline the whole source line.
+		const bool bColumnsSpecified = Diagnostic->ColumnStart != -1 && Diagnostic->ColumnEnd != -1;
+
+		// ColumnStart/ColumnEnd are 1-based, closed interval. FTextRange is 0 based, half-closed interval. 
+		FTextRange UnderlineRange(Diagnostic->ColumnStart - 1, Diagnostic->ColumnEnd);
+
+		// The highlighting lines have to match the runs and not exceed their bounds, so we chop them up. 
+		for (TSharedRef<IRun> Run : LineData.Runs)
+		{
+			if (!bColumnsSpecified || !Run->GetTextRange().Intersect(UnderlineRange).IsEmpty())
+			{
+				LineHighlightsToAdd.Add(FTextLineHighlight(
+					LineNumber,
+					bColumnsSpecified ? UnderlineRange : Run->GetTextRange(),
+					FSlateTextUnderlineLineHighlighter::DefaultZIndex,
+					UnderlineLineHighlighter.ToSharedRef()));
+			}
+		}
+	}
+
+	return LineData;
+}
+
+FPCGHLSLSyntaxHighlighter::FSyntaxTextStyle FPCGHLSLSyntaxHighlighter::GetSyntaxTextStyle()
+{
+	return FSyntaxTextStyle(
+		FAppStyle::Get().GetWidgetStyle<FTextBlockStyle>("SyntaxHighlight.SourceCode.Normal"),
+		FAppStyle::Get().GetWidgetStyle<FTextBlockStyle>("SyntaxHighlight.SourceCode.Operator"),
+		FAppStyle::Get().GetWidgetStyle<FTextBlockStyle>("SyntaxHighlight.SourceCode.Keyword"),
+		FAppStyle::Get().GetWidgetStyle<FTextBlockStyle>("SyntaxHighlight.SourceCode.String"),
+		FAppStyle::Get().GetWidgetStyle<FTextBlockStyle>("SyntaxHighlight.SourceCode.Number"),
+		FAppStyle::Get().GetWidgetStyle<FTextBlockStyle>("SyntaxHighlight.SourceCode.Comment"),
+		FAppStyle::Get().GetWidgetStyle<FTextBlockStyle>("SyntaxHighlight.SourceCode.PreProcessorKeyword"),
+		FAppStyle::Get().GetWidgetStyle<FTextBlockStyle>("SyntaxHighlight.SourceCode.Error"));
+}
diff --git a/Engine/Plugins/PCG/Source/PCGEditor/Private/PCGHLSLSyntaxHighlighter.h b/Engine/Plugins/PCG/Source/PCGEditor/Private/PCGHLSLSyntaxHighlighter.h
new file mode 100644
--- /dev/null
+++ b/Engine/Plugins/PCG/Source/PCGEditor/Private/PCGHLSLSyntaxHighlighter.h
@@ -0,0 +1,37 @@
+// Copyright Epic Games, Inc. All Rights Reserved.
+
+#pragma once
+
+#include "CoreMinimal.h"
+
+#include "Compute/PCGCompilerDiagnostic.h"
+
+#include "Framework/Text/SyntaxTokenizer.h"
+#include "Text/HLSLSyntaxHighlighterMarshaller.h"
+
+class FTextLayout;
+struct FPCGCompilerDiagnostic;
+
+/**
+ * Get/set the raw text to/from a text layout, and also inject syntax highlighting for our rich-text markup
+ */
+class FPCGHLSLSyntaxHighlighter : public FHLSLSyntaxHighlighterMarshaller
+{
+public:
+	static TSharedRef<FPCGHLSLSyntaxHighlighter> Create(const FSyntaxTextStyle& InSyntaxTextStyle = GetSyntaxTextStyle());
+	static FSyntaxTextStyle GetSyntaxTextStyle();
+
+	void SetCompilerMessages(const FPCGCompilerDiagnostics& InCompilerMessages);
+	void ClearCompilerMessages();
+
+protected:
+	FPCGHLSLSyntaxHighlighter(TSharedPtr<ISyntaxTokenizer> InTokenizer, const FSyntaxTextStyle& InSyntaxTextStyle);
+
+	virtual void ParseTokens(const FString& SourceString, FTextLayout& TargetTextLayout, TArray<ISyntaxTokenizer::FTokenizedLine> TokenizedLines) override;
+	virtual FTextLayout::FNewLineData ProcessTokenizedLine(const ISyntaxTokenizer::FTokenizedLine& TokenizedLine, const int32& LineNumber, const FString& SourceString, EParseState& CurrentParseState) override;
+
+	TMultiMap<int32 /*Line*/, FPCGCompilerDiagnostic> CompilerMessages;
+
+private:
+	TArray<FTextLineHighlight> LineHighlightsToAdd;
+};
diff --git a/Engine/Plugins/PCG/Source/PCGEditor/Private/Widgets/SPCGEditorNodeSource.cpp b/Engine/Plugins/PCG/Source/PCGEditor/Private/Widgets/SPCGEditorNodeSource.cpp
new file mode 100644
--- /dev/null
+++ b/Engine/Plugins/PCG/Source/PCGEditor/Private/Widgets/SPCGEditorNodeSource.cpp
@@ -0,0 +1,238 @@
+// Copyright Epic Games, Inc. All Rights Reserved.
+
+#include "SPCGEditorNodeSource.h"
+
+#include "Compute/IPCGNodeSourceTextProvider.h"
+
+#include "PCGEditor.h"
+#include "PCGHLSLSyntaxHighlighter.h"
+#include "SPCGNodeSourceTextBox.h"
+
+#include "Widgets/Input/SButton.h"
+#include "Widgets/Layout/SSpacer.h"
+
+#define LOCTEXT_NAMESPACE "SPCGEditorNodeSource"
+
+SPCGEditorNodeSource::SPCGEditorNodeSource()
+	: SyntaxHighlighterDeclarations(FPCGHLSLSyntaxHighlighter::Create())
+	, SyntaxHighlighterShaderText(FPCGHLSLSyntaxHighlighter::Create())
+{
+}
+
+void SPCGEditorNodeSource::Construct(const FArguments& InArgs, TSharedPtr<FPCGEditor> InPCGEditor)
+{
+	SExpandableArea::FArguments ExpandableAreaArgs;
+	ExpandableAreaArgs.AreaTitleFont(FAppStyle::GetFontStyle("DetailsView.CategoryFontStyle"));
+	
+	ChildSlot
+	[
+		SNew(SSplitter)
+		.Orientation(Orient_Vertical)
+		.PhysicalSplitterHandleSize(4.0f)
+		.HitDetectionSplitterHandleSize(6.0f)
+		+SSplitter::Slot()
+		.Value(1.0)
+		.SizeRule(this, &SPCGEditorNodeSource::GetDeclarationsSectionSizeRule)
+		[
+			SAssignNew(DeclarationsExpandableArea, SExpandableArea) = ExpandableAreaArgs
+			.AreaTitle(LOCTEXT("PCGNodeSource_Declarations_Title", "Declarations (Read-Only)"))
+			.InitiallyCollapsed(true)
+			.AllowAnimatedTransition(false)
+			.BodyContent()
+			[
+				SAssignNew(DeclarationsTextBox, SPCGNodeSourceTextBox)
+				.Text(this, &SPCGEditorNodeSource::GetDeclarationsAsText)
+				.IsReadOnly(true)
+				.Marshaller(SyntaxHighlighterDeclarations)
+			]
+		]
+		+SSplitter::Slot()
+		.Value(1.0)
+		.SizeRule(this, &SPCGEditorNodeSource::GetShaderFunctionsSectionSizeRule)
+		[
+			SAssignNew(ShaderFunctionsExpandableArea, SExpandableArea) = ExpandableAreaArgs
+			.AreaTitle(LOCTEXT("PCGNodeSource_ShaderFunctions_Title", "Shader Functions"))
+			.InitiallyCollapsed(true)
+			.AllowAnimatedTransition(false)
+			.BodyContent()
+			[
+				SAssignNew(ShaderFunctionsTextBox, SPCGNodeSourceTextBox)
+				.Text(this, &SPCGEditorNodeSource::GetShaderFunctionsAsText)
+				.IsReadOnly(this, &SPCGEditorNodeSource::IsShaderTextReadOnly)
+				.Marshaller(SyntaxHighlighterDeclarations)
+				.OnTextChanged(this, &SPCGEditorNodeSource::OnShaderFunctionsTextChanged)
+				.OnTextChangesApplied(this, &SPCGEditorNodeSource::OnShaderTextChangesApplied)
+			]
+		]
+		+SSplitter::Slot()
+		.Value(1.0)
+		[
+			SNew(SVerticalBox)
+			+SVerticalBox::Slot()
+			.AutoHeight()
+			[
+				ConstructNonExpandableHeaderWidget(ExpandableAreaArgs.AreaTitle(LOCTEXT("PCGNodeSource_ShaderText_Title", "Shader Text")))
+			]
+			+SVerticalBox::Slot()
+			.FillHeight(1.0f)
+			[
+				SAssignNew(ShaderTextTextBox, SPCGNodeSourceTextBox)
+				.Text(this, &SPCGEditorNodeSource::GetShaderTextAsText)
+				.IsReadOnly(this, &SPCGEditorNodeSource::IsShaderTextReadOnly)
+				.Marshaller(SyntaxHighlighterShaderText)
+				.OnTextChanged(this, &SPCGEditorNodeSource::OnShaderTextChanged)
+				.OnTextChangesApplied(this, &SPCGEditorNodeSource::OnShaderTextChangesApplied)
+			]
+		]
+	];
+}
+
+void SPCGEditorNodeSource::SetTextProviderObject(UObject* InProviderObject)
+{
+	ShaderTextProviderObject = InProviderObject;
+
+	Refresh();
+}
+
+void SPCGEditorNodeSource::Refresh()
+{
+	DeclarationsTextBox->Refresh();
+	ShaderTextTextBox->Refresh();
+}
+
+SSplitter::ESizeRule SPCGEditorNodeSource::GetDeclarationsSectionSizeRule() const
+{
+	SSplitter::ESizeRule SizeRule = SSplitter::ESizeRule::FractionOfParent;
+
+	if (ensure(DeclarationsExpandableArea.IsValid()))
+	{
+		SizeRule = DeclarationsExpandableArea->IsExpanded() ? SSplitter::ESizeRule::FractionOfParent : SSplitter::ESizeRule::SizeToContent;
+	}
+
+	return SizeRule;
+}
+
+SSplitter::ESizeRule SPCGEditorNodeSource::GetShaderFunctionsSectionSizeRule() const
+{
+	SSplitter::ESizeRule SizeRule = SSplitter::ESizeRule::FractionOfParent;
+
+	if (ensure(ShaderFunctionsExpandableArea.IsValid()))
+	{
+		SizeRule = ShaderFunctionsExpandableArea->IsExpanded() ? SSplitter::ESizeRule::FractionOfParent : SSplitter::ESizeRule::SizeToContent;
+	}
+
+	return SizeRule;
+}
+
+IPCGNodeSourceTextProvider* SPCGEditorNodeSource::GetTextProviderInterface() const
+{
+	return Cast<IPCGNodeSourceTextProvider>(ShaderTextProviderObject);
+}
+
+FText SPCGEditorNodeSource::GetDeclarationsAsText() const
+{
+	if (GetTextProviderInterface())
+	{
+		return FText::FromString(GetTextProviderInterface()->GetDeclarationsText());
+	}
+
+	return FText::GetEmpty();
+}
+
+FText SPCGEditorNodeSource::GetShaderFunctionsAsText() const
+{
+	if (GetTextProviderInterface())
+	{
+		return FText::FromString(GetTextProviderInterface()->GetShaderFunctionsText());
+	}
+
+	return FText::GetEmpty();
+}
+
+FText SPCGEditorNodeSource::GetShaderTextAsText() const
+{
+	if (GetTextProviderInterface())
+	{
+		return FText::FromString(GetTextProviderInterface()->GetShaderText());
+	}
+
+	return FText::GetEmpty();
+}
+
+bool SPCGEditorNodeSource::IsShaderTextReadOnly() const
+{
+	if (IPCGNodeSourceTextProvider* Provider = Cast<IPCGNodeSourceTextProvider>(ShaderTextProviderObject))
+	{
+		return Provider->IsShaderTextReadOnly();
+	}
+
+	return false;
+}
+
+void SPCGEditorNodeSource::OnShaderFunctionsTextChanged(const FText& InText) const
+{
+	if (GetTextProviderInterface())
+	{
+		GetTextProviderInterface()->SetShaderFunctionsText(InText.ToString());
+	}
+}
+
+void SPCGEditorNodeSource::OnShaderTextChanged(const FText& InText) const
+{
+	// Always clear compiler messages when text is edited so that markup doesn't lurk in random places.
+	SyntaxHighlighterShaderText->ClearCompilerMessages();
+
+	if (GetTextProviderInterface())
+	{
+		GetTextProviderInterface()->SetShaderText(InText.ToString());
+	}
+}
+
+void SPCGEditorNodeSource::OnShaderTextChangesApplied() const
+{
+	if (GetTextProviderInterface())
+	{
+		GetTextProviderInterface()->ApplySourceChanges();
+	}
+}
+
+void SPCGEditorNodeSource::OnDiagnosticsUpdated(const FPCGCompilerDiagnostics& InDiagnostics) const
+{
+	SyntaxHighlighterShaderText->SetCompilerMessages(InDiagnostics);
+	ShaderTextTextBox->Refresh();
+}
+
+TSharedRef<SWidget> SPCGEditorNodeSource::ConstructNonExpandableHeaderWidget(const SExpandableArea::FArguments& InArgs) const
+{
+	return SNew(SBorder)
+		.BorderImage(FStyleDefaults::GetNoBrush())
+		.BorderBackgroundColor(FLinearColor::Transparent)
+		.Padding(0.0f)
+		[
+			SNew(SButton)
+			.ButtonStyle(FCoreStyle::Get(), "NoBorder")
+			.ContentPadding(InArgs._HeaderPadding)
+			.ForegroundColor(FSlateColor::UseForeground())
+			[
+				SNew(SHorizontalBox)
+				+SHorizontalBox::Slot()
+				.AutoWidth()
+				.Padding(InArgs._AreaTitlePadding)
+				.VAlign(VAlign_Center)
+				[
+					SNew(SSpacer)
+					.Size(InArgs._Style->CollapsedImage.GetImageSize())
+				]
+				+SHorizontalBox::Slot()
+				.FillWidth(1.0f)
+				.VAlign(VAlign_Center)
+				[
+					SNew(STextBlock)
+					.Text(InArgs._AreaTitle)
+					.Font(InArgs._AreaTitleFont)
+				]
+			]
+		];
+}
+
+#undef LOCTEXT_NAMESPACE
diff --git a/Engine/Plugins/PCG/Source/PCGEditor/Private/Widgets/SPCGEditorNodeSource.h b/Engine/Plugins/PCG/Source/PCGEditor/Private/Widgets/SPCGEditorNodeSource.h
new file mode 100644
--- /dev/null
+++ b/Engine/Plugins/PCG/Source/PCGEditor/Private/Widgets/SPCGEditorNodeSource.h
@@ -0,0 +1,64 @@
+// Copyright Epic Games, Inc. All Rights Reserved.
+
+#pragma once
+
+#include "Widgets/SCompoundWidget.h"
+
+#include "Widgets/Layout/SExpandableArea.h"
+#include "Widgets/Views/ITableRow.h"
+#include "Widgets/Views/SListView.h"
+#include "Widgets/Views/STableRow.h"
+
+class FPCGEditor;
+class FPCGHLSLSyntaxHighlighter;
+class IPCGNodeSourceTextProvider;
+class SPCGNodeSourceTextBox;
+struct FPCGCompilerDiagnostics;
+
+/** Window for editing node source text. */
+class SPCGEditorNodeSource : public SCompoundWidget
+{
+public:
+	SLATE_BEGIN_ARGS(SPCGEditorNodeSource) { }
+	SLATE_END_ARGS()
+
+	SPCGEditorNodeSource();
+
+	void Construct(const FArguments& InArgs, TSharedPtr<FPCGEditor> InPCGEditor);
+
+	void SetTextProviderObject(UObject* InProviderObject);
+	const UObject* GetTextProviderObject() const { return ShaderTextProviderObject.IsValid() ? ShaderTextProviderObject.Get() : nullptr; }
+
+	void OnDiagnosticsUpdated(const FPCGCompilerDiagnostics& InDiagnostics) const;
+
+private:
+	void Refresh();
+
+	SSplitter::ESizeRule GetDeclarationsSectionSizeRule() const;
+	SSplitter::ESizeRule GetShaderFunctionsSectionSizeRule() const;
+	IPCGNodeSourceTextProvider* GetTextProviderInterface() const;
+
+	FText GetDeclarationsAsText() const;
+	FText GetShaderFunctionsAsText() const;
+	FText GetShaderTextAsText() const;
+
+	bool IsShaderTextReadOnly() const;
+
+	void OnShaderFunctionsTextChanged(const FText& InText) const;
+	void OnShaderTextChanged(const FText& InText) const;
+	void OnShaderTextChangesApplied() const;
+
+	TSharedRef<SWidget> ConstructNonExpandableHeaderWidget(const SExpandableArea::FArguments& InArgs) const;
+
+	TSharedRef<FPCGHLSLSyntaxHighlighter> SyntaxHighlighterDeclarations;
+	TSharedRef<FPCGHLSLSyntaxHighlighter> SyntaxHighlighterShaderText;
+
+	TWeakObjectPtr<UObject> ShaderTextProviderObject;
+
+	TSharedPtr<SExpandableArea> DeclarationsExpandableArea;
+	TSharedPtr<SExpandableArea> ShaderFunctionsExpandableArea;
+
+	TSharedPtr<SPCGNodeSourceTextBox> DeclarationsTextBox;
+	TSharedPtr<SPCGNodeSourceTextBox> ShaderFunctionsTextBox;
+	TSharedPtr<SPCGNodeSourceTextBox> ShaderTextTextBox;
+};
diff --git a/Engine/Plugins/PCG/Source/PCGEditor/Private/Widgets/SPCGNodeSourceTextBox.cpp b/Engine/Plugins/PCG/Source/PCGEditor/Private/Widgets/SPCGNodeSourceTextBox.cpp
new file mode 100644
--- /dev/null
+++ b/Engine/Plugins/PCG/Source/PCGEditor/Private/Widgets/SPCGNodeSourceTextBox.cpp
@@ -0,0 +1,739 @@
+// Copyright Epic Games, Inc. All Rights Reserved.
+
+#include "SPCGNodeSourceTextBox.h"
+
+#include "PCGEditorStyle.h"
+#include "SPCGShaderTextSearchWidget.h"
+
+#include "Framework/Application/SlateApplication.h"
+#include "Framework/Application/SlateUser.h"
+#include "Framework/MultiBox/MultiBoxBuilder.h"
+#include "Styling/AppStyle.h" 
+#include "Widgets/Layout/SGridPanel.h"
+#include "Widgets/Layout/SScrollBar.h"
+#include "Widgets/Text/SMultiLineEditableText.h"
+
+#define LOCTEXT_NAMESPACE "PCGShaderTextDocumentTextBox"
+
+static int32 GetNumSpacesAtStartOfLine(const FString& InLine)
+{
+	int32 NumSpaces = 0;
+	for (const TCHAR Char : InLine)
+	{
+		if ((Char != TEXT(' ')))
+		{
+			break;
+		}
+				
+		NumSpaces++;
+	}
+
+	return NumSpaces;
+}
+
+static bool IsOpenBrace(const TCHAR& InCharacter)
+{
+	return (InCharacter == TEXT('{') || InCharacter == TEXT('[') || InCharacter == TEXT('('));
+}
+
+static bool IsCloseBrace(const TCHAR& InCharacter)
+{
+	return (InCharacter == TEXT('}') || InCharacter == TEXT(']') || InCharacter == TEXT(')'));
+}
+
+static TCHAR GetMatchedCloseBrace(const TCHAR& InCharacter)
+{
+	return InCharacter == TEXT('{') ? TEXT('}') :
+		(InCharacter == TEXT('[') ? TEXT(']') : TEXT(')'));
+}
+
+static bool IsWhiteSpace(const TCHAR& InCharacter)
+{
+	return InCharacter == TEXT(' ') || InCharacter == TEXT('\r') || InCharacter == TEXT('\n');
+}
+
+FPCGNodeSourceEditorTextBoxCommands::FPCGNodeSourceEditorTextBoxCommands() 
+	: TCommands<FPCGNodeSourceEditorTextBoxCommands>(
+		"PCGShaderTextEditorDocumentTextBox", // Context name for fast lookup
+		NSLOCTEXT("Contexts", "PCGShaderTextEditorDocumentTextBox", "PCG Node Source Editor Document TextBox"), // Localized context name for display
+		NAME_None,
+		FAppStyle::GetAppStyleSetName()
+	)
+{
+}
+
+void FPCGNodeSourceEditorTextBoxCommands::RegisterCommands()
+{
+	UI_COMMAND(ApplyChanges, "Apply Changes", "Apply Changes", EUserInterfaceActionType::Button, FInputChord(EKeys::Enter, EModifierKey::Alt));
+
+	UI_COMMAND(Search, "Search", "Search for a String", EUserInterfaceActionType::Button, FInputChord(EKeys::F, EModifierKey::Control));
+	UI_COMMAND(NextOccurrence, "Next Occurrence", "Go to Next Occurrence", EUserInterfaceActionType::Button, FInputChord(EKeys::F3, EModifierKey::None));
+	UI_COMMAND(PreviousOccurrence, "Previous Occurrence", "Go to Previous Occurrence", EUserInterfaceActionType::Button, FInputChord(EKeys::F3, EModifierKey::Shift));
+
+	UI_COMMAND(ToggleComment, "Toggle Comment", "Comment/Uncomment selected lines", EUserInterfaceActionType::Button, FInputChord(EKeys::Slash, EModifierKey::Control));
+}
+
+SPCGNodeSourceTextBox::SPCGNodeSourceTextBox()
+	: bIsSearchBarHidden(true)
+	, TopLevelCommandList(MakeShared<FUICommandList>())
+	, TextCommandList(MakeShared<FUICommandList>())
+{
+}
+
+SPCGNodeSourceTextBox::~SPCGNodeSourceTextBox()
+{
+}
+
+void SPCGNodeSourceTextBox::Construct(const FArguments& InArgs)
+{
+	RegisterCommands();
+	
+	OnTextChangesAppliedCallback = InArgs._OnTextChangesApplied;
+
+	const TSharedPtr<SScrollBar> HScrollBar = SNew(SScrollBar)
+		.Orientation(EOrientation::Orient_Horizontal);
+	
+	const TSharedPtr<SScrollBar> VScrollBar = SNew(SScrollBar)
+		.Orientation(EOrientation::Orient_Vertical);
+	
+	const FTextBlockStyle& TextStyle = FAppStyle::Get().GetWidgetStyle<FTextBlockStyle>("MessageLog");
+	const FSlateFontInfo &Font = TextStyle.Font;
+
+	const bool bReadOnly = InArgs._IsReadOnly.Get();
+	const FPCGNodeSourceEditorTextBoxCommands& Commands = FPCGNodeSourceEditorTextBoxCommands::Get();
+	
+	Text = SNew(SMultiLineEditableText)
+			.Font(Font)
+			.TextStyle(&TextStyle)
+			.Text(InArgs._Text)
+			.OnTextChanged(InArgs._OnTextChanged)
+			.OnKeyCharHandler(this, &SPCGNodeSourceTextBox::OnTextKeyChar)
+			.OnKeyDownHandler(this, &SPCGNodeSourceTextBox::OnTextKeyDown)
+			// By default, the Tab key gets routed to "next widget". We want to disable that behaviour.
+			.OnIsTypedCharValid_Lambda([](const TCHAR InChar) { return true; })
+			.Marshaller(InArgs._Marshaller)
+			.AutoWrapText(false)
+			.ClearTextSelectionOnFocusLoss(false)
+			.AllowContextMenu(true)
+			.ContextMenuExtender(
+				FMenuExtensionDelegate::CreateLambda(
+					[this, bReadOnly, Commands](FMenuBuilder& InBuilder)
+					{
+						if (!bReadOnly)
+						{
+							InBuilder.PushCommandList(TextCommandList);
+							InBuilder.AddMenuEntry(Commands.ApplyChanges);
+							InBuilder.AddMenuEntry(Commands.ToggleComment);
+						}
+						
+					}))
+			.IsReadOnly(bReadOnly)
+			.HScrollBar(HScrollBar)
+			.VScrollBar(VScrollBar);
+
+	SearchBar = SNew(SPCGShaderTextSearchWidget)
+		.OnTextChanged(this, &SPCGNodeSourceTextBox::OnSearchTextChanged)
+		.OnTextCommitted(this, &SPCGNodeSourceTextBox::OnSearchTextCommitted)
+		.SearchResultData(this, &SPCGNodeSourceTextBox::GetSearchResultData)
+		.OnResultNavigationButtonClicked(this, &SPCGNodeSourceTextBox::OnSearchResultNavigationButtonClicked);
+
+	ChildSlot
+	[
+		SAssignNew(TabBody, SVerticalBox)
+		+SVerticalBox::Slot()
+		[
+			SNew(SBorder)
+			.BorderImage(FPCGEditorStyle::Get().GetBrush("TextEditor.Border"))
+			.BorderBackgroundColor(FLinearColor(0.55f, 0.55f, 0.55f, 1.0f))
+			[
+				SNew(SGridPanel)
+				.FillColumn(0,1.0f)
+				.FillRow(0,1.0f)
+				+SGridPanel::Slot(0,0)
+				[
+					Text.ToSharedRef()
+				]
+				+SGridPanel::Slot(1,0)
+				[
+					VScrollBar.ToSharedRef()
+				]
+				+SGridPanel::Slot(0,1)
+				[
+					HScrollBar.ToSharedRef()
+				]
+			]
+		]
+	];
+}
+
+void SPCGNodeSourceTextBox::RegisterCommands()
+{
+	const FPCGNodeSourceEditorTextBoxCommands& Commands = FPCGNodeSourceEditorTextBoxCommands::Get();
+	
+	TopLevelCommandList->MapAction(Commands.Search, FExecuteAction::CreateSP(this, &SPCGNodeSourceTextBox::OnTriggerSearch));
+	TopLevelCommandList->MapAction(Commands.NextOccurrence, FExecuteAction::CreateSP(this, &SPCGNodeSourceTextBox::OnGoToNextOccurrence));
+	TopLevelCommandList->MapAction(Commands.PreviousOccurrence, FExecuteAction::CreateSP(this, &SPCGNodeSourceTextBox::OnGoToPreviousOccurrence));
+	TextCommandList->MapAction(Commands.ApplyChanges, FExecuteAction::CreateSP(this, &SPCGNodeSourceTextBox::OnApplyChanges));
+	TextCommandList->MapAction(Commands.ToggleComment, FExecuteAction::CreateSP(this, &SPCGNodeSourceTextBox::OnToggleComment));
+}
+
+FReply SPCGNodeSourceTextBox::OnPreviewKeyDown(const FGeometry& MyGeometry, const FKeyEvent& InKeyEvent)
+{
+	const FKey Key = InKeyEvent.GetKey();
+	
+	if (Key == EKeys::Escape)
+	{
+		if (HandleEscape())
+		{
+			return FReply::Handled();
+		}
+	}
+
+	if (TopLevelCommandList->ProcessCommandBindings(InKeyEvent))
+	{
+		return FReply::Handled();
+	}
+	
+	return SCompoundWidget::OnPreviewKeyDown(MyGeometry, InKeyEvent);
+}
+
+bool SPCGNodeSourceTextBox::HandleEscape()
+{
+	if (HideSearchBar())
+	{
+		return true;
+	}
+
+	return false;
+}
+
+void SPCGNodeSourceTextBox::ShowSearchBar()
+{
+	if (bIsSearchBarHidden)
+	{
+		bIsSearchBarHidden = false;
+		
+		TabBody->InsertSlot(0)
+		.AutoHeight()
+		[
+			SearchBar.ToSharedRef()
+		];
+	}
+}
+
+bool SPCGNodeSourceTextBox::HideSearchBar() 
+{
+	if (!bIsSearchBarHidden)
+	{
+		bIsSearchBarHidden = true;
+		SearchBar->ClearSearchText();
+		TabBody->RemoveSlot(SearchBar.ToSharedRef());
+			
+		FSlateApplication::Get().ForEachUser([&](FSlateUser& User) {
+			User.SetFocus(Text.ToSharedRef(), EFocusCause::SetDirectly);
+		});
+
+		return true;
+	}
+	
+	return false;
+}
+
+void SPCGNodeSourceTextBox::OnTriggerSearch()
+{
+	ShowSearchBar();
+
+	FText SelectedText = Text->GetSelectedText();
+	
+	SearchBar->TriggerSearch(SelectedText);
+}
+
+void SPCGNodeSourceTextBox::Refresh() const
+{
+	Text->Refresh();
+}
+
+void SPCGNodeSourceTextBox::OnSearchTextChanged(const FText& InTextToSearch)
+{
+	// we start the search from the beginning of current selection.
+	// goto clears the selection, but it will be restored by the first search
+	Text->GoTo(Text->GetSelection().GetBeginning());
+
+	Text->SetSearchText(InTextToSearch);
+}
+
+void SPCGNodeSourceTextBox::OnSearchTextCommitted(const FText& InTextToSearch, ETextCommit::Type InCommitType)
+{
+	if (!InTextToSearch.EqualTo(Text->GetSearchText()))
+	{
+		Text->SetSearchText(InTextToSearch);
+	}
+	else
+	{
+		if (InCommitType == ETextCommit::Type::OnEnter)
+		{
+			OnSearchResultNavigationButtonClicked(SSearchBox::SearchDirection::Next);
+		}
+	}
+}
+
+TOptional<SSearchBox::FSearchResultData> SPCGNodeSourceTextBox::GetSearchResultData() const
+{
+	FText SearchText = Text->GetSearchText();
+	
+	if (!SearchText.IsEmpty())
+	{
+		SSearchBox::FSearchResultData Result;
+		Result.CurrentSearchResultIndex = Text->GetSearchResultIndex();
+		Result.NumSearchResults = Text->GetNumSearchResults();
+		
+		return Result;
+	}
+		
+	return TOptional<SSearchBox::FSearchResultData>();
+}
+
+void SPCGNodeSourceTextBox::OnSearchResultNavigationButtonClicked(SSearchBox::SearchDirection InDirection)
+{
+	Text->AdvanceSearch(InDirection == SSearchBox::SearchDirection::Previous);
+}
+
+void SPCGNodeSourceTextBox::OnGoToNextOccurrence()
+{
+	OnSearchResultNavigationButtonClicked(SSearchBox::SearchDirection::Next);
+}
+
+void SPCGNodeSourceTextBox::OnGoToPreviousOccurrence()
+{
+	OnSearchResultNavigationButtonClicked(SSearchBox::SearchDirection::Previous);
+}
+
+FReply SPCGNodeSourceTextBox::OnTextKeyDown(
+	const FGeometry& MyGeometry,
+	const FKeyEvent& InKeyEvent) const
+{
+	if (TextCommandList->ProcessCommandBindings(InKeyEvent))
+	{
+		return FReply::Handled();
+	}
+
+	// Let SMultiLineEditableText::OnKeyDown handle it.
+	return FReply::Unhandled();
+}
+
+void SPCGNodeSourceTextBox::OnApplyChanges() const
+{
+	OnTextChangesAppliedCallback.ExecuteIfBound();
+}
+
+void SPCGNodeSourceTextBox::OnToggleComment() const
+{
+	SMultiLineEditableText::FScopedEditableTextTransaction Transaction(Text);
+	
+	const FTextLocation CursorLocation = Text->GetCursorLocation();
+	
+	const FTextSelection Selection = Text->GetSelection();
+	const FTextLocation SelectionStart =
+		CursorLocation == Selection.GetBeginning() ? Selection.GetEnd() : Selection.GetBeginning();
+	
+	// Need to shift the selection according to the new indentation
+	FTextLocation NewCursorLocation = CursorLocation;
+	FTextLocation NewSelectionStart = SelectionStart;
+		
+	const int32 StartLine = Selection.GetBeginning().GetLineIndex();
+	const int32 EndLine = Selection.GetEnd().GetLineIndex();
+
+	bool bShouldComment = false;
+	bool bAreAllLinesEmpty = true;
+	int32 MinNumLeadingSpaces = INDEX_NONE;
+	for (int32 Index = StartLine; Index <= EndLine; Index++)
+	{
+		FString Line;
+		Text->GetTextLine(Index, Line);
+		// Empty lines are not considered when deciding whether to comment/uncomment
+		const int32 NumSpaces = GetNumSpacesAtStartOfLine(Line);
+		if (NumSpaces >= Line.Len())	
+		{
+			continue;
+		}
+
+		bAreAllLinesEmpty = false;
+		
+		FStringView LineView(&Line[NumSpaces]);
+		if (!LineView.StartsWith(TEXT("//")))
+		{
+			bShouldComment = true;
+		}
+
+		MinNumLeadingSpaces =
+			MinNumLeadingSpaces == INDEX_NONE ?
+			NumSpaces : FMath::Min(MinNumLeadingSpaces, NumSpaces);
+	}
+
+	// Find nearest tab
+	MinNumLeadingSpaces = MinNumLeadingSpaces / 4 * 4;
+
+	// Single line comment or single line no-op should move the cursor down
+	const bool bShouldGoToNewLine = !Text->AnyTextSelected() && (bShouldComment || bAreAllLinesEmpty);
+	
+	for (int32 Index = StartLine; Index <= EndLine; Index++)
+	{
+		FString Line;
+		Text->GetTextLine(Index, Line);
+
+		if (!Line.IsEmpty())
+		{
+			const int32 NumSpaces = GetNumSpacesAtStartOfLine(Line);
+			if (NumSpaces < Line.Len())
+			{
+				int32 LineShift = 0;
+				if (bShouldComment)
+				{
+					Text->GoTo(FTextLocation(Index, MinNumLeadingSpaces));
+					Text->InsertTextAtCursor(TEXT("// "));
+					LineShift = 3;
+				}
+				else
+				{
+					int32 CommentOffset = Line.Find("//");
+					if (ensure(CommentOffset != INDEX_NONE))
+					{
+						int32 ContentOffset =
+							Line.IsValidIndex(CommentOffset+2) && Line[CommentOffset+2] == TEXT(' ') ?
+								3 : 2;
+						Text->SelectText({Index, CommentOffset}, {Index, CommentOffset+ContentOffset});
+						Text->DeleteSelectedText();
+						LineShift = -ContentOffset;
+					}
+				}
+
+				if (Index == CursorLocation.GetLineIndex())
+				{
+					NewCursorLocation = FTextLocation(CursorLocation, LineShift);
+				}
+					
+				if (Index == SelectionStart.GetLineIndex())
+				{
+					NewSelectionStart = FTextLocation(SelectionStart, LineShift);
+				}	
+			}
+		}
+	}
+
+	if (bShouldGoToNewLine)
+	{
+		// Automatically go to the next line if there is no selection
+		int32 LineIndex = CursorLocation.GetLineIndex() + 1;
+		if (LineIndex < Text->GetTextLineCount())
+		{
+			FString Line;
+			Text->GetTextLine(LineIndex, Line);
+			
+			int32 Offset = FMath::Min(CursorLocation.GetOffset(), Line.Len());
+			
+			Text->GoTo({LineIndex, Offset});
+		}
+		else
+		{
+			Text->GoTo(NewCursorLocation);
+		}
+	}
+	else
+	{
+		Text->SelectText(NewSelectionStart, NewCursorLocation);
+	}
+}
+
+FReply SPCGNodeSourceTextBox::OnTextKeyChar(
+	const FGeometry& MyGeometry,
+	const FCharacterEvent& InCharacterEvent) const
+{
+	if (Text->IsTextReadOnly())
+	{
+		return FReply::Unhandled();
+	}
+
+	const TCHAR Character = InCharacterEvent.GetCharacter();
+
+	if (Character == TEXT('\b'))
+	{
+		if (!Text->AnyTextSelected())
+		{
+			// if we are deleting a single open brace
+			// look for a matching close brace following it and delete the close brace as well
+			const FTextLocation CursorLocation = Text->GetCursorLocation();
+			const int Offset = CursorLocation.GetOffset();
+			FString Line;
+			Text->GetTextLine(CursorLocation.GetLineIndex(), Line);
+
+			if (Line.IsValidIndex(Offset) && Line.IsValidIndex(Offset-1))
+			{
+				if (IsOpenBrace(Line[Offset-1]))
+				{
+					if (Line[Offset] == GetMatchedCloseBrace(Line[Offset-1]))
+					{
+						SMultiLineEditableText::FScopedEditableTextTransaction ScopedTransaction(Text);
+						Text->SelectText(FTextLocation(CursorLocation, -1),FTextLocation(CursorLocation, 1));
+						Text->DeleteSelectedText();
+						Text->GoTo(FTextLocation(CursorLocation, -1));
+						return FReply::Handled();
+					}
+				}
+			}
+		}
+
+		return FReply::Unhandled();
+	}
+	else if (Character == TEXT('\t'))
+	{
+		SMultiLineEditableText::FScopedEditableTextTransaction Transaction(Text);
+
+		const FTextLocation CursorLocation = Text->GetCursorLocation();
+
+		const bool bShouldIncreaseIndentation = InCharacterEvent.GetModifierKeys().IsShiftDown() ? false : true;
+		
+		// When there is no text selected, shift tab should also decrease line indentation
+		const bool bShouldIndentLine = Text->AnyTextSelected() || (!Text->AnyTextSelected() && !bShouldIncreaseIndentation);
+		
+		if (bShouldIndentLine)
+		{
+			// Indent the whole line if there is a text selection
+			const FTextSelection Selection = Text->GetSelection();
+			const FTextLocation SelectionStart =
+				CursorLocation == Selection.GetBeginning() ? Selection.GetEnd() : Selection.GetBeginning();
+
+			// Shift the selection according to the new indentation
+			FTextLocation NewCursorLocation;
+			FTextLocation NewSelectionStart;
+		
+			const int32 StartLine = Selection.GetBeginning().GetLineIndex();
+			const int32 EndLine = Selection.GetEnd().GetLineIndex();
+		
+			for (int32 Index = StartLine; Index <= EndLine; Index++)
+			{
+				const FTextLocation LineStart(Index, 0);
+				Text->GoTo(LineStart);
+				
+				FString Line;
+				Text->GetTextLine(Index, Line);
+				const int32 NumSpaces = GetNumSpacesAtStartOfLine(Line);
+				const int32 NumExtraSpaces = NumSpaces % 4;
+
+				// Tab to nearest 4.
+				int32 NumSpacesForIndentation;
+				if (bShouldIncreaseIndentation)
+				{
+					NumSpacesForIndentation = NumExtraSpaces == 0 ? 4 : 4 - NumExtraSpaces ;
+					Text->InsertTextAtCursor(FString::ChrN(NumSpacesForIndentation, TEXT(' ')));
+				}
+				else
+				{
+					NumSpacesForIndentation = NumExtraSpaces == 0 ? FMath::Min(4, NumSpaces) : NumExtraSpaces;
+					Text->SelectText(LineStart,FTextLocation(LineStart, NumSpacesForIndentation));
+					Text->DeleteSelectedText();
+				}
+
+				const int32 CursorShiftDirection = bShouldIncreaseIndentation ? 1 : -1;
+				const int32 CursorShift = NumSpacesForIndentation * CursorShiftDirection;
+				
+				if (Index == CursorLocation.GetLineIndex())
+				{
+					NewCursorLocation = FTextLocation(CursorLocation, CursorShift);
+				}
+				
+				if (Index == SelectionStart.GetLineIndex())
+				{
+					NewSelectionStart = FTextLocation(SelectionStart, CursorShift);
+				}
+			}
+			
+			Text->SelectText(NewSelectionStart, NewCursorLocation);
+		}
+		else
+		{
+			FString Line;
+			Text->GetCurrentTextLine(Line);
+
+			const int32 Offset = CursorLocation.GetOffset();
+
+			// Tab to nearest 4.
+			if (ensure(bShouldIncreaseIndentation))
+			{
+				const int32 NumSpacesForIndentation = 4 - Offset % 4;
+				Text->InsertTextAtCursor(FString::ChrN(NumSpacesForIndentation, TEXT(' ')));
+			}
+		}
+
+		return FReply::Handled();
+	}
+	else if (Character == TEXT('\n') || Character == TEXT('\r'))
+	{
+		SMultiLineEditableText::FScopedEditableTextTransaction Transaction(Text);
+		
+		// at this point, the text after the text cursor is already in a new line
+		HandleAutoIndent();
+		
+		return FReply::Handled();
+	}
+	else if (IsOpenBrace(Character))
+	{
+		const TCHAR CloseBrace = GetMatchedCloseBrace(Character);
+		const FTextLocation CursorLocation = Text->GetCursorLocation();
+		FString Line;
+		Text->GetCurrentTextLine(Line);
+		
+		bool bShouldAutoInsertBraces = false;
+
+		if (CursorLocation.GetOffset() < Line.Len())
+		{
+			const TCHAR NextChar = Text->GetCharacterAt(CursorLocation);
+
+			if (IsWhiteSpace(NextChar))
+			{
+				bShouldAutoInsertBraces = true;
+			}
+			else if (IsCloseBrace(NextChar))
+			{
+				int32 BraceBalancePrior = 0;
+				for (int32 Index = 0; Index < CursorLocation.GetOffset() && Index < Line.Len(); Index++)
+				{
+					BraceBalancePrior += (Line[Index] == Character);
+					BraceBalancePrior -= (Line[Index] == CloseBrace);
+				}
+				
+				int32 BraceBalanceLater = 0;
+				for (int32 Index = CursorLocation.GetOffset(); Index < Line.Len(); Index++)
+				{
+					BraceBalanceLater += (Line[Index] == Character);
+					BraceBalanceLater -= (Line[Index] == CloseBrace);
+				}
+
+				if (BraceBalancePrior >= -BraceBalanceLater)
+				{
+					bShouldAutoInsertBraces = true;
+				}
+			}
+		}
+		else
+		{
+			bShouldAutoInsertBraces = true;
+		}
+		
+		// auto insert if we have more open braces
+		// on the left side than close braces on the right side
+		if (bShouldAutoInsertBraces)
+		{
+			// auto insert the matched close brace
+			SMultiLineEditableText::FScopedEditableTextTransaction Transaction(Text);
+			Text->InsertTextAtCursor(FString::Chr(Character));
+			Text->InsertTextAtCursor(FString::Chr(CloseBrace));
+			const FTextLocation NewCursorLocation(Text->GetCursorLocation(), -1);
+			Text->GoTo(NewCursorLocation);
+			return FReply::Handled();
+		}
+
+		return FReply::Unhandled();
+	}
+	else if (IsCloseBrace(Character))
+	{
+		if (!Text->AnyTextSelected())
+		{
+			const FTextLocation CursorLocation = Text->GetCursorLocation();	
+			FString Line;
+			Text->GetTextLine(CursorLocation.GetLineIndex(), Line);
+			
+			const int32 Offset =CursorLocation.GetOffset();
+			
+			if (Line.IsValidIndex(Offset))
+			{
+				if (Line[Offset] == Character)
+				{
+					// avoid creating a duplicated close brace and simply
+					// advance the cursor
+					Text->GoTo(FTextLocation(CursorLocation, 1));
+					return FReply::Handled();
+				}
+			}
+		}
+		
+		return FReply::Unhandled();
+	}
+	else
+	{
+		// Let SMultiLineEditableText::OnKeyChar handle it.
+		return FReply::Unhandled();
+	}
+}
+
+void SPCGNodeSourceTextBox::HandleAutoIndent() const
+{
+	const FTextLocation CursorLocation = Text->GetCursorLocation();
+	const int32 CurLineIndex = CursorLocation.GetLineIndex();
+	const int32 LastLineIndex = CurLineIndex - 1;
+
+	if (LastLineIndex > 0)
+	{
+		FString LastLine;
+		Text->GetTextLine(LastLineIndex, LastLine);
+
+		const int32 NumSpaces = GetNumSpacesAtStartOfLine(LastLine);
+
+		const int32 NumSpacesForCurrentIndentation = NumSpaces/4*4;
+		const FString CurrentIndentation = FString::ChrN(NumSpacesForCurrentIndentation, TEXT(' '));
+		const FString NextIndentation = FString::ChrN(NumSpacesForCurrentIndentation + 4, TEXT(' '));
+	
+		// See what the open/close curly brace balance is.
+		int32 BraceBalance = 0;
+		for (const TCHAR Char : LastLine)
+		{
+			BraceBalance += (Char == TEXT('{'));
+			BraceBalance -= (Char == TEXT('}'));
+		}
+
+		if (BraceBalance <= 0)
+		{
+			Text->InsertTextAtCursor(CurrentIndentation);
+		}
+		else
+		{
+			Text->InsertTextAtCursor(NextIndentation);
+		
+			// Look for an extra close curly brace and auto-indent it as well
+			FString CurLine;
+			Text->GetTextLine(CurLineIndex, CurLine);
+
+			BraceBalance = 0;
+			int32 CloseBraceOffset = 0;
+			for (const TCHAR Char : CurLine)
+			{
+				BraceBalance += (Char == TEXT('{'));
+				BraceBalance -= (Char == TEXT('}'));
+
+				// Found the first extra '}'
+				if (BraceBalance < 0)
+				{
+					break;
+				}
+
+				CloseBraceOffset++;
+			}
+
+			if (BraceBalance < 0)
+			{
+				const FTextLocation SavedCursorLocation = Text->GetCursorLocation();
+				const FTextLocation CloseBraceLocation(CurLineIndex, CloseBraceOffset);
+			
+				// Create a new line and apply indentation for the close curly brace
+				FString NewLineAndIndent(TEXT("\n"));
+				NewLineAndIndent.Append(CurrentIndentation);
+			
+				Text->GoTo(CloseBraceLocation);
+				Text->InsertTextAtCursor(NewLineAndIndent);
+				// Recover cursor location
+				Text->GoTo(SavedCursorLocation);
+			}
+		}
+	}
+}
+
+#undef LOCTEXT_NAMESPACE
diff --git a/Engine/Plugins/PCG/Source/PCGEditor/Private/Widgets/SPCGNodeSourceTextBox.h b/Engine/Plugins/PCG/Source/PCGEditor/Private/Widgets/SPCGNodeSourceTextBox.h
new file mode 100644
--- /dev/null
+++ b/Engine/Plugins/PCG/Source/PCGEditor/Private/Widgets/SPCGNodeSourceTextBox.h
@@ -0,0 +1,114 @@
+// Copyright Epic Games, Inc. All Rights Reserved.
+
+#pragma once
+
+#include "Widgets/SCompoundWidget.h"
+#include "Framework/Commands/Commands.h"
+
+#include "CoreMinimal.h"
+#include "Framework/SlateDelegates.h"
+#include "Widgets/Input/SSearchBox.h"
+
+class ITextLayoutMarshaller;
+class SMultiLineEditableText;
+class SPCGShaderTextSearchWidget;
+class SVerticalBox;
+
+class FPCGNodeSourceEditorTextBoxCommands :
+	public TCommands<FPCGNodeSourceEditorTextBoxCommands>
+{
+public:
+	FPCGNodeSourceEditorTextBoxCommands();
+	
+	// TCommands<> overrides
+	virtual void RegisterCommands() override;
+	
+	TSharedPtr<FUICommandInfo> Search;
+	TSharedPtr<FUICommandInfo> NextOccurrence;
+	TSharedPtr<FUICommandInfo> PreviousOccurrence;
+	
+	TSharedPtr<FUICommandInfo> ApplyChanges;
+	TSharedPtr<FUICommandInfo> ToggleComment;
+};
+
+/** Text box for editing node source text. */
+class SPCGNodeSourceTextBox : public SCompoundWidget
+{
+public:
+	SPCGNodeSourceTextBox();
+	virtual ~SPCGNodeSourceTextBox() override;
+
+	DECLARE_DELEGATE(FOnTextChangesApplied);
+
+	SLATE_BEGIN_ARGS(SPCGNodeSourceTextBox) {};
+		/** The initial text that will appear in the widget. */
+		SLATE_ATTRIBUTE(FText, Text)
+
+		/** Text to search for (a new search is triggered whenever this text changes) */
+		SLATE_ATTRIBUTE(FText, SearchText)
+
+		/** The marshaller used to get/set the raw text to/from the text layout. */
+		SLATE_ARGUMENT(TSharedPtr<ITextLayoutMarshaller>, Marshaller)
+
+		/** Sets whether this text box can actually be modified interactively by the user */
+		SLATE_ATTRIBUTE(bool, IsReadOnly)
+
+		/** Called whenever the text is changed programmatically or interactively by the user */
+		SLATE_EVENT(FOnTextChanged, OnTextChanged)
+
+		/** Called when user applies the changes */
+		SLATE_EVENT(FOnTextChangesApplied, OnTextChangesApplied)
+	SLATE_END_ARGS()
+
+	void Construct(const FArguments& InArgs);
+
+	void Refresh() const;
+
+private:
+	//~ Begin SWidget interface
+	// Use Preview Key Down so that the commands related to this tab is given the priority
+	virtual FReply OnPreviewKeyDown(const FGeometry& MyGeometry, const FKeyEvent& InKeyEvent) override;
+	//~ End SWidget interface
+
+	void RegisterCommands();
+	
+	bool HandleEscape();
+	
+	void ShowSearchBar();
+	bool HideSearchBar();
+	
+	void OnTriggerSearch();
+	void OnSearchTextChanged(const FText& InTextToSearch);
+	void OnSearchTextCommitted(const FText& InTextToSearch, ETextCommit::Type InCommitType);
+	TOptional<SSearchBox::FSearchResultData> GetSearchResultData() const;
+	void OnSearchResultNavigationButtonClicked(SSearchBox::SearchDirection InDirection);
+	void OnGoToNextOccurrence();
+	void OnGoToPreviousOccurrence();
+
+	FReply OnTextKeyDown(const FGeometry& MyGeometry, const FKeyEvent& InCharacterEvent) const;
+
+	void OnApplyChanges() const;
+
+	void OnToggleComment() const;
+	
+	FReply OnTextKeyChar(const FGeometry& MyGeometry, const FCharacterEvent& InCharacterEvent) const;
+	
+	// Figure out if we need to auto-indent.
+	void HandleAutoIndent() const;
+
+private:
+	TSharedPtr<SVerticalBox> TabBody;
+	
+	TSharedPtr<SMultiLineEditableText> Text;
+
+	TSharedPtr<SPCGShaderTextSearchWidget> SearchBar;
+	bool bIsSearchBarHidden = true;
+
+	// Top level command list is processed during preview key down
+	// such that its commands are given the priority
+	TSharedRef<FUICommandList> TopLevelCommandList;
+	
+	TSharedRef<FUICommandList> TextCommandList;
+	
+	FOnTextChangesApplied OnTextChangesAppliedCallback;
+};
diff --git a/Engine/Plugins/PCG/Source/PCGEditor/Private/Widgets/SPCGShaderTextSearchWidget.cpp b/Engine/Plugins/PCG/Source/PCGEditor/Private/Widgets/SPCGShaderTextSearchWidget.cpp
new file mode 100644
--- /dev/null
+++ b/Engine/Plugins/PCG/Source/PCGEditor/Private/Widgets/SPCGShaderTextSearchWidget.cpp
@@ -0,0 +1,103 @@
+// Copyright Epic Games, Inc. All Rights Reserved.
+
+#include "SPCGShaderTextSearchWidget.h"
+
+#include "PCGEditorStyle.h"
+
+#include "Widgets/Input/SSearchBox.h"
+#include "Widgets/SBoxPanel.h"
+
+#include "Framework/Application/SlateApplication.h"
+
+#define LOCTEXT_NAMESPACE "PCGShaderTextSearchWidget"
+
+SPCGShaderTextSearchWidget::SPCGShaderTextSearchWidget()
+{
+}
+
+SPCGShaderTextSearchWidget::~SPCGShaderTextSearchWidget()
+{
+}
+
+void SPCGShaderTextSearchWidget::Construct(const FArguments& InArgs)
+{
+	const FSearchBoxStyle& SearchBoxStyle = FPCGEditorStyle::Get().GetWidgetStyle<FSearchBoxStyle>(TEXT("TextEditor.SearchBoxStyle"));
+	
+	SearchBox =
+		SNew(SSearchBox)
+			.HintText(NSLOCTEXT("SearchBox", "HelpHint", "Search For Text"))
+			.Style(&SearchBoxStyle)
+			.OnTextChanged(InArgs._OnTextChanged)
+			.OnTextCommitted(InArgs._OnTextCommitted)
+			.SearchResultData(InArgs._SearchResultData)
+			.SelectAllTextWhenFocused(true)
+			.DelayChangeNotificationsWhileTyping(true)
+			.MinDesiredWidth(200)
+			.OnSearch(InArgs._OnResultNavigationButtonClicked);
+	
+	ChildSlot
+	[
+		SNew(SVerticalBox)
+		+ SVerticalBox::Slot()
+		.Padding(6.f)
+		.AutoHeight()
+		[
+			SNew(SVerticalBox)
+			+ SVerticalBox::Slot()
+			.AutoHeight()
+			[
+				SearchBox.ToSharedRef()
+			]
+		]
+	];
+}
+
+void SPCGShaderTextSearchWidget::FocusSearchBox() const
+{
+	FSlateApplication::Get().SetKeyboardFocus(SearchBox, EFocusCause::SetDirectly);
+}
+
+void SPCGShaderTextSearchWidget::TriggerSearch(const FText& InNewSearchText) const
+{
+	FocusSearchBox();
+
+	// Multiline search is not supported, sanitize the input to be single line text.
+	FString SingleLineString = InNewSearchText.ToString();
+	{
+		SingleLineString.GetCharArray().RemoveAll([&](const TCHAR InChar) -> bool
+		{
+			if (InChar != 0)
+			{
+				const bool bIsCharAllowed = !FChar::IsLinebreak(InChar);
+				return !bIsCharAllowed;
+			}
+			return false;
+		});
+	}
+	
+	FText SingleLineSearchText = FText::FromString(SingleLineString);
+	
+	// Clear the text to trigger a fresh search. Sometimes, the search text can be the same but starting from different place.
+	SearchBox->SetText(FText::GetEmpty());
+	
+	if (InNewSearchText.IsEmpty())
+	{
+		SearchBox->SetText(LastSearchedText);
+	}
+	else
+	{
+		SearchBox->SetText(SingleLineSearchText);
+	}
+
+	SearchBox->SelectAllText();
+}
+
+void SPCGShaderTextSearchWidget::ClearSearchText()
+{
+	// Save the last searched text so that if the next search is triggered without any selected text, we can initialize the search
+	// using it, standard behavior in regular text editors.
+	LastSearchedText = SearchBox->GetText();
+	SearchBox->SetText(FText::GetEmpty());
+}
+
+#undef LOCTEXT_NAMESPACE
diff --git a/Engine/Plugins/PCG/Source/PCGEditor/Private/Widgets/SPCGShaderTextSearchWidget.h b/Engine/Plugins/PCG/Source/PCGEditor/Private/Widgets/SPCGShaderTextSearchWidget.h
new file mode 100644
--- /dev/null
+++ b/Engine/Plugins/PCG/Source/PCGEditor/Private/Widgets/SPCGShaderTextSearchWidget.h
@@ -0,0 +1,50 @@
+// Copyright Epic Games, Inc. All Rights Reserved.
+
+#pragma once
+
+#include "Widgets/SCompoundWidget.h"
+
+#include "CoreMinimal.h"
+#include "Framework/SlateDelegates.h"
+#include "Widgets/Input/SSearchBox.h"
+
+class SSearchBox;
+class STextBlock;
+struct FTextLocation;
+
+/** Search text field. */
+class SPCGShaderTextSearchWidget : public SCompoundWidget
+{
+public:
+	SPCGShaderTextSearchWidget();
+	~SPCGShaderTextSearchWidget() override;
+
+	SLATE_BEGIN_ARGS(SPCGShaderTextSearchWidget) {};
+		/** Invoked whenever the text changes */
+		SLATE_EVENT(FOnTextChanged, OnTextChanged )
+
+		/** Invoked whenever the text is committed (e.g. user presses enter) */
+		SLATE_EVENT( FOnTextCommitted, OnTextCommitted )
+
+		/** This will add a next and previous button to your search box */
+		SLATE_EVENT(SSearchBox::FOnSearch, OnResultNavigationButtonClicked )
+	
+		/** Optional search result data to be shown in the search bar. */
+		SLATE_ATTRIBUTE(TOptional<SSearchBox::FSearchResultData>, SearchResultData )
+
+		/** Callback delegate to have first chance handling of the OnKeyDown event */
+		SLATE_EVENT(FOnKeyDown, OnKeyDownHandler)
+	SLATE_END_ARGS()
+
+	void Construct(const FArguments& InArgs);
+
+	void FocusSearchBox() const;
+
+	void TriggerSearch(const FText& InNewSearchText) const;
+	
+	void ClearSearchText();
+	
+private:
+	TSharedPtr<SSearchBox> SearchBox;
+	FText LastSearchedText;
+};
