Minor adjustment to prevent potential crash in loops/subgraphs in specific case with BP actors.

#rnx

diff --git a/Engine/Plugins/PCG/Source/PCG/Private/Elements/PCGLoopElement.cpp b/Engine/Plugins/PCG/Source/PCG/Private/Elements/PCGLoopElement.cpp
--- a/Engine/Plugins/PCG/Source/PCG/Private/Elements/PCGLoopElement.cpp
+++ b/Engine/Plugins/PCG/Source/PCG/Private/Elements/PCGLoopElement.cpp
@@ -344,7 +344,7 @@ bool FPCGLoopElement::ExecuteInternal(FPCGContext* InContext) const
 	{
 		// when woken up, get the output data from the subgraph
 		// and copy it to the current context output data, and finally return true
-		UPCGSubsystem* Subsystem = Context->SourceComponent->GetSubsystem();
+		UPCGSubsystem* Subsystem = Context->SourceComponent.IsValid() ? Context->SourceComponent->GetSubsystem() : nullptr;
 		if (Subsystem)
 		{
 			// If this was running iterations, we need to build a list of pins we'll ignore for non-terminal tasks
diff --git a/Engine/Plugins/PCG/Source/PCG/Private/PCGSubgraph.cpp b/Engine/Plugins/PCG/Source/PCG/Private/PCGSubgraph.cpp
--- a/Engine/Plugins/PCG/Source/PCG/Private/PCGSubgraph.cpp
+++ b/Engine/Plugins/PCG/Source/PCG/Private/PCGSubgraph.cpp
@@ -774,7 +774,7 @@ bool FPCGSubgraphElement::ExecuteInternal(FPCGContext* InContext) const
 		{
 			// when woken up, get the output data from the subgraph
 			// and copy it to the current context output data, and finally return true
-			UPCGSubsystem* Subsystem = Context->SourceComponent->GetSubsystem();
+			UPCGSubsystem* Subsystem = Context->SourceComponent.IsValid() ? Context->SourceComponent->GetSubsystem() : nullptr;
 			if (Subsystem)
 			{
 				if (Context->SubgraphTaskIds.Num() > 0)
