[PCG] Refactor GPU attribute tracking - change attribute table to map from attribute name to attribute ID & type, rather than the map key being name & ID.

* This reflects the way we use the table in practice, and allows attributes to be registered with unknown type (and still be retrieved later).
* Add mechanism to resolve types at execution time. An SM Spawner instance packer will know attribute names to pack but not their types, so add a way to resolve during data binding.

#rnx

diff --git a/Engine/Plugins/PCG/Source/PCG/Private/Compute/Elements/PCGComputeGraphElement.cpp b/Engine/Plugins/PCG/Source/PCG/Private/Compute/Elements/PCGComputeGraphElement.cpp
--- a/Engine/Plugins/PCG/Source/PCG/Private/Compute/Elements/PCGComputeGraphElement.cpp
+++ b/Engine/Plugins/PCG/Source/PCG/Private/Compute/Elements/PCGComputeGraphElement.cpp
@@ -100,6 +100,8 @@ bool FPCGComputeGraphElement::ExecuteInternal(FPCGContext* InContext) const
 	// 4. Initialize and parse incoming data for data sizes, attributes, etc that will drive buffer allocations and dispatch thread counts.
 	if (!Context->DataBinding)
 	{
+		Graph->FillInMissingAttributeTableTypes(Context->InputData);
+
 		UPCGDataBinding* DataBindingObject = FPCGContext::NewObject_AnyThread<UPCGDataBinding>(Context);
 		Context->DataBinding.Reset(DataBindingObject);
 
diff --git a/Engine/Plugins/PCG/Source/PCG/Private/Compute/Elements/PCGCustomHLSL.cpp b/Engine/Plugins/PCG/Source/PCG/Private/Compute/Elements/PCGCustomHLSL.cpp
--- a/Engine/Plugins/PCG/Source/PCG/Private/Compute/Elements/PCGCustomHLSL.cpp
+++ b/Engine/Plugins/PCG/Source/PCG/Private/Compute/Elements/PCGCustomHLSL.cpp
@@ -306,24 +306,25 @@ FPCGDataCollectionDesc UPCGCustomHLSLSettings::ComputeOutputPinDataDesc(const UP
 		PCGComputeHelpers::ComputeOutputPinDataDesc(*Properties, this, Binding, PinDesc);
 	}
 
-	const TMap<FPCGKernelAttributeKey, int32>& GlobalAttributeLookupTable = ensure(Binding && Binding->Graph) ? Binding->Graph->GetAttributeLookupTable() : TMap<FPCGKernelAttributeKey, int32>();
-
-	for (const FPCGKernelAttributeKey& AttributeKey : KernelAttributeKeys)
+	if (const TMap<FName, FPCGKernelAttributeIDAndType>* GlobalAttributeLookupTable = ensure(Binding && Binding->Graph) ? &Binding->Graph->GetAttributeLookupTable() : nullptr)
 	{
-		// Add attributes that will be created for this pin on the GPU.
-		if (const TArray<TTuple<FPCGKernelAttributeKey, bool>>* Keys = PinToAttributeKeys.Find(OutputPin->Properties.Label))
+		for (const FPCGKernelAttributeKey& AttributeKey : KernelAttributeKeys)
 		{
-			const TTuple<FPCGKernelAttributeKey, bool>* Pair = Keys->FindByPredicate([AttributeKey](const TTuple<FPCGKernelAttributeKey, bool>& Pair) { return Pair.Key == AttributeKey; });
-			const bool bCreatedOnGPU = Pair && Pair->Value;
-
-			if (bCreatedOnGPU)
+			// Add attributes that will be created for this pin on the GPU.
+			if (const TArray<TTuple<FPCGKernelAttributeKey, bool>>* Keys = PinToAttributeKeys.Find(OutputPin->Properties.Label))
 			{
-				for (FPCGDataDesc& DataDesc : PinDesc.DataDescs)
+				const TTuple<FPCGKernelAttributeKey, bool>* Pair = Keys->FindByPredicate([AttributeKey](const TTuple<FPCGKernelAttributeKey, bool>& Pair) { return Pair.Key == AttributeKey; });
+				const bool bCreatedOnGPU = Pair && Pair->Value;
+
+				if (bCreatedOnGPU)
 				{
-					if (const int* Index = GlobalAttributeLookupTable.Find(AttributeKey))
+					for (FPCGDataDesc& DataDesc : PinDesc.DataDescs)
 					{
-						const FPCGKernelAttributeDesc AttributeDesc(*Index, AttributeKey.Type, AttributeKey.Name);
-						DataDesc.AttributeDescs.AddUnique(AttributeDesc);
+						if (const FPCGKernelAttributeIDAndType* IDAndType = GlobalAttributeLookupTable->Find(AttributeKey.Name))
+						{
+							const FPCGKernelAttributeDesc AttributeDesc(IDAndType->Id, IDAndType->Type, AttributeKey.Name);
+							DataDesc.AttributeDescs.AddUnique(AttributeDesc);
+						}
 					}
 				}
 			}
@@ -1449,7 +1450,7 @@ bool UPCGCustomHLSLSettings::AreKernelAttributesValid(FPCGContext* InContext, FT
 	return true;
 }
 
-FString UPCGCustomHLSLSettings::GetCookedKernelSource(const TMap<FPCGKernelAttributeKey, int32>& GlobalAttributeLookupTable) const
+FString UPCGCustomHLSLSettings::GetCookedKernelSource(const TMap<FName, FPCGKernelAttributeIDAndType>& GlobalAttributeLookupTable) const
 {
 	const FIntVector GroupSize = GetThreadGroupSize();
 
@@ -1463,10 +1464,11 @@ FString UPCGCustomHLSLSettings::GetCookedKernelSource(const TMap<FPCGKernelAttri
 	Functions.ReplaceInline(TEXT("\r"), TEXT(""));
 #endif
 
-	for (const TPair<FPCGKernelAttributeKey, int32>& Pair : GlobalAttributeLookupTable)
+	for (const TPair<FName, FPCGKernelAttributeIDAndType>& Pair : GlobalAttributeLookupTable)
 	{
-		const FString SourceDefinition = PCGHLSLElement::GetKernelAttributeKeyAsString(Pair.Key);
-		const FString AttributeIndexAsString = FString::FromInt(Pair.Value);
+		FPCGKernelAttributeKey AttributeKey(Pair.Key, Pair.Value.Type);
+		const FString SourceDefinition = PCGHLSLElement::GetKernelAttributeKeyAsString(AttributeKey);
+		const FString AttributeIndexAsString = FString::FromInt(Pair.Value.Id);
 
 		Source.ReplaceInline(*SourceDefinition, *AttributeIndexAsString);
 		Functions.ReplaceInline(*SourceDefinition, *AttributeIndexAsString);
diff --git a/Engine/Plugins/PCG/Source/PCG/Private/Compute/PCGComputeGraph.cpp b/Engine/Plugins/PCG/Source/PCG/Private/Compute/PCGComputeGraph.cpp
--- a/Engine/Plugins/PCG/Source/PCG/Private/Compute/PCGComputeGraph.cpp
+++ b/Engine/Plugins/PCG/Source/PCG/Private/Compute/PCGComputeGraph.cpp
@@ -6,6 +6,7 @@
 #include "PCGModule.h"
 #include "PCGNode.h"
 #include "Compute/PCGCompilerDiagnostic.h"
+#include "Metadata/PCGMetadata.h"
 
 #include "ComputeFramework/ComputeKernelCompileResult.h"
 
@@ -85,3 +86,47 @@ void UPCGComputeGraph::OnKernelCompilationComplete(int32 InKernelIndex, FCompute
 		UE_LOG(LogPCG, Verbose, TEXT("Compilation message ignored for kernel index %d which has no associated node."), InKernelIndex);
 	}
 }
+
+void UPCGComputeGraph::FillInMissingAttributeTableTypes(const FPCGDataCollection& InComputeGraphElementInputData)
+{
+	TRACE_CPUPROFILER_EVENT_SCOPE(UPCGComputeGraph::FillInMissingAttributeTableTypes);
+
+	// Local to loop body but hoisted for perf.
+	TArray<FName> AttributeNames;
+	TArray<EPCGMetadataTypes> AttributeTypes;
+
+	for (TPair<FName, FPCGKernelAttributeIDAndType>& Entry : GlobalAttributeLookupTable)
+	{
+		if (Entry.Value.Type != EPCGKernelAttributeType::None)
+		{
+			continue;
+		}
+
+		for (const FPCGTaggedData& Data : InComputeGraphElementInputData.TaggedData)
+		{
+			AttributeNames.Reset();
+			AttributeTypes.Reset();
+
+			if (const UPCGMetadata* Metadata = Data.Data ? Data.Data->ConstMetadata() : nullptr)
+			{
+				Metadata->GetAttributes(AttributeNames, AttributeTypes);
+			}
+
+			for (int Index = 0; Index < AttributeNames.Num(); ++Index)
+			{
+				if (AttributeNames[Index] == Entry.Key)
+				{
+					Entry.Value.Type = PCGDataForGPUHelpers::GetAttributeTypeFromMetadataType(AttributeTypes[Index]);
+					break;
+				}
+			}
+
+			if (Entry.Value.Type != EPCGKernelAttributeType::None)
+			{
+				break;
+			}
+		}
+
+		ensure(Entry.Value.Type != EPCGKernelAttributeType::None);
+	}
+}
diff --git a/Engine/Plugins/PCG/Source/PCG/Private/Compute/PCGComputeGraph.h b/Engine/Plugins/PCG/Source/PCG/Private/Compute/PCGComputeGraph.h
--- a/Engine/Plugins/PCG/Source/PCG/Private/Compute/PCGComputeGraph.h
+++ b/Engine/Plugins/PCG/Source/PCG/Private/Compute/PCGComputeGraph.h
@@ -24,8 +24,11 @@ public:
 	void OnKernelCompilationComplete(int32 InKernelIndex, FComputeKernelCompileResults const& InCompileResults) override;
 	//~End UComputeGraph interface
 
+	/** Some attributes in table may be missing type information that we could not infer statically. Fill in missing types from the execution-time data. */
+	void FillInMissingAttributeTableTypes(const FPCGDataCollection& InComputeGraphElementInputData);
+
 	/** Get the global attribute indices. */
-	const TMap<FPCGKernelAttributeKey, int32>& GetAttributeLookupTable() const { return GlobalAttributeLookupTable; }
+	const TMap<FName, FPCGKernelAttributeIDAndType>& GetAttributeLookupTable() const { return GlobalAttributeLookupTable; }
 
 public:
 	TMap<TObjectKey<const UPCGNode>, TArray<FComputeKernelCompileMessage>> KernelToCompileMessages;
@@ -48,7 +51,7 @@ public:
 
 protected:
 	UPROPERTY()
-	TMap<FPCGKernelAttributeKey, int32 /* Attribute Index */> GlobalAttributeLookupTable;
+	TMap<FName /* Attribute name */, FPCGKernelAttributeIDAndType> GlobalAttributeLookupTable;
 
 	friend class FPCGGraphCompilerGPU;
 };
diff --git a/Engine/Plugins/PCG/Source/PCG/Private/Compute/PCGDataForGPU.cpp b/Engine/Plugins/PCG/Source/PCG/Private/Compute/PCGDataForGPU.cpp
--- a/Engine/Plugins/PCG/Source/PCG/Private/Compute/PCGDataForGPU.cpp
+++ b/Engine/Plugins/PCG/Source/PCG/Private/Compute/PCGDataForGPU.cpp
@@ -407,7 +407,7 @@ FPCGDataDesc::FPCGDataDesc(EPCGDataType InType, int InElementCount)
 	InitializeAttributeDescs(nullptr);
 }
 
-FPCGDataDesc::FPCGDataDesc(const UPCGData* Data, const TMap<FPCGKernelAttributeKey, int32>& GlobalAttributeLookupTable)
+FPCGDataDesc::FPCGDataDesc(const UPCGData* Data, const TMap<FName, FPCGKernelAttributeIDAndType>& GlobalAttributeLookupTable)
 {
 	check(Data);
 
@@ -443,7 +443,7 @@ uint32 FPCGDataDesc::ComputePackedSize() const
 	return DataSizeBytes;
 }
 
-void FPCGDataDesc::InitializeAttributeDescs(const UPCGMetadata* Metadata, const TMap<FPCGKernelAttributeKey, int32>& GlobalAttributeLookupTable)
+void FPCGDataDesc::InitializeAttributeDescs(const UPCGMetadata* Metadata, const TMap<FName, FPCGKernelAttributeIDAndType>& GlobalAttributeLookupTable)
 {
 	if (Type == EPCGDataType::Point)
 	{
@@ -477,8 +477,6 @@ void FPCGDataDesc::InitializeAttributeDescs(const UPCGMetadata* Metadata, const
 				continue;
 			}
 
-			const FPCGKernelAttributeKey AttributeKey = { AttributeType, AttributeName };
-
 			// Ignore excess attributes.
 			if (CustomAttributeIndex >= MAX_NUM_CUSTOM_ATTRS)
 			{
@@ -494,14 +492,14 @@ void FPCGDataDesc::InitializeAttributeDescs(const UPCGMetadata* Metadata, const
 				continue;
 			}
 
-			if (const int32* AttributeId = GlobalAttributeLookupTable.Find(AttributeKey))
+			if (const FPCGKernelAttributeIDAndType* AttributeIdAndType = GlobalAttributeLookupTable.Find(AttributeName))
 			{
-				AttributeDescs.Emplace(*AttributeId, AttributeType, AttributeName);
+				AttributeDescs.Emplace(AttributeIdAndType->Id, AttributeType, AttributeName);
 				++NumAttributesFromLUT;
 			}
 			else
 			{
-				DelayedAttributeKeys.Add(AttributeKey);
+				DelayedAttributeKeys.Emplace(AttributeName, AttributeType);
 			}
 		}
 
@@ -516,7 +514,7 @@ void FPCGDataDesc::InitializeAttributeDescs(const UPCGMetadata* Metadata, const
 FPCGDataCollectionDesc FPCGDataCollectionDesc::BuildFromInputDataCollectionAndInputPinLabel(
 	const FPCGDataCollection& InDataCollection,
 	FName InputPinLabel,
-	const TMap<FPCGKernelAttributeKey, int32>& InAttributeLookupTable)
+	const TMap<FName, FPCGKernelAttributeIDAndType>& InAttributeLookupTable)
 {
 	FPCGDataCollectionDesc Desc;
 	TArray<FPCGTaggedData> DataForPin = InDataCollection.GetInputsByPin(InputPinLabel);
diff --git a/Engine/Plugins/PCG/Source/PCG/Private/Elements/PCGCopyPoints.cpp b/Engine/Plugins/PCG/Source/PCG/Private/Elements/PCGCopyPoints.cpp
--- a/Engine/Plugins/PCG/Source/PCG/Private/Elements/PCGCopyPoints.cpp
+++ b/Engine/Plugins/PCG/Source/PCG/Private/Elements/PCGCopyPoints.cpp
@@ -25,7 +25,7 @@ FText UPCGCopyPointsSettings::GetNodeTooltipText() const
 }
 #endif
 
-FString UPCGCopyPointsSettings::GetCookedKernelSource(const TMap<FPCGKernelAttributeKey, int>& GlobalAttributeLookupTable) const
+FString UPCGCopyPointsSettings::GetCookedKernelSource(const TMap<FName, FPCGKernelAttributeIDAndType>& GlobalAttributeLookupTable) const
 {
 	FString TemplateFile;
 	LoadShaderSourceFile(TEXT("/Plugin/PCG/Private/Elements/PCGCopyPoints.usf"), EShaderPlatform::SP_PCD3D_SM5, &TemplateFile, nullptr);
diff --git a/Engine/Plugins/PCG/Source/PCG/Private/Graph/PCGGraphCompilerGPU.cpp b/Engine/Plugins/PCG/Source/PCG/Private/Graph/PCGGraphCompilerGPU.cpp
--- a/Engine/Plugins/PCG/Source/PCG/Private/Graph/PCGGraphCompilerGPU.cpp
+++ b/Engine/Plugins/PCG/Source/PCG/Private/Graph/PCGGraphCompilerGPU.cpp
@@ -597,9 +597,26 @@ void FPCGGraphCompilerGPU::BuildGPUGraphTask(
 
 					for (const FPCGKernelAttributeKey& Key : Settings->GetKernelAttributeKeys())
 					{
-						if (!ComputeGraph->GlobalAttributeLookupTable.Find(Key))
+						FPCGKernelAttributeIDAndType* FoundEntry = ComputeGraph->GlobalAttributeLookupTable.Find(Key.Name);
+						if (FoundEntry)
 						{
-							ComputeGraph->GlobalAttributeLookupTable.Add(Key, ComputeGraph->GlobalAttributeLookupTable.Num() + PCGComputeConstants::NUM_RESERVED_ATTRS);
+							// If we already registered this attribute, populate the type if not populated already, or otherwise ensure type matches.
+							if (Key.Type != EPCGKernelAttributeType::None)
+							{
+								if (FoundEntry->Type == EPCGKernelAttributeType::None)
+								{
+									FoundEntry->Type = Key.Type;
+								}
+								else
+								{
+									ensure(Key.Type == FoundEntry->Type);
+								}
+							}
+						}
+						else
+						{
+							const FPCGKernelAttributeIDAndType IDAndType(ComputeGraph->GlobalAttributeLookupTable.Num() + PCGComputeConstants::NUM_RESERVED_ATTRS, Key.Type);
+							ComputeGraph->GlobalAttributeLookupTable.Add(Key.Name, IDAndType);
 						}
 					}
 				}
@@ -931,13 +948,12 @@ void FPCGGraphCompilerGPU::BuildGPUGraphTask(
 			{
 				UE_LOG(LogPCG, Warning, TEXT("ATTRIBUTE LOOK-UP TABLE [%s]"), *Settings->GetDefaultNodeTitle().ToString());
 
-				for (const TPair<FPCGKernelAttributeKey, int32 /* Attribute Index */>& Pair : ComputeGraph->GlobalAttributeLookupTable)
+				for (const TPair<FName, FPCGKernelAttributeIDAndType>& Pair : ComputeGraph->GlobalAttributeLookupTable)
 				{
-					const FString TypeString = UEnum::GetValueAsString(Pair.Key.Type);
-					const FString NameString = Pair.Key.Name.ToString();
-					const FString IndexString = FString::FromInt(Pair.Value);
-
-					UE_LOG(LogPCG, Warning, TEXT("%s: %s (%s)"), *IndexString, *NameString, *TypeString);
+					UE_LOG(LogPCG, Warning, TEXT("%s: %s (%s)"),
+						*FString::FromInt(Pair.Value.Id),
+						*Pair.Key.ToString(),
+						*UEnum::GetValueAsString(Pair.Value.Type));
 				}
 			}
 #endif
diff --git a/Engine/Plugins/PCG/Source/PCG/Public/Compute/Elements/PCGCustomHLSL.h b/Engine/Plugins/PCG/Source/PCG/Public/Compute/Elements/PCGCustomHLSL.h
--- a/Engine/Plugins/PCG/Source/PCG/Public/Compute/Elements/PCGCustomHLSL.h
+++ b/Engine/Plugins/PCG/Source/PCG/Public/Compute/Elements/PCGCustomHLSL.h
@@ -67,7 +67,8 @@ public:
 	virtual void ApplyPreconfiguredSettings(const FPCGPreConfiguredSettingsInfo& PreconfigureInfo) override;
 
 	virtual bool IsKernelValid(FPCGContext* InContext = nullptr, bool bQuiet = true) const override;
-	virtual FString GetCookedKernelSource(const TMap<FPCGKernelAttributeKey, int>& GlobalAttributeLookupTable) const override;
+	virtual FString GetCookedKernelSource(const TMap<FName, FPCGKernelAttributeIDAndType>& GlobalAttributeLookupTable) const override;
+	virtual const TArray<FPCGKernelAttributeKey> GetKernelAttributeKeys() const { return KernelAttributeKeys; }
 	virtual int ComputeKernelThreadCount(const UPCGDataBinding* Binding) const override;
 	virtual FPCGDataCollectionDesc ComputeOutputPinDataDesc(const UPCGPin* OutputPin, const UPCGDataBinding* Binding) const override;
 
@@ -175,6 +176,12 @@ protected:
 	/** Helper data and functions that can be used from the shader code. Intended to be viewed using the Node Source Editor window. */
 	UPROPERTY(Transient, VisibleAnywhere, Category = "Declarations|Helpers", meta = (MultiLine = true))
 	FString HelperDeclarations;
+
+	/** Attributes statically detected as being read, written, or created by this node. */
+	TArray<FPCGKernelAttributeKey> KernelAttributeKeys;
+
+	/** Maps pins to their attribute keys and whether or not they were created on the GPU. */
+	TMap<FName, TArray<TTuple<FPCGKernelAttributeKey, bool /*bCreatedOnGPU*/>>> PinToAttributeKeys;
 };
 
 class FPCGCustomHLSLElement : public IPCGElement
diff --git a/Engine/Plugins/PCG/Source/PCG/Public/Compute/PCGDataForGPU.h b/Engine/Plugins/PCG/Source/PCG/Public/Compute/PCGDataForGPU.h
--- a/Engine/Plugins/PCG/Source/PCG/Public/Compute/PCGDataForGPU.h
+++ b/Engine/Plugins/PCG/Source/PCG/Public/Compute/PCGDataForGPU.h
@@ -10,6 +10,7 @@
 class UPCGMetadata;
 class UPCGPin;
 class UPCGSettings;
+enum class EPCGMetadataTypes : uint8;
 
 enum class EPCGUnpackDataCollectionResult
 {
@@ -20,7 +21,8 @@ enum class EPCGUnpackDataCollectionResult
 UENUM()
 enum class EPCGKernelAttributeType : uint8
 {
-	Bool = 0,
+	None = 0,
+	Bool,
 	Int,
 	Float,
 	Float2,
@@ -38,16 +40,42 @@ struct FPCGKernelAttributeKey
 {
 	GENERATED_BODY()
 
-	UPROPERTY(EditAnywhere, Category = "Settings")
-	EPCGKernelAttributeType Type = EPCGKernelAttributeType::Float;
+	FPCGKernelAttributeKey() = default;
+
+	explicit FPCGKernelAttributeKey(FName InName, EPCGKernelAttributeType InType)
+		: Name(InName)
+		, Type(InType)
+	{}
 
 	UPROPERTY(EditAnywhere, Category = "Settings")
 	FName Name = NAME_None;
 
+	UPROPERTY(EditAnywhere, Category = "Settings")
+	EPCGKernelAttributeType Type = EPCGKernelAttributeType::Float;
+
 	bool operator==(const FPCGKernelAttributeKey& Other) const;
 	friend uint32 GetTypeHash(const FPCGKernelAttributeKey& In);
 };
 
+USTRUCT()
+struct FPCGKernelAttributeIDAndType
+{
+	GENERATED_BODY()
+
+	FPCGKernelAttributeIDAndType() = default;
+
+	explicit FPCGKernelAttributeIDAndType(int32 InId, EPCGKernelAttributeType InType)
+		: Id(InId)
+		, Type(InType)
+	{}
+
+	UPROPERTY()
+	int32 Id;
+
+	UPROPERTY()
+	EPCGKernelAttributeType Type;
+};
+
 struct FPCGKernelAttributeDesc
 {
 	explicit FPCGKernelAttributeDesc(int32 InIndex, EPCGKernelAttributeType InType, FName InName)
@@ -67,7 +95,7 @@ struct FPCGKernelAttributeDesc
 struct FPCGDataDesc
 {
 	FPCGDataDesc(EPCGDataType InType, int32 InElementCount);
-	FPCGDataDesc(const UPCGData* Data, const TMap<FPCGKernelAttributeKey, int32>& GlobalAttributeLookupTable);
+	FPCGDataDesc(const UPCGData* Data, const TMap<FName, FPCGKernelAttributeIDAndType>& GlobalAttributeLookupTable);
 
 	uint32 ComputePackedSize() const;
 
@@ -76,7 +104,7 @@ struct FPCGDataDesc
 	int32 ElementCount = 0;
 
 private:
-	void InitializeAttributeDescs(const UPCGMetadata* Metadata, const TMap<FPCGKernelAttributeKey, int32>& GlobalAttributeLookupTable = {});
+	void InitializeAttributeDescs(const UPCGMetadata* Metadata, const TMap<FName, FPCGKernelAttributeIDAndType>& GlobalAttributeLookupTable = {});
 };
 
 struct FPCGDataCollectionDesc
@@ -84,7 +112,7 @@ struct FPCGDataCollectionDesc
 	static FPCGDataCollectionDesc BuildFromInputDataCollectionAndInputPinLabel(
 		const FPCGDataCollection& InDataCollection,
 		FName InputPinLabel,
-		const TMap<FPCGKernelAttributeKey, int32>& InAttributeLookupTable);
+		const TMap<FName, FPCGKernelAttributeIDAndType>& InAttributeLookupTable);
 
 	/** Computes the size (in bytes) of the data collection after packing. Also produces the offset (in bytes) for each data in the packed collection. */
 	uint32 ComputePackedSize(TArray<uint32>* OutDataAddresses = nullptr) const;
@@ -116,3 +144,8 @@ struct FPCGDataForGPU
 	// Since the compute graph is collapsed to a single element, all data crossing from CPU to GPU is in a single collection.
 	FPCGDataCollection InputDataCollection;
 };
+
+namespace PCGDataForGPUHelpers
+{
+	EPCGKernelAttributeType GetAttributeTypeFromMetadataType(EPCGMetadataTypes MetadataType);
+}
diff --git a/Engine/Plugins/PCG/Source/PCG/Public/Elements/PCGCopyPoints.h b/Engine/Plugins/PCG/Source/PCG/Public/Elements/PCGCopyPoints.h
--- a/Engine/Plugins/PCG/Source/PCG/Public/Elements/PCGCopyPoints.h
+++ b/Engine/Plugins/PCG/Source/PCG/Public/Elements/PCGCopyPoints.h
@@ -53,7 +53,7 @@ public:
 	virtual bool DisplayExecuteOnGPUSetting() const override { return true; }
 #endif
 
-	virtual FString GetCookedKernelSource(const TMap<FPCGKernelAttributeKey, int>& GlobalAttributeLookupTable) const override;
+	virtual FString GetCookedKernelSource(const TMap<FName, FPCGKernelAttributeIDAndType>& GlobalAttributeLookupTable) const override;
 	virtual int ComputeKernelThreadCount(const UPCGDataBinding* Binding) const override;
 	virtual FPCGDataCollectionDesc ComputeOutputPinDataDesc(const UPCGPin* OutputPin, const UPCGDataBinding* Binding) const override;
 	virtual void CreateAdditionalInputDataInterfaces(TArray<TObjectPtr<UComputeDataInterface>>& OutDataInterfaces) const override;
diff --git a/Engine/Plugins/PCG/Source/PCG/Public/PCGSettings.h b/Engine/Plugins/PCG/Source/PCG/Public/PCGSettings.h
--- a/Engine/Plugins/PCG/Source/PCG/Public/PCGSettings.h
+++ b/Engine/Plugins/PCG/Source/PCG/Public/PCGSettings.h
@@ -458,12 +458,6 @@ protected:
 	TArray<FPCGPinProperties> DefaultPointInputPinProperties() const;
 	TArray<FPCGPinProperties> DefaultPointOutputPinProperties() const;
 
-	/** Attributes statically detected as being read, written, or created by this node. */
-	TArray<FPCGKernelAttributeKey> KernelAttributeKeys;
-
-	/** Maps pins to their attribute keys and whether or not they were created on the GPU. */
-	TMap<FName, TArray<TTuple<FPCGKernelAttributeKey, bool /*bCreatedOnGPU*/>>> PinToAttributeKeys;
-
 public:
 #if WITH_EDITORONLY_DATA
 	/** The version number of the data after load and after any data migration. */
@@ -522,12 +516,12 @@ public:
 	virtual bool ShouldExecuteOnGPU() const { return bExecuteOnGPU; }
 
 	virtual bool IsKernelValid(FPCGContext* InContext = nullptr, bool bQuiet = true) const;
-	virtual FString GetCookedKernelSource(const TMap<FPCGKernelAttributeKey, int>& GlobalAttributeLookupTable) const { return TEXT(""); }
+	virtual FString GetCookedKernelSource(const TMap<FName, FPCGKernelAttributeIDAndType>& GlobalAttributeLookupTable) const { return TEXT(""); }
 	virtual FString GetKernelEntryPoint() const { return TEXT("Main"); }
 	virtual FIntVector GetThreadGroupSize() const { return FIntVector(64, 1, 1); }
 
 	/** Get a list of the attributes read or written by this node. */
-	virtual const TArray<FPCGKernelAttributeKey>& GetKernelAttributeKeys() const { return KernelAttributeKeys; }
+	virtual const TArray<FPCGKernelAttributeKey> GetKernelAttributeKeys() const { return {}; }
 
 	/** Compute how many threads should be dispatched to execute this node on the GPU. */
 	virtual int ComputeKernelThreadCount(const UPCGDataBinding* Binding) const { return 0; };
