[PCG] Factorize Save Data Asset to work with any asset

#rnx

diff --git a/Engine/Plugins/PCG/Source/PCG/Private/PCGAssetExporterUtils.cpp b/Engine/Plugins/PCG/Source/PCG/Private/PCGAssetExporterUtils.cpp
--- a/Engine/Plugins/PCG/Source/PCG/Private/PCGAssetExporterUtils.cpp
+++ b/Engine/Plugins/PCG/Source/PCG/Private/PCGAssetExporterUtils.cpp
@@ -14,6 +14,11 @@
 
 #define LOCTEXT_NAMESPACE "PCGAssetExporterUtils"
 
+namespace PCGAssetExporterUtils
+{
+	static const FText NonEditorError = LOCTEXT("CannotExportInNonEditor", "PCG Asset Exporter Utils cannot be used in non-editor builds.");
+}
+
 UPackage* UPCGAssetExporterUtils::CreateAsset(UPCGAssetExporter* Exporter, FPCGAssetExporterParameters Parameters)
 {
 	return CreateAsset(Exporter, Parameters, nullptr);
@@ -34,14 +39,35 @@ UPackage* UPCGAssetExporterUtils::CreateAsset(UPCGAssetExporter* Exporter, const
 		Exporter->AddToRoot();
 	}
 
-	ON_SCOPE_EXIT
-	{
-		if (bExporterNeedsRooting)
+	UPackage* ResultPackage = CreateAsset(Exporter->GetAssetType(),
+		Parameters,
+		[Exporter](const FString& PackagePath, UObject* Asset) -> bool
 		{
-			Exporter->RemoveFromRoot();
-		}
-	};
+			return Exporter->Export(PackagePath, CastChecked<UPCGDataAsset>(Asset));
+		},
+		Context);
+
+	if (bExporterNeedsRooting)
+	{
+		Exporter->RemoveFromRoot();
+	}
 
+	return ResultPackage;
+#else
+	PCGLog::LogErrorOnGraph(PCGAssetExporterUtils::NonEditorError, Context);
+	return nullptr;
+#endif // WITH_EDITOR
+}
+
+UPackage* UPCGAssetExporterUtils::CreateAsset(const UClass* AssetClass, const FPCGAssetExporterParameters& Parameters, TFunctionRef<bool(const FString&, UObject*)> ExportFunc, FPCGContext* Context)
+{
+#if WITH_EDITOR
+	if (!AssetClass)
+	{
+		PCGLog::LogErrorOnGraph(LOCTEXT("MissingAssetClass", "Unable to create asset without a class."), Context);
+		return nullptr;
+	}
+	
 	FString AssetName = Parameters.AssetName;
 	FString AssetPath = Parameters.AssetPath;
 	FString PackageName = FPaths::Combine(AssetPath, AssetName);
@@ -51,7 +77,7 @@ UPackage* UPCGAssetExporterUtils::CreateAsset(UPCGAssetExporter* Exporter, const
 		FSaveAssetDialogConfig SaveAssetDialogConfig;
 		SaveAssetDialogConfig.DefaultPath = AssetPath;
 		SaveAssetDialogConfig.DefaultAssetName = AssetName;
-		SaveAssetDialogConfig.AssetClassNames.Add(Exporter->GetAssetType()->GetClassPathName());
+		SaveAssetDialogConfig.AssetClassNames.Add(AssetClass->GetClassPathName());
 		SaveAssetDialogConfig.ExistingAssetPolicy = ESaveAssetDialogExistingAssetPolicy::AllowButWarn;
 		SaveAssetDialogConfig.DialogTitleOverride = NSLOCTEXT("PCGAssetExporter", "SaveAssetToFileDialogTitle", "Save PCG Asset");
 
@@ -81,13 +107,13 @@ UPackage* UPCGAssetExporterUtils::CreateAsset(UPCGAssetExporter* Exporter, const
 
 	UPackage* Package = FPackageName::DoesPackageExist(PackageName) ? LoadPackage(nullptr, *PackageName, LOAD_None) : nullptr;
 
-	UPCGDataAsset* Asset = nullptr;
+	UObject* Asset = nullptr;
 	bool NewAssetCreated = false;
 
 	if (Package)
 	{
 		UObject* Object = FindObjectFast<UObject>(Package, *AssetName);
-		if (Object && Object->GetClass() != Exporter->GetAssetType())
+		if (Object && Object->GetClass() != AssetClass)
 		{
 			Object->SetFlags(RF_Transient);
 			Object->Rename(nullptr, GetTransientPackage(), REN_DontCreateRedirectors | REN_NonTransactional);
@@ -95,7 +121,7 @@ UPackage* UPCGAssetExporterUtils::CreateAsset(UPCGAssetExporter* Exporter, const
 		}
 		else
 		{
-			Asset = Cast<UPCGDataAsset>(Object);
+			Asset = Object;
 		}
 	}
 	else
@@ -116,7 +142,7 @@ UPackage* UPCGAssetExporterUtils::CreateAsset(UPCGAssetExporter* Exporter, const
 	if (!Asset)
 	{
 		const EObjectFlags Flags = RF_Public | RF_Standalone | RF_Transactional;
-		Asset = NewObject<UPCGDataAsset>(Package, Exporter->GetAssetType(), FName(*AssetName), Flags);
+		Asset = NewObject<UObject>(Package, AssetClass, FName(*AssetName), Flags);
 	}
 
 	if (Asset)
@@ -127,7 +153,7 @@ UPackage* UPCGAssetExporterUtils::CreateAsset(UPCGAssetExporter* Exporter, const
 			FAssetRegistryModule::AssetCreated(Asset);
 		}
 
-		if (Exporter->Export(PackageName, Asset))
+		if (ExportFunc(PackageName, Asset))
 		{
 			// Make sure everybody knows we changed the asset
 			if (!NewAssetCreated)
@@ -149,7 +175,7 @@ UPackage* UPCGAssetExporterUtils::CreateAsset(UPCGAssetExporter* Exporter, const
 
 	return Package;
 #else
-	PCGLog::LogErrorOnGraph(LOCTEXT("CannotExportInNonEditor", "PCG Asset Exporter Utils cannot be used in non-editor builds."), Context);
+	PCGLog::LogErrorOnGraph(PCGAssetExporterUtils::NonEditorError, Context);
 	return nullptr;
 #endif
 }
@@ -209,7 +235,7 @@ void UPCGAssetExporterUtils::UpdateAssets(const TArray<FAssetData>& PCGAssets, c
 		FEditorFileUtils::PromptForCheckoutAndSave(PackagesToSave, /*bCheckDirty=*/false, /*bPromptToSave=*/false);
 	}
 #else
-	PCGLog::LogErrorOnGraph(LOCTEXT("CannotExportInNonEditor", "PCG Asset Exporter Utils cannot be used in non-editor builds."), Context);
+	PCGLog::LogErrorOnGraph(PCGAssetExporterUtils::NonEditorError, Context);
 #endif // WITH_EDITOR
 }
 
diff --git a/Engine/Plugins/PCG/Source/PCG/Public/PCGAssetExporterUtils.h b/Engine/Plugins/PCG/Source/PCG/Public/PCGAssetExporterUtils.h
--- a/Engine/Plugins/PCG/Source/PCG/Public/PCGAssetExporterUtils.h
+++ b/Engine/Plugins/PCG/Source/PCG/Public/PCGAssetExporterUtils.h
@@ -30,4 +30,14 @@ public:
 	// Native API equivalent
 	static UPackage* CreateAsset(UPCGAssetExporter* Exporter, const FPCGAssetExporterParameters& Parameters, FPCGContext* InOptionalContext);
 	static void UpdateAssets(const TArray<FAssetData>& PCGAssets, const FPCGAssetExporterParameters& Parameters, FPCGContext* InOptionalContext);
+
+	// Generic equivalent
+	static UPackage* CreateAsset(const UClass* AssetClass, const FPCGAssetExporterParameters& Parameters,  TFunctionRef<bool(const FString& PackagePath, UObject* Asset)> ExportFunc, FPCGContext* InOptionalContext);
+
+	template <typename T>
+	static UPackage* CreateAsset(const FPCGAssetExporterParameters& Parameters,  TFunctionRef<bool(const FString& PackagePath, UObject* Asset)> ExportFunc, FPCGContext* InOptionalContext)
+	{
+		static_assert(std::is_base_of_v<UObject, T>, "T needs to be a UObject");
+		return CreateAsset(T::StaticClass(), Parameters, std::move(ExportFunc), InOptionalContext);
+	}
 };
\ No newline at end of file
