[PCG] Add DynMesh as new data type + node in PCG

* Unfortunately we don't have a way to expose external data types in PCG Editor at the moment, so it needs to pollute the main plugin

#rnx

diff --git a/Engine/Plugins/PCG/Source/PCG/Public/PCGCommon.h b/Engine/Plugins/PCG/Source/PCG/Public/PCGCommon.h
--- a/Engine/Plugins/PCG/Source/PCG/Public/PCGCommon.h
+++ b/Engine/Plugins/PCG/Source/PCG/Public/PCGCommon.h
@@ -75,9 +75,10 @@ enum class EPCGDataType : uint32
 
 	Volume = 1 << 7,
 	Primitive = 1 << 8,
+	DynamicMesh = 1 << 10,
 
 	/** Simple concrete data. */
-	Concrete = Point | PolyLine | Surface | Volume | Primitive,
+	Concrete = Point | PolyLine | Surface | Volume | Primitive | DynamicMesh,
 
 	/** Boolean operations like union, difference, intersection. */
 	Composite = 1 << 9 UMETA(Hidden),
@@ -119,7 +120,8 @@ enum class EPCGExclusiveDataType : uint8
 	Settings UMETA(Hidden),
 	Other,
 	Any,
-	PointOrParam
+	PointOrParam,
+	DynamicMesh
 };
 
 namespace PCGPinConstants
diff --git a/Engine/Plugins/PCG/Source/PCG/Public/PCGSettings.h b/Engine/Plugins/PCG/Source/PCG/Public/PCGSettings.h
--- a/Engine/Plugins/PCG/Source/PCG/Public/PCGSettings.h
+++ b/Engine/Plugins/PCG/Source/PCG/Public/PCGSettings.h
@@ -68,6 +68,7 @@ enum class EPCGSettingsType : uint8
 	GraphParameters,
 	Reroute,
 	GPU,
+	DynamicMesh,
 };
 
 #if WITH_EDITOR
diff --git a/Engine/Plugins/PCG/Source/PCGEditor/Private/PCGEditorCommon.h b/Engine/Plugins/PCG/Source/PCGEditor/Private/PCGEditorCommon.h
--- a/Engine/Plugins/PCG/Source/PCGEditor/Private/PCGEditorCommon.h
+++ b/Engine/Plugins/PCG/Source/PCGEditor/Private/PCGEditorCommon.h
@@ -48,6 +48,7 @@ namespace FPCGEditorCommon
 	const FName RenderTargetDataType = FName(TEXT("Render Target Data"));
 	const FName VolumeDataType = FName(TEXT("Volume Data"));
 	const FName PrimitiveDataType = FName(TEXT("Primitive Data"));
+	const FName DynamicMeshDataType = FName(TEXT("Dynamic Mesh Data"));
 
 	extern TAutoConsoleVariable<bool> CVarShowAdvancedAttributesFields;
 
diff --git a/Engine/Plugins/PCG/Source/PCGEditor/Private/PCGEditorGraphNodeBase.cpp b/Engine/Plugins/PCG/Source/PCGEditor/Private/PCGEditorGraphNodeBase.cpp
--- a/Engine/Plugins/PCG/Source/PCGEditor/Private/PCGEditorGraphNodeBase.cpp
+++ b/Engine/Plugins/PCG/Source/PCGEditor/Private/PCGEditorGraphNodeBase.cpp
@@ -947,6 +947,10 @@ FEdGraphPinType UPCGEditorGraphNodeBase::GetPinType(const UPCGPin* InPin)
 		{
 			EdPinType.PinSubCategory = FPCGEditorCommon::VolumeDataType;
 		}
+		else if (CheckType(EPCGDataType::DynamicMesh))
+		{
+			EdPinType.PinSubCategory = FPCGEditorCommon::DynamicMeshDataType;
+		}
 		else if (CheckType(EPCGDataType::Primitive))
 		{
 			EdPinType.PinSubCategory = FPCGEditorCommon::PrimitiveDataType;
diff --git a/Engine/Plugins/PCG/Source/PCGEditor/Private/PCGEditorSettings.cpp b/Engine/Plugins/PCG/Source/PCGEditor/Private/PCGEditorSettings.cpp
--- a/Engine/Plugins/PCG/Source/PCGEditor/Private/PCGEditorSettings.cpp
+++ b/Engine/Plugins/PCG/Source/PCGEditor/Private/PCGEditorSettings.cpp
@@ -32,6 +32,7 @@ UPCGEditorSettings::UPCGEditorSettings(const FObjectInitializer& ObjectInitializ
 	HierarchicalGenerationNodeColor = FLinearColor(1.0f, 0.132868f, 0.0f);
 	GraphParametersNodeColor = FLinearColor::Yellow;
 	RerouteNodeColor = FLinearColor(0.5f, 1.0f, 0.83f);
+	DynamicMeshNodeColor = FLinearColor(0.42f, 0.25f, 1.0f);
 
 	DefaultPinColor = FLinearColor(0.29f, 0.29f, 0.29f);
 	SpatialDataPinColor = FLinearColor(1.0f, 1.0f, 1.0f);
@@ -45,6 +46,7 @@ UPCGEditorSettings::UPCGEditorSettings(const FObjectInitializer& ObjectInitializ
 	RenderTargetDataPinColor = FLinearColor(0.8f, 0.18f, 0.12f);
 	VolumeDataPinColor = FLinearColor(0.79f, 0.06f, 0.5f);
 	PrimitiveDataPinColor = FLinearColor(0.22f, 0.05f, 1.0f);
+	DynamicMeshPinColor = FLinearColor(0.42f, 0.25f, 1.0f);
 
 	ParamDataPinColor = FLinearColor(1.0f, 0.38f, 0.02f);
 	UnknownDataPinColor = FLinearColor(0.3f, 0.3f, 0.3f);
@@ -98,6 +100,8 @@ FLinearColor UPCGEditorSettings::GetColor(UPCGSettings* Settings) const
 			return GraphParametersNodeColor;
 		case EPCGSettingsType::Reroute:
 			return RerouteNodeColor;
+		case EPCGSettingsType::DynamicMesh:
+			return DynamicMeshNodeColor;
 		case EPCGSettingsType::Generic: // falls through
 		default:
 			// Finally, we couldn't find any match, so return the default value
@@ -142,6 +146,10 @@ FLinearColor UPCGEditorSettings::GetPinColor(const FEdGraphPinType& PinType) con
 		{
 			return VolumeDataPinColor;
 		}
+		else if (PinType.PinSubCategory == FPCGEditorCommon::DynamicMeshDataType)
+		{
+			return DynamicMeshPinColor;
+		}
 		else if (PinType.PinSubCategory == FPCGEditorCommon::PrimitiveDataType)
 		{
 			return PrimitiveDataPinColor;
diff --git a/Engine/Plugins/PCG/Source/PCGEditor/Public/PCGEditorSettings.h b/Engine/Plugins/PCG/Source/PCGEditor/Public/PCGEditorSettings.h
--- a/Engine/Plugins/PCG/Source/PCGEditor/Public/PCGEditorSettings.h
+++ b/Engine/Plugins/PCG/Source/PCGEditor/Public/PCGEditorSettings.h
@@ -108,6 +108,10 @@ public:
 	UPROPERTY(EditAnywhere, config, Category = Node, meta = (HideAlphaChannel))
 	FLinearColor RerouteNodeColor;
 
+	/** Color used for dynamic mesh nodes */
+	UPROPERTY(EditAnywhere, config, Category = Node, meta = (HideAlphaChannel))
+	FLinearColor DynamicMeshNodeColor;
+
 	/** Default pin color */
 	UPROPERTY(EditAnywhere, config, Category = Node, meta = (HideAlphaChannel))
 	FLinearColor DefaultPinColor;
@@ -160,6 +164,10 @@ public:
 	UPROPERTY(EditAnywhere, config, Category = Node, meta = (HideAlphaChannel, DisplayName = "Attribute Set Pin Color"))
 	FLinearColor ParamDataPinColor;
 
+	/** Color used for data pins of type Dynamic Mesh */
+	UPROPERTY(EditAnywhere, config, Category = Node, meta = (HideAlphaChannel))
+	FLinearColor DynamicMeshPinColor;
+
 	/** Color used for other/unknown data types */
 	UPROPERTY(EditAnywhere, config, Category = Node, meta = (HideAlphaChannel))
 	FLinearColor UnknownDataPinColor;
diff --git a/Engine/Plugins/PCGInterops/PCGGeometryScriptInterop/Source/PCGGeometryScriptInterop/PCGGeometryScriptInterop.Build.cs b/Engine/Plugins/PCGInterops/PCGGeometryScriptInterop/Source/PCGGeometryScriptInterop/PCGGeometryScriptInterop.Build.cs
--- a/Engine/Plugins/PCGInterops/PCGGeometryScriptInterop/Source/PCGGeometryScriptInterop/PCGGeometryScriptInterop.Build.cs
+++ b/Engine/Plugins/PCGInterops/PCGGeometryScriptInterop/Source/PCGGeometryScriptInterop/PCGGeometryScriptInterop.Build.cs
@@ -35,6 +35,7 @@ namespace UnrealBuildTool.Rules
 					new string[]
 					{
 						"UnrealEd",
+						"PCGEditor"
 					}
 				);
 			}
