[PCG] Expose Duplicate data to BP

#rnx

diff --git a/Engine/Plugins/PCG/Source/PCG/Private/Helpers/PCGBlueprintHelpers.cpp b/Engine/Plugins/PCG/Source/PCG/Private/Helpers/PCGBlueprintHelpers.cpp
--- a/Engine/Plugins/PCG/Source/PCG/Private/Helpers/PCGBlueprintHelpers.cpp
+++ b/Engine/Plugins/PCG/Source/PCG/Private/Helpers/PCGBlueprintHelpers.cpp
@@ -270,4 +270,9 @@ void UPCGBlueprintHelpers::RefreshPCGRuntimeComponent(UPCGComponent* InComponent
 
 		PCGSubsystem->RefreshRuntimeGenComponent(InComponent, EPCGChangeType::GenerationGrid);
 	}
+}
+
+UPCGData* UPCGBlueprintHelpers::DuplicateData(const UPCGData* InData, FPCGContext& Context, bool bInitializeMetadata)
+{
+	return InData ? InData->DuplicateData(&Context, bInitializeMetadata) : nullptr;
 }
\ No newline at end of file
diff --git a/Engine/Plugins/PCG/Source/PCG/Public/Helpers/PCGBlueprintHelpers.h b/Engine/Plugins/PCG/Source/PCG/Public/Helpers/PCGBlueprintHelpers.h
--- a/Engine/Plugins/PCG/Source/PCG/Public/Helpers/PCGBlueprintHelpers.h
+++ b/Engine/Plugins/PCG/Source/PCG/Public/Helpers/PCGBlueprintHelpers.h
@@ -92,6 +92,11 @@ public:
 	/** Refresh a component set to Generate At Runtime, if some parameters changed. Can also flush the cache. */
 	UFUNCTION(BlueprintCallable, Category = "PCG|Runtime", meta = (ScriptMethod, DisplayName = "Refresh PCG Runtime Component"))
 	static void RefreshPCGRuntimeComponent(UPCGComponent* InComponent, const bool bFlushCache = false);
+
+	// Implementation note: Needs to be done outside of UPCGData because of circular dependency between PCGContext.h and PCGData.h
+	/** Return a copy of the data, with Metadata inheritance for spatial data. */
+	UFUNCTION(BlueprintCallable, Category="PCG|Data", meta = (ScriptMethod))
+	static UPCGData* DuplicateData(const UPCGData* InData, UPARAM(ref) FPCGContext& Context, bool bInitializeMetadata = true);
 };
 
 #if UE_ENABLE_INCLUDE_ORDER_DEPRECATED_IN_5_2
