[PCG] Add option to do N:N operation on the Distance node

* Right now distance is for each source, check against all targets
* Give the option to check the source only against the corresponding target, off by default.

#rnx

diff --git a/Engine/Plugins/PCG/Source/PCG/Private/Elements/PCGDistance.cpp b/Engine/Plugins/PCG/Source/PCG/Private/Elements/PCGDistance.cpp
--- a/Engine/Plugins/PCG/Source/PCG/Private/Elements/PCGDistance.cpp
+++ b/Engine/Plugins/PCG/Source/PCG/Private/Elements/PCGDistance.cpp
@@ -115,7 +115,7 @@ bool FPCGDistanceElement::ExecuteInternal(FPCGContext* Context) const
 
 	if (Context->Node && !Context->Node->IsInputPinConnected(PCGDistance::TargetLabel))
 	{
-		// If Target pin is unconnected then we no-op and pass through all data from Target pin.
+		// If Target pin is unconnected then we no-op and pass through all data from Source pin.
 		Context->OutputData = PCGGather::GatherDataForPin(Context->InputData, PCGDistance::SourceLabel);
 		return true;
 	}
@@ -152,15 +152,29 @@ bool FPCGDistanceElement::ExecuteInternal(FPCGContext* Context) const
 		if (!TargetPointData)
 		{
 			PCGE_LOG(Error, GraphAndLog, FText::Format(LOCTEXT("CannotConvertToPoint", "Cannot convert target '{0}' into Point data"), FText::FromString(Target.Data->GetClass()->GetName())));
-			continue;			
+			continue;
 		}
 
 		TargetPointDatas.Add(TargetPointData);
 	}
 
+	if (TargetPointDatas.IsEmpty())
+	{
+		// If Target pin has no valid data then we no-op and pass through all data from Source pin.
+		Context->OutputData = PCGGather::GatherDataForPin(Context->InputData, PCGDistance::SourceLabel);
+		return true;
+	}
+
+	if (Settings->bCheckSourceAgainstRespectiveTarget && Sources.Num() > 1 && TargetPointDatas.Num() > 1 && Sources.Num() != TargetPointDatas.Num())
+	{
+		PCGLog::InputOutput::LogInvalidCardinalityError(PCGDistance::SourceLabel, PCGDistance::TargetLabel, Context);
+		return true;
+	}
+
 	// First find the total Input bounds which will determine the size of each cell
-	for (const FPCGTaggedData& Source : Sources) 
+	for (int32 SourceIndex = 0; SourceIndex < Sources.Num(); ++SourceIndex) 
 	{
+		const FPCGTaggedData& Source = Sources[SourceIndex];
 		// Add the point bounds to the input cell
 
 		const UPCGSpatialData* SourceData = Cast<UPCGSpatialData>(Source.Data);
@@ -249,7 +263,7 @@ bool FPCGDistanceElement::ExecuteInternal(FPCGContext* Context) const
 			}
 		}
 
-		auto ProcessDistance = [SourceShape, TargetShape, &TargetPointDatas, MaximumDistance, MaximumDistanceRecip, bSetDensity, bOutputDistanceVector, &OutPoints, &SourcePoints, &ResultCache, bWriteToAttribute = Accessor.IsValid()](int32 ReadIndex, int32 WriteIndex)
+		auto ProcessDistance = [SourceShape, TargetShape, &TargetPointDatas, MaximumDistance, MaximumDistanceRecip, bSetDensity, bOutputDistanceVector, &OutPoints, &SourcePoints, &ResultCache, bWriteToAttribute = Accessor.IsValid(), Settings, SourceIndex](int32 ReadIndex, int32 WriteIndex)
 		{
 			FPCGPoint& OutPoint = OutPoints[WriteIndex];
 			const FPCGPoint& SourcePoint = SourcePoints[ReadIndex];
@@ -291,13 +305,27 @@ bool FPCGDistanceElement::ExecuteInternal(FPCGContext* Context) const
 				}
 			};
 
-			for (const UPCGPointData* TargetPointData : TargetPointDatas)
+			auto CheckAgainstTargetPointData = [&SourceQueryBounds, &CalculateSDF](const UPCGPointData* TargetPointData)
 			{
+				check(TargetPointData);
+				
 				const UPCGPointData::PointOctree& Octree = TargetPointData->GetOctree();
 
 				Octree.FindElementsWithBoundsTest(
 						FBoxCenterAndExtent(SourceQueryBounds.Origin, SourceQueryBounds.BoxExtent),
 						CalculateSDF);
+			};
+
+			if (Settings->bCheckSourceAgainstRespectiveTarget)
+			{
+				CheckAgainstTargetPointData(TargetPointDatas[SourceIndex % TargetPointDatas.Num()]);
+			}
+			else
+			{
+				for (const UPCGPointData* TargetPointData : TargetPointDatas)
+				{
+					CheckAgainstTargetPointData(TargetPointData);
+				}
 			}
 
 			const double Distance = FMath::Sign(MinDistanceSquared) * FMath::Sqrt(FMath::Abs(MinDistanceSquared));
diff --git a/Engine/Plugins/PCG/Source/PCG/Public/Elements/PCGDistance.h b/Engine/Plugins/PCG/Source/PCG/Public/Elements/PCGDistance.h
--- a/Engine/Plugins/PCG/Source/PCG/Public/Elements/PCGDistance.h
+++ b/Engine/Plugins/PCG/Source/PCG/Public/Elements/PCGDistance.h
@@ -80,6 +80,10 @@ public:
 	UPROPERTY(BlueprintReadWrite, EditAnywhere, Category = Settings, meta = (PCG_Overridable))
 	PCGDistanceShape TargetShape = PCGDistanceShape::SphereBounds;
 
+	/** If this option is on, each source will be tested against its respective target (for a N:N operation). Source and Target num must be the same (or 1).  */
+	UPROPERTY(BlueprintReadWrite, EditAnywhere, Category = Settings, meta = (PCG_Overridable))
+	bool bCheckSourceAgainstRespectiveTarget = false;
+
 #if WITH_EDITORONLY_DATA
 	UPROPERTY(meta=(DeprecatedProperty, DeprecationMessage="Use OutputAttribute selector instead."))
 	FName AttributeName_DEPRECATED = PCGDistanceConstants::DefaultOutputAttributeName;
