Added physics test to pathfinding traversal

diff --git a/Engine/Plugins/PCG/Source/PCG/Private/Elements/PCGPathfindingElement.cpp b/Engine/Plugins/PCG/Source/PCG/Private/Elements/PCGPathfindingElement.cpp
--- a/Engine/Plugins/PCG/Source/PCG/Private/Elements/PCGPathfindingElement.cpp
+++ b/Engine/Plugins/PCG/Source/PCG/Private/Elements/PCGPathfindingElement.cpp
@@ -3,6 +3,7 @@
 #include "Elements/PCGPathfindingElement.h"
 
 #include "PCGContext.h"
+#include "PCGComponent.h"
 #include "Data/PCGPointData.h"
 #include "Data/PCGSplineData.h"
 #include "Helpers/PCGHelpers.h"
@@ -12,6 +13,8 @@
 #include "SpatialAlgo/PCGAStar.h"
 
 #include "Components/SplineComponent.h"
+#include "Engine/HitResult.h"
+#include "Engine/World.h"
 
 #define LOCTEXT_NAMESPACE "PCGPathfindingElement"
 
@@ -51,6 +54,12 @@ namespace PCGPathfindingElement::Helpers
 	}
 } // namespace PCGPathfindingElement::Helpers
 
+UPCGPathfindingSettings::UPCGPathfindingSettings()
+{
+	// In most cases, we're not going to be interested in checking for occlusion by the landscape itself, as we'll be pathfinding on the landscape.
+	PathTraceParams.SelectLandscapeHits = EPCGWorldQuerySelectLandscapeHits::Exclude;
+}
+
 TArray<FPCGPinProperties> UPCGPathfindingSettings::InputPinProperties() const
 {
 	TArray<FPCGPinProperties> Properties;
@@ -92,6 +101,12 @@ EPCGChangeType UPCGPathfindingSettings::GetChangeTypeForProperty(const FName& In
 }
 #endif // WITH_EDITOR
 
+bool FPCGPathfindingElement::IsCacheable(const UPCGSettings* InSettings) const
+{
+	const UPCGPathfindingSettings* Settings = Cast<const UPCGPathfindingSettings>(InSettings);
+	return !Settings || !Settings->bUsePathTraces;
+}
+
 bool FPCGPathfindingElement::PrepareDataInternal(FPCGContext* InContext) const
 {
 	TRACE_CPUPROFILER_EVENT_SCOPE(FPCGPathfindingElement::PrepareData);
@@ -166,14 +181,49 @@ bool FPCGPathfindingElement::PrepareDataInternal(FPCGContext* InContext) const
 			}
 		}
 
+		TFunction<bool(const FVector&, const FVector&)> LineTraceTest = [](const FVector&, const FVector&) -> bool { return true; };
+
+		if (Settings->bUsePathTraces)
+		{
+			if (UWorld* World = Context->SourceComponent.Get() ? Context->SourceComponent->GetWorld() : nullptr)
+			{
+				FPCGWorldRaycastQueryParams PathTraceParams = Settings->PathTraceParams;
+				PathTraceParams.Initialize();
+
+				TWeakObjectPtr<UPCGComponent> OriginatingComponent = Context->SourceComponent;
+				FCollisionObjectQueryParams ObjectQueryParams(PathTraceParams.CollisionChannel);
+				FCollisionQueryParams Params;
+				Params.bTraceComplex = PathTraceParams.bTraceComplex;
+
+				LineTraceTest = [World, ObjectQueryParams = MoveTemp(ObjectQueryParams), Params = MoveTemp(Params), OriginatingComponent, PathTraceParams = MoveTemp(PathTraceParams)](const FVector& StartPosition, const FVector& EndPosition) -> bool
+				{
+					TArray<FHitResult> OutHits;
+					if (World->LineTraceMultiByObjectType(OutHits, StartPosition, EndPosition, ObjectQueryParams, Params))
+					{
+						TOptional<FHitResult> HitResult = PCGWorldQueryHelpers::FilterRayHitResults(&PathTraceParams, OriginatingComponent, OutHits);
+						return !HitResult.IsSet();
+					}
+					else
+					{
+						return true;
+					}
+				};
+			}
+		}
+
 		if (CostAccessor)
 		{
 			if (Settings->CostFunctionMode == EPCGPathfindingCostFunctionMode::FitnessScore)
 			{
 				const double MaxFitnessPenaltyFactor = FMath::Max(Settings->MaximumFitnessPenaltyFactor, 1.0);
 
-				OutSearchState.CostFunction = [FitnessAccessor = CostAccessor, MaxFitnessPenaltyFactor](const double PreviousNodeCost, const FPCGPoint* PreviousNodePoint, const double DistanceToCurrentSquared, const FPCGPoint* CurrentNodePoint)
+				OutSearchState.CostFunction = [FitnessAccessor = CostAccessor, MaxFitnessPenaltyFactor, PathTraceTest = MoveTemp(LineTraceTest)](const double PreviousNodeCost, const FPCGPoint* PreviousNodePoint, const double DistanceToCurrentSquared, const FPCGPoint* CurrentNodePoint)
 				{
+					if(!PathTraceTest(PreviousNodePoint->Transform.GetLocation(), CurrentNodePoint->Transform.GetLocation()))
+					{
+						return std::numeric_limits<double>::max();
+					}
+
 					double FitnessScore = 1.0;
 					FPCGAttributeAccessorKeysPoints Key(*CurrentNodePoint);
 
@@ -185,8 +235,13 @@ bool FPCGPathfindingElement::PrepareDataInternal(FPCGContext* InContext) const
 			}
 			else if (Settings->CostFunctionMode == EPCGPathfindingCostFunctionMode::CostMultipler)
 			{
-				OutSearchState.CostFunction = [MultiplierAccessor = CostAccessor](const double PreviousNodeCost, const FPCGPoint* PreviousNodePoint, const double DistanceToCurrentSquared, const FPCGPoint* CurrentNodePoint)
+				OutSearchState.CostFunction = [MultiplierAccessor = CostAccessor, PathTraceTest = MoveTemp(LineTraceTest)](const double PreviousNodeCost, const FPCGPoint* PreviousNodePoint, const double DistanceToCurrentSquared, const FPCGPoint* CurrentNodePoint)
 				{
+					if (!PathTraceTest(PreviousNodePoint->Transform.GetLocation(), CurrentNodePoint->Transform.GetLocation()))
+					{
+						return std::numeric_limits<double>::max();
+					}
+
 					double Multiplier = 1.0;
 					FPCGAttributeAccessorKeysPoints Key(*CurrentNodePoint);
 
@@ -201,6 +256,21 @@ bool FPCGPathfindingElement::PrepareDataInternal(FPCGContext* InContext) const
 				checkNoEntry();
 			}
 		}
+		else if (Settings->bUsePathTraces)
+		{
+			// Use distance but with line trace
+			OutSearchState.CostFunction = [PathTraceTest = MoveTemp(LineTraceTest)](const double PreviousNodeCost, const FPCGPoint* PreviousNodePoint, const double DistanceToCurrentSquared, const FPCGPoint* CurrentNodePoint)
+			{
+				if (!PathTraceTest(PreviousNodePoint->Transform.GetLocation(), CurrentNodePoint->Transform.GetLocation()))
+				{
+					return std::numeric_limits<double>::max();
+				}
+				else
+				{
+					return PCGSpatialAlgo::AStar::Cost::CalculateCost_EuclideanDistance(PreviousNodeCost, PreviousNodePoint, DistanceToCurrentSquared, CurrentNodePoint);
+				}
+			};
+		}
 
 		PCGSpatialAlgo::AStar::Initialize(OutSearchState.OriginatingPointData, SearchSettings, OutSearchState);
 
diff --git a/Engine/Plugins/PCG/Source/PCG/Private/SpatialAlgo/PCGAStar.cpp b/Engine/Plugins/PCG/Source/PCG/Private/SpatialAlgo/PCGAStar.cpp
--- a/Engine/Plugins/PCG/Source/PCG/Private/SpatialAlgo/PCGAStar.cpp
+++ b/Engine/Plugins/PCG/Source/PCG/Private/SpatialAlgo/PCGAStar.cpp
@@ -104,6 +104,12 @@ bool ExecuteSearchIteration(const FSearchSettings& SearchSettings, FSearchState&
 
 		const double TentativeNewLocalCost = SearchState.CostFunction(NodeList[CurrentNodeIndex].LocalCost, NodeList[CurrentNodeIndex].PCGPoint, DistanceToPointSquared, Point);
 
+		// Blocked paths don't need expansion here, and serve no purpose at this point (at least from that direction) from being visited.
+		if (TentativeNewLocalCost == std::numeric_limits<double>::max())
+		{
+			return;
+		}
+
 		// Not tracking this point yet--add it to the node list and map it.
 		if (!bTrackedPoint)
 		{
diff --git a/Engine/Plugins/PCG/Source/PCG/Public/Elements/PCGPathfindingElement.h b/Engine/Plugins/PCG/Source/PCG/Public/Elements/PCGPathfindingElement.h
--- a/Engine/Plugins/PCG/Source/PCG/Public/Elements/PCGPathfindingElement.h
+++ b/Engine/Plugins/PCG/Source/PCG/Public/Elements/PCGPathfindingElement.h
@@ -4,6 +4,7 @@
 
 #include "PCGSettings.h"
 
+#include "Data/PCGWorldData.h"
 #include "Elements/PCGTimeSlicedElementBase.h"
 #include "SpatialAlgo/PCGAStar.h"
 
@@ -29,12 +30,14 @@ enum class EPCGPathfindingCostFunctionMode : uint8
 /** Finds the optimal path across the points of a given point cloud--should one exist--when provided a start and goal
  * location, and a maximum jump distance between points. Can return a partial path.
  */
-UCLASS(MinimalAPI, BlueprintType, ClassGroup = (Procedural))
+UCLASS(MinimalAPI, BlueprintType, ClassGroup = (Procedural), hidecategories="Data|Attributes")
 class UPCGPathfindingSettings : public UPCGSettings
 {
 	GENERATED_BODY()
 
 public:
+	UPCGPathfindingSettings();
+
 	//~Begin UPCGSettings interface
 #if WITH_EDITOR
 	virtual FName GetDefaultNodeName() const override { return FName(TEXT("PathfindingElement")); }
@@ -81,6 +84,13 @@ public:
 	UPROPERTY(EditAnywhere, BlueprintReadWrite, Category = Settings, meta = (ClampMin = "1.0", EditCondition = "CostFunctionMode == EPCGPathfindingCostFunctionMode::FitnessScore", EditConditionHides, PCG_Overridable))
 	double MaximumFitnessPenaltyFactor = 10.0;
 
+	/** Controls whether raycasts will be used to test for collisions along the path (hit results will be considered obstacles for the pathfinding). */
+	UPROPERTY(EditAnywhere, BlueprintReadWrite, Category = Settings, meta = (PCG_Overridable))
+	bool bUsePathTraces = false;
+
+	UPROPERTY(EditAnywhere, BlueprintReadWrite, Category = Settings, meta = (EditCondition = "bUsePathTraces", EditConditionHides, PCG_Overridable, ShowOnlyInnerProperties))
+	FPCGWorldRaycastQueryParams PathTraceParams;
+
 	/** Even if the path is not complete, return the most optimal and viable partial path to the goal. */
 	UPROPERTY(EditAnywhere, BlueprintReadWrite, Category = Settings, meta = (PCG_Overridable))
 	bool bAcceptPartialPath = true;
@@ -100,6 +110,8 @@ public:
 
 class FPCGPathfindingElement : public TPCGTimeSlicedElementBase<PCGSpatialAlgo::AStar::FSearchSettings, PCGSpatialAlgo::AStar::FSearchState>
 {
+public:
+	virtual bool IsCacheable(const UPCGSettings* InSettings) const override;
 protected:
 	virtual bool PrepareDataInternal(FPCGContext* InContext) const override;
 	virtual bool ExecuteInternal(FPCGContext* InContext) const override;
