[PCG] Fix crash in source editor code due to bad use of temporary object.

#rnx

diff --git a/Engine/Plugins/PCG/Source/PCGEditor/Private/PCGEditor.cpp b/Engine/Plugins/PCG/Source/PCGEditor/Private/PCGEditor.cpp
--- a/Engine/Plugins/PCG/Source/PCGEditor/Private/PCGEditor.cpp
+++ b/Engine/Plugins/PCG/Source/PCGEditor/Private/PCGEditor.cpp
@@ -3174,10 +3174,6 @@ void FPCGEditor::OnNodeSourceCompiled(const UPCGNode* InNode, const FPCGCompiler
 	{
 		NodeSourceWidget->OnDiagnosticsUpdated(InDiagnostics);
 	}
-	else
-	{
-		NodeSourceWidget->OnDiagnosticsUpdated({});
-	}
 }
 
 void FPCGEditor::OnMapChanged(UWorld* InWorld, EMapChangeType InMapChangedType)
