[PCG] Fix issue with detecting multiple "Select" comparators on Select Grammar Node

#rnx

diff --git a/Engine/Plugins/PCG/Source/PCG/Private/Elements/Grammar/PCGSelectGrammar.cpp b/Engine/Plugins/PCG/Source/PCG/Private/Elements/Grammar/PCGSelectGrammar.cpp
--- a/Engine/Plugins/PCG/Source/PCG/Private/Elements/Grammar/PCGSelectGrammar.cpp
+++ b/Engine/Plugins/PCG/Source/PCG/Private/Elements/Grammar/PCGSelectGrammar.cpp
@@ -8,6 +8,8 @@
 #include "Elements/Metadata/PCGMetadataElementCommon.h"
 #include "Helpers/PCGPropertyHelpers.h"
 
+#include "Algo/Count.h"
+
 #define LOCTEXT_NAMESPACE "PCGSelectGrammarElement"
 
 namespace PCGSelectGrammar
@@ -243,9 +245,10 @@ bool FPCGSelectGrammarElement::ExecuteInternal(FPCGContext* InContext) const
 		CriteriaMap[Criterion.Key].Emplace(std::move(Criterion));
 	}
 
-	for (const TTuple<FName, TArray<FPCGSelectGrammarCriterion, TSizedInlineAllocator<16, 32>>>& Tuple : CriteriaMap)
+	for (const TTuple<FName, TArray<FPCGSelectGrammarCriterion, TInlineAllocator<16>>>& Tuple : CriteriaMap)
 	{
-		if (Tuple.Get<1>().Num() > 1)
+		const int32 SelectionCount = Algo::CountIf(Tuple.Get<1>(), [](const FPCGSelectGrammarCriterion& InCriterion) { return InCriterion.Comparator == EPCGSelectGrammarComparator::Select; });
+		if (SelectionCount > 1)
 		{
 			PCGLog::LogWarningOnGraph(FText::Format(LOCTEXT("MultipleSelectSameCriteria", "Multiple 'Select' comparators found on the criteria with same key: {0}"), FText::FromName(Tuple.Get<0>())), InContext);
 		}
