[PCG] Add the concept of data used multiple times

* It is computed at the compilation stage and stored on the task input.
* There is also a runtime component, that is enabled with a virtual function on the element.
* Can be used in the end to steal data and work in place.

#rnx

diff --git a/Engine/Plugins/PCG/Source/PCG/Private/Graph/PCGGraphCompiler.cpp b/Engine/Plugins/PCG/Source/PCG/Private/Graph/PCGGraphCompiler.cpp
--- a/Engine/Plugins/PCG/Source/PCG/Private/Graph/PCGGraphCompiler.cpp
+++ b/Engine/Plugins/PCG/Source/PCG/Private/Graph/PCGGraphCompiler.cpp
@@ -33,6 +33,48 @@ namespace PCGGraphCompiler
 		TEXT("pcg.GraphExecution.GPU.Enable"),
 		true,
 		TEXT("Whether to emit compatible nodes as compute graphs to execute on the GPU."));
+
+	/**
+	 * Utility structure to store "FPCGGraphTaskInput" in a hashable set
+	 * and query it easily using only the task id and its upstream pin, to keep track of inputs used multiple times.
+	 * TaskInput is the same as another if they have the same task id and the same upstream pin (or they are both null).
+	 */
+	struct FTaskInputOrigin
+	{
+		FTaskInputOrigin(FPCGGraphTaskInput* InTaskInput)
+		{
+			// Only set it if the upstream properties is set
+			TaskInput = (InTaskInput && InTaskInput->UpstreamPin.IsSet()) ? InTaskInput : nullptr;
+		}
+		
+		friend uint32 GetTypeHash(const FTaskInputOrigin& Value)
+		{
+			check(!Value.TaskInput || Value.TaskInput->UpstreamPin.IsSet());
+			return Value.TaskInput ? HashCombine(Value.TaskInput->TaskId, GetTypeHash(Value.TaskInput->UpstreamPin.GetValue())) : PointerHash(nullptr);
+		}
+
+		bool operator==(const FTaskInputOrigin& Other) const
+		{
+			if (!TaskInput && !Other.TaskInput)
+			{
+				return true;
+			}
+			else if (!TaskInput || !Other.TaskInput || TaskInput->TaskId != Other.TaskInput->TaskId)
+			{
+				return false;
+			}
+			else
+			{
+				// By construction those are set.
+				const FPCGPinProperties& Properties = TaskInput->UpstreamPin.GetValue();
+				const FPCGPinProperties& OtherProperties = Other.TaskInput->UpstreamPin.GetValue();
+
+				return Properties == OtherProperties;
+			}
+		}
+		
+		FPCGGraphTaskInput* TaskInput = nullptr;
+	};
 }
 
 TArray<FPCGGraphTask> FPCGGraphCompiler::CompileGraph(UPCGGraph* InGraph, FPCGTaskId& NextId, FPCGStackContext& InOutStackContext)
@@ -1128,6 +1170,31 @@ void FPCGGraphCompiler::CompileTopGraph(UPCGGraph* InGraph, uint32 GenerationGri
 		CalculateDynamicActivePinDependencies(CompiledTasks[TaskIndex].NodeId, CompiledTasks);
 	}
 
+	// Also we keep track of all the inputs of each task. If the input is in SeenInputs it means it is used multiple times, so mark it as is.
+	TSet<PCGGraphCompiler::FTaskInputOrigin> SeenInputs;
+	for (int TaskIndex = 0; TaskIndex < CompiledTasks.Num(); ++TaskIndex)
+	{
+		for (FPCGGraphTaskInput& TaskInput : CompiledTasks[TaskIndex].Inputs)
+		{
+			PCGGraphCompiler::FTaskInputOrigin InputOrigin(&TaskInput);
+			if (!InputOrigin.TaskInput)
+			{
+				continue;
+			}
+
+			if (PCGGraphCompiler::FTaskInputOrigin* It = SeenInputs.Find(InputOrigin))
+			{
+				check(It->TaskInput);
+				It->TaskInput->bIsUsedMultipleTimes = true;
+				TaskInput.bIsUsedMultipleTimes = true;
+			}
+			else
+			{
+				SeenInputs.Emplace(std::move(InputOrigin));
+			}
+		}
+	}
+
 	const int TaskNum = CompiledTasks.Num();
 	const FPCGTaskId PreExecuteTaskId = FPCGTaskId(TaskNum);
 	const FPCGTaskId PostExecuteTaskId = PreExecuteTaskId + 1;
diff --git a/Engine/Plugins/PCG/Source/PCG/Private/Graph/PCGGraphExecutor.cpp b/Engine/Plugins/PCG/Source/PCG/Private/Graph/PCGGraphExecutor.cpp
--- a/Engine/Plugins/PCG/Source/PCG/Private/Graph/PCGGraphExecutor.cpp
+++ b/Engine/Plugins/PCG/Source/PCG/Private/Graph/PCGGraphExecutor.cpp
@@ -2240,11 +2240,18 @@ void FPCGGraphExecutor::BuildTaskInput(FPCGGraphTask& Task, TArray<FPCGTaskId>&
 
 		// Apply labelling on data; technically, we should ensure that we do this only for pass-through nodes,
 		// Otherwise we could also null out the label on the input...
-		if (Input.DownstreamPin.IsSet())
+		// Also marking the data used multiple times if the given input is used multiple times.
+		if (Input.DownstreamPin.IsSet() || Input.bIsUsedMultipleTimes)
 		{
 			for (int32 TaggedDataIndex = TaggedDataOffset; TaggedDataIndex < Task.TaskInput.TaggedData.Num(); ++TaggedDataIndex)
 			{
-				Task.TaskInput.TaggedData[TaggedDataIndex].Pin = Input.DownstreamPin.GetValue().Label;
+				FPCGTaggedData& TaggedData = Task.TaskInput.TaggedData[TaggedDataIndex];
+				if (Input.DownstreamPin.IsSet())
+				{
+					TaggedData.Pin = Input.DownstreamPin.GetValue().Label;
+				}
+
+				TaggedData.bIsUsedMultipleTimes |= Input.bIsUsedMultipleTimes;
 			}
 		}
 	}
diff --git a/Engine/Plugins/PCG/Source/PCG/Private/Graph/PCGGraphExecutor.h b/Engine/Plugins/PCG/Source/PCG/Private/Graph/PCGGraphExecutor.h
--- a/Engine/Plugins/PCG/Source/PCG/Private/Graph/PCGGraphExecutor.h
+++ b/Engine/Plugins/PCG/Source/PCG/Private/Graph/PCGGraphExecutor.h
@@ -63,6 +63,9 @@ struct FPCGGraphTaskInput
 
 	/** Whether the input provides any data. For the post execute task, only the output node will provide data. */
 	bool bProvideData;
+	
+	/** Whether the input is used multiple times (previous task has multiple successors). Useful to know if the data can be stolen or not. True by default.*/
+	bool bIsUsedMultipleTimes = true;
 
 	static inline const TOptional<FPCGPinProperties> NoPin = TOptional<FPCGPinProperties>();
 };
diff --git a/Engine/Plugins/PCG/Source/PCG/Private/PCGElement.cpp b/Engine/Plugins/PCG/Source/PCG/Private/PCGElement.cpp
--- a/Engine/Plugins/PCG/Source/PCG/Private/PCGElement.cpp
+++ b/Engine/Plugins/PCG/Source/PCG/Private/PCGElement.cpp
@@ -401,6 +401,43 @@ void IPCGElement::PostExecute(FPCGContext* Context) const
 		}
 	}
 
+	// Analyze if the output data is used multiple times, if the element requires it.
+	if (ShouldVerifyIfOutputsAreUsedMultipleTimes(Settings))
+	{
+		for (FPCGTaggedData& OutputData : Context->OutputData.TaggedData)
+		{
+			// Enforce that pinless data is always used multiple times, or if the debug mode is enabled.
+			if (OutputData.bPinlessData || (Settings && Settings->CanBeDebugged() && Settings->bDebug))
+			{
+				OutputData.bIsUsedMultipleTimes = true;
+				continue;
+			}
+			
+			// For data that are marked to be used multiple times, they are potentially not used multiple times if they are not passthrough
+			// (hence if they are not in the input). So set them back to false in that case. It will be set to true again by the executor
+			// if it is actually used in multiple places.
+			auto PassthroughPredicate = [&OutputData](const FPCGTaggedData& InputData) { return InputData.Data == OutputData.Data; };
+			if (OutputData.bIsUsedMultipleTimes && !Context->InputData.TaggedData.ContainsByPredicate(PassthroughPredicate))
+			{
+#if !UE_BUILD_SHIPPING
+				OutputData.OriginatingNode = Context->Node;
+#endif //!UE_BUILD_SHIPPING
+				OutputData.bIsUsedMultipleTimes = false;
+			}
+			
+			// We also need to verify that the data is not used in other outputs.
+			auto SameDataDifferentTaggedData = [&OutputData](const FPCGTaggedData& OtherOutputData)
+			{
+				return (&OtherOutputData != &OutputData) && (OtherOutputData.Data == OutputData.Data);
+			};
+			
+			if (!OutputData.bIsUsedMultipleTimes && Context->OutputData.TaggedData.ContainsByPredicate(SameDataDifferentTaggedData))
+			{
+				OutputData.bIsUsedMultipleTimes = true;
+			}
+		}
+	}
+	
 #if WITH_EDITOR
 	// Register the element to the component indicating the element has run and can have dynamic tracked keys.
 	if (Settings && Settings->CanDynamicallyTrackKeys() && Context->SourceComponent.IsValid())
diff --git a/Engine/Plugins/PCG/Source/PCG/Private/PCGPin.cpp b/Engine/Plugins/PCG/Source/PCG/Private/PCGPin.cpp
--- a/Engine/Plugins/PCG/Source/PCG/Private/PCGPin.cpp
+++ b/Engine/Plugins/PCG/Source/PCG/Private/PCGPin.cpp
@@ -94,6 +94,19 @@ bool FPCGPinProperties::operator==(const FPCGPinProperties& Other) const
 		bInvisiblePin == Other.bInvisiblePin;
 }
 
+uint32 GetTypeHash(const FPCGPinProperties& Value)
+{
+	uint32 Hash = GetTypeHash(Value.Label);
+	Hash = HashCombine(Hash, static_cast<uint32>(Value.AllowedTypes));
+	Hash = HashCombine(Hash, Value.bAllowMultipleConnections);
+	Hash = HashCombine(Hash, Value.bAllowMultipleData);
+	Hash = HashCombine(Hash, static_cast<uint32>(Value.Usage));
+	Hash = HashCombine(Hash, static_cast<uint32>(Value.PinStatus));
+	Hash = HashCombine(Hash, Value.bInvisiblePin);
+
+	return Hash;
+}
+
 void FPCGPinProperties::PostSerialize(const FArchive& Ar)
 {
 #if WITH_EDITOR
diff --git a/Engine/Plugins/PCG/Source/PCG/Public/PCGData.h b/Engine/Plugins/PCG/Source/PCG/Public/PCGData.h
--- a/Engine/Plugins/PCG/Source/PCG/Public/PCGData.h
+++ b/Engine/Plugins/PCG/Source/PCG/Public/PCGData.h
@@ -103,6 +103,18 @@ struct PCG_API FPCGTaggedData
 	// Special flag for data that are forwarded to other nodes, but without a pin. Useful for internal data.
 	UPROPERTY()
 	bool bPinlessData = false;
+	
+	/**
+	* Special flag to be modified by execution when a data is used multiple times (in this node or other nodes),
+	* to enable optimization when they are not. Always assume that it is true by default.
+	*/
+	UPROPERTY(BlueprintReadOnly, Transient, Category = Data)
+	bool bIsUsedMultipleTimes = true;
+
+#if !UE_BUILD_SHIPPING
+	// For debug tracking.
+	TWeakObjectPtr<const UPCGNode> OriginatingNode;
+#endif // !UE_BUILD_SHIPPING
 
 	// Used to re-order elements when using per-data caching
 	int32 OriginalIndex = INDEX_NONE;
diff --git a/Engine/Plugins/PCG/Source/PCG/Public/PCGElement.h b/Engine/Plugins/PCG/Source/PCG/Public/PCGElement.h
--- a/Engine/Plugins/PCG/Source/PCG/Public/PCGElement.h
+++ b/Engine/Plugins/PCG/Source/PCG/Public/PCGElement.h
@@ -109,6 +109,9 @@ public:
 
 	/** Returns true if the node can be cached (e.g. does not create artifacts & does not depend on untracked data */
 	virtual bool IsCacheable(const UPCGSettings* InSettings) const { return true; }
+	
+	/** Returns true if the node outputs requires to update the output tagged data to detect if the data is used multiple times. */
+	virtual bool ShouldVerifyIfOutputsAreUsedMultipleTimes(const UPCGSettings* InSettings) const { return false; }
 
 	/** Whether to do a 'deep' fine-grained CRC of the output data to pass to downstream nodes. Can be expensive so should be used sparingly. */
 	virtual bool ShouldComputeFullOutputDataCrc(FPCGContext* Context) const { return false; }
diff --git a/Engine/Plugins/PCG/Source/PCG/Public/PCGPin.h b/Engine/Plugins/PCG/Source/PCG/Public/PCGPin.h
--- a/Engine/Plugins/PCG/Source/PCG/Public/PCGPin.h
+++ b/Engine/Plugins/PCG/Source/PCG/Public/PCGPin.h
@@ -91,6 +91,8 @@ struct PCG_API FPCGPinProperties
 	// Convert the bIsAdvanced boolean to PinStatus for deprecation purposes.
 	void PostSerialize(const FArchive& Ar);
 
+	friend uint32 GetTypeHash(const FPCGPinProperties& Value);
+
 #if WITH_EDITOR
 	bool CanEditChange(const FEditPropertyChain& PropertyChain) const;
 #endif
