[PCG] First version of GPU path for spawning instanced SMs.

Uses the new scene compute update interface to spawn instances direct on GPU without CPU involvement.

Current limitations:
* Only Nanite meshes are supported, non-Nanite will follow later.
* Meshes are required to be hardcoded onto SM Spawner node, spawn-by-attribute will follow later.
* Only floating point types (float, vector, etc) can be packed into instance data, integral types will follow later.
* Only instance packers that can declare upfront the list of attribute names are supported, regex and others will follow later.

Current issues:
* Some events flush the GPU Scene which can clear out the generated instances, we need hooks to refresh in these cases.

#rnx

diff --git a/Engine/Plugins/PCG/Shaders/Private/Elements/PCGCopyPoints.usf b/Engine/Plugins/PCG/Shaders/Private/Elements/PCGCopyPoints.usf
--- a/Engine/Plugins/PCG/Shaders/Private/Elements/PCGCopyPoints.usf
+++ b/Engine/Plugins/PCG/Shaders/Private/Elements/PCGCopyPoints.usf
@@ -109,9 +109,8 @@ void CopyAttributesFromTarget(uint TargetDataIndex, uint TargetElemIndex, uint O
 [numthreads(64, 1, 1)]
 void Main(uint3 GroupId : SV_GroupID, uint GroupIndex : SV_GroupIndex)
 {
-	// Mark the kernel as having executed. Must run before we early via thread index, because the kernel is still 'executed' even if the number
-	// of threads to iterate on is zero. Even if GetNumThreads() returns 0, the kernel will still have been dispatched on a single thread to set
-	// this flag.
+	// Mark the kernel as having executed. Must run before we early out via thread index, because the kernel is still 'executed' even if the number of
+	// threads to iterate on is zero. Even if GetNumThreads() returns 0, the kernel will still have been dispatched on a single thread to set this flag.
 	if (all(GroupId == 0) && GroupIndex == 0)
 	{
 		Out_SetAsExecutedInternal();
diff --git a/Engine/Plugins/PCG/Shaders/Private/Elements/PCGStaticMeshSpawner.usf b/Engine/Plugins/PCG/Shaders/Private/Elements/PCGStaticMeshSpawner.usf
new file mode 100644
--- /dev/null
+++ b/Engine/Plugins/PCG/Shaders/Private/Elements/PCGStaticMeshSpawner.usf
@@ -0,0 +1,117 @@
+// Copyright Epic Games, Inc. All Rights Reserved.
+
+// Static mesh spawner kernel. Selects a primitive using a primitive cumulative distribution function (CDF) and packs attributes in custom floats.
+// Passes through points to output.
+
+#include "/Engine/Private/ComputeShaderUtils.ush"
+#include "/Engine/Private/Quaternion.ush"
+#include "/Plugin/PCG/Private/PCGShaderUtils.ush"
+
+[numthreads(64, 1, 1)]
+void Main(uint3 GroupId : SV_GroupID, uint GroupIndex : SV_GroupIndex)
+{
+	// Mark the kernel as having executed. Must run before we early out via thread index, because the kernel is still 'executed' even if the number of
+	// threads to iterate on is zero. Even if GetNumThreads() returns 0, the kernel will still have been dispatched on a single thread to set this flag.
+	if (all(GroupId == 0) && GroupIndex == 0)
+	{
+		Out_SetAsExecutedInternal();
+	}
+
+	const uint ThreadIndex = GetUnWrappedDispatchThreadId(GroupId, GroupIndex, 64);
+	if (ThreadIndex >= InstanceData_GetNumInstancesAllPrimitives())
+	{
+		return;
+	}
+
+	// Break out thread index into DataIndex and ElementIndex
+	uint DataIndex, ElementIndex;
+	if (!In_GetThreadData(ThreadIndex, DataIndex, ElementIndex))
+	{
+		return;
+	}
+
+	const int NumPrimitives = InstanceData_GetNumPrimitives();
+
+	uint PrimitiveIndex = 0;
+
+	if (NumPrimitives > 1)
+	{
+		// Randomly select a primitive by inverting the cumulative distribution function.
+		uint Seed = In_GetSeed(DataIndex, ElementIndex);
+		const float RandomDraw = FRand(Seed);
+
+		for (; PrimitiveIndex < uint(NumPrimitives) - 1; ++PrimitiveIndex)
+		{
+			if (RandomDraw < SMSpawner_GetPrimitiveSelectionCDF(PrimitiveIndex))
+			{
+				break;
+			}
+		}
+	}
+
+	const int InstanceIndexInt = InstanceData_GetIndexToWriteTo(PrimitiveIndex);
+	if (InstanceIndexInt < 0)
+	{
+		// Saturated instances for this primitive. No-op.
+		return;
+	}
+
+	uint Out_DataIndex, Out_ElementIndex;
+	if (!Out_GetThreadData(ThreadIndex, Out_DataIndex, Out_ElementIndex))
+	{
+		// Should never happen.
+		return;
+	}
+
+	const uint InstanceIndex = InstanceIndexInt;
+
+	if (In_IsPointRemoved(DataIndex, ElementIndex))
+	{
+		// Cull the instance by setting the transform to 0.
+		InstanceData_WriteInstanceLocalToWorld(InstanceIndex, (float4x4)0);
+		
+		// Propagate the 'Removed' status to the output point if either source or target point is removed. Otherwise, this point will not be culled.
+		Out_RemovePoint(Out_DataIndex, Out_ElementIndex);
+
+		return;
+	}
+
+	// Write instance transform.
+	const float4x4 Transform = In_GetPointTransform(DataIndex, ElementIndex);
+	InstanceData_WriteInstanceLocalToWorld(InstanceIndex, transpose(Transform));
+
+	// Write instance custom float data.
+	const uint DataAddress = In_GetDataAddress(DataIndex);
+	const uint AttributeHeadersAddress = In_GetDataAttributeHeadersAddress(DataAddress);
+	const uint NumAttributes = SMSpawner_GetNumAttributes();
+	for (uint AttributeIndex = 0; AttributeIndex < NumAttributes; ++AttributeIndex)
+	{
+		const uint4 AttributeIdOffsetStride = SMSpawner_GetAttributeIdOffsetStride(AttributeIndex);
+		const uint AttributeId = AttributeIdOffsetStride[0];
+		const uint AttributeOffset = AttributeIdOffsetStride[1];
+		const uint AttributeStride = AttributeIdOffsetStride[2];
+
+		const uint AttributeHeaderAddress = In_GetAttributeHeaderAddress(AttributeHeadersAddress, AttributeId);
+		const uint AttributeAddress = In_GetAttributeAddress(AttributeHeaderAddress);
+
+		for (uint FloatIndex = 0; FloatIndex < AttributeStride; ++FloatIndex)
+		{
+			const uint AttributeValueAddress = AttributeAddress + ElementIndex * AttributeStride * 4;
+			
+			// Assumes floating point data. TODO add type information and deal with different types in different ways?
+			InstanceData_WriteCustomFloatRaw(InstanceIndex, AttributeOffset + FloatIndex, In_LoadBufferInternal(AttributeValueAddress));
+		}
+	}
+
+	// Write output point data (pass through).
+	// TODO - support writing mesh bounds.
+	Out_SetPosition(Out_DataIndex, Out_ElementIndex, In_GetPosition(DataIndex, ElementIndex));
+	Out_SetRotation(Out_DataIndex, Out_ElementIndex, In_GetRotation(DataIndex, ElementIndex));
+	Out_SetScale(Out_DataIndex, Out_ElementIndex, In_GetScale(DataIndex, ElementIndex));
+	Out_SetBoundsMin(Out_DataIndex, Out_ElementIndex, In_GetBoundsMin(DataIndex, ElementIndex));
+	Out_SetBoundsMax(Out_DataIndex, Out_ElementIndex, In_GetBoundsMax(DataIndex, ElementIndex));
+	Out_SetColor(Out_DataIndex, Out_ElementIndex, In_GetColor(DataIndex, ElementIndex));
+	Out_SetSeed(Out_DataIndex, Out_ElementIndex, In_GetSeed(DataIndex, ElementIndex));
+	Out_SetDensity(Out_DataIndex, Out_ElementIndex, In_GetDensity(DataIndex, ElementIndex));
+	Out_SetSteepness(Out_DataIndex, Out_ElementIndex, In_GetSteepness(DataIndex, ElementIndex));
+}
diff --git a/Engine/Plugins/PCG/Shaders/Private/PCGInstanceDataInterface.ush b/Engine/Plugins/PCG/Shaders/Private/PCGInstanceDataInterface.ush
new file mode 100644
--- /dev/null
+++ b/Engine/Plugins/PCG/Shaders/Private/PCGInstanceDataInterface.ush
@@ -0,0 +1,60 @@
+// Copyright Epic Games, Inc. All Rights Reserved.
+
+#include "/Engine/Private/Common.ush"
+
+#define MAX_PRIMITIVES 64
+
+uint4 {DataInterfaceName}_PrimitiveParameters[MAX_PRIMITIVES]; // (NumInstancesAllocated, InstanceOffset, Unused, Unused)
+
+uint {DataInterfaceName}_NumPrimitives;
+uint {DataInterfaceName}_NumInstancesAllPrimitives;
+uint {DataInterfaceName}_NumCustomFloatsPerInstance;
+
+RWStructuredBuffer<uint4> {DataInterfaceName}_InstanceData;
+RWStructuredBuffer<uint> {DataInterfaceName}_InstanceCustomFloatData;
+RWStructuredBuffer<uint> {DataInterfaceName}_WriteCounters;
+
+uint InstanceData_GetNumPrimitives_{DataInterfaceName}()
+{
+	return {DataInterfaceName}_NumPrimitives;
+}
+
+uint InstanceData_GetNumInstancesAllPrimitives_{DataInterfaceName}()
+{
+	return {DataInterfaceName}_NumInstancesAllPrimitives;
+}
+
+int InstanceData_GetIndexToWriteTo_{DataInterfaceName}(uint InPrimitiveIndex)
+{
+	uint ValueBefore;
+	InterlockedAdd({DataInterfaceName}_WriteCounters[InPrimitiveIndex], 1u, ValueBefore);
+
+	const uint NumInstancesAllocated = {DataInterfaceName}_PrimitiveParameters[InPrimitiveIndex][0];
+	if (ValueBefore >= NumInstancesAllocated)
+	{
+		// All instances for this primitive have been written to.
+		return -1;
+	}
+
+	const uint InstanceOffset = {DataInterfaceName}_PrimitiveParameters[InPrimitiveIndex][1];
+	return (int)(ValueBefore + InstanceOffset);
+}
+
+void InstanceData_WriteInstanceLocalToWorld_{DataInterfaceName}(uint InInstanceIndex, float4x4 InLocalToRelativeWorld)
+{
+	InLocalToRelativeWorld = transpose(InLocalToRelativeWorld);
+	const uint StrideUint4s = 3;
+	{DataInterfaceName}_InstanceData[InInstanceIndex * StrideUint4s + 0] = asuint(InLocalToRelativeWorld[0]);
+	{DataInterfaceName}_InstanceData[InInstanceIndex * StrideUint4s + 1] = asuint(InLocalToRelativeWorld[1]);
+	{DataInterfaceName}_InstanceData[InInstanceIndex * StrideUint4s + 2] = asuint(InLocalToRelativeWorld[2]);
+}
+
+void InstanceData_WriteCustomFloat_{DataInterfaceName}(uint InInstanceIndex, uint InCustomFloatIndex, float InValue)
+{
+	{DataInterfaceName}_InstanceCustomFloatData[{DataInterfaceName}_NumCustomFloatsPerInstance * InInstanceIndex + InCustomFloatIndex] = asuint(InValue);
+}
+
+void InstanceData_WriteCustomFloatRaw_{DataInterfaceName}(uint InInstanceIndex, uint InCustomFloatIndex, uint InValueAsUint)
+{
+	{DataInterfaceName}_InstanceCustomFloatData[{DataInterfaceName}_NumCustomFloatsPerInstance * InInstanceIndex + InCustomFloatIndex] = InValueAsUint;
+}
diff --git a/Engine/Plugins/PCG/Shaders/Private/PCGSceneWriter.usf b/Engine/Plugins/PCG/Shaders/Private/PCGSceneWriter.usf
new file mode 100644
--- /dev/null
+++ b/Engine/Plugins/PCG/Shaders/Private/PCGSceneWriter.usf
@@ -0,0 +1,81 @@
+// Copyright Epic Games, Inc. All Rights Reserved.
+
+// Scene writer - used to inject final instance data into the GPU Scene. Dispatched in the scene update timeline.
+
+#include "/Engine/Private/Common.ush"
+#include "/Engine/Private/GPUScene/GPUSceneWriter.ush"
+#include "/Engine/Public/Platform.ush"
+
+uint InPrimitiveIndex;
+uint InNumInstancesAllocatedInGPUScene;
+uint InInstanceOffset;
+StructuredBuffer<uint4> InInstanceData;
+StructuredBuffer<uint> InInstanceCustomFloatData;
+StructuredBuffer<uint> InWriteCounters;
+uint InPrimitiveId;
+uint InCustomDataCount;
+uint InPayloadDataFlags;
+
+void CullInstance(uint InstanceId)
+{
+	// This will minimally initialize the instance data such that it will be ignored and culled
+	WriteInstancePrimitiveIdAndFlags(InstanceId, InPrimitiveId, INSTANCE_SCENE_DATA_FLAG_HIDDEN);
+}
+
+[numthreads(NUM_THREADS_PER_GROUP, 1, 1)]
+void SceneWriter_CS(const uint ThreadId : SV_DispatchThreadID)
+{
+	const uint LocalInstanceId = ThreadId;
+	if (LocalInstanceId >= InNumInstancesAllocatedInGPUScene)
+	{
+		return;
+	}
+
+	// Retrieve the primitive data
+	FPrimitiveSceneData PrimitiveData = GetPrimitiveData(InPrimitiveId);
+	
+	const uint InstanceId = PrimitiveData.InstanceSceneDataOffset + LocalInstanceId;
+	
+	// Cull any instances at tail end (assume all valid instances are written).
+	if (LocalInstanceId >= InWriteCounters[InPrimitiveIndex])
+	{
+		CullInstance(InstanceId);
+		return;
+	}
+	
+	const uint StrideFloat4s = 3;
+	float4x4 LocalToWorld;
+	LocalToWorld[0] = asfloat(InInstanceData[(InInstanceOffset + LocalInstanceId) * StrideFloat4s + 0]);
+	LocalToWorld[1] = asfloat(InInstanceData[(InInstanceOffset + LocalInstanceId) * StrideFloat4s + 1]);
+	LocalToWorld[2] = asfloat(InInstanceData[(InInstanceOffset + LocalInstanceId) * StrideFloat4s + 2]);
+	LocalToWorld[3] = float4(0.0f, 0.0f, 0.0f, 1.0f);
+	LocalToWorld = transpose(LocalToWorld);
+
+	InitializeInstanceSceneDataWS(
+		InstanceId,
+		InPrimitiveId,
+		LocalInstanceId,
+ 		InPayloadDataFlags,
+ 		InCustomDataCount,
+ 		/*RandomId=*/0.0f,
+		DFFromTileOffset(LWCPromote(LocalToWorld))
+	);
+
+	const FDFMatrix Transform = MakeDFMatrix4x3(PrimitiveData.PositionHigh*0.0, LocalToWorld);
+	WriteInstanceLocalToWorld(InstanceId, InPrimitiveId, Transform, false);
+
+	FInstancePayloadDataOffsets PayloadOffsets = GetInstancePayloadDataOffsets(InPrimitiveId, InPayloadDataFlags, LocalInstanceId);
+
+	const uint NumCustomFloat4s = (InCustomDataCount + 3) / 4;
+	for (uint Float4Index = 0; Float4Index < NumCustomFloat4s; ++Float4Index)
+	{
+		const uint BaseIndex = (InInstanceOffset + LocalInstanceId) * InCustomDataCount + Float4Index * 4;
+
+		uint4 CustomData;
+		CustomData[0] = InInstanceCustomFloatData[BaseIndex];
+		CustomData[1] = InInstanceCustomFloatData[BaseIndex + 1];
+		CustomData[2] = InInstanceCustomFloatData[BaseIndex + 2];
+		CustomData[3] = InInstanceCustomFloatData[BaseIndex + 3];
+		WriteInstanceCustomData(PayloadOffsets, Float4Index, asfloat(CustomData));
+	}
+}
diff --git a/Engine/Plugins/PCG/Source/PCG/PCG.Build.cs b/Engine/Plugins/PCG/Source/PCG/PCG.Build.cs
--- a/Engine/Plugins/PCG/Source/PCG/PCG.Build.cs
+++ b/Engine/Plugins/PCG/Source/PCG/PCG.Build.cs
@@ -29,6 +29,7 @@ namespace UnrealBuildTool.Rules
 				{
 					"ComputeFramework",
 					"PCGCompute",
+					"Renderer",
 					"Voronoi",
 				});
 
diff --git a/Engine/Plugins/PCG/Source/PCG/Private/Components/PCGProceduralISMComponent.cpp b/Engine/Plugins/PCG/Source/PCG/Private/Components/PCGProceduralISMComponent.cpp
new file mode 100644
--- /dev/null
+++ b/Engine/Plugins/PCG/Source/PCG/Private/Components/PCGProceduralISMComponent.cpp
@@ -0,0 +1,529 @@
+// Copyright Epic Games, Inc. All Rights Reserved.
+
+#include "Components/PCGProceduralISMComponent.h"
+
+#include "PCGComponent.h"
+#include "PCGModule.h"
+#include "PCGManagedResource.h"
+#include "Helpers/PCGHelpers.h"
+#include "MeshSelectors/PCGISMDescriptor.h"
+
+#include "InstanceDataSceneProxy.h"
+#include "NaniteSceneProxy.h"
+#include "SceneInterface.h"
+#include "Engine/StaticMesh.h"
+#include "Serialization/ArchiveCrc32.h"
+#include "VT/RuntimeVirtualTexture.h"
+
+UPCGManagedProceduralISMComponent* PCGManagedProceduralISMComponent::GetOrCreateManagedProceduralISMC(AActor* InTargetActor, UPCGComponent* InSourceComponent, uint64 InSettingsUID, const FPCGProceduralISMCBuilderParameters& InParams)
+{
+	check(InTargetActor && InSourceComponent);
+
+	const UStaticMesh* StaticMesh = InParams.Descriptor.StaticMesh;
+	if (!StaticMesh)
+	{
+		return nullptr;
+	}
+
+	TRACE_CPUPROFILER_EVENT_SCOPE(UPCGActorHelpers::GetOrCreateManagedISMC);
+
+	FPCGProceduralISMComponentDescriptor Descriptor = InParams.Descriptor;
+
+	{
+		TRACE_CPUPROFILER_EVENT_SCOPE(UPCGActorHelpers::GetOrCreateManagedISMC::FindMatchingMISMC);
+		UPCGManagedProceduralISMComponent* MatchingResource = nullptr;
+		InSourceComponent->ForEachManagedResource([&MatchingResource, &InParams, &InTargetActor, &Descriptor, InSettingsUID](UPCGManagedResource* InResource)
+		{
+			// Early out if already found a match
+			if (MatchingResource)
+			{
+				return;
+			}
+
+			if (UPCGManagedProceduralISMComponent* Resource = Cast<UPCGManagedProceduralISMComponent>(InResource))
+			{
+				if (Resource->GetSettingsUID() != InSettingsUID || !Resource->CanBeUsed())
+				{
+					return;
+				}
+
+				if (UPCGProceduralISMComponent* ISMC = Resource->GetComponent())
+				{
+					if (IsValid(ISMC) &&
+						ISMC->GetOwner() == InTargetActor &&
+						Resource->GetDescriptor() == Descriptor)
+					{
+						MatchingResource = Resource;
+					}
+				}
+			}
+		});
+
+		if (MatchingResource)
+		{
+			TRACE_CPUPROFILER_EVENT_SCOPE(UPCGActorHelpers::GetOrCreateManagedISMC::MarkAsUsed);
+			MatchingResource->MarkAsUsed();
+
+			if (UPCGProceduralISMComponent* ISMC = MatchingResource->GetComponent())
+			{
+				ISMC->ComponentTags = Descriptor.ComponentTags;
+				ISMC->ComponentTags.AddUnique(PCGHelpers::DefaultPCGTag);
+				ISMC->ComponentTags.AddUnique(InSourceComponent->GetFName());
+			}
+
+			return MatchingResource;
+		}
+	}
+
+	// No matching ISM component found, let's create a new one
+	FString ComponentName = TEXT("PISM_") + StaticMesh->GetName();
+
+	const EObjectFlags ObjectFlags = RF_Transient;
+	UPCGProceduralISMComponent* ISMC = NewObject<UPCGProceduralISMComponent>(InTargetActor, UPCGProceduralISMComponent::StaticClass(), MakeUniqueObjectName(InTargetActor, UPCGProceduralISMComponent::StaticClass(), FName(ComponentName)), ObjectFlags);
+	Descriptor.InitComponent(ISMC);
+
+	ISMC->RegisterComponent();
+	InTargetActor->AddInstanceComponent(ISMC);
+
+	ISMC->AttachToComponent(InTargetActor->GetRootComponent(), FAttachmentTransformRules(EAttachmentRule::KeepRelative, EAttachmentRule::KeepWorld, EAttachmentRule::KeepWorld, false));
+
+	// Create managed resource on source component
+	UPCGManagedProceduralISMComponent* Resource = NewObject<UPCGManagedProceduralISMComponent>(InSourceComponent);
+	Resource->SetComponent(ISMC);
+	Resource->SetDescriptor(Descriptor);
+	if (InTargetActor->GetRootComponent())
+	{
+		Resource->SetRootLocation(InTargetActor->GetRootComponent()->GetComponentLocation());
+	}
+
+	Resource->SetSettingsUID(InSettingsUID);
+	InSourceComponent->AddToManagedResources(Resource);
+
+	if (ISMC)
+	{
+		ISMC->ComponentTags.AddUnique(PCGHelpers::DefaultPCGTag);
+		ISMC->ComponentTags.AddUnique(InSourceComponent->GetFName());
+	}
+
+	return Resource;
+}
+
+FPCGProceduralISMComponentDescriptor& FPCGProceduralISMComponentDescriptor::operator=(const FPCGSoftISMComponentDescriptor& Other)
+{
+	InstanceStartCullDistance = Other.InstanceStartCullDistance;
+	InstanceEndCullDistance = Other.InstanceEndCullDistance;
+	OverlayMaterial = Other.OverlayMaterial.LoadSynchronous();
+	StaticMesh = Other.StaticMesh.LoadSynchronous();
+	ComponentTags = Other.ComponentTags;
+
+	Algo::Transform(Other.OverrideMaterials, OverrideMaterials, [](const TSoftObjectPtr<UMaterialInterface>& OverrideMaterial)
+	{
+		return OverrideMaterial.LoadSynchronous();
+	});
+
+	Algo::Transform(Other.RuntimeVirtualTextures, RuntimeVirtualTextures, [](const TSoftObjectPtr<URuntimeVirtualTexture>& RVT)
+	{
+		return RVT.LoadSynchronous();
+	});
+
+	return *this;
+}
+
+void FPCGProceduralISMComponentDescriptor::InitComponent(UPCGProceduralISMComponent* InProceduralISMComponent) const
+{
+	InProceduralISMComponent->SetStaticMesh(StaticMesh);
+	InProceduralISMComponent->OverrideMaterials = OverrideMaterials;
+	InProceduralISMComponent->OverlayMaterial = OverlayMaterial;
+	InProceduralISMComponent->RuntimeVirtualTextures = RuntimeVirtualTextures;
+	InProceduralISMComponent->SetNumInstances(NumInstances);
+	InProceduralISMComponent->SetNumCustomDataFloats(NumCustomFloats);
+	InProceduralISMComponent->SetLocalBounds(LocalBounds);
+	InProceduralISMComponent->SetCullDistances(InstanceStartCullDistance, InstanceEndCullDistance);
+	InProceduralISMComponent->ComponentTags = ComponentTags;
+}
+
+void FPCGProceduralISMComponentDescriptor::InitFrom(UPCGProceduralISMComponent* InProceduralISMComponent)
+{
+	StaticMesh = InProceduralISMComponent->GetStaticMesh();
+	OverrideMaterials = InProceduralISMComponent->OverrideMaterials;
+	OverlayMaterial = InProceduralISMComponent->OverlayMaterial;
+	RuntimeVirtualTextures = InProceduralISMComponent->RuntimeVirtualTextures;
+	NumInstances = InProceduralISMComponent->GetNumInstances();
+	NumCustomFloats = InProceduralISMComponent->GetNumCustomDataFloats();
+	InProceduralISMComponent->GetLocalBounds(LocalBounds);
+	InProceduralISMComponent->GetCullDistances(InstanceStartCullDistance, InstanceEndCullDistance);
+}
+
+uint32 FPCGProceduralISMComponentDescriptor::ComputeHash() const
+{
+	FArchiveCrc32 CrcArchive;
+
+	Hash = 0; // we don't want the hash to impact the calculation
+	CrcArchive << *this;
+	Hash = CrcArchive.GetCrc();
+
+	return Hash;
+}
+
+bool FPCGProceduralISMComponentDescriptor::operator!=(const FPCGProceduralISMComponentDescriptor& Other) const
+{
+	return !(*this == Other);
+}
+
+bool FPCGProceduralISMComponentDescriptor::operator==(const FPCGProceduralISMComponentDescriptor& Other) const
+{
+	return StaticMesh == Other.StaticMesh
+		&& OverrideMaterials == Other.OverrideMaterials
+		&& OverlayMaterial == Other.OverlayMaterial
+		&& RuntimeVirtualTextures == Other.RuntimeVirtualTextures
+		&& NumInstances == Other.NumInstances
+		&& NumCustomFloats == Other.NumCustomFloats
+		&& LocalBounds == Other.LocalBounds
+		&& InstanceStartCullDistance == Other.InstanceStartCullDistance
+		&& InstanceEndCullDistance == Other.InstanceEndCullDistance;
+}
+
+UPCGProceduralISMComponent::UPCGProceduralISMComponent(const FObjectInitializer& ObjectInitializer)
+	: Super(ObjectInitializer)
+{
+	// Disable unsupported rendering features (currently require instance data on CPU).
+	bAffectDynamicIndirectLighting = false;
+	bAffectDistanceFieldLighting = false;
+	bVisibleInRayTracing = false;
+
+	Mobility = EComponentMobility::Static;
+	BodyInstance.bSimulatePhysics = false;
+	SetGenerateOverlapEvents(false);
+
+	bEnableVertexColorMeshPainting = false;
+
+	bNavigationRelevant = false;
+	bCanEverAffectNavigation = false;
+	
+	bEnableAutoLODGeneration = false;
+
+	bIsEditorOnly = true;
+
+#if STATS
+	{
+		UObject const* StatObject = this->AdditionalStatObject();
+		if (!StatObject)
+		{
+			StatObject = this;
+		}
+		StatId = StatObject->GetStatID(true);
+	}
+#endif
+}
+
+FPrimitiveSceneProxy* UPCGProceduralISMComponent::CreateStaticMeshSceneProxy(Nanite::FMaterialAudit& NaniteMaterials, bool bCreateNanite)
+{
+	LLM_SCOPE(ELLMTag::InstancedMesh);
+
+	if (CheckPSOPrecachingAndBoostPriority() && GetPSOPrecacheProxyCreationStrategy() == EPSOPrecacheProxyCreationStrategy::DelayUntilPSOPrecached)
+	{
+		UE_LOG(LogPCG, Verbose, TEXT("Skipping CreateSceneProxy for PCGProceduralISMComponent %s (PCGProceduralISMComponent PSOs are still compiling)"), *GetFullName());
+		return nullptr;
+	}
+
+	FInstanceSceneDataBuffers InstanceSceneDataBuffers(/*InbInstanceDataIsGPUOnly=*/true);
+
+	FInstanceSceneDataBuffers::FAccessTag AccessTag(PointerHash(this));
+	FInstanceSceneDataBuffers::FWriteView ProxyData = InstanceSceneDataBuffers.BeginWriteAccess(AccessTag);
+
+	InstanceSceneDataBuffers.SetPrimitiveLocalToWorld(GetRenderMatrix(), AccessTag);
+
+	ProxyData.NumInstancesGPUOnly = GetNumInstances();
+	ProxyData.NumCustomDataFloats = GetNumCustomDataFloats();
+	ProxyData.InstanceLocalBounds.SetNum(1);
+	ProxyData.InstanceLocalBounds[0] = GetStaticMesh()->GetBounds();
+
+	ProxyData.Flags.bHasPerInstanceCustomData = ProxyData.NumCustomDataFloats > 0;
+
+	InstanceSceneDataBuffers.EndWriteAccess(AccessTag);
+	InstanceSceneDataBuffers.ValidateData();
+
+	auto InstanceDataSceneProxy = MakeShared<FInstanceDataSceneProxy, ESPMode::ThreadSafe>(MoveTemp(InstanceSceneDataBuffers));
+
+	if (bCreateNanite)
+	{
+		return ::new Nanite::FSceneProxy(NaniteMaterials, this, InstanceDataSceneProxy);
+	}
+	else
+	{
+		UE_LOG(LogPCG, Warning, TEXT("Only Nanite meshes are supported by PCGProceduralISMComponent currently, component will not render."));
+		return nullptr;
+	}
+}
+
+FPrimitiveSceneProxy* UPCGProceduralISMComponent::CreateSceneProxy()
+{
+	// Verify that both mesh and instance data is valid before using it.
+	const bool bIsMeshAndInstanceDataValid =
+		NumInstances > 0 &&
+		GetStaticMesh() &&
+		GetStaticMesh()->IsCompiling() == false &&
+		GetStaticMesh()->HasValidRenderData();
+
+	if (!bIsMeshAndInstanceDataValid)
+	{
+		return nullptr;
+	}
+
+	ValidateComponentSetup();
+
+	return Super::CreateSceneProxy();
+}
+
+FMatrix UPCGProceduralISMComponent::GetRenderMatrix() const
+{
+	// Apply the translated space to the render matrix.
+	return GetComponentTransform().ToMatrixWithScale();
+}
+
+FBoxSphereBounds UPCGProceduralISMComponent::CalcBounds(const FTransform& BoundTransform) const
+{
+	return LocalBounds.TransformBy(BoundTransform);
+}
+
+#if WITH_EDITOR
+FBox UPCGProceduralISMComponent::GetStreamingBounds() const
+{
+	return FBox::BuildAABB(Bounds.Origin, Bounds.BoxExtent);
+}
+#endif
+
+bool UPCGProceduralISMComponent::GetMaterialStreamingData(int32 MaterialIndex, FPrimitiveMaterialInfo& MaterialData) const
+{
+	// Same thing as StaticMesh but we take the full bounds to cover the instances.
+	if (GetStaticMesh())
+	{
+		MaterialData.Material = GetMaterial(MaterialIndex);
+		MaterialData.UVChannelData = GetStaticMesh()->GetUVChannelData(MaterialIndex);
+		MaterialData.PackedRelativeBox = PackedRelativeBox_Identity;
+	}
+	return MaterialData.IsValid();
+}
+
+bool UPCGProceduralISMComponent::BuildTextureStreamingDataImpl(ETextureStreamingBuildType BuildType, EMaterialQualityLevel::Type QualityLevel, ERHIFeatureLevel::Type FeatureLevel, TSet<FGuid>& DependentResources, bool& bOutSupportsBuildTextureStreamingData)
+{
+#if WITH_EDITORONLY_DATA // Only rebuild the data in editor 
+	if (GetNumInstances() > 0)
+	{
+		return Super::BuildTextureStreamingDataImpl(BuildType, QualityLevel, FeatureLevel, DependentResources, bOutSupportsBuildTextureStreamingData);
+	}
+#endif
+	return true;
+}
+
+void UPCGProceduralISMComponent::GetStreamingRenderAssetInfo(FStreamingTextureLevelContext& LevelContext, TArray<FStreamingRenderAssetPrimitiveInfo>& OutStreamingRenderAssets) const
+{
+	// Don't only look the instance count but also if the bound is valid, as derived classes might not set PerInstanceSMData.
+	if (GetNumInstances() > 0 || Bounds.SphereRadius > 0)
+	{
+		return Super::GetStreamingRenderAssetInfo(LevelContext, OutStreamingRenderAssets);
+	}
+}
+
+void UPCGProceduralISMComponent::SetCullDistances(int32 InStartCullDistance, int32 InEndCullDistance)
+{
+	if (InstanceStartCullDistance != InStartCullDistance || InstanceEndCullDistance != InEndCullDistance)
+	{
+		InstanceStartCullDistance = InStartCullDistance;
+		InstanceEndCullDistance = InEndCullDistance;
+
+		if (GetScene() && SceneProxy)
+		{
+			GetScene()->UpdateInstanceCullDistance(this, InStartCullDistance, InEndCullDistance);
+		}
+	}
+}
+
+void UPCGProceduralISMComponent::ValidateComponentSetup()
+{
+	if (bAffectDynamicIndirectLighting)
+	{
+		UE_LOG(LogPCG, Warning, TEXT("Affecting indirect lighting is not currently supported by PCGProceduralISMComponent, disabling."));
+		bAffectDynamicIndirectLighting = false;
+	}
+
+	if (bAffectDistanceFieldLighting)
+	{
+		UE_LOG(LogPCG, Warning, TEXT("Affecting distance field lighting is not currently supported by PCGProceduralISMComponent, disabling."));
+		bAffectDistanceFieldLighting = false;
+	}
+
+	if (bVisibleInRayTracing)
+	{
+		UE_LOG(LogPCG, Warning, TEXT("Ray tracing is not currently supported by PCGProceduralISMComponent, disabling."));
+		bVisibleInRayTracing = false;
+	}
+}
+
+void UPCGProceduralISMComponent::SetNumInstances(int32 InNumInstances)
+{
+	if (NumInstances != InNumInstances)
+	{
+		NumInstances = InNumInstances;
+		MarkRenderStateDirty();
+	}
+}
+
+void UPCGProceduralISMComponent::SetNumCustomDataFloats(int32 InNumCustomDataFloats)
+{
+	if (FMath::Max(InNumCustomDataFloats, 0) != NumCustomDataFloats)
+	{
+		NumCustomDataFloats = FMath::Max(InNumCustomDataFloats, 0);
+		MarkRenderStateDirty();
+	}
+}
+
+void UPCGProceduralISMComponent::SetLocalBounds(const FBox& InLocalBounds)
+{
+	if (InLocalBounds != LocalBounds)
+	{
+		LocalBounds = InLocalBounds;
+		MarkRenderStateDirty();
+	}
+}
+
+#if WITH_EDITOR
+void UPCGProceduralISMComponent::PostEditUndo()
+{
+	Super::PostEditUndo();
+
+	MarkRenderStateDirty();
+}
+#endif
+
+void UPCGManagedProceduralISMComponent::PostLoad()
+{
+	Super::PostLoad();
+
+	// Cache raw ptr
+	GetComponent();
+}
+
+void UPCGManagedProceduralISMComponent::ForgetComponent()
+{
+	Super::ForgetComponent();
+	CachedRawComponentPtr = nullptr;
+}
+
+void UPCGManagedProceduralISMComponent::SetDescriptor(const FPCGProceduralISMComponentDescriptor& InDescriptor)
+{
+	Descriptor = InDescriptor;
+}
+
+bool UPCGManagedProceduralISMComponent::ReleaseIfUnused(TSet<TSoftObjectPtr<AActor>>& OutActorsToDelete)
+{
+	if (Super::ReleaseIfUnused(OutActorsToDelete) || !GetComponent())
+	{
+		return true;
+	}
+	else if (GetComponent()->GetNumInstances() == 0)
+	{
+		GeneratedComponent->DestroyComponent();
+		ForgetComponent();
+		return true;
+	}
+	else
+	{
+		return false;
+	}
+}
+
+void UPCGManagedProceduralISMComponent::ResetComponent()
+{
+	if (UPCGProceduralISMComponent* ISMC = GetComponent())
+	{
+		ISMC->UpdateBounds();
+	}
+}
+
+void UPCGManagedProceduralISMComponent::MarkAsUsed()
+{
+	const bool bWasMarkedUnused = bIsMarkedUnused;
+	Super::MarkAsUsed();
+
+	if (!bWasMarkedUnused)
+	{
+		return;
+	}
+
+	if (UPCGProceduralISMComponent* ISMC = GetComponent())
+	{
+		// Keep track of the current root location so if we reuse this later we are able to update this appropriately
+		if (USceneComponent* RootComponent = ISMC->GetAttachmentRoot())
+		{
+			bHasRootLocation = true;
+			RootLocation = RootComponent->GetComponentLocation();
+		}
+		else
+		{
+			bHasRootLocation = false;
+			RootLocation = FVector::ZeroVector;
+		}
+
+		// Reset the rotation/scale to be identity otherwise if the root component transform has changed, the final transform will be wrong
+		// Since this is technically 'moving' the ISM, we need to unregister it before moving otherwise we could get a warning that we're moving a component with static mobility
+		ISMC->UnregisterComponent();
+		ISMC->SetWorldTransform(FTransform(FQuat::Identity, RootLocation, FVector::OneVector));
+		ISMC->RegisterComponent();
+	}
+}
+
+void UPCGManagedProceduralISMComponent::MarkAsReused()
+{
+	Super::MarkAsReused();
+
+	if (UPCGProceduralISMComponent* ISMC = GetComponent())
+	{
+		// Reset the rotation/scale to be identity otherwise if the root component transform has changed, the final transform will be wrong
+		FVector TentativeRootLocation = RootLocation;
+
+		if (!bHasRootLocation)
+		{
+			if (USceneComponent* RootComponent = ISMC->GetAttachmentRoot())
+			{
+				TentativeRootLocation = RootComponent->GetComponentLocation();
+			}
+		}
+
+		// Since this is technically 'moving' the ISM, we need to unregister it before moving otherwise we could get a warning that we're moving a component with static mobility
+		ISMC->UnregisterComponent();
+		ISMC->SetWorldTransform(FTransform(FQuat::Identity, TentativeRootLocation, FVector::OneVector));
+		ISMC->RegisterComponent();
+	}
+}
+
+void UPCGManagedProceduralISMComponent::SetRootLocation(const FVector& InRootLocation)
+{
+	bHasRootLocation = true;
+	RootLocation = InRootLocation;
+}
+
+UPCGProceduralISMComponent* UPCGManagedProceduralISMComponent::GetComponent() const
+{
+	if (!CachedRawComponentPtr)
+	{
+		UPCGProceduralISMComponent* GeneratedComponentPtr = Cast<UPCGProceduralISMComponent>(GeneratedComponent.Get());
+
+		// Implementation note:
+		// There is no surefire way to make sure that we can use the raw pointer UNLESS it is from the same owner
+		if (GeneratedComponentPtr && Cast<UPCGComponent>(GetOuter()) && GeneratedComponentPtr->GetOwner() == Cast<UPCGComponent>(GetOuter())->GetOwner())
+		{
+			CachedRawComponentPtr = GeneratedComponentPtr;
+		}
+
+		return GeneratedComponentPtr;
+	}
+
+	return CachedRawComponentPtr;
+}
+
+void UPCGManagedProceduralISMComponent::SetComponent(UPCGProceduralISMComponent* InComponent)
+{
+	GeneratedComponent = InComponent;
+	CachedRawComponentPtr = InComponent;
+}
diff --git a/Engine/Plugins/PCG/Source/PCG/Private/Components/PCGProceduralISMComponent.h b/Engine/Plugins/PCG/Source/PCG/Private/Components/PCGProceduralISMComponent.h
new file mode 100644
--- /dev/null
+++ b/Engine/Plugins/PCG/Source/PCG/Private/Components/PCGProceduralISMComponent.h
@@ -0,0 +1,266 @@
+// Copyright Epic Games, Inc. All Rights Reserved.
+
+#pragma once
+
+#include "Components/StaticMeshComponent.h"
+#include "PCGManagedResource.h"
+
+#include "CoreMinimal.h"
+#include "EngineDefines.h"
+#include "InstanceDataSceneProxy.h"
+#include "StaticMeshResources.h"
+#include "StaticMeshSceneProxy.h"
+#include "Misc/Guid.h"
+#include "Stats/Stats.h"
+#include "UObject/ObjectMacros.h"
+
+#include "PCGProceduralISMComponent.generated.h"
+
+class FPrimitiveSceneProxy;
+class UPCGComponent;
+class UPCGManagedProceduralISMComponent;
+struct FPCGSoftISMComponentDescriptor;
+
+/** Experimental instanced primitive that is not backed by CPU instance data. Instance data must be written direct to the GPU scene. */
+UCLASS(ClassGroup = Rendering, meta = (BlueprintSpawnableComponent), Blueprintable, MinimalAPI, hideCategories = (Collision, Cooking, HLOD, MeshPainting, Navigation, Physics, RayTracing))
+class UPCGProceduralISMComponent : public UStaticMeshComponent
+{
+	GENERATED_UCLASS_BODY()
+
+public:
+	UFUNCTION(BlueprintCallable, Category = "Components|PCGProceduralISMC")
+	int32 GetNumInstances() const { return NumInstances; }
+
+	UFUNCTION(BlueprintCallable, Category = "Components|PCGProceduralISMC")
+	void SetNumInstances(int32 InNumInstances);
+
+	UFUNCTION(BlueprintCallable, Category = "Components|PCGProceduralISMC")
+	int32 GetNumCustomDataFloats() const { return NumCustomDataFloats; }
+
+	/** Update number of custom data entries per instance. */
+	UFUNCTION(BlueprintCallable, Category = "Components|PCGProceduralISMC")
+	void SetNumCustomDataFloats(int32 InNumCustomDataFloats);
+
+	/** Get local bounds that were previously set for all instances of this component. */
+	void GetLocalBounds(FBox& OutLocalBounds) { OutLocalBounds = LocalBounds; }
+
+	/** Set local bounds for all instances of this component (bounds are not computed automatically from instances). */
+	UFUNCTION(BlueprintCallable, Category = "Components|PCGProceduralISMC")
+	void SetLocalBounds(const FBox& InLocalBounds);
+
+	/** Clear all instances being rendered by this component. */
+	UFUNCTION(BlueprintCallable, Category="Components|PCGProceduralISMC")
+	void ClearInstances() { SetNumInstances(0); }
+
+	/** Sets the fading start and culling end distances for this component. */
+	UFUNCTION(BlueprintCallable, Category = "Components|PCGProceduralISMC")
+	void SetCullDistances(int32 InStartCullDistance, int32 InEndCullDistance);
+
+	/** Gets the fading start and culling end distances for this component. */
+	UFUNCTION(BlueprintCallable, Category = "Components|PCGProceduralISMC")
+	void GetCullDistances(int32& OutStartCullDistance, int32& OutEndCullDistance) const { OutStartCullDistance = InstanceStartCullDistance; OutEndCullDistance = InstanceEndCullDistance; }
+
+	//~ Begin UActorComponent Interface
+	bool ShouldCreatePhysicsState() const override { return false; }
+	bool IsHLODRelevant() const override { return false; }
+	//~ End UActorComponent Interface
+
+	//~ Begin UPrimitiveComponent Interface
+	virtual FPrimitiveSceneProxy* CreateSceneProxy() override;
+	virtual FMatrix GetRenderMatrix() const override;
+	virtual FBoxSphereBounds CalcBounds(const FTransform& BoundTransform) const override;
+	bool SupportsStaticLighting() const override { return false; }
+#if WITH_EDITOR
+	virtual FBox GetStreamingBounds() const override;
+#endif
+	virtual bool BuildTextureStreamingDataImpl(ETextureStreamingBuildType BuildType, EMaterialQualityLevel::Type QualityLevel, ERHIFeatureLevel::Type FeatureLevel, TSet<FGuid>& DependentResources, bool& bOutSupportsBuildTextureStreamingData) override;
+	virtual void GetStreamingRenderAssetInfo(FStreamingTextureLevelContext& LevelContext, TArray<FStreamingRenderAssetPrimitiveInfo>& OutStreamingRenderAssets) const override;
+	//~ End UPrimitiveComponent Interface
+
+	//~ Begin INavRelevantInterface Interface.
+	bool IsNavigationRelevant() const override { return false; }
+	//~ End INavRelevantInterface Interface.
+
+	//~ Begin UStaticMeshComponentInterface
+public:
+	virtual bool GetMaterialStreamingData(int32 MaterialIndex, FPrimitiveMaterialInfo& MaterialData) const override;
+
+protected:
+	virtual FPrimitiveSceneProxy* CreateStaticMeshSceneProxy(Nanite::FMaterialAudit& NaniteMaterials, bool bCreateNanite) override;
+	//~ End UStaticMeshComponentInterface
+
+	//~ Begin UObject Interface
+#if WITH_EDITOR
+protected:
+	virtual void PostEditUndo() override;
+#endif
+	//~ End UObject Interface
+
+protected:
+	/** Validate component setup, disable features that are not supported (typically those that require instance transforms on CPU). */
+	virtual void ValidateComponentSetup();
+
+protected:
+	UPROPERTY(EditAnywhere, Category = Instances)
+	int32 NumInstances = 0;
+
+	/** Defines the number of floats that will be available per instance for custom data */
+	UPROPERTY(EditAnywhere, Category = Instances)
+	int32 NumCustomDataFloats;
+
+	static inline const float DefaultBoundsExtent = 500.0f;
+
+	UPROPERTY(EditAnywhere, Category = Instances)
+	FBox LocalBounds = FBox(-FVector::OneVector * DefaultBoundsExtent, FVector::OneVector * DefaultBoundsExtent);
+
+	/** Distance from camera at which each instance begins to fade out. */
+	UPROPERTY(EditAnywhere, BlueprintReadOnly, Category = Culling)
+	int32 InstanceStartCullDistance;
+
+	/** Distance from camera at which each instance completely fades out. */
+	UPROPERTY(EditAnywhere, BlueprintReadOnly, Category = Culling)
+	int32 InstanceEndCullDistance;
+
+#if STATS
+	/** Used for dynamic stats */
+	TStatId StatId;
+#endif
+};
+
+/** Struct that holds properties that can be used to initialize Procedural ISM Components. */
+USTRUCT()
+struct FPCGProceduralISMComponentDescriptor
+{
+	GENERATED_BODY()
+
+public:
+	virtual ~FPCGProceduralISMComponentDescriptor() = default;
+
+	FPCGProceduralISMComponentDescriptor& operator=(const FPCGSoftISMComponentDescriptor& Other);
+
+	virtual void InitComponent(UPCGProceduralISMComponent* InProceduralISMComponent) const;
+
+	virtual void InitFrom(UPCGProceduralISMComponent* InProceduralISMComponent);
+
+	virtual uint32 ComputeHash() const;
+
+	bool operator!=(const FPCGProceduralISMComponentDescriptor& Other) const;
+	bool operator==(const FPCGProceduralISMComponentDescriptor& Other) const;
+
+	friend inline uint32 GetTypeHash(const FPCGProceduralISMComponentDescriptor& Key)
+	{
+		return Key.GetTypeHash();
+	}
+
+	uint32 GetTypeHash() const
+	{
+		if (Hash == 0)
+		{
+			ComputeHash();
+		}
+		return Hash;
+	}
+
+public:
+	UPROPERTY()
+	mutable uint32 Hash = 0;
+
+	UPROPERTY()
+	TObjectPtr<UStaticMesh> StaticMesh = nullptr;
+
+	UPROPERTY()
+	TArray<TObjectPtr<UMaterialInterface>> OverrideMaterials;
+
+	UPROPERTY()
+	TObjectPtr<UMaterialInterface> OverlayMaterial;
+
+	UPROPERTY()
+	TArray<TObjectPtr<URuntimeVirtualTexture>> RuntimeVirtualTextures;
+
+	UPROPERTY()
+	int32 NumInstances;
+
+	UPROPERTY()
+	int32 NumCustomFloats;
+
+	UPROPERTY()
+	FBox LocalBounds;
+
+	UPROPERTY()
+	int32 InstanceStartCullDistance;
+
+	UPROPERTY()
+	int32 InstanceEndCullDistance;
+
+	UPROPERTY()
+	TArray<FName> ComponentTags;
+};
+
+struct FPCGProceduralISMCBuilderParameters
+{
+	friend inline uint32 GetTypeHash(const FPCGProceduralISMCBuilderParameters& Key)
+	{
+		return GetTypeHash(Key.Descriptor);
+	}
+
+	inline bool operator==(const FPCGProceduralISMCBuilderParameters& Other) const { return Descriptor == Other.Descriptor && bAllowDescriptorChanges == Other.bAllowDescriptorChanges; }
+
+	FPCGProceduralISMComponentDescriptor Descriptor;
+
+	bool bAllowDescriptorChanges = true;
+};
+
+UCLASS(BlueprintType)
+class PCG_API UPCGManagedProceduralISMComponent : public UPCGManagedComponent
+{
+	GENERATED_BODY()
+
+public:
+	//~Begin UObject interface
+	virtual void PostLoad() override;
+	//~End UObject interface
+
+	//~Begin UPCGManagedResource interface
+	virtual bool ReleaseIfUnused(TSet<TSoftObjectPtr<AActor>>& OutActorsToDelete) override;
+	//~End UPCGManagedResource interface
+
+	//~Begin UPCGManagedComponents interface
+	virtual void ResetComponent() override;
+	virtual bool SupportsComponentReset() const override { return true; }
+	virtual void MarkAsUsed() override;
+	virtual void MarkAsReused() override;
+	virtual void ForgetComponent() override;
+	//~End UPCGManagedComponents interface
+
+	UPCGProceduralISMComponent* GetComponent() const;
+	void SetComponent(UPCGProceduralISMComponent* InComponent);
+
+	void SetDescriptor(const FPCGProceduralISMComponentDescriptor& InDescriptor);
+	const FPCGProceduralISMComponentDescriptor& GetDescriptor() const { return Descriptor; }
+
+	void SetRootLocation(const FVector& InRootLocation);
+
+	uint64 GetSettingsUID() const { return SettingsUID; }
+	void SetSettingsUID(uint64 InSettingsUID) { SettingsUID = InSettingsUID; }
+
+protected:
+	UPROPERTY()
+	FPCGProceduralISMComponentDescriptor Descriptor;
+
+	UPROPERTY()
+	bool bHasRootLocation = false;
+
+	UPROPERTY()
+	FVector RootLocation = FVector::ZeroVector;
+
+	UPROPERTY(Transient)
+	uint64 SettingsUID = -1; // purposefully a value that will never happen in data
+
+	// Cached raw pointer to ISM component
+	mutable UPCGProceduralISMComponent* CachedRawComponentPtr = nullptr;
+};
+
+namespace PCGManagedProceduralISMComponent
+{
+	UPCGManagedProceduralISMComponent* GetOrCreateManagedProceduralISMC(AActor* InTargetActor, UPCGComponent* InSourceComponent, uint64 InSettingsUID, const FPCGProceduralISMCBuilderParameters& InParams);
+}
diff --git a/Engine/Plugins/PCG/Source/PCG/Private/Compute/DataInterfaces/Elements/PCGStaticMeshSpawnerDataInterface.cpp b/Engine/Plugins/PCG/Source/PCG/Private/Compute/DataInterfaces/Elements/PCGStaticMeshSpawnerDataInterface.cpp
new file mode 100644
--- /dev/null
+++ b/Engine/Plugins/PCG/Source/PCG/Private/Compute/DataInterfaces/Elements/PCGStaticMeshSpawnerDataInterface.cpp
@@ -0,0 +1,126 @@
+// Copyright Epic Games, Inc. All Rights Reserved.
+
+#include "PCGStaticMeshSpawnerDataInterface.h"
+
+#include "Compute/PCGComputeCommon.h"
+#include "Compute/PCGDataBinding.h"
+#include "Elements/PCGStaticMeshSpawner.h"
+
+#include "GlobalRenderResources.h"
+#include "RHIResources.h"
+#include "ShaderCompilerCore.h"
+#include "ShaderCore.h"
+#include "ShaderParameterMetadataBuilder.h"
+#include "ComputeFramework/ShaderParamTypeDefinition.h"
+
+void UPCGStaticMeshSpawnerDataInterface::GetSupportedInputs(TArray<FShaderFunctionDefinition>& OutFunctions) const
+{
+	OutFunctions.AddDefaulted_GetRef()
+		.SetName(TEXT("SMSpawner_GetNumAttributes"))
+		.AddReturnType(EShaderFundamentalType::Uint);
+
+	OutFunctions.AddDefaulted_GetRef()
+		.SetName(TEXT("SMSpawner_GetAttributeIdOffsetStride"))
+		.AddReturnType(EShaderFundamentalType::Uint, 4)
+		.AddParam(EShaderFundamentalType::Uint); // InAttributeIndex
+
+	OutFunctions.AddDefaulted_GetRef()
+		.SetName(TEXT("SMSpawner_GetNumPrimitives"))
+		.AddReturnType(EShaderFundamentalType::Uint);
+
+	OutFunctions.AddDefaulted_GetRef()
+		.SetName(TEXT("SMSpawner_GetPrimitiveSelectionCDF"))
+		.AddReturnType(EShaderFundamentalType::Float)
+		.AddParam(EShaderFundamentalType::Uint); // InPrimitiveIndex
+}
+
+BEGIN_SHADER_PARAMETER_STRUCT(FPCGStaticMeshSpawnerDataInterfaceParameters,)
+	SHADER_PARAMETER(uint32, NumAttributes)
+	SHADER_PARAMETER_ARRAY(FUintVector4, AttributeIdOffsetStrides, [UPCGStaticMeshSpawnerDataInterface::MAX_ATTRIBUTES])
+	SHADER_PARAMETER(uint32, NumPrimitives)
+	SHADER_PARAMETER_SCALAR_ARRAY(float, SelectionCDF, [PCGComputeConstants::MAX_PRIMITIVE_COMPONENTS_PER_SPAWNER])
+END_SHADER_PARAMETER_STRUCT()
+
+void UPCGStaticMeshSpawnerDataInterface::GetShaderParameters(TCHAR const* UID, FShaderParametersMetadataBuilder& InOutBuilder, FShaderParametersMetadataAllocations& InOutAllocations) const
+{
+	InOutBuilder.AddNestedStruct<FPCGStaticMeshSpawnerDataInterfaceParameters>(UID);
+}
+
+void UPCGStaticMeshSpawnerDataInterface::GetHLSL(FString& OutHLSL, FString const& InDataInterfaceName) const
+{
+	TMap<FString, FStringFormatArg> TemplateArgs =
+	{
+		{ TEXT("DataInterfaceName"), InDataInterfaceName },
+		{ TEXT("MaxAttributes"), MAX_ATTRIBUTES },
+		{ TEXT("MaxPrimitives"), PCGComputeConstants::MAX_PRIMITIVE_COMPONENTS_PER_SPAWNER },
+	};
+
+	OutHLSL += FString::Format(TEXT(
+		"uint {DataInterfaceName}_NumAttributes;\n"
+		"uint4 {DataInterfaceName}_AttributeIdOffsetStrides[{MaxAttributes}];\n"
+		"uint SMSpawner_GetNumAttributes_{DataInterfaceName}() { return {DataInterfaceName}_NumAttributes; }\n"
+		"uint4 SMSpawner_GetAttributeIdOffsetStride_{DataInterfaceName}(uint InAttributeIndex) { return {DataInterfaceName}_AttributeIdOffsetStrides[InAttributeIndex]; }\n"
+		"\n"
+		"uint {DataInterfaceName}_NumPrimitives;\n"
+		"DECLARE_SCALAR_ARRAY(float, {DataInterfaceName}_SelectionCDF, {MaxPrimitives});\n"
+		"uint SMSpawner_GetNumPrimitives_{DataInterfaceName}() { return {DataInterfaceName}_NumPrimitives; }\n"
+		"float SMSpawner_GetPrimitiveSelectionCDF_{DataInterfaceName}(uint InPrimitiveIndex) { return GET_SCALAR_ARRAY_ELEMENT({DataInterfaceName}_SelectionCDF, InPrimitiveIndex); }\n"
+		), TemplateArgs);
+}
+
+UComputeDataProvider* UPCGStaticMeshSpawnerDataInterface::CreateDataProvider(TObjectPtr<UObject> InBinding, uint64 InInputMask, uint64 InOutputMask) const
+{
+	UPCGDataBinding* Binding = CastChecked<UPCGDataBinding>(InBinding);
+	const FPCGSpawnerPrimitives* Primitives = Binding->MeshSpawnersToPrimitives.Find(Settings);
+	if (!ensure(Primitives))
+	{
+		return nullptr;
+	}
+	if (!ensure(!Primitives->Primitives.IsEmpty()))
+	{
+		return nullptr;
+	}
+
+	TObjectPtr<UPCGStaticMeshSpawnerDataProvider> DataProvider = NewObject<UPCGStaticMeshSpawnerDataProvider>();
+	DataProvider->Settings = Cast<UPCGStaticMeshSpawnerSettings>(Settings);
+	DataProvider->AttributeIdOffsetStrides = Primitives->AttributeIdOffsetStrides;
+	DataProvider->SelectionCDF = Primitives->SelectionCDF;
+
+	return DataProvider;
+}
+
+FComputeDataProviderRenderProxy* UPCGStaticMeshSpawnerDataProvider::GetRenderProxy()
+{
+	return new FPCGStaticMeshSpawnerDataProviderProxy(AttributeIdOffsetStrides, SelectionCDF);
+}
+
+bool FPCGStaticMeshSpawnerDataProviderProxy::IsValid(FValidationData const& InValidationData) const
+{
+	if (InValidationData.ParameterStructSize != sizeof(FParameters))
+	{
+		return false;
+	}
+
+	return true;
+}
+
+void FPCGStaticMeshSpawnerDataProviderProxy::GatherDispatchData(FDispatchData const& InDispatchData)
+{
+	const TStridedView<FParameters> ParameterArray = MakeStridedParameterView<FParameters>(InDispatchData);
+	for (int32 InvocationIndex = 0; InvocationIndex < ParameterArray.Num(); ++InvocationIndex)
+	{
+		FParameters& Parameters = ParameterArray[InvocationIndex];
+
+		Parameters.NumAttributes = AttributeIdOffsetStrides.Num();
+		for (int32 Index = 0; Index < AttributeIdOffsetStrides.Num(); ++Index)
+		{
+			Parameters.AttributeIdOffsetStrides[Index] = AttributeIdOffsetStrides[Index];
+		}
+
+		Parameters.NumPrimitives = SelectionCDF.Num();
+		for (int32 Index = 0; Index < SelectionCDF.Num(); ++Index)
+		{
+			GET_SCALAR_ARRAY_ELEMENT(Parameters.SelectionCDF, Index) = SelectionCDF[Index];
+		}
+	}
+}
diff --git a/Engine/Plugins/PCG/Source/PCG/Private/Compute/DataInterfaces/Elements/PCGStaticMeshSpawnerDataInterface.h b/Engine/Plugins/PCG/Source/PCG/Private/Compute/DataInterfaces/Elements/PCGStaticMeshSpawnerDataInterface.h
new file mode 100644
--- /dev/null
+++ b/Engine/Plugins/PCG/Source/PCG/Private/Compute/DataInterfaces/Elements/PCGStaticMeshSpawnerDataInterface.h
@@ -0,0 +1,73 @@
+// Copyright Epic Games, Inc. All Rights Reserved.
+
+#pragma once
+
+#include "Compute/DataInterfaces/PCGComputeDataInterface.h"
+#include "ComputeFramework/ComputeDataProvider.h"
+
+#include "PCGStaticMeshSpawnerDataInterface.generated.h"
+
+class FPCGStaticMeshSpawnerDataInterfaceParameters;
+class UPCGStaticMeshSpawnerSettings;
+class UPCGSettings;
+
+/** Data Interface to marshal Static Mesh Spawner settings to the GPU. */
+UCLASS(ClassGroup = (Procedural))
+class UPCGStaticMeshSpawnerDataInterface : public UPCGComputeDataInterface
+{
+	GENERATED_BODY()
+
+public:
+	//~ Begin UComputeDataInterface Interface
+	TCHAR const* GetClassName() const override { return TEXT("PCGStaticMeshSpawner"); }
+	void GetSupportedInputs(TArray<FShaderFunctionDefinition>& OutFunctions) const override;
+	void GetShaderParameters(TCHAR const* UID, FShaderParametersMetadataBuilder& InOutBuilder, FShaderParametersMetadataAllocations& InOutAllocations) const override;
+	void GetHLSL(FString& OutHLSL, FString const& InDataInterfaceName) const override;
+	UComputeDataProvider* CreateDataProvider(TObjectPtr<UObject> InBinding, uint64 InInputMask, uint64 InOutputMask) const override;
+	//~ End UComputeDataInterface Interface
+
+	UPROPERTY()
+	TObjectPtr<const UPCGSettings> Settings = nullptr;
+
+	static constexpr uint32 MAX_ATTRIBUTES = 64;
+};
+
+UCLASS()
+class UPCGStaticMeshSpawnerDataProvider : public UComputeDataProvider
+{
+	GENERATED_BODY()
+
+public:
+	//~ Begin UComputeDataProvider Interface
+	FComputeDataProviderRenderProxy* GetRenderProxy() override;
+	//~ End UComputeDataProvider Interface
+
+	UPROPERTY()
+	TObjectPtr<const UPCGStaticMeshSpawnerSettings> Settings = nullptr;
+
+	UPROPERTY()
+	TArray<FUintVector4> AttributeIdOffsetStrides;
+
+	UPROPERTY()
+	TArray<float> SelectionCDF;
+};
+
+class FPCGStaticMeshSpawnerDataProviderProxy : public FComputeDataProviderRenderProxy
+{
+public:
+	FPCGStaticMeshSpawnerDataProviderProxy(TArray<FUintVector4>& InAttributeIdOffsetStrides, TArray<float> InSelectionCDF)
+		: AttributeIdOffsetStrides(InAttributeIdOffsetStrides)
+		, SelectionCDF(InSelectionCDF)
+	{}
+
+	//~ Begin FComputeDataProviderRenderProxy Interface
+	bool IsValid(FValidationData const& InValidationData) const override;
+	void GatherDispatchData(FDispatchData const& InDispatchData) override;
+	//~ End FComputeDataProviderRenderProxy Interface
+
+protected:
+	using FParameters = FPCGStaticMeshSpawnerDataInterfaceParameters;
+
+	TArray<FUintVector4> AttributeIdOffsetStrides;
+	TArray<float> SelectionCDF;
+};
diff --git a/Engine/Plugins/PCG/Source/PCG/Private/Compute/DataInterfaces/PCGInstanceDataInterface.cpp b/Engine/Plugins/PCG/Source/PCG/Private/Compute/DataInterfaces/PCGInstanceDataInterface.cpp
new file mode 100644
--- /dev/null
+++ b/Engine/Plugins/PCG/Source/PCG/Private/Compute/DataInterfaces/PCGInstanceDataInterface.cpp
@@ -0,0 +1,296 @@
+// Copyright Epic Games, Inc. All Rights Reserved.
+
+#include "PCGInstanceDataInterface.h"
+
+#include "PCGData.h"
+#include "PCGNode.h"
+#include "PCGSettings.h"
+#include "PCGSceneWriterCS.h"
+#include "Compute/PCGComputeCommon.h"
+#include "Compute/PCGDataBinding.h"
+#include "Data/PCGTextureData.h"
+
+#include "GlobalRenderResources.h"
+#include "GlobalShader.h"
+#include "GPUSceneWriter.h"
+#include "InstanceDataSceneProxy.h"
+#include "PrimitiveSceneInfo.h"
+#include "PrimitiveSceneProxy.h"
+#include "RendererUtils.h"
+#include "RenderGraphBuilder.h"
+#include "RenderGraphUtils.h"
+#include "RHIResources.h"
+#include "RHIStaticStates.h"
+#include "SceneComputeUpdates.h"
+#include "ShaderCompilerCore.h"
+#include "ShaderCore.h"
+#include "ShaderParameterMetadataBuilder.h"
+#include "ShaderParameterStruct.h"
+#include "TextureResource.h"
+#include "Components/PrimitiveComponent.h"
+#include "ComputeFramework/ComputeKernelPermutationSet.h"
+#include "ComputeFramework/ShaderParamTypeDefinition.h"
+#include "Engine/Texture.h"
+#include "Engine/World.h"
+
+#include UE_INLINE_GENERATED_CPP_BY_NAME(PCGInstanceDataInterface)
+
+void UPCGInstanceDataInterface::GetSupportedOutputs(TArray<FShaderFunctionDefinition>& OutFunctions) const
+{
+	OutFunctions.AddDefaulted_GetRef()
+		.SetName(TEXT("InstanceData_GetNumPrimitives"))
+		.AddReturnType(EShaderFundamentalType::Uint);
+
+	OutFunctions.AddDefaulted_GetRef()
+		.SetName(TEXT("InstanceData_GetNumInstancesAllPrimitives"))
+		.AddReturnType(EShaderFundamentalType::Uint);
+
+	OutFunctions.AddDefaulted_GetRef()
+		.SetName(TEXT("InstanceData_GetIndexToWriteTo"))
+		.AddReturnType(EShaderFundamentalType::Int)
+		.AddParam(EShaderFundamentalType::Uint); // InPrimitiveIndex
+
+	OutFunctions.AddDefaulted_GetRef()
+		.SetName(TEXT("InstanceData_WriteInstanceLocalToWorld"))
+		.AddParam(EShaderFundamentalType::Uint) // InInstanceIndex
+		.AddParam(EShaderFundamentalType::Float, 4, 4);  // InTransform
+
+	OutFunctions.AddDefaulted_GetRef()
+		.SetName(TEXT("InstanceData_WriteCustomFloat"))
+		.AddParam(EShaderFundamentalType::Uint) // InInstanceIndex
+		.AddParam(EShaderFundamentalType::Uint) // InCustomFloatIndex
+		.AddParam(EShaderFundamentalType::Float); // InValue
+
+	OutFunctions.AddDefaulted_GetRef()
+		.SetName(TEXT("InstanceData_WriteCustomFloatRaw"))
+		.AddParam(EShaderFundamentalType::Uint) // InInstanceIndex
+		.AddParam(EShaderFundamentalType::Uint) // InCustomFloatIndex
+		.AddParam(EShaderFundamentalType::Uint); // InValueAsUint
+}
+
+BEGIN_SHADER_PARAMETER_STRUCT(FPCGInstanceDataInterfaceParameters,)
+	SHADER_PARAMETER_ARRAY(FUintVector4, PrimitiveParameters, [PCGComputeConstants::MAX_PRIMITIVE_COMPONENTS_PER_SPAWNER]) // (NumInstancesAllocated, InstanceOffset, Unused, Unused)
+	SHADER_PARAMETER(uint32, NumPrimitives)
+	SHADER_PARAMETER(uint32, NumInstancesAllPrimitives)
+	SHADER_PARAMETER(uint32, NumCustomFloatsPerInstance)
+	SHADER_PARAMETER_RDG_BUFFER_UAV(RWStructuredBuffer<FUintVector4>, InstanceData)
+	SHADER_PARAMETER_RDG_BUFFER_UAV(RWStructuredBuffer<uint32>, InstanceCustomFloatData)
+	SHADER_PARAMETER_RDG_BUFFER_UAV(RWStructuredBuffer<uint32>, WriteCounters)
+END_SHADER_PARAMETER_STRUCT()
+
+void UPCGInstanceDataInterface::GetShaderParameters(TCHAR const* UID, FShaderParametersMetadataBuilder& InOutBuilder, FShaderParametersMetadataAllocations& InOutAllocations) const
+{
+	InOutBuilder.AddNestedStruct<FPCGInstanceDataInterfaceParameters>(UID);
+}
+
+TCHAR const* UPCGInstanceDataInterface::TemplateFilePath = TEXT("/Plugin/PCG/Private/PCGInstanceDataInterface.ush");
+
+TCHAR const* UPCGInstanceDataInterface::GetShaderVirtualPath() const
+{
+	return TemplateFilePath;
+}
+
+void UPCGInstanceDataInterface::GetShaderHash(FString& InOutKey) const
+{
+	GetShaderFileHash(TemplateFilePath, EShaderPlatform::SP_PCD3D_SM5).AppendString(InOutKey);
+}
+
+void UPCGInstanceDataInterface::GetHLSL(FString& OutHLSL, FString const& InDataInterfaceName) const
+{
+	TMap<FString, FStringFormatArg> TemplateArgs =
+	{
+		{ TEXT("DataInterfaceName"), InDataInterfaceName },
+	};
+
+	FString TemplateFile;
+	if (ensure(LoadShaderSourceFile(TemplateFilePath, EShaderPlatform::SP_PCD3D_SM5, &TemplateFile, nullptr)))
+	{
+		OutHLSL += FString::Format(*TemplateFile, TemplateArgs);
+	}
+}
+
+UComputeDataProvider* UPCGInstanceDataInterface::CreateDataProvider(TObjectPtr<UObject> InBinding, uint64 InInputMask, uint64 InOutputMask) const
+{
+	UPCGDataBinding* Binding = CastChecked<UPCGDataBinding>(InBinding);
+
+	FPCGSpawnerPrimitives* FoundPrimitives = Binding->MeshSpawnersToPrimitives.Find(ProducerSettings);
+	if (!FoundPrimitives || FoundPrimitives->Primitives.IsEmpty())
+	{
+		return nullptr;
+	}
+
+	if (!ensure(FoundPrimitives->Primitives.Num() <= PCGComputeConstants::MAX_PRIMITIVE_COMPONENTS_PER_SPAWNER))
+	{
+		// Last resort - should be clamped earlier during setup and not come to this.
+		FoundPrimitives->Primitives.SetNum(PCGComputeConstants::MAX_PRIMITIVE_COMPONENTS_PER_SPAWNER);
+	}
+
+	TObjectPtr<UPCGInstanceDataProvider> DataProvider = NewObject<UPCGInstanceDataProvider>();
+	
+	DataProvider->Primitives = FoundPrimitives->Primitives;
+	DataProvider->NumCustomFloatsPerInstance = FoundPrimitives->NumCustomFloats;
+	
+	const UPCGNode* Node = ProducerSettings ? Cast<UPCGNode>(ProducerSettings->GetOuter()) : nullptr;
+	const UPCGPin* InputPin = Node ? Node->GetInputPin(InputPinProvidingData) : nullptr;
+	check(ProducerSettings && InputPin);
+	const FPCGDataCollectionDesc InputDataDesc = ProducerSettings->ComputeInputPinDataDesc(InputPin, Binding);
+	DataProvider->NumInstancesAllPrimitives = InputDataDesc.ComputeDataElementCount(EPCGDataType::Point);
+
+	return DataProvider;
+}
+
+FComputeDataProviderRenderProxy* UPCGInstanceDataProvider::GetRenderProxy()
+{
+	TArray<FPrimitiveSceneProxy*> PrimitiveProxies;
+	PrimitiveProxies.Reserve(Primitives.Num());
+	
+	for (TObjectPtr<UPrimitiveComponent>& Primitive : Primitives)
+	{
+		UWorld* World = Primitive ? Primitive->GetWorld() : nullptr;
+		FSceneInterface* Scene = World ? World->Scene : nullptr;
+		if (FPrimitiveSceneProxy* PrimitiveSceneProxy = Primitive ? Primitive->GetSceneProxy() : nullptr)
+		{
+			PrimitiveProxies.Add(PrimitiveSceneProxy);
+		}
+	}
+
+	return new FPCGInstanceDataProviderProxy(MoveTemp(PrimitiveProxies), NumInstancesAllPrimitives, NumCustomFloatsPerInstance);
+}
+
+FPCGInstanceDataProviderProxy::FPCGInstanceDataProviderProxy(TArray<FPrimitiveSceneProxy*>&& InPrimitiveSceneProxies, uint32 InNumInstancesAllPrimitives, uint32 InNumCustomFloatsPerInstance)
+	: PrimitiveSceneProxies(MoveTemp(InPrimitiveSceneProxies))
+	, NumInstancesAllPrimitives(InNumInstancesAllPrimitives)
+	, NumCustomFloatsPerInstance(InNumCustomFloatsPerInstance)
+{
+}
+
+bool FPCGInstanceDataProviderProxy::IsValid(FValidationData const& InValidationData) const
+{
+	if (InValidationData.ParameterStructSize != sizeof(FParameters))
+	{
+		return false;
+	}
+
+	if (PrimitiveSceneProxies.IsEmpty())
+	{
+		UE_LOG(LogPCG, Warning, TEXT("Proxy invalid due to missing primitive scene proxy."));
+		return false;
+	}
+
+	if (NumInstancesAllPrimitives == 0)
+	{
+		UE_LOG(LogPCG, Warning, TEXT("Proxy invalid due to 0 instance count."));
+		return false;
+	}
+
+	return true;
+}
+
+void FPCGInstanceDataProviderProxy::AllocateResources(FRDGBuilder& GraphBuilder, FAllocationData const& InAllocationData)
+{
+	ensure(NumInstancesAllPrimitives > 0);
+	const uint32 StrideUint4s = 3;
+	InstanceData = GraphBuilder.CreateBuffer(FRDGBufferDesc::CreateStructuredDesc(sizeof(FUintVector4) * StrideUint4s, NumInstancesAllPrimitives), TEXT("PCGInstanceDataBuffer"));
+	InstanceDataSRV = GraphBuilder.CreateSRV(InstanceData);
+	InstanceDataUAV = GraphBuilder.CreateUAV(InstanceData);
+
+	const uint32 CustomFloatsRequired = FMath::Max(1u, NumInstancesAllPrimitives * NumCustomFloatsPerInstance);
+	InstanceCustomFloatData = GraphBuilder.CreateBuffer(FRDGBufferDesc::CreateStructuredDesc(sizeof(uint32), CustomFloatsRequired), TEXT("PCGInstanceCustomFloatDataBuffer"));
+	InstanceCustomFloatDataSRV = GraphBuilder.CreateSRV(InstanceCustomFloatData);
+	InstanceCustomFloatDataUAV = GraphBuilder.CreateUAV(InstanceCustomFloatData);
+
+	WriteCounters = GraphBuilder.CreateBuffer(FRDGBufferDesc::CreateStructuredDesc(sizeof(uint32), PrimitiveSceneProxies.Num()), TEXT("PCGWriteCounters"));
+	WriteCountersSRV = GraphBuilder.CreateSRV(WriteCounters);
+	WriteCountersUAV = GraphBuilder.CreateUAV(WriteCounters);
+	TArray<uint32> Zeros;
+	Zeros.SetNumZeroed(PrimitiveSceneProxies.Num());
+	GraphBuilder.QueueBufferUpload(WriteCounters, Zeros.GetData(), Zeros.Num() * Zeros.GetTypeSize());
+}
+
+void FPCGInstanceDataProviderProxy::GatherDispatchData(FDispatchData const& InDispatchData)
+{
+	const TStridedView<FParameters> ParameterArray = MakeStridedParameterView<FParameters>(InDispatchData);
+	for (int32 InvocationIndex = 0; InvocationIndex < ParameterArray.Num(); ++InvocationIndex)
+	{
+		FParameters& Parameters = ParameterArray[InvocationIndex];
+		Parameters.NumPrimitives = PrimitiveSceneProxies.Num();
+		Parameters.NumInstancesAllPrimitives = NumInstancesAllPrimitives;
+		Parameters.NumCustomFloatsPerInstance = NumCustomFloatsPerInstance;
+		Parameters.InstanceData = InstanceDataUAV;
+		Parameters.InstanceCustomFloatData = InstanceCustomFloatDataUAV;
+		Parameters.WriteCounters = WriteCountersUAV;
+
+		uint32 CumulativeInstanceCount = 0;
+		for (int32 Index = 0; Index < PrimitiveSceneProxies.Num(); ++Index)
+		{
+			const uint32 PrimitiveInstanceCount = PrimitiveSceneProxies[Index]->GetInstanceDataHeader().NumInstances;
+
+			Parameters.PrimitiveParameters[Index] = FUintVector4(
+				PrimitiveInstanceCount,
+				CumulativeInstanceCount,
+				0, 0); // Unused
+
+			CumulativeInstanceCount += PrimitiveInstanceCount;
+		}
+	}
+}
+
+void FPCGInstanceDataProviderProxy::PostGraphDispatch(FRDGBuilder& GraphBuilder) const
+{
+	if (NumInstancesAllPrimitives == 0)
+	{
+		return;
+	}
+
+	const FSceneComputeUpdatesBlackboardEntry* SceneComputeUpdatesBB = GraphBuilder.Blackboard.Get<FSceneComputeUpdatesBlackboardEntry>();
+	if (!SceneComputeUpdatesBB || !ensure(SceneComputeUpdatesBB->SceneComputeUpdates))
+	{
+		UE_LOG(LogPCG, Warning, TEXT("Scene Compute Updates interface missing from graph builder blackboard. Scene will not be updated."));
+		return;
+	}
+
+	uint32 CumulativeInstanceCount = 0;
+
+	// One dispatch per primitive. If the scene compute update API supported writing to multiple primitives this could be done in one unified call.
+	for (int32 PrimitiveIndex = 0; PrimitiveIndex < PrimitiveSceneProxies.Num(); ++PrimitiveIndex)
+	{
+		FPrimitiveSceneProxy* PrimitiveSceneProxy = PrimitiveSceneProxies[PrimitiveIndex];
+		FPrimitiveSceneInfo* SceneInfo = PrimitiveSceneProxy ? PrimitiveSceneProxy->GetPrimitiveSceneInfo() : nullptr;
+		if (!ensure(SceneInfo))
+		{
+			return;
+		}
+
+		const int32 NumInstancesAllocatedInGPUScene = PrimitiveSceneProxy->GetInstanceDataHeader().NumInstances;
+
+		// Post graph dispatch, the instance data will have been populated. Enqueue a scene update to write the instance data into the scene.
+		// Will get executed on same frame and in same builder.
+		SceneComputeUpdatesBB->SceneComputeUpdates->EnqueueUpdate(SceneInfo,
+			FGPUSceneWriteDelegate::CreateLambda([PrimitiveIndex, NumInstancesAllocatedInGPUScene, CumulativeInstanceCount, InstanceDataSRV=InstanceDataSRV, InstanceCustomFloatDataSRV=InstanceCustomFloatDataSRV, WriteCountersSRV=WriteCountersSRV](FRDGBuilder& GraphBuilder, const FGPUSceneWriteDelegateParams& Params)
+			{
+				RDG_EVENT_SCOPE(GraphBuilder, "SceneComputeUpdateInterface->EnqueueUpdate");
+				check(Params.PersistentPrimitiveId != (uint32)INDEX_NONE);
+
+				// Write instances
+				FPCGSceneWriterCS::FParameters* Parameters = GraphBuilder.AllocParameters<FPCGSceneWriterCS::FParameters>();
+				Parameters->InPrimitiveIndex = PrimitiveIndex;
+				Parameters->InNumInstancesAllocatedInGPUScene = NumInstancesAllocatedInGPUScene;
+				Parameters->InInstanceOffset = CumulativeInstanceCount;
+				Parameters->InInstanceData = InstanceDataSRV;
+				Parameters->InInstanceCustomFloatData = InstanceCustomFloatDataSRV;
+				Parameters->InWriteCounters = WriteCountersSRV;
+				Parameters->InPrimitiveId = Params.PersistentPrimitiveId;
+				Parameters->GPUSceneWriterParameters = Params.GPUWriteParams;
+				Parameters->InCustomDataCount = Params.NumCustomDataFloats;
+				Parameters->InPayloadDataFlags = Params.PackedInstanceSceneDataFlags;
+			
+				// We're using a custom compute shader here rather than using a compute graph kernel, because ultimately the scene update will incorporate the
+				// new instance data. In the future we will not write directly to the scene here.
+				TShaderMapRef<FPCGSceneWriterCS> Shader(GetGlobalShaderMap(GMaxRHIFeatureLevel));
+				const int GroupCount = FMath::DivideAndRoundUp<int>(NumInstancesAllocatedInGPUScene, FPCGSceneWriterCS::NumThreadsPerGroup);
+				FComputeShaderUtils::AddPass(GraphBuilder, RDG_EVENT_NAME("PCGWriteInstanceData"), ERDGPassFlags::Compute | ERDGPassFlags::NeverCull, Shader, Parameters, FIntVector(GroupCount, 1, 1));
+			}));
+
+		CumulativeInstanceCount += NumInstancesAllocatedInGPUScene;
+	}
+}
diff --git a/Engine/Plugins/PCG/Source/PCG/Private/Compute/DataInterfaces/PCGInstanceDataInterface.h b/Engine/Plugins/PCG/Source/PCG/Private/Compute/DataInterfaces/PCGInstanceDataInterface.h
new file mode 100644
--- /dev/null
+++ b/Engine/Plugins/PCG/Source/PCG/Private/Compute/DataInterfaces/PCGInstanceDataInterface.h
@@ -0,0 +1,100 @@
+// Copyright Epic Games, Inc. All Rights Reserved.
+
+#pragma once
+
+#include "Compute/DataInterfaces/PCGComputeDataInterface.h"
+
+#include "RHIResources.h"
+#include "ComputeFramework/ComputeDataProvider.h"
+
+#include "PCGInstanceDataInterface.generated.h"
+
+class FPCGInstanceDataInterfaceParameters;
+class FPrimitiveSceneProxy;
+class UPCGSettings;
+
+/** Data Interface for writing instance data (transform and custom floats) to an intermediate buffer, which is then injected into the GPU Scene. */
+UCLASS(ClassGroup = (Procedural))
+class UPCGInstanceDataInterface : public UPCGComputeDataInterface
+{
+	GENERATED_BODY()
+
+public:
+	//~ Begin UComputeDataInterface Interface
+	TCHAR const* GetClassName() const override { return TEXT("PCGInstanceData"); }
+	void GetSupportedOutputs(TArray<FShaderFunctionDefinition>& OutFunctions) const override;
+	void GetShaderParameters(TCHAR const* UID, FShaderParametersMetadataBuilder& InOutBuilder, FShaderParametersMetadataAllocations& InOutAllocations) const override;
+	TCHAR const* GetShaderVirtualPath() const override;
+	void GetShaderHash(FString& InOutKey) const override;
+	void GetHLSL(FString& OutHLSL, FString const& InDataInterfaceName) const override;
+	UComputeDataProvider* CreateDataProvider(TObjectPtr<UObject> InBinding, uint64 InInputMask, uint64 InOutputMask) const override;
+	//~ End UComputeDataInterface Interface
+
+	void SetProducerSettings(const UPCGSettings* InProducerSettings) { ProducerSettings = InProducerSettings; }
+
+public:
+	UPROPERTY()
+	FName InputPinProvidingData = NAME_None;
+
+private:
+	/** Settings of node that produces this data, normally the upstream node. */
+	UPROPERTY()
+	TObjectPtr<const UPCGSettings> ProducerSettings;
+
+	static TCHAR const* TemplateFilePath;
+};
+
+/** Compute Framework Data Provider for writing to the GPU Scene. */
+UCLASS()
+class UPCGInstanceDataProvider : public UComputeDataProvider
+{
+	GENERATED_BODY()
+
+public:
+	//~ Begin UComputeDataProvider Interface
+	FComputeDataProviderRenderProxy* GetRenderProxy() override;
+	//~ End UComputeDataProvider Interface
+
+	UPROPERTY()
+	TArray<TObjectPtr<UPrimitiveComponent>> Primitives;
+
+	uint32 NumInstancesAllPrimitives = 0;
+	uint32 NumCustomFloatsPerInstance = 0;
+};
+
+class FPCGInstanceDataProviderProxy : public FComputeDataProviderRenderProxy
+{
+public:
+	FPCGInstanceDataProviderProxy(TArray<FPrimitiveSceneProxy*>&& InPrimitiveSceneProxies, uint32 InNumInstances, uint32 InNumCustomFloatsPerInstance);
+
+	//~ Begin FComputeDataProviderRenderProxy Interface
+	bool IsValid(FValidationData const& InValidationData) const override;
+	void AllocateResources(FRDGBuilder& GraphBuilder, FAllocationData const& InAllocationData);
+	void GatherDispatchData(FDispatchData const& InDispatchData) override;
+	void PostGraphDispatch(FRDGBuilder& GraphBuilder) const override;
+	//~ End FComputeDataProviderRenderProxy Interface
+
+protected:
+	using FParameters = FPCGInstanceDataInterfaceParameters;
+
+	TArray<FPrimitiveSceneProxy*> PrimitiveSceneProxies;
+	TArray<uint32> PrimitiveCounts;
+
+	/** Instance transforms of all instances across all primitives. Three float4s per instance. */
+	FRDGBufferRef InstanceData = nullptr;
+	FRDGBufferSRVRef InstanceDataSRV = nullptr;
+	FRDGBufferUAVRef InstanceDataUAV = nullptr;
+
+	/** Custom float data for all instances across all primitives. N per instance. Number of custom floats N comes from inputs (and is static). */
+	FRDGBufferRef InstanceCustomFloatData = nullptr;
+	FRDGBufferSRVRef InstanceCustomFloatDataSRV = nullptr;
+	FRDGBufferUAVRef InstanceCustomFloatDataUAV = nullptr;
+
+	/** A per-primitive instance counter. Updated atomically and used to place instances into a segmented array (one segment per primitive). */
+	FRDGBufferRef WriteCounters = nullptr;
+	FRDGBufferSRVRef WriteCountersSRV = nullptr;
+	FRDGBufferUAVRef WriteCountersUAV = nullptr;
+
+	uint32 NumInstancesAllPrimitives = 0;
+	uint32 NumCustomFloatsPerInstance = 0;
+};
diff --git a/Engine/Plugins/PCG/Source/PCG/Private/Compute/Elements/PCGComputeGraphElement.cpp b/Engine/Plugins/PCG/Source/PCG/Private/Compute/Elements/PCGComputeGraphElement.cpp
--- a/Engine/Plugins/PCG/Source/PCG/Private/Compute/Elements/PCGComputeGraphElement.cpp
+++ b/Engine/Plugins/PCG/Source/PCG/Private/Compute/Elements/PCGComputeGraphElement.cpp
@@ -3,9 +3,16 @@
 #include "Compute/Elements/PCGComputeGraphElement.h"
 
 #include "PCGComponent.h"
+#include "PCGManagedResource.h"
 #include "PCGModule.h"
 #include "PCGSubsystem.h"
+#include "Components/PCGProceduralISMComponent.h"
+#include "Compute/PCGComputeCommon.h"
 #include "Compute/DataInterfaces/PCGDataCollectionDataInterface.h"
+#include "Elements/PCGStaticMeshSpawner.h"
+#include "Engine/StaticMesh.h"
+#include "InstanceDataPackers/PCGInstanceDataPackerBase.h"
+#include "MeshSelectors/PCGMeshSelectorWeighted.h"
 
 #include "ComputeWorkerInterface.h"
 #include "ComputeFramework/ComputeFramework.h"
@@ -100,6 +107,17 @@ bool FPCGComputeGraphElement::ExecuteInternal(FPCGContext* InContext) const
 	// 4. Initialize and parse incoming data for data sizes, attributes, etc that will drive buffer allocations and dispatch thread counts.
 	if (!Context->DataBinding)
 	{
+		for (TWeakObjectPtr<const UPCGNode> Node : Graph->KernelToNode)
+		{
+			const UPCGSettings* Settings = Node.Get() ? Node->GetSettings() : nullptr;
+
+			if (!Settings || !Settings->IsKernelValid(Context, /*bQuiet=*/false))
+			{
+				return true;
+			}
+		}
+
+		// When attribute table was built at compile time, some attribute types may not be known. Fill them out now.
 		Graph->FillInMissingAttributeTableTypes(Context->InputData);
 
 		UPCGDataBinding* DataBindingObject = FPCGContext::NewObject_AnyThread<UPCGDataBinding>(Context);
@@ -123,6 +141,8 @@ bool FPCGComputeGraphElement::ExecuteInternal(FPCGContext* InContext) const
 
 		DataForGPU.InputPinLabelAliases = Graph->InputPinLabelAliases;
 
+		const bool bAnyComponentsSetup = SetupProceduralISMComponents(InContext, DataBindingObject);
+
 		Context->ComputeGraphInstance.CreateDataProviders(Graph.Get(), 0, Context->DataBinding.Get());
 
 		// Register all providers running async operations. TODO review if we should have a general API like "RunsAsyncOperations()"?
@@ -152,14 +172,11 @@ bool FPCGComputeGraphElement::ExecuteInternal(FPCGContext* InContext) const
 			}
 		}
 
-		for (TWeakObjectPtr<const UPCGNode> Node : Graph->KernelToNode)
+		if (bAnyComponentsSetup)
 		{
-			const UPCGSettings* Settings = Node.Get() ? Node->GetSettings() : nullptr;
-
-			if (!Settings || !Settings->IsKernelValid(Context, /*bQuiet=*/false))
-			{
-				return true;
-			}
+			// Delay to give time for proxy to settle and for scene update to allocate space in the GPU scene.
+			SleepUntilNextFrame();
+			return false;
 		}
 	}
 
@@ -229,12 +246,111 @@ bool FPCGComputeGraphElement::ExecuteInternal(FPCGContext* InContext) const
 	}
 }
 
+bool FPCGComputeGraphElement::SetupProceduralISMComponents(FPCGContext* InContext, UPCGDataBinding* Binding) const
+{
+	bool bAnyComponentsSetup = false;
+
+	for (const UPCGSettings* Settings : Graph->StaticMeshSpawners)
+	{
+		const UPCGStaticMeshSpawnerSettings* SpawnerSettings = Cast<UPCGStaticMeshSpawnerSettings>(Settings);
+		if (!ensure(SpawnerSettings))
+		{
+			continue;
+		}
+
+		const UPCGMeshSelectorWeighted* Selector = Cast<UPCGMeshSelectorWeighted>(SpawnerSettings->MeshSelectorParameters);
+		AActor* TargetActor = SpawnerSettings->TargetActor.Get() ? SpawnerSettings->TargetActor.Get() : InContext->GetTargetActor(nullptr);
+		if (!ensure(Selector) || !ensure(TargetActor))
+		{
+			continue;
+		}
+
+		const UPCGNode* Node = Cast<const UPCGNode>(SpawnerSettings->GetOuter());
+		const UPCGPin* Pin = Node ? Node->GetInputPin(PCGPinConstants::DefaultInputLabel) : nullptr;
+		if (!ensure(Pin))
+		{
+			continue;
+		}
+
+		const FPCGDataCollectionDesc InputDataDesc = SpawnerSettings->ComputeInputPinDataDesc(Pin, Binding);
+
+		const uint32 PointCount = InputDataDesc.ComputeDataElementCount(EPCGDataType::Point);
+		
+		uint32 CustomFloatCount = 0;
+		TArray<FUint32Vector4> AttributeIdOffsetStrides;
+		if (SpawnerSettings->InstanceDataPackerParameters && SpawnerSettings->InstanceDataPackerParameters->GetAttributeNames(/*OutNames=*/nullptr))
+		{
+			TArray<FName> AttributeNames;
+			SpawnerSettings->InstanceDataPackerParameters->GetAttributeNames(&AttributeNames);
+
+			PCGDataForGPUHelpers::ComputeCustomFloatPacking(AttributeNames, Binding, InputDataDesc, CustomFloatCount, AttributeIdOffsetStrides);
+		}
+
+		bool AttributeSetupDone = false;
+
+		float TotalWeight = 0.0f;
+		for (const FPCGMeshSelectorWeightedEntry& Entry : Selector->MeshEntries)
+		{
+			TotalWeight += Entry.Weight;
+		}
+
+		if (!ensure(TotalWeight > UE_SMALL_NUMBER))
+		{
+			continue;
+		}
+
+		float CumulativeWeight = 0.0f;
+
+		for (const FPCGMeshSelectorWeightedEntry& Entry : Selector->MeshEntries)
+		{
+			UStaticMesh* StaticMesh = Entry.Descriptor.StaticMesh.LoadSynchronous();
+			const float Weight = float(Entry.Weight) / TotalWeight;
+			CumulativeWeight += Weight;
+
+			FPCGProceduralISMCBuilderParameters Params;
+			Params.Descriptor.NumInstances = FMath::CeilToInt(PointCount * Weight);
+			Params.Descriptor.LocalBounds = InContext->SourceComponent->GetGridBounds().ShiftBy(-InContext->SourceComponent->GetOwner()->GetActorLocation());
+			Params.Descriptor.NumCustomFloats = CustomFloatCount;
+			Params.Descriptor = Entry.Descriptor;
+			Params.bAllowDescriptorChanges = false;
+
+			UPCGManagedProceduralISMComponent* MISMC = PCGManagedProceduralISMComponent::GetOrCreateManagedProceduralISMC(TargetActor, InContext->SourceComponent.Get(), SpawnerSettings->UID, Params);
+
+			check(MISMC);
+			MISMC->SetCrc(InContext->DependenciesCrc);
+
+			// Don't bother registering the resource change as we're transient anyway.
+			//InContext->TouchedResources.Emplace(MISMC);
+
+			FPCGSpawnerPrimitives& Primitives = Binding->MeshSpawnersToPrimitives.FindOrAdd(SpawnerSettings);
+			Primitives.Primitives.Add(MISMC->GetComponent());
+			Primitives.SelectionCDF.Add(CumulativeWeight);
+			if (!AttributeSetupDone)
+			{
+				Primitives.NumCustomFloats = CustomFloatCount;
+				Primitives.AttributeIdOffsetStrides = MoveTemp(AttributeIdOffsetStrides);
+				AttributeSetupDone = true;
+			}
+
+			bAnyComponentsSetup = true;
+
+			if (Primitives.Primitives.Num() >= PCGComputeConstants::MAX_PRIMITIVE_COMPONENTS_PER_SPAWNER)
+			{
+				UE_LOG(LogPCG, Warning, TEXT("Attempted to emit too many primitive components, terminated after creating %d."), Primitives.Primitives.Num());
+				break;
+			}
+		}
+	}
+
+	return bAnyComponentsSetup;
+}
+
 void FPCGComputeGraphElement::PostExecuteInternal(FPCGContext* InContext) const
 {
 	check(InContext);
 	FPCGComputeGraphContext* Context = static_cast<FPCGComputeGraphContext*>(InContext);
 
-	if (!ensure(Context->DataBinding))
+	if (!Context->DataBinding)
 	{
 		return;
 	}
diff --git a/Engine/Plugins/PCG/Source/PCG/Private/Compute/Elements/PCGComputeGraphElement.h b/Engine/Plugins/PCG/Source/PCG/Private/Compute/Elements/PCGComputeGraphElement.h
--- a/Engine/Plugins/PCG/Source/PCG/Private/Compute/Elements/PCGComputeGraphElement.h
+++ b/Engine/Plugins/PCG/Source/PCG/Private/Compute/Elements/PCGComputeGraphElement.h
@@ -69,6 +69,8 @@ protected:
 	// TODO - need to accumulate dependencies from compute graph nodes.
 	virtual bool IsCacheable(const UPCGSettings* InSettings) const override { return false; }
 
+	bool SetupProceduralISMComponents(FPCGContext* InContext, UPCGDataBinding* Binding) const;
+
 	void ResetAsyncOperations(FPCGContext* InContext) const;
 
 #if WITH_EDITOR
diff --git a/Engine/Plugins/PCG/Source/PCG/Private/Compute/Elements/PCGCustomHLSL.cpp b/Engine/Plugins/PCG/Source/PCG/Private/Compute/Elements/PCGCustomHLSL.cpp
--- a/Engine/Plugins/PCG/Source/PCG/Private/Compute/Elements/PCGCustomHLSL.cpp
+++ b/Engine/Plugins/PCG/Source/PCG/Private/Compute/Elements/PCGCustomHLSL.cpp
@@ -1629,13 +1629,8 @@ FString UPCGCustomHLSLSettings::GetCookedKernelSource(const TMap<FName, FPCGKern
 
 bool FPCGCustomHLSLElement::ExecuteInternal(FPCGContext* Context) const
 {
-	check(Context);
-
-	const UPCGCustomHLSLSettings* Settings = Context->GetInputSettings<UPCGCustomHLSLSettings>();
-	check(Settings);
-
-	Settings->IsKernelValid(Context, /*bQuiet=*/false);
-
+	// This element does not support CPU execution and we are never supposed to land here.
+	check(false);
 	return true;
 }
 
diff --git a/Engine/Plugins/PCG/Source/PCG/Private/Compute/PCGComputeCommon.h b/Engine/Plugins/PCG/Source/PCG/Private/Compute/PCGComputeCommon.h
--- a/Engine/Plugins/PCG/Source/PCG/Private/Compute/PCGComputeCommon.h
+++ b/Engine/Plugins/PCG/Source/PCG/Private/Compute/PCGComputeCommon.h
@@ -58,6 +58,8 @@ namespace PCGComputeConstants
 	constexpr EPCGDataType AllowedDataCollectionTypes = EPCGDataType::Point | EPCGDataType::Param;
 
 	constexpr uint32 KernelExecutedFlag = 1 << 31;
+
+	constexpr uint32 MAX_PRIMITIVE_COMPONENTS_PER_SPAWNER = 64;
 }
 
 namespace PCGComputeHelpers
diff --git a/Engine/Plugins/PCG/Source/PCG/Private/Compute/PCGComputeGraph.h b/Engine/Plugins/PCG/Source/PCG/Private/Compute/PCGComputeGraph.h
--- a/Engine/Plugins/PCG/Source/PCG/Private/Compute/PCGComputeGraph.h
+++ b/Engine/Plugins/PCG/Source/PCG/Private/Compute/PCGComputeGraph.h
@@ -45,10 +45,14 @@ public:
 	UPROPERTY()
 	TMap<TObjectPtr<const UPCGPin>, FName> OutputCPUPinToInputGPUPinAlias;
 
-	// Node corresponding to each kernel, useful for compilation feedback.
+	/** Node corresponding to each kernel, useful for compilation feedback. */
 	UPROPERTY()
 	TArray<TWeakObjectPtr<const UPCGNode>> KernelToNode;
 
+	/** List of settings for all nodes that spawn static meshes, so we can do required primitive & DI setup when compute graph element executes. */
+	UPROPERTY()
+	TArray<TObjectPtr<const UPCGSettings>> StaticMeshSpawners;
+
 protected:
 	UPROPERTY()
 	TMap<FName /* Attribute name */, FPCGKernelAttributeIDAndType> GlobalAttributeLookupTable;
diff --git a/Engine/Plugins/PCG/Source/PCG/Private/Compute/PCGDataBinding.h b/Engine/Plugins/PCG/Source/PCG/Private/Compute/PCGDataBinding.h
--- a/Engine/Plugins/PCG/Source/PCG/Private/Compute/PCGDataBinding.h
+++ b/Engine/Plugins/PCG/Source/PCG/Private/Compute/PCGDataBinding.h
@@ -12,6 +12,27 @@
 class UPCGComponent;
 class UPCGComputeGraph;
 
+USTRUCT()
+struct FPCGSpawnerPrimitives
+{
+	GENERATED_BODY()
+
+	UPROPERTY()
+	TArray<TObjectPtr<UPrimitiveComponent>> Primitives;
+
+	/** Cumulative distribution function values (one per primitive) to enable choosing a primitive based on a random draw value. */
+	UPROPERTY()
+	TArray<float> SelectionCDF;
+
+	// Same for all primitives
+	UPROPERTY()
+	uint32 NumCustomFloats = 0;
+
+	// Same for all primitives
+	UPROPERTY()
+	TArray<FUintVector4> AttributeIdOffsetStrides;
+};
+
 UCLASS(Transient, Category = PCG)
 class UPCGDataBinding : public UObject
 {
@@ -22,8 +43,12 @@ public:
 
 	FPCGDataCollection OutputDataCollection;
 
-	const UPCGComputeGraph* Graph = nullptr;
+	UPROPERTY()
+	TObjectPtr<const UPCGComputeGraph> Graph = nullptr;
 
 	// Data collection pointer + set of pins that get data from the collection (cross barrier)
 	FPCGDataForGPU DataForGPU;
+
+	UPROPERTY()
+	TMap<TObjectPtr<const UPCGSettings>, FPCGSpawnerPrimitives> MeshSpawnersToPrimitives;
 };
diff --git a/Engine/Plugins/PCG/Source/PCG/Private/Compute/PCGDataForGPU.cpp b/Engine/Plugins/PCG/Source/PCG/Private/Compute/PCGDataForGPU.cpp
--- a/Engine/Plugins/PCG/Source/PCG/Private/Compute/PCGDataForGPU.cpp
+++ b/Engine/Plugins/PCG/Source/PCG/Private/Compute/PCGDataForGPU.cpp
@@ -6,6 +6,8 @@
 #include "PCGParamData.h"
 #include "PCGPoint.h"
 #include "Compute/PCGComputeCommon.h"
+#include "Compute/PCGComputeGraph.h"
+#include "Compute/PCGDataBinding.h"
 #include "Data/PCGPointData.h"
 #include "Helpers/PCGAsync.h"
 #include "Metadata/PCGMetadata.h"
@@ -383,6 +385,40 @@ namespace PCGDataForGPUHelpers
 
 		return true;
 	}
+
+	void ComputeCustomFloatPacking(
+		TArray<FName>& InAttributeNames,
+		const UPCGDataBinding* InBinding,
+		const FPCGDataCollectionDesc& InDataCollectionDescription,
+		uint32& OutCustomFloatCount,
+		TArray<FUint32Vector4>& OutAttributeIdOffsetStrides)
+	{
+		check(InBinding);
+		check(InBinding->Graph);
+		const TMap<FName, FPCGKernelAttributeIDAndType>& GlobalAttributeLookupTable = InBinding->Graph->GetAttributeLookupTable();
+
+		uint32 OffsetFloats = 0;
+
+		for (FName AttributeName : InAttributeNames)
+		{
+			const FPCGKernelAttributeIDAndType* FoundAttribute = GlobalAttributeLookupTable.Find(AttributeName);
+			if (!FoundAttribute)
+			{
+				ensure(false);
+				continue;
+			}
+
+			const uint32 AttributeId = static_cast<uint32>(FoundAttribute->Id);
+			const EPCGKernelAttributeType AttributeType = FoundAttribute->Type;
+			const uint32 StrideFloats = GetAttributeTypeStrideBytes(AttributeType) / sizeof(float);
+
+			OutAttributeIdOffsetStrides.Emplace(AttributeId, OffsetFloats, StrideFloats, /*Unused*/0);
+
+			OffsetFloats += StrideFloats;
+		}
+
+		OutCustomFloatCount = OffsetFloats;
+	}
 }
 
 bool FPCGKernelAttributeKey::operator==(const FPCGKernelAttributeKey& Other) const
@@ -863,7 +899,7 @@ EPCGUnpackDataCollectionResult FPCGDataCollectionDesc::UnpackDataCollection(cons
 	const uint32 NumPackedFloats = InPackedData.Num() / 4;
 
 	// Most significant bit of NumData is reserved to flag whether or not the kernel executed.
-	ensureAlwaysMsgf(DataAsUint[0] & PCGComputeConstants::KernelExecutedFlag, TEXT("Tried to unpack a GPU data collection, but the compute shader did not execute."));
+	ensureMsgf(DataAsUint[0] & PCGComputeConstants::KernelExecutedFlag, TEXT("Tried to unpack a GPU data collection, but the compute shader did not execute."));
 	const uint32 NumData = DataAsUint[0] & ~PCGComputeConstants::KernelExecutedFlag;
 
 	if (NumData != DataDescs.Num())
diff --git a/Engine/Plugins/PCG/Source/PCG/Private/Elements/PCGCopyPoints.cpp b/Engine/Plugins/PCG/Source/PCG/Private/Elements/PCGCopyPoints.cpp
--- a/Engine/Plugins/PCG/Source/PCG/Private/Elements/PCGCopyPoints.cpp
+++ b/Engine/Plugins/PCG/Source/PCG/Private/Elements/PCGCopyPoints.cpp
@@ -114,8 +114,11 @@ int UPCGCopyPointsSettings::ComputeKernelThreadCount(const UPCGDataBinding* Bind
 
 void UPCGCopyPointsSettings::CreateAdditionalInputDataInterfaces(TArray<TObjectPtr<UComputeDataInterface>>& OutDataInterfaces) const
 {
-	TObjectPtr<UPCGCopyPointsDataInterface> DataInterface = Cast<UPCGCopyPointsDataInterface>(OutDataInterfaces.Add_GetRef(NewObject<UPCGCopyPointsDataInterface>()));
-	DataInterface->Settings = this;
+	Super::CreateAdditionalInputDataInterfaces(OutDataInterfaces);
+
+	TObjectPtr<UPCGCopyPointsDataInterface> NodeDI = NewObject<UPCGCopyPointsDataInterface>();
+	NodeDI->Settings = this;
+	OutDataInterfaces.Add(NodeDI);
 }
 
 TArray<FPCGPinProperties> UPCGCopyPointsSettings::InputPinProperties() const
@@ -143,7 +146,7 @@ bool FPCGCopyPointsElement::ExecuteInternal(FPCGContext* Context) const
 	TRACE_CPUPROFILER_EVENT_SCOPE(FPCGCopyPointsElement::Execute);
 
 	const UPCGCopyPointsSettings* Settings = Context->GetInputSettings<UPCGCopyPointsSettings>();
-	check(Settings);
+	check(Settings && !Settings->ShouldExecuteOnGPU());
 
 	const EPCGCopyPointsInheritanceMode RotationInheritance = Settings->RotationInheritance;
 	const EPCGCopyPointsInheritanceMode ScaleInheritance = Settings->ScaleInheritance;
diff --git a/Engine/Plugins/PCG/Source/PCG/Private/Elements/PCGStaticMeshSpawner.cpp b/Engine/Plugins/PCG/Source/PCG/Private/Elements/PCGStaticMeshSpawner.cpp
--- a/Engine/Plugins/PCG/Source/PCG/Private/Elements/PCGStaticMeshSpawner.cpp
+++ b/Engine/Plugins/PCG/Source/PCG/Private/Elements/PCGStaticMeshSpawner.cpp
@@ -5,6 +5,9 @@
 #include "PCGComponent.h"
 #include "PCGCustomVersion.h"
 #include "PCGManagedResource.h"
+#include "Compute/PCGComputeCommon.h"
+#include "Compute/DataInterfaces/PCGInstanceDataInterface.h"
+#include "Compute/DataInterfaces/Elements/PCGStaticMeshSpawnerDataInterface.h"
 #include "Data/PCGPointData.h"
 #include "Data/PCGSpatialData.h"
 #include "Elements/PCGStaticMeshSpawnerContext.h"
@@ -41,6 +44,123 @@ UPCGStaticMeshSpawnerSettings::UPCGStaticMeshSpawnerSettings(const FObjectInitia
 	}
 }
 
+TCHAR const* UPCGStaticMeshSpawnerSettings::TemplateFilePath = TEXT("/Plugin/PCG/Private/Elements/PCGStaticMeshSpawner.usf");
+
+bool UPCGStaticMeshSpawnerSettings::IsKernelValid(FPCGContext* InContext, bool bQuiet) const
+{
+	if (!Super::IsKernelValid(InContext, bQuiet))
+	{
+		return false;
+	}
+
+	UPCGMeshSelectorWeighted* Selector = Cast<UPCGMeshSelectorWeighted>(MeshSelectorParameters);
+
+	if (!Selector)
+	{
+		PCG_KERNEL_VALIDATION_ERR(InContext, this, bQuiet, LOCTEXT("InvalidMeshSelector", "Currently GPU Static Mesh Spawner nodes must use PCGMeshSelectorWeighted as the mesh selector type."));
+		return false;
+	}
+
+	if (Selector->MeshEntries.IsEmpty())
+	{
+		PCG_KERNEL_VALIDATION_ERR(InContext, this, bQuiet, LOCTEXT("NoMeshEntries", "No meshes specified."));
+		return false;
+	}
+
+	for (const FPCGMeshSelectorWeightedEntry& Entry : Selector->MeshEntries)
+	{
+		UStaticMesh* StaticMesh = Entry.Descriptor.StaticMesh.LoadSynchronous();
+
+		if (!StaticMesh)
+		{
+			PCG_KERNEL_VALIDATION_ERR(InContext, this, bQuiet, LOCTEXT("UnassignedMesh", "Unassigned mesh."));
+			return false;
+		}
+
+		if (!StaticMesh->HasValidNaniteData())
+		{
+			PCG_KERNEL_VALIDATION_ERR(InContext, this, bQuiet, LOCTEXT("NonNaniteMesh", "Only Nanite meshes are currently supported by the GPU Static Mesh Spawner path."));
+			return false;
+		}
+	}
+
+	// Currently instance packers must be able to specify a full list of attribute names upfront, to build the attribute table at compile time.
+	// TODO: We should be able to augment a static attribute table with new attributes at execution time, which will allow other types like regex.
+	if (InstanceDataPackerParameters && !InstanceDataPackerParameters->GetAttributeNames(/*OutNames=*/nullptr))
+	{
+		PCG_KERNEL_VALIDATION_ERR(InContext, this, bQuiet, LOCTEXT("InvalidInstancePacker", "Selected instance packer does not support GPU execution."));
+		return false;
+	}
+
+	return true;
+}
+
+FString UPCGStaticMeshSpawnerSettings::GetCookedKernelSource(const TMap<FName, FPCGKernelAttributeIDAndType>& GlobalAttributeLookupTable) const
+{
+	FString TemplateFile;
+	ensure(LoadShaderSourceFile(TemplateFilePath, EShaderPlatform::SP_PCD3D_SM5, &TemplateFile, nullptr));
+
+	return TemplateFile;
+}
+
+FPCGDataCollectionDesc UPCGStaticMeshSpawnerSettings::ComputeOutputPinDataDesc(const UPCGPin* OutputPin, const UPCGDataBinding* Binding) const
+{
+	check(OutputPin);
+
+	const UPCGNode* Node = CastChecked<UPCGNode>(GetOuter());
+	const UPCGPin* InputPin = Node->GetInputPin(PCGPinConstants::DefaultInputLabel);
+
+	// Passthrough - output description matches input description.
+	return ComputeInputPinDataDesc(InputPin, Binding);
+}
+
+const TArray<FPCGKernelAttributeKey> UPCGStaticMeshSpawnerSettings::GetKernelAttributeKeys() const
+{
+	TArray<FPCGKernelAttributeKey> AttributeKeys;
+
+	if (InstanceDataPackerParameters)
+	{
+		TArray<FName> AttributeNames;
+		if (ensure(InstanceDataPackerParameters->GetAttributeNames(&AttributeNames)))
+		{
+			for (const FName& AttributeName : AttributeNames)
+			{
+				// We don't know the type statically before execution, leave unset.
+				AttributeKeys.Emplace(AttributeName, EPCGKernelAttributeType::None);
+			}
+		}
+	}
+
+	return AttributeKeys;
+}
+
+int UPCGStaticMeshSpawnerSettings::ComputeKernelThreadCount(const UPCGDataBinding* Binding) const
+{
+	const UPCGNode* Node = CastChecked<UPCGNode>(GetOuter());
+	const UPCGPin* Pin = Node->GetInputPin(PCGPinConstants::DefaultInputLabel);
+
+	return ComputeInputPinDataDesc(Pin, Binding).ComputeDataElementCount(EPCGDataType::Point);
+}
+
+void UPCGStaticMeshSpawnerSettings::CreateAdditionalInputDataInterfaces(TArray<TObjectPtr<UComputeDataInterface>>& OutDataInterfaces) const
+{
+	Super::CreateAdditionalInputDataInterfaces(OutDataInterfaces);
+
+	TObjectPtr<UPCGStaticMeshSpawnerDataInterface> NodeDI = NewObject<UPCGStaticMeshSpawnerDataInterface>();
+	NodeDI->Settings = this;
+	OutDataInterfaces.Add(NodeDI);
+}
+
+void UPCGStaticMeshSpawnerSettings::CreateAdditionalOutputDataInterfaces(TArray<TObjectPtr<UComputeDataInterface>>& OutDataInterfaces) const
+{
+	Super::CreateAdditionalOutputDataInterfaces(OutDataInterfaces);
+
+	UPCGInstanceDataInterface* InstanceDI = NewObject<UPCGInstanceDataInterface>();
+	InstanceDI->SetProducerSettings(this);
+	InstanceDI->InputPinProvidingData = PCGPinConstants::DefaultInputLabel;
+	OutDataInterfaces.Add(InstanceDI);
+}
+
 #if WITH_EDITOR
 FText UPCGStaticMeshSpawnerSettings::GetDefaultNodeTitle() const
 {
@@ -289,7 +409,7 @@ bool FPCGStaticMeshSpawnerElement::ExecuteInternal(FPCGContext* InContext) const
 	TRACE_CPUPROFILER_EVENT_SCOPE(FPCGStaticMeshSpawnerElement::Execute);
 	FPCGStaticMeshSpawnerContext* Context = static_cast<FPCGStaticMeshSpawnerContext*>(InContext);
 	const UPCGStaticMeshSpawnerSettings* Settings = Context->GetInputSettings<UPCGStaticMeshSpawnerSettings>();
-	check(Settings);
+	check(Settings && !Settings->ShouldExecuteOnGPU());
 
 	while(!Context->MeshInstancesData.IsEmpty())
 	{
diff --git a/Engine/Plugins/PCG/Source/PCG/Private/Graph/PCGGraphCompilerGPU.cpp b/Engine/Plugins/PCG/Source/PCG/Private/Graph/PCGGraphCompilerGPU.cpp
--- a/Engine/Plugins/PCG/Source/PCG/Private/Graph/PCGGraphCompilerGPU.cpp
+++ b/Engine/Plugins/PCG/Source/PCG/Private/Graph/PCGGraphCompilerGPU.cpp
@@ -14,6 +14,7 @@
 #include "Compute/Elements/PCGComputeGraphElement.h"
 #include "Compute/PCGComputeCommon.h"
 #include "Compute/PCGComputeKernelSource.h"
+#include "Elements/PCGStaticMeshSpawner.h"
 #include "Graph/PCGGraphCompiler.h"
 #include "Graph/PCGGraphExecutor.h"
 
@@ -24,7 +25,7 @@
 namespace PCGGraphCompilerGPU
 {
 #if WITH_EDITOR
-	static TAutoConsoleVariable<bool> CVarEnableGPUDebugging(
+	static TAutoConsoleVariable<bool> CVarEnableDebugging(
 		TEXT("pcg.GraphExecution.GPU.EnableDebugging"),
 		false,
 		TEXT("Enable verbose logging of GPU compilation and execution."));
@@ -705,6 +706,11 @@ void FPCGGraphCompilerGPU::BuildGPUGraphTask(
 		const UPCGSettings* Settings = Node ? Node->GetSettings() : nullptr;
 		check(Settings && Settings->bEnabled && Settings->ShouldExecuteOnGPU());
 
+		if (Settings->IsA<UPCGStaticMeshSpawnerSettings>())
+		{
+			ComputeGraph->StaticMeshSpawners.Add(Settings);
+		}
+
 		// One data interface may be connected to multiple downstream pins. Because each entry in these arrays produces its own bindings and graph edges,
 		// we also need to specify the pin label to uniquely describe each edge.
 		TArray<TPair</*DataInterfaceIndex=*/int, /*PinLabel=*/FName>> InputDataInterfaceIndexAndPin;
@@ -848,10 +854,6 @@ void FPCGGraphCompilerGPU::BuildGPUGraphTask(
 			}
 		}
 
-		// TODO add graph data interface (graph params). Reference: UOptimusGraphDataInterface.
-		//Element->Graph->DataInterfaces.Add(NewObject< UPCGDataCollectionDataInterface>(InGraph));
-		//Element->Graph->DataInterfaceToBinding.Add(0);
-
 		// TODO once we support cooking for different platforms/configs, don't create the interface if logging is not present.
 		if (Settings->bPrintShaderDebugValues)
 		{
@@ -944,7 +946,7 @@ void FPCGGraphCompilerGPU::BuildGPUGraphTask(
 			}
 
 #if WITH_EDITOR
-			if (PCGGraphCompilerGPU::CVarEnableGPUDebugging.GetValueOnAnyThread())
+			if (PCGGraphCompilerGPU::CVarEnableDebugging.GetValueOnAnyThread())
 			{
 				UE_LOG(LogPCG, Warning, TEXT("ATTRIBUTE LOOK-UP TABLE [%s]"), *Settings->GetDefaultNodeTitle().ToString());
 
diff --git a/Engine/Plugins/PCG/Source/PCG/Private/Helpers/PCGActorHelpers.cpp b/Engine/Plugins/PCG/Source/PCG/Private/Helpers/PCGActorHelpers.cpp
--- a/Engine/Plugins/PCG/Source/PCG/Private/Helpers/PCGActorHelpers.cpp
+++ b/Engine/Plugins/PCG/Source/PCG/Private/Helpers/PCGActorHelpers.cpp
@@ -103,6 +103,8 @@ UPCGManagedISMComponent* UPCGActorHelpers::GetOrCreateManagedISMC(AActor* InTarg
 
 	if (InParams.bAllowDescriptorChanges)
 	{
+		// It's potentially less efficient to put Nanite meshes inside of HISMs so decay those to ISM in this case.
+		// Note the equality here, not a IsA because we do not want to change derived types either
 		if (Descriptor.ComponentClass == UHierarchicalInstancedStaticMeshComponent::StaticClass())
 		{
 			// Done as in InstancedStaticMesh.cpp
@@ -165,8 +167,6 @@ UPCGManagedISMComponent* UPCGActorHelpers::GetOrCreateManagedISMC(AActor* InTarg
 
 	FString ComponentName;
 
-	// It's potentially less efficient to put nanite meshes inside of HISMs so decay those to ISM in this case.
-	// Note the equality here, not a IsA because we do not want to change derived types either
 	if (Descriptor.ComponentClass == UHierarchicalInstancedStaticMeshComponent::StaticClass())
 	{
 		ComponentName = TEXT("HISM_");
diff --git a/Engine/Plugins/PCG/Source/PCG/Private/InstanceDataPackers/PCGInstanceDataPackerByAttribute.cpp b/Engine/Plugins/PCG/Source/PCG/Private/InstanceDataPackers/PCGInstanceDataPackerByAttribute.cpp
--- a/Engine/Plugins/PCG/Source/PCG/Private/InstanceDataPackers/PCGInstanceDataPackerByAttribute.cpp
+++ b/Engine/Plugins/PCG/Source/PCG/Private/InstanceDataPackers/PCGInstanceDataPackerByAttribute.cpp
@@ -69,4 +69,17 @@ void UPCGInstanceDataPackerByAttribute::PackInstances_Implementation(FPCGContext
 	PackCustomDataFromAccessors(InstanceList, std::move(SelectedAccessors), std::move(SelectedKeys), OutPackedCustomData);
 }
 
+bool UPCGInstanceDataPackerByAttribute::GetAttributeNames(TArray<FName>* OutNames)
+{
+	if (OutNames)
+	{
+		for (const FPCGAttributePropertyInputSelector& AttributeSelector : AttributeSelectors)
+		{
+			OutNames->Add(AttributeSelector.GetAttributeName());
+		}
+	}
+
+	return true;
+}
+
 #undef LOCTEXT_NAMESPACE
diff --git a/Engine/Plugins/PCG/Source/PCG/Private/PCGElement.cpp b/Engine/Plugins/PCG/Source/PCG/Private/PCGElement.cpp
--- a/Engine/Plugins/PCG/Source/PCG/Private/PCGElement.cpp
+++ b/Engine/Plugins/PCG/Source/PCG/Private/PCGElement.cpp
@@ -264,10 +264,20 @@ void IPCGElement::PreExecute(FPCGContext* Context) const
 	// Will override the settings if there is any override.
 	Context->OverrideSettings();
 
-	if (CVarPCGAllowPerDataCaching.GetValueOnAnyThread())
+	const UPCGSettings* Settings = Context->GetInputSettings<UPCGSettings>();
+
+	// If we were supposed to execute on GPU and end up here, then GPU compilation failed. Run validation in noisy mode
+	// to generate runtime errors/graph errors, then pass through.
+	if (Settings && Settings->ShouldExecuteOnGPU())
 	{
-		const UPCGSettings* Settings = Context->GetInputSettings<UPCGSettings>();
+		Settings->IsKernelValid(Context, /*bQuiet=*/false);
+		DisabledPassThroughData(Context);
+		Context->CurrentPhase = EPCGExecutionPhase::PostExecute;
+		return;
+	}
 
+	if (CVarPCGAllowPerDataCaching.GetValueOnAnyThread())
+	{
 		// Default implementation when the entries in a primary loop can be processed independently, e.g. they can appear in the cache separately
 		// Implementation note: this supposes that the current node has only ONE required pin, and not multiple.
 		// For more complex cases (such as multiple required pins, whether cartesian or matching), the implementation should use common code instead to streamline this process -
diff --git a/Engine/Plugins/PCG/Source/PCG/Public/Compute/PCGDataForGPU.h b/Engine/Plugins/PCG/Source/PCG/Public/Compute/PCGDataForGPU.h
--- a/Engine/Plugins/PCG/Source/PCG/Public/Compute/PCGDataForGPU.h
+++ b/Engine/Plugins/PCG/Source/PCG/Public/Compute/PCGDataForGPU.h
@@ -7,6 +7,7 @@
 
 #include "PCGDataForGPU.generated.h"
 
+class UPCGDataBinding;
 class UPCGMetadata;
 class UPCGPin;
 class UPCGSettings;
@@ -147,5 +148,14 @@ struct FPCGDataForGPU
 
 namespace PCGDataForGPUHelpers
 {
+	/** Returns GPU type that will be used to represent the given metadata type. */
 	EPCGKernelAttributeType GetAttributeTypeFromMetadataType(EPCGMetadataTypes MetadataType);
+
+	/** Compute how attributes can be packed to custom floats. */
+	void ComputeCustomFloatPacking(
+		TArray<FName>& InAttributeNames,
+		const UPCGDataBinding* InBinding,
+		const FPCGDataCollectionDesc& InDataCollectionDescription,
+		uint32& OutCustomFloatCount,
+		TArray<FUint32Vector4>& OutAttributeIdOffsetStrides);
 }
diff --git a/Engine/Plugins/PCG/Source/PCG/Public/Elements/PCGStaticMeshSpawner.h b/Engine/Plugins/PCG/Source/PCG/Public/Elements/PCGStaticMeshSpawner.h
--- a/Engine/Plugins/PCG/Source/PCG/Public/Elements/PCGStaticMeshSpawner.h
+++ b/Engine/Plugins/PCG/Source/PCG/Public/Elements/PCGStaticMeshSpawner.h
@@ -28,8 +28,17 @@ class PCG_API UPCGStaticMeshSpawnerSettings : public UPCGSettings
 public:
 	UPCGStaticMeshSpawnerSettings(const FObjectInitializer &ObjectInitializer);
 
-#if WITH_EDITOR
 	// ~Begin UPCGSettings interface
+	virtual bool IsKernelValid(FPCGContext* InContext = nullptr, bool bQuiet = true) const;
+	virtual FString GetCookedKernelSource(const TMap<FName, FPCGKernelAttributeIDAndType>& GlobalAttributeLookupTable) const override;
+	virtual const TArray<FPCGKernelAttributeKey> GetKernelAttributeKeys() const override;
+	virtual int ComputeKernelThreadCount(const UPCGDataBinding* Binding) const override;
+	virtual FPCGDataCollectionDesc ComputeOutputPinDataDesc(const UPCGPin* OutputPin, const UPCGDataBinding* Binding) const override;
+	virtual void CreateAdditionalInputDataInterfaces(TArray<TObjectPtr<UComputeDataInterface>>& OutDataInterfaces) const override;
+	virtual void CreateAdditionalOutputDataInterfaces(TArray<TObjectPtr<UComputeDataInterface>>& OutDataInterfaces) const override;
+#if WITH_EDITOR
+
+	virtual bool DisplayExecuteOnGPUSetting() const override { return true; }
 	virtual FName GetDefaultNodeName() const override { return FName(TEXT("StaticMeshSpawner")); }
 	virtual FText GetDefaultNodeTitle() const override;
 	virtual EPCGSettingsType GetType() const override { return EPCGSettingsType::Spawner; }
@@ -116,6 +125,9 @@ public:
 protected:
 	void RefreshMeshSelector();
 	void RefreshInstancePacker();
+
+private:
+	static TCHAR const* TemplateFilePath;
 };
 
 class PCG_API FPCGStaticMeshSpawnerElement : public IPCGElement
diff --git a/Engine/Plugins/PCG/Source/PCG/Public/InstanceDataPackers/PCGInstanceDataPackerBase.h b/Engine/Plugins/PCG/Source/PCG/Public/InstanceDataPackers/PCGInstanceDataPackerBase.h
--- a/Engine/Plugins/PCG/Source/PCG/Public/InstanceDataPackers/PCGInstanceDataPackerBase.h
+++ b/Engine/Plugins/PCG/Source/PCG/Public/InstanceDataPackers/PCGInstanceDataPackerBase.h
@@ -49,6 +49,9 @@ public:
 
 	/** Build a PackedCustomData by processing each accessor in order for each point in the InstanceList */
 	void PackCustomDataFromAccessors(const FPCGMeshInstanceList& InstanceList, TArray<TUniquePtr<const IPCGAttributeAccessor>> Accessors, TArray<TUniquePtr<const IPCGAttributeAccessorKeys>> AccessorKeys, FPCGPackedCustomData& OutPackedCustomData) const;
+
+	/** If OutNames is not null, returns a list of all attributes that will be packed. Returns true if this list can be statically determined (prior to execution). */
+	virtual bool GetAttributeNames(TArray<FName>* OutNames) { return false; }
 };
 
 #if UE_ENABLE_INCLUDE_ORDER_DEPRECATED_IN_5_2
diff --git a/Engine/Plugins/PCG/Source/PCG/Public/InstanceDataPackers/PCGInstanceDataPackerByAttribute.h b/Engine/Plugins/PCG/Source/PCG/Public/InstanceDataPackers/PCGInstanceDataPackerByAttribute.h
--- a/Engine/Plugins/PCG/Source/PCG/Public/InstanceDataPackers/PCGInstanceDataPackerByAttribute.h
+++ b/Engine/Plugins/PCG/Source/PCG/Public/InstanceDataPackers/PCGInstanceDataPackerByAttribute.h
@@ -18,8 +18,8 @@ public:
 	virtual void PostLoad() override;
 	// ~End UObject interface
 #endif
-
 	virtual void PackInstances_Implementation(UPARAM(ref) FPCGContext& Context, const UPCGSpatialData* InSpatialData, UPARAM(ref) const FPCGMeshInstanceList& InstanceList, FPCGPackedCustomData& OutPackedCustomData) const override;
+	virtual bool GetAttributeNames(TArray<FName>* OutNames) override;
 
 public:
 	UPROPERTY(BlueprintReadWrite, EditAnywhere, Category = InstanceDataPacker, meta = (PCG_DiscardPropertySelection))
diff --git a/Engine/Plugins/PCG/Source/PCGCompute/PCGCompute.Build.cs b/Engine/Plugins/PCG/Source/PCGCompute/PCGCompute.Build.cs
--- a/Engine/Plugins/PCG/Source/PCGCompute/PCGCompute.Build.cs
+++ b/Engine/Plugins/PCG/Source/PCGCompute/PCGCompute.Build.cs
@@ -17,6 +17,12 @@ namespace UnrealBuildTool.Rules
 					"RenderCore",
 					"Renderer",
 				});
+
+			PrivateDependencyModuleNames.AddRange(
+				new string[]
+				{
+					"Engine",
+				});
 		}
 	}
 }
diff --git a/Engine/Plugins/PCG/Source/PCGCompute/Private/PCGSceneWriterCS.cpp b/Engine/Plugins/PCG/Source/PCGCompute/Private/PCGSceneWriterCS.cpp
new file mode 100644
--- /dev/null
+++ b/Engine/Plugins/PCG/Source/PCGCompute/Private/PCGSceneWriterCS.cpp
@@ -0,0 +1,14 @@
+// Copyright Epic Games, Inc. All Rights Reserved.
+
+#include "PCGSceneWriterCS.h"
+
+void FPCGSceneWriterCS::ModifyCompilationEnvironment(const FGlobalShaderPermutationParameters& Parameters, FShaderCompilerEnvironment& OutEnvironment)
+{
+	FGlobalShader::ModifyCompilationEnvironment(Parameters, OutEnvironment);
+
+	OutEnvironment.SetDefine(TEXT("VF_SUPPORTS_PRIMITIVE_SCENE_DATA"), 1); // Fixes errors that Primitive uniform buffer unbound.
+	OutEnvironment.SetDefine(TEXT("USE_GLOBAL_GPU_SCENE_DATA_RW"), 1); // Directs shader to use GPUSceneFrameNumber from global scope.
+	OutEnvironment.SetDefine(TEXT("NUM_THREADS_PER_GROUP"), NumThreadsPerGroup);
+}
+
+IMPLEMENT_GLOBAL_SHADER(FPCGSceneWriterCS, "/PCGComputeShaders/PCGSceneWriter.usf", "SceneWriter_CS", SF_Compute);
diff --git a/Engine/Plugins/PCG/Source/PCGCompute/Public/PCGSceneWriterCS.h b/Engine/Plugins/PCG/Source/PCGCompute/Public/PCGSceneWriterCS.h
new file mode 100644
--- /dev/null
+++ b/Engine/Plugins/PCG/Source/PCGCompute/Public/PCGSceneWriterCS.h
@@ -0,0 +1,32 @@
+// Copyright Epic Games, Inc. All Rights Reserved.
+
+#pragma once
+
+#include "GlobalShader.h"
+#include "GPUSceneWriter.h"
+#include "ShaderParameterStruct.h"
+
+class FPCGSceneWriterCS : public FGlobalShader
+{
+public:
+	DECLARE_EXPORTED_GLOBAL_SHADER(FPCGSceneWriterCS, PCGCOMPUTE_API);
+	SHADER_USE_PARAMETER_STRUCT(FPCGSceneWriterCS, FGlobalShader);
+
+	BEGIN_SHADER_PARAMETER_STRUCT(FParameters, PCGCOMPUTE_API)
+		SHADER_PARAMETER(uint32, InPrimitiveIndex)
+		SHADER_PARAMETER(uint32, InNumInstancesAllocatedInGPUScene)
+		SHADER_PARAMETER(uint32, InInstanceOffset)
+		SHADER_PARAMETER_RDG_BUFFER_SRV(StructuredBuffer<FUintVector4>, InInstanceData)
+		SHADER_PARAMETER_RDG_BUFFER_SRV(StructuredBuffer<uint32>, InInstanceCustomFloatData)
+		SHADER_PARAMETER_RDG_BUFFER_SRV(StructuredBuffer<uint32>, InWriteCounters)
+		SHADER_PARAMETER(uint32, InPrimitiveId)
+		SHADER_PARAMETER(uint32, InCustomDataCount)
+		SHADER_PARAMETER(uint32, InPayloadDataFlags)
+		SHADER_PARAMETER_STRUCT_INCLUDE(FGPUSceneWriterParameters, GPUSceneWriterParameters)
+	END_SHADER_PARAMETER_STRUCT()
+
+public:
+	static constexpr uint32 NumThreadsPerGroup = 64;
+
+	static void ModifyCompilationEnvironment(const FGlobalShaderPermutationParameters& Parameters, FShaderCompilerEnvironment& OutEnvironment);
+};
diff --git a/Engine/Source/Runtime/Engine/Public/PrimitiveSceneProxy.h b/Engine/Source/Runtime/Engine/Public/PrimitiveSceneProxy.h
--- a/Engine/Source/Runtime/Engine/Public/PrimitiveSceneProxy.h
+++ b/Engine/Source/Runtime/Engine/Public/PrimitiveSceneProxy.h
@@ -1161,7 +1161,7 @@ public:
 
 	/**
 	 */
-	FInstanceDataBufferHeader GetInstanceDataHeader() const;
+	ENGINE_API FInstanceDataBufferHeader GetInstanceDataHeader() const;
 
 #if UE_WITH_PSO_PRECACHING
 	ENGINE_API void BoostPrecachedPSORequestsOnDraw();
