[PCG] Static analysis - remove misleading and unnecessary pointer guard.

#rnx

diff --git a/Engine/Plugins/PCG/Source/PCG/Private/Components/PCGProceduralISMComponent.cpp b/Engine/Plugins/PCG/Source/PCG/Private/Components/PCGProceduralISMComponent.cpp
--- a/Engine/Plugins/PCG/Source/PCG/Private/Components/PCGProceduralISMComponent.cpp
+++ b/Engine/Plugins/PCG/Source/PCG/Private/Components/PCGProceduralISMComponent.cpp
@@ -99,11 +99,8 @@ UPCGManagedProceduralISMComponent* PCGManagedProceduralISMComponent::GetOrCreate
 	Resource->SetSettingsUID(InSettingsUID);
 	InSourceComponent->AddToManagedResources(Resource);
 
-	if (ISMC)
-	{
-		ISMC->ComponentTags.AddUnique(PCGHelpers::DefaultPCGTag);
-		ISMC->ComponentTags.AddUnique(InSourceComponent->GetFName());
-	}
+	ISMC->ComponentTags.AddUnique(PCGHelpers::DefaultPCGTag);
+	ISMC->ComponentTags.AddUnique(InSourceComponent->GetFName());
 
 	return Resource;
 }
