[PCG] Move Editor visualizations into public API

* Necessary for Dynamic Mesh visualisation

#rnx

diff --git a/Engine/Plugins/PCG/Source/PCGEditor/Private/DataVisualizations/PCGDataVisualizationHelpers.h b/Engine/Plugins/PCG/Source/PCGEditor/Private/DataVisualizations/PCGDataVisualizationHelpers.h
deleted file mode 100644
diff --git a/Engine/Plugins/PCG/Source/PCGEditor/Private/DataVisualizations/PCGParamDataVisualization.h b/Engine/Plugins/PCG/Source/PCGEditor/Private/DataVisualizations/PCGParamDataVisualization.h
deleted file mode 100644
diff --git a/Engine/Plugins/PCG/Source/PCGEditor/Private/DataVisualizations/PCGSpatialDataVisualization.h b/Engine/Plugins/PCG/Source/PCGEditor/Private/DataVisualizations/PCGSpatialDataVisualization.h
deleted file mode 100644
diff --git a/Engine/Plugins/PCG/Source/PCGEditor/Private/DataVisualizations/PCGSplineDataVisualization.h b/Engine/Plugins/PCG/Source/PCGEditor/Private/DataVisualizations/PCGSplineDataVisualization.h
deleted file mode 100644
diff --git a/Engine/Plugins/PCG/Source/PCGEditor/Private/PCGEditorModule.h b/Engine/Plugins/PCG/Source/PCGEditor/Private/PCGEditorModule.h
deleted file mode 100644
diff --git a/Engine/Plugins/PCG/Source/PCGEditor/Public/DataVisualizations/PCGDataVisualizationHelpers.h b/Engine/Plugins/PCG/Source/PCGEditor/Public/DataVisualizations/PCGDataVisualizationHelpers.h
new file mode 100644
--- /dev/null
+++ b/Engine/Plugins/PCG/Source/PCGEditor/Public/DataVisualizations/PCGDataVisualizationHelpers.h
@@ -0,0 +1,13 @@
+// Copyright Epic Games, Inc. All Rights Reserved.
+
+#pragma once
+
+#include "PCGDataVisualization.h"
+
+class UPCGData;
+
+namespace PCGDataVisualizationHelpers
+{
+	FPCGTableVisualizerColumnInfo CreateColumnInfo(const UPCGData* Data, FName Id, const FText& Label, TFunction<TSharedPtr<const IPCGAttributeAccessor>()> CreateAccessorFunc = nullptr);
+	void CreateMetadataColumnInfos(const UPCGData* Data, TArray<FPCGTableVisualizerColumnInfo>& OutColumnInfos);
+}
diff --git a/Engine/Plugins/PCG/Source/PCGEditor/Public/DataVisualizations/PCGParamDataVisualization.h b/Engine/Plugins/PCG/Source/PCGEditor/Public/DataVisualizations/PCGParamDataVisualization.h
new file mode 100644
--- /dev/null
+++ b/Engine/Plugins/PCG/Source/PCGEditor/Public/DataVisualizations/PCGParamDataVisualization.h
@@ -0,0 +1,14 @@
+// Copyright Epic Games, Inc. All Rights Reserved.
+
+#pragma once
+
+#include "PCGDataVisualization.h"
+
+class UPCGData;
+
+class PCGEDITOR_API IPCGParamDataVisualization : public IPCGDataVisualization
+{
+public:
+	virtual void ExecuteDebugDisplay(FPCGContext* Context, const UPCGSettingsInterface* SettingsInterface, const UPCGData* Data, AActor* TargetActor) const { /* Do nothing. */ }
+	virtual FPCGTableVisualizerInfo GetTableVisualizerInfo(const UPCGData* Data) const override;
+};
diff --git a/Engine/Plugins/PCG/Source/PCGEditor/Public/DataVisualizations/PCGSpatialDataVisualization.h b/Engine/Plugins/PCG/Source/PCGEditor/Public/DataVisualizations/PCGSpatialDataVisualization.h
new file mode 100644
--- /dev/null
+++ b/Engine/Plugins/PCG/Source/PCGEditor/Public/DataVisualizations/PCGSpatialDataVisualization.h
@@ -0,0 +1,23 @@
+// Copyright Epic Games, Inc. All Rights Reserved.
+
+#pragma once
+
+#include "PCGDataVisualization.h"
+#include "PCGEditorModule.h"
+
+class AActor;
+class UPCGData;
+class UPCGPointData;
+struct FPCGContext;
+
+/** Default implementation for spatial data. Collapses to a PointData representation. */
+class PCGEDITOR_API IPCGSpatialDataVisualization : public IPCGDataVisualization
+{
+public:
+	// ~Begin IPCGDataVisualization interface
+	virtual void ExecuteDebugDisplay(FPCGContext* Context, const UPCGSettingsInterface* SettingsInterface, const UPCGData* Data, AActor* TargetActor) const override;
+	virtual FPCGTableVisualizerInfo GetTableVisualizerInfo(const UPCGData* Data) const override;
+	// ~End IPCGDataVisualization interface
+
+	virtual const UPCGPointData* CollapseToDebugPointData(FPCGContext* Context, const UPCGData* Data) const;
+};
diff --git a/Engine/Plugins/PCG/Source/PCGEditor/Public/DataVisualizations/PCGSplineDataVisualization.h b/Engine/Plugins/PCG/Source/PCGEditor/Public/DataVisualizations/PCGSplineDataVisualization.h
new file mode 100644
--- /dev/null
+++ b/Engine/Plugins/PCG/Source/PCGEditor/Public/DataVisualizations/PCGSplineDataVisualization.h
@@ -0,0 +1,23 @@
+// Copyright Epic Games, Inc. All Rights Reserved.
+
+#pragma once
+
+#include "PCGSpatialDataVisualization.h"
+#include "PCGEditorModule.h"
+
+class UPCGData;
+class UPCGPointData;
+struct FPCGContext;
+
+class PCGEDITOR_API IPCGSplineDataVisualization : public IPCGSpatialDataVisualization
+{
+public:
+	// ~Begin IPCGDataVisualization interface
+	virtual FPCGTableVisualizerInfo GetTableVisualizerInfo(const UPCGData* Data) const override;
+	// ~End IPCGDataVisualization interface
+
+	// ~Begin IPCGSpatialDataVisualization interface
+	/** Overrides collapse behavior to show the spline control points. */
+	virtual const UPCGPointData* CollapseToDebugPointData(FPCGContext* Context, const UPCGData* Data) const override;
+	// ~End IPCGSpatialDataVisualization interface
+};
diff --git a/Engine/Plugins/PCG/Source/PCGEditor/Public/PCGEditorModule.h b/Engine/Plugins/PCG/Source/PCGEditor/Public/PCGEditorModule.h
new file mode 100644
--- /dev/null
+++ b/Engine/Plugins/PCG/Source/PCGEditor/Public/PCGEditorModule.h
@@ -0,0 +1,62 @@
+// Copyright Epic Games, Inc. All Rights Reserved.
+
+#pragma once
+
+#include "Editor/IPCGEditorModule.h"
+
+#include "AssetTypeCategories.h"
+#include "Modules/ModuleInterface.h"
+#include "Toolkits/AssetEditorToolkit.h"
+
+class ILevelEditor;
+// Logs
+DECLARE_LOG_CATEGORY_EXTERN(LogPCGEditor, Log, All);
+
+class FMenuBuilder;
+class FPCGEditorGraphNodeFactory;
+class IAssetTypeActions;
+
+class FPCGEditorModule : public IPCGEditorModule
+{
+public:
+	// ~IModuleInterface implementation
+	virtual void StartupModule() override;
+	virtual void ShutdownModule() override;
+	virtual bool SupportsDynamicReloading() override;
+	// ~End IModuleInterface implementation
+
+	// ~IPCGEditorModule implementation
+	virtual TWeakPtr<IPCGEditorProgressNotification> CreateProgressNotification(const FTextFormat& TextFormat, bool bCanCancel) override;
+	virtual void ReleaseProgressNotification(TWeakPtr<IPCGEditorProgressNotification> InNotification) override;
+	virtual void SetOutlinerUIRefreshDelay(float InDelay) override;
+	// ~End IPCGEditorModule implementation
+
+protected:
+	void RegisterDetailsCustomizations();
+	void UnregisterDetailsCustomizations();
+	void RegisterMenuExtensions();
+	void UnregisterMenuExtensions();
+	void PopulateMenuActions(FMenuBuilder& MenuBuilder);
+	void RegisterSettings();
+	void UnregisterSettings();
+	void RegisterPCGDataVisualizations();
+	void UnregisterPCGDataVisualizations();
+
+	void OnLevelEditorCreated(TSharedPtr<ILevelEditor> InLevelEditor);
+	void RegisterOnEditorModeChange();
+	void OnEditorModeIDChanged(const FEditorModeID& EditorModeID, bool bIsEntering);
+
+	bool ShouldDisableCPUThrottling();
+
+	TArray<TSharedRef<IAssetTypeActions>> RegisteredAssetTypeActions;
+	static EAssetTypeCategories::Type PCGAssetCategory;
+
+	TSharedPtr<FExtensibilityManager> MenuExtensibilityManager;
+	TSharedPtr<FExtensibilityManager> ToolBarExtensibilityManager;
+
+	TSharedPtr<FPCGEditorGraphNodeFactory> GraphNodeFactory;
+
+	FDelegateHandle ShouldDisableCPUThrottlingDelegateHandle;
+
+	TSet<TSharedPtr<IPCGEditorProgressNotification>> ActiveNotifications;
+};
