[PCG] Fix crash when compiling graph after closing PCGEditor window.

#rnx

diff --git a/Engine/Plugins/PCG/Source/PCGEditor/Private/PCGEditor.cpp b/Engine/Plugins/PCG/Source/PCGEditor/Private/PCGEditor.cpp
--- a/Engine/Plugins/PCG/Source/PCGEditor/Private/PCGEditor.cpp
+++ b/Engine/Plugins/PCG/Source/PCGEditor/Private/PCGEditor.cpp
@@ -2913,6 +2913,7 @@ void FPCGEditor::OnClose()
 	if (PCGGraphBeingEdited)
 	{
 		PCGGraphBeingEdited->OnGraphChangedDelegate.RemoveAll(this);
+		PCGGraphBeingEdited->OnNodeSourceCompiledDelegate.RemoveAll(this);
 
 		if (PCGGraphBeingEdited->IsInspecting())
 		{
