[PCG] Fix ensure when a CPU version of one of our GPU-enabled built-in nodes is used just upstream of a GPU node.

We should not ever be recursively computing data descriptions across CPU nodes.

#rnx

diff --git a/Engine/Plugins/PCG/Source/PCG/Private/Elements/PCGCopyPoints.cpp b/Engine/Plugins/PCG/Source/PCG/Private/Elements/PCGCopyPoints.cpp
--- a/Engine/Plugins/PCG/Source/PCG/Private/Elements/PCGCopyPoints.cpp
+++ b/Engine/Plugins/PCG/Source/PCG/Private/Elements/PCGCopyPoints.cpp
@@ -34,6 +34,11 @@ FString UPCGCopyPointsSettings::GetCookedKernelSource(const TMap<FName, FPCGKern
 
 FPCGDataCollectionDesc UPCGCopyPointsSettings::ComputeOutputPinDataDesc(const UPCGPin* OutputPin, const UPCGDataBinding* Binding) const
 {
+	if (!ShouldExecuteOnGPU())
+	{
+		return Super::ComputeOutputPinDataDesc(OutputPin, Binding);
+	}
+
 	check(OutputPin);
 
 	const UPCGNode* Node = CastChecked<UPCGNode>(GetOuter());
diff --git a/Engine/Plugins/PCG/Source/PCG/Private/Elements/PCGStaticMeshSpawner.cpp b/Engine/Plugins/PCG/Source/PCG/Private/Elements/PCGStaticMeshSpawner.cpp
--- a/Engine/Plugins/PCG/Source/PCG/Private/Elements/PCGStaticMeshSpawner.cpp
+++ b/Engine/Plugins/PCG/Source/PCG/Private/Elements/PCGStaticMeshSpawner.cpp
@@ -105,6 +105,11 @@ FString UPCGStaticMeshSpawnerSettings::GetCookedKernelSource(const TMap<FName, F
 
 FPCGDataCollectionDesc UPCGStaticMeshSpawnerSettings::ComputeOutputPinDataDesc(const UPCGPin* OutputPin, const UPCGDataBinding* Binding) const
 {
+	if (!ShouldExecuteOnGPU())
+	{
+		return Super::ComputeOutputPinDataDesc(OutputPin, Binding);
+	}
+
 	check(OutputPin);
 
 	const UPCGNode* Node = CastChecked<UPCGNode>(GetOuter());
diff --git a/Engine/Plugins/PCG/Source/PCG/Public/PCGSettings.h b/Engine/Plugins/PCG/Source/PCG/Public/PCGSettings.h
--- a/Engine/Plugins/PCG/Source/PCG/Public/PCGSettings.h
+++ b/Engine/Plugins/PCG/Source/PCG/Public/PCGSettings.h
@@ -514,7 +514,7 @@ protected:
 	// GPU section
 public:
 	/** Whether this node should be executed on the GPU. */
-	virtual bool ShouldExecuteOnGPU() const { return bExecuteOnGPU; }
+	bool ShouldExecuteOnGPU() const { return bExecuteOnGPU; }
 
 	virtual bool IsKernelValid(FPCGContext* InContext = nullptr, bool bQuiet = true) const;
 	virtual FString GetCookedKernelSource(const TMap<FName, FPCGKernelAttributeIDAndType>& GlobalAttributeLookupTable) const { return TEXT(""); }
