[PCG] Add int32 accessors for properties

* By default it was always converting to int64 for safety measures, but we have more mecanism to auto cast now, so it should be safe

#rnx

diff --git a/Engine/Plugins/PCG/Source/PCG/Private/Metadata/Accessors/PCGAttributeAccessorHelpers.cpp b/Engine/Plugins/PCG/Source/PCG/Private/Metadata/Accessors/PCGAttributeAccessorHelpers.cpp
--- a/Engine/Plugins/PCG/Source/PCG/Private/Metadata/Accessors/PCGAttributeAccessorHelpers.cpp
+++ b/Engine/Plugins/PCG/Source/PCG/Private/Metadata/Accessors/PCGAttributeAccessorHelpers.cpp
@@ -287,11 +287,20 @@ namespace PCGAttributeAccessorHelpers
 		{
 			if (NumericProperty->IsFloatingPoint())
 			{
+				// As floating properties are mostly all double in UE, convert float to double attributes.
 				return Functor(Signature<FPCGNumericPropertyAccessor<double>>{}, NumericProperty);
 			}
 			else if (NumericProperty->IsInteger())
 			{
-				return Functor(Signature<FPCGNumericPropertyAccessor<int64>>{}, NumericProperty);
+				// But for int32/int64 we can distinguish between the two. Everything of size 32 or less is will be an int32, 64bits integers will be int64.
+				if (NumericProperty->IsA<FInt64Property>() || NumericProperty->IsA<FUInt64Property>())
+				{
+					return Functor(Signature<FPCGNumericPropertyAccessor<int64>>{}, NumericProperty);
+				}
+				else
+				{
+					return Functor(Signature<FPCGNumericPropertyAccessor<int32>>{}, NumericProperty);
+				}
 			}
 		}
 		else if (const FBoolProperty* BoolProperty = CastField<FBoolProperty>(InProperty))
diff --git a/Engine/Plugins/PCG/Source/PCG/Public/Metadata/Accessors/PCGPropertyAccessor.h b/Engine/Plugins/PCG/Source/PCG/Public/Metadata/Accessors/PCGPropertyAccessor.h
--- a/Engine/Plugins/PCG/Source/PCG/Public/Metadata/Accessors/PCGPropertyAccessor.h
+++ b/Engine/Plugins/PCG/Source/PCG/Public/Metadata/Accessors/PCGPropertyAccessor.h
@@ -145,7 +145,7 @@ public:
 			{
 				if constexpr (std::is_integral_v<T>)
 				{
-					Property->SetIntPropertyValue(PropertyAddressData, Value);
+					Property->SetIntPropertyValue(PropertyAddressData, static_cast<int64>(Value));
 				}
 				else
 				{
