[PCG] Update tests after changes from property accessors

#rnx

diff --git a/Engine/Plugins/PCG/Source/PCG/Private/Tests/Accessors/PCGAttributeExtractorTest.cpp b/Engine/Plugins/PCG/Source/PCG/Private/Tests/Accessors/PCGAttributeExtractorTest.cpp
--- a/Engine/Plugins/PCG/Source/PCG/Private/Tests/Accessors/PCGAttributeExtractorTest.cpp
+++ b/Engine/Plugins/PCG/Source/PCG/Private/Tests/Accessors/PCGAttributeExtractorTest.cpp
@@ -102,7 +102,7 @@ bool FPCGAttributeSingleGetPropertyTest::RunTest(const FString& Parameters)
 
 	bTestPassed &= Verify(double(Struct.FloatValue), PCG_STRUCT_NAME_CHECK(FloatValue));
 	bTestPassed &= Verify(Struct.DoubleValue, PCG_STRUCT_NAME_CHECK(DoubleValue));
-	bTestPassed &= Verify(int64(Struct.Int32Value), PCG_STRUCT_NAME_CHECK(Int32Value));
+	bTestPassed &= Verify(Struct.Int32Value, PCG_STRUCT_NAME_CHECK(Int32Value));
 	bTestPassed &= Verify(Struct.IntValue, PCG_STRUCT_NAME_CHECK(IntValue));
 	bTestPassed &= Verify(Struct.Vector2Value, PCG_STRUCT_NAME_CHECK(Vector2Value));
 	bTestPassed &= Verify(Struct.VectorValue, PCG_STRUCT_NAME_CHECK(VectorValue));
diff --git a/Engine/Plugins/PCG/Source/PCG/Private/Tests/Elements/Metadata/PCGMetadataMathsOpTest.cpp b/Engine/Plugins/PCG/Source/PCG/Private/Tests/Elements/Metadata/PCGMetadataMathsOpTest.cpp
--- a/Engine/Plugins/PCG/Source/PCG/Private/Tests/Elements/Metadata/PCGMetadataMathsOpTest.cpp
+++ b/Engine/Plugins/PCG/Source/PCG/Private/Tests/Elements/Metadata/PCGMetadataMathsOpTest.cpp
@@ -204,7 +204,7 @@ bool FPCGMetadataMathsOp_OneMinus::RunTest(const FString& Parameters)
 	// Point Data
 	FPCGAttributePropertyInputSelector Selector;
 	Selector.SetPointProperty(EPCGPointProperties::Seed);
-	bSuccess &= ExecuteTest(TestParams(OneMinus, GeneratePointData(), TArray<int64>({0}), TArray<int64>({0})), Selector);
+	bSuccess &= ExecuteTest(TestParams(OneMinus, GeneratePointData(), TArray<int32>({0}), TArray<int32>({0})), Selector);
 	Selector.SetPointProperty(EPCGPointProperties::Density);
 	bSuccess &= ExecuteTest(TestParams(OneMinus, GeneratePointData(), TArray({0.0}), TArray({0.0})), Selector);
 	Selector.SetPointProperty(EPCGPointProperties::Position);
@@ -235,7 +235,7 @@ bool FPCGMetadataMathsOp_Sign::RunTest(const FString& Parameters)
 	// Point Data
 	FPCGAttributePropertyInputSelector Selector;
 	Selector.SetPointProperty(EPCGPointProperties::Seed);
-	bSuccess &= ExecuteTest(TestParams(Sign, GeneratePointData(), TArray<int64>({0}), TArray<int64>({1})), Selector);
+	bSuccess &= ExecuteTest(TestParams(Sign, GeneratePointData(), TArray<int32>({0}), TArray<int32>({1})), Selector);
 	Selector.SetPointProperty(EPCGPointProperties::Density);
 	bSuccess &= ExecuteTest(TestParams(Sign, GeneratePointData(), TArray({0.0}), TArray({1.0})), Selector);
 	Selector.SetPointProperty(EPCGPointProperties::Position);
@@ -264,7 +264,7 @@ bool FPCGMetadataMathsOp_Frac::RunTest(const FString& Parameters)
 	// Point Data
 	FPCGAttributePropertyInputSelector Selector;
 	Selector.SetPointProperty(EPCGPointProperties::Seed);
-	bSuccess &= ExecuteTest(TestParams(Frac, GeneratePointData(), TArray<int64>({0}), TArray<int64>({0})), Selector);
+	bSuccess &= ExecuteTest(TestParams(Frac, GeneratePointData(), TArray<int32>({0}), TArray<int32>({0})), Selector);
 	Selector.SetPointProperty(EPCGPointProperties::Density);
 	bSuccess &= ExecuteTest(TestParams(Frac, GeneratePointData(), TArray({0.0}), TArray({0.0})), Selector);
 	Selector.SetPointProperty(EPCGPointProperties::Position);
@@ -293,7 +293,7 @@ bool FPCGMetadataMathsOp_Truncate::RunTest(const FString& Parameters)
 	// Point Data
 	FPCGAttributePropertyInputSelector Selector;
 	Selector.SetPointProperty(EPCGPointProperties::Seed);
-	bSuccess &= ExecuteTest(TestParams(Truncate, GeneratePointData(), TArray<int64>({0}), TArray<int64>({1})), Selector);
+	bSuccess &= ExecuteTest(TestParams(Truncate, GeneratePointData(), TArray<int32>({0}), TArray<int32>({1})), Selector);
 	Selector.SetPointProperty(EPCGPointProperties::Density);
 	bSuccess &= ExecuteTest(TestParams(Truncate, GeneratePointData(), TArray({0.0}), TArray({1.0})), Selector);
 	Selector.SetPointProperty(EPCGPointProperties::Position);
@@ -322,7 +322,7 @@ bool FPCGMetadataMathsOp_Add::RunTest(const FString& Parameters)
 	// Point Data
 	FPCGAttributePropertyInputSelector Selector;
 	Selector.SetPointProperty(EPCGPointProperties::Seed);
-	bSuccess &= ExecuteTest(TestParams(Add, GeneratePointData(2), TArray<int64>({0}), TArray<int64>({3})), Selector);
+	bSuccess &= ExecuteTest(TestParams(Add, GeneratePointData(2), TArray<int32>({0}), TArray<int32>({3})), Selector);
 	Selector.SetPointProperty(EPCGPointProperties::Density);
 	bSuccess &= ExecuteTest(TestParams(Add, GeneratePointData(2), TArray({0.0}), TArray({3.0})), Selector);
 	Selector.SetPointProperty(EPCGPointProperties::Position);
@@ -351,7 +351,7 @@ bool FPCGMetadataMathsOp_Subtract::RunTest(const FString& Parameters)
 	// Point Data
 	FPCGAttributePropertyInputSelector Selector;
 	Selector.SetPointProperty(EPCGPointProperties::Seed);
-	bSuccess &= ExecuteTest(TestParams(Subtract, GeneratePointData(2), TArray<int64>({0}), TArray<int64>({-1})), Selector);
+	bSuccess &= ExecuteTest(TestParams(Subtract, GeneratePointData(2), TArray<int32>({0}), TArray<int32>({-1})), Selector);
 	Selector.SetPointProperty(EPCGPointProperties::Density);
 	bSuccess &= ExecuteTest(TestParams(Subtract, GeneratePointData(2), TArray({0.0}), TArray({1.0})), Selector);
 	Selector.SetPointProperty(EPCGPointProperties::Position);
@@ -380,7 +380,7 @@ bool FPCGMetadataMathsOp_Multiply::RunTest(const FString& Parameters)
 	// Point Data
 	FPCGAttributePropertyInputSelector Selector;
 	Selector.SetPointProperty(EPCGPointProperties::Seed);
-	bSuccess &= ExecuteTest(TestParams(Multiply, GeneratePointData(2), TArray<int64>({0}), TArray<int64>({2})), Selector);
+	bSuccess &= ExecuteTest(TestParams(Multiply, GeneratePointData(2), TArray<int32>({0}), TArray<int32>({2})), Selector);
 	Selector.SetPointProperty(EPCGPointProperties::Density);
 	bSuccess &= ExecuteTest(TestParams(Multiply, GeneratePointData(2), TArray({0.0}), TArray({2.0})), Selector);
 	Selector.SetPointProperty(EPCGPointProperties::Position);
@@ -409,7 +409,7 @@ bool FPCGMetadataMathsOp_Divide::RunTest(const FString& Parameters)
 	// Point Data
 	FPCGAttributePropertyInputSelector Selector;
 	Selector.SetPointProperty(EPCGPointProperties::Seed);
-	bSuccess &= ExecuteTest(TestParams(Divide, GeneratePointData(2), TArray<int64>({0}), TArray<int64>({0})), Selector);
+	bSuccess &= ExecuteTest(TestParams(Divide, GeneratePointData(2), TArray<int32>({0}), TArray<int32>({0})), Selector);
 	Selector.SetPointProperty(EPCGPointProperties::Density);
 	bSuccess &= ExecuteTest(TestParams(Divide, GeneratePointData(2), TArray({0.0}), TArray({2.0})), Selector);
 	Selector.SetPointProperty(EPCGPointProperties::Position);
@@ -481,7 +481,7 @@ bool FPCGMetadataMathsOp_Abs::RunTest(const FString& Parameters)
 	// Point Data
 	FPCGAttributePropertyInputSelector Selector;
 	Selector.SetPointProperty(EPCGPointProperties::Seed);
-	bSuccess &= ExecuteTest(TestParams(Abs, GeneratePointData(), TArray<int64>({0}), TArray<int64>({1})), Selector);
+	bSuccess &= ExecuteTest(TestParams(Abs, GeneratePointData(), TArray<int32>({0}), TArray<int32>({1})), Selector);
 	Selector.SetPointProperty(EPCGPointProperties::Density);
 	bSuccess &= ExecuteTest(TestParams(Abs, GeneratePointData(), TArray({0.0}), TArray({1.0})), Selector);
 	Selector.SetPointProperty(EPCGPointProperties::Position);
@@ -510,7 +510,7 @@ bool FPCGMetadataMathsOp_Inc::RunTest(const FString& Parameters)
 	// Point Data
 	FPCGAttributePropertyInputSelector Selector;
 	Selector.SetPointProperty(EPCGPointProperties::Seed);
-	bSuccess &= ExecuteTest(TestParams(Inc, GeneratePointData(), TArray<int64>({0}), TArray<int64>({2})), Selector);
+	bSuccess &= ExecuteTest(TestParams(Inc, GeneratePointData(), TArray<int32>({0}), TArray<int32>({2})), Selector);
 	Selector.SetPointProperty(EPCGPointProperties::Density);
 	bSuccess &= ExecuteTest(TestParams(Inc, GeneratePointData(), TArray({0.0}), TArray({2.0})), Selector);
 	Selector.SetPointProperty(EPCGPointProperties::Position);
@@ -539,7 +539,7 @@ bool FPCGMetadataMathsOp_Dec::RunTest(const FString& Parameters)
 	// Point Data
 	FPCGAttributePropertyInputSelector Selector;
 	Selector.SetPointProperty(EPCGPointProperties::Seed);
-	bSuccess &= ExecuteTest(TestParams(Dec, GeneratePointData(), TArray<int64>({0}), TArray<int64>({0})), Selector);
+	bSuccess &= ExecuteTest(TestParams(Dec, GeneratePointData(), TArray<int32>({0}), TArray<int32>({0})), Selector);
 	Selector.SetPointProperty(EPCGPointProperties::Density);
 	bSuccess &= ExecuteTest(TestParams(Dec, GeneratePointData(), TArray({0.0}), TArray({0.0})), Selector);
 	Selector.SetPointProperty(EPCGPointProperties::Position);
@@ -568,7 +568,7 @@ bool FPCGMetadataMathsOp_Negate::RunTest(const FString& Parameters)
 	// Point Data
 	FPCGAttributePropertyInputSelector Selector;
 	Selector.SetPointProperty(EPCGPointProperties::Seed);
-	bSuccess &= ExecuteTest(TestParams(Negate, GeneratePointData(), TArray<int64>({0}), TArray<int64>({-1})), Selector);
+	bSuccess &= ExecuteTest(TestParams(Negate, GeneratePointData(), TArray<int32>({0}), TArray<int32>({-1})), Selector);
 	Selector.SetPointProperty(EPCGPointProperties::Density);
 	bSuccess &= ExecuteTest(TestParams(Negate, GeneratePointData(), TArray({0.0}), TArray({-1.0})), Selector);
 	Selector.SetPointProperty(EPCGPointProperties::Position);
diff --git a/Engine/Plugins/PCG/Source/PCG/Private/Tests/Elements/PCGDataTableRowToParamDataTest.cpp b/Engine/Plugins/PCG/Source/PCG/Private/Tests/Elements/PCGDataTableRowToParamDataTest.cpp
--- a/Engine/Plugins/PCG/Source/PCG/Private/Tests/Elements/PCGDataTableRowToParamDataTest.cpp
+++ b/Engine/Plugins/PCG/Source/PCG/Private/Tests/Elements/PCGDataTableRowToParamDataTest.cpp
@@ -56,7 +56,7 @@ CCC,"C Name","C String","333","3333","3.3","3.33","(X=3.0,Y=3.0)","(X=3.0,Y=3.0,
 	// this is because sometimes conversion occurs as PCG can't represent all types
 	bSuccess |= PCGDataTableRowToParamDataTestHelpers::TestAttribute<FName,FName>(*this, RowData, *Params, TEXT("Name"), TEXT("B Name"));
 	bSuccess |= PCGDataTableRowToParamDataTestHelpers::TestAttribute<FString,FString>(*this, RowData, *Params, TEXT("String"), TEXT("B String"));
-	bSuccess |= PCGDataTableRowToParamDataTestHelpers::TestAttribute<int32,int64>(*this, RowData, *Params, TEXT("I32"), 222);
+	bSuccess |= PCGDataTableRowToParamDataTestHelpers::TestAttribute<int32,int32>(*this, RowData, *Params, TEXT("I32"), 222);
 	bSuccess |= PCGDataTableRowToParamDataTestHelpers::TestAttribute<int64,int64>(*this, RowData, *Params, TEXT("I64"), 2222);
 	bSuccess |= PCGDataTableRowToParamDataTestHelpers::TestAttribute<float,double>(*this, RowData, *Params, TEXT("F32"), 2.2);
 	bSuccess |= PCGDataTableRowToParamDataTestHelpers::TestAttribute<double,double>(*this, RowData, *Params, TEXT("F64"), 2.22);
diff --git a/Engine/Plugins/PCG/Source/PCG/Private/Tests/Elements/PCGGetActorPropertyTest.cpp b/Engine/Plugins/PCG/Source/PCG/Private/Tests/Elements/PCGGetActorPropertyTest.cpp
--- a/Engine/Plugins/PCG/Source/PCG/Private/Tests/Elements/PCGGetActorPropertyTest.cpp
+++ b/Engine/Plugins/PCG/Source/PCG/Private/Tests/Elements/PCGGetActorPropertyTest.cpp
@@ -251,7 +251,7 @@ bool FPCGPropertyToParamDataPropertyTypeTest::RunTest(const FString& Parameters)
 	Actor->DummyStruct.Level2Struct.DoubleArrayProperty = { 0.1, 0.2, 0.3 };
 
 	// Basic properties
-	bSuccess &= VerifyAttributeValueValid(this, TestData, GET_MEMBER_NAME_CHECKED(APCGUnitTestDummyActor, IntProperty), 42ll, ExtraTestWhat);
+	bSuccess &= VerifyAttributeValueValid(this, TestData, GET_MEMBER_NAME_CHECKED(APCGUnitTestDummyActor, IntProperty), 42, ExtraTestWhat);
 	bSuccess &= VerifyAttributeValueValid(this, TestData, GET_MEMBER_NAME_CHECKED(APCGUnitTestDummyActor, Int64Property), 42ll, ExtraTestWhat);
 	bSuccess &= VerifyAttributeValueValid(this, TestData, GET_MEMBER_NAME_CHECKED(APCGUnitTestDummyActor, FloatProperty), 1.0, ExtraTestWhat);
 	bSuccess &= VerifyAttributeValueValid(this, TestData, GET_MEMBER_NAME_CHECKED(APCGUnitTestDummyActor, DoubleProperty), 1.0, ExtraTestWhat);
@@ -292,7 +292,7 @@ bool FPCGPropertyToParamDataPropertyTypeTest::RunTest(const FString& Parameters)
 	bSuccess &= VerifyAttributeValuesValid(this, TestData, GET_MEMBER_NAME_CHECKED(APCGUnitTestDummyActor, ObjectProperty), ObjectPropertyNames, ExtraTestWhat, 42ll, 1.0);
 
 	// Arrays of supported properties
-	bSuccess &= VerifyAttributeValuesValid(this, TestData, GET_MEMBER_NAME_CHECKED(APCGUnitTestDummyActor, ArrayOfIntsProperty), { GET_MEMBER_NAME_CHECKED(APCGUnitTestDummyActor, ArrayOfIntsProperty) }, ExtraTestWhat, 42ll, 43ll, 44ll);
+	bSuccess &= VerifyAttributeValuesValid(this, TestData, GET_MEMBER_NAME_CHECKED(APCGUnitTestDummyActor, ArrayOfIntsProperty), { GET_MEMBER_NAME_CHECKED(APCGUnitTestDummyActor, ArrayOfIntsProperty) }, ExtraTestWhat, 42, 43, 44);
 	bSuccess &= VerifyAttributeValuesValid(this, TestData, GET_MEMBER_NAME_CHECKED(APCGUnitTestDummyActor, ArrayOfVectorsProperty), { GET_MEMBER_NAME_CHECKED(APCGUnitTestDummyActor, ArrayOfVectorsProperty) }, ExtraTestWhat, VectorValue, SecondVectorValue);
 	bSuccess &= VerifyAttributeValuesValid(this, TestData, GET_MEMBER_NAME_CHECKED(APCGUnitTestDummyActor, ArrayOfObjectsProperty), { GET_MEMBER_NAME_CHECKED(APCGUnitTestDummyActor, ArrayOfObjectsProperty) }, ExtraTestWhat, FSoftObjectPath(ObjectValue), FSoftObjectPath(SecondObjectValue));
 
@@ -303,7 +303,7 @@ bool FPCGPropertyToParamDataPropertyTypeTest::RunTest(const FString& Parameters)
 
 	// Extractors
 	bSuccess &= VerifyAttributeValuesValid(this, TestData, TEXT("DummyStruct.FloatProperty"), { GET_MEMBER_NAME_CHECKED(FPCGDummyGetPropertyStruct, FloatProperty) }, ExtraTestWhat, 1.2);
-	bSuccess &= VerifyAttributeValuesValid(this, TestData, TEXT("DummyStruct.IntArrayProperty"), { GET_MEMBER_NAME_CHECKED(FPCGDummyGetPropertyStruct, IntArrayProperty) }, ExtraTestWhat, 5ll, 6ll, 7ll);
+	bSuccess &= VerifyAttributeValuesValid(this, TestData, TEXT("DummyStruct.IntArrayProperty"), { GET_MEMBER_NAME_CHECKED(FPCGDummyGetPropertyStruct, IntArrayProperty) }, ExtraTestWhat, 5, 6, 7);
 	bSuccess &= VerifyAttributeValuesValid(this, TestData, TEXT("DummyStruct.Level2Struct.DoubleArrayProperty"), { GET_MEMBER_NAME_CHECKED(FPCGDummyGetPropertyLevel2Struct, DoubleArrayProperty) }, ExtraTestWhat, 0.1, 0.2, 0.3);
 
 	// Extracting the DummyStruct should only extract the float, as arrays and deeper structs are discarded
@@ -345,7 +345,7 @@ bool FPCGPropertyToParamDataActorFindTest::RunTest(const FString& Parameters)
 		PCGTestsCommon::FTestData TestData(Seed, Settings, APCGUnitTestDummyActor::StaticClass());
 		Cast<APCGUnitTestDummyActor>(TestData.TestActor)->IntProperty = 42;
 
-		bSuccess &= VerifyAttributeValueValid(this, TestData, GET_MEMBER_NAME_CHECKED(APCGUnitTestDummyActor, IntProperty), 42ll, "PropertyToParamDataActorFindTest_Self_Class");
+		bSuccess &= VerifyAttributeValueValid(this, TestData, GET_MEMBER_NAME_CHECKED(APCGUnitTestDummyActor, IntProperty), 42, "PropertyToParamDataActorFindTest_Self_Class");
 	}
 
 	// Self by tag
@@ -357,7 +357,7 @@ bool FPCGPropertyToParamDataActorFindTest::RunTest(const FString& Parameters)
 		TestData.TestActor->Tags.Add(Tag);
 		Cast<APCGUnitTestDummyActor>(TestData.TestActor)->IntProperty = 42;
 
-		bSuccess &= VerifyAttributeValueValid(this, TestData, GET_MEMBER_NAME_CHECKED(APCGUnitTestDummyActor, IntProperty), 42ll, "PropertyToParamDataActorFindTest_Self_Tag");
+		bSuccess &= VerifyAttributeValueValid(this, TestData, GET_MEMBER_NAME_CHECKED(APCGUnitTestDummyActor, IntProperty), 42, "PropertyToParamDataActorFindTest_Self_Tag");
 	}
 
 	// TODO: Need a good way to spawn actors with parenting relation between them
@@ -370,7 +370,7 @@ bool FPCGPropertyToParamDataActorFindTest::RunTest(const FString& Parameters)
 	//	TestData.AddActor(APCGUnitTestDummyActor::StaticClass(), true);
 	//	Cast<APCGUnitTestDummyActor>(TestData.TestActor)->IntProperty = 42;
 
-	//	bSuccess &= VerifyAttributeValueValid(this, TestData, GET_MEMBER_NAME_CHECKED(APCGUnitTestDummyActor, IntProperty), 42ll, "PropertyToParamDataActorFindTest_Parent_Class");
+	//	bSuccess &= VerifyAttributeValueValid(this, TestData, GET_MEMBER_NAME_CHECKED(APCGUnitTestDummyActor, IntProperty), 42, "PropertyToParamDataActorFindTest_Parent_Class");
 	//}
 
 	//// Root by tag
@@ -384,7 +384,7 @@ bool FPCGPropertyToParamDataActorFindTest::RunTest(const FString& Parameters)
 	//	TestData.TestActor->Tags.Add(Tag);
 	//	Cast<APCGUnitTestDummyActor>(TestData.TestActor)->IntProperty = 42;
 
-	//	bSuccess &= VerifyAttributeValueValid(this, TestData, GET_MEMBER_NAME_CHECKED(APCGUnitTestDummyActor, IntProperty), 42ll, "PropertyToParamDataActorFindTest_Root_Tag");
+	//	bSuccess &= VerifyAttributeValueValid(this, TestData, GET_MEMBER_NAME_CHECKED(APCGUnitTestDummyActor, IntProperty), 42, "PropertyToParamDataActorFindTest_Root_Tag");
 	//}
 
 	//// Root by tag including children
@@ -400,7 +400,7 @@ bool FPCGPropertyToParamDataActorFindTest::RunTest(const FString& Parameters)
 
 	//	TestData.AddActor(AActor::StaticClass(), true);
 
-	//	bSuccess &= VerifyAttributeValueValid(this, TestData, GET_MEMBER_NAME_CHECKED(APCGUnitTestDummyActor, IntProperty), 42ll, "PropertyToParamDataActorFindTest_Root_Tag_Children");
+	//	bSuccess &= VerifyAttributeValueValid(this, TestData, GET_MEMBER_NAME_CHECKED(APCGUnitTestDummyActor, IntProperty), 42, "PropertyToParamDataActorFindTest_Root_Tag_Children");
 	//	Settings->bIncludeChildren = false;
 	//}
 
@@ -415,7 +415,7 @@ bool FPCGPropertyToParamDataActorFindTest::RunTest(const FString& Parameters)
 		UPCGUnitTestDummyComponent* Component = Cast<UPCGUnitTestDummyComponent>(TestData.TestActor->AddComponentByClass(Settings->ComponentClass, false, FTransform::Identity, false));
 		Component->IntProperty = 42;
 
-		bSuccess &= VerifyAttributeValueValid(this, TestData, GET_MEMBER_NAME_CHECKED(UPCGUnitTestDummyComponent, IntProperty), 42ll, "PropertyToParamDataActorFindTest_Self_Tag_Component");
+		bSuccess &= VerifyAttributeValueValid(this, TestData, GET_MEMBER_NAME_CHECKED(UPCGUnitTestDummyComponent, IntProperty), 42, "PropertyToParamDataActorFindTest_Self_Tag_Component");
 		Settings->bSelectComponent = false;
 	}
 
@@ -427,7 +427,7 @@ bool FPCGPropertyToParamDataActorFindTest::RunTest(const FString& Parameters)
 		PCGTestsCommon::FTestData TestData(Seed, Settings, APCGUnitTestDummyActor::StaticClass());
 		Cast<APCGUnitTestDummyActor>(TestData.TestActor)->IntProperty = 42;
 
-		bSuccess &= VerifyAttributeValueValid(this, TestData, GET_MEMBER_NAME_CHECKED(APCGUnitTestDummyActor, IntProperty), 42ll, "PropertyToParamDataActorFindTest_World_Class");
+		bSuccess &= VerifyAttributeValueValid(this, TestData, GET_MEMBER_NAME_CHECKED(APCGUnitTestDummyActor, IntProperty), 42, "PropertyToParamDataActorFindTest_World_Class");
 	}
 
 	// World by tag
@@ -439,7 +439,7 @@ bool FPCGPropertyToParamDataActorFindTest::RunTest(const FString& Parameters)
 		TestData.TestActor->Tags.Add(Tag);
 		Cast<APCGUnitTestDummyActor>(TestData.TestActor)->IntProperty = 42;
 
-		bSuccess &= VerifyAttributeValueValid(this, TestData, GET_MEMBER_NAME_CHECKED(APCGUnitTestDummyActor, IntProperty), 42ll, "PropertyToParamDataActorFindTest_World_Tag");
+		bSuccess &= VerifyAttributeValueValid(this, TestData, GET_MEMBER_NAME_CHECKED(APCGUnitTestDummyActor, IntProperty), 42, "PropertyToParamDataActorFindTest_World_Tag");
 	}
 
 	return bSuccess;
