Fixed duplication error on copy-paste.

#rnx

diff --git a/Engine/Plugins/PCG/Source/PCG/Private/PCGData.cpp b/Engine/Plugins/PCG/Source/PCG/Private/PCGData.cpp
--- a/Engine/Plugins/PCG/Source/PCG/Private/PCGData.cpp
+++ b/Engine/Plugins/PCG/Source/PCG/Private/PCGData.cpp
@@ -80,6 +80,12 @@ void UPCGData::InitUID()
 	UID = ++UIDCounter;
 }
 
+void UPCGData::PostEditImport()
+{
+	Super::PostEditImport();
+	InitUID();
+}
+
 void UPCGData::Flatten()
 {
 	if (UPCGMetadata* Metadata = MutableMetadata())
diff --git a/Engine/Plugins/PCG/Source/PCG/Public/PCGData.h b/Engine/Plugins/PCG/Source/PCG/Public/PCGData.h
--- a/Engine/Plugins/PCG/Source/PCG/Public/PCGData.h
+++ b/Engine/Plugins/PCG/Source/PCG/Public/PCGData.h
@@ -59,6 +59,7 @@ public:
 
 	// ~Begin UObject interface
 	virtual void PostDuplicate(bool bDuplicateForPIE) override { InitUID(); }
+	virtual void PostEditImport() override;
 	// ~End UObject interface
 
 	// Metadata ops, to be implemented if data supports Metadata
