[PCG] Add missing transpose on RotationMatrix in Get/SetPointTransform helper functions in PCGDataCollectionDataInterface.ush
* Was causing rotational artifacts for GPU Static Mesh Spawning

#rnx

diff --git a/Engine/Plugins/PCG/Shaders/Private/PCGDataCollectionDataInterface.ush b/Engine/Plugins/PCG/Shaders/Private/PCGDataCollectionDataInterface.ush
--- a/Engine/Plugins/PCG/Shaders/Private/PCGDataCollectionDataInterface.ush
+++ b/Engine/Plugins/PCG/Shaders/Private/PCGDataCollectionDataInterface.ush
@@ -376,7 +376,7 @@ uint GetSeed_{DataInterfaceName}(uint DataIndex, uint ElementIndex)
 float4x4 GetPointTransform_{DataInterfaceName}(uint InDataIndex, uint InElementIndex)
 {
 	const FQuat Rotation = GetRotation_{DataInterfaceName}(InDataIndex, InElementIndex);
-	const half3x3 RotationMatrix = QuatToMatrix(Rotation);
+	const half3x3 RotationMatrix = transpose(QuatToMatrix(Rotation));
 	const float3 Position = GetPosition_{DataInterfaceName}(InDataIndex, InElementIndex);
 	const float3 Scale = GetScale_{DataInterfaceName}(InDataIndex, InElementIndex);
 	
@@ -449,7 +449,7 @@ void SetPointTransform_{DataInterfaceName}(uint InDataIndex, uint InElementIndex
 	SetScale_{DataInterfaceName}(InDataIndex, InElementIndex, Scale);
 
 	// Assumes Transform axes have the correct handedness. We could do more work to fix this up too if required.
-	float3x3 RotationMatix = Transform;
+	float3x3 RotationMatix = transpose(Transform);
 	// Required unfortunately, QuatFromMatrix expects orthonormalized.
 	RotationMatix[0] /= Scale.x;
 	RotationMatix[1] /= Scale.y;
