[PCG] Remove ensure that fires on load for HiGen graphs that only utilize unbounded grid size

#rnx

diff --git a/Engine/Plugins/PCG/Source/PCG/Private/PCGSubsystem.cpp b/Engine/Plugins/PCG/Source/PCG/Private/PCGSubsystem.cpp
--- a/Engine/Plugins/PCG/Source/PCG/Private/PCGSubsystem.cpp
+++ b/Engine/Plugins/PCG/Source/PCG/Private/PCGSubsystem.cpp
@@ -1011,7 +1011,6 @@ FPCGTaskId UPCGSubsystem::ForAllOverlappingCells(UPCGComponent* InComponent, con
 
 	if (GridSizes.IsEmpty())
 	{
-		ensureMsgf(false, TEXT("Must have grid sizes."));
 		return InvalidPCGTaskId;
 	}
 
@@ -1696,7 +1695,7 @@ bool UPCGSubsystem::GetStackContext(const UPCGComponent* InComponent, FPCGStackC
 
 	if (bDoesComponentExecute)
 	{
-		GetGraphCompiler()->GetCompiledTasks(Graph, InComponent->GetGenerationGridSize(), OutStackContext);
+		GetGraphCompiler()->GetCompiledTasks(Graph, InComponent->GetGenerationGridSize(), OutStackContext, /*bIsCooking=*/false);
 		return true;
 	}
 	else
