[PCG] Level to Asset : Make sure to load all actors for WP worlds

#rnx

diff --git a/Engine/Plugins/PCG/Source/PCGEditor/Private/PCGLevelToAsset.cpp b/Engine/Plugins/PCG/Source/PCGEditor/Private/PCGLevelToAsset.cpp
--- a/Engine/Plugins/PCG/Source/PCGEditor/Private/PCGLevelToAsset.cpp
+++ b/Engine/Plugins/PCG/Source/PCGEditor/Private/PCGLevelToAsset.cpp
@@ -19,6 +19,7 @@
 #include "Components/InstancedStaticMeshComponent.h"
 #include "UObject/UObjectGlobals.h"
 #include "UObject/Package.h"
+#include "WorldPartition/WorldPartition.h"
 
 void UPCGLevelToAsset::CreateOrUpdatePCGAssets(const TArray<FAssetData>& WorldAssets, const FPCGAssetExporterParameters& InParameters, TSubclassOf<UPCGLevelToAsset> ExporterSubclass)
 {
@@ -208,6 +209,14 @@ bool UPCGLevelToAsset::BP_ExportWorld_Implementation(UWorld* World, const FStrin
 		ActorIndexAttribute->SetValue(RootPoint.MetadataEntry, 0);
 	}
 
+
+	// Make sure all actors are loaded
+	TArray<FWorldPartitionReference> ActorReferences;
+	if (UWorldPartition* WorldPartition = World->GetWorldPartition())
+	{
+		WorldPartition->LoadAllActors(ActorReferences);
+	}
+
 	// Build actor-index map
 	TMap<AActor*, int> ActorIndexMap;
 	int LastActorIndex = 1; // Since the root is the "first" point we'll have, we'll have the map start from 1.
