[PCG] GPU: Fix attribute usage validation which broke in previous commit.

#rnx

diff --git a/Engine/Plugins/PCG/Source/PCG/Private/Compute/Elements/PCGComputeGraphElement.cpp b/Engine/Plugins/PCG/Source/PCG/Private/Compute/Elements/PCGComputeGraphElement.cpp
--- a/Engine/Plugins/PCG/Source/PCG/Private/Compute/Elements/PCGComputeGraphElement.cpp
+++ b/Engine/Plugins/PCG/Source/PCG/Private/Compute/Elements/PCGComputeGraphElement.cpp
@@ -107,16 +107,6 @@ bool FPCGComputeGraphElement::ExecuteInternal(FPCGContext* InContext) const
 	// 4. Initialize and parse incoming data for data sizes, attributes, etc that will drive buffer allocations and dispatch thread counts.
 	if (!Context->DataBinding)
 	{
-		for (TWeakObjectPtr<const UPCGNode> Node : Graph->KernelToNode)
-		{
-			const UPCGSettings* Settings = Node.Get() ? Node->GetSettings() : nullptr;
-
-			if (!Settings || !Settings->IsKernelValid(Context, /*bQuiet=*/false))
-			{
-				return true;
-			}
-		}
-
 		// When attribute table was built at compile time, some attribute types may not be known. Fill them out now.
 		Graph->FillInMissingAttributeTableTypes(Context->InputData);
 
@@ -141,6 +131,17 @@ bool FPCGComputeGraphElement::ExecuteInternal(FPCGContext* InContext) const
 
 		DataForGPU.InputPinLabelAliases = Graph->InputPinLabelAliases;
 
+		// Perform validation after input data is initialized.
+		for (TWeakObjectPtr<const UPCGNode> Node : Graph->KernelToNode)
+		{
+			const UPCGSettings* Settings = Node.Get() ? Node->GetSettings() : nullptr;
+
+			if (!Settings || !Settings->IsKernelValid(Context, /*bQuiet=*/false))
+			{
+				return true;
+			}
+		}
+
 		const bool bAnyComponentsSetup = SetupProceduralISMComponents(InContext, DataBindingObject);
 
 		Context->ComputeGraphInstance.CreateDataProviders(Graph.Get(), 0, Context->DataBinding.Get());
