[PCG] Fix case where EndPlay is called first on PCGPartitionActor and second on Original component volume
- Make sure RemoveGraphInstance is called on EndPlay for RuntimeGenerated PAs
- Remove ensure with comment

#rnx

diff --git a/Engine/Plugins/PCG/Source/PCG/Private/Grid/PCGPartitionActor.cpp b/Engine/Plugins/PCG/Source/PCG/Private/Grid/PCGPartitionActor.cpp
--- a/Engine/Plugins/PCG/Source/PCG/Private/Grid/PCGPartitionActor.cpp
+++ b/Engine/Plugins/PCG/Source/PCG/Private/Grid/PCGPartitionActor.cpp
@@ -243,6 +243,16 @@ void APCGPartitionActor::BeginPlay()
 
 void APCGPartitionActor::EndPlay(const EEndPlayReason::Type EndPlayReason)
 {
+	// It is possible for a PA to get called EndPlay before its original volume actor so make sure to cleanup properly here
+	if (IsRuntimeGenerated())
+	{
+		TMap<TObjectPtr<UPCGComponent>, TObjectPtr<UPCGComponent>> OriginalToLocalCopy(OriginalToLocal);
+		for (const TTuple<TObjectPtr<UPCGComponent>, TObjectPtr<UPCGComponent>>& OriginalToLocalItem : OriginalToLocalCopy)
+		{
+			RemoveGraphInstance(OriginalToLocalItem.Key.Get());
+		}
+	}
+
 	UnregisterPCG();
 
 	Super::EndPlay(EndPlayReason);
diff --git a/Engine/Plugins/PCG/Source/PCG/Private/RuntimeGen/PCGRuntimeGenScheduler.cpp b/Engine/Plugins/PCG/Source/PCG/Private/RuntimeGen/PCGRuntimeGenScheduler.cpp
--- a/Engine/Plugins/PCG/Source/PCG/Private/RuntimeGen/PCGRuntimeGenScheduler.cpp
+++ b/Engine/Plugins/PCG/Source/PCG/Private/RuntimeGen/PCGRuntimeGenScheduler.cpp
@@ -796,7 +796,8 @@ void FPCGRuntimeGenScheduler::OnOriginalComponentUnregistered(UPCGComponent* InO
 		// Get the generated component for this key (might be a local component).
 		UPCGComponent* ComponentToCleanup = bIsOriginalComponent ? InOriginalComponent : ActorAndComponentMapping->GetLocalComponent(GridDescriptor, GenerationKey.GetGridCoords(), InOriginalComponent);
 
-		if (ensure(ComponentToCleanup))
+		// It is possible for a PartitionActor's LocalComponent to have been cleaned up by the APCGPartitionActor::EndPlay call depending on the order in which actors get called
+		if (ComponentToCleanup)
 		{
 			CleanupComponent(GenerationKey, ComponentToCleanup);
 		}
