[PCG] Fix unitialized fields

#rnx

diff --git a/Engine/Plugins/PCG/Source/PCG/Public/Compute/PCGDataForGPU.h b/Engine/Plugins/PCG/Source/PCG/Public/Compute/PCGDataForGPU.h
--- a/Engine/Plugins/PCG/Source/PCG/Public/Compute/PCGDataForGPU.h
+++ b/Engine/Plugins/PCG/Source/PCG/Public/Compute/PCGDataForGPU.h
@@ -71,10 +71,10 @@ struct FPCGKernelAttributeIDAndType
 	{}
 
 	UPROPERTY()
-	int32 Id;
+	int32 Id = INDEX_NONE;
 
 	UPROPERTY()
-	EPCGKernelAttributeType Type;
+	EPCGKernelAttributeType Type = EPCGKernelAttributeType::Invalid;
 };
 
 struct FPCGKernelAttributeDesc
