Split part of the Fuse element functionality in an 'Attract' node that can find attractors based on distance or attribute criteria (in distance) then interpolate between source and target.
Also has a straight 'interpolate' functionality driven by an index.

diff --git a/Engine/Plugins/PCG/Source/PCG/Private/Elements/PCGAttractElement.cpp b/Engine/Plugins/PCG/Source/PCG/Private/Elements/PCGAttractElement.cpp
new file mode 100644
--- /dev/null
+++ b/Engine/Plugins/PCG/Source/PCG/Private/Elements/PCGAttractElement.cpp
@@ -0,0 +1,637 @@
+// Copyright Epic Games, Inc. All Rights Reserved.
+
+#include "Elements/PCGAttractElement.h"
+
+#include "PCGContext.h"
+#include "Data/PCGPointData.h"
+#include "Elements/PCGSortAttributes.h"
+#include "Helpers/PCGAsync.h"
+#include "Helpers/PCGHelpers.h"
+#include "Metadata/Accessors/PCGAttributeAccessorHelpers.h"
+#include "SpatialAlgo/PCGOctreeQueries.h"
+
+#include "Algo/MaxElement.h"
+
+#define LOCTEXT_NAMESPACE "PCGAttractElement"
+
+namespace PCGAttractElement
+{
+	namespace Constants
+	{
+		const FName InputSourceLabel = TEXT("Source");
+		const FName InputTargetLabel = TEXT("Target");
+		const FName AttractIndexName = TEXT("AttractIndex");
+	}
+
+	namespace Helpers
+	{
+		static void InterpolatePoints(FAttractState& AttractData, const int32 InterpolatorIndex, const UPCGAttractSettings* Settings)
+		{
+			FAttractState::FAttributeInterpolationData& Interpolator = AttractData.Interpolators[InterpolatorIndex];
+
+			// Process
+			int32 OutIndex = 0;
+			const int32 SourcePointCount = AttractData.SourcePointData->GetPoints().Num();
+
+			for (int32 SourceIndex = 0; SourceIndex < SourcePointCount; ++SourceIndex)
+			{
+				const int32 TargetIndex = AttractData.Mapping[SourceIndex];
+
+				if (TargetIndex == INDEX_NONE)
+				{
+					// Point was not matched, we're either not touching it or dropping it.
+					if (!Settings->bRemoveUnattractedPoints)
+					{
+						++OutIndex;
+					}
+				}
+				else if (AttractData.bTargetIsSource && TargetIndex == SourceIndex)
+				{
+					// Nothing to do regardless of weights
+					++OutIndex; 
+				}
+				else
+				{
+					double Alpha = 0.0;
+					if (AttractData.SourceWeights.IsEmpty() && AttractData.TargetWeights.IsEmpty())
+					{
+						Alpha = Settings->Weight;
+					}
+					else if (AttractData.SourceWeights.IsEmpty())
+					{
+						Alpha = AttractData.TargetWeights[TargetIndex];
+					}
+					else if (AttractData.TargetWeights.IsEmpty())
+					{
+						Alpha = AttractData.SourceWeights[SourceIndex];
+					}
+					else
+					{
+						const double SourceWeight = AttractData.SourceWeights[SourceIndex];
+						const double TargetWeight = AttractData.TargetWeights[TargetIndex];
+						Alpha = (SourceWeight + TargetWeight > UE_SMALL_NUMBER) ? (TargetWeight / (SourceWeight + TargetWeight)) : 0.5;
+					}
+
+					auto Callback = [OutIndex, SourceIndex, TargetIndex, Alpha, &Interpolator, &AttractData](auto Dummy)
+					{
+						using OutputType = decltype(Dummy);
+						OutputType SourceValue{}, TargetValue{};
+						Interpolator.SourceAccessor->Get(SourceValue, SourceIndex, *Interpolator.SourceKeys);
+						Interpolator.TargetAccessor->Get(TargetValue, TargetIndex, *Interpolator.TargetKeys, EPCGAttributeAccessorFlags::AllowBroadcastAndConstructible);
+
+						OutputType OutputValue{};
+						if constexpr (PCG::Private::MetadataTraits<OutputType>::CanInterpolate)
+						{
+							if constexpr (std::is_same_v<OutputType, FTransform>)
+							{
+								// TODO: use BlendWith or PCGMetadataTransfromSettings::LerpTransform (this one has more options)
+								OutputValue = SourceValue;
+								OutputValue.BlendWith(TargetValue, Alpha);
+							}
+							else
+							{
+								// TODO: use PCGMetadataMaths::Lerp or PCG::Private::MetadataTraits<OutputType>::WeightedSum or FMath::Lerp
+								OutputValue = FMath::Lerp(SourceValue, TargetValue, Alpha);
+							}
+						}
+						else // Take value based on maximal weight
+						{
+							// TODO - Optimization: we don't need to query both values when we're in this case, as we'll use only one of them.
+							OutputValue = (Alpha <= 0.5 ? SourceValue : TargetValue);
+						}
+
+						Interpolator.OutputAccessor->Set(OutputValue, OutIndex, *AttractData.OutputKeys);
+					};
+
+					PCGMetadataAttribute::CallbackWithRightType(Interpolator.SourceAccessor->GetUnderlyingType(), Callback);
+					++OutIndex;
+				}
+			}
+		}
+
+		// Runs a single iteration of the attract algorithm.
+		static void RunAttractIteration(const int32 Index, FAttractState& OutAttractData, const UPCGAttractSettings* AttractSettings)
+		{
+			const TArray<FPCGPoint>& SourcePoints = OutAttractData.SourcePointData->GetPoints();
+
+			if (AttractSettings->Mode == EPCGAttractMode::Closest)
+			{
+				// Note: there is a slight behavior difference here in the target==source case, we don't want a point to be attracted to itself
+				if (OutAttractData.bTargetIsSource)
+				{
+					int32 ClosestTargetIndex = INDEX_NONE;
+					double MinDistanceSquared = std::numeric_limits<double>::max();
+					bool bFoundColocatedPoint = false;
+
+					UPCGOctreeQueries::ForEachPointInsideSphere(
+						OutAttractData.TargetPointData,
+						SourcePoints[Index].Transform.GetLocation(),
+						AttractSettings->Distance,
+						[Index, &OutAttractData, &MinDistanceSquared, &ClosestTargetIndex, &bFoundColocatedPoint](const FPCGPointRef& PointRef, const double DistanceSquared)
+						{
+							const int32 TargetIndex = PointRef.Point - OutAttractData.TargetPointData->GetPoints().GetData();
+							// Ignore self-selection
+							if (TargetIndex == Index)
+							{
+								return;
+							}
+
+							if (FMath::IsNearlyZero(DistanceSquared))
+							{
+								if (bFoundColocatedPoint)
+								{
+									// Implementation note: this is an arbitrary choice. We could have randomized the targets in a consistent way though.
+									ClosestTargetIndex = FMath::Min(TargetIndex, ClosestTargetIndex);
+								}
+								else
+								{
+									bFoundColocatedPoint = true;
+									ClosestTargetIndex = TargetIndex;
+								}
+
+								return;
+							}
+
+							// Non-colocated points will just be on a distance basis
+							if (DistanceSquared < MinDistanceSquared)
+							{
+								MinDistanceSquared = DistanceSquared;
+								ClosestTargetIndex = TargetIndex;
+							}
+						});
+
+					// Update assignment based on closest target index.
+					OutAttractData.Mapping[Index] = ClosestTargetIndex;
+				}
+				else
+				{
+					// Simple closest point query
+					if (const FPCGPoint* ClosestPointInTarget = UPCGOctreeQueries::GetClosestPoint(
+						OutAttractData.TargetPointData,
+						SourcePoints[Index].Transform.GetLocation(),
+						/*bInDiscardCenter=*/false,
+						AttractSettings->Distance))
+					{
+						const int32 TargetIndex = ClosestPointInTarget - OutAttractData.TargetPointData->GetPoints().GetData();
+						OutAttractData.Mapping[Index] = TargetIndex;
+					}
+				}
+			}
+			else // Min/Max attribute test inside of search radius
+			{
+				UPCGOctreeQueries::ForEachPointInsideSphere(
+					OutAttractData.TargetPointData,
+					SourcePoints[Index].Transform.GetLocation(),
+					AttractSettings->Distance,
+					[Index, &OutAttractData](const FPCGPointRef& PointRef, const double DistanceSquared)
+					{
+						const int32 TargetIndex = PointRef.Point - OutAttractData.TargetPointData->GetPoints().GetData();
+
+						if (OutAttractData.Mapping[Index] == INDEX_NONE)
+						{
+							OutAttractData.Mapping[Index] = TargetIndex;
+						}
+						else
+						{
+							// Implementation note: if we stored 'sorted target indices indices' in the mapping instead
+							// we could remove the secondary find here, but it would require another pass to go back to real indices
+							int32 SortedTargetIndex = OutAttractData.SortedTargetIndices.IndexOfByKey(TargetIndex);
+							check(SortedTargetIndex != INDEX_NONE);
+
+							int32 SortedPreviousIndex = OutAttractData.SortedTargetIndices.IndexOfByKey(OutAttractData.Mapping[Index]);
+							check(SortedPreviousIndex != INDEX_NONE);
+
+							if (SortedTargetIndex < SortedPreviousIndex)
+							{
+								OutAttractData.Mapping[Index] = TargetIndex;
+							}
+						}
+					});
+			}
+		}
+	} // namespace Helpers
+
+	namespace Algorithm
+	{
+		static bool Sequential(const FPCGContext* InContext, FAttractState& OutAttractData, const UPCGAttractSettings* AttractSettings)
+		{
+			TRACE_CPUPROFILER_EVENT_SCOPE(FPCGAttractElement::Sequential);
+
+			check(InContext);
+
+			while (OutAttractData.IterationIndex < OutAttractData.SourcePointData->GetNumPoints())
+			{
+				Helpers::RunAttractIteration(OutAttractData.IterationIndex++, OutAttractData, AttractSettings);
+
+				if (OutAttractData.IterationIndex != OutAttractData.SourcePointData->GetNumPoints() && InContext->ShouldStop())
+				{
+					// Not done
+					return false;
+				}
+			}
+
+			// Fully done
+			return true;
+		}
+	} // namespace Algorithm
+} // namespace PCGAttractElement
+
+UPCGAttractSettings::UPCGAttractSettings()
+{
+	AttractorIndexAttribute.Update(PCGAttractElement::Constants::AttractIndexName.ToString());
+
+	FPCGAttributePropertyInputSelector DefaultSourceInput;
+	DefaultSourceInput.SetPointProperty(EPCGPointProperties::Position);
+	FPCGAttributePropertyInputSelector DefaultTargetInput;
+	DefaultTargetInput.SetPointProperty(EPCGPointProperties::Position);
+	SourceAndTargetAttributeMapping.Add(DefaultSourceInput, DefaultTargetInput);
+
+	OutputAttractIndexAttribute.Update(PCGAttractElement::Constants::AttractIndexName.ToString());
+}
+
+#if WITH_EDITOR
+FText UPCGAttractSettings::GetNodeTooltipText() const
+{
+	return LOCTEXT("NodeTooltip", "Attracts source points to target points based on a max distance and a criteria.");
+}
+#endif // WITH_EDITOR
+
+TArray<FPCGPinProperties> UPCGAttractSettings::InputPinProperties() const
+{
+	TArray<FPCGPinProperties> Properties;
+
+	FPCGPinProperties& SourcePinProperty = Properties.Emplace_GetRef(PCGAttractElement::Constants::InputSourceLabel, EPCGDataType::Point);
+	SourcePinProperty.SetRequiredPin();
+
+	// TODO: For now, only allow a single target input. Only N:0 and N:1 are currently supported.
+	FPCGPinProperties& TargetPinProperty = Properties.Emplace_GetRef(PCGAttractElement::Constants::InputTargetLabel, EPCGDataType::Point);
+	TargetPinProperty.SetRequiredPin();
+#if WITH_EDITOR
+	SourcePinProperty.Tooltip = LOCTEXT("SourcePinTooltip", "The source points to be attracted by the target points.");
+	TargetPinProperty.Tooltip = LOCTEXT("TargetPinTooltip", "The target points that will attract the source points.");
+#endif // WITH_EDITOR
+
+	return Properties;
+}
+
+FPCGElementPtr UPCGAttractSettings::CreateElement() const
+{
+	return MakeShared<FPCGAttractElement>();
+}
+
+bool FPCGAttractElement::PrepareDataInternal(FPCGContext* InContext) const
+{
+	TRACE_CPUPROFILER_EVENT_SCOPE(FPCGAttractElement::PrepareData);
+
+	const UPCGAttractSettings* Settings = InContext->GetInputSettings<UPCGAttractSettings>();
+	check(Settings);
+
+	ContextType* Context = static_cast<ContextType*>(InContext);
+	check(Context);
+
+	const TArray<FPCGTaggedData> SourceInputs = Context->InputData.GetInputsByPin(PCGAttractElement::Constants::InputSourceLabel);
+	const TArray<FPCGTaggedData> TargetInputs = Context->InputData.GetInputsByPin(PCGAttractElement::Constants::InputTargetLabel);
+
+	// Early out if no source to modify or there is no target to be attracted to.
+	if (SourceInputs.IsEmpty() || TargetInputs.IsEmpty())
+	{
+		Context->OutputData.TaggedData = SourceInputs;
+		return true;
+	}
+
+	// Only N:1 and N:N are currently supported.
+	if (TargetInputs.Num() > 1 && TargetInputs.Num() != SourceInputs.Num())
+	{
+		PCGLog::InputOutput::LogInvalidCardinalityError(PCGAttractElement::Constants::InputSourceLabel, PCGAttractElement::Constants::InputTargetLabel, InContext);
+		return true;
+	}
+
+	// Additional validation: if the attract operation would do nothing, log a warning
+	if (!Settings->bRemoveUnattractedPoints && !Settings->bOutputAttractIndex && Settings->SourceAndTargetAttributeMapping.IsEmpty())
+	{
+		PCGLog::LogWarningOnGraph(LOCTEXT("NoOperation", "Attract node settings will do nothing. Sources will be forwarded."), Context);
+		Context->OutputData.TaggedData = SourceInputs;
+		return true;
+	}
+
+	Context->InitializePerExecutionState([Settings](const ContextType*, ExecStateType& OutState)
+	{
+		OutState.AttractFunction = PCGAttractElement::Algorithm::Sequential;
+		check(OutState.AttractFunction);
+
+		return EPCGTimeSliceInitResult::Success;
+	});
+
+	Context->InitializePerIterationStates(SourceInputs.Num(), [&SourceInputs, &TargetInputs, Settings, Context](IterStateType& OutState, const ExecStateType&, const uint32 IterationIndex)
+	{
+		const UPCGPointData* SourcePointData = CastChecked<UPCGPointData>(SourceInputs[IterationIndex].Data);
+
+		if (!SourcePointData || SourcePointData->IsEmpty())
+		{
+			return EPCGTimeSliceInitResult::NoOperation;
+		}
+
+		const UPCGPointData* TargetPointData = CastChecked<UPCGPointData>(TargetInputs[IterationIndex % TargetInputs.Num()].Data);
+
+		UPCGPointData* OutPointData = FPCGContext::NewObject_AnyThread<UPCGPointData>(Context);
+		OutPointData->InitializeFromData(SourcePointData);
+		Context->OutputData.TaggedData.Add_GetRef(SourceInputs[IterationIndex]).Data = OutPointData;
+
+		const bool bTargetIsSource = (TargetPointData == SourcePointData);
+
+		OutState = PCGAttractElement::FAttractState
+		{
+			.TargetPointData = TargetPointData,
+			.SourcePointData = SourcePointData,
+			.OutPointData = OutPointData,
+
+			.IterationIndex = 0,
+			.bTargetIsSource = bTargetIsSource
+		};
+
+		OutState.Mapping.SetNumUninitialized(SourcePointData->GetNumPoints());
+
+		if (Settings->Mode == EPCGAttractMode::FromIndex)
+		{
+			const FPCGAttributePropertyInputSelector AttractorIndexSelector = Settings->AttractorIndexAttribute.CopyAndFixLast(SourcePointData);
+			TUniquePtr<const IPCGAttributeAccessor> AttractorIndexAccessor = PCGAttributeAccessorHelpers::CreateConstAccessor(SourcePointData, AttractorIndexSelector);
+			TUniquePtr<const IPCGAttributeAccessorKeys> AttractorIndexKeys = PCGAttributeAccessorHelpers::CreateConstKeys(SourcePointData, AttractorIndexSelector);
+			if (!AttractorIndexAccessor || !AttractorIndexKeys)
+			{
+				PCGLog::Metadata::LogFailToCreateAccessorError(AttractorIndexSelector, Context);
+				return EPCGTimeSliceInitResult::NoOperation;
+			}
+
+			if (!AttractorIndexAccessor->GetRange(MakeArrayView(OutState.Mapping), 0, *AttractorIndexKeys))
+			{
+				PCGLog::Metadata::LogFailToGetAttributeError(AttractorIndexSelector, Context);
+				return EPCGTimeSliceInitResult::NoOperation;
+			}
+
+			// Fix indices that would be invalid vs target
+			const int32 TargetPointCount = TargetPointData->GetNumPoints();
+			for (int32& MapIndex : OutState.Mapping)
+			{
+				if (MapIndex < 0 || MapIndex >= TargetPointCount)
+				{
+					MapIndex = INDEX_NONE;
+				}
+			}
+		}
+		else
+		{
+			if (bTargetIsSource)
+			{
+				// Self-assignment
+				for (int i = 0; i < OutState.Mapping.Num(); ++i)
+				{
+					OutState.Mapping[i] = i;
+				}
+			}
+			else
+			{
+				// Point to no target
+				for (int i = 0; i < OutState.Mapping.Num(); ++i)
+				{
+					OutState.Mapping[i] = INDEX_NONE;
+				}
+			}
+		}
+
+		// Prioritize selected attract target by comparing attributes for the min or max value.
+		if (Settings->Mode == EPCGAttractMode::MinAttribute || Settings->Mode == EPCGAttractMode::MaxAttribute)
+		{
+			const FPCGAttributePropertyInputSelector TargetSelector = Settings->TargetAttribute.CopyAndFixLast(TargetPointData);
+			TUniquePtr<const IPCGAttributeAccessor> TargetAccessor = PCGAttributeAccessorHelpers::CreateConstAccessor(TargetPointData, TargetSelector);
+			TUniquePtr<const IPCGAttributeAccessorKeys> TargetKeys = PCGAttributeAccessorHelpers::CreateConstKeys(TargetPointData, TargetSelector);
+			if (!TargetAccessor || !TargetKeys)
+			{
+				PCGLog::Metadata::LogFailToCreateAccessorError(TargetSelector, Context);
+				return EPCGTimeSliceInitResult::NoOperation;
+			}
+
+			// Build order of target values - we don't need to keep the values, just the relative ordering is going to be fine.
+			TArray<int> TargetIndices;
+			TargetIndices.SetNumUninitialized(TargetPointData->GetNumPoints());
+			for (int i = 0; i < TargetIndices.Num(); ++i)
+			{
+				TargetIndices[i] = i;
+			}
+
+			PCGAttributeAccessorHelpers::SortByAttribute(*TargetAccessor, *TargetKeys, TargetIndices, Settings->Mode == EPCGAttractMode::MinAttribute);
+			OutState.SortedTargetIndices = std::move(TargetIndices);
+		}
+
+		if (Settings->bUseSourceWeight)
+		{
+			const FPCGAttributePropertyInputSelector SourceWeightSelector = Settings->SourceWeightAttribute.CopyAndFixLast(SourcePointData);
+			const TUniquePtr<const IPCGAttributeAccessor> SourceWeightAccessor = PCGAttributeAccessorHelpers::CreateConstAccessor(SourcePointData, SourceWeightSelector);
+			const TUniquePtr<const IPCGAttributeAccessorKeys> SourceWeightKeys = PCGAttributeAccessorHelpers::CreateConstKeys(SourcePointData, SourceWeightSelector);
+			if (!SourceWeightAccessor)
+			{
+				PCGLog::Metadata::LogFailToCreateAccessorError(SourceWeightSelector, Context);
+				return EPCGTimeSliceInitResult::NoOperation;
+			}
+
+			OutState.SourceWeights.SetNumUninitialized(SourcePointData->GetNumPoints());
+
+			if (!SourceWeightAccessor->GetRange(MakeArrayView(OutState.SourceWeights), 0, *SourceWeightKeys))
+			{
+				PCGLog::Metadata::LogFailToGetAttributeError(SourceWeightSelector, Context);
+				return EPCGTimeSliceInitResult::NoOperation;
+			}
+		}
+
+		if (Settings->bUseTargetWeight)
+		{
+			const FPCGAttributePropertyInputSelector TargetWeightSelector = Settings->TargetWeightAttribute.CopyAndFixLast(TargetPointData);
+			const TUniquePtr<const IPCGAttributeAccessor> TargetWeightAccessor = PCGAttributeAccessorHelpers::CreateConstAccessor(TargetPointData, TargetWeightSelector);
+			const TUniquePtr<const IPCGAttributeAccessorKeys> TargetWeightKeys = PCGAttributeAccessorHelpers::CreateConstKeys(TargetPointData, TargetWeightSelector);
+			if (!TargetWeightAccessor)
+			{
+				PCGLog::Metadata::LogFailToCreateAccessorError(TargetWeightSelector, Context);
+				return EPCGTimeSliceInitResult::NoOperation;
+			}
+
+			OutState.TargetWeights.SetNumUninitialized(TargetPointData->GetNumPoints());
+
+			if (!TargetWeightAccessor->GetRange(MakeArrayView(OutState.TargetWeights), 0, *TargetWeightKeys))
+			{
+				PCGLog::Metadata::LogFailToGetAttributeError(TargetWeightSelector, Context);
+				return EPCGTimeSliceInitResult::NoOperation;
+			}
+		}
+
+		for (const auto& SourceAndTargetMapping : Settings->SourceAndTargetAttributeMapping)
+		{
+			const FPCGAttributePropertyInputSelector SourceSelector = SourceAndTargetMapping.Key.CopyAndFixLast(SourcePointData);
+
+			TUniquePtr<const IPCGAttributeAccessor> SourceAccessor = PCGAttributeAccessorHelpers::CreateConstAccessor(SourcePointData, SourceSelector);
+			TUniquePtr<const IPCGAttributeAccessorKeys> SourceKeys = PCGAttributeAccessorHelpers::CreateConstKeys(SourcePointData, SourceSelector);
+
+			if (!SourceAccessor || !SourceKeys)
+			{
+				PCGLog::Metadata::LogFailToCreateAccessorError(SourceSelector, Context);
+				continue;
+			}
+
+			const FPCGAttributePropertyInputSelector TargetSelector = SourceAndTargetMapping.Value.CopyAndFixLast(TargetPointData);
+
+			TUniquePtr<const IPCGAttributeAccessor> TargetAccessor = PCGAttributeAccessorHelpers::CreateConstAccessor(TargetPointData, TargetSelector);
+			TUniquePtr<const IPCGAttributeAccessorKeys> TargetKeys = PCGAttributeAccessorHelpers::CreateConstKeys(TargetPointData, TargetSelector);
+
+			if (!TargetAccessor || !TargetKeys)
+			{
+				PCGLog::Metadata::LogFailToCreateAccessorError(SourceSelector, Context);
+				continue;
+			}
+
+			// Finally, validate that the target attribute type can be broadcasted to the source type.
+			if (!PCG::Private::IsBroadcastableOrConstructible(TargetAccessor->GetUnderlyingType(), SourceAccessor->GetUnderlyingType()))
+			{
+				PCGLog::Metadata::LogIncomparableAttributesError(TargetSelector, SourceSelector, Context);
+				continue;
+			}
+			
+			FPCGAttributePropertyOutputSelector OutputSelector;
+			OutputSelector.ImportFromOtherSelector(SourceSelector);
+			OutputSelector = OutputSelector.CopyAndFixSource(&SourceSelector, OutPointData);
+
+			TUniquePtr<IPCGAttributeAccessor> OutputAccessor = PCGAttributeAccessorHelpers::CreateAccessor(OutPointData, OutputSelector);
+
+			if (!OutputAccessor)
+			{
+				PCGLog::Metadata::LogFailToCreateAccessorError(OutputSelector, Context);
+				continue;
+			}
+
+			PCGAttractElement::FAttractState::FAttributeInterpolationData& Interpolator = OutState.Interpolators.Emplace_GetRef();
+			Interpolator.SourceAccessor = std::move(SourceAccessor);
+			Interpolator.SourceKeys = std::move(SourceKeys);
+			Interpolator.TargetAccessor = std::move(TargetAccessor);
+			Interpolator.TargetKeys = std::move(TargetKeys);
+			Interpolator.OutputAccessor = std::move(OutputAccessor);
+		}
+
+		if (Settings->bOutputAttractIndex)
+		{
+			const FPCGAttributePropertyOutputNoSourceSelector AttractIndexSelector = Settings->OutputAttractIndexAttribute;
+
+			// Make sure we create the attribute if needed.
+			if(AttractIndexSelector.IsBasicAttribute())
+			{
+				OutPointData->Metadata->FindOrCreateAttribute<int32>(AttractIndexSelector.GetName(), 0, /*bAllowInterpolation=*/false);
+			}
+
+			TUniquePtr<IPCGAttributeAccessor> AttractIndexAccessor = PCGAttributeAccessorHelpers::CreateAccessor(OutPointData, AttractIndexSelector);
+			if (!AttractIndexAccessor)
+			{
+				PCGLog::Metadata::LogFailToCreateAccessorError(AttractIndexSelector, Context);
+				return EPCGTimeSliceInitResult::NoOperation;
+			}
+
+			OutState.OutputAttractIndexAccessor = std::move(AttractIndexAccessor);
+		}
+
+		return EPCGTimeSliceInitResult::Success;
+	});
+
+	return true;
+}
+
+bool FPCGAttractElement::ExecuteInternal(FPCGContext* InContext) const
+{
+	TRACE_CPUPROFILER_EVENT_SCOPE(FPCGAttractElement::Execute);
+
+	ContextType* TimeSlicedContext = static_cast<ContextType*>(InContext);
+	check(TimeSlicedContext);
+
+	if (!TimeSlicedContext->DataIsPreparedForExecution())
+	{
+		return true;
+	}
+
+	if (TimeSlicedContext->GetExecutionStateResult() == EPCGTimeSliceInitResult::NoOperation)
+	{
+		TimeSlicedContext->OutputData = TimeSlicedContext->InputData;
+		return true;
+	}
+
+	const UPCGAttractSettings* Settings = TimeSlicedContext->GetInputSettings<UPCGAttractSettings>();
+	check(Settings);
+
+	return ExecuteSlice(TimeSlicedContext, [Settings](const ContextType* Context, const ExecStateType& ExecState, IterStateType& AttractData, const uint32 IterIndex)
+	{
+		if (Context->GetIterationStateResult(IterIndex) == EPCGTimeSliceInitResult::NoOperation)
+		{
+			return true;
+		}
+
+		if (!AttractData.bAttractPhaseDone)
+		{
+			if (Settings->Mode == EPCGAttractMode::FromIndex)
+			{
+				// Nothing to do, we've already gotten the data
+			}
+			else if (!ExecState.AttractFunction(Context, AttractData, Settings))
+			{
+				return false;
+			}
+
+			AttractData.bAttractPhaseDone = true;
+		}
+
+		// Copy points from source that are kept across
+		TArray<int32> Mapping;
+
+		if (Settings->bRemoveUnattractedPoints)
+		{
+			TArray<FPCGPoint>& OutPoints = AttractData.OutPointData->GetMutablePoints();
+			const TArray<FPCGPoint>& SourcePoints = AttractData.SourcePointData->GetPoints();
+
+			Mapping.Reserve(AttractData.Mapping.Num());
+
+			OutPoints.Reserve(SourcePoints.Num());
+			for (int PointIndex = 0; PointIndex < SourcePoints.Num(); ++PointIndex)
+			{
+				if (AttractData.Mapping[PointIndex] != INDEX_NONE)
+				{
+					OutPoints.Add(SourcePoints[PointIndex]);
+					Mapping.Add(AttractData.Mapping[PointIndex]);
+				}
+			}
+		}
+		else
+		{
+			AttractData.OutPointData->GetMutablePoints() = AttractData.SourcePointData->GetPoints();
+			Mapping = AttractData.Mapping;
+		}
+
+		// After we've created the points, we can now create the keys
+		AttractData.OutputKeys = PCGAttributeAccessorHelpers::CreateKeys(AttractData.OutPointData, FPCGAttributePropertyOutputSelector());
+
+		// Set the attact index if required
+		if (AttractData.OutputAttractIndexAccessor)
+		{
+			AttractData.OutputAttractIndexAccessor->SetRange<int32>(Mapping, 0, *AttractData.OutputKeys);
+		}
+
+		// Apply weighting from source to target
+		while (AttractData.InterpolationIndex < AttractData.Interpolators.Num())
+		{
+			PCGAttractElement::Helpers::InterpolatePoints(AttractData, AttractData.InterpolationIndex++, Settings);
+
+			// Check for time-slice
+			if (AttractData.InterpolationIndex != AttractData.Interpolators.Num() && Context->ShouldStop())
+			{
+				return false;
+			}
+		}
+
+		AttractData.OutPointData->GetMutablePoints().Shrink();
+
+		return true;
+	});
+}
+
+#undef LOCTEXT_NAMESPACE
diff --git a/Engine/Plugins/PCG/Source/PCG/Public/Elements/PCGAttractElement.h b/Engine/Plugins/PCG/Source/PCG/Public/Elements/PCGAttractElement.h
new file mode 100644
--- /dev/null
+++ b/Engine/Plugins/PCG/Source/PCG/Public/Elements/PCGAttractElement.h
@@ -0,0 +1,141 @@
+// Copyright Epic Games, Inc. All Rights Reserved.
+
+#pragma once
+
+#include "PCGSettings.h"
+
+#include "Elements/PCGTimeSlicedElementBase.h"
+#include "Metadata/PCGMetadataAttributeTpl.h"
+#include "Metadata/Accessors/PCGPropertyAccessor.h"
+
+#include "PCGAttractElement.generated.h"
+
+class UPCGPointData;
+
+UENUM(BlueprintType)
+enum class EPCGAttractMode : uint8
+{
+	Closest UMETA(Tooltip = "Attract to nearest target."),
+	MinAttribute UMETA(Tooltip = "Attract to the minimum target inside the search radius."),
+	MaxAttribute UMETA(Tooltip = "Attract to the maximum target inside the search radius."),
+	FromIndex UMETA(Tooltip = "Attracts to a fixed target point based on an attribute.")
+};
+
+/** Attracts points (interpolates) from the source towards points from the target. */
+UCLASS(MinimalAPI, BlueprintType, ClassGroup = (Procedural))
+class UPCGAttractSettings : public UPCGSettings
+{
+	GENERATED_BODY()
+
+public:
+	UPCGAttractSettings();
+
+	//~Begin UPCGSettings interface
+#if WITH_EDITOR
+	virtual FName GetDefaultNodeName() const override { return FName(TEXT("AttractElement")); }
+	virtual FText GetDefaultNodeTitle() const override { return NSLOCTEXT("PCGAttractElement", "NodeTitle", "Attract"); }
+	virtual EPCGSettingsType GetType() const override { return EPCGSettingsType::PointOps; }
+	virtual FText GetNodeTooltipText() const override;
+#endif
+
+protected:
+	virtual TArray<FPCGPinProperties> InputPinProperties() const override;
+	virtual TArray<FPCGPinProperties> OutputPinProperties() const override { return Super::DefaultPointOutputPinProperties(); }
+	virtual FPCGElementPtr CreateElement() const override;
+	//~End UPCGSettings interface
+
+public:
+	/** Controls the criteria used for the attract operation. */
+	UPROPERTY(EditAnywhere, BlueprintReadWrite, Category = Settings, meta = (PCG_Overridable))
+	EPCGAttractMode Mode = EPCGAttractMode::Closest;
+
+	/** Index attribute on the source that maps a point to a point from the target. */
+	UPROPERTY(EditAnywhere, BlueprintReadWrite, Category = Settings, meta = (EditCondition = "Mode == EPCGAttractMode::FromIndex", EditConditionHides, PCG_Overridable))
+	FPCGAttributePropertyInputSelector AttractorIndexAttribute;
+
+	/** Will be used to determine which points to attract. */
+	UPROPERTY(EditAnywhere, BlueprintReadWrite, Category = Settings, meta = (EditCondition = "Mode != EPCGAttractMode::FromIndex", ClampMin = "0.01", PCG_Overridable))
+	double Distance = 100.0;
+
+	/** Can optionally remove points that weren't attracted to points on the target. Will have no effect when the source is the target. */
+	UPROPERTY(EditAnywhere, BlueprintReadWrite, Category = Settings, meta = (PCG_Overridable))
+	bool bRemoveUnattractedPoints = false;
+
+	/** The target attribute used when attracted by attribute value. */
+	UPROPERTY(EditAnywhere, BlueprintReadWrite, Category = Settings, meta = (EditCondition = "Mode == EPCGAttractMode::MinAttribute || Mode == EPCGAttractMode::MaxAttribute", EditConditionHides, PCG_Overridable))
+	FPCGAttributePropertyInputSelector TargetAttribute;
+
+	UPROPERTY(EditAnywhere, BlueprintReadWrite, Category = Settings, meta = (InlineEditConditionToggle, PCG_Overridable))
+	bool bUseSourceWeight = false;
+
+	/** This attribute will determine the weight of the fusion result for the source point. It will be normalized to the range of [0..1]. */
+	UPROPERTY(EditAnywhere, BlueprintReadWrite, Category = Settings, meta = (EditCondition = "bUseSourceWeight", PCG_Overridable))
+	FPCGAttributePropertyInputSelector SourceWeightAttribute;
+
+	UPROPERTY(EditAnywhere, BlueprintReadWrite, Category = Settings, meta = (InlineEditConditionToggle, PCG_Overridable))
+	bool bUseTargetWeight = false;
+
+	/** This attribute will determine the weight of the fusion result for the target point. It will be normalized to the range of [0..1]. */
+	UPROPERTY(EditAnywhere, BlueprintReadWrite, Category = Settings, meta = (EditCondition = "bUseTargetWeight", PCG_Overridable))
+	FPCGAttributePropertyInputSelector TargetWeightAttribute;
+
+	UPROPERTY(EditAnywhere, BlueprintReadWrite, Category = Settings, meta = (ClampMin = "0.0", ClampMax = "1.0", EditCondition = "!bUseSourceWeight && !bUseTargetWeight", PCG_Overridable))
+	double Weight = 1.0;
+	
+	UPROPERTY(EditAnywhere, BlueprintReadWrite, Category = Settings)
+	TMap<FPCGAttributePropertyInputSelector, FPCGAttributePropertyInputSelector> SourceAndTargetAttributeMapping;
+
+	UPROPERTY(EditAnywhere, BlueprintReadWrite, Category = Settings, meta = (InlineEditConditionToggle, PCG_Overridable))
+	bool bOutputAttractIndex = false;
+
+	UPROPERTY(EditAnywhere, BlueprintReadWrite, Category = Settings, meta = (EditCondition = "bOutputAttractIndex", PCG_Overridable))
+	FPCGAttributePropertyOutputNoSourceSelector OutputAttractIndexAttribute;
+};
+
+namespace PCGAttractElement
+{
+	struct FAttractState
+	{
+		const UPCGPointData* TargetPointData = nullptr;
+		const UPCGPointData* SourcePointData = nullptr;
+		UPCGPointData* OutPointData = nullptr;
+
+		TArray<int32> Mapping;
+		TArray<int32> SortedTargetIndices;
+
+		TUniquePtr<IPCGAttributeAccessorKeys> OutputKeys;
+		TUniquePtr<IPCGAttributeAccessor> OutputAttractIndexAccessor;
+
+		TArray<double> SourceWeights;
+		TArray<double> TargetWeights;
+
+		struct FAttributeInterpolationData
+		{
+			TUniquePtr<const IPCGAttributeAccessor> SourceAccessor;
+			TUniquePtr<const IPCGAttributeAccessorKeys> SourceKeys;
+			TUniquePtr<const IPCGAttributeAccessor> TargetAccessor;
+			TUniquePtr<const IPCGAttributeAccessorKeys> TargetKeys;
+			TUniquePtr<IPCGAttributeAccessor> OutputAccessor;
+		};
+
+		TArray<FAttributeInterpolationData> Interpolators;
+
+		int32 IterationIndex = INDEX_NONE;
+		bool bTargetIsSource = false;
+		bool bAttractPhaseDone = false;
+		int32 InterpolationIndex = 0;
+	};
+
+	struct FExecState
+	{
+		using AttractSignature = bool(*)(const FPCGContext* InContext, FAttractState& AttractData, const UPCGAttractSettings* AttractSettings);
+		AttractSignature AttractFunction = nullptr;
+	};
+}
+
+class FPCGAttractElement : public TPCGTimeSlicedElementBase<PCGAttractElement::FExecState, PCGAttractElement::FAttractState>
+{
+protected:
+	virtual bool PrepareDataInternal(FPCGContext* InContext) const override;
+	virtual bool ExecuteInternal(FPCGContext* InContext) const override;
+};
