Added exponential on grid size that shifts to the next (& more) cell sizes.

diff --git a/Engine/Plugins/PCG/Source/PCG/Private/Elements/PCGHiGenGridSize.cpp b/Engine/Plugins/PCG/Source/PCG/Private/Elements/PCGHiGenGridSize.cpp
--- a/Engine/Plugins/PCG/Source/PCG/Private/Elements/PCGHiGenGridSize.cpp
+++ b/Engine/Plugins/PCG/Source/PCG/Private/Elements/PCGHiGenGridSize.cpp
@@ -18,6 +18,24 @@ namespace PCGHiGenGridSizeConstants
 	const FName CellVolumeOutputLabel = TEXT("Grid Cell Volume");
 }
 
+EPCGHiGenGrid UPCGHiGenGridSizeSettings::GetGrid() const
+{
+	if (GetOuter() && GetOuter()->GetOuter() && GetOuter()->GetOuter()->IsA<UPCGGraph>())
+	{
+		UPCGGraph* Graph = Cast<UPCGGraph>(GetOuter()->GetOuter());
+		return static_cast<EPCGHiGenGrid>(static_cast<uint32>(HiGenGridSize) << Graph->GetGridExponential());
+	}
+	else
+	{
+		return HiGenGridSize;
+	}
+}
+
+uint32 UPCGHiGenGridSizeSettings::GetGridSize() const
+{
+	return (HiGenGridSize == EPCGHiGenGrid::Unbounded) ? PCGHiGenGrid::UnboundedGridSize() : PCGHiGenGrid::GridToGridSize(GetGrid());
+}
+
 #if WITH_EDITOR
 FName UPCGHiGenGridSizeSettings::GetDefaultNodeName() const
 {
diff --git a/Engine/Plugins/PCG/Source/PCG/Private/PCGCommon.cpp b/Engine/Plugins/PCG/Source/PCG/Private/PCGCommon.cpp
--- a/Engine/Plugins/PCG/Source/PCG/Private/PCGCommon.cpp
+++ b/Engine/Plugins/PCG/Source/PCG/Private/PCGCommon.cpp
@@ -119,6 +119,17 @@ double FPCGRuntimeGenerationRadii::GetGenerationRadiusFromGrid(EPCGHiGenGrid Gri
 		case EPCGHiGenGrid::Grid512: return GenerationRadius51200;
 		case EPCGHiGenGrid::Grid1024: return GenerationRadius102400;
 		case EPCGHiGenGrid::Grid2048: return GenerationRadius204800;
+		case EPCGHiGenGrid::Grid4096: return GenerationRadius204800 * (1 << 1);
+		case EPCGHiGenGrid::Grid8192: return GenerationRadius204800 * (1 << 2);
+		case EPCGHiGenGrid::Grid16384: return GenerationRadius204800 * (1 << 3);
+		case EPCGHiGenGrid::Grid32768: return GenerationRadius204800 * (1 << 4);
+		case EPCGHiGenGrid::Grid65536: return GenerationRadius204800 * (1 << 5);
+		case EPCGHiGenGrid::Grid131072: return GenerationRadius204800 * (1 << 6);
+		case EPCGHiGenGrid::Grid262144: return GenerationRadius204800 * (1 << 7);
+		case EPCGHiGenGrid::Grid524288: return GenerationRadius204800 * (1 << 8);
+		case EPCGHiGenGrid::Grid1048576: return GenerationRadius204800 * (1 << 9);
+		case EPCGHiGenGrid::Grid2097152: return GenerationRadius204800 * (1 << 10);
+		case EPCGHiGenGrid::Grid4194304: return GenerationRadius204800 * (1 << 11);
 		case EPCGHiGenGrid::Unbounded: return GenerationRadius;
 	}
 
diff --git a/Engine/Plugins/PCG/Source/PCG/Private/PCGComponent.cpp b/Engine/Plugins/PCG/Source/PCG/Private/PCGComponent.cpp
--- a/Engine/Plugins/PCG/Source/PCG/Private/PCGComponent.cpp
+++ b/Engine/Plugins/PCG/Source/PCG/Private/PCGComponent.cpp
@@ -1093,7 +1093,7 @@ double UPCGComponent::GetGenerationRadiusFromGrid(EPCGHiGenGrid Grid) const
 	const UPCGGraph* Graph = GetGraph();
 	if (ensure(Graph))
 	{
-		return Graph->GenerationRadii.GetGenerationRadiusFromGrid(Grid);
+		return Graph->GetGridGenerationRadiusFromGrid(Grid);
 	}
 
 	return 0;
@@ -1109,7 +1109,7 @@ double UPCGComponent::GetCleanupRadiusFromGrid(EPCGHiGenGrid Grid) const
 	const UPCGGraph* Graph = GetGraph();
 	if (ensure(Graph))
 	{
-		return Graph->GenerationRadii.GetCleanupRadiusFromGrid(Grid);
+		return Graph->GetGridCleanupRadiusFromGrid(Grid);
 	}
 
 	return 0;
diff --git a/Engine/Plugins/PCG/Source/PCG/Private/PCGGraph.cpp b/Engine/Plugins/PCG/Source/PCG/Private/PCGGraph.cpp
--- a/Engine/Plugins/PCG/Source/PCG/Private/PCGGraph.cpp
+++ b/Engine/Plugins/PCG/Source/PCG/Private/PCGGraph.cpp
@@ -530,7 +530,7 @@ void UPCGGraph::AddReferencedObjects(UObject* InThis, FReferenceCollector& Colle
 uint32 UPCGGraph::GetDefaultGridSize() const
 {
 	ensure(IsHierarchicalGenerationEnabled());
-	return PCGHiGenGrid::IsValidGrid(HiGenGridSize) ? PCGHiGenGrid::GridToGridSize(HiGenGridSize) : PCGHiGenGrid::UnboundedGridSize();
+	return PCGHiGenGrid::IsValidGrid(HiGenGridSize) ? (PCGHiGenGrid::GridToGridSize(HiGenGridSize) * (1 << HiGenExponential)) : PCGHiGenGrid::UnboundedGridSize();
 }
 
 UPCGNode* UPCGGraph::AddNodeOfType(TSubclassOf<class UPCGSettings> InSettingsClass, UPCGSettings*& OutDefaultNodeSettings)
@@ -1087,6 +1087,46 @@ void UPCGGraph::GetGridSizes(PCGHiGenGrid::FSizeArray& OutGridSizes, bool& bOutH
 	return;
 }
 
+double UPCGGraph::GetGridGenerationRadiusFromGrid(EPCGHiGenGrid Grid) const
+{
+	if (Grid == EPCGHiGenGrid::Unbounded || Grid == EPCGHiGenGrid::Uninitialized)
+	{
+		return GenerationRadii.GetGenerationRadiusFromGrid(Grid);
+	}
+	// If the queried grid is smaller than the min grid including the exponent, we'll take the min grid and scale it down
+	else if (static_cast<uint32>(Grid) < (static_cast<uint32>(EPCGHiGenGrid::GridMin) << HiGenExponential))
+	{
+		check(static_cast<uint32>(Grid) >= static_cast<uint32>(EPCGHiGenGrid::GridMin));
+		uint32 Multiplier = static_cast<uint32>(Grid) / static_cast<uint32>(EPCGHiGenGrid::GridMin);
+		return GenerationRadii.GetGenerationRadiusFromGrid(EPCGHiGenGrid::GridMin) * Multiplier;
+	}
+	else
+	{
+		EPCGHiGenGrid AdjustedGrid = static_cast<EPCGHiGenGrid>(static_cast<uint32>(Grid) >> HiGenExponential);
+		return GenerationRadii.GetGenerationRadiusFromGrid(AdjustedGrid) * (1ULL << HiGenExponential);
+	}
+}
+
+double UPCGGraph::GetGridCleanupRadiusFromGrid(EPCGHiGenGrid Grid) const
+{
+	if (Grid == EPCGHiGenGrid::Unbounded || Grid == EPCGHiGenGrid::Uninitialized)
+	{
+		return GenerationRadii.GetGenerationRadiusFromGrid(Grid);
+	}
+	// If the queried grid is smaller than the min grid including the exponent, we'll take the min grid and scale it down
+	else if (static_cast<uint32>(Grid) < (static_cast<uint32>(EPCGHiGenGrid::GridMin) << HiGenExponential))
+	{
+		check(static_cast<uint32>(Grid) >= static_cast<uint32>(EPCGHiGenGrid::GridMin));
+		uint32 Multiplier = static_cast<uint32>(Grid) / static_cast<uint32>(EPCGHiGenGrid::GridMin);
+		return GenerationRadii.GetCleanupRadiusFromGrid(EPCGHiGenGrid::GridMin) * Multiplier;
+	}
+	else
+	{
+		EPCGHiGenGrid AdjustedGrid = static_cast<EPCGHiGenGrid>(static_cast<uint32>(Grid) >> HiGenExponential);
+		return GenerationRadii.GetCleanupRadiusFromGrid(AdjustedGrid) * (1ULL << HiGenExponential);
+	}
+}
+
 #if WITH_EDITOR
 void UPCGGraph::DisableNotificationsForEditor()
 {
@@ -1452,7 +1492,8 @@ void UPCGGraph::PostEditChangeProperty(FPropertyChangedEvent& PropertyChangedEve
 	}
 	else if (PropertyName == GET_MEMBER_NAME_CHECKED(UPCGGraph, HiGenGridSize)
 		|| PropertyName == GET_MEMBER_NAME_CHECKED(UPCGGraph, bUseHierarchicalGeneration)
-		|| PropertyName == GET_MEMBER_NAME_CHECKED(UPCGGraph, bUse2DGrid))
+		|| PropertyName == GET_MEMBER_NAME_CHECKED(UPCGGraph, bUse2DGrid)
+		|| PropertyName == GET_MEMBER_NAME_CHECKED(UPCGGraph, HiGenExponential))
 	{
 		// The higen settings change the structure of the graph (presence or absence of links between grid levels).
 		NotifyGraphChanged(EPCGChangeType::Structural | EPCGChangeType::GenerationGrid);
diff --git a/Engine/Plugins/PCG/Source/PCG/Public/Elements/PCGHiGenGridSize.h b/Engine/Plugins/PCG/Source/PCG/Public/Elements/PCGHiGenGridSize.h
--- a/Engine/Plugins/PCG/Source/PCG/Public/Elements/PCGHiGenGridSize.h
+++ b/Engine/Plugins/PCG/Source/PCG/Public/Elements/PCGHiGenGridSize.h
@@ -15,8 +15,8 @@ class UPCGHiGenGridSizeSettings : public UPCGSettings
 	GENERATED_BODY()
 
 public:
-	EPCGHiGenGrid GetGrid() const { return HiGenGridSize; }
-	uint32 GetGridSize() const { return (HiGenGridSize == EPCGHiGenGrid::Unbounded) ? PCGHiGenGrid::UnboundedGridSize() : PCGHiGenGrid::GridToGridSize(HiGenGridSize); }
+	EPCGHiGenGrid GetGrid() const;
+	uint32 GetGridSize() const;
 
 	//~Begin UPCGSettings interface
 #if WITH_EDITOR
diff --git a/Engine/Plugins/PCG/Source/PCG/Public/PCGCommon.h b/Engine/Plugins/PCG/Source/PCG/Public/PCGCommon.h
--- a/Engine/Plugins/PCG/Source/PCG/Public/PCGCommon.h
+++ b/Engine/Plugins/PCG/Source/PCG/Public/PCGCommon.h
@@ -236,9 +236,20 @@ enum class EPCGHiGenGrid : uint32
 	Grid512 = 512 UMETA(DisplayName = "51200"),
 	Grid1024 = 1024 UMETA(DisplayName = "102400"),
 	Grid2048 = 2048 UMETA(DisplayName = "204800"),
-	
+	Grid4096 = 4096 UMETA(Hidden),
+	Grid8192 = 8192 UMETA(Hidden),
+	Grid16384 = 16384 UMETA(Hidden),
+	Grid32768 = 32768 UMETA(Hidden),
+	Grid65536 = 65536 UMETA(Hidden),
+	Grid131072 = 131072 UMETA(Hidden),
+	Grid262144 = 262144 UMETA(Hidden),
+	Grid524288 = 524288 UMETA(Hidden),
+	Grid1048576 = 1048576 UMETA(Hidden),
+	Grid2097152 = 2097152 UMETA(Hidden),
+	Grid4194304 = 4194304 UMETA(Hidden),
+
 	GridMin = Grid4 UMETA(Hidden),
-	GridMax = Grid2048 UMETA(Hidden),
+	GridMax = Grid4194304 UMETA(Hidden),
 
 	// Should execute once rather than executing on any grid
 	Unbounded = 1u << 31,
diff --git a/Engine/Plugins/PCG/Source/PCG/Public/PCGGraph.h b/Engine/Plugins/PCG/Source/PCG/Public/PCGGraph.h
--- a/Engine/Plugins/PCG/Source/PCG/Public/PCGGraph.h
+++ b/Engine/Plugins/PCG/Source/PCG/Public/PCGGraph.h
@@ -254,7 +254,7 @@ public:
 	bool bIgnoreLandscapeTracking = false;
 #endif
 
-	UPROPERTY(EditAnywhere, Category = Settings)
+	UPROPERTY(EditAnywhere, Category = "Settings|Advanced")
 	bool bLandscapeUsesMetadata = true;
 
 	/** Creates a node using the given settings interface. Does not manage ownership - done outside of this method. */
@@ -330,6 +330,15 @@ public:
 	/** Determine the relevant grid sizes by inspecting all HiGenGridSize nodes. */
 	void GetGridSizes(PCGHiGenGrid::FSizeArray& OutGridSizes, bool& bOutHasUnbounded) const;
 
+	/** Returns exponential on grid size, which represents a shift in the grid */
+	uint32 GetGridExponential() const { return HiGenExponential; }
+
+	/** Gets generation radius from grid, considering grid exponential. */
+	double GetGridGenerationRadiusFromGrid(EPCGHiGenGrid Grid) const;
+
+	/** Gets cleanup radius from grid, considering grid exponential. */
+	double GetGridCleanupRadiusFromGrid(EPCGHiGenGrid Grid) const;
+
 #if WITH_EDITOR
 	void DisableNotificationsForEditor();
 	void EnableNotificationsForEditor();
@@ -424,6 +433,10 @@ protected:
 	UPROPERTY(EditAnywhere, Category = Settings, meta = (DisplayName = "HiGen Default Grid Size", EditCondition = "bUseHierarchicalGeneration"))
 	EPCGHiGenGrid HiGenGridSize = EPCGHiGenGrid::Grid256;
 
+	/** Shifts the grid sizes upwards based on the value, which allows to use larger grids. A value of 1 will effectively use the graph's Grid-400 values x 2 for the actual Grid-800 sizes and so on. */
+	UPROPERTY(EditAnywhere, Category = Settings, meta = (ClampMax = "10", DisplayName = "HiGen Grid Size Exponential", EditCondition = "bUseHierarchicalGeneration"))
+	uint32 HiGenExponential = 0;
+
 	UPROPERTY(EditAnywhere, Category = Settings, meta = (DisplayName = "2D Grid"))
 	bool bUse2DGrid = true;
 
diff --git a/Engine/Plugins/PCG/Source/PCGEditor/Private/SPCGEditorGraphNode.cpp b/Engine/Plugins/PCG/Source/PCGEditor/Private/SPCGEditorGraphNode.cpp
--- a/Engine/Plugins/PCG/Source/PCGEditor/Private/SPCGEditorGraphNode.cpp
+++ b/Engine/Plugins/PCG/Source/PCGEditor/Private/SPCGEditorGraphNode.cpp
@@ -974,6 +974,17 @@ FLinearColor SPCGEditorGraphNode::GetGridLabelColor(EPCGHiGenGrid NodeGrid)
 	{
 	case EPCGHiGenGrid::Unbounded:
 		return FColor(255, 255, 255, 255);
+	case EPCGHiGenGrid::Grid4194304: // fall-through
+	case EPCGHiGenGrid::Grid2097152: // fall-through
+	case EPCGHiGenGrid::Grid1048576: // fall-through
+	case EPCGHiGenGrid::Grid524288: // fall-through
+	case EPCGHiGenGrid::Grid262144: // fall-through
+	case EPCGHiGenGrid::Grid131072: // fall-through
+	case EPCGHiGenGrid::Grid65536: // fall-through
+	case EPCGHiGenGrid::Grid32768: // fall-through
+	case EPCGHiGenGrid::Grid16384: // fall-through
+	case EPCGHiGenGrid::Grid8192: // fall-through
+	case EPCGHiGenGrid::Grid4096: // fall-through
 	case EPCGHiGenGrid::Grid2048:
 		return FColor(53, 60, 171, 255);
 	case EPCGHiGenGrid::Grid1024:
