Small UX improvement on the add component node

#rnx

diff --git a/Engine/Plugins/PCG/Source/PCG/Private/Elements/PCGAddComponent.cpp b/Engine/Plugins/PCG/Source/PCG/Private/Elements/PCGAddComponent.cpp
--- a/Engine/Plugins/PCG/Source/PCG/Private/Elements/PCGAddComponent.cpp
+++ b/Engine/Plugins/PCG/Source/PCG/Private/Elements/PCGAddComponent.cpp
@@ -348,7 +348,7 @@ bool FPCGAddComponentElement::ExecuteInternal(FPCGContext* InContext) const
 		UPCGData* OutputData = OriginalData->DuplicateData(Context);
 
 		// Add component reference attribute
-		FPCGAttributePropertyOutputSelector ComponentReferenceSelector = Settings->ComponentReferenceAttribute; // No copy source here because there's no 1:1 mapping
+		FPCGAttributePropertyOutputNoSourceSelector ComponentReferenceSelector = Settings->ComponentReferenceAttribute; // No copy source here because there's no 1:1 mapping
 
 		// Create attribute if needed
 		if (ComponentReferenceSelector.IsBasicAttribute())
diff --git a/Engine/Plugins/PCG/Source/PCG/Public/Elements/PCGAddComponent.h b/Engine/Plugins/PCG/Source/PCG/Public/Elements/PCGAddComponent.h
--- a/Engine/Plugins/PCG/Source/PCG/Public/Elements/PCGAddComponent.h
+++ b/Engine/Plugins/PCG/Source/PCG/Public/Elements/PCGAddComponent.h
@@ -97,7 +97,7 @@ public:
 
 	/** Specifies what attribute to write the component reference to. */
 	UPROPERTY(BlueprintReadWrite, EditAnywhere, Category = Settings, meta = (PCG_Overridable, PCG_DiscardPropertySelection, PCG_DiscardExtraSelection))
-	FPCGAttributePropertyOutputSelector ComponentReferenceAttribute;
+	FPCGAttributePropertyOutputNoSourceSelector ComponentReferenceAttribute;
 
 protected:
 #if WITH_EDITOR
