Added reflection direction in world ray hit

diff --git a/Engine/Plugins/PCG/Source/PCG/Private/Data/PCGWorldData.cpp b/Engine/Plugins/PCG/Source/PCG/Private/Data/PCGWorldData.cpp
--- a/Engine/Plugins/PCG/Source/PCG/Private/Data/PCGWorldData.cpp
+++ b/Engine/Plugins/PCG/Source/PCG/Private/Data/PCGWorldData.cpp
@@ -238,7 +238,7 @@ bool UPCGWorldRayHitData::SamplePoint(const FTransform& InTransform, const FBox&
 			const FHitResult& Hit = HitResult.GetValue();
 			OutPoint = FPCGPoint(PCGWorldQueryHelpers::GetOrthonormalImpactTransform(Hit), 1.0f, UPCGBlueprintHelpers::ComputeSeedFromPosition(Hit.Location));
 			// TODO: Pre-create attributes within caller or find a better solution than per point
-			PCGWorldQueryHelpers::ApplyRayHitMetadata(Hit, QueryParams, OutPoint, OutMetadata, World, /*bShouldCreateAttributes=*/true);
+			PCGWorldQueryHelpers::ApplyRayHitMetadata(Hit, QueryParams, QueryParams.RayDirection, OutPoint, OutMetadata, World, /*bShouldCreateAttributes=*/true);
 			return true;
 		}
 	}
diff --git a/Engine/Plugins/PCG/Source/PCG/Private/Elements/PCGWorldRaycast.cpp b/Engine/Plugins/PCG/Source/PCG/Private/Elements/PCGWorldRaycast.cpp
--- a/Engine/Plugins/PCG/Source/PCG/Private/Elements/PCGWorldRaycast.cpp
+++ b/Engine/Plugins/PCG/Source/PCG/Private/Elements/PCGWorldRaycast.cpp
@@ -348,7 +348,7 @@ bool FPCGWorldRaycastElement::ExecuteInternal(FPCGContext* InContext) const
 				OutPoint.Density = 1.f;
 				UPCGBlueprintHelpers::SetSeedFromPosition(OutPoint);
 
-				bAttributeSuccess &= PCGWorldQueryHelpers::ApplyRayHitMetadata(Hit, WorldQueryParams, OutputPoints.Last(), OutMetadata, World, /*bShouldCreateAttributes=*/false);
+				bAttributeSuccess &= PCGWorldQueryHelpers::ApplyRayHitMetadata(Hit, WorldQueryParams, Direction, OutputPoints.Last(), OutMetadata, World, /*bShouldCreateAttributes=*/false);
 			}
 			else if (Settings->bKeepOriginalPointOnMiss)
 			{
diff --git a/Engine/Plugins/PCG/Source/PCG/Private/Helpers/PCGWorldQueryHelpers.cpp b/Engine/Plugins/PCG/Source/PCG/Private/Helpers/PCGWorldQueryHelpers.cpp
--- a/Engine/Plugins/PCG/Source/PCG/Private/Helpers/PCGWorldQueryHelpers.cpp
+++ b/Engine/Plugins/PCG/Source/PCG/Private/Helpers/PCGWorldQueryHelpers.cpp
@@ -195,6 +195,7 @@ namespace PCGWorldQueryHelpers
 		bResult &= CreateAttribute(PCGWorldQueryConstants::ImpactAttribute, QueryParams.bGetImpact, true);
 		bResult &= CreateAttribute(PCGWorldQueryConstants::ImpactPointAttribute, QueryParams.bGetImpactPoint, FVector::ZeroVector);
 		bResult &= CreateAttribute(PCGWorldQueryConstants::ImpactNormalAttribute, QueryParams.bGetImpactNormal, FVector::ZeroVector);
+		bResult &= CreateAttribute(PCGWorldQueryConstants::ImpactReflectionAttribute, QueryParams.bGetReflection, FVector::ZeroVector);
 		bResult &= CreateAttribute(PCGWorldQueryConstants::ImpactDistanceAttribute, QueryParams.bGetDistance, 0.0);
 		bResult &= CreateAttribute(PCGWorldQueryConstants::LocalImpactPointAttribute, QueryParams.bGetLocalImpactPoint, FVector::ZeroVector);
 		bResult &= CreateAttribute(PCGPointDataConstants::ActorReferenceAttribute, QueryParams.bGetReferenceToActorHit, FSoftObjectPath());
@@ -234,6 +235,7 @@ namespace PCGWorldQueryHelpers
 	bool ApplyRayHitMetadata(
 		const FHitResult& HitResult,
 		const FPCGWorldRaycastQueryParams& QueryParams,
+		const FVector& RayDirection,
 		FPCGPoint& OutPoint,
 		UPCGMetadata* OutMetadata,
 		TWeakObjectPtr<UWorld> World,
@@ -268,10 +270,18 @@ namespace PCGWorldQueryHelpers
 
 		const FHitResult& Hit = HitResult;
 
+		FVector ReflectionVector = FVector::ZeroVector;
+		if (QueryParams.bGetReflection)
+		{
+			ReflectionVector = RayDirection - 2.0 * (FVector(Hit.ImpactNormal) | RayDirection) * FVector(Hit.ImpactNormal);
+			ReflectionVector.Normalize();
+		}
+
 		bool bResult = true;
 		// Note: The T/F Impact attribute is true by default, so no need to set it directly.
 		bResult &= ApplyAttribute(PCGWorldQueryConstants::ImpactPointAttribute, FVector(Hit.ImpactPoint), QueryParams.bGetImpactPoint);
 		bResult &= ApplyAttribute(PCGWorldQueryConstants::ImpactNormalAttribute, FVector(Hit.ImpactNormal), QueryParams.bGetImpactNormal);
+		bResult &= ApplyAttribute(PCGWorldQueryConstants::ImpactReflectionAttribute, ReflectionVector, QueryParams.bGetReflection);
 		bResult &= ApplyAttribute(PCGWorldQueryConstants::ImpactDistanceAttribute, (Hit.ImpactPoint - Hit.TraceStart).Length(), QueryParams.bGetDistance);
 		bResult &= ApplyAttribute(PCGPointDataConstants::ActorReferenceAttribute, FSoftObjectPath(Hit.GetActor()), QueryParams.bGetReferenceToActorHit);
 		bResult &= ApplyAttribute(PCGWorldQueryConstants::PhysicalMaterialReferenceAttribute, FSoftObjectPath(Hit.PhysMaterial.Get()), QueryParams.bGetReferenceToPhysicalMaterial);
diff --git a/Engine/Plugins/PCG/Source/PCG/Public/Data/PCGWorldData.h b/Engine/Plugins/PCG/Source/PCG/Public/Data/PCGWorldData.h
--- a/Engine/Plugins/PCG/Source/PCG/Public/Data/PCGWorldData.h
+++ b/Engine/Plugins/PCG/Source/PCG/Public/Data/PCGWorldData.h
@@ -114,6 +114,10 @@ struct FPCGWorldRaycastQueryParams : public FPCGWorldCommonQueryParams
 	UPROPERTY(BlueprintReadWrite, EditAnywhere, Category = "Data|Attributes", meta = (PCG_Overridable))
 	uint8 bGetImpactNormal : 1 = false;
 
+	/** Create an attribute for the reflection vector based on the ray incoming direction and the impact normal. */
+	UPROPERTY(BlueprintReadWrite, EditAnywhere, Category = "Data|Attributes", meta = (PCG_Overridable))
+	uint8 bGetReflection : 1 = false;
+
 	/** Create an attribute for the distance between the ray origin and the impact point. */
 	UPROPERTY(BlueprintReadWrite, EditAnywhere, Category = "Data|Attributes", meta = (PCG_Overridable))
 	uint8 bGetDistance : 1 = false;
diff --git a/Engine/Plugins/PCG/Source/PCG/Public/Helpers/PCGWorldQueryHelpers.h b/Engine/Plugins/PCG/Source/PCG/Public/Helpers/PCGWorldQueryHelpers.h
--- a/Engine/Plugins/PCG/Source/PCG/Public/Helpers/PCGWorldQueryHelpers.h
+++ b/Engine/Plugins/PCG/Source/PCG/Public/Helpers/PCGWorldQueryHelpers.h
@@ -30,6 +30,7 @@ namespace PCGWorldQueryConstants
 	const FName ImpactAttribute = TEXT("ImpactResult");
 	const FName ImpactPointAttribute = TEXT("ImpactPoint");
 	const FName ImpactNormalAttribute = TEXT("ImpactNormal");
+	const FName ImpactReflectionAttribute = TEXT("ImpactReflection");
 	const FName ImpactDistanceAttribute = TEXT("ImpactDistance");
 	const FName LocalImpactPointAttribute = TEXT("ImpactLocalPoint");
 	const FName PhysicalMaterialReferenceAttribute = TEXT("PhysicalMaterial");
@@ -71,6 +72,7 @@ namespace PCGWorldQueryHelpers
 	/** Applies common world ray hit results to attributes. */
 	bool ApplyRayHitMetadata(const FHitResult& HitResult,
 		const FPCGWorldRaycastQueryParams& QueryParams,
+		const FVector& RayDirection,
 		FPCGPoint& OutPoint,
 		UPCGMetadata* OutMetadata,
 		TWeakObjectPtr<UWorld> World,
