[PCG] GPU Static Mesh Spawner - support spawn by-attribute.

Carves out space for InputPointCount X PrimitiveCount instances for now (when using by-attribute spawning - weighted spawning remains the same). In the future dynamic analysis of the spawner input points could reduce this.

Requires string handling which isn't supported in HLSL. Represent using ints so they can still be read/selected/written in kernels.

String table is built at execution time which has all strings (or SOP paths) that enter the GPU graph. We need to know all strings that can arrive on the input pin of an SM spawner so we know what primitives to set up. Did a number of updates to data analysis/description to enable this.

Also added "dump data description" logging to GPU nodes for easy inspection of the data analysis.

#rnx

diff --git a/Engine/Plugins/PCG/Shaders/Private/Elements/PCGCopyPoints.usf b/Engine/Plugins/PCG/Shaders/Private/Elements/PCGCopyPoints.usf
--- a/Engine/Plugins/PCG/Shaders/Private/Elements/PCGCopyPoints.usf
+++ b/Engine/Plugins/PCG/Shaders/Private/Elements/PCGCopyPoints.usf
@@ -20,6 +20,7 @@
 #define MAX_NUM_ATTRS 128
 #define NUM_RESERVED_ATTRS 32
 
+// TODO pass in attribute IDs to copy.
 void CopyAttributesFromSource(uint SourceDataIndex, uint SourceElemIndex, uint OutDataIndex, uint OutElemIndex)
 {
 	const uint SourceDataAddress = Source_GetDataAddress(SourceDataIndex);
@@ -63,6 +64,7 @@ void CopyAttributesFromSource(uint SourceDataIndex, uint SourceElemIndex, uint O
 	}
 }
 
+// TODO pass in attribute IDs to copy.
 void CopyAttributesFromTarget(uint TargetDataIndex, uint TargetElemIndex, uint OutDataIndex, uint OutElemIndex)
 {
 	const uint TargetDataAddress = Target_GetDataAddress(TargetDataIndex);
diff --git a/Engine/Plugins/PCG/Shaders/Private/Elements/PCGStaticMeshSpawner.usf b/Engine/Plugins/PCG/Shaders/Private/Elements/PCGStaticMeshSpawner.usf
--- a/Engine/Plugins/PCG/Shaders/Private/Elements/PCGStaticMeshSpawner.usf
+++ b/Engine/Plugins/PCG/Shaders/Private/Elements/PCGStaticMeshSpawner.usf
@@ -1,12 +1,21 @@
 // Copyright Epic Games, Inc. All Rights Reserved.
 
-// Static mesh spawner kernel. Selects a primitive using a primitive cumulative distribution function (CDF) and packs attributes in custom floats.
+// Static mesh spawner kernel. Selects a primitive either using a primitive cumulative distribution function (CDF) or by an attribute, and packs attributes in custom floats.
 // Passes through points to output.
 
+// When spawning by attribute, we don't know in advance how many instances of each mesh will be spawned, so the worst case NumPoints * NumPrims is allocated in the GPU Scene,
+// and NumPoints * NumPrims threads are dispatched here, hence the adjustments to ThreadIndex below and usage of bWriteOutput.
+
 #include "/Engine/Private/ComputeShaderUtils.ush"
 #include "/Engine/Private/Quaternion.ush"
 #include "/Plugin/PCG/Private/PCGShaderUtils.ush"
 
+// TODO: should be code-gen'd from PCGComputeCommon.h
+#define MAX_NUM_ATTRS 128
+#define NUM_RESERVED_ATTRS 32
+
+void CopyAttributesToOuput(uint In_DataIndex, uint In_ElementIndex, uint Out_DataIndex, uint Out_ElementIndex);
+
 [numthreads(64, 1, 1)]
 void Main(uint3 GroupId : SV_GroupID, uint GroupIndex : SV_GroupIndex)
 {
@@ -23,9 +32,12 @@ void Main(uint3 GroupId : SV_GroupID, uint GroupIndex : SV_GroupIndex)
 		return;
 	}
 
-	// Break out thread index into DataIndex and ElementIndex
-	uint DataIndex, ElementIndex;
-	if (!In_GetThreadData(ThreadIndex, DataIndex, ElementIndex))
+	// Break out thread index into DataIndex and In_ElementIndex. For spawning by attribute we spawn NumPoints * NumPrims threads.
+	const uint NumInputPoints = In_GetNumElements();
+	const uint PointIndex = ThreadIndex % NumInputPoints;
+
+	uint In_DataIndex, In_ElementIndex;
+	if (!In_GetThreadData(PointIndex, In_DataIndex, In_ElementIndex))
 	{
 		return;
 	}
@@ -36,19 +48,36 @@ void Main(uint3 GroupId : SV_GroupID, uint GroupIndex : SV_GroupIndex)
 
 	if (NumPrimitives > 1)
 	{
-		// Randomly select a primitive by inverting the cumulative distribution function.
-		uint Seed = In_GetSeed(DataIndex, ElementIndex);
-		const float RandomDraw = FRand(Seed);
-
-		for (; PrimitiveIndex < uint(NumPrimitives) - 1; ++PrimitiveIndex)
+		const uint SelectorAttributeId = SMSpawner_GetSelectorAttributeId();
+		
+		if (SelectorAttributeId != (uint)-1)
+		{
+			// Select using string key
+			const int MeshPathStringKey = In_GetStringKey(In_DataIndex, In_ElementIndex, SelectorAttributeId);
+			PrimitiveIndex = SMSpawner_GetPrimitiveIndexFromStringKey(MeshPathStringKey);
+		}
+		else
 		{
-			if (RandomDraw < SMSpawner_GetPrimitiveSelectionCDF(PrimitiveIndex))
+			// Randomly select a primitive by inverting the cumulative distribution function.
+			uint Seed = In_GetSeed(In_DataIndex, In_ElementIndex);
+			const float RandomDraw = FRand(Seed);
+
+			for (; PrimitiveIndex < uint(NumPrimitives) - 1; ++PrimitiveIndex)
 			{
-				break;
+				if (RandomDraw < SMSpawner_GetPrimitiveSelectionCDF(PrimitiveIndex))
+				{
+					break;
+				}
 			}
 		}
 	}
 
+	if (PrimitiveIndex >= NumPrimitives)
+	{
+		// Should never happen.
+		return;
+	}
+
 	const int InstanceIndexInt = InstanceData_GetIndexToWriteTo(PrimitiveIndex);
 	if (InstanceIndexInt < 0)
 	{
@@ -56,8 +85,10 @@ void Main(uint3 GroupId : SV_GroupID, uint GroupIndex : SV_GroupIndex)
 		return;
 	}
 
+	// We run NumPoints x NumPrimitives threads. The output points only have to be written for the first NumPoints threads, not the rest.
+	const bool bWriteOutput = ThreadIndex < NumInputPoints;
 	uint Out_DataIndex, Out_ElementIndex;
-	if (!Out_GetThreadData(ThreadIndex, Out_DataIndex, Out_ElementIndex))
+	if (bWriteOutput && !Out_GetThreadData(ThreadIndex, Out_DataIndex, Out_ElementIndex))
 	{
 		// Should never happen.
 		return;
@@ -65,23 +96,26 @@ void Main(uint3 GroupId : SV_GroupID, uint GroupIndex : SV_GroupIndex)
 
 	const uint InstanceIndex = InstanceIndexInt;
 
-	if (In_IsPointRemoved(DataIndex, ElementIndex))
+	if (In_IsPointRemoved(In_DataIndex, In_ElementIndex))
 	{
 		// Cull the instance by setting the transform to 0.
 		InstanceData_WriteInstanceLocalToWorld(InstanceIndex, (float4x4)0);
 		
 		// Propagate the 'Removed' status to the output point if either source or target point is removed. Otherwise, this point will not be culled.
-		Out_RemovePoint(Out_DataIndex, Out_ElementIndex);
+		if (bWriteOutput)
+		{
+			Out_RemovePoint(Out_DataIndex, Out_ElementIndex);
+		}
 
 		return;
 	}
 
 	// Write instance transform.
-	const float4x4 Transform = In_GetPointTransform(DataIndex, ElementIndex);
+	const float4x4 Transform = In_GetPointTransform(In_DataIndex, In_ElementIndex);
 	InstanceData_WriteInstanceLocalToWorld(InstanceIndex, transpose(Transform));
 
 	// Write instance custom float data.
-	const uint DataAddress = In_GetDataAddress(DataIndex);
+	const uint DataAddress = In_GetDataAddress(In_DataIndex);
 	const uint AttributeHeadersAddress = In_GetDataAttributeHeadersAddress(DataAddress);
 	const uint NumAttributes = SMSpawner_GetNumAttributes();
 	for (uint AttributeIndex = 0; AttributeIndex < NumAttributes; ++AttributeIndex)
@@ -96,7 +130,7 @@ void Main(uint3 GroupId : SV_GroupID, uint GroupIndex : SV_GroupIndex)
 
 		for (uint FloatIndex = 0; FloatIndex < AttributeStride; ++FloatIndex)
 		{
-			const uint AttributeValueAddress = AttributeAddress + ElementIndex * AttributeStride * 4;
+			const uint AttributeValueAddress = AttributeAddress + In_ElementIndex * AttributeStride * 4;
 			
 			// Assumes floating point data. TODO add type information and deal with different types in different ways?
 			InstanceData_WriteCustomFloatRaw(InstanceIndex, AttributeOffset + FloatIndex, In_LoadBufferInternal(AttributeValueAddress));
@@ -105,13 +139,62 @@ void Main(uint3 GroupId : SV_GroupID, uint GroupIndex : SV_GroupIndex)
 
 	// Write output point data (pass through).
 	// TODO - support writing mesh bounds.
-	Out_SetPosition(Out_DataIndex, Out_ElementIndex, In_GetPosition(DataIndex, ElementIndex));
-	Out_SetRotation(Out_DataIndex, Out_ElementIndex, In_GetRotation(DataIndex, ElementIndex));
-	Out_SetScale(Out_DataIndex, Out_ElementIndex, In_GetScale(DataIndex, ElementIndex));
-	Out_SetBoundsMin(Out_DataIndex, Out_ElementIndex, In_GetBoundsMin(DataIndex, ElementIndex));
-	Out_SetBoundsMax(Out_DataIndex, Out_ElementIndex, In_GetBoundsMax(DataIndex, ElementIndex));
-	Out_SetColor(Out_DataIndex, Out_ElementIndex, In_GetColor(DataIndex, ElementIndex));
-	Out_SetSeed(Out_DataIndex, Out_ElementIndex, In_GetSeed(DataIndex, ElementIndex));
-	Out_SetDensity(Out_DataIndex, Out_ElementIndex, In_GetDensity(DataIndex, ElementIndex));
-	Out_SetSteepness(Out_DataIndex, Out_ElementIndex, In_GetSteepness(DataIndex, ElementIndex));
+	if (bWriteOutput)
+	{
+		Out_SetPosition(Out_DataIndex, Out_ElementIndex, In_GetPosition(In_DataIndex, In_ElementIndex));
+		Out_SetRotation(Out_DataIndex, Out_ElementIndex, In_GetRotation(In_DataIndex, In_ElementIndex));
+		Out_SetScale(Out_DataIndex, Out_ElementIndex, In_GetScale(In_DataIndex, In_ElementIndex));
+		Out_SetBoundsMin(Out_DataIndex, Out_ElementIndex, In_GetBoundsMin(In_DataIndex, In_ElementIndex));
+		Out_SetBoundsMax(Out_DataIndex, Out_ElementIndex, In_GetBoundsMax(In_DataIndex, In_ElementIndex));
+		Out_SetColor(Out_DataIndex, Out_ElementIndex, In_GetColor(In_DataIndex, In_ElementIndex));
+		Out_SetSeed(Out_DataIndex, Out_ElementIndex, In_GetSeed(In_DataIndex, In_ElementIndex));
+		Out_SetDensity(Out_DataIndex, Out_ElementIndex, In_GetDensity(In_DataIndex, In_ElementIndex));
+		Out_SetSteepness(Out_DataIndex, Out_ElementIndex, In_GetSteepness(In_DataIndex, In_ElementIndex));
+
+		CopyAttributesToOuput(In_DataIndex, In_ElementIndex, Out_DataIndex, Out_ElementIndex);
+	}
+}
+
+// TODO pass in attribute IDs to copy.
+void CopyAttributesToOuput(uint In_DataIndex, uint In_ElementIndex, uint Out_DataIndex, uint Out_ElementIndex)
+{
+	const uint In_DataAddress = In_GetDataAddress(In_DataIndex);
+	const uint In_HeadersAddress = In_GetDataAttributeHeadersAddress(In_DataAddress);
+
+	const uint Out_DataAddress = Out_GetDataAddress(Out_DataIndex);
+	const uint Out_HeadersAddress = Out_GetDataAttributeHeadersAddress(Out_DataAddress);
+
+	const uint NumAttributes = In_GetDataNumAttributes(In_DataAddress);
+	uint NumAttributesProcessed = 0;
+	
+	// Loop over all attribute headers, if the address is non-zero, then copy it from pin Source to pin Out.
+	for (int AttributeIndex = NUM_RESERVED_ATTRS; AttributeIndex < MAX_NUM_ATTRS; ++AttributeIndex)
+	{
+		const uint In_HeaderAddress = In_GetAttributeHeaderAddress(In_HeadersAddress, AttributeIndex);
+		const uint In_Address = In_GetAttributeAddress(In_HeaderAddress);
+
+		const uint OutHeaderAddress = Out_GetAttributeHeaderAddress(Out_HeadersAddress, AttributeIndex);
+		const uint Out_Address = Out_GetAttributeAddress(OutHeaderAddress);
+
+		const uint Stride = In_GetAttributeStride(In_HeaderAddress);
+
+		if (In_Address != 0 && Out_Address != 0)
+		{
+			const uint In_BaseElementAddress = In_Address + In_ElementIndex * Stride;
+			const uint Out_BaseElementAddress = Out_Address + Out_ElementIndex * Stride;
+
+			for (int I = 0; I < Stride; I += 4)
+			{
+				const uint In_ElementAddress = In_BaseElementAddress + I;
+				const uint Out_ElementAddress = Out_BaseElementAddress + I;
+				Out_StoreBufferInternal(Out_ElementAddress, In_LoadBufferInternal(In_ElementAddress));
+			}
+
+			// We can early-out when we've looked at all the possible attributes
+			if (++NumAttributesProcessed >= NumAttributes)
+			{
+				break;
+			}
+		}
+	}
 }
diff --git a/Engine/Plugins/PCG/Shaders/Private/PCGDataCollectionDataInterface.ush b/Engine/Plugins/PCG/Shaders/Private/PCGDataCollectionDataInterface.ush
--- a/Engine/Plugins/PCG/Shaders/Private/PCGDataCollectionDataInterface.ush
+++ b/Engine/Plugins/PCG/Shaders/Private/PCGDataCollectionDataInterface.ush
@@ -270,6 +270,12 @@ float4 GetQuat_{DataInterfaceName}(uint DataIndex, uint ElementIndex, uint Attri
 	return asfloat({DataInterfaceName}_DataCollectionBuffer.Load4(AttributeAddress + ElementIndex * 4 * 4));
 }
 
+int GetStringKey_{DataInterfaceName}(uint DataIndex, uint ElementIndex, uint AttributeId)
+{
+	// String keys are represented as ints.
+	return GetInt_{DataInterfaceName}(DataIndex, ElementIndex, AttributeId);
+}
+
 // #################### ATTRIBUTE SETTERS ##########################
 
 void SetBool_{DataInterfaceName}(uint DataIndex, uint ElementIndex, uint AttributeId, bool Value)
@@ -326,6 +332,12 @@ void SetQuat_{DataInterfaceName}(uint DataIndex, uint ElementIndex, uint Attribu
 	{DataInterfaceName}_DataCollectionBuffer.Store4(AttributeAddress + ElementIndex * 4 * 4, asuint(Value));
 }
 
+void SetStringKey_{DataInterfaceName}(uint DataIndex, uint ElementIndex, uint AttributeId, int Value)
+{
+	// String keys are represented as ints.
+	SetInt_{DataInterfaceName}(DataIndex, ElementIndex, AttributeId, Value);
+}
+
 // #################### POINT ATTRIBUTE GETTERS ##########################
 
 float3 GetPosition_{DataInterfaceName}(uint DataIndex, uint ElementIndex)
diff --git a/Engine/Plugins/PCG/Source/PCG/Private/Compute/DataInterfaces/Elements/PCGStaticMeshSpawnerDataInterface.cpp b/Engine/Plugins/PCG/Source/PCG/Private/Compute/DataInterfaces/Elements/PCGStaticMeshSpawnerDataInterface.cpp
--- a/Engine/Plugins/PCG/Source/PCG/Private/Compute/DataInterfaces/Elements/PCGStaticMeshSpawnerDataInterface.cpp
+++ b/Engine/Plugins/PCG/Source/PCG/Private/Compute/DataInterfaces/Elements/PCGStaticMeshSpawnerDataInterface.cpp
@@ -15,9 +15,21 @@
 
 void UPCGStaticMeshSpawnerDataInterface::GetSupportedInputs(TArray<FShaderFunctionDefinition>& OutFunctions) const
 {
+	OutFunctions.AddDefaulted_GetRef()
+		.SetName(TEXT("SMSpawner_GetSelectorAttributeId"))
+		.AddReturnType(EShaderFundamentalType::Uint); // Attribute id to get mesh path string key from, or invalid if we should use CDF instead.
+
 	OutFunctions.AddDefaulted_GetRef()
 		.SetName(TEXT("SMSpawner_GetNumAttributes"))
-		.AddReturnType(EShaderFundamentalType::Uint);
+		.AddReturnType(EShaderFundamentalType::Uint); // Num attributes
+
+	OutFunctions.AddDefaulted_GetRef()
+		.SetName(TEXT("SMSpawner_GetNumPrimitives"))
+		.AddReturnType(EShaderFundamentalType::Uint); // Num primitives
+
+	OutFunctions.AddDefaulted_GetRef()
+		.SetName(TEXT("SMSpawner_GetNumInputPoints"))
+		.AddReturnType(EShaderFundamentalType::Uint); // Num input points
 
 	OutFunctions.AddDefaulted_GetRef()
 		.SetName(TEXT("SMSpawner_GetAttributeIdOffsetStride"))
@@ -25,20 +37,24 @@ void UPCGStaticMeshSpawnerDataInterface::GetSupportedInputs(TArray<FShaderFuncti
 		.AddParam(EShaderFundamentalType::Uint); // InAttributeIndex
 
 	OutFunctions.AddDefaulted_GetRef()
-		.SetName(TEXT("SMSpawner_GetNumPrimitives"))
-		.AddReturnType(EShaderFundamentalType::Uint);
+		.SetName(TEXT("SMSpawner_GetPrimitiveIndexFromStringKey"))
+		.AddReturnType(EShaderFundamentalType::Uint) // Primitive index
+		.AddParam(EShaderFundamentalType::Uint); // InMeshPathStringKey
 
 	OutFunctions.AddDefaulted_GetRef()
 		.SetName(TEXT("SMSpawner_GetPrimitiveSelectionCDF"))
-		.AddReturnType(EShaderFundamentalType::Float)
+		.AddReturnType(EShaderFundamentalType::Float) // CDF value
 		.AddParam(EShaderFundamentalType::Uint); // InPrimitiveIndex
 }
 
 BEGIN_SHADER_PARAMETER_STRUCT(FPCGStaticMeshSpawnerDataInterfaceParameters,)
-	SHADER_PARAMETER(uint32, NumAttributes)
 	SHADER_PARAMETER_ARRAY(FUintVector4, AttributeIdOffsetStrides, [UPCGStaticMeshSpawnerDataInterface::MAX_ATTRIBUTES])
-	SHADER_PARAMETER(uint32, NumPrimitives)
+	SHADER_PARAMETER_ARRAY(FUintVector4, PrimitiveStringKeys, [PCGComputeConstants::MAX_PRIMITIVE_COMPONENTS_PER_SPAWNER])
 	SHADER_PARAMETER_SCALAR_ARRAY(float, SelectionCDF, [PCGComputeConstants::MAX_PRIMITIVE_COMPONENTS_PER_SPAWNER])
+	SHADER_PARAMETER(uint32, NumAttributes)
+	SHADER_PARAMETER(uint32, NumPrimitives)
+	SHADER_PARAMETER(uint32, SelectorAttributeId)
+	SHADER_PARAMETER(uint32, NumInputPoints)
 END_SHADER_PARAMETER_STRUCT()
 
 void UPCGStaticMeshSpawnerDataInterface::GetShaderParameters(TCHAR const* UID, FShaderParametersMetadataBuilder& InOutBuilder, FShaderParametersMetadataAllocations& InOutAllocations) const
@@ -56,15 +72,37 @@ void UPCGStaticMeshSpawnerDataInterface::GetHLSL(FString& OutHLSL, FString const
 	};
 
 	OutHLSL += FString::Format(TEXT(
+		"uint {DataInterfaceName}_SelectorAttributeId;\n"
+		"uint SMSpawner_GetSelectorAttributeId_{DataInterfaceName}() { return {DataInterfaceName}_SelectorAttributeId; }\n"
+		"\n"
+		"uint {DataInterfaceName}_NumInputPoints;\n"
+		"uint SMSpawner_GetNumInputPoints_{DataInterfaceName}() { return {DataInterfaceName}_NumInputPoints; }\n"
+		"\n"
 		"uint {DataInterfaceName}_NumAttributes;\n"
 		"uint4 {DataInterfaceName}_AttributeIdOffsetStrides[{MaxAttributes}];\n"
 		"uint SMSpawner_GetNumAttributes_{DataInterfaceName}() { return {DataInterfaceName}_NumAttributes; }\n"
 		"uint4 SMSpawner_GetAttributeIdOffsetStride_{DataInterfaceName}(uint InAttributeIndex) { return {DataInterfaceName}_AttributeIdOffsetStrides[InAttributeIndex]; }\n"
 		"\n"
 		"uint {DataInterfaceName}_NumPrimitives;\n"
+		"uint4 {DataInterfaceName}_PrimitiveStringKeys[{MaxPrimitives}];"
 		"DECLARE_SCALAR_ARRAY(float, {DataInterfaceName}_SelectionCDF, {MaxPrimitives});\n"
 		"uint SMSpawner_GetNumPrimitives_{DataInterfaceName}() { return {DataInterfaceName}_NumPrimitives; }\n"
+		"\n"
+		"uint SMSpawner_GetPrimitiveIndexFromStringKey_{DataInterfaceName}(uint InMeshPathStringKey)\n"
+		"{\n"
+		"	for (uint Index = 0; Index < {DataInterfaceName}_NumPrimitives; ++Index)\n"
+		"	{\n"
+		"		if ({DataInterfaceName}_PrimitiveStringKeys[Index][0] == InMeshPathStringKey)\n"
+		"		{\n"
+		"			return Index;\n"
+		"		}\n"
+		"	}\n"
+		"	\n"
+		"	return (uint)-1;\n"
+		"}\n"
+		"\n"
 		"float SMSpawner_GetPrimitiveSelectionCDF_{DataInterfaceName}(uint InPrimitiveIndex) { return GET_SCALAR_ARRAY_ELEMENT({DataInterfaceName}_SelectionCDF, InPrimitiveIndex); }\n"
+		"\n"
 		), TemplateArgs);
 }
 
@@ -85,13 +123,14 @@ UComputeDataProvider* UPCGStaticMeshSpawnerDataInterface::CreateDataProvider(TOb
 	DataProvider->Settings = Cast<UPCGStaticMeshSpawnerSettings>(Settings);
 	DataProvider->AttributeIdOffsetStrides = Primitives->AttributeIdOffsetStrides;
 	DataProvider->SelectionCDF = Primitives->SelectionCDF;
-
+	DataProvider->SelectorAttributeId = Primitives->SelectorAttributeId;
+	DataProvider->PrimitiveStringKeys = Primitives->PrimitiveStringKeys;
 	return DataProvider;
 }
 
 FComputeDataProviderRenderProxy* UPCGStaticMeshSpawnerDataProvider::GetRenderProxy()
 {
-	return new FPCGStaticMeshSpawnerDataProviderProxy(AttributeIdOffsetStrides, SelectionCDF);
+	return new FPCGStaticMeshSpawnerDataProviderProxy(NumInputPoints, AttributeIdOffsetStrides, SelectorAttributeId, PrimitiveStringKeys, SelectionCDF);
 }
 
 bool FPCGStaticMeshSpawnerDataProviderProxy::IsValid(FValidationData const& InValidationData) const
@@ -112,12 +151,19 @@ void FPCGStaticMeshSpawnerDataProviderProxy::GatherDispatchData(FDispatchData co
 		FParameters& Parameters = ParameterArray[InvocationIndex];
 
 		Parameters.NumAttributes = AttributeIdOffsetStrides.Num();
+		Parameters.NumPrimitives = SelectionCDF.Num();
+		Parameters.SelectorAttributeId = (uint32)SelectorAttributeId;
+
 		for (int32 Index = 0; Index < AttributeIdOffsetStrides.Num(); ++Index)
 		{
 			Parameters.AttributeIdOffsetStrides[Index] = AttributeIdOffsetStrides[Index];
 		}
 
-		Parameters.NumPrimitives = SelectionCDF.Num();
+		for (int32 Index = 0; Index < PrimitiveStringKeys.Num(); ++Index)
+		{
+			Parameters.PrimitiveStringKeys[Index] = FUintVector4(PrimitiveStringKeys[Index], /*Unused*/0, /*Unused*/0, /*Unused*/0);
+		}
+
 		for (int32 Index = 0; Index < SelectionCDF.Num(); ++Index)
 		{
 			GET_SCALAR_ARRAY_ELEMENT(Parameters.SelectionCDF, Index) = SelectionCDF[Index];
diff --git a/Engine/Plugins/PCG/Source/PCG/Private/Compute/DataInterfaces/Elements/PCGStaticMeshSpawnerDataInterface.h b/Engine/Plugins/PCG/Source/PCG/Private/Compute/DataInterfaces/Elements/PCGStaticMeshSpawnerDataInterface.h
--- a/Engine/Plugins/PCG/Source/PCG/Private/Compute/DataInterfaces/Elements/PCGStaticMeshSpawnerDataInterface.h
+++ b/Engine/Plugins/PCG/Source/PCG/Private/Compute/DataInterfaces/Elements/PCGStaticMeshSpawnerDataInterface.h
@@ -48,16 +48,28 @@ public:
 	UPROPERTY()
 	TArray<FUintVector4> AttributeIdOffsetStrides;
 
+	UPROPERTY()
+	TArray<uint32> PrimitiveStringKeys;
+
 	UPROPERTY()
 	TArray<float> SelectionCDF;
+
+	UPROPERTY()
+	int32 SelectorAttributeId = -1;
+
+	UPROPERTY()
+	int32 NumInputPoints = 0;
 };
 
 class FPCGStaticMeshSpawnerDataProviderProxy : public FComputeDataProviderRenderProxy
 {
 public:
-	FPCGStaticMeshSpawnerDataProviderProxy(TArray<FUintVector4>& InAttributeIdOffsetStrides, TArray<float> InSelectionCDF)
-		: AttributeIdOffsetStrides(InAttributeIdOffsetStrides)
+	FPCGStaticMeshSpawnerDataProviderProxy(int32 InNumInputPoints, TArray<FUintVector4>& InAttributeIdOffsetStrides, int32 InSelectorAttributeId, TArray<uint32>& InPrimitiveStringKeys, TArray<float> InSelectionCDF)
+		: NumInputPoints(InNumInputPoints)
+		, AttributeIdOffsetStrides(InAttributeIdOffsetStrides)
 		, SelectionCDF(InSelectionCDF)
+		, SelectorAttributeId(InSelectorAttributeId)
+		, PrimitiveStringKeys(InPrimitiveStringKeys)
 	{}
 
 	//~ Begin FComputeDataProviderRenderProxy Interface
@@ -68,6 +80,11 @@ public:
 protected:
 	using FParameters = FPCGStaticMeshSpawnerDataInterfaceParameters;
 
+	int32 NumInputPoints = 0;
+
 	TArray<FUintVector4> AttributeIdOffsetStrides;
 	TArray<float> SelectionCDF;
+
+	int32 SelectorAttributeId = -1;
+	TArray<uint32> PrimitiveStringKeys;
 };
diff --git a/Engine/Plugins/PCG/Source/PCG/Private/Compute/DataInterfaces/PCGDataCollectionDataInterface.cpp b/Engine/Plugins/PCG/Source/PCG/Private/Compute/DataInterfaces/PCGDataCollectionDataInterface.cpp
--- a/Engine/Plugins/PCG/Source/PCG/Private/Compute/DataInterfaces/PCGDataCollectionDataInterface.cpp
+++ b/Engine/Plugins/PCG/Source/PCG/Private/Compute/DataInterfaces/PCGDataCollectionDataInterface.cpp
@@ -227,6 +227,13 @@ void UPCGDataCollectionDataInterface::GetSupportedInputs(TArray<FShaderFunctionD
 			.AddParam(EShaderFundamentalType::Uint)
 			.AddParam(EShaderFundamentalType::Uint)
 			.AddParam(EShaderFundamentalType::Uint);
+
+		OutFunctions.AddDefaulted_GetRef()
+			.SetName(TEXT("GetStringKey"))
+			.AddReturnType(EShaderFundamentalType::Int) // String key represented by int
+			.AddParam(EShaderFundamentalType::Uint) // DataIndex
+			.AddParam(EShaderFundamentalType::Uint) // ElementIndex
+			.AddParam(EShaderFundamentalType::Uint); // AttributeId
 	}
 
 	// Attribute Setters
@@ -293,6 +300,13 @@ void UPCGDataCollectionDataInterface::GetSupportedInputs(TArray<FShaderFunctionD
 			.AddParam(EShaderFundamentalType::Uint)
 			.AddParam(EShaderFundamentalType::Uint)
 			.AddParam(EShaderFundamentalType::Float, 4, 4);
+
+		OutFunctions.AddDefaulted_GetRef()
+			.SetName(TEXT("SetStringKey"))
+			.AddParam(EShaderFundamentalType::Uint) // DataIndex
+			.AddParam(EShaderFundamentalType::Uint) // ElementIndex
+			.AddParam(EShaderFundamentalType::Uint) // AttributeId
+			.AddParam(EShaderFundamentalType::Int); // String key represented by int
 	}
 
 	// Point Attribute Getters
@@ -519,7 +533,7 @@ bool UPCGDataCollectionDataInterface::GetRequiresReadback() const
 
 FComputeDataProviderRenderProxy* UPCGDataCollectionDataProvider::GetRenderProxy()
 {
-	FPCGDataCollectionDataProviderProxy* Proxy = new FPCGDataCollectionDataProviderProxy(Binding, PinDesc, ReadbackMode);
+	FPCGDataCollectionDataProviderProxy* Proxy = new FPCGDataCollectionDataProviderProxy(PinDesc, ReadbackMode);
 
 	if (ReadbackMode != EPCGReadbackMode::None)
 	{
@@ -579,7 +593,7 @@ bool UPCGDataCollectionDataProvider::ProcessReadBackData(FPCGComputeGraphContext
 	}
 
 	FPCGDataCollection DataFromGPU;
-	const EPCGUnpackDataCollectionResult Result = PinDesc.UnpackDataCollection(RawReadbackData, OutputPinLabelAlias, DataFromGPU);
+	const EPCGUnpackDataCollectionResult Result = PinDesc.UnpackDataCollection(RawReadbackData, OutputPinLabelAlias, InContext->DataBinding->GetStringTable(), DataFromGPU);
 
 	if (Result == EPCGUnpackDataCollectionResult::DataMismatch)
 	{
@@ -640,12 +654,8 @@ bool UPCGDataCollectionDataProvider::ProcessReadBackData(FPCGComputeGraphContext
 	return true;
 }
 
-FPCGDataCollectionDataProviderProxy::FPCGDataCollectionDataProviderProxy(
-	TWeakObjectPtr<UPCGDataBinding> InBinding,
-	const FPCGDataCollectionDesc& InPinDesc,
-	EPCGReadbackMode InReadbackMode)
+FPCGDataCollectionDataProviderProxy::FPCGDataCollectionDataProviderProxy(const FPCGDataCollectionDesc& InPinDesc, EPCGReadbackMode InReadbackMode)
 	: ReadbackMode(InReadbackMode)
-	, Binding(InBinding)
 	, PinDesc(InPinDesc)
 {
 	SizeBytes = PinDesc.ComputePackedSize();
@@ -658,12 +668,6 @@ bool FPCGDataCollectionDataProviderProxy::IsValid(FValidationData const& InValid
 		return false;
 	}
 
-	if (!Binding.IsValid())
-	{
-		UE_LOG(LogPCG, Error, TEXT("Proxy invalid due to missing data binding."));
-		return false;
-	}
-
 	if (SizeBytes <= 0)
 	{
 		UE_LOG(LogPCG, Error, TEXT("Proxy invalid due to invalid buffer size."));
diff --git a/Engine/Plugins/PCG/Source/PCG/Private/Compute/DataInterfaces/PCGDataCollectionDataInterface.h b/Engine/Plugins/PCG/Source/PCG/Private/Compute/DataInterfaces/PCGDataCollectionDataInterface.h
--- a/Engine/Plugins/PCG/Source/PCG/Private/Compute/DataInterfaces/PCGDataCollectionDataInterface.h
+++ b/Engine/Plugins/PCG/Source/PCG/Private/Compute/DataInterfaces/PCGDataCollectionDataInterface.h
@@ -108,10 +108,7 @@ protected:
 class FPCGDataCollectionDataProviderProxy : public FComputeDataProviderRenderProxy
 {
 public:
-	FPCGDataCollectionDataProviderProxy(
-		TWeakObjectPtr<UPCGDataBinding> InBinding,
-		const FPCGDataCollectionDesc& InPinDesc,
-		EPCGReadbackMode InReadbackMode);
+	FPCGDataCollectionDataProviderProxy(const FPCGDataCollectionDesc& InPinDesc, EPCGReadbackMode InReadbackMode);
 
 	//~ Begin FComputeDataProviderRenderProxy Interface
 	bool IsValid(FValidationData const& InValidationData) const override;
@@ -130,7 +127,6 @@ protected:
 
 	using FParameters = FPCGDataCollectionDataInterfaceParameters;
 
-	TWeakObjectPtr<UPCGDataBinding> Binding;
 	FPCGDataCollectionDesc PinDesc;
 
 	FRDGBufferRef Buffer = nullptr;
diff --git a/Engine/Plugins/PCG/Source/PCG/Private/Compute/DataInterfaces/PCGDataCollectionUploadDataInterface.cpp b/Engine/Plugins/PCG/Source/PCG/Private/Compute/DataInterfaces/PCGDataCollectionUploadDataInterface.cpp
--- a/Engine/Plugins/PCG/Source/PCG/Private/Compute/DataInterfaces/PCGDataCollectionUploadDataInterface.cpp
+++ b/Engine/Plugins/PCG/Source/PCG/Private/Compute/DataInterfaces/PCGDataCollectionUploadDataInterface.cpp
@@ -25,7 +25,8 @@ UComputeDataProvider* UPCGDataCollectionUploadDataInterface::CreateDataProvider(
 	Provider->PinDesc = FPCGDataCollectionDesc::BuildFromInputDataCollectionAndInputPinLabel(
 		Binding->DataForGPU.InputDataCollection,
 		DownstreamInputPinLabelAliases[0],
-		Binding->Graph->GetAttributeLookupTable());
+		Binding->GetAttributeLookupTable(),
+		Binding->GetStringTable());
 
 	Provider->DownstreamInputPinLabels = DownstreamInputPinLabelAliases;
 
@@ -34,31 +35,24 @@ UComputeDataProvider* UPCGDataCollectionUploadDataInterface::CreateDataProvider(
 
 FComputeDataProviderRenderProxy* UPCGDataProviderDataCollectionUpload::GetRenderProxy()
 {
-	return new FPCGDataProviderDataCollectionUploadProxy(Binding, PinDesc, DownstreamInputPinLabels);
+	TArray<uint32> PackedDataCollection;
+
+	// Use any downstream input pin label to grab data from the collection.
+	check(!DownstreamInputPinLabels.IsEmpty());
+	PinDesc.PackDataCollection(Binding->DataForGPU.InputDataCollection, DownstreamInputPinLabels[0], Binding->GetStringTable(), PackedDataCollection);
+
+	return new FPCGDataProviderDataCollectionUploadProxy(PinDesc, MoveTemp(PackedDataCollection), DownstreamInputPinLabels);
 }
 
-FPCGDataProviderDataCollectionUploadProxy::FPCGDataProviderDataCollectionUploadProxy(
-	TWeakObjectPtr<UPCGDataBinding> InBinding,
-	const FPCGDataCollectionDesc& InPinDesc,
-	const TArray<FName>& InDownstreamInputPinLabels)
-	: FPCGDataCollectionDataProviderProxy(InBinding, InPinDesc, EPCGReadbackMode::None)
+FPCGDataProviderDataCollectionUploadProxy::FPCGDataProviderDataCollectionUploadProxy(const FPCGDataCollectionDesc& InPinDesc, TArray<uint32>&& InPackedDataCollection, const TArray<FName>& InDownstreamInputPinLabels)
+	: FPCGDataCollectionDataProviderProxy(InPinDesc, EPCGReadbackMode::None)
+	, PackedDataCollection(MoveTemp(InPackedDataCollection))
+	, DownstreamInputPinLabels(InDownstreamInputPinLabels)
 {
-	DownstreamInputPinLabels = InDownstreamInputPinLabels;
 }
 
 void FPCGDataProviderDataCollectionUploadProxy::AllocateResources(FRDGBuilder& GraphBuilder, FAllocationData const& InAllocationData)
 {
-	if (!ensure(Binding.IsValid()))
-	{
-		return;
-	}
-
-	TArray<uint32> PackedDataCollection;
-
-	// Use any downstream input pin label to grab data from the collection.
-	check(!DownstreamInputPinLabels.IsEmpty());
-	PinDesc.PackDataCollection(Binding->DataForGPU.InputDataCollection, DownstreamInputPinLabels[0], PackedDataCollection);
-	
 	const FRDGBufferDesc Desc = FRDGBufferDesc::CreateStructuredDesc(sizeof(uint32), PackedDataCollection.Num());
 	Buffer = GraphBuilder.CreateBuffer(Desc, TEXT("PCGDataCollectionUploadBuffer"));
 	BufferUAV = GraphBuilder.CreateUAV(Buffer);
diff --git a/Engine/Plugins/PCG/Source/PCG/Private/Compute/DataInterfaces/PCGDataCollectionUploadDataInterface.h b/Engine/Plugins/PCG/Source/PCG/Private/Compute/DataInterfaces/PCGDataCollectionUploadDataInterface.h
--- a/Engine/Plugins/PCG/Source/PCG/Private/Compute/DataInterfaces/PCGDataCollectionUploadDataInterface.h
+++ b/Engine/Plugins/PCG/Source/PCG/Private/Compute/DataInterfaces/PCGDataCollectionUploadDataInterface.h
@@ -39,11 +39,13 @@ public:
 class FPCGDataProviderDataCollectionUploadProxy : public FPCGDataCollectionDataProviderProxy
 {
 public:
-	FPCGDataProviderDataCollectionUploadProxy(TWeakObjectPtr<UPCGDataBinding> InBinding, const FPCGDataCollectionDesc& InPinDesc, const TArray<FName>& InDownstreamInputPinLabels);
+	FPCGDataProviderDataCollectionUploadProxy(const FPCGDataCollectionDesc& InPinDesc, TArray<uint32>&& PackedDataCollection, const TArray<FName>& InDownstreamInputPinLabels);
 
 	//~ Begin FComputeDataProviderRenderProxy Interface
 	void AllocateResources(FRDGBuilder& GraphBuilder, FAllocationData const& InAllocationData) override;
 	//~ End FComputeDataProviderRenderProxy Interface
 
+	TArray<uint32> PackedDataCollection;
+
 	TArray<FName> DownstreamInputPinLabels;
 };
diff --git a/Engine/Plugins/PCG/Source/PCG/Private/Compute/DataInterfaces/PCGInstanceDataInterface.cpp b/Engine/Plugins/PCG/Source/PCG/Private/Compute/DataInterfaces/PCGInstanceDataInterface.cpp
--- a/Engine/Plugins/PCG/Source/PCG/Private/Compute/DataInterfaces/PCGInstanceDataInterface.cpp
+++ b/Engine/Plugins/PCG/Source/PCG/Private/Compute/DataInterfaces/PCGInstanceDataInterface.cpp
@@ -130,11 +130,14 @@ UComputeDataProvider* UPCGInstanceDataInterface::CreateDataProvider(TObjectPtr<U
 	DataProvider->Primitives = FoundPrimitives->Primitives;
 	DataProvider->NumCustomFloatsPerInstance = FoundPrimitives->NumCustomFloats;
 	
-	const UPCGNode* Node = ProducerSettings ? Cast<UPCGNode>(ProducerSettings->GetOuter()) : nullptr;
-	const UPCGPin* InputPin = Node ? Node->GetInputPin(InputPinProvidingData) : nullptr;
-	check(ProducerSettings && InputPin);
-	const FPCGDataCollectionDesc InputDataDesc = ProducerSettings->ComputeInputPinDataDesc(InputPin, Binding);
+	check(ProducerSettings);
+	const FPCGDataCollectionDesc InputDataDesc = ProducerSettings->ComputeInputPinDataDesc(InputPinProvidingData, Binding);
 	DataProvider->NumInstancesAllPrimitives = InputDataDesc.ComputeDataElementCount(EPCGDataType::Point);
+	if (FoundPrimitives->SelectorAttributeId != -1)
+	{
+		// When selecting primitives dynamically (by attribute), we don't know statically how many instances will end up in each primitive, so run worst case.
+		DataProvider->NumInstancesAllPrimitives *= FoundPrimitives->Primitives.Num();
+	}
 
 	return DataProvider;
 }
diff --git a/Engine/Plugins/PCG/Source/PCG/Private/Compute/Elements/PCGComputeGraphElement.cpp b/Engine/Plugins/PCG/Source/PCG/Private/Compute/Elements/PCGComputeGraphElement.cpp
--- a/Engine/Plugins/PCG/Source/PCG/Private/Compute/Elements/PCGComputeGraphElement.cpp
+++ b/Engine/Plugins/PCG/Source/PCG/Private/Compute/Elements/PCGComputeGraphElement.cpp
@@ -12,6 +12,7 @@
 #include "Elements/PCGStaticMeshSpawner.h"
 #include "Engine/StaticMesh.h"
 #include "InstanceDataPackers/PCGInstanceDataPackerBase.h"
+#include "MeshSelectors/PCGMeshSelectorByAttribute.h"
 #include "MeshSelectors/PCGMeshSelectorWeighted.h"
 
 #include "ComputeWorkerInterface.h"
@@ -107,29 +108,10 @@ bool FPCGComputeGraphElement::ExecuteInternal(FPCGContext* InContext) const
 	// 4. Initialize and parse incoming data for data sizes, attributes, etc that will drive buffer allocations and dispatch thread counts.
 	if (!Context->DataBinding)
 	{
-		// When attribute table was built at compile time, some attribute types may not be known. Fill them out now.
-		Graph->FillInMissingAttributeTableTypes(Context->InputData);
+		UPCGDataBinding* DataBinding = FPCGContext::NewObject_AnyThread<UPCGDataBinding>(Context);
+		Context->DataBinding.Reset(DataBinding);
 
-		UPCGDataBinding* DataBindingObject = FPCGContext::NewObject_AnyThread<UPCGDataBinding>(Context);
-		Context->DataBinding.Reset(DataBindingObject);
-
-		DataBindingObject->SourceComponent = Context->SourceComponent;
-		DataBindingObject->Graph = Graph.Get();
-
-		FPCGDataForGPU& DataForGPU = DataBindingObject->DataForGPU;
-
-		DataForGPU.InputDataCollection = Context->InputData;
-
-		// Link each input pin to the data collection, so that data providers can find the data.
-		for (TWeakObjectPtr<const UPCGPin>& InputPinPtr : Graph->PinsReceivingDataFromCPU)
-		{
-			if (const UPCGPin* InputPin = InputPinPtr.Get())
-			{
-				DataForGPU.InputPins.Add(InputPin);
-			}
-		}
-
-		DataForGPU.InputPinLabelAliases = Graph->InputPinLabelAliases;
+		Context->DataBinding->Initialize(Graph.Get(), Context->SourceComponent, Context->InputData, Graph->GetAttributeLookupTable());
 
 		// Perform validation after input data is initialized.
 		for (TWeakObjectPtr<const UPCGNode> Node : Graph->KernelToNode)
@@ -142,7 +124,7 @@ bool FPCGComputeGraphElement::ExecuteInternal(FPCGContext* InContext) const
 			}
 		}
 
-		const bool bAnyComponentsSetup = SetupProceduralISMComponents(InContext, DataBindingObject);
+		const bool bAnyComponentsSetup = SetupProceduralISMComponents(Context, DataBinding);
 
 		Context->ComputeGraphInstance.CreateDataProviders(Graph.Get(), 0, Context->DataBinding.Get());
 
@@ -247,8 +229,12 @@ bool FPCGComputeGraphElement::ExecuteInternal(FPCGContext* InContext) const
 	}
 }
 
-bool FPCGComputeGraphElement::SetupProceduralISMComponents(FPCGContext* InContext, UPCGDataBinding* Binding) const
+bool FPCGComputeGraphElement::SetupProceduralISMComponents(FPCGContext* InContext, UPCGDataBinding* InBinding) const
 {
+	TRACE_CPUPROFILER_EVENT_SCOPE(FPCGComputeGraphElement::SetupProceduralISMComponents);
+	check(InContext);
+	check(InBinding);
+
 	bool bAnyComponentsSetup = false;
 
 	for (const UPCGSettings* Settings : Graph->StaticMeshSpawners)
@@ -259,60 +245,144 @@ bool FPCGComputeGraphElement::SetupProceduralISMComponents(FPCGContext* InContex
 			continue;
 		}
 
-		const UPCGMeshSelectorWeighted* Selector = Cast<UPCGMeshSelectorWeighted>(SpawnerSettings->MeshSelectorParameters);
 		AActor* TargetActor = SpawnerSettings->TargetActor.Get() ? SpawnerSettings->TargetActor.Get() : InContext->GetTargetActor(nullptr);
-		if (!ensure(Selector) || !ensure(TargetActor))
+		if (!ensure(SpawnerSettings->MeshSelectorParameters) || !ensure(TargetActor))
 		{
 			continue;
 		}
 
-		const UPCGNode* Node = Cast<const UPCGNode>(SpawnerSettings->GetOuter());
-		const UPCGPin* Pin = Node ? Node->GetInputPin(PCGPinConstants::DefaultInputLabel) : nullptr;
-		if (!ensure(Pin))
-		{
-			continue;
-		}
+		const FPCGDataCollectionDesc InputDataDesc = SpawnerSettings->ComputeInputPinDataDesc(PCGPinConstants::DefaultInputLabel, InBinding);
 
-		const FPCGDataCollectionDesc InputDataDesc = SpawnerSettings->ComputeInputPinDataDesc(Pin, Binding);
+		const uint32 InputPointCount = InputDataDesc.ComputeDataElementCount(EPCGDataType::Point);
+		const FBox LocalBounds = InContext->SourceComponent->GetGridBounds().ShiftBy(-InContext->SourceComponent->GetOwner()->GetActorLocation());
 
-		const uint32 PointCount = InputDataDesc.ComputeDataElementCount(EPCGDataType::Point);
-		
 		uint32 CustomFloatCount = 0;
+
 		TArray<FUint32Vector4> AttributeIdOffsetStrides;
 		if (SpawnerSettings->InstanceDataPackerParameters && SpawnerSettings->InstanceDataPackerParameters->GetAttributeNames(/*OutNames=*/nullptr))
 		{
 			TArray<FName> AttributeNames;
 			SpawnerSettings->InstanceDataPackerParameters->GetAttributeNames(&AttributeNames);
 
-			PCGDataForGPUHelpers::ComputeCustomFloatPacking(AttributeNames, Binding, InputDataDesc, CustomFloatCount, AttributeIdOffsetStrides);
+			PCGDataForGPUHelpers::ComputeCustomFloatPacking(AttributeNames, InBinding, InputDataDesc, CustomFloatCount, AttributeIdOffsetStrides);
 		}
 
-		bool AttributeSetupDone = false;
+		TArray<float> PrimitiveSelectionCDF;
+		int32 SelectorAttributeId = -1;
+		TArray<uint32> PrimitiveStringKeys; // ok index is just array index! TODO
+		TArray<FPCGProceduralISMComponentDescriptor> ComponentsToCreate;
 
-		float TotalWeight = 0.0f;
-		for (const FPCGMeshSelectorWeightedEntry& Entry : Selector->MeshEntries)
+		if (const UPCGMeshSelectorByAttribute* SelectorByAttribute = Cast<UPCGMeshSelectorByAttribute>(SpawnerSettings->MeshSelectorParameters))
 		{
-			TotalWeight += Entry.Weight;
-		}
+			const FName SelectedName = SelectorByAttribute->AttributeName;
+			if (SelectedName == NAME_None)
+			{
+				UE_LOG(LogPCG, Error, TEXT("Invalid mesh selector attribute specified '%s'."), *SelectedName.ToString());
+				continue;
+			}
+
+			const FPCGKernelAttributeIDAndType* FoundAttribute = InBinding->GetAttributeLookupTable().Find(SelectedName);
+			if (!FoundAttribute)
+			{
+				UE_LOG(LogPCG, Error, TEXT("Mesh selector attribute '%s' not found."), *SelectedName.ToString());
+				continue;
+			}
+
+			SelectorAttributeId = FoundAttribute->Id;
+
+			// Compute how many unique incoming values we will receive.
+			TSet<int32> UniqueStringKeys;
+			for (const FPCGDataDesc& Desc : InputDataDesc.DataDescs)
+			{
+				for (const FPCGKernelAttributeDesc& AttributeDesc : Desc.AttributeDescs)
+				{
+					if (AttributeDesc.Name == SelectedName)
+					{
+						UniqueStringKeys.Append(AttributeDesc.UniqueStringKeys);
+					}
+				}
+			}
+
+			for (int32 StringKey : UniqueStringKeys)
+			{
+				if (!ensure(InBinding->GetStringTable().IsValidIndex(StringKey)))
+				{
+					continue;
+				}
+
+				const FString& MeshPathString = InBinding->GetStringTable()[StringKey];
+				
+				FPCGProceduralISMComponentDescriptor Descriptor;
+				Descriptor = SelectorByAttribute->TemplateDescriptor;
+				Descriptor.NumInstances = InputPointCount;
+				Descriptor.LocalBounds = LocalBounds;
+				Descriptor.NumCustomFloats = CustomFloatCount;
+				
+				Descriptor.StaticMesh = Cast<UStaticMesh>(FSoftObjectPath(MeshPathString).TryLoad());
+				if (!Descriptor.StaticMesh)
+				{
+					UE_LOG(LogPCG, Error, TEXT("Could not load static mesh from path '%s'."), *MeshPathString);
+					continue;
+				}
+
+				PrimitiveStringKeys.Emplace(static_cast<uint32>(StringKey));
+				ComponentsToCreate.Add(MoveTemp(Descriptor));
+			}
 
-		if (!ensure(TotalWeight > UE_SMALL_NUMBER))
+			PrimitiveSelectionCDF.SetNumZeroed(ComponentsToCreate.Num());
+		}
+		else if (const UPCGMeshSelectorWeighted* SelectorWeighted = Cast<UPCGMeshSelectorWeighted>(SpawnerSettings->MeshSelectorParameters))
 		{
-			continue;
+			if (SelectorWeighted->MeshEntries.IsEmpty())
+			{
+				UE_LOG(LogPCG, Error, TEXT("No mesh entries provided."));
+				continue;
+			}
+
+			float CumulativeWeight = 0.0f;
+
+			float TotalWeight = 0.0f;
+			for (const FPCGMeshSelectorWeightedEntry& Entry : SelectorWeighted->MeshEntries)
+			{
+				TotalWeight += Entry.Weight;
+			}
+
+			if (!ensure(TotalWeight > UE_SMALL_NUMBER))
+			{
+				continue;
+			}
+
+			PrimitiveSelectionCDF.Reserve(SelectorWeighted->MeshEntries.Num());
+
+			for (const FPCGMeshSelectorWeightedEntry& Entry : SelectorWeighted->MeshEntries)
+			{
+				UStaticMesh* StaticMesh = Entry.Descriptor.StaticMesh.LoadSynchronous();
+
+				const float Weight = float(Entry.Weight) / TotalWeight;
+				CumulativeWeight += Weight;
+				PrimitiveSelectionCDF.Add(CumulativeWeight);
+
+				FPCGProceduralISMComponentDescriptor Descriptor;
+				Descriptor = Entry.Descriptor;
+				Descriptor.NumInstances = FMath::CeilToInt(InputPointCount * Weight);
+				Descriptor.LocalBounds = LocalBounds;
+				Descriptor.NumCustomFloats = CustomFloatCount;
+				Descriptor.StaticMesh = StaticMesh;
+				ComponentsToCreate.Add(MoveTemp(Descriptor));
+			}
 		}
 
-		float CumulativeWeight = 0.0f;
+		FPCGSpawnerPrimitives& Primitives = InBinding->MeshSpawnersToPrimitives.FindOrAdd(SpawnerSettings);
+		Primitives.NumCustomFloats = CustomFloatCount;
+		Primitives.AttributeIdOffsetStrides = MoveTemp(AttributeIdOffsetStrides);
+		Primitives.SelectorAttributeId = SelectorAttributeId;
+		Primitives.SelectionCDF = MoveTemp(PrimitiveSelectionCDF);
+		Primitives.PrimitiveStringKeys = MoveTemp(PrimitiveStringKeys);
 
-		for (const FPCGMeshSelectorWeightedEntry& Entry : Selector->MeshEntries)
+		for (const FPCGProceduralISMComponentDescriptor& Desc : ComponentsToCreate)
 		{
-			UStaticMesh* StaticMesh = Entry.Descriptor.StaticMesh.LoadSynchronous();
-			const float Weight = float(Entry.Weight) / TotalWeight;
-			CumulativeWeight += Weight;
-
 			FPCGProceduralISMCBuilderParameters Params;
-			Params.Descriptor.NumInstances = FMath::CeilToInt(PointCount * Weight);
-			Params.Descriptor.LocalBounds = InContext->SourceComponent->GetGridBounds().ShiftBy(-InContext->SourceComponent->GetOwner()->GetActorLocation());
-			Params.Descriptor.NumCustomFloats = CustomFloatCount;
-			Params.Descriptor = Entry.Descriptor;
+			Params.Descriptor = Desc;
 			Params.bAllowDescriptorChanges = false;
 
 			UPCGManagedProceduralISMComponent* MISMC = PCGManagedProceduralISMComponent::GetOrCreateManagedProceduralISMC(TargetActor, InContext->SourceComponent.Get(), SpawnerSettings->UID, Params);
@@ -323,15 +393,7 @@ bool FPCGComputeGraphElement::SetupProceduralISMComponents(FPCGContext* InContex
 			// Don't bother registering the resource change as we're transient anyway.
 			//InContext->TouchedResources.Emplace(MISMC);
 
-			FPCGSpawnerPrimitives& Primitives = Binding->MeshSpawnersToPrimitives.FindOrAdd(SpawnerSettings);
 			Primitives.Primitives.Add(MISMC->GetComponent());
-			Primitives.SelectionCDF.Add(CumulativeWeight);
-			if (!AttributeSetupDone)
-			{
-				Primitives.NumCustomFloats = CustomFloatCount;
-				Primitives.AttributeIdOffsetStrides = MoveTemp(AttributeIdOffsetStrides);
-				AttributeSetupDone = true;
-			}
 
 			bAnyComponentsSetup = true;
 
diff --git a/Engine/Plugins/PCG/Source/PCG/Private/Compute/Elements/PCGCustomHLSL.cpp b/Engine/Plugins/PCG/Source/PCG/Private/Compute/Elements/PCGCustomHLSL.cpp
--- a/Engine/Plugins/PCG/Source/PCG/Private/Compute/Elements/PCGCustomHLSL.cpp
+++ b/Engine/Plugins/PCG/Source/PCG/Private/Compute/Elements/PCGCustomHLSL.cpp
@@ -31,7 +31,7 @@ namespace PCGHLSLElement
 	/** Second capture: Function name (Get or Set) */
 	constexpr int AttributeFunctionCaptureGroup = 2;
 
-	/** Third capture: Attribute type (e.g.Int, Float, Rotator, etc.) */
+	/** Third capture: Attribute type (e.g. Int, Float, Rotator, etc.) */
 	constexpr int AttributeTypeCaptureGroup = 3;
 
 	/** Fourth capture: Attribute name (supports a-z, A-Z, 0-9, ' ', '-', '_', and '/') */
@@ -282,6 +282,7 @@ int UPCGCustomHLSLSettings::ComputeKernelThreadCount(const UPCGDataBinding* Bind
 FPCGDataCollectionDesc UPCGCustomHLSLSettings::ComputeOutputPinDataDesc(const UPCGPin* OutputPin, const UPCGDataBinding* Binding) const
 {
 	check(OutputPin);
+	check(Binding);
 	FPCGDataCollectionDesc PinDesc;
 
 	const FPCGPinPropertiesGPU* Properties = GetOutputPinPropertiesGPU(OutputPin->Properties.Label);
@@ -306,29 +307,119 @@ FPCGDataCollectionDesc UPCGCustomHLSLSettings::ComputeOutputPinDataDesc(const UP
 		PCGComputeHelpers::ComputeOutputPinDataDesc(*Properties, this, Binding, PinDesc);
 	}
 
-	if (const TMap<FName, FPCGKernelAttributeIDAndType>* GlobalAttributeLookupTable = ensure(Binding && Binding->Graph) ? &Binding->Graph->GetAttributeLookupTable() : nullptr)
+	const TMap<FName, FPCGKernelAttributeIDAndType>& GlobalAttributeLookupTable = Binding->GetAttributeLookupTable();
+
+	for (const FPCGKernelAttributeKey& AttributeKey : KernelAttributeKeys)
+	{
+		// Add attributes that will be created for this pin on the GPU.
+		if (const TArray<TTuple<FPCGKernelAttributeKey, bool>>* Keys = PinToAttributeKeys.Find(OutputPin->Properties.Label))
+		{
+			const TTuple<FPCGKernelAttributeKey, bool>* Pair = Keys->FindByPredicate([AttributeKey](const TTuple<FPCGKernelAttributeKey, bool>& Pair) { return Pair.Key == AttributeKey; });
+			if (!Pair || !Pair->Value)
+			{
+				// Not created on GPU.
+				continue;
+			}
+
+			for (FPCGDataDesc& DataDesc : PinDesc.DataDescs)
+			{
+				if (const FPCGKernelAttributeIDAndType* IDAndType = GlobalAttributeLookupTable.Find(AttributeKey.Name))
+				{
+					const FPCGKernelAttributeDesc AttributeDesc(IDAndType->Id, IDAndType->Type, AttributeKey.Name);
+					DataDesc.AttributeDescs.AddUnique(AttributeDesc);
+				}
+			}
+		}
+	}
+
+	// Try to propagate string keys across node. Not trivial because there could be one or more string key attributes on input pins and on output pins,
+	// and it is in general hard to determine from source which string keys from input are being written to outputs. Try first collecting all string keys
+	// from matching attribute names (across all input pins), and then fall back to collecting keys from all string key attributes across all inputs.
+	if (const UPCGNode* Node = Cast<UPCGNode>(GetOuter()))
 	{
-		for (const FPCGKernelAttributeKey& AttributeKey : KernelAttributeKeys)
+		TArray<FPCGDataCollectionDesc> RelevantInputDataDescs;
+
+		// Collect descriptions of input data items that have string key attributes.
+		for (const UPCGPin* InputPin : Node->GetInputPins())
 		{
-			// Add attributes that will be created for this pin on the GPU.
-			if (const TArray<TTuple<FPCGKernelAttributeKey, bool>>* Keys = PinToAttributeKeys.Find(OutputPin->Properties.Label))
+			const FPCGDataCollectionDesc InputPinDesc = ComputeInputPinDataDesc(InputPin, Binding);
+
+			bool bFoundStringKeyAttribute = false;
+
+			for (FPCGDataDesc& DataDesc : PinDesc.DataDescs)
+			{
+				for (FPCGKernelAttributeDesc& AttributeDesc : DataDesc.AttributeDescs)
+				{
+					bFoundStringKeyAttribute |= (AttributeDesc.Type == EPCGKernelAttributeType::StringKey);
+				}
+			}
+
+			if (bFoundStringKeyAttribute)
 			{
-				const TTuple<FPCGKernelAttributeKey, bool>* Pair = Keys->FindByPredicate([AttributeKey](const TTuple<FPCGKernelAttributeKey, bool>& Pair) { return Pair.Key == AttributeKey; });
-				const bool bCreatedOnGPU = Pair && Pair->Value;
+				RelevantInputDataDescs.Add(InputPinDesc);
+			}
+		}
 
-				if (bCreatedOnGPU)
+		if (!RelevantInputDataDescs.IsEmpty())
+		{
+			for (FPCGDataDesc& DataDesc : PinDesc.DataDescs)
+			{
+				for (FPCGKernelAttributeDesc& AttributeDesc : DataDesc.AttributeDescs)
 				{
-					for (FPCGDataDesc& DataDesc : PinDesc.DataDescs)
+					if (AttributeDesc.Type != EPCGKernelAttributeType::StringKey)
+					{
+						continue;
+					}
+
+					bool bFoundMatchingAttribute = false;
+
+					for (const FPCGDataCollectionDesc& InputPinDataDesc : RelevantInputDataDescs)
+					{
+						// Try to find string keys for matching attributes on inputs. E.g. if we are processing an output attribute named 'MeshPath',
+						// look at data on all input pins for an attribute named MeshPath and assume we could use any of its values - copy the string keys.
+						for (const FPCGDataDesc& InputDataDesc : InputPinDataDesc.DataDescs)
+						{
+							for (const FPCGKernelAttributeDesc& InputAttributeDesc : InputDataDesc.AttributeDescs)
+							{
+								if (InputAttributeDesc.Type == EPCGKernelAttributeType::StringKey && InputAttributeDesc.Name == AttributeDesc.Name)
+								{
+									AttributeDesc.UniqueStringKeys.Append(InputAttributeDesc.UniqueStringKeys);
+									bFoundMatchingAttribute = true;
+									break;
+								}
+							}
+						}
+					}
+
+					if (!bFoundMatchingAttribute)
 					{
-						if (const FPCGKernelAttributeIDAndType* IDAndType = GlobalAttributeLookupTable->Find(AttributeKey.Name))
+						// We didn't find an exact attribute. Fall back to finding any and all string keys. This is concerning and perhaps we can
+						// have additional hinting mechanisms in the kernel source or in the node UI.
+						for (const FPCGDataCollectionDesc& InputPinDataDesc : RelevantInputDataDescs)
 						{
-							const FPCGKernelAttributeDesc AttributeDesc(IDAndType->Id, IDAndType->Type, AttributeKey.Name);
-							DataDesc.AttributeDescs.AddUnique(AttributeDesc);
+							// Try to find string keys for matching attributes on inputs. E.g. if we are processing an output attribute named 'MeshPath',
+							// look at data on all input pins for an attribute named MeshPath and assume we could use any of its values - copy the string keys.
+							for (const FPCGDataDesc& InputDataDesc : InputPinDataDesc.DataDescs)
+							{
+								for (const FPCGKernelAttributeDesc& InputAttributeDesc : InputDataDesc.AttributeDescs)
+								{
+									if (InputAttributeDesc.Type == EPCGKernelAttributeType::StringKey)
+									{
+										AttributeDesc.UniqueStringKeys.Append(InputAttributeDesc.UniqueStringKeys);
+									}
+								}
+							}
 						}
 					}
 				}
 			}
 		}
+		else
+		{
+			// If there were no string keys found on any input pin then we are in a bad place. String values cannot be built on the GPU, they must
+			// come in through an input.
+			UE_LOG(LogPCG, Warning, TEXT("No incoming attributes to obtain string keys from."));
+		}
 	}
 
 	return PinDesc;
@@ -475,11 +566,15 @@ FString UPCGCustomHLSLSettings::GetShaderText() const
 void UPCGCustomHLSLSettings::SetShaderFunctionsText(const FString& NewFunctionsText)
 {
 	ShaderFunctions = NewFunctionsText;
+
+	UpdateAttributeKeys();
 }
 
 void UPCGCustomHLSLSettings::SetShaderText(const FString& NewText)
 {
 	ShaderSource = NewText;
+
+	UpdateAttributeKeys();
 }
 
 bool UPCGCustomHLSLSettings::IsShaderTextReadOnly() const
@@ -598,7 +693,7 @@ void UPCGCustomHLSLSettings::UpdateInputDeclarations()
 			"uint {0}_GetNumData();\n"
 			"uint {0}_GetNumElements();\n"
 			"\n"
-			"// Valid types: bool, int, float, float2, float3, float4, Rotator (float3), Quat (float4), Transform (float4x4)\n"
+			"// Valid types: bool, int, float, float2, float3, float4, Rotator (float3), Quat (float4), Transform (float4x4), StringKey (int)\n"
 			"\n"
 			"<type> {0}_Get<type>(uint DataIndex, uint ElementIndex, uint AttributeId);\n"
 			"<type> {0}_Get<type>(uint DataIndex, uint ElementIndex, 'AttributeName');\n"),
@@ -732,7 +827,7 @@ void UPCGCustomHLSLSettings::UpdateOutputDeclarations()
 		}
 
 		OutputDeclarations += FString::Format(TEXT(
-			"// Valid types: bool, int, float, float2, float3, float4, Rotator (float3), Quat (float4), Transform (float4x4)\n"
+			"// Valid types: bool, int, float, float2, float3, float4, Rotator (float3), Quat (float4), Transform (float4x4), StringKey (int)\n"
 			"\n"
 			"void {0}_Set<type>(uint DataIndex, uint ElementIndex, uint AttributeId, <type> Value);\n"
 			"void {0}_Set<type>(uint DataIndex, uint ElementIndex, 'AttributeName', <type> Value);\n"),
diff --git a/Engine/Plugins/PCG/Source/PCG/Private/Compute/PCGComputeGraph.cpp b/Engine/Plugins/PCG/Source/PCG/Private/Compute/PCGComputeGraph.cpp
--- a/Engine/Plugins/PCG/Source/PCG/Private/Compute/PCGComputeGraph.cpp
+++ b/Engine/Plugins/PCG/Source/PCG/Private/Compute/PCGComputeGraph.cpp
@@ -86,47 +86,3 @@ void UPCGComputeGraph::OnKernelCompilationComplete(int32 InKernelIndex, FCompute
 		UE_LOG(LogPCG, Verbose, TEXT("Compilation message ignored for kernel index %d which has no associated node."), InKernelIndex);
 	}
 }
-
-void UPCGComputeGraph::FillInMissingAttributeTableTypes(const FPCGDataCollection& InComputeGraphElementInputData)
-{
-	TRACE_CPUPROFILER_EVENT_SCOPE(UPCGComputeGraph::FillInMissingAttributeTableTypes);
-
-	// Local to loop body but hoisted for perf.
-	TArray<FName> AttributeNames;
-	TArray<EPCGMetadataTypes> AttributeTypes;
-
-	for (TPair<FName, FPCGKernelAttributeIDAndType>& Entry : GlobalAttributeLookupTable)
-	{
-		if (Entry.Value.Type != EPCGKernelAttributeType::None)
-		{
-			continue;
-		}
-
-		for (const FPCGTaggedData& Data : InComputeGraphElementInputData.TaggedData)
-		{
-			AttributeNames.Reset();
-			AttributeTypes.Reset();
-
-			if (const UPCGMetadata* Metadata = Data.Data ? Data.Data->ConstMetadata() : nullptr)
-			{
-				Metadata->GetAttributes(AttributeNames, AttributeTypes);
-			}
-
-			for (int Index = 0; Index < AttributeNames.Num(); ++Index)
-			{
-				if (AttributeNames[Index] == Entry.Key)
-				{
-					Entry.Value.Type = PCGDataForGPUHelpers::GetAttributeTypeFromMetadataType(AttributeTypes[Index]);
-					break;
-				}
-			}
-
-			if (Entry.Value.Type != EPCGKernelAttributeType::None)
-			{
-				break;
-			}
-		}
-
-		ensure(Entry.Value.Type != EPCGKernelAttributeType::None);
-	}
-}
diff --git a/Engine/Plugins/PCG/Source/PCG/Private/Compute/PCGComputeGraph.h b/Engine/Plugins/PCG/Source/PCG/Private/Compute/PCGComputeGraph.h
--- a/Engine/Plugins/PCG/Source/PCG/Private/Compute/PCGComputeGraph.h
+++ b/Engine/Plugins/PCG/Source/PCG/Private/Compute/PCGComputeGraph.h
@@ -24,10 +24,7 @@ public:
 	void OnKernelCompilationComplete(int32 InKernelIndex, FComputeKernelCompileResults const& InCompileResults) override;
 	//~End UComputeGraph interface
 
-	/** Some attributes in table may be missing type information that we could not infer statically. Fill in missing types from the execution-time data. */
-	void FillInMissingAttributeTableTypes(const FPCGDataCollection& InComputeGraphElementInputData);
-
-	/** Get the global attribute indices. */
+	/** Get the global attribute information collated during compilation. */
 	const TMap<FName, FPCGKernelAttributeIDAndType>& GetAttributeLookupTable() const { return GlobalAttributeLookupTable; }
 
 public:
@@ -53,6 +50,9 @@ public:
 	UPROPERTY()
 	TArray<TObjectPtr<const UPCGSettings>> StaticMeshSpawners;
 
+	UPROPERTY()
+	bool bLogDataDescriptions = false;
+
 protected:
 	UPROPERTY()
 	TMap<FName /* Attribute name */, FPCGKernelAttributeIDAndType> GlobalAttributeLookupTable;
diff --git a/Engine/Plugins/PCG/Source/PCG/Private/Compute/PCGDataBinding.cpp b/Engine/Plugins/PCG/Source/PCG/Private/Compute/PCGDataBinding.cpp
new file mode 100644
--- /dev/null
+++ b/Engine/Plugins/PCG/Source/PCG/Private/Compute/PCGDataBinding.cpp
@@ -0,0 +1,237 @@
+// Copyright Epic Games, Inc. All Rights Reserved.
+
+#include "Compute/PCGDataBinding.h"
+
+#include "PCGNode.h"
+#include "PCGParamData.h"
+#include "PCGSettings.h"
+#include "Compute/PCGComputeCommon.h"
+#include "Compute/PCGComputeGraph.h"
+#include "Data/PCGPointData.h"
+#include "Metadata/PCGMetadata.h"
+
+void UPCGDataBinding::Initialize(
+	const UPCGComputeGraph* InComputeGraph,
+	TWeakObjectPtr<UPCGComponent> InSourceComponent,
+	const FPCGDataCollection& InComputeGraphElementInputData,
+	const TMap<FName, FPCGKernelAttributeIDAndType>& InStaticAttributeTable)
+{
+	check(InComputeGraph);
+
+	Graph = InComputeGraph;
+	SourceComponent = InSourceComponent;
+	GlobalAttributeLookupTable = InStaticAttributeTable;
+
+	InitializeInputData(InComputeGraphElementInputData);
+
+	FillInMissingAttributeTableTypes(InComputeGraphElementInputData);
+
+#if PCG_KERNEL_LOGGING_ENABLED
+	if (Graph->bLogDataDescriptions)
+	{
+		DebugLogDataDescriptions();
+	}
+#endif
+}
+
+void UPCGDataBinding::InitializeInputData(const FPCGDataCollection& InComputeGraphElementInputData)
+{
+	DataForGPU.InputDataCollection = InComputeGraphElementInputData;
+
+	// Link each input pin to the data collection, so that data providers can find the data.
+	for (const TWeakObjectPtr<const UPCGPin>& InputPinPtr : Graph->PinsReceivingDataFromCPU)
+	{
+		if (const UPCGPin* InputPin = InputPinPtr.Get())
+		{
+			DataForGPU.InputPins.Add(InputPin);
+		}
+	}
+
+	DataForGPU.InputPinLabelAliases = Graph->InputPinLabelAliases;
+
+	BuildStringTable();
+}
+
+void UPCGDataBinding::BuildStringTable()
+{
+	for (const FPCGTaggedData& Data : DataForGPU.InputDataCollection.TaggedData)
+	{
+		const UPCGMetadata* Metadata = Data.Data ? Data.Data->ConstMetadata() : nullptr;
+		if (!Metadata)
+		{
+			continue;
+		}
+
+		TArray<FName> AttributeNames;
+		TArray<EPCGMetadataTypes> AttributeTypes;
+		Metadata->GetAttributes(AttributeNames, AttributeTypes);
+
+		for (int AttributeIndex = 0; AttributeIndex < AttributeNames.Num(); ++AttributeIndex)
+		{
+			const FName AttributeName = AttributeNames[AttributeIndex];
+			const EPCGKernelAttributeType AttributeType = PCGDataForGPUHelpers::GetAttributeTypeFromMetadataType(AttributeTypes[AttributeIndex]);
+
+			if (AttributeType == EPCGKernelAttributeType::StringKey)
+			{
+				const FPCGMetadataAttributeBase* AttributeBase = Metadata->GetConstAttribute(AttributeName);
+				if (!AttributeBase)
+				{
+					UE_LOG(LogPCG, Warning, TEXT("Attribute '%s' not found."), *AttributeName.ToString());
+					continue;
+				}
+
+				auto AddEntryToStringTable = [AttributeBase, &StringTable=StringTable](int64 InValueKey)
+				{
+					if (AttributeBase->GetTypeId() == PCG::Private::MetadataTypes<FSoftObjectPath>::Id)
+					{
+						StringTable.AddUnique(static_cast<const FPCGMetadataAttribute<FSoftObjectPath>*>(AttributeBase)->GetValue(InValueKey).ToString());
+					}
+					else if (AttributeBase->GetTypeId() == PCG::Private::MetadataTypes<FString>::Id)
+					{
+						StringTable.AddUnique(static_cast<const FPCGMetadataAttribute<FString>*>(AttributeBase)->GetValue(InValueKey));
+					}
+					else
+					{
+						checkNoEntry();
+					}
+				};
+
+				if (const UPCGPointData* PointData = Cast<UPCGPointData>(Data.Data))
+				{
+					for (const FPCGPoint& Point : PointData->GetPoints())
+					{
+						AddEntryToStringTable(Point.MetadataEntry);
+					}
+				}
+				else if (const UPCGParamData* ParamData = Cast<UPCGParamData>(Data.Data))
+				{
+					const int32 NumElements = Metadata->GetItemCountForChild();
+
+					for (int64 MetadataKey = 0; MetadataKey < NumElements; ++MetadataKey)
+					{
+						AddEntryToStringTable(MetadataKey);
+					}
+				}
+				else { /** TODO: More attribute types! */ }
+			}
+		}
+	}
+}
+
+void UPCGDataBinding::FillInMissingAttributeTableTypes(const FPCGDataCollection& InComputeGraphElementInputData)
+{
+	TRACE_CPUPROFILER_EVENT_SCOPE(UPCGDataBinding::FillInMissingAttributeTableTypes);
+
+	// Local to loop body but hoisted for perf.
+	TArray<FName> AttributeNames;
+	TArray<EPCGMetadataTypes> AttributeTypes;
+
+	for (TPair<FName, FPCGKernelAttributeIDAndType>& Entry : GlobalAttributeLookupTable)
+	{
+		if (Entry.Value.Type != EPCGKernelAttributeType::None)
+		{
+			continue;
+		}
+
+		for (const FPCGTaggedData& Data : InComputeGraphElementInputData.TaggedData)
+		{
+			AttributeNames.Reset();
+			AttributeTypes.Reset();
+
+			if (const UPCGMetadata* Metadata = Data.Data ? Data.Data->ConstMetadata() : nullptr)
+			{
+				Metadata->GetAttributes(AttributeNames, AttributeTypes);
+			}
+
+			for (int Index = 0; Index < AttributeNames.Num(); ++Index)
+			{
+				if (AttributeNames[Index] == Entry.Key)
+				{
+					Entry.Value.Type = PCGDataForGPUHelpers::GetAttributeTypeFromMetadataType(AttributeTypes[Index]);
+					break;
+				}
+			}
+
+			if (Entry.Value.Type != EPCGKernelAttributeType::None)
+			{
+				break;
+			}
+		}
+	}
+}
+
+void UPCGDataBinding::DebugLogDataDescriptions()
+{
+#if PCG_KERNEL_LOGGING_ENABLED
+	UE_LOG(LogPCG, Display, TEXT("\n### ATTRIBUTE TABLE ###"));
+	for (TPair<FName, FPCGKernelAttributeIDAndType>& AttributeEntry : GlobalAttributeLookupTable)
+	{
+		UE_LOG(LogPCG, Display, TEXT("\tName: %s\t\tID: %d\t\tType: %d)"), *AttributeEntry.Get<0>().ToString(), AttributeEntry.Get<1>().Id, AttributeEntry.Get<1>().Type);
+	}
+
+	UE_LOG(LogPCG, Display, TEXT("\n### STRING TABLE ###"));
+	for (int32 i = 0; i < StringTable.Num(); ++i)
+	{
+		UE_LOG(LogPCG, Display, TEXT("\t%d: %s"), i, *StringTable[i]);
+	}
+
+	const UEnum* PCGDataTypeEnum = StaticEnum<EPCGDataType>();
+	const UEnum* PCGKernelAttributeTypeEnum = StaticEnum<EPCGKernelAttributeType>();
+	check(PCGDataTypeEnum && PCGKernelAttributeTypeEnum);
+
+	auto LogDataDescription = [this, PCGDataTypeEnum, PCGKernelAttributeTypeEnum](const FPCGDataCollectionDesc& InDataDesc)
+	{
+		for (int32 DataIndex = 0; DataIndex < InDataDesc.DataDescs.Num(); ++DataIndex)
+		{
+			const FPCGDataDesc& DataDesc = InDataDesc.DataDescs[DataIndex];
+
+			UE_LOG(LogPCG, Display, TEXT("\t\tData Index %d"), DataIndex);
+			UE_LOG(LogPCG, Display, TEXT("\t\t\tType: %s"), *PCGDataTypeEnum->GetNameStringByValue(static_cast<int64>(DataDesc.Type)));
+			UE_LOG(LogPCG, Display, TEXT("\t\t\tElementCount: %d"), DataDesc.ElementCount);
+			UE_LOG(LogPCG, Display, TEXT("\t\t\tAttributes (%d)"), DataDesc.AttributeDescs.Num());
+				
+			for (int32 AttributeIndex = 0; AttributeIndex < DataDesc.AttributeDescs.Num(); ++AttributeIndex)
+			{
+				const FPCGKernelAttributeDesc& AttributeDesc = DataDesc.AttributeDescs[AttributeIndex];
+
+				FString UniqueStringKeys;
+				for (int32 StringKey : AttributeDesc.UniqueStringKeys)
+				{
+					UniqueStringKeys += TEXT(", ");
+					UniqueStringKeys += FString::Format(TEXT("{0}"), { StringKey });
+				}
+
+				UE_LOG(LogPCG, Display, TEXT("\t\t\t\tID: %d\t\tName: %s\t\tType: %s\t\tUniqueStringKeys%s"),
+					AttributeDesc.Index,
+					*AttributeDesc.Name.ToString(),
+					*(PCGKernelAttributeTypeEnum->GetNameStringByValue(static_cast<int64>(AttributeDesc.Type))),
+					UniqueStringKeys.IsEmpty() ? TEXT(": ") : *UniqueStringKeys);
+			}
+		}
+	};
+
+	UE_LOG(LogPCG, Display, TEXT("\n### INPUT PIN DATA DESCRIPTIONS ###"));
+	for (const TWeakObjectPtr<const UPCGNode>& Node : Graph->KernelToNode)
+	{
+		const UPCGSettings* Settings = Node.IsValid() ? Node->GetSettings() : nullptr;
+		if (Settings && Settings->bDumpDataDescriptions)
+		{
+			UE_LOG(LogPCG, Display, TEXT("Node: %s"), *Node->GetNodeTitle(EPCGNodeTitleType::ListView).ToString());
+
+			for (const UPCGPin* Pin : Node->GetInputPins())
+			{
+				const FPCGDataCollectionDesc DataDesc = Settings->ComputeInputPinDataDesc(Pin, this);
+				UE_LOG(LogPCG, Display, TEXT("\tInput Pin: %s (%d data)"), *Pin->Properties.Label.ToString(), DataDesc.DataDescs.Num());
+				LogDataDescription(DataDesc);
+			}
+
+			for (const UPCGPin* Pin : Node->GetOutputPins())
+			{
+				const FPCGDataCollectionDesc DataDesc = Settings->ComputeOutputPinDataDesc(Pin, this);
+				UE_LOG(LogPCG, Display, TEXT("\tOutput Pin: %s (%d data)"), *Pin->Properties.Label.ToString(), DataDesc.DataDescs.Num());
+				LogDataDescription(DataDesc);
+			}
+		}
+	}
+#endif
+}
diff --git a/Engine/Plugins/PCG/Source/PCG/Private/Compute/PCGDataBinding.h b/Engine/Plugins/PCG/Source/PCG/Private/Compute/PCGDataBinding.h
--- a/Engine/Plugins/PCG/Source/PCG/Private/Compute/PCGDataBinding.h
+++ b/Engine/Plugins/PCG/Source/PCG/Private/Compute/PCGDataBinding.h
@@ -5,8 +5,6 @@
 #include "PCGData.h"
 #include "Compute/PCGDataForGPU.h"
 
-#include "ComputeFramework/ComputeGraphInstance.h"
-
 #include "PCGDataBinding.generated.h"
 
 class UPCGComponent;
@@ -31,6 +29,12 @@ struct FPCGSpawnerPrimitives
 	// Same for all primitives
 	UPROPERTY()
 	TArray<FUintVector4> AttributeIdOffsetStrides;
+
+	UPROPERTY()
+	int32 SelectorAttributeId = -1;
+
+	UPROPERTY()
+	TArray<uint32> PrimitiveStringKeys;
 };
 
 UCLASS(Transient, Category = PCG)
@@ -39,16 +43,45 @@ class UPCGDataBinding : public UObject
 	GENERATED_BODY()
 
 public:
+	void Initialize(const UPCGComputeGraph* InComputeGraph, TWeakObjectPtr<UPCGComponent> InSourceComponent, const FPCGDataCollection& InComputeGraphElementInputData, const TMap<FName, FPCGKernelAttributeIDAndType>& InStaticAttributeTable);
+
+	/** Get the global attribute information for this graph instance. */
+	const TMap<FName, FPCGKernelAttributeIDAndType>& GetAttributeLookupTable() const { return GlobalAttributeLookupTable; }
+
+	/** Get the set of all unique strings entering the compute graph. */
+	const TArray<FString>& GetStringTable() const { return StringTable; }
+
+private:
+	void InitializeInputData(const FPCGDataCollection& InComputeGraphElementInputData);
+
+	void BuildStringTable();
+
+	/** Some attributes in table may be missing type information that we could not infer statically. Fill in missing types from the execution-time data. */
+	void FillInMissingAttributeTableTypes(const FPCGDataCollection& InComputeGraphElementInputData);
+
+	void DebugLogDataDescriptions();
+
+public:
+	UPROPERTY()
 	TWeakObjectPtr<UPCGComponent> SourceComponent;
 
+	UPROPERTY()
 	FPCGDataCollection OutputDataCollection;
 
 	UPROPERTY()
 	TObjectPtr<const UPCGComputeGraph> Graph = nullptr;
 
 	// Data collection pointer + set of pins that get data from the collection (cross barrier)
+	UPROPERTY()
 	FPCGDataForGPU DataForGPU;
 
 	UPROPERTY()
 	TMap<TObjectPtr<const UPCGSettings>, FPCGSpawnerPrimitives> MeshSpawnersToPrimitives;
+
+protected:
+	UPROPERTY()
+	TMap<FName /* Attribute name */, FPCGKernelAttributeIDAndType> GlobalAttributeLookupTable;
+
+	UPROPERTY()
+	TArray<FString> StringTable;
 };
diff --git a/Engine/Plugins/PCG/Source/PCG/Private/Compute/PCGDataForGPU.cpp b/Engine/Plugins/PCG/Source/PCG/Private/Compute/PCGDataForGPU.cpp
--- a/Engine/Plugins/PCG/Source/PCG/Private/Compute/PCGDataForGPU.cpp
+++ b/Engine/Plugins/PCG/Source/PCG/Private/Compute/PCGDataForGPU.cpp
@@ -58,6 +58,10 @@ namespace PCGDataForGPUHelpers
 			return EPCGKernelAttributeType::Quat;
 		case EPCGMetadataTypes::Transform:
 			return EPCGKernelAttributeType::Transform;
+		case EPCGMetadataTypes::SoftObjectPath:
+		case EPCGMetadataTypes::String:
+			return EPCGKernelAttributeType::StringKey;
+		case EPCGMetadataTypes::SoftClassPath:
 		default:
 			return EPCGKernelAttributeType::Invalid;
 		}
@@ -70,6 +74,7 @@ namespace PCGDataForGPUHelpers
 		case EPCGKernelAttributeType::Bool:
 		case EPCGKernelAttributeType::Int:
 		case EPCGKernelAttributeType::Float:
+		case EPCGKernelAttributeType::StringKey:
 			return 4;
 		case EPCGKernelAttributeType::Float2:
 			return 8;
@@ -87,7 +92,7 @@ namespace PCGDataForGPUHelpers
 		}
 	}
 
-	bool PackAttributeHelper(const FPCGMetadataAttributeBase* InAttributeBase, const FPCGKernelAttributeDesc& InAttributeDesc, PCGMetadataEntryKey InEntryKey, TArray<uint32>& OutPackedDataCollection, uint32 ElementIndex)
+	bool PackAttributeHelper(const FPCGMetadataAttributeBase* InAttributeBase, const FPCGKernelAttributeDesc& InAttributeDesc, PCGMetadataEntryKey InEntryKey, const TArray<FString>& InStringTable, TArray<uint32>& OutPackedDataCollection, uint32 ElementIndex)
 	{
 		check(InAttributeBase);
 
@@ -212,6 +217,24 @@ namespace PCGDataForGPUHelpers
 			OutPackedDataCollection[ElementIndex + 15] = FMath::AsUInt(static_cast<float>(Matrix.M[3][3]));
 			break;
 		}
+		case PCG::Private::MetadataTypes<FString>::Id:
+		{
+			// String stored as an integer for reading/writing in kernel, and accompanying string table in data description.
+			const FPCGMetadataAttribute<FString>* Attribute = static_cast<const FPCGMetadataAttribute<FString>*>(InAttributeBase);
+			const int32 Value = InStringTable.IndexOfByKey(Attribute->GetValue(ValueKey));
+			check(StrideBytes == 4);
+			OutPackedDataCollection[ElementIndex] = Value;
+			break;
+		}
+		case PCG::Private::MetadataTypes<FSoftObjectPath>::Id:
+		{
+			// SOP path string stored as an integer for reading/writing in kernel, and accompanying string table in data description.
+			const FPCGMetadataAttribute<FSoftObjectPath>* Attribute = static_cast<const FPCGMetadataAttribute<FSoftObjectPath>*>(InAttributeBase);
+			const int32 Value = InStringTable.IndexOfByKey(Attribute->GetValue(ValueKey).ToString());
+			check(StrideBytes == 4);
+			OutPackedDataCollection[ElementIndex] = Value;
+			break;
+		}
 		default:
 			return false;
 		}
@@ -259,17 +282,21 @@ namespace PCGDataForGPUHelpers
 		{
 			return Metadata->FindOrCreateAttribute<FTransform>(AttributeDesc.Name);
 		}
+		case EPCGKernelAttributeType::StringKey:
+		{
+			return Metadata->FindOrCreateAttribute<FString>(AttributeDesc.Name);
+		}
 		default:
 			return nullptr;
 		}
 	}
 
-	bool UnpackAttributeHelper(const void* InPackedData, uint32 ElementIndex, FPCGMetadataAttributeBase* AttributeBase, const FPCGKernelAttributeDesc& AttributeDesc, PCGMetadataEntryKey EntryKey)
+	bool UnpackAttributeHelper(const void* InPackedData, const TArray<FString>& InStringTable, uint32 ElementIndex, FPCGMetadataAttributeBase* AttributeBase, const FPCGKernelAttributeDesc& AttributeDesc, PCGMetadataEntryKey EntryKey)
 	{
 		check(InPackedData && AttributeBase);
 
 		const float* DataAsFloat = static_cast<const float*>(InPackedData);
-		const int* DataAsInt = static_cast<const int*>(InPackedData);
+		const int32* DataAsInt = static_cast<const int32*>(InPackedData);
 
 		switch (AttributeDesc.Type)
 		{
@@ -285,7 +312,7 @@ namespace PCGDataForGPUHelpers
 		{
 			FPCGMetadataAttribute<int>* Attribute = static_cast<FPCGMetadataAttribute<int>*>(AttributeBase);
 
-			const int Value = DataAsInt[ElementIndex];
+			const int32 Value = DataAsInt[ElementIndex];
 			Attribute->SetValue(EntryKey, Value);
 			break;
 		}
@@ -379,6 +406,22 @@ namespace PCGDataForGPUHelpers
 			Attribute->SetValue(EntryKey, Value);
 			break;
 		}
+		case EPCGKernelAttributeType::StringKey:
+		{
+			FPCGMetadataAttribute<FString>* Attribute = static_cast<FPCGMetadataAttribute<FString>*>(AttributeBase);
+			
+			const int32 StringKey = DataAsInt[ElementIndex];
+			if (InStringTable.IsValidIndex(StringKey))
+			{
+				Attribute->SetValue(EntryKey, InStringTable[StringKey]);
+			}
+			else
+			{
+				UE_LOG(LogPCG, Error, TEXT("String retrieval failed for string key %d. String table size is %d."), StringKey, InStringTable.Num());
+				ensure(false);
+			}
+			break;
+		}
 		default:
 			return false;
 		}
@@ -395,7 +438,7 @@ namespace PCGDataForGPUHelpers
 	{
 		check(InBinding);
 		check(InBinding->Graph);
-		const TMap<FName, FPCGKernelAttributeIDAndType>& GlobalAttributeLookupTable = InBinding->Graph->GetAttributeLookupTable();
+		const TMap<FName, FPCGKernelAttributeIDAndType>& GlobalAttributeLookupTable = InBinding->GetAttributeLookupTable();
 
 		uint32 OffsetFloats = 0;
 
@@ -404,12 +447,16 @@ namespace PCGDataForGPUHelpers
 			const FPCGKernelAttributeIDAndType* FoundAttribute = GlobalAttributeLookupTable.Find(AttributeName);
 			if (!FoundAttribute)
 			{
-				ensure(false);
 				continue;
 			}
 
-			const uint32 AttributeId = static_cast<uint32>(FoundAttribute->Id);
 			const EPCGKernelAttributeType AttributeType = FoundAttribute->Type;
+			if (AttributeType == EPCGKernelAttributeType::None)
+			{
+				continue;
+			}
+
+			const uint32 AttributeId = static_cast<uint32>(FoundAttribute->Id);
 			const uint32 StrideFloats = GetAttributeTypeStrideBytes(AttributeType) / sizeof(float);
 
 			OutAttributeIdOffsetStrides.Emplace(AttributeId, OffsetFloats, StrideFloats, /*Unused*/0);
@@ -440,17 +487,18 @@ FPCGDataDesc::FPCGDataDesc(EPCGDataType InType, int InElementCount)
 	: Type(InType)
 	, ElementCount(InElementCount)
 {
-	InitializeAttributeDescs(nullptr);
+	TArray<FString> StringTableDummy;
+	InitializeAttributeDescs(nullptr, {}, StringTableDummy);
 }
 
-FPCGDataDesc::FPCGDataDesc(const UPCGData* Data, const TMap<FName, FPCGKernelAttributeIDAndType>& GlobalAttributeLookupTable)
+FPCGDataDesc::FPCGDataDesc(const UPCGData* InData, const TMap<FName, FPCGKernelAttributeIDAndType>& InGlobalAttributeLookupTable, const TArray<FString>& InStringTable)
 {
-	check(Data);
+	check(InData);
 
-	Type = Data->GetDataType();
-	ElementCount = PCGComputeHelpers::GetElementCount(Data);
+	Type = InData->GetDataType();
+	ElementCount = PCGComputeHelpers::GetElementCount(InData);
 
-	InitializeAttributeDescs(Data->ConstMetadata(), GlobalAttributeLookupTable);
+	InitializeAttributeDescs(InData, InGlobalAttributeLookupTable, InStringTable);
 }
 
 uint32 FPCGDataDesc::ComputePackedSize() const
@@ -479,21 +527,23 @@ uint32 FPCGDataDesc::ComputePackedSize() const
 	return DataSizeBytes;
 }
 
-void FPCGDataDesc::InitializeAttributeDescs(const UPCGMetadata* Metadata, const TMap<FName, FPCGKernelAttributeIDAndType>& GlobalAttributeLookupTable)
+void FPCGDataDesc::InitializeAttributeDescs(const UPCGData* InData, const TMap<FName, FPCGKernelAttributeIDAndType>& InGlobalAttributeLookupTable, const TArray<FString>& InStringTable)
 {
+	TRACE_CPUPROFILER_EVENT_SCOPE(FPCGDataDesc::InitializeAttributeDescs);
+
 	if (Type == EPCGDataType::Point)
 	{
 		AttributeDescs.Append(PCGDataForGPUConstants::PointPropertyDescs, NUM_POINT_PROPERTIES);
 	}
 	else { /* TODO: More types! */ }
 
-	if (Metadata)
+	if (const UPCGMetadata* Metadata = InData ? InData->ConstMetadata() : nullptr)
 	{
 		TArray<FName> AttributeNames;
 		TArray<EPCGMetadataTypes> AttributeTypes;
 		Metadata->GetAttributes(AttributeNames, AttributeTypes);
 
-		TArray<FPCGKernelAttributeKey> DelayedAttributeKeys; // Attribute keys that don't exist in the global lookup table must be delayed so we can append them at the end.
+		TArray<TPair<FPCGKernelAttributeKey, TArray<int32>>> DelayedAttributeKeysAndStringKeys; // Attribute keys that don't exist in the global lookup table must be delayed so we can append them at the end.
 		int NumAttributesFromLUT = 0; // Keep track of how many attributes come from the LUT. This will help give us the starting index for our delayed attributes.
 
 		for (int CustomAttributeIndex = 0; CustomAttributeIndex < AttributeNames.Num(); ++CustomAttributeIndex)
@@ -528,21 +578,92 @@ void FPCGDataDesc::InitializeAttributeDescs(const UPCGMetadata* Metadata, const
 				continue;
 			}
 
-			if (const FPCGKernelAttributeIDAndType* AttributeIdAndType = GlobalAttributeLookupTable.Find(AttributeName))
+			TArray<int32> UniqueStringKeys;
+
+			if (AttributeType == EPCGKernelAttributeType::StringKey)
+			{
+				const FPCGMetadataAttributeBase* AttributeBase = Metadata->GetConstAttribute(AttributeName);
+				check(AttributeBase);
+
+				if (Type == EPCGDataType::Point && ensure(InData->IsA<UPCGPointData>()))
+				{
+					const UPCGPointData* PointData = CastChecked<UPCGPointData>(InData);
+
+					for (int32 PointIndex = 0; PointIndex < PointData->GetPoints().Num(); ++PointIndex)
+					{
+						const PCGMetadataValueKey ValueKey = AttributeBase->GetValueKey(PointData->GetPoints()[PointIndex].MetadataEntry);
+
+						int StringTableIndex = INDEX_NONE;
+
+						if (AttributeBase->GetTypeId() == PCG::Private::MetadataTypes<FSoftObjectPath>::Id)
+						{
+							StringTableIndex = InStringTable.IndexOfByKey(static_cast<const FPCGMetadataAttribute<FSoftObjectPath>*>(AttributeBase)->GetValue(ValueKey).ToString());
+						}
+						else if (AttributeBase->GetTypeId() == PCG::Private::MetadataTypes<FString>::Id)
+						{
+							StringTableIndex = InStringTable.IndexOfByKey(static_cast<const FPCGMetadataAttribute<FString>*>(AttributeBase)->GetValue(ValueKey));
+						}
+						else
+						{
+							// Should not get here if attribute type is string key.
+							checkNoEntry();
+						}
+
+						if (StringTableIndex != INDEX_NONE)
+						{
+							UniqueStringKeys.AddUnique(StringTableIndex);
+						}
+					}
+				}
+				else if (Type == EPCGDataType::Param && ensure(InData->IsA<UPCGParamData>()))
+				{
+					const int32 NumElements = Metadata->GetItemCountForChild();
+
+					for (int64 MetadataKey = 0; MetadataKey < NumElements; ++MetadataKey)
+					{
+						int StringTableIndex = INDEX_NONE;
+
+						if (AttributeBase->GetTypeId() == PCG::Private::MetadataTypes<FSoftObjectPath>::Id)
+						{
+							StringTableIndex = InStringTable.IndexOfByKey(static_cast<const FPCGMetadataAttribute<FSoftObjectPath>*>(AttributeBase)->GetValue(MetadataKey).ToString());
+						}
+						else if (AttributeBase->GetTypeId() == PCG::Private::MetadataTypes<FString>::Id)
+						{
+							StringTableIndex = InStringTable.IndexOfByKey(static_cast<const FPCGMetadataAttribute<FString>*>(AttributeBase)->GetValue(MetadataKey));
+						}
+						else
+						{
+							// Should not get here if attribute type is string key.
+							checkNoEntry();
+						}
+
+						if (StringTableIndex != INDEX_NONE)
+						{
+							UniqueStringKeys.AddUnique(StringTableIndex);
+						}
+					}
+				}
+				else { /* TODO: More types! */ }
+			}
+
+			if (const FPCGKernelAttributeIDAndType* AttributeIdAndType = InGlobalAttributeLookupTable.Find(AttributeName))
 			{
-				AttributeDescs.Emplace(AttributeIdAndType->Id, AttributeType, AttributeName);
+				AttributeDescs.Emplace(AttributeIdAndType->Id, AttributeType, AttributeName, MoveTemp(UniqueStringKeys));
 				++NumAttributesFromLUT;
 			}
 			else
 			{
-				DelayedAttributeKeys.Emplace(AttributeName, AttributeType);
+				TPair<FPCGKernelAttributeKey, TArray<int32>> AttributeAndStringKeys;
+				AttributeAndStringKeys.Get<0>() = FPCGKernelAttributeKey(AttributeName, AttributeType);
+				AttributeAndStringKeys.Get<1>() = MoveTemp(UniqueStringKeys);
+				DelayedAttributeKeysAndStringKeys.Add(MoveTemp(AttributeAndStringKeys));
 			}
 		}
 
-		for (int DelayedAttributeIndex = 0; DelayedAttributeIndex < DelayedAttributeKeys.Num(); ++DelayedAttributeIndex)
+		for (int DelayedAttributeIndex = 0; DelayedAttributeIndex < DelayedAttributeKeysAndStringKeys.Num(); ++DelayedAttributeIndex)
 		{
-			const FPCGKernelAttributeKey& AttributeKey = DelayedAttributeKeys[DelayedAttributeIndex];
-			AttributeDescs.Emplace(NUM_RESERVED_ATTRS + DelayedAttributeIndex + NumAttributesFromLUT, AttributeKey.Type, AttributeKey.Name);
+			const FPCGKernelAttributeKey& AttributeKey = DelayedAttributeKeysAndStringKeys[DelayedAttributeIndex].Get<0>();
+			AttributeDescs.Emplace(NUM_RESERVED_ATTRS + DelayedAttributeIndex + NumAttributesFromLUT, AttributeKey.Type, AttributeKey.Name, MoveTemp(DelayedAttributeKeysAndStringKeys[DelayedAttributeIndex].Get<1>()));
 		}
 	}
 }
@@ -550,9 +671,10 @@ void FPCGDataDesc::InitializeAttributeDescs(const UPCGMetadata* Metadata, const
 FPCGDataCollectionDesc FPCGDataCollectionDesc::BuildFromInputDataCollectionAndInputPinLabel(
 	const FPCGDataCollection& InDataCollection,
 	FName InputPinLabel,
-	const TMap<FName, FPCGKernelAttributeIDAndType>& InAttributeLookupTable)
+	const TMap<FName, FPCGKernelAttributeIDAndType>& InAttributeLookupTable,
+	const TArray<FString>& InStringTable)
 {
-	FPCGDataCollectionDesc Desc;
+	FPCGDataCollectionDesc CollectionDesc;
 	TArray<FPCGTaggedData> DataForPin = InDataCollection.GetInputsByPin(InputPinLabel);
 
 	for (const FPCGTaggedData& Data : DataForPin)
@@ -562,10 +684,10 @@ FPCGDataCollectionDesc FPCGDataCollectionDesc::BuildFromInputDataCollectionAndIn
 			continue;
 		}
 
-		Desc.DataDescs.Emplace(Data.Data, InAttributeLookupTable);
+		CollectionDesc.DataDescs.Emplace(Data.Data, InAttributeLookupTable, InStringTable);
 	}
 
-	return Desc;
+	return CollectionDesc;
 }
 
 uint32 FPCGDataCollectionDesc::ComputePackedSize(TArray<uint32>* OutDataAddresses) const
@@ -598,7 +720,7 @@ uint32 FPCGDataCollectionDesc::ComputePackedSize(TArray<uint32>* OutDataAddresse
 	return TotalCollectionSizeBytes;
 }
 
-void FPCGDataCollectionDesc::PackDataCollection(const FPCGDataCollection& InDataCollection, FName InPin, TArray<uint32>& OutPackedDataCollection) const
+void FPCGDataCollectionDesc::PackDataCollection(const FPCGDataCollection& InDataCollection, FName InPin, const TArray<FString>& InStringTable, TArray<uint32>& OutPackedDataCollection) const
 {
 	const TArray<FPCGTaggedData> InputData = InDataCollection.GetInputsByPin(InPin);
 	const uint32 NumData = InputData.Num();
@@ -659,7 +781,7 @@ void FPCGDataCollectionDesc::PackDataCollection(const FPCGDataCollection& InData
 
 					if (AttributeBase) // Pack attribute
 					{
-						ensure(PCGDataForGPUHelpers::PackAttributeHelper(AttributeBase, AttributeDesc, Points[ElementIndex].MetadataEntry, OutPackedDataCollection, PackedDataElementIndex));
+						ensure(PCGDataForGPUHelpers::PackAttributeHelper(AttributeBase, AttributeDesc, Points[ElementIndex].MetadataEntry, InStringTable, OutPackedDataCollection, PackedDataElementIndex));
 					}
 					else // Pack property
 					{
@@ -772,6 +894,12 @@ void FPCGDataCollectionDesc::PackDataCollection(const FPCGDataCollection& InData
 
 			for (const FPCGKernelAttributeDesc& AttributeDesc : AttributeDescs)
 			{
+				const FPCGMetadataAttributeBase* AttributeBase = Metadata->GetConstAttribute(AttributeDesc.Name);
+				if (!AttributeBase)
+				{
+					continue;
+				}
+
 				const uint32 AttributeId = AttributeDesc.Index;
 				const uint32 AttributeStrideBytes = PCGDataForGPUHelpers::GetAttributeTypeStrideBytes(AttributeDesc.Type);
 				const uint32 AttributeNumComponents = AttributeStrideBytes / sizeof(uint32); // E.g. float3 has 3 components
@@ -784,17 +912,13 @@ void FPCGDataCollectionDesc::PackDataCollection(const FPCGDataCollection& InData
 				OutPackedDataCollection[AttributeHeaderIndex + 0] = PackedIdAndStride;
 				OutPackedDataCollection[AttributeHeaderIndex + 1] = CurrentAttributeAddress;
 
-				const FPCGMetadataAttributeBase* AttributeBase = Metadata->GetConstAttribute(AttributeDesc.Name);
-
 				for (uint32 ElementIndex = 0; ElementIndex < NumElements; ++ElementIndex)
 				{
 					const uint32 PackedDataElementIndex = AttributeIndex + (ElementIndex * AttributeNumComponents);
 					const int64 MetadataKey = ElementIndex;
 
-					if (AttributeBase) // Pack attribute
-					{
-						ensure(PCGDataForGPUHelpers::PackAttributeHelper(AttributeBase, AttributeDesc, MetadataKey, OutPackedDataCollection, PackedDataElementIndex));
-					}
+					// Pack attribute
+					ensure(PCGDataForGPUHelpers::PackAttributeHelper(AttributeBase, AttributeDesc, MetadataKey, InStringTable, OutPackedDataCollection, PackedDataElementIndex));
 				}
 
 				CurrentAttributeAddress += NumElements * AttributeNumComponents * 4;
@@ -889,12 +1013,12 @@ void FPCGDataCollectionDesc::PrepareBufferForKernelOutput(TArray<uint32>& OutPac
 	}
 }
 
-EPCGUnpackDataCollectionResult FPCGDataCollectionDesc::UnpackDataCollection(const TArray<uint8>& InPackedData, FName InPin, FPCGDataCollection& OutDataCollection) const
+EPCGUnpackDataCollectionResult FPCGDataCollectionDesc::UnpackDataCollection(const TArray<uint8>& InPackedData, FName InPin, const TArray<FString>& InStringTable, FPCGDataCollection& OutDataCollection) const
 {
 	const void* PackedData = InPackedData.GetData();
 	const float* DataAsFloat = static_cast<const float*>(PackedData);
 	const uint32* DataAsUint = static_cast<const uint32*>(PackedData);
-	const int* DataAsInt = static_cast<const int*>(PackedData);
+	const int32* DataAsInt = static_cast<const int32*>(PackedData);
 
 	const uint32 NumPackedFloats = InPackedData.Num() / 4;
 
@@ -972,7 +1096,7 @@ EPCGUnpackDataCollectionResult FPCGDataCollectionDesc::UnpackDataCollection(cons
 					if (AttributeBase) // Unpack attribute
 					{
 						Metadata->InitializeOnSet(OutPoints[ElementIndex].MetadataEntry);
-						ensure(PCGDataForGPUHelpers::UnpackAttributeHelper(PackedData, PackedDataElementIndex, AttributeBase, AttributeDesc, OutPoints[ElementIndex].MetadataEntry));
+						ensure(PCGDataForGPUHelpers::UnpackAttributeHelper(PackedData, InStringTable, PackedDataElementIndex, AttributeBase, AttributeDesc, OutPoints[ElementIndex].MetadataEntry));
 					}
 					else // Unpack property
 					{
@@ -1063,7 +1187,7 @@ EPCGUnpackDataCollectionResult FPCGDataCollectionDesc::UnpackDataCollection(cons
 						}
 						case POINT_SEED_ATTRIBUTE_ID:
 						{
-							const int Seed = DataAsInt[PackedDataElementIndex];
+							const int32 Seed = DataAsInt[PackedDataElementIndex];
 
 							OutPoints[ElementIndex].Seed = Seed;
 							break;
@@ -1136,7 +1260,7 @@ EPCGUnpackDataCollectionResult FPCGDataCollectionDesc::UnpackDataCollection(cons
 						const uint32 PackedDataElementIndex = AttributeIndex + ElementIndex * AttributeNumComponents;
 						check(PackedDataElementIndex + AttributeNumComponents <= NumPackedFloats);
 
-						ensure(PCGDataForGPUHelpers::UnpackAttributeHelper(PackedData, PackedDataElementIndex, AttributeBase, AttributeDesc, ElementIndex));
+						ensure(PCGDataForGPUHelpers::UnpackAttributeHelper(PackedData, InStringTable, PackedDataElementIndex, AttributeBase, AttributeDesc, ElementIndex));
 					}
 				});
 			}
diff --git a/Engine/Plugins/PCG/Source/PCG/Private/Elements/PCGStaticMeshSpawner.cpp b/Engine/Plugins/PCG/Source/PCG/Private/Elements/PCGStaticMeshSpawner.cpp
--- a/Engine/Plugins/PCG/Source/PCG/Private/Elements/PCGStaticMeshSpawner.cpp
+++ b/Engine/Plugins/PCG/Source/PCG/Private/Elements/PCGStaticMeshSpawner.cpp
@@ -6,6 +6,7 @@
 #include "PCGCustomVersion.h"
 #include "PCGManagedResource.h"
 #include "Compute/PCGComputeCommon.h"
+#include "Compute/PCGDataBinding.h"
 #include "Compute/DataInterfaces/PCGInstanceDataInterface.h"
 #include "Compute/DataInterfaces/Elements/PCGStaticMeshSpawnerDataInterface.h"
 #include "Data/PCGPointData.h"
@@ -53,36 +54,44 @@ bool UPCGStaticMeshSpawnerSettings::IsKernelValid(FPCGContext* InContext, bool b
 		return false;
 	}
 
-	UPCGMeshSelectorWeighted* Selector = Cast<UPCGMeshSelectorWeighted>(MeshSelectorParameters);
-
-	if (!Selector)
-	{
-		PCG_KERNEL_VALIDATION_ERR(InContext, this, bQuiet, LOCTEXT("InvalidMeshSelector", "Currently GPU Static Mesh Spawner nodes must use PCGMeshSelectorWeighted as the mesh selector type."));
-		return false;
-	}
-
-	if (Selector->MeshEntries.IsEmpty())
-	{
-		PCG_KERNEL_VALIDATION_ERR(InContext, this, bQuiet, LOCTEXT("NoMeshEntries", "No meshes specified."));
-		return false;
-	}
-
-	for (const FPCGMeshSelectorWeightedEntry& Entry : Selector->MeshEntries)
+	if (UPCGMeshSelectorWeighted* SelectorWeighted = Cast<UPCGMeshSelectorWeighted>(MeshSelectorParameters))
 	{
-		UStaticMesh* StaticMesh = Entry.Descriptor.StaticMesh.LoadSynchronous();
-
-		if (!StaticMesh)
+		if (SelectorWeighted->MeshEntries.IsEmpty())
 		{
-			PCG_KERNEL_VALIDATION_ERR(InContext, this, bQuiet, LOCTEXT("UnassignedMesh", "Unassigned mesh."));
+			PCG_KERNEL_VALIDATION_ERR(InContext, this, bQuiet, LOCTEXT("NoMeshEntries", "No meshes specified in selector."));
 			return false;
 		}
 
-		if (!StaticMesh->HasValidNaniteData())
+		for (const FPCGMeshSelectorWeightedEntry& Entry : SelectorWeighted->MeshEntries)
 		{
-			PCG_KERNEL_VALIDATION_ERR(InContext, this, bQuiet, LOCTEXT("NonNaniteMesh", "Only Nanite meshes are currently supported by the GPU Static Mesh Spawner path."));
+			UStaticMesh* StaticMesh = Entry.Descriptor.StaticMesh.LoadSynchronous();
+
+			if (!StaticMesh)
+			{
+				PCG_KERNEL_VALIDATION_ERR(InContext, this, bQuiet, LOCTEXT("UnassignedMesh", "Unassigned mesh."));
+				return false;
+			}
+
+			if (!StaticMesh->HasValidNaniteData())
+			{
+				PCG_KERNEL_VALIDATION_ERR(InContext, this, bQuiet, LOCTEXT("NonNaniteMesh", "Only Nanite meshes are currently supported by the GPU Static Mesh Spawner path."));
+				return false;
+			}
+		}
+	}
+	else if (UPCGMeshSelectorByAttribute* SelectorByAttribute = Cast<UPCGMeshSelectorByAttribute>(MeshSelectorParameters))
+	{
+		if (SelectorByAttribute->AttributeName == NAME_None)
+		{
+			PCG_KERNEL_VALIDATION_ERR(InContext, this, bQuiet, LOCTEXT("NoMeshAttribute", "Mesh selector attribute name is 'None'."));
 			return false;
 		}
 	}
+	else
+	{
+		PCG_KERNEL_VALIDATION_ERR(InContext, this, bQuiet, LOCTEXT("InvalidMeshSelector", "Currently GPU Static Mesh Spawner nodes must use PCGMeshSelectorWeighted or UPCGMeshSelectorByAttribute as the mesh selector type."));
+		return false;
+	}
 
 	// Currently instance packers must be able to specify a full list of attribute names upfront, to build the attribute table at compile time.
 	// TODO: We should be able to augment a static attribute table with new attributes at execution time, which will allow other types like regex.
@@ -112,11 +121,8 @@ FPCGDataCollectionDesc UPCGStaticMeshSpawnerSettings::ComputeOutputPinDataDesc(c
 
 	check(OutputPin);
 
-	const UPCGNode* Node = CastChecked<UPCGNode>(GetOuter());
-	const UPCGPin* InputPin = Node->GetInputPin(PCGPinConstants::DefaultInputLabel);
-
 	// Passthrough - output description matches input description.
-	return ComputeInputPinDataDesc(InputPin, Binding);
+	return ComputeInputPinDataDesc(PCGPinConstants::DefaultInputLabel, Binding);
 }
 
 const TArray<FPCGKernelAttributeKey> UPCGStaticMeshSpawnerSettings::GetKernelAttributeKeys() const
@@ -136,6 +142,12 @@ const TArray<FPCGKernelAttributeKey> UPCGStaticMeshSpawnerSettings::GetKernelAtt
 		}
 	}
 
+	if (const UPCGMeshSelectorByAttribute* Selector = Cast<UPCGMeshSelectorByAttribute>(MeshSelectorParameters))
+	{
+		// Add an attribute key for the given attribute so we register it as being read.
+		AttributeKeys.Emplace(Selector->AttributeName, EPCGKernelAttributeType::None);
+	}
+
 	return AttributeKeys;
 }
 
@@ -144,7 +156,18 @@ int UPCGStaticMeshSpawnerSettings::ComputeKernelThreadCount(const UPCGDataBindin
 	const UPCGNode* Node = CastChecked<UPCGNode>(GetOuter());
 	const UPCGPin* Pin = Node->GetInputPin(PCGPinConstants::DefaultInputLabel);
 
-	return ComputeInputPinDataDesc(Pin, Binding).ComputeDataElementCount(EPCGDataType::Point);
+	uint32 ElementCount = ComputeInputPinDataDesc(Pin, Binding).ComputeDataElementCount(EPCGDataType::Point);
+
+	if (const FPCGSpawnerPrimitives* FoundPrimitives = Binding->MeshSpawnersToPrimitives.Find(this))
+	{
+		// When selecting primitives dynamically (by attribute), we don't know statically how many instances will end up in each primitive, so run worst case.
+		if (FoundPrimitives->SelectorAttributeId != -1)
+		{
+			ElementCount *= FoundPrimitives->Primitives.Num();
+		}
+	}
+
+	return ElementCount;
 }
 
 void UPCGStaticMeshSpawnerSettings::CreateAdditionalInputDataInterfaces(TArray<TObjectPtr<UComputeDataInterface>>& OutDataInterfaces) const
diff --git a/Engine/Plugins/PCG/Source/PCG/Private/Graph/PCGGraphCompilerGPU.cpp b/Engine/Plugins/PCG/Source/PCG/Private/Graph/PCGGraphCompilerGPU.cpp
--- a/Engine/Plugins/PCG/Source/PCG/Private/Graph/PCGGraphCompilerGPU.cpp
+++ b/Engine/Plugins/PCG/Source/PCG/Private/Graph/PCGGraphCompilerGPU.cpp
@@ -865,6 +865,8 @@ void FPCGGraphCompilerGPU::BuildGPUGraphTask(
 			OutputDataInterfaceIndexAndPin.Emplace(DebugDIIndex, NAME_None);
 		}
 
+		ComputeGraph->bLogDataDescriptions |= Settings->bDumpDataDescriptions;
+
 		// Now that all data interfaces added, create the (trivial) binding mapping. All map to primary binding, index 0.
 		ComputeGraph->DataInterfaceToBinding.SetNumZeroed(ComputeGraph->DataInterfaces.Num());
 
diff --git a/Engine/Plugins/PCG/Source/PCG/Private/PCGSettings.cpp b/Engine/Plugins/PCG/Source/PCG/Private/PCGSettings.cpp
--- a/Engine/Plugins/PCG/Source/PCG/Private/PCGSettings.cpp
+++ b/Engine/Plugins/PCG/Source/PCG/Private/PCGSettings.cpp
@@ -224,6 +224,14 @@ bool UPCGSettings::IsKernelValid(FPCGContext* InContext, bool bQuiet) const
 	return true;
 }
 
+FPCGDataCollectionDesc UPCGSettings::ComputeInputPinDataDesc(const FName& InputPinLabel, const UPCGDataBinding* Binding) const
+{
+	const UPCGNode* Node = Cast<UPCGNode>(GetOuter());
+	const UPCGPin* InputPin = Node ? Node->GetInputPin(InputPinLabel) : nullptr;
+	check(InputPin);
+	return ComputeInputPinDataDesc(InputPin, Binding);
+}
+
 FPCGDataCollectionDesc UPCGSettings::ComputeInputPinDataDesc(const UPCGPin* InputPin, const UPCGDataBinding* Binding) const
 {
 	check(InputPin && Binding);
@@ -265,7 +273,8 @@ FPCGDataCollectionDesc UPCGSettings::ComputeOutputPinDataDesc(const UPCGPin* Out
 		return FPCGDataCollectionDesc::BuildFromInputDataCollectionAndInputPinLabel(
 			Binding->DataForGPU.InputDataCollection,
 			*FoundPinAlias,
-			Binding->Graph->GetAttributeLookupTable());
+			Binding->GetAttributeLookupTable(),
+			Binding->GetStringTable());
 	}
 
 	ensure(false);
diff --git a/Engine/Plugins/PCG/Source/PCG/Public/Compute/PCGDataForGPU.h b/Engine/Plugins/PCG/Source/PCG/Public/Compute/PCGDataForGPU.h
--- a/Engine/Plugins/PCG/Source/PCG/Public/Compute/PCGDataForGPU.h
+++ b/Engine/Plugins/PCG/Source/PCG/Public/Compute/PCGDataForGPU.h
@@ -32,6 +32,7 @@ enum class EPCGKernelAttributeType : uint8
 	Rotator,
 	Quat,
 	Transform,
+	StringKey,
 
 	Invalid = std::numeric_limits<uint8>::max()
 };
@@ -86,17 +87,28 @@ struct FPCGKernelAttributeDesc
 	{
 	}
 
+	explicit FPCGKernelAttributeDesc(int32 InIndex, EPCGKernelAttributeType InType, FName InName, TArray<int32>&& InUniqueStringKeys)
+		: Index(InIndex)
+		, Type(InType)
+		, Name(InName)
+		, UniqueStringKeys(MoveTemp(InUniqueStringKeys))
+	{
+	}
+
 	int32 Index = INDEX_NONE;
 	EPCGKernelAttributeType Type = EPCGKernelAttributeType::Float;
 	FName Name = NAME_None;
 
+	/* All possible string keys arriving on this attribute (string keys are indices into the string table in the data binding). */
+	TArray<int32> UniqueStringKeys;
+
 	bool operator==(const FPCGKernelAttributeDesc& Other) const;
 };
 
 struct FPCGDataDesc
 {
 	FPCGDataDesc(EPCGDataType InType, int32 InElementCount);
-	FPCGDataDesc(const UPCGData* Data, const TMap<FName, FPCGKernelAttributeIDAndType>& GlobalAttributeLookupTable);
+	FPCGDataDesc(const UPCGData* InData, const TMap<FName, FPCGKernelAttributeIDAndType>& InGlobalAttributeLookupTable, const TArray<FString>& InStringTable);
 
 	uint32 ComputePackedSize() const;
 
@@ -105,7 +117,7 @@ struct FPCGDataDesc
 	int32 ElementCount = 0;
 
 private:
-	void InitializeAttributeDescs(const UPCGMetadata* Metadata, const TMap<FName, FPCGKernelAttributeIDAndType>& GlobalAttributeLookupTable = {});
+	void InitializeAttributeDescs(const UPCGData* InData, const TMap<FName, FPCGKernelAttributeIDAndType>& InGlobalAttributeLookupTable, const TArray<FString>& InStringTable);
 };
 
 struct FPCGDataCollectionDesc
@@ -113,19 +125,20 @@ struct FPCGDataCollectionDesc
 	static FPCGDataCollectionDesc BuildFromInputDataCollectionAndInputPinLabel(
 		const FPCGDataCollection& InDataCollection,
 		FName InputPinLabel,
-		const TMap<FName, FPCGKernelAttributeIDAndType>& InAttributeLookupTable);
+		const TMap<FName, FPCGKernelAttributeIDAndType>& InAttributeLookupTable,
+		const TArray<FString>& InStringTable);
 
 	/** Computes the size (in bytes) of the data collection after packing. Also produces the offset (in bytes) for each data in the packed collection. */
 	uint32 ComputePackedSize(TArray<uint32>* OutDataAddresses = nullptr) const;
 
 	/** Pack a data collection into the GPU data format. DataDescs defines which attributes are packed. */
-	void PackDataCollection(const FPCGDataCollection& InDataCollection, FName InPin, TArray<uint32>& OutPackedDataCollection) const;
+	void PackDataCollection(const FPCGDataCollection& InDataCollection, FName InPin, const TArray<FString>& InStringTable, TArray<uint32>& OutPackedDataCollection) const;
 
 	/** Allocates the correct size and sets up header. Initializes data count to 0, which kernel will then overwrite if it executes at least one thread. */
 	void PrepareBufferForKernelOutput(TArray<uint32>& OutPackedDataCollection);
 
 	/** Unpack a buffer of 8-bit uints to a data collection. */
-	EPCGUnpackDataCollectionResult UnpackDataCollection(const TArray<uint8>& InPackedData, FName InPin, FPCGDataCollection& OutDataCollection) const;
+	EPCGUnpackDataCollectionResult UnpackDataCollection(const TArray<uint8>& InPackedData, FName InPin, const TArray<FString>& InStringTable, FPCGDataCollection& OutDataCollection) const;
 
 	/** Compute total number of processing elements of the given type. */
 	uint32 ComputeDataElementCount(EPCGDataType InDataType) const;
@@ -136,13 +149,20 @@ struct FPCGDataCollectionDesc
 	TArray<FPCGDataDesc> DataDescs;
 };
 
+USTRUCT()
 struct FPCGDataForGPU
 {
+	GENERATED_BODY()
+
 	// All the node input pins that have edges that cross from CPU to GPU.
-	TSet<const UPCGPin*> InputPins;
+	UPROPERTY()
+	TSet<TObjectPtr<const UPCGPin>> InputPins;
+
+	UPROPERTY()
 	TMap<TObjectPtr<const UPCGPin>, FName> InputPinLabelAliases;
 
 	// Since the compute graph is collapsed to a single element, all data crossing from CPU to GPU is in a single collection.
+	UPROPERTY()
 	FPCGDataCollection InputDataCollection;
 };
 
diff --git a/Engine/Plugins/PCG/Source/PCG/Public/PCGSettings.h b/Engine/Plugins/PCG/Source/PCG/Public/PCGSettings.h
--- a/Engine/Plugins/PCG/Source/PCG/Public/PCGSettings.h
+++ b/Engine/Plugins/PCG/Source/PCG/Public/PCGSettings.h
@@ -528,6 +528,7 @@ public:
 	virtual int ComputeKernelThreadCount(const UPCGDataBinding* Binding) const { return 0; };
 
 	/** Compute a description of all data arriving on InputPin. */
+	FPCGDataCollectionDesc ComputeInputPinDataDesc(const FName& InputPinLabel, const UPCGDataBinding* Binding) const;
 	FPCGDataCollectionDesc ComputeInputPinDataDesc(const UPCGPin* InputPin, const UPCGDataBinding* Binding) const;
 
 	/** Compute a description of data that will be output from OutputPinLabel/OutputPin. */
@@ -550,15 +551,19 @@ public:
 	bool bExecuteOnGPU = false;
 
 	/** Dump the cooked HLSL into the log after it is generated. */
-	UPROPERTY(EditAnywhere, Category = "GPU", meta = (EditCondition = "bExecuteOnGPU", EditConditionHides))
+	UPROPERTY(EditAnywhere, Category = "GPU", AdvancedDisplay, meta = (EditCondition = "bExecuteOnGPU", EditConditionHides))
 	bool bDumpCookedHLSL = false;
 
+	/** Dump the data descriptions of input/output pins to the log. */
+	UPROPERTY(EditAnywhere, Category = "GPU", AdvancedDisplay, meta = (EditCondition = "bExecuteOnGPU", EditConditionHides))
+	bool bDumpDataDescriptions = false;
+
 	/** Enable use of 'WriteDebugValue(uint Index, float Value)' function in your kernel. Allows you to write float values to a buffer for logging on the CPU. */
-	UPROPERTY(EditAnywhere, Category = "GPU", meta = (EditCondition = "bExecuteOnGPU", EditConditionHides))
+	UPROPERTY(EditAnywhere, Category = "GPU", AdvancedDisplay, meta = (EditCondition = "bExecuteOnGPU", EditConditionHides))
 	bool bPrintShaderDebugValues = false;
 
 	/** Size (in number of floats) of the shader debug print buffer. */
-	UPROPERTY(EditAnywhere, Category = "GPU", meta = (EditCondition="bExecuteOnGPU && bPrintShaderDebugValues", EditConditionHides))
+	UPROPERTY(EditAnywhere, Category = "GPU", AdvancedDisplay, meta = (EditCondition="bExecuteOnGPU && bPrintShaderDebugValues", EditConditionHides))
 	int DebugBufferSize = 16;
 
 private:
