[PCG] Add profiling scopes to non-trivial work in data interfaces to help diagnose performance issues.

#rnx

diff --git a/Engine/Plugins/PCG/Source/PCG/Private/Compute/DataInterfaces/Elements/PCGStaticMeshSpawnerDataInterface.cpp b/Engine/Plugins/PCG/Source/PCG/Private/Compute/DataInterfaces/Elements/PCGStaticMeshSpawnerDataInterface.cpp
--- a/Engine/Plugins/PCG/Source/PCG/Private/Compute/DataInterfaces/Elements/PCGStaticMeshSpawnerDataInterface.cpp
+++ b/Engine/Plugins/PCG/Source/PCG/Private/Compute/DataInterfaces/Elements/PCGStaticMeshSpawnerDataInterface.cpp
@@ -108,6 +108,7 @@ void UPCGStaticMeshSpawnerDataInterface::GetHLSL(FString& OutHLSL, FString const
 
 UComputeDataProvider* UPCGStaticMeshSpawnerDataInterface::CreateDataProvider(TObjectPtr<UObject> InBinding, uint64 InInputMask, uint64 InOutputMask) const
 {
+	TRACE_CPUPROFILER_EVENT_SCOPE(UPCGStaticMeshSpawnerDataInterface::CreateDataProvider);
 	UPCGDataBinding* Binding = CastChecked<UPCGDataBinding>(InBinding);
 	const FPCGSpawnerPrimitives* Primitives = Binding->MeshSpawnersToPrimitives.Find(Settings);
 	if (!ensure(Primitives))
diff --git a/Engine/Plugins/PCG/Source/PCG/Private/Compute/DataInterfaces/PCGCustomKernelDataInterface.cpp b/Engine/Plugins/PCG/Source/PCG/Private/Compute/DataInterfaces/PCGCustomKernelDataInterface.cpp
--- a/Engine/Plugins/PCG/Source/PCG/Private/Compute/DataInterfaces/PCGCustomKernelDataInterface.cpp
+++ b/Engine/Plugins/PCG/Source/PCG/Private/Compute/DataInterfaces/PCGCustomKernelDataInterface.cpp
@@ -69,6 +69,7 @@ void UPCGCustomKernelDataInterface::GetHLSL(FString& OutHLSL, FString const& InD
 
 UComputeDataProvider* UPCGCustomKernelDataInterface::CreateDataProvider(TObjectPtr<UObject> InBinding, uint64 InInputMask, uint64 InOutputMask) const
 {
+	TRACE_CPUPROFILER_EVENT_SCOPE(UPCGCustomKernelDataInterface::CreateDataProvider);
 	check(Settings);
 
 	UPCGDataBinding* Binding = CastChecked<UPCGDataBinding>(InBinding);
diff --git a/Engine/Plugins/PCG/Source/PCG/Private/Compute/DataInterfaces/PCGDataCollectionDataInterface.cpp b/Engine/Plugins/PCG/Source/PCG/Private/Compute/DataInterfaces/PCGDataCollectionDataInterface.cpp
--- a/Engine/Plugins/PCG/Source/PCG/Private/Compute/DataInterfaces/PCGDataCollectionDataInterface.cpp
+++ b/Engine/Plugins/PCG/Source/PCG/Private/Compute/DataInterfaces/PCGDataCollectionDataInterface.cpp
@@ -490,6 +490,7 @@ void UPCGDataCollectionDataInterface::GetHLSL(FString& OutHLSL, FString const& I
 
 UComputeDataProvider* UPCGDataCollectionDataInterface::CreateDataProvider(TObjectPtr<UObject> InBinding, uint64 InInputMask, uint64 InOutputMask) const
 {
+	TRACE_CPUPROFILER_EVENT_SCOPE(UPCGDataCollectionDataInterface::CreateDataProvider);
 	UPCGDataBinding* Binding = CastChecked<UPCGDataBinding>(InBinding);
 
 	UPCGDataCollectionDataProvider* Provider = NewObject<UPCGDataCollectionDataProvider>();
@@ -533,6 +534,7 @@ bool UPCGDataCollectionDataInterface::GetRequiresReadback() const
 
 FComputeDataProviderRenderProxy* UPCGDataCollectionDataProvider::GetRenderProxy()
 {
+	TRACE_CPUPROFILER_EVENT_SCOPE(UPCGDataCollectionDataProvider::GetRenderProxy);
 	FPCGDataCollectionDataProviderProxy* Proxy = new FPCGDataCollectionDataProviderProxy(PinDesc, ReadbackMode);
 
 	if (ReadbackMode != EPCGReadbackMode::None)
@@ -577,6 +579,7 @@ FComputeDataProviderRenderProxy* UPCGDataCollectionDataProvider::GetRenderProxy(
 
 bool UPCGDataCollectionDataProvider::ProcessReadBackData(FPCGComputeGraphContext* InContext)
 {
+	TRACE_CPUPROFILER_EVENT_SCOPE(UPCGDataCollectionDataProvider::ProcessReadBackData);
 	check(InContext);
 	check(ProducerSettings);
 
@@ -691,6 +694,7 @@ void FPCGDataCollectionDataProviderProxy::GatherDispatchData(FDispatchData const
 
 void FPCGDataCollectionDataProviderProxy::AllocateResources(FRDGBuilder& GraphBuilder, FAllocationData const& InAllocationData)
 {
+	TRACE_CPUPROFILER_EVENT_SCOPE(FPCGDataCollectionDataProviderProxy::AllocateResources);
 	check(SizeBytes > 0);
 
 	// Initialize with an empty data collection. The kernel may not run, for example if indirect dispatch args end up being 0. Ensure
diff --git a/Engine/Plugins/PCG/Source/PCG/Private/Compute/DataInterfaces/PCGDataCollectionUploadDataInterface.cpp b/Engine/Plugins/PCG/Source/PCG/Private/Compute/DataInterfaces/PCGDataCollectionUploadDataInterface.cpp
--- a/Engine/Plugins/PCG/Source/PCG/Private/Compute/DataInterfaces/PCGDataCollectionUploadDataInterface.cpp
+++ b/Engine/Plugins/PCG/Source/PCG/Private/Compute/DataInterfaces/PCGDataCollectionUploadDataInterface.cpp
@@ -14,6 +14,7 @@
 
 UComputeDataProvider* UPCGDataCollectionUploadDataInterface::CreateDataProvider(TObjectPtr<UObject> InBinding, uint64 InInputMask, uint64 InOutputMask) const
 {
+	TRACE_CPUPROFILER_EVENT_SCOPE(UPCGDataCollectionUploadDataInterface::CreateDataProvider);
 	UPCGDataBinding* Binding = CastChecked<UPCGDataBinding>(InBinding);
 
 	UPCGDataProviderDataCollectionUpload* Provider = NewObject<UPCGDataProviderDataCollectionUpload>();
@@ -35,6 +36,7 @@ UComputeDataProvider* UPCGDataCollectionUploadDataInterface::CreateDataProvider(
 
 FComputeDataProviderRenderProxy* UPCGDataProviderDataCollectionUpload::GetRenderProxy()
 {
+	TRACE_CPUPROFILER_EVENT_SCOPE(UPCGDataProviderDataCollectionUpload::GetRenderProxy);
 	TArray<uint32> PackedDataCollection;
 
 	// Use any downstream input pin label to grab data from the collection.
@@ -53,6 +55,8 @@ FPCGDataProviderDataCollectionUploadProxy::FPCGDataProviderDataCollectionUploadP
 
 void FPCGDataProviderDataCollectionUploadProxy::AllocateResources(FRDGBuilder& GraphBuilder, FAllocationData const& InAllocationData)
 {
+	TRACE_CPUPROFILER_EVENT_SCOPE(FPCGDataProviderDataCollectionUploadProxy::AllocateResources);
+
 	const FRDGBufferDesc Desc = FRDGBufferDesc::CreateStructuredDesc(sizeof(uint32), PackedDataCollection.Num());
 	Buffer = GraphBuilder.CreateBuffer(Desc, TEXT("PCGDataCollectionUploadBuffer"));
 	BufferUAV = GraphBuilder.CreateUAV(Buffer);
diff --git a/Engine/Plugins/PCG/Source/PCG/Private/Compute/DataInterfaces/PCGInstanceDataInterface.cpp b/Engine/Plugins/PCG/Source/PCG/Private/Compute/DataInterfaces/PCGInstanceDataInterface.cpp
--- a/Engine/Plugins/PCG/Source/PCG/Private/Compute/DataInterfaces/PCGInstanceDataInterface.cpp
+++ b/Engine/Plugins/PCG/Source/PCG/Private/Compute/DataInterfaces/PCGInstanceDataInterface.cpp
@@ -111,6 +111,7 @@ void UPCGInstanceDataInterface::GetHLSL(FString& OutHLSL, FString const& InDataI
 
 UComputeDataProvider* UPCGInstanceDataInterface::CreateDataProvider(TObjectPtr<UObject> InBinding, uint64 InInputMask, uint64 InOutputMask) const
 {
+	TRACE_CPUPROFILER_EVENT_SCOPE(UPCGInstanceDataInterface::CreateDataProvider);
 	UPCGDataBinding* Binding = CastChecked<UPCGDataBinding>(InBinding);
 
 	FPCGSpawnerPrimitives* FoundPrimitives = Binding->MeshSpawnersToPrimitives.Find(ProducerSettings);
@@ -144,6 +145,7 @@ UComputeDataProvider* UPCGInstanceDataInterface::CreateDataProvider(TObjectPtr<U
 
 FComputeDataProviderRenderProxy* UPCGInstanceDataProvider::GetRenderProxy()
 {
+	TRACE_CPUPROFILER_EVENT_SCOPE(UPCGInstanceDataProvider::GetRenderProxy);
 	TArray<FPrimitiveSceneProxy*> PrimitiveProxies;
 	PrimitiveProxies.Reserve(Primitives.Num());
 	
@@ -240,6 +242,8 @@ void FPCGInstanceDataProviderProxy::GatherDispatchData(FDispatchData const& InDi
 
 void FPCGInstanceDataProviderProxy::PostGraphDispatch(FRDGBuilder& GraphBuilder) const
 {
+	TRACE_CPUPROFILER_EVENT_SCOPE(FPCGInstanceDataProviderProxy::PostGraphDispatch);
+
 	if (NumInstancesAllPrimitives == 0)
 	{
 		return;
diff --git a/Engine/Plugins/PCG/Source/PCG/Private/Compute/DataInterfaces/PCGLandscapeDataInterface.cpp b/Engine/Plugins/PCG/Source/PCG/Private/Compute/DataInterfaces/PCGLandscapeDataInterface.cpp
--- a/Engine/Plugins/PCG/Source/PCG/Private/Compute/DataInterfaces/PCGLandscapeDataInterface.cpp
+++ b/Engine/Plugins/PCG/Source/PCG/Private/Compute/DataInterfaces/PCGLandscapeDataInterface.cpp
@@ -254,6 +254,7 @@ void UPCGLandscapeDataInterface::GetHLSL(FString& OutHLSL, FString const& InData
 
 UComputeDataProvider* UPCGLandscapeDataInterface::CreateDataProvider(TObjectPtr<UObject> InBinding, uint64 InInputMask, uint64 InOutputMask) const
 {
+	TRACE_CPUPROFILER_EVENT_SCOPE(UPCGLandscapeDataInterface::CreateDataProvider);
 	UPCGDataBinding* Binding = CastChecked<UPCGDataBinding>(InBinding);
 	check(Binding->SourceComponent.IsValid() && Binding->SourceComponent.Get());
 
@@ -313,6 +314,7 @@ UComputeDataProvider* UPCGLandscapeDataInterface::CreateDataProvider(TObjectPtr<
 
 void UPCGLandscapeDataProvider::Initialize(ALandscape* InLandscape, const FBox& Bounds)
 {
+	TRACE_CPUPROFILER_EVENT_SCOPE(UPCGLandscapeDataProvider::Initialize);
 	check(IsInGameThread());
 
 	const ULandscapeInfo* LandscapeInfo = InLandscape ? InLandscape->GetLandscapeInfo() : nullptr;
diff --git a/Engine/Plugins/PCG/Source/PCG/Private/Compute/DataInterfaces/PCGTextureDataInterface.cpp b/Engine/Plugins/PCG/Source/PCG/Private/Compute/DataInterfaces/PCGTextureDataInterface.cpp
--- a/Engine/Plugins/PCG/Source/PCG/Private/Compute/DataInterfaces/PCGTextureDataInterface.cpp
+++ b/Engine/Plugins/PCG/Source/PCG/Private/Compute/DataInterfaces/PCGTextureDataInterface.cpp
@@ -69,6 +69,7 @@ void UPCGTextureDataInterface::GetHLSL(FString& OutHLSL, FString const& InDataIn
 
 UComputeDataProvider* UPCGTextureDataInterface::CreateDataProvider(TObjectPtr<UObject> InBinding, uint64 InInputMask, uint64 InOutputMask) const
 {
+	TRACE_CPUPROFILER_EVENT_SCOPE(UPCGTextureDataInterface::CreateDataProvider);
 	UPCGDataBinding* Binding = CastChecked<UPCGDataBinding>(InBinding);
 
 	TObjectPtr<UPCGTextureDataProvider> DataProvider = NewObject<UPCGTextureDataProvider>();
diff --git a/Engine/Plugins/PCG/Source/PCG/Private/Compute/Elements/PCGComputeGraphElement.cpp b/Engine/Plugins/PCG/Source/PCG/Private/Compute/Elements/PCGComputeGraphElement.cpp
--- a/Engine/Plugins/PCG/Source/PCG/Private/Compute/Elements/PCGComputeGraphElement.cpp
+++ b/Engine/Plugins/PCG/Source/PCG/Private/Compute/Elements/PCGComputeGraphElement.cpp
@@ -57,6 +57,8 @@ bool FPCGComputeGraphElement::ExecuteInternal(FPCGContext* InContext) const
 	// 7. Execution is complete when any async readbacks are complete.
 	if (Context->bAllAsyncOperationsDone)
 	{
+		TRACE_CPUPROFILER_EVENT_SCOPE(Context->bAllAsyncOperationsDone);
+
 		Context->bExecutionSuccess = true;
 
 		for (UComputeDataProvider* DataProvider : Context->ComputeGraphInstance.GetDataProviders())
@@ -126,7 +128,10 @@ bool FPCGComputeGraphElement::ExecuteInternal(FPCGContext* InContext) const
 
 		const bool bAnyComponentsSetup = SetupProceduralISMComponents(Context, DataBinding);
 
-		Context->ComputeGraphInstance.CreateDataProviders(Graph.Get(), 0, Context->DataBinding.Get());
+		{
+			TRACE_CPUPROFILER_EVENT_SCOPE(ComputeGraphInstance.CreateDataProviders);
+			Context->ComputeGraphInstance.CreateDataProviders(Graph.Get(), 0, Context->DataBinding.Get());
+		}
 
 		// Register all providers running async operations. TODO review if we should have a general API like "RunsAsyncOperations()"?
 		for (UComputeDataProvider* DataProvider : Context->ComputeGraphInstance.GetDataProviders())
@@ -168,6 +173,8 @@ bool FPCGComputeGraphElement::ExecuteInternal(FPCGContext* InContext) const
 	// 1. Prepare render resources. In editor, this will trigger shader compile if not compiled already.
 	if (!Graph->GetRenderProxy())
 	{
+		TRACE_CPUPROFILER_EVENT_SCOPE(Graph->UpdateResources);
+
 		Graph->UpdateResources();
 
 		SleepUntilNextFrame();
@@ -197,12 +204,16 @@ bool FPCGComputeGraphElement::ExecuteInternal(FPCGContext* InContext) const
 	}
 
 	// 5. Enqueue work to be executed when the GPU processes the current frame.
-	Context->bGraphEnqueued = Context->ComputeGraphInstance.EnqueueWork(
-		Graph.Get(),
-		InContext->SourceComponent->GetScene(),
-		ComputeTaskExecutionGroup::EndOfFrameUpdate,
-		InContext->SourceComponent->GetOwner()->GetFName(),
-		FSimpleDelegate());
+	{
+		TRACE_CPUPROFILER_EVENT_SCOPE(ComputeGraphInstance.EnqueueWork);
+
+		Context->bGraphEnqueued = Context->ComputeGraphInstance.EnqueueWork(
+			Graph.Get(),
+			InContext->SourceComponent->GetScene(),
+			ComputeTaskExecutionGroup::EndOfFrameUpdate,
+			InContext->SourceComponent->GetOwner()->GetFName(),
+			FSimpleDelegate());
+	}
 
 	if (ensure(Context->bGraphEnqueued))
 	{
diff --git a/Engine/Plugins/PCG/Source/PCG/Private/Compute/Elements/PCGCustomHLSL.cpp b/Engine/Plugins/PCG/Source/PCG/Private/Compute/Elements/PCGCustomHLSL.cpp
--- a/Engine/Plugins/PCG/Source/PCG/Private/Compute/Elements/PCGCustomHLSL.cpp
+++ b/Engine/Plugins/PCG/Source/PCG/Private/Compute/Elements/PCGCustomHLSL.cpp
@@ -1129,6 +1129,8 @@ TArray<FName> UPCGCustomHLSLSettings::GetInputPinNamesAndNone() const
 
 bool UPCGCustomHLSLSettings::IsKernelValid(FPCGContext* InContext, bool bQuiet) const
 {
+	TRACE_CPUPROFILER_EVENT_SCOPE(UPCGCustomHLSLSettings::IsKernelValid);
+
 	if (OutputPins.IsEmpty())
 	{
 		PCG_KERNEL_VALIDATION_ERR(InContext, this, bQuiet, LOCTEXT("NoOutputs", "Custom HLSL nodes must have at least one output."));
@@ -1388,6 +1390,8 @@ bool UPCGCustomHLSLSettings::IsKernelValid(FPCGContext* InContext, bool bQuiet)
 
 bool UPCGCustomHLSLSettings::AreKernelAttributesValid(FPCGContext* InContext, FText* OutErrorText) const
 {
+	TRACE_CPUPROFILER_EVENT_SCOPE(UPCGCustomHLSLSettings::AreKernelAttributesValid);
+
 	// The context can either be a compute graph element context (if the compute graph was successfully created), otherwise
 	// it will be the original CPU node context. We need the former to run the following validation.
 	if (!InContext || !InContext->IsComputeContext())
diff --git a/Engine/Plugins/PCG/Source/PCG/Private/Compute/PCGDataBinding.cpp b/Engine/Plugins/PCG/Source/PCG/Private/Compute/PCGDataBinding.cpp
--- a/Engine/Plugins/PCG/Source/PCG/Private/Compute/PCGDataBinding.cpp
+++ b/Engine/Plugins/PCG/Source/PCG/Private/Compute/PCGDataBinding.cpp
@@ -16,6 +16,7 @@ void UPCGDataBinding::Initialize(
 	const FPCGDataCollection& InComputeGraphElementInputData,
 	const TMap<FName, FPCGKernelAttributeIDAndType>& InStaticAttributeTable)
 {
+	TRACE_CPUPROFILER_EVENT_SCOPE(UPCGDataBinding::Initialize);
 	check(InComputeGraph);
 
 	Graph = InComputeGraph;
@@ -36,6 +37,8 @@ void UPCGDataBinding::Initialize(
 
 void UPCGDataBinding::InitializeInputData(const FPCGDataCollection& InComputeGraphElementInputData)
 {
+	TRACE_CPUPROFILER_EVENT_SCOPE(UPCGDataBinding::InitializeInputData);
+
 	DataForGPU.InputDataCollection = InComputeGraphElementInputData;
 
 	// Link each input pin to the data collection, so that data providers can find the data.
@@ -54,6 +57,8 @@ void UPCGDataBinding::InitializeInputData(const FPCGDataCollection& InComputeGra
 
 void UPCGDataBinding::BuildStringTable()
 {
+	TRACE_CPUPROFILER_EVENT_SCOPE(UPCGDataBinding::BuildStringTable);
+
 	for (const FPCGTaggedData& Data : DataForGPU.InputDataCollection.TaggedData)
 	{
 		const UPCGMetadata* Metadata = Data.Data ? Data.Data->ConstMetadata() : nullptr;
diff --git a/Engine/Plugins/PCG/Source/PCG/Private/Elements/PCGStaticMeshSpawner.cpp b/Engine/Plugins/PCG/Source/PCG/Private/Elements/PCGStaticMeshSpawner.cpp
--- a/Engine/Plugins/PCG/Source/PCG/Private/Elements/PCGStaticMeshSpawner.cpp
+++ b/Engine/Plugins/PCG/Source/PCG/Private/Elements/PCGStaticMeshSpawner.cpp
@@ -49,6 +49,8 @@ TCHAR const* UPCGStaticMeshSpawnerSettings::TemplateFilePath = TEXT("/Plugin/PCG
 
 bool UPCGStaticMeshSpawnerSettings::IsKernelValid(FPCGContext* InContext, bool bQuiet) const
 {
+	TRACE_CPUPROFILER_EVENT_SCOPE(UPCGStaticMeshSpawnerSettings::IsKernelValid);
+
 	if (!Super::IsKernelValid(InContext, bQuiet))
 	{
 		return false;
