[PCG] Make sure tasks not using the PCGAsync API still get a proper out of tick budget

#rnx

diff --git a/Engine/Plugins/PCG/Source/PCG/Private/Graph/PCGGraphExecutor.cpp b/Engine/Plugins/PCG/Source/PCG/Private/Graph/PCGGraphExecutor.cpp
--- a/Engine/Plugins/PCG/Source/PCG/Private/Graph/PCGGraphExecutor.cpp
+++ b/Engine/Plugins/PCG/Source/PCG/Private/Graph/PCGGraphExecutor.cpp
@@ -15,6 +15,7 @@
 #include "Compute/Elements/PCGComputeGraphElement.h"
 #include "Grid/PCGPartitionActor.h"
 #include "Helpers/PCGActorHelpers.h"
+#include "Helpers/PCGAsync.h"
 #include "Helpers/PCGHelpers.h"
 #include "Metadata/PCGMetadata.h"
 #include "Utils/PCGGraphExecutionLogging.h"
@@ -1536,8 +1537,6 @@ bool FPCGGraphExecutor::ExecuteScheduling(double EndTime, TSharedPtr<FPCGGraphAc
 						ActiveTask->Context->AsyncState.NumAvailableTasks = -1;
 						ActiveTask->Context->AsyncState.EndTime = EndTime;
 						ActiveTask->Context->AsyncState.bIsRunningOnMainThread = false;
-						ActiveTask->Context->AsyncState.bIsRunningOutOfTick = true;
-						ActiveTask->Context->AsyncState.bIsOutOfTickBudgetSet = false;
 
 						// Capture copy of SharedPtr so task stays valid while Async task is alive
 						ActiveTask->StartExecuting();
@@ -1545,6 +1544,10 @@ bool FPCGGraphExecutor::ExecuteScheduling(double EndTime, TSharedPtr<FPCGGraphAc
 						{
 							TRACE_CPUPROFILER_EVENT_SCOPE(FPCGGraphExecutor::ExecuteAsyncTask);
 
+							// Set Out Of tick budget time (minimum 1ms)
+							const float OutOfTickBudgetInSeconds = FMath::Max(1.0f, FPCGAsync::ConsoleVar::CVarAsyncOutOfTickBudgetInMilliseconds.GetValueOnAnyThread()) / 1000.f;
+							ActiveTask->Context->AsyncState.EndTime = FPlatformTime::Seconds() + OutOfTickBudgetInSeconds;
+
 							const bool bIsDone = ActiveTask->bWasCancelled || ActiveTask->Element->Execute(ActiveTask->Context.Get());
 							const bool bIsPaused = ActiveTask->Context->bIsPaused;
 
diff --git a/Engine/Plugins/PCG/Source/PCG/Private/Helpers/PCGAsyncState.cpp b/Engine/Plugins/PCG/Source/PCG/Private/Helpers/PCGAsyncState.cpp
--- a/Engine/Plugins/PCG/Source/PCG/Private/Helpers/PCGAsyncState.cpp
+++ b/Engine/Plugins/PCG/Source/PCG/Private/Helpers/PCGAsyncState.cpp
@@ -18,9 +18,3 @@ void FPCGAsyncState::Reset()
 
 	bStarted = false;
 }
-
-void FPCGAsyncState::SetStarted(double InEndTime)
-{
-	bStarted = true;
-	EndTime = InEndTime;
-}
\ No newline at end of file
diff --git a/Engine/Plugins/PCG/Source/PCG/Public/Helpers/PCGAsync.h b/Engine/Plugins/PCG/Source/PCG/Public/Helpers/PCGAsync.h
--- a/Engine/Plugins/PCG/Source/PCG/Public/Helpers/PCGAsync.h
+++ b/Engine/Plugins/PCG/Source/PCG/Public/Helpers/PCGAsync.h
@@ -75,7 +75,6 @@ namespace FPCGAsync
 		const bool bIsInGameThread = IsInGameThread();
 		const bool bEnableTimeSlicing = bInEnableTimeSlicing && ((bIsInGameThread && !ConsoleVar::CVarDisableAsyncTimeSlicingOnGameThread.GetValueOnAnyThread()) || (!bIsInGameThread && !ConsoleVar::CVarDisableAsyncTimeSlicing.GetValueOnAnyThread()));
 				
-		const float OutOfTickBudgetInSeconds = ConsoleVar::CVarAsyncOutOfTickBudgetInMilliseconds.GetValueOnAnyThread() / 1000.f;
 
 		if (AsyncState.NumAvailableTasks == 0 || ChunkSize <= 0 || NumIterations <= 0)
 		{
@@ -97,17 +96,7 @@ namespace FPCGAsync
 
 			Initialize();
 
-			double EndTime = AsyncState.EndTime;
-			
-			if (AsyncState.bIsRunningOutOfTick && !AsyncState.bIsOutOfTickBudgetSet && bEnableTimeSlicing && OutOfTickBudgetInSeconds > 0.f)
-			{
-				EndTime = FPlatformTime::Seconds() + OutOfTickBudgetInSeconds;
-
-				// Prevent multiple AsyncProcessing calls in the same node execute to reset the budget
-				AsyncState.bIsOutOfTickBudgetSet = true;
-			}
-
-			AsyncState.SetStarted(EndTime);
+			AsyncState.bStarted = true;
 
 			// At the beginning, dispatch the chunks that each task needs to process.
 			// It will be stored on the AsyncState so that it can be reused.
diff --git a/Engine/Plugins/PCG/Source/PCG/Public/Helpers/PCGAsyncState.h b/Engine/Plugins/PCG/Source/PCG/Public/Helpers/PCGAsyncState.h
--- a/Engine/Plugins/PCG/Source/PCG/Public/Helpers/PCGAsyncState.h
+++ b/Engine/Plugins/PCG/Source/PCG/Public/Helpers/PCGAsyncState.h
@@ -64,7 +64,4 @@ struct FPCGAsyncState
 
 	/** To be called at the end of an async processing. */
 	PCG_API void Reset();
-
-	/** To be called to change the state to started and provide an EndTime for timeslicing */
-	PCG_API void SetStarted(double InEndTime);
 };
\ No newline at end of file
