[PCG] Rename Grammar Nodes

Volume Slicer -> Duplicate Cross-Sections
Segment Slicer -> Subdivide Segment
Spline Slicer -> Subdivide Spline

#rnx

diff --git a/Engine/Plugins/PCG/Config/BasePCG.ini b/Engine/Plugins/PCG/Config/BasePCG.ini
--- a/Engine/Plugins/PCG/Config/BasePCG.ini
+++ b/Engine/Plugins/PCG/Config/BasePCG.ini
@@ -100,4 +100,18 @@
 +EnumRedirects=(OldName="EPCGMedadataCompareOperation",NewName="EPCGMetadataCompareOperation")
 +EnumRedirects=(OldName="EPCGMedadataVectorOperation",NewName="EPCGMetadataVectorOperation")
 +EnumRedirects=(OldName="EPCGMedadataBitwiseOperation",NewName="EPCGMetadataBitwiseOperation")
-+EnumRedirects=(OldName="EPCGMedadataRotatorOperation",NewName="EPCGMetadataRotatorOperation")
\ No newline at end of file
++EnumRedirects=(OldName="EPCGMedadataRotatorOperation",NewName="EPCGMetadataRotatorOperation")
+
+// START - TODO: Remove before 5.5
++ClassRedirects=(OldName="/Script/PCG.PCGSplineSlicerSettings",NewName="/Script/PCG.PCGSubdivideSplineSettings")
++ClassRedirects=(OldName="/Script/PCG.PCGSegmentSlicerSettings",NewName="/Script/PCG.PCGSubdivideSegmentSettings")
++ClassRedirects=(OldName="/Script/PCG.PCGSlicingBaseSettings",NewName="/Script/PCG.PCGSubdivisionBaseSettings")
++ClassRedirects=(OldName="/Script/PCG.PCGVolumeSlicerSettings",NewName="/Script/PCG.PCGDuplicateCrossSectionsSettings")
+
++StructRedirects=(OldName="/Script/PCG.PCGSlicingSubmodule",NewName="/Script/PCG.PCGSubdivisionSubmodule")
++StructRedirects=(OldName="/Script/PCG.PCGSlicingModuleAttributeNames",NewName="/Script/PCG.PCGSubdivisionModuleAttributeNames")
+
++PropertyRedirects=(OldName="/Script/PCG.PCGSubdivideSegmentSettings.SlicingAxis",NewName="/Script/PCG.PCGSubdivideSegmentSettings.SubdivisionAxis")
++PropertyRedirects=(OldName="/Script/PCG.PCGSubdivideSegmentSettings.bAcceptIncompleteSlicing",NewName="/Script/PCG.PCGSubdivideSegmentSettings.bAcceptIncompleteSubdivision")
++PropertyRedirects=(OldName="/Script/PCG.PCGSubdivideSplineSettings.bAcceptIncompleteSlicing",NewName="/Script/PCG.PCGSubdivideSplineSettings.bAcceptIncompleteSubdivision")
+// END - TODO : Remove before 5.5
diff --git a/Engine/Plugins/PCG/Source/PCG/Private/Elements/Grammar/PCGDuplicateCrossSections.cpp b/Engine/Plugins/PCG/Source/PCG/Private/Elements/Grammar/PCGDuplicateCrossSections.cpp
new file mode 100644
--- /dev/null
+++ b/Engine/Plugins/PCG/Source/PCG/Private/Elements/Grammar/PCGDuplicateCrossSections.cpp
@@ -0,0 +1,209 @@
+// Copyright Epic Games, Inc. All Rights Reserved.
+
+#include "Elements/Grammar/PCGDuplicateCrossSections.h"
+
+#include "PCGContext.h"
+#include "PCGParamData.h"
+#include "Data/PCGPointData.h"
+#include "Data/PCGSplineData.h"
+#include "Grammar/PCGGrammarParser.h"
+#include "Metadata/Accessors/IPCGAttributeAccessor.h"
+#include "Metadata/Accessors/PCGAttributeAccessorHelpers.h"
+#include "Metadata/Accessors/PCGAttributeAccessorKeys.h"
+
+#define LOCTEXT_NAMESPACE "PCGDuplicateCrossSectionsElement"
+
+#if WITH_EDITOR
+FName UPCGDuplicateCrossSectionsSettings::GetDefaultNodeName() const
+{
+	return FName(TEXT("DuplicateCrossSections"));
+}
+
+FText UPCGDuplicateCrossSectionsSettings::GetDefaultNodeTitle() const
+{
+	return LOCTEXT("NodeTitle", "Duplicate Cross-Sections");
+}
+#endif // WITH_EDITOR
+
+FPCGElementPtr UPCGDuplicateCrossSectionsSettings::CreateElement() const
+{
+	return MakeShared<FPCGDuplicateCrossSectionsElement>();
+}
+
+TArray<FPCGPinProperties> UPCGDuplicateCrossSectionsSettings::InputPinProperties() const
+{
+	TArray<FPCGPinProperties> Result;
+	FPCGPinProperties& InputPin = Result.Emplace_GetRef(PCGPinConstants::DefaultInputLabel, EPCGDataType::Spline);
+	InputPin.SetRequiredPin();
+
+	if (bModuleInfoAsInput)
+	{
+		FPCGPinProperties& ModuleInfoPin = Result.Emplace_GetRef(PCGSubdivisionBase::Constants::ModulesInfoPinLabel, EPCGDataType::Param);
+		ModuleInfoPin.SetRequiredPin();
+	}
+
+	return Result;
+}
+
+TArray<FPCGPinProperties> UPCGDuplicateCrossSectionsSettings::OutputPinProperties() const
+{
+	TArray<FPCGPinProperties> Result;
+	Result.Emplace_GetRef(PCGPinConstants::DefaultOutputLabel, EPCGDataType::Spline);
+
+	return Result;
+}
+
+bool FPCGDuplicateCrossSectionsElement::ExecuteInternal(FPCGContext* InContext) const
+{
+	TRACE_CPUPROFILER_EVENT_SCOPE(FPCGDuplicateCrossSectionsElement::Execute);
+
+	check(InContext);
+
+	const UPCGDuplicateCrossSectionsSettings* Settings = InContext->GetInputSettings<UPCGDuplicateCrossSectionsSettings>();
+	check(Settings);
+
+	const TArray<FPCGTaggedData> Inputs = InContext->InputData.GetInputsByPin(PCGPinConstants::DefaultInputLabel);
+	TArray<FPCGTaggedData>& Outputs = InContext->OutputData.TaggedData;
+
+	const UPCGParamData* ModuleInfoParamData = nullptr;
+	PCGSubdivisionBase::FModuleInfoMap ModulesInfo = GetModulesInfoMap(InContext, Settings, ModuleInfoParamData);
+
+	const bool bMatchAndSetAttributes = Settings->bForwardAttributesFromModulesInfo && ModuleInfoParamData;
+
+	for (const FPCGTaggedData& Input : Inputs)
+	{
+		const UPCGSplineData* InputSplineData = Cast<const UPCGSplineData>(Input.Data);
+
+		if (!InputSplineData)
+		{
+			continue;
+		}
+
+		const TArray<FInterpCurvePointVector>& ControlPoints = InputSplineData->SplineStruct.SplineCurves.Position.Points;
+		if (ControlPoints.Num() < 2)
+		{
+			continue;
+		}
+
+		FVector ExtrudeVector = Settings->ExtrudeVector;
+
+		auto GetValueFromAttribute = [InputSplineData, InContext]<typename T>(const FPCGAttributePropertyInputSelector& InSelector, T& OutValue) -> bool
+		{
+			const FPCGAttributePropertyInputSelector Selector = InSelector.CopyAndFixLast(InputSplineData);
+			const TUniquePtr<const IPCGAttributeAccessor> Accessor = PCGAttributeAccessorHelpers::CreateConstAccessor(InputSplineData, Selector);
+			if (!Accessor)
+			{
+				PCGLog::Metadata::LogFailToCreateAccessorError(Selector, InContext);
+				return false;
+			}
+
+			if (!Accessor->Get(OutValue, FPCGAttributeAccessorKeysEntries(PCGInvalidEntryKey), EPCGAttributeAccessorFlags::AllowBroadcastAndConstructible))
+			{
+				PCGLog::Metadata::LogFailToGetAttributeError<T>(Selector, Accessor.Get(), InContext);
+				return false;
+			}
+
+			return true;
+		};
+
+		if (Settings->bExtrudeVectorAsAttribute && !GetValueFromAttribute(Settings->ExtrudeVectorAttribute, ExtrudeVector))
+		{
+			continue;
+		}
+
+		const double ExtrudeLength = ExtrudeVector.Length();
+		if (FMath::IsNearlyZero(ExtrudeLength))
+		{
+			continue;
+		}
+
+		// Set seed if required
+		int32 AdditionalSeed = 0;
+		if (Settings->bUseSeedAttribute)
+		{
+			const FPCGAttributePropertyInputSelector Selector = Settings->SeedAttribute.CopyAndFixLast(InputSplineData);
+			TUniquePtr<const IPCGAttributeAccessor> SeedAccessor = PCGAttributeAccessorHelpers::CreateConstAccessor(InputSplineData, Selector);
+
+			if (!SeedAccessor)
+			{
+				PCGLog::Metadata::LogFailToCreateAccessorError(Selector, InContext);
+			}
+			// Otherwise, get the value, if it fails, the attribute wasn't compatible
+			else if (!SeedAccessor->Get(AdditionalSeed, FPCGAttributeAccessorKeysEntries(PCGInvalidEntryKey), EPCGAttributeAccessorFlags::AllowBroadcastAndConstructible))
+			{
+				PCGLog::Metadata::LogFailToGetAttributeError<int32>(Selector, SeedAccessor.Get(), InContext);
+			}
+		}
+
+		const FVector ExtrudeDirection = ExtrudeVector / ExtrudeLength;
+
+		double MinSize = 0.0;
+		PCGGrammar::FTokenizedGrammar TokenizedGrammar = GetTokenizedGrammar(InContext, InputSplineData, Settings, ModulesInfo, MinSize);
+
+		if (!TokenizedGrammar.IsValid())
+		{
+			continue;
+		}
+
+		TArray<PCGSubdivisionBase::TModuleInstance<PCGGrammar::FTokenizedModule>> Instances;
+		double RemainingLength = 0.0;
+		const bool bHeightSubdivideSuccess = PCGSubdivisionBase::Subdivide(*TokenizedGrammar.ModuleGrammar, ExtrudeLength, Instances, RemainingLength, InContext, AdditionalSeed);
+
+		if (!bHeightSubdivideSuccess)
+		{
+			PCGLog::LogErrorOnGraph(LOCTEXT("SubdivideFailed", "Grammar doesn't fit."), InContext);
+			continue;
+		}
+
+		FVector CurrentDisplacement = FVector::ZeroVector;
+		int32 SplineIndex = 0;
+
+		for (const PCGSubdivisionBase::TModuleInstance<PCGGrammar::FTokenizedModule>& Instance : Instances)
+		{
+			const FName Symbol = Instance.Module->Descriptor->Symbol;
+			const FPCGSubdivisionSubmodule& CurrentBlock = ModulesInfo[Symbol];
+			const FVector Size = ExtrudeDirection * CurrentBlock.Size * (FVector::OneVector + Instance.ExtraScale);
+			
+			FPCGSplineStruct NewSpline = InputSplineData->SplineStruct;
+			NewSpline.Transform.AddToTranslation(CurrentDisplacement);
+
+			UPCGSplineData* NewSplineData = FPCGContext::NewObject_AnyThread<UPCGSplineData>(InContext);
+			NewSplineData->Initialize(NewSpline);
+			NewSplineData->InitializeFromData(InputSplineData);
+
+			NewSplineData->Metadata->FindOrCreateAttribute<FName>(Settings->SymbolAttributeName, Symbol, false, false);
+
+			if (!bMatchAndSetAttributes && Settings->bOutputDebugColorAttribute)
+			{
+				NewSplineData->Metadata->FindOrCreateAttribute<FVector4>(Settings->DebugColorAttributeName, FVector4(CurrentBlock.DebugColor, 1.0), false, false);
+			}
+
+			if (Settings->bOutputSizeAttribute)
+			{
+				NewSplineData->Metadata->FindOrCreateAttribute<FVector>(Settings->SizeAttributeName, Size, false, false);
+			}
+
+			if (!bMatchAndSetAttributes && Settings->bOutputScalableAttribute)
+			{
+				NewSplineData->Metadata->FindOrCreateAttribute<bool>(Settings->ScalableAttributeName, CurrentBlock.bScalable, false, false);
+			}
+
+			if (Settings->bOutputSplineIndexAttribute)
+			{
+				NewSplineData->Metadata->FindOrCreateAttribute<int32>(Settings->SplineIndexAttributeName, SplineIndex++, false, false);
+			}
+
+			Outputs.Add_GetRef(Input).Data = NewSplineData;
+			CurrentDisplacement += Size;
+		}
+	}
+
+	if (bMatchAndSetAttributes)
+	{
+		MatchAndSetAttributes(Inputs, Outputs, ModuleInfoParamData, Settings);
+	}
+
+	return true;
+}
+
+#undef LOCTEXT_NAMESPACE
diff --git a/Engine/Plugins/PCG/Source/PCG/Private/Elements/Grammar/PCGSegmentSlicer.cpp b/Engine/Plugins/PCG/Source/PCG/Private/Elements/Grammar/PCGSegmentSlicer.cpp
deleted file mode 100644
diff --git a/Engine/Plugins/PCG/Source/PCG/Private/Elements/Grammar/PCGSlicingBase.cpp b/Engine/Plugins/PCG/Source/PCG/Private/Elements/Grammar/PCGSlicingBase.cpp
deleted file mode 100644
diff --git a/Engine/Plugins/PCG/Source/PCG/Private/Elements/Grammar/PCGSplineSlicer.cpp b/Engine/Plugins/PCG/Source/PCG/Private/Elements/Grammar/PCGSplineSlicer.cpp
deleted file mode 100644
diff --git a/Engine/Plugins/PCG/Source/PCG/Private/Elements/Grammar/PCGSubdivideSegment.cpp b/Engine/Plugins/PCG/Source/PCG/Private/Elements/Grammar/PCGSubdivideSegment.cpp
new file mode 100644
--- /dev/null
+++ b/Engine/Plugins/PCG/Source/PCG/Private/Elements/Grammar/PCGSubdivideSegment.cpp
@@ -0,0 +1,398 @@
+// Copyright Epic Games, Inc. All Rights Reserved.
+
+#include "Elements/Grammar/PCGSubdivideSegment.h"
+
+#include "PCGContext.h"
+#include "PCGParamData.h"
+#include "Data/PCGPointData.h"
+#include "Elements/Metadata/PCGMetadataElementCommon.h"
+#include "Grammar/PCGGrammarParser.h"
+#include "Metadata/Accessors/IPCGAttributeAccessor.h"
+#include "Metadata/Accessors/PCGAttributeAccessorHelpers.h"
+#include "Metadata/Accessors/PCGAttributeAccessorKeys.h"
+
+#define LOCTEXT_NAMESPACE "PCGSubdivideSegmentElement"
+
+class PCGSubdivideSegmentHelpers
+{
+public:
+	struct FParameters
+	{
+		FPCGMetadataAttribute<FName>* SymbolAttribute = nullptr;
+		FPCGMetadataAttribute<FVector4>* DebugColorAttribute = nullptr;
+		FPCGMetadataAttribute<int32>* ModuleIndexAttribute = nullptr;
+		FPCGMetadataAttribute<bool>* IsFirstPointAttribute = nullptr;
+		FPCGMetadataAttribute<bool>* IsFinalPointAttribute = nullptr;
+		FPCGMetadataAttribute<int32>* ExtremityNeighborIndexAttribute = nullptr;
+
+		PCGSubdivisionBase::FModuleInfoMap ModulesInfo;
+		TMap<FString, PCGGrammar::FTokenizedGrammar> CachedModules;
+		TArray<int32> CornerIndexes;
+
+		FVector SubdivisionDirection;
+		FVector PerpendicularSubdivisionDirection;
+		int32 AdditionalSeed = 0;
+
+		const UPCGSubdivideSegmentSettings* Settings = nullptr;
+		FPCGContext* Context = nullptr;
+		const TArray<FPCGPoint>* InPoints = nullptr;
+		TArray<FPCGPoint>* OutPoints = nullptr;
+		UPCGMetadata* OutputMetadata = nullptr;
+	};
+
+	static void Process(FParameters& InOutParameters, const FString& InGrammar, const bool bFlipAxis, int32 Index)
+	{
+		if (!InOutParameters.CachedModules.Contains(InGrammar))
+		{
+			double MinSize;
+			InOutParameters.CachedModules.Emplace(InGrammar, PCGSubdivisionBase::GetTokenizedGrammar(InOutParameters.Context, InGrammar, InOutParameters.ModulesInfo, MinSize));
+		}
+
+		const PCGGrammar::FTokenizedGrammar& CurrentTokenizedGrammar = InOutParameters.CachedModules[InGrammar];
+
+		if(!CurrentTokenizedGrammar.IsValid())
+		{
+			return;
+		}
+
+		const FPCGPoint& Point = (*InOutParameters.InPoints)[Index];
+		FTransform TransformNoTranslation = Point.Transform;
+		TransformNoTranslation.SetLocation(FVector::ZeroVector);
+
+		const int32 FirstModuleIndex = InOutParameters.OutPoints->Num();
+
+		FBox Segment = Point.GetLocalBounds();
+		FVector PointScaledSize = Point.GetScaledLocalSize();
+		if (bFlipAxis)
+		{
+			// Swap coordinates on the subdivision direction
+			const FVector PreviousMin = Segment.Min;
+			Segment.Min = Segment.Min * InOutParameters.PerpendicularSubdivisionDirection + Segment.Max * InOutParameters.SubdivisionDirection;
+			Segment.Max = Segment.Max * InOutParameters.PerpendicularSubdivisionDirection + PreviousMin * InOutParameters.SubdivisionDirection;
+			PointScaledSize *= (InOutParameters.PerpendicularSubdivisionDirection - InOutParameters.SubdivisionDirection);
+		}
+
+		const FVector Direction = TransformNoTranslation.TransformVectorNoScale(InOutParameters.SubdivisionDirection).GetSafeNormal();
+		const FVector OtherDirection = TransformNoTranslation.TransformVectorNoScale(PointScaledSize * InOutParameters.PerpendicularSubdivisionDirection) * 0.5;
+		const FVector HalfExtents2D = PointScaledSize * InOutParameters.PerpendicularSubdivisionDirection * 0.5;
+		const double Size = PointScaledSize.Dot(InOutParameters.SubdivisionDirection);
+
+		TArray<PCGSubdivisionBase::TModuleInstance<PCGGrammar::FTokenizedModule>> ModulesInstances;
+		double RemainingSubdivide;
+		const bool bSubdivideSuccess = PCGSubdivisionBase::Subdivide(*CurrentTokenizedGrammar.ModuleGrammar, Size, ModulesInstances, RemainingSubdivide, InOutParameters.Context, InOutParameters.AdditionalSeed);
+
+		if (!bSubdivideSuccess)
+		{
+			return;
+		}
+
+		if (!InOutParameters.Settings->bAcceptIncompleteSubdivision && !FMath::IsNearlyZero(RemainingSubdivide))
+		{
+			PCGLog::LogWarningOnGraph(LOCTEXT("FailSubdivisionFullLength", "One segment has an incomplete subdivision (grammar doesn't fit the whole segment)."), InOutParameters.Context);
+			return;
+		}
+
+		// Now we have our segment subdivided, create the final points
+		const FVector InitialPos = Point.Transform.TransformPosition(Segment.Min);
+		FVector CurrentPos = InitialPos;
+		int32 ModuleIndex = 0;
+
+		for (int32 ModuleInstanceIndex = 0; ModuleInstanceIndex < ModulesInstances.Num(); ModuleInstanceIndex++)
+		{
+			const PCGSubdivisionBase::TModuleInstance<PCGGrammar::FTokenizedModule>& ModuleInstance = ModulesInstances[ModuleInstanceIndex];
+
+			const FName Symbol = ModuleInstance.Module->Descriptor->Symbol;
+			const FVector Scale = FVector::OneVector + (InOutParameters.SubdivisionDirection * ModuleInstance.ExtraScale);
+			const FPCGSubdivisionSubmodule& SubdivisionSubmodule = InOutParameters.ModulesInfo[Symbol];
+			const bool bIsFirstModule = (ModuleInstanceIndex == 0);
+			const bool bIsFinalModule = (ModuleInstanceIndex == ModulesInstances.Num() - 1);
+
+			const double HalfDisplacement = SubdivisionSubmodule.Size * 0.5;
+			const double HalfScaledDisplacement = Scale.Dot(InOutParameters.SubdivisionDirection) * HalfDisplacement;
+
+			const FVector LocalBoundsExtents = InOutParameters.SubdivisionDirection * HalfDisplacement + HalfExtents2D;
+			const FVector HalfStep = HalfScaledDisplacement * Direction;
+			const FVector Position = CurrentPos + HalfStep;
+			CurrentPos = Position + HalfStep;
+
+			FPCGPoint& OutPoint = InOutParameters.OutPoints->Emplace_GetRef(Point);
+			OutPoint.Transform = FTransform(Point.Transform.GetRotation(), Position + OtherDirection, Scale);
+			OutPoint.SetLocalBounds(FBox(-LocalBoundsExtents, LocalBoundsExtents));
+			InOutParameters.OutputMetadata->InitializeOnSet(OutPoint.MetadataEntry);
+			if (InOutParameters.SymbolAttribute)
+			{
+				InOutParameters.SymbolAttribute->SetValue(OutPoint.MetadataEntry, Symbol);
+			}
+
+			if (InOutParameters.DebugColorAttribute)
+			{
+				InOutParameters.DebugColorAttribute->SetValue(OutPoint.MetadataEntry, FVector4(SubdivisionSubmodule.DebugColor, 1.0));
+			}
+
+			if (InOutParameters.ModuleIndexAttribute)
+			{
+				InOutParameters.ModuleIndexAttribute->SetValue(OutPoint.MetadataEntry, ModuleIndex++);
+			}
+
+			if (bIsFirstModule && InOutParameters.IsFirstPointAttribute)
+			{
+				InOutParameters.IsFirstPointAttribute->SetValue(OutPoint.MetadataEntry, true);
+			}
+
+			if (bIsFinalModule && InOutParameters.IsFinalPointAttribute)
+			{
+				InOutParameters.IsFinalPointAttribute->SetValue(OutPoint.MetadataEntry, true);
+			}
+		}
+
+		if (bFlipAxis)
+		{
+			InOutParameters.CornerIndexes.Add(InOutParameters.OutPoints->Num() - 1);
+			InOutParameters.CornerIndexes.Add(FirstModuleIndex);
+		}
+		else
+		{
+			InOutParameters.CornerIndexes.Add(FirstModuleIndex);
+			InOutParameters.CornerIndexes.Add(InOutParameters.OutPoints->Num() - 1);
+		}
+	}
+};
+
+#if WITH_EDITOR
+FName UPCGSubdivideSegmentSettings::GetDefaultNodeName() const
+{
+	return FName(TEXT("SubdivideSegment"));
+}
+
+FText UPCGSubdivideSegmentSettings::GetDefaultNodeTitle() const
+{
+	return LOCTEXT("NodeTitle", "Subdivide Segment");
+}
+#endif // WITH_EDITOR
+
+FPCGElementPtr UPCGSubdivideSegmentSettings::CreateElement() const
+{
+	return MakeShared<FPCGSegmentSubdivisionElement>();
+}
+
+TArray<FPCGPinProperties> UPCGSubdivideSegmentSettings::InputPinProperties() const
+{
+	TArray<FPCGPinProperties> Result;
+	FPCGPinProperties& InputPin = Result.Emplace_GetRef(PCGPinConstants::DefaultInputLabel, EPCGDataType::Point);
+	InputPin.SetRequiredPin();
+
+	if (bModuleInfoAsInput)
+	{
+		FPCGPinProperties& ModuleInfoPin = Result.Emplace_GetRef(PCGSubdivisionBase::Constants::ModulesInfoPinLabel, EPCGDataType::Param);
+		ModuleInfoPin.SetRequiredPin();
+	}
+
+	return Result;
+}
+
+TArray<FPCGPinProperties> UPCGSubdivideSegmentSettings::OutputPinProperties() const
+{
+	return Super::DefaultPointOutputPinProperties();
+}
+
+bool FPCGSegmentSubdivisionElement::ExecuteInternal(FPCGContext* InContext) const
+{
+	TRACE_CPUPROFILER_EVENT_SCOPE(FPCGSubdivideSegmentElement::Execute);
+
+	check(InContext);
+
+	const UPCGSubdivideSegmentSettings* Settings = InContext->GetInputSettings<UPCGSubdivideSegmentSettings>();
+	check(Settings);
+
+	const TArray<FPCGTaggedData> Inputs = InContext->InputData.GetInputsByPin(PCGPinConstants::DefaultInputLabel);
+	TArray<FPCGTaggedData>& Outputs = InContext->OutputData.TaggedData;
+
+	PCGSubdivideSegmentHelpers::FParameters Parameters{};
+
+	Parameters.SubdivisionDirection = FVector::ZeroVector;
+
+	switch (Settings->SubdivisionAxis)
+	{
+	case EPCGSplitAxis::X:
+		Parameters.SubdivisionDirection = FVector::XAxisVector;
+		break;
+	case EPCGSplitAxis::Y:
+		Parameters.SubdivisionDirection = FVector::YAxisVector;
+		break;
+	case EPCGSplitAxis::Z:
+		Parameters.SubdivisionDirection = FVector::ZAxisVector;
+		break;
+	default:
+		PCGLog::LogErrorOnGraph(LOCTEXT("InvalidAxis", "Invalid Subdivision Axis enum value."), InContext);
+		return true;
+	}
+
+	Parameters.PerpendicularSubdivisionDirection = FVector::OneVector - Parameters.SubdivisionDirection;
+
+	const UPCGParamData* ModuleInfoParamData = nullptr;
+	Parameters.ModulesInfo = GetModulesInfoMap(InContext, Settings, ModuleInfoParamData);
+
+	for (const FPCGTaggedData& Input : Inputs)
+	{
+		const UPCGPointData* InputPointData = Cast<const UPCGPointData>(Input.Data);
+
+		if (!InputPointData)
+		{
+			continue;
+		}
+
+		TUniquePtr<const IPCGAttributeAccessor> GrammarAccessor;
+		TUniquePtr<const IPCGAttributeAccessor> FlipAxisAccessor;
+		TUniquePtr<const IPCGAttributeAccessorKeys> Keys;
+
+		if (Settings->GrammarSelection.bGrammarAsAttribute)
+		{
+			const FPCGAttributePropertyInputSelector Selector = Settings->GrammarSelection.GrammarAttribute.CopyAndFixLast(InputPointData);
+			GrammarAccessor = PCGAttributeAccessorHelpers::CreateConstAccessor(InputPointData, Selector);
+			Keys = PCGAttributeAccessorHelpers::CreateConstKeys(InputPointData, Selector);
+			if (!GrammarAccessor || !Keys)
+			{
+				PCGLog::Metadata::LogFailToCreateAccessorError(Selector, InContext);
+				continue;
+			}
+		}
+
+		if (Settings->bFlipAxisAsAttribute)
+		{
+			const FPCGAttributePropertyInputSelector Selector = Settings->FlipAxisAttribute.CopyAndFixLast(InputPointData);
+			FlipAxisAccessor = PCGAttributeAccessorHelpers::CreateConstAccessor(InputPointData, Selector);
+			if (!Keys)
+			{
+				Keys = PCGAttributeAccessorHelpers::CreateConstKeys(InputPointData, Selector);
+			}
+
+			if (!FlipAxisAccessor || !Keys)
+			{
+				PCGLog::Metadata::LogFailToCreateAccessorError(Selector, InContext);
+				continue;
+			}
+		}
+
+		UPCGPointData* OutputPointData = FPCGContext::NewObject_AnyThread<UPCGPointData>(InContext);
+		OutputPointData->InitializeFromData(InputPointData);
+
+		Parameters.Settings = Settings;
+		Parameters.Context = InContext;
+		Parameters.InPoints = &InputPointData->GetPoints();
+		Parameters.OutPoints = &OutputPointData->GetMutablePoints();
+		Parameters.OutputMetadata = OutputPointData->Metadata;
+
+		auto CreateAndValidateAttribute = [InContext, OutputMetadata = OutputPointData->Metadata]<typename T>(const FName AttributeName, const T DefaultValue, const bool bShouldCreate, FPCGMetadataAttribute<T>*& OutAttribute) -> bool
+		{
+			if (bShouldCreate)
+			{
+				OutAttribute = OutputMetadata->FindOrCreateAttribute<T>(AttributeName, DefaultValue, false, true);
+				if (!OutAttribute)
+				{
+					PCGLog::Metadata::LogFailToCreateAttributeError<T>(AttributeName, InContext);
+					return false;
+				}
+			}
+
+			return true;
+		};
+
+		if (!CreateAndValidateAttribute(Settings->SymbolAttributeName, FName(NAME_None), true, Parameters.SymbolAttribute)
+			|| !CreateAndValidateAttribute(Settings->DebugColorAttributeName, FVector4::Zero(), Settings->bOutputDebugColorAttribute, Parameters.DebugColorAttribute)
+			|| !CreateAndValidateAttribute(Settings->ModuleIndexAttributeName, -1, Settings->bOutputModuleIndexAttribute, Parameters.ModuleIndexAttribute)
+			|| !CreateAndValidateAttribute(Settings->IsFirstAttributeName, false, Settings->bOutputExtremityAttributes, Parameters.IsFirstPointAttribute)
+			|| !CreateAndValidateAttribute(Settings->IsFinalAttributeName, false, Settings->bOutputExtremityAttributes, Parameters.IsFinalPointAttribute)
+			|| !CreateAndValidateAttribute(Settings->ExtremityNeighborIndexAttributeName, -1, Settings->bOutputExtremityNeighborIndexAttribute, Parameters.ExtremityNeighborIndexAttribute))
+		{
+			continue;
+		}
+
+		// Set seed if required
+		Parameters.AdditionalSeed = 0;
+		if (Settings->bUseSeedAttribute)
+		{
+			const FPCGAttributePropertyInputSelector Selector = Settings->SeedAttribute.CopyAndFixLast(InputPointData);
+			TUniquePtr<const IPCGAttributeAccessor> SeedAccessor = PCGAttributeAccessorHelpers::CreateConstAccessor(InputPointData, Selector);
+			if (!Keys)
+			{
+				Keys = PCGAttributeAccessorHelpers::CreateConstKeys(InputPointData, Selector);
+			}
+
+			if (!SeedAccessor || !Keys)
+			{
+				PCGLog::Metadata::LogFailToCreateAccessorError(Selector, InContext);
+			}
+			// Otherwise, get the value, if it fails, the attribute wasn't compatible
+			else if (!SeedAccessor->Get(Parameters.AdditionalSeed, FPCGAttributeAccessorKeysEntries(PCGInvalidEntryKey), EPCGAttributeAccessorFlags::AllowBroadcastAndConstructible))
+			{
+				PCGLog::Metadata::LogFailToGetAttributeError<int32>(Selector, SeedAccessor.Get(), InContext);
+			}
+		}
+
+		if (Settings->GrammarSelection.bGrammarAsAttribute && Settings->bFlipAxisAsAttribute)
+		{
+			auto Process = [&Parameters](const FString& InGrammar, const bool bFlipAxis, int32 Index) -> void
+			{
+				PCGSubdivideSegmentHelpers::Process(Parameters, InGrammar, bFlipAxis, Index);
+			};
+
+			PCGMetadataElementCommon::ApplyOnMultiAccessors<FString, bool>(*Keys, { GrammarAccessor.Get(), FlipAxisAccessor.Get() }, Process);
+		}
+		else if (Settings->GrammarSelection.bGrammarAsAttribute)
+		{
+			auto Process = [&Parameters, bShouldFlipAxis = Settings->bShouldFlipAxis](const FString& InGrammar, int32 Index) -> void
+			{
+				PCGSubdivideSegmentHelpers::Process(Parameters, InGrammar, bShouldFlipAxis, Index);
+			};
+
+			PCGMetadataElementCommon::ApplyOnAccessor<FString>(*Keys, *GrammarAccessor, Process);
+		}
+		else if (Settings->bFlipAxisAsAttribute)
+		{
+			auto Process = [&Parameters, Grammar = Settings->GrammarSelection.GrammarString](const bool bFlipAxis, int32 Index) -> void
+			{
+				PCGSubdivideSegmentHelpers::Process(Parameters, Grammar, bFlipAxis, Index);
+			};
+
+			PCGMetadataElementCommon::ApplyOnAccessor<bool>(*Keys, *FlipAxisAccessor, Process);
+		}
+		else
+		{
+			for (int32 SegmentIndex = 0; SegmentIndex < Parameters.InPoints->Num(); ++SegmentIndex)
+			{
+				PCGSubdivideSegmentHelpers::Process(Parameters, Settings->GrammarSelection.GrammarString, Settings->bShouldFlipAxis, SegmentIndex);
+			}
+		}
+
+		if (!Parameters.OutPoints->IsEmpty())
+		{
+			// Set the extremity neighbor indexes
+			if (Parameters.ExtremityNeighborIndexAttribute)
+			{
+				// Go 2 by 2
+				check(Parameters.CornerIndexes.Num() % 2 == 0);
+
+				for (int32 j = 0; j < Parameters.CornerIndexes.Num(); j += 2)
+				{
+					const int32 PreviousModuleIndex = Parameters.CornerIndexes[j == 0 ? Parameters.CornerIndexes.Num() - 1 : j - 1];
+					const int32 NextModuleIndex = Parameters.CornerIndexes[j == Parameters.CornerIndexes.Num() - 2 ? 0 : j + 2];
+
+					Parameters.ExtremityNeighborIndexAttribute->SetValue((*Parameters.OutPoints)[Parameters.CornerIndexes[j]].MetadataEntry, PreviousModuleIndex);
+					Parameters.ExtremityNeighborIndexAttribute->SetValue((*Parameters.OutPoints)[Parameters.CornerIndexes[j+1]].MetadataEntry, NextModuleIndex);
+				}
+			}
+
+			FPCGTaggedData& Output = Outputs.Emplace_GetRef(Input);
+			Output.Data = OutputPointData;
+		}
+	}
+
+	if (Settings->bForwardAttributesFromModulesInfo && ModuleInfoParamData)
+	{
+		MatchAndSetAttributes(Inputs, Outputs, ModuleInfoParamData, Settings);
+	}
+
+	return true;
+}
+
+#undef LOCTEXT_NAMESPACE
diff --git a/Engine/Plugins/PCG/Source/PCG/Private/Elements/Grammar/PCGSubdivideSpline.cpp b/Engine/Plugins/PCG/Source/PCG/Private/Elements/Grammar/PCGSubdivideSpline.cpp
new file mode 100644
--- /dev/null
+++ b/Engine/Plugins/PCG/Source/PCG/Private/Elements/Grammar/PCGSubdivideSpline.cpp
@@ -0,0 +1,449 @@
+// Copyright Epic Games, Inc. All Rights Reserved.
+
+#include "Elements/Grammar/PCGSubdivideSpline.h"
+
+#include "PCGContext.h"
+#include "Data/PCGPointData.h"
+#include "Data/PCGSplineData.h"
+#include "Elements/Metadata/PCGMetadataElementCommon.h"
+#include "Helpers/PCGHelpers.h"
+#include "Kismet/KismetMathLibrary.h"
+#include "Metadata/Accessors/PCGAttributeAccessorHelpers.h"
+
+#define LOCTEXT_NAMESPACE "PCGSubdivideSplineElement"
+
+namespace PCGSubdivideSplineHelpers
+{
+	struct FParameters
+	{
+		FPCGContext* Context = nullptr;
+
+		PCGSubdivisionBase::FModuleInfoMap ModulesInfo;
+		TMap<FString, PCGGrammar::FTokenizedGrammar> CachedModules;
+		bool bAcceptIncompleteSubdivision = false;
+		double ModuleHeight = 0.0;
+		int32 AdditionalSeed = 0;
+
+		const UPCGPolyLineData* PolyLineData = nullptr;
+		TArray<FPCGPoint>* OutPoints = nullptr;
+		UPCGMetadata* OutputMetadata = nullptr;
+
+		FPCGMetadataAttribute<FName>* SymbolAttribute = nullptr;
+		FPCGMetadataAttribute<FVector4>* DebugColorAttribute = nullptr;
+		FPCGMetadataAttribute<int32>* ModuleIndexAttribute = nullptr;
+		FPCGMetadataAttribute<bool>* IsFirstPointAttribute = nullptr;
+		FPCGMetadataAttribute<bool>* IsFinalPointAttribute = nullptr;
+	};
+
+	double GetLinearDistanceBetweenPolyLineAlphas(const UPCGPolyLineData* PolyLineData, double FirstAlpha, double SecondAlpha)
+	{
+		check(PolyLineData);
+
+		FirstAlpha = FMath::Clamp(FirstAlpha, 0, 1);
+		SecondAlpha = FMath::Clamp(SecondAlpha, 0, 1);
+		const FVector FirstPoint = PolyLineData->GetLocationAtAlpha(FirstAlpha);
+		const FVector SecondPoint = PolyLineData->GetLocationAtAlpha(SecondAlpha);
+		return FVector::Distance(FirstPoint, SecondPoint);
+	}
+
+	// Using a numerical method (s.a. bisection), determine the spline alpha from a segment length. Bisection can be slower, but guaranteed to converge.
+	static double FindRootAtLinearDistance_Bisection(const UPCGPolyLineData* PolyLineData, const double SegmentLength, const double StartingAlpha, const double Tolerance = 1)
+	{
+		check(PolyLineData);
+
+		// Bisect a number of times, before taking an estimate
+		static constexpr uint16 BisectionCountLimit = 64u;
+
+		double Low = StartingAlpha;
+		double High = 1.0;
+		double Estimate = 0.0;
+
+		for (uint16 BisectionCount = 0u; BisectionCount < BisectionCountLimit; ++BisectionCount)
+		{
+			Estimate = (Low + High) * 0.5;
+			// TODO: If performance is an issue, this can deal with Dist^2 instead to save the sqrt. But, it will skew the tolerance.
+			const double LinearDistance = GetLinearDistanceBetweenPolyLineAlphas(PolyLineData, StartingAlpha, Estimate);
+			check(LinearDistance >= 0);
+
+			if (FMath::IsNearlyEqual(LinearDistance, SegmentLength, Tolerance))
+			{
+				return Estimate;
+			}
+
+			if (LinearDistance < SegmentLength)
+			{
+				Low = Estimate;
+			}
+			else
+			{
+				High = Estimate;
+			}
+		}
+
+		// Couldn't find the root yet, so estimate
+		return (Low + High) * 0.5;
+	}
+
+	static void Process(FParameters& InOutParameters, const FString& InGrammar)
+	{
+		if (!InOutParameters.CachedModules.Contains(InGrammar))
+		{
+			double MinSize;
+			InOutParameters.CachedModules.Emplace(InGrammar, PCGSubdivisionBase::GetTokenizedGrammar(InOutParameters.Context, InGrammar, InOutParameters.ModulesInfo, MinSize));
+		}
+
+		const PCGGrammar::FTokenizedGrammar& CurrentTokenizedGrammar = InOutParameters.CachedModules[InGrammar];
+
+		if (!CurrentTokenizedGrammar.IsValid())
+		{
+			return;
+		}
+
+		const double SplineLength = InOutParameters.PolyLineData->GetLength();
+
+		/* Implementation Note: Subdivided spline length will always be equal or greater than discretized linear length, depending on the curvature of the spline.
+		 * For extremely long or curvy splines, this can result in the tokenized grammar being cut short. Alternative subdivision solutions may need to be explored.
+		 */
+		TArray<PCGSubdivisionBase::TModuleInstance<PCGGrammar::FTokenizedModule>> ModulesInstances;
+		double RemainingSubdivide;
+		if (!Subdivide(*CurrentTokenizedGrammar.ModuleGrammar, SplineLength, ModulesInstances, RemainingSubdivide, InOutParameters.Context, InOutParameters.AdditionalSeed))
+		{
+			return;
+		}
+
+		if (!InOutParameters.bAcceptIncompleteSubdivision && !FMath::IsNearlyZero(RemainingSubdivide))
+		{
+			PCGLog::LogWarningOnGraph(LOCTEXT("FailSubdivisionFullLength", "The spline has an incomplete subdivision (grammar doesn't fit the whole segment)."), InOutParameters.Context);
+			return;
+		}
+
+		// Since we're essentially mapping spline-space modules onto linear-space modules (they aren't deformed),
+		// Then we need to refit potentially, knowing that having scaled up some modules might make the total length too large.
+		// We'll do a few passes to converge to a better fit.
+		TArray<double> ModuleAlphas;
+		int LengthConvergencePass = 0;
+		const int MaxLengthConvergencePass = 10;
+
+		constexpr double OvershootToleranceFactor = 0.02; // 2%
+		constexpr double OvershootToleranceMax = 5.0; // 5 cm
+		double MinimumModuleSize = std::numeric_limits<double>::max();
+		for (const PCGSubdivisionBase::TModuleInstance<PCGGrammar::FTokenizedModule>& ModuleInstance : ModulesInstances)
+		{
+			const FName Symbol = ModuleInstance.Module->Descriptor->Symbol;
+			const FPCGSubdivisionSubmodule& SubdivisionSubmodule = InOutParameters.ModulesInfo[Symbol];
+			MinimumModuleSize = FMath::Min(MinimumModuleSize, SubdivisionSubmodule.Size);
+		}
+
+		const double OvershootTolerance = FMath::Min(MinimumModuleSize * OvershootToleranceFactor, OvershootToleranceMax);
+
+		for(;LengthConvergencePass < MaxLengthConvergencePass; ++LengthConvergencePass)
+		{
+			ModuleAlphas.Reset(ModulesInstances.Num() + 1);
+			ModuleAlphas.Emplace(0.0);
+
+			double ScaledLength = 0.0;
+			double OvershootDistance = 0.0;
+			bool bAtSplineEnd = false;
+
+			for (const PCGSubdivisionBase::TModuleInstance<PCGGrammar::FTokenizedModule>& ModuleInstance : ModulesInstances)
+			{
+				const FName Symbol = ModuleInstance.Module->Descriptor->Symbol;
+				const FPCGSubdivisionSubmodule& SubdivisionSubmodule = InOutParameters.ModulesInfo[Symbol];
+				const double SubmoduleSize = SubdivisionSubmodule.Size * (1.0 + ModuleInstance.ExtraScale);
+
+				// TODO: modules that can be deformed can also have extra scale, but the move in spline space (so their real length needs to be measured)
+				if (ModuleInstance.ExtraScale > 0)
+				{
+					ScaledLength += SubdivisionSubmodule.Size * ModuleInstance.ExtraScale;
+				}
+
+				if (!bAtSplineEnd)
+				{
+					const double PreviousAlpha = ModuleAlphas.Last();
+					const double CurrentAlpha = FindRootAtLinearDistance_Bisection(InOutParameters.PolyLineData, SubmoduleSize, PreviousAlpha);
+
+					ModuleAlphas.Add(CurrentAlpha);
+
+					if (FMath::IsNearlyEqual(CurrentAlpha, 1.0))
+					{
+						bAtSplineEnd = true;
+
+						// In this case, we know that the end of the module might be really after the end of the spline, so we need to compare the distance and add it to the overshoot.
+						const FVector SubmoduleStartPoint = InOutParameters.PolyLineData->GetLocationAtAlpha(ModuleAlphas.Last(1));
+						const FVector SplineEndPoint = InOutParameters.PolyLineData->GetLocationAtAlpha(ModuleAlphas.Last());
+
+						const double LastSegmentOvershoot = SubmoduleSize - (SplineEndPoint - SubmoduleStartPoint).Length();
+						if (LastSegmentOvershoot > 0)
+						{
+							OvershootDistance += LastSegmentOvershoot;
+						}
+					}
+				}
+				else
+				{
+					OvershootDistance += SubmoduleSize;
+				}
+			}
+
+			// Evaluate if we need another pass
+			if (LengthConvergencePass < MaxLengthConvergencePass && OvershootDistance > OvershootTolerance && ScaledLength > 0)
+			{
+				double UpdateExtraScaleFactor = 1.0;
+
+				// If the leeway we currently have (scaled length) is smaller than the overshoot distance, we are unlikely to converge, and in this instance it's probably best to set the scale to zero across the board.
+				if (OvershootDistance > ScaledLength)
+				{
+					UpdateExtraScaleFactor = 0.0;
+				}
+				// Otherwise, we want to move towards a better solution, but not too quickly lest we undershoot (which we don't detect here).
+				else
+				{
+					UpdateExtraScaleFactor = 0.25 * (3.0 + ((ScaledLength - OvershootDistance) / ScaledLength));
+					check(UpdateExtraScaleFactor >= 0.0 && UpdateExtraScaleFactor <= 1.0);
+				}
+
+				for (PCGSubdivisionBase::TModuleInstance<PCGGrammar::FTokenizedModule>& ModuleInstance : ModulesInstances)
+				{
+					if (ModuleInstance.ExtraScale > 0)
+					{
+						ModuleInstance.ExtraScale *= UpdateExtraScaleFactor;
+					}
+				}
+			}
+			else
+			{
+				// TODO: possibly log an error here, but it's subtle to get to a place where it's not too much
+				break;
+			}
+		}
+
+		const bool bHasMetadata = InOutParameters.SymbolAttribute || InOutParameters.DebugColorAttribute || InOutParameters.ModuleIndexAttribute || InOutParameters.IsFirstPointAttribute || InOutParameters.IsFinalPointAttribute;
+
+		FVector PreviousSegmentEndPoint = (ModuleAlphas.IsEmpty() ? FVector::Zero() : InOutParameters.PolyLineData->GetLocationAtAlpha(ModuleAlphas[0]));
+		const int32 NumIterations = FMath::Min(ModulesInstances.Num(), ModuleAlphas.Num() - 1);
+		for (int32 ModuleInstanceIndex = 0; ModuleInstanceIndex < NumIterations; ModuleInstanceIndex++)
+		{
+			const PCGSubdivisionBase::TModuleInstance<PCGGrammar::FTokenizedModule>& ModuleInstance = ModulesInstances[ModuleInstanceIndex];
+			const FName Symbol = ModuleInstance.Module->Descriptor->Symbol;
+
+			const FPCGSubdivisionSubmodule& SubdivisionSubmodule = InOutParameters.ModulesInfo[Symbol];
+
+			const double& PreviousAlpha = ModuleAlphas[ModuleInstanceIndex];
+			const double& SplineAlpha = ModuleAlphas[ModuleInstanceIndex + 1];
+
+			// Move to the next segment of the spline
+			const FVector SegmentStartPoint = PreviousSegmentEndPoint;
+			const FVector SegmentEndPoint = InOutParameters.PolyLineData->GetLocationAtAlpha(SplineAlpha);
+			PreviousSegmentEndPoint = SegmentEndPoint;
+			const FVector SubdivisionVector = SegmentEndPoint - SegmentStartPoint;
+			const FVector SubdivisionDirection = SubdivisionVector.GetSafeNormal();
+
+			// Since its discretized, we won't take the transform's position, but we'll use the up vector--to create the module rotation--and the scale
+			FTransform CenterPointTransform = InOutParameters.PolyLineData->GetTransformAtAlpha((SplineAlpha + PreviousAlpha) * 0.5);
+
+			const FVector Position = SegmentStartPoint + (SubdivisionVector * 0.5) + FVector(0, 0, InOutParameters.ModuleHeight * 0.5);
+			const FRotator Rotation = FRotationMatrix::MakeFromXZ(SubdivisionDirection, CenterPointTransform.GetRotation().GetUpVector()).Rotator();
+			const FVector Scale = FVector(1.0 + ModuleInstance.ExtraScale, 1.0, 1.0);
+
+			FPCGPoint& OutPoint = InOutParameters.OutPoints->Emplace_GetRef(FTransform(Rotation, Position, Scale), /*InDensity=*/1, PCGHelpers::ComputeSeedFromPosition(Position));
+
+			const double HalfSubmoduleSize = SubdivisionSubmodule.Size * 0.5;
+			OutPoint.SetLocalBounds(FBox(FVector(-HalfSubmoduleSize, 0, 0), FVector(HalfSubmoduleSize, 1, InOutParameters.ModuleHeight)));
+
+			// Now, handle the metadata attributes
+			if (!bHasMetadata)
+			{
+				continue;
+			}
+
+			InOutParameters.OutputMetadata->InitializeOnSet(OutPoint.MetadataEntry);
+			if (InOutParameters.SymbolAttribute)
+			{
+				InOutParameters.SymbolAttribute->SetValue(OutPoint.MetadataEntry, Symbol);
+			}
+
+			if (InOutParameters.DebugColorAttribute)
+			{
+				InOutParameters.DebugColorAttribute->SetValue(OutPoint.MetadataEntry, FVector4(SubdivisionSubmodule.DebugColor, 1.0));
+			}
+
+			if (InOutParameters.ModuleIndexAttribute)
+			{
+				InOutParameters.ModuleIndexAttribute->SetValue(OutPoint.MetadataEntry, ModuleInstanceIndex);
+			}
+
+			const bool bIsFirstModule = (ModuleInstanceIndex == 0);
+			if (bIsFirstModule && InOutParameters.IsFirstPointAttribute)
+			{
+				InOutParameters.IsFirstPointAttribute->SetValue(OutPoint.MetadataEntry, true);
+			}
+
+			const bool bIsFinalModule = (ModuleInstanceIndex == NumIterations - 1);
+			if (bIsFinalModule && InOutParameters.IsFinalPointAttribute)
+			{
+				InOutParameters.IsFinalPointAttribute->SetValue(OutPoint.MetadataEntry, true);
+			}
+		}
+	}
+}
+
+FPCGElementPtr UPCGSubdivideSplineSettings::CreateElement() const
+{
+	return MakeShared<FPCGSubdivideSplineElement>();
+}
+
+TArray<FPCGPinProperties> UPCGSubdivideSplineSettings::InputPinProperties() const
+{
+	TArray<FPCGPinProperties> PinProperties;
+	FPCGPinProperties& InputPin = PinProperties.Emplace_GetRef(PCGPinConstants::DefaultInputLabel, EPCGDataType::PolyLine);
+	InputPin.SetRequiredPin();
+
+	if (bModuleInfoAsInput)
+	{
+		FPCGPinProperties& ModuleInfoPin = PinProperties.Emplace_GetRef(PCGSubdivisionBase::Constants::ModulesInfoPinLabel, EPCGDataType::Param);
+		ModuleInfoPin.SetRequiredPin();
+	}
+
+	return PinProperties;
+}
+
+TArray<FPCGPinProperties> UPCGSubdivideSplineSettings::OutputPinProperties() const
+{
+	return Super::DefaultPointOutputPinProperties();
+}
+
+bool FPCGSubdivideSplineElement::ExecuteInternal(FPCGContext* InContext) const
+{
+	TRACE_CPUPROFILER_EVENT_SCOPE(FPCGSubdivideSplineElement::Execute);
+
+	const UPCGSubdivideSplineSettings* Settings = InContext->GetInputSettings<UPCGSubdivideSplineSettings>();
+	check(Settings);
+
+	const TArray<FPCGTaggedData> Inputs = InContext->InputData.GetInputsByPin(PCGPinConstants::DefaultInputLabel);
+	TArray<FPCGTaggedData>& Outputs = InContext->OutputData.TaggedData;
+
+	const UPCGParamData* ModuleInfoParamData = nullptr;
+	const FModuleInfoMap ModulesInfo = GetModulesInfoMap(InContext, Settings, ModuleInfoParamData);
+
+	double ModuleHeight = Settings->bModuleHeightAsAttribute ? 0.0 : Settings->ModuleHeight;
+
+	PCGSubdivideSplineHelpers::FParameters Parameters
+	{
+		.Context = InContext,
+		.ModulesInfo = ModulesInfo,
+		.bAcceptIncompleteSubdivision = Settings->bAcceptIncompleteSubdivision
+	};
+
+	for (const FPCGTaggedData& Input : Inputs)
+	{
+		const UPCGPolyLineData* InputPolyLineData = Cast<const UPCGPolyLineData>(Input.Data);
+		if (!InputPolyLineData)
+		{
+			continue;
+		}
+
+		if (Settings->bModuleHeightAsAttribute)
+		{
+			check(InputPolyLineData->Metadata);
+			const FPCGAttributePropertyInputSelector Selector = Settings->ModuleHeightAttribute.CopyAndFixLast(InputPolyLineData);
+			TUniquePtr<const IPCGAttributeAccessor> HeightAccessor = PCGAttributeAccessorHelpers::CreateConstAccessor(InputPolyLineData, Selector);
+
+			if (!HeightAccessor)
+			{
+				PCGLog::Metadata::LogFailToCreateAccessorError(Selector, InContext);
+				continue;
+			}
+
+			if (!HeightAccessor->Get<double>(ModuleHeight, FPCGAttributeAccessorKeysEntries(PCGInvalidEntryKey), EPCGAttributeAccessorFlags::AllowBroadcastAndConstructible))
+			{
+				PCGLog::Metadata::LogFailToGetAttributeError(Selector, InContext);
+				continue;
+			}
+		}
+
+		UPCGPointData* OutputPointData = FPCGContext::NewObject_AnyThread<UPCGPointData>(InContext);
+		OutputPointData->InitializeFromData(InputPolyLineData);
+		FPCGTaggedData& Output = Outputs.Emplace_GetRef(Input);
+		Output.Data = OutputPointData;
+
+		// Update 'per-input' parameters
+		Parameters.PolyLineData = InputPolyLineData;
+		Parameters.OutPoints = &OutputPointData->GetMutablePoints();
+		Parameters.OutputMetadata = OutputPointData->Metadata;
+		Parameters.ModuleHeight = ModuleHeight;
+
+		auto CreateAndValidateAttribute = [InContext, OutputMetadata = OutputPointData->Metadata]<typename T>(const FName AttributeName, const T DefaultValue, const bool bShouldCreate, FPCGMetadataAttribute<T>*& OutAttribute) -> bool
+		{
+			if (bShouldCreate)
+			{
+				OutAttribute = OutputMetadata->FindOrCreateAttribute<T>(AttributeName, DefaultValue, false, true);
+				if (!OutAttribute)
+				{
+					PCGLog::Metadata::LogFailToCreateAttributeError<T>(AttributeName, InContext);
+					return false;
+				}
+			}
+
+			return true;
+		};
+
+		if (!CreateAndValidateAttribute(Settings->SymbolAttributeName, FName(NAME_None), true, Parameters.SymbolAttribute)
+			|| !CreateAndValidateAttribute(Settings->DebugColorAttributeName, FVector4::Zero(), Settings->bOutputDebugColorAttribute, Parameters.DebugColorAttribute)
+			|| !CreateAndValidateAttribute(Settings->ModuleIndexAttributeName, -1, Settings->bOutputModuleIndexAttribute, Parameters.ModuleIndexAttribute)
+			|| !CreateAndValidateAttribute(Settings->IsFirstAttributeName, false, Settings->bOutputExtremityAttributes, Parameters.IsFirstPointAttribute)
+			|| !CreateAndValidateAttribute(Settings->IsFinalAttributeName, false, Settings->bOutputExtremityAttributes, Parameters.IsFinalPointAttribute))
+		{
+			continue;
+		}
+
+		// Set seed if required
+		Parameters.AdditionalSeed = 0;
+		if (Settings->bUseSeedAttribute)
+		{
+			const FPCGAttributePropertyInputSelector Selector = Settings->SeedAttribute.CopyAndFixLast(InputPolyLineData);
+			TUniquePtr<const IPCGAttributeAccessor> SeedAccessor = PCGAttributeAccessorHelpers::CreateConstAccessor(InputPolyLineData, Selector);
+
+			if (!SeedAccessor)
+			{
+				PCGLog::Metadata::LogFailToCreateAccessorError(Selector, InContext);
+			}
+			// Otherwise, get the value, if it fails, the attribute wasn't compatible
+			else if (!SeedAccessor->Get(Parameters.AdditionalSeed, FPCGAttributeAccessorKeysEntries(PCGInvalidEntryKey), EPCGAttributeAccessorFlags::AllowBroadcastAndConstructible))
+			{
+				PCGLog::Metadata::LogFailToGetAttributeError<int32>(Selector, SeedAccessor.Get(), InContext);
+			}
+		}
+
+		if (Settings->GrammarSelection.bGrammarAsAttribute)
+		{
+			const FPCGAttributePropertyInputSelector Selector = Settings->GrammarSelection.GrammarAttribute.CopyAndFixLast(InputPolyLineData);
+			const TUniquePtr<const IPCGAttributeAccessor> GrammarAccessor = PCGAttributeAccessorHelpers::CreateConstAccessor(InputPolyLineData, Selector);
+			const FPCGAttributeAccessorKeysEntries Keys(PCGInvalidEntryKey);
+			if (!GrammarAccessor)
+			{
+				PCGLog::Metadata::LogFailToCreateAccessorError(Selector, InContext);
+				continue;
+			}
+
+			auto Process = [&Parameters](const FString& InGrammar, int32 Index) -> void
+			{
+				PCGSubdivideSplineHelpers::Process(Parameters, InGrammar);
+			};
+
+			PCGMetadataElementCommon::ApplyOnAccessor<FString>(Keys, *GrammarAccessor, Process);
+		}
+		else
+		{
+			PCGSubdivideSplineHelpers::Process(Parameters, Settings->GrammarSelection.GrammarString);
+		}
+	}
+
+	if (Settings->bForwardAttributesFromModulesInfo && ModuleInfoParamData)
+	{
+		MatchAndSetAttributes(Inputs, Outputs, ModuleInfoParamData, Settings);
+	}
+
+	return true;
+}
+
+#undef LOCTEXT_NAMESPACE
diff --git a/Engine/Plugins/PCG/Source/PCG/Private/Elements/Grammar/PCGSubdivisionBase.cpp b/Engine/Plugins/PCG/Source/PCG/Private/Elements/Grammar/PCGSubdivisionBase.cpp
new file mode 100644
--- /dev/null
+++ b/Engine/Plugins/PCG/Source/PCG/Private/Elements/Grammar/PCGSubdivisionBase.cpp
@@ -0,0 +1,473 @@
+// Copyright Epic Games, Inc. All Rights Reserved.
+
+#include "Elements/Grammar/PCGSubdivisionBase.h"
+
+#include "PCGContext.h"
+#include "PCGParamData.h"
+#include "Data/PCGPointData.h"
+#include "Elements/Metadata/PCGMetadataElementCommon.h"
+#include "Grammar/PCGGrammarParser.h"
+#include "Helpers/PCGPropertyHelpers.h"
+#include "Metadata/Accessors/IPCGAttributeAccessor.h"
+#include "Metadata/Accessors/PCGAttributeAccessorHelpers.h"
+#include "Metadata/Accessors/PCGAttributeAccessorKeys.h"
+
+#define LOCTEXT_NAMESPACE "PCGSubdivisionBaseElement"
+
+namespace PCGSubdivisionBase
+{
+	static const FText DuplicatedSymbolText = LOCTEXT("SymbolDuplicate", "Symbol {0} is duplicated, ignored.");
+}
+
+void UPCGSubdivisionBaseSettings::PostLoad()
+{
+	Super::PostLoad();
+
+#if WITH_EDITOR
+	PRAGMA_DISABLE_DEPRECATION_WARNINGS
+	if (bGrammarAsAttribute_DEPRECATED != false)
+	{
+		GrammarSelection.bGrammarAsAttribute = bGrammarAsAttribute_DEPRECATED;
+		bGrammarAsAttribute_DEPRECATED = false;
+	}
+
+	if (!Grammar_DEPRECATED.IsEmpty())
+	{
+		GrammarSelection.GrammarString = Grammar_DEPRECATED;
+		Grammar_DEPRECATED.Empty();
+	}
+	PRAGMA_ENABLE_DEPRECATION_WARNINGS
+#endif // WITH_EDITOR
+}
+
+FPCGSubdivisionBaseElement::FModuleInfoMap FPCGSubdivisionBaseElement::GetModulesInfoMap(FPCGContext* InContext, const TArray<FPCGSubdivisionSubmodule>& SubmodulesInfo, const UPCGParamData*& OutModuleInfoParamData) const
+{
+	PCGSubdivisionBase::FModuleInfoMap ModulesInfo;
+	OutModuleInfoParamData = nullptr;
+
+	ModulesInfo.Reserve(SubmodulesInfo.Num());
+	for (const FPCGSubdivisionSubmodule& SubdivisionModule : SubmodulesInfo)
+	{
+		if (ModulesInfo.Contains(SubdivisionModule.Symbol))
+		{
+			PCGLog::LogWarningOnGraph(FText::Format(PCGSubdivisionBase::DuplicatedSymbolText, FText::FromName(SubdivisionModule.Symbol)), InContext);
+			continue;
+		}
+
+		ModulesInfo.Emplace(SubdivisionModule.Symbol, SubdivisionModule);
+	}
+
+	return ModulesInfo;
+}
+
+FPCGSubdivisionBaseElement::FModuleInfoMap FPCGSubdivisionBaseElement::GetModulesInfoMap(FPCGContext* InContext, const FPCGSubdivisionModuleAttributeNames& InSubdivisionModuleAttributeNames, const UPCGParamData*& OutModuleInfoParamData) const
+{
+	PCGSubdivisionBase::FModuleInfoMap ModulesInfo;
+	OutModuleInfoParamData = nullptr;
+
+	const TArray<FPCGTaggedData> ModulesInfoInputs = InContext->InputData.GetInputsByPin(PCGSubdivisionBase::Constants::ModulesInfoPinLabel);
+
+	if (ModulesInfoInputs.IsEmpty())
+	{
+		PCGLog::LogWarningOnGraph(LOCTEXT("NoModuleInfo", "No data was found on the module info pin."), InContext);
+		return ModulesInfo;
+	}
+
+	const UPCGParamData* ParamData = Cast<UPCGParamData>(ModulesInfoInputs[0].Data);
+	if (!ParamData)
+	{
+		PCGLog::LogWarningOnGraph(LOCTEXT("ModuleInfoWrongType", "Module info input is not of type attribute set."), InContext);
+		return ModulesInfo;
+	}
+
+	TMap<FName, TTuple<FName, bool>> PropertyNameMapping;
+	PropertyNameMapping.Emplace(GET_MEMBER_NAME_CHECKED(FPCGSubdivisionSubmodule, Symbol), {InSubdivisionModuleAttributeNames.SymbolAttributeName, /*bCanBeDefaulted=*/false});
+	PropertyNameMapping.Emplace(GET_MEMBER_NAME_CHECKED(FPCGSubdivisionSubmodule, Size), {InSubdivisionModuleAttributeNames.SizeAttributeName, /*bCanBeDefaulted=*/false});
+	PropertyNameMapping.Emplace(GET_MEMBER_NAME_CHECKED(FPCGSubdivisionSubmodule, bScalable), {InSubdivisionModuleAttributeNames.ScalableAttributeName, /*bCanBeDefaulted=*/!InSubdivisionModuleAttributeNames.bProvideScalable});
+	PropertyNameMapping.Emplace(GET_MEMBER_NAME_CHECKED(FPCGSubdivisionSubmodule, DebugColor), {InSubdivisionModuleAttributeNames.DebugColorAttributeName, /*bCanBeDefaulted=*/!InSubdivisionModuleAttributeNames.bProvideDebugColor});
+
+	const TArray<FPCGSubdivisionSubmodule> AllModules = PCGPropertyHelpers::ExtractAttributeSetAsArrayOfStructs<FPCGSubdivisionSubmodule>(ParamData, &PropertyNameMapping, InContext);
+
+	ModulesInfo.Reserve(AllModules.Num());
+
+	for (int32 i = 0; i < AllModules.Num(); ++i)
+	{
+		if (ModulesInfo.Contains(AllModules[i].Symbol))
+		{
+			PCGLog::LogWarningOnGraph(FText::Format(PCGSubdivisionBase::DuplicatedSymbolText, FText::FromName(AllModules[i].Symbol)), InContext);
+			continue;
+		}
+
+		ModulesInfo.Emplace(AllModules[i].Symbol, std::move(AllModules[i]));
+	}
+
+	OutModuleInfoParamData = ParamData;
+
+	return ModulesInfo;
+}
+
+PCGSubdivisionBase::FModuleInfoMap FPCGSubdivisionBaseElement::GetModulesInfoMap(FPCGContext* InContext, const UPCGSubdivisionBaseSettings* InSettings, const UPCGParamData*& OutModuleInfoParamData) const
+{
+	if (InSettings->bModuleInfoAsInput)
+	{
+		return GetModulesInfoMap(InContext, InSettings->ModulesInfoAttributeNames, OutModuleInfoParamData);
+	}
+	else
+	{
+		return GetModulesInfoMap(InContext, InSettings->ModulesInfo, OutModuleInfoParamData);
+	}
+}
+
+PCGGrammar::FTokenizedGrammar FPCGSubdivisionBaseElement::GetTokenizedGrammar(FPCGContext* InContext, const UPCGData* InputData, const UPCGSubdivisionBaseSettings* InSettings, const FModuleInfoMap& InModulesInfo, double& OutMinSize) const
+{
+	FString Grammar = InSettings->GrammarSelection.GrammarString;
+
+	if (InSettings->GrammarSelection.bGrammarAsAttribute)
+	{
+		const FPCGAttributePropertyInputSelector Selector = InSettings->GrammarSelection.GrammarAttribute.CopyAndFixLast(InputData);
+		const TUniquePtr<const IPCGAttributeAccessor> Accessor = PCGAttributeAccessorHelpers::CreateConstAccessor(InputData, Selector);
+		if (!Accessor)
+		{
+			PCGLog::Metadata::LogFailToCreateAccessorError(Selector, InContext);
+			return {};
+		}
+
+		if (!Accessor->Get(Grammar, FPCGAttributeAccessorKeysEntries(PCGInvalidEntryKey), EPCGAttributeAccessorFlags::AllowBroadcastAndConstructible))
+		{
+			PCGLog::Metadata::LogFailToGetAttributeError<FString>(Selector, Accessor.Get(), InContext);
+			return {};
+		}
+	}
+
+	return PCGSubdivisionBase::GetTokenizedGrammar(InContext, Grammar, InModulesInfo, OutMinSize);
+}
+
+PCGGrammar::FTokenizedGrammar PCGSubdivisionBase::GetTokenizedGrammar(FPCGContext* InContext, const FString& InGrammar, const FModuleInfoMap& InModulesInfo, double& OutMinSize)
+{
+	FPCGGrammarResult Result = PCGGrammar::Parse(InGrammar);
+
+	// TODO: Add quiet mode
+	{
+		if (!Result.bSuccess)
+		{
+			PCGLog::LogErrorOnGraph(LOCTEXT("GrammarParseFail", "Problem while parsing grammar:"), InContext);
+		}
+
+		for (const FPCGGrammarResult::FLog& Log : Result.GetLogs())
+		{
+			switch (Log.Verbosity)
+			{
+			case FPCGGrammarResult::ELogType::Error:
+				PCGLog::LogErrorOnGraph(Log.Message, InContext);
+				break;
+			case FPCGGrammarResult::ELogType::Warning:
+				PCGLog::LogWarningOnGraph(Log.Message, InContext);
+				break;
+			default:
+				UE_LOG(LogPCG, Log, TEXT("%s"), *Log.Message.ToString());
+				break;
+			}
+		}
+	}
+
+	if (!Result.bSuccess)
+	{
+		return {};
+	}
+
+	// Build equivalent grammar tree with the size information.
+	TSet<FName> UnmatchedTokens;
+	auto BuildNode = [&UnmatchedTokens, &InModulesInfo](PCGGrammar::FTokenizedModule& Module, const PCGGrammar::FModuleDescriptor& Descriptor, auto& RecursiveFn) -> void
+	{
+		if(Descriptor.Type == PCGGrammar::EModuleType::Literal)
+		{
+			if (const FPCGSubdivisionSubmodule* It = InModulesInfo.Find(Descriptor.Symbol))
+			{
+				Module.UnitSize = It->Size;
+				Module.ConcreteUnitSize = It->Size;
+				Module.bScalable = It->bScalable;
+				Module.bIsValid = true;
+			}
+			else
+			{
+				UnmatchedTokens.Add(Descriptor.Symbol);
+			}
+		}
+		else
+		{
+			// Go through submodules
+			Module.Submodules.Reserve(Descriptor.Submodules.Num());
+
+			// In any of the modules that can be expanded, the minimum concrete size is going to be either the MinSize if it's greater than 0, otherwise it will be the minimum of the minimum concrete sizes of all submodules
+			double MinSubmoduleConcreteSize = 0.0;
+
+			for (const PCGGrammar::FModuleDescriptor& SubDescriptor : Descriptor.Submodules)
+			{
+				PCGGrammar::FTokenizedModule& Submodule = Module.Submodules.Emplace_GetRef(&SubDescriptor);
+				RecursiveFn(Submodule, SubDescriptor, RecursiveFn);
+
+				if (Submodule.bIsValid)
+				{
+					// Update size/scalable, with Sequence/Root: sum and Stochastic/Priority : min
+					if (Descriptor.Type == PCGGrammar::EModuleType::Root || Descriptor.Type == PCGGrammar::EModuleType::Sequence)
+					{
+						Module.UnitSize += Submodule.GetMinSize();
+					}
+					else if (Descriptor.Type == PCGGrammar::EModuleType::Stochastic || Descriptor.Type == PCGGrammar::EModuleType::Priority)
+					{
+						if (Module.bIsValid)
+						{
+							Module.UnitSize = FMath::Min(Module.UnitSize, Submodule.GetMinSize());
+						}
+						else
+						{
+							Module.UnitSize = Submodule.GetMinSize();
+						}
+					}
+
+					if (Module.bIsValid)
+					{
+						MinSubmoduleConcreteSize = FMath::Min(MinSubmoduleConcreteSize, Submodule.GetMinConcreteSize());
+					}
+					else
+					{
+						MinSubmoduleConcreteSize = Submodule.GetMinConcreteSize();
+					}
+
+					Module.bScalable |= Submodule.bScalable;
+					Module.bIsValid = true;
+				}
+			}
+
+			// If all expansions of the module lead to nothing, we might as well remove it.
+			if (Module.bIsValid && MinSubmoduleConcreteSize > 0)
+			{
+				Module.ConcreteUnitSize = FMath::Max(Module.UnitSize, MinSubmoduleConcreteSize);
+			}
+			else
+			{
+				Module.bIsValid = false;
+			}
+		}
+	};
+
+	if (Result.Root.Submodules.IsEmpty())
+	{
+		return {};
+	}
+
+	PCGGrammar::FTokenizedGrammar TokenizedGrammar;
+	TokenizedGrammar.ParsedGrammar = MakeShared<PCGGrammar::FModuleDescriptor>(MoveTemp(Result.Root));
+	TokenizedGrammar.ModuleGrammar = MakeShared<PCGGrammar::FTokenizedModule>(TokenizedGrammar.ParsedGrammar.Get());
+	BuildNode(*TokenizedGrammar.ModuleGrammar, *TokenizedGrammar.ParsedGrammar, BuildNode);
+	OutMinSize = TokenizedGrammar.ModuleGrammar->GetMinSize();
+
+	// TODO Add quiet mode
+	for (FName UnmatchedToken : UnmatchedTokens)
+	{
+		FText WarningMessage = FText::Format(LOCTEXT("UnmatchedTokensInGrammar", "Unmatched token found in grammar: {0}."), FText::FromName(UnmatchedToken));
+		PCGLog::LogWarningOnGraph(WarningMessage, InContext);
+	}
+
+	return TokenizedGrammar;
+}
+
+TMap<FString, PCGGrammar::FTokenizedGrammar> FPCGSubdivisionBaseElement::GetTokenizedGrammarForPoints(FPCGContext* InContext, const UPCGPointData* InputData, const UPCGSubdivisionBaseSettings* InSettings, const FModuleInfoMap& InModulesInfo, double& OutMinSize) const
+{
+	TMap<FString, PCGGrammar::FTokenizedGrammar> Result;
+
+	if (InSettings->GrammarSelection.bGrammarAsAttribute)
+	{
+		const FPCGAttributePropertyInputSelector Selector = InSettings->GrammarSelection.GrammarAttribute.CopyAndFixLast(InputData);
+		const TUniquePtr<const IPCGAttributeAccessor> Accessor = PCGAttributeAccessorHelpers::CreateConstAccessor(InputData, Selector);
+		const TUniquePtr<const IPCGAttributeAccessorKeys> Keys = PCGAttributeAccessorHelpers::CreateConstKeys(InputData, Selector);
+		if (!Accessor || !Keys)
+		{
+			PCGLog::Metadata::LogFailToCreateAccessorError(Selector, InContext);
+			return Result;
+		}
+
+		const bool bSuccess = PCGMetadataElementCommon::ApplyOnAccessor<FString>(*Keys, *Accessor, [&Result](const FString& InValue, int32)
+		{
+			if (!Result.Contains(InValue))
+			{
+				Result.Emplace(InValue);
+			}
+		});
+
+		if (!bSuccess)
+		{
+			PCGLog::Metadata::LogFailToGetAttributeError<FString>(Selector, Accessor.Get(), InContext);
+			return Result;
+		}
+	}
+	else
+	{
+		Result.Emplace(InSettings->GrammarSelection.GrammarString);
+	}
+
+	for (auto& [Grammar, TokenizeGrammar] : Result)
+	{
+		TokenizeGrammar = PCGSubdivisionBase::GetTokenizedGrammar(InContext, Grammar, InModulesInfo, OutMinSize);
+	}
+
+	return Result;
+}
+
+bool FPCGSubdivisionBaseElement::MatchAndSetAttributes(const TArray<FPCGTaggedData>& InputData, TArray<FPCGTaggedData>& OutputData, const UPCGParamData* InModuleInfoParamData, const UPCGSubdivisionBaseSettings* InSettings) const
+{
+	check(InModuleInfoParamData && InModuleInfoParamData->Metadata);
+
+	const UPCGMetadata* InputMetadata = InModuleInfoParamData->Metadata;
+
+	// We prepare everything to process all the output data afterward.
+	// Build the Symbol -> EntryKey mapping
+	// Since we don't know if it is a Name or a String, we need to get both.
+	TMap<FName, PCGMetadataEntryKey> SymbolToEntryKeyMapping;
+	if (const FPCGMetadataAttribute<FName>* InSymbolNameAttribute = InputMetadata->GetConstTypedAttribute<FName>(PCGSubdivisionBase::Constants::SymbolAttributeName))
+	{
+		for (PCGMetadataEntryKey EntryKey = InputMetadata->GetItemKeyCountForParent(); EntryKey < InputMetadata->GetItemCountForChild(); ++EntryKey)
+		{
+			const FName Symbol = InSymbolNameAttribute->GetValueFromItemKey(EntryKey);
+			if (!SymbolToEntryKeyMapping.Contains(Symbol))
+			{
+				SymbolToEntryKeyMapping.Emplace(Symbol, EntryKey);
+			}
+		}
+	}
+	else if (const FPCGMetadataAttribute<FString>* InSymbolStrAttribute = InputMetadata->GetConstTypedAttribute<FString>(PCGSubdivisionBase::Constants::SymbolAttributeName))
+	{
+		for (PCGMetadataEntryKey EntryKey = InputMetadata->GetItemKeyCountForParent(); EntryKey < InputMetadata->GetItemCountForChild(); ++EntryKey)
+		{
+			const FName Symbol(InSymbolStrAttribute->GetValueFromItemKey(EntryKey));
+			if (!SymbolToEntryKeyMapping.Contains(Symbol))
+			{
+				SymbolToEntryKeyMapping.Emplace(Symbol, EntryKey);
+			}
+		}
+	}
+	else
+	{
+		return false;
+	}
+
+	// Also gather all the attribute in the input metadata
+	TArray<FName> AttributeNames;
+	TArray<EPCGMetadataTypes> AttributeTypes;
+	InModuleInfoParamData->Metadata->GetAttributes(AttributeNames, AttributeTypes);
+	check(AttributeNames.Num() == AttributeTypes.Num());
+
+	for (FPCGTaggedData& TaggedData : OutputData)
+	{
+		// Be careful with the const cast, only allow it if the data is not present in the input too (forwarded)
+		if (InputData.ContainsByPredicate([TaggedData](const FPCGTaggedData& InData) { return InData.Data == TaggedData.Data; }))
+		{
+			continue;
+		}
+
+		UPCGData* OutData = const_cast<UPCGData*>(TaggedData.Data.Get());
+		check(OutData && OutData->MutableMetadata());
+
+		UPCGMetadata* OutputMetadata = OutData->MutableMetadata();
+
+		// Look for the Symbol Attribute in the output metadata to query its value.
+		const FPCGMetadataAttribute<FName>* OutSymbolAttribute = OutputMetadata->GetConstTypedAttribute<FName>(InSettings->SymbolAttributeName);
+		if (!ensure(OutSymbolAttribute))
+		{
+			continue;
+		}
+
+		if (UPCGPointData* OutPointData = Cast<UPCGPointData>(OutData))
+		{
+			TArray<const FPCGMetadataAttributeBase*> InAttributes;
+			TArray<FPCGMetadataAttributeBase*> OutAttributes;
+
+			InAttributes.Reserve(AttributeNames.Num());
+			OutAttributes.Reserve(AttributeNames.Num());
+
+			// Copy all the attributes in the output metadata
+			for (const FName AttributeName : AttributeNames)
+			{
+				// Skip the symbol, it already exists, but perhaps with a different attribute name.
+				if (AttributeName == PCGSubdivisionBase::Constants::SymbolAttributeName)
+				{
+					continue;
+				}
+
+				// Skip any attribute that already exists in the output data
+				if (OutputMetadata->HasAttribute(AttributeName))
+				{
+					continue;
+				}
+
+				InAttributes.Add(InputMetadata->GetConstAttribute(AttributeName));
+				OutAttributes.Add(OutputMetadata->CopyAttribute(InAttributes.Last(), AttributeName, /*bKeepParent=*/false, /*bCopyEntries=*/false, /*bCopyValues=*/true));
+			}
+
+			for (const FPCGPoint& OutPoint : OutPointData->GetPoints())
+			{
+				const FName Symbol = OutSymbolAttribute->GetValueFromItemKey(OutPoint.MetadataEntry);
+
+				if (!SymbolToEntryKeyMapping.Contains(Symbol))
+				{
+					continue;
+				}
+
+				const PCGMetadataEntryKey InputEntryKey = SymbolToEntryKeyMapping[Symbol];
+
+				if (InputEntryKey == PCGInvalidEntryKey)
+				{
+					continue;
+				}
+
+				for (int32 i = 0; i < InAttributes.Num(); ++i)
+				{
+					PCGMetadataValueKey InputValueKey = InAttributes[i]->GetValueKey(InputEntryKey);
+					OutAttributes[i]->SetValueFromValueKey(OutPoint.MetadataEntry, InputValueKey);
+				}
+			}
+		}
+		else
+		{
+			// If we have a spatial data that is not a point data, we can only operate on Default values.
+			// So add all the attributes, and set the default value to the same value for the chosen entry key.
+			const FName Symbol = OutSymbolAttribute->GetValue(PCGDefaultValueKey);
+			const PCGMetadataEntryKey InputEntryKey = SymbolToEntryKeyMapping[Symbol];
+
+			auto AddAttribute = [InputMetadata, OutputMetadata, InputEntryKey]<typename T>(T Dummy, const FName AttributeName)
+			{
+				const FPCGMetadataAttribute<T>* InAttribute = InputMetadata->GetConstTypedAttribute<T>(AttributeName);
+				check(InAttribute);
+
+				// Overwrite any existing attribute
+				if (OutputMetadata->HasAttribute(AttributeName))
+				{
+					OutputMetadata->DeleteAttribute(AttributeName);
+				}
+
+				OutputMetadata->CreateAttribute<T>(AttributeName, InAttribute->GetValueFromItemKey(InputEntryKey), InAttribute->AllowsInterpolation(), /*bOverrideParent=*/true);
+			};
+
+			for (int32 i = 0; i < AttributeNames.Num(); ++i)
+			{
+				// Skip the symbol attribute, as it is already in the OutData (with perhaps a different name)
+				if (AttributeNames[i] == PCGSubdivisionBase::Constants::SymbolAttributeName)
+				{
+					continue;
+				}
+
+				// Skip any attribute that already exists in the output data
+				if (OutputMetadata->HasAttribute(AttributeNames[i]))
+				{
+					continue;
+				}
+
+				PCGMetadataAttribute::CallbackWithRightType(static_cast<int16>(AttributeTypes[i]), AddAttribute, AttributeNames[i]);
+			}
+		}
+	}
+
+	return true;
+}
+
+#undef LOCTEXT_NAMESPACE
diff --git a/Engine/Plugins/PCG/Source/PCG/Private/Elements/Grammar/PCGVolumeSlicer.cpp b/Engine/Plugins/PCG/Source/PCG/Private/Elements/Grammar/PCGVolumeSlicer.cpp
deleted file mode 100644
diff --git a/Engine/Plugins/PCG/Source/PCG/Public/Elements/Grammar/PCGDuplicateCrossSections.h b/Engine/Plugins/PCG/Source/PCG/Public/Elements/Grammar/PCGDuplicateCrossSections.h
new file mode 100644
--- /dev/null
+++ b/Engine/Plugins/PCG/Source/PCG/Public/Elements/Grammar/PCGDuplicateCrossSections.h
@@ -0,0 +1,53 @@
+// Copyright Epic Games, Inc. All Rights Reserved.
+
+#pragma once
+
+#include "PCGSettings.h"
+#include "Metadata/PCGAttributePropertySelector.h"
+#include "Elements/Grammar/PCGSubdivisionBase.h"
+
+#include "PCGDuplicateCrossSections.generated.h"
+
+UCLASS(MinimalAPI, BlueprintType, ClassGroup = (Procedural))
+class UPCGDuplicateCrossSectionsSettings : public UPCGSubdivisionBaseSettings
+{
+	GENERATED_BODY()
+
+public:
+	//~Begin UPCGSettings interface
+#if WITH_EDITOR
+	virtual FName GetDefaultNodeName() const override;
+	virtual FText GetDefaultNodeTitle() const override;
+#endif
+
+protected:
+	virtual TArray<FPCGPinProperties> InputPinProperties() const override;
+	virtual TArray<FPCGPinProperties> OutputPinProperties() const override;
+	virtual FPCGElementPtr CreateElement() const override;
+	//~End UPCGSettings interface
+
+public:
+	/** Set it to true if you want the extrude vector to be taken from the input spline as attribute, or fixed in the settings. */
+	UPROPERTY(BlueprintReadWrite, EditAnywhere, Category = Settings)
+	bool bExtrudeVectorAsAttribute = false;
+
+	UPROPERTY(BlueprintReadWrite, EditAnywhere, Category = Settings, meta = (EditCondition = "!bExtrudeVectorAsAttribute", EditConditionHides, DisplayAfter = bExtrudeVectorAsAttribute, PCG_Overridable))
+	FVector ExtrudeVector = FVector(0.0, 0.0, 1000.0);
+
+	/** Attribute to be taken from the input spline containing the extrude vector for the slicing. */
+	UPROPERTY(BlueprintReadWrite, EditAnywhere, Category = Settings, meta = (EditCondition = "bExtrudeVectorAsAttribute", EditConditionHides, DisplayAfter = bExtrudeVectorAsAttribute, PCG_Overridable))
+	FPCGAttributePropertyInputSelector ExtrudeVectorAttribute;
+
+	UPROPERTY(BlueprintReadWrite, EditAnywhere, Category = "Settings|Extra Output Attributes", meta = (PCG_Overridable, InlineEditConditionToggle))
+	bool bOutputSplineIndexAttribute = true;
+
+	/** Name of the spline index output attribute name. */
+	UPROPERTY(BlueprintReadWrite, EditAnywhere, Category = "Settings|Extra Output Attributes", meta = (PCG_Overridable, EditCondition = "bOutputSplineIndexAttribute"))
+	FName SplineIndexAttributeName = TEXT("SplineIndex");
+};
+
+class FPCGDuplicateCrossSectionsElement : public FPCGSubdivisionBaseElement
+{
+protected:
+	virtual bool ExecuteInternal(FPCGContext* InContext) const override;
+};
diff --git a/Engine/Plugins/PCG/Source/PCG/Public/Elements/Grammar/PCGSegmentSlicer.h b/Engine/Plugins/PCG/Source/PCG/Public/Elements/Grammar/PCGSegmentSlicer.h
deleted file mode 100644
diff --git a/Engine/Plugins/PCG/Source/PCG/Public/Elements/Grammar/PCGSelectGrammar.h b/Engine/Plugins/PCG/Source/PCG/Public/Elements/Grammar/PCGSelectGrammar.h
--- a/Engine/Plugins/PCG/Source/PCG/Public/Elements/Grammar/PCGSelectGrammar.h
+++ b/Engine/Plugins/PCG/Source/PCG/Public/Elements/Grammar/PCGSelectGrammar.h
@@ -3,7 +3,7 @@
 #pragma once
 
 #include "PCGSettings.h"
-#include "PCGSlicingBase.h"
+#include "PCGSubdivisionBase.h"
 
 #include "PCGSelectGrammar.generated.h"
 
@@ -142,7 +142,7 @@ public:
 };
 
 // TODO: For V2, this element could be async and/or timesliced
-class FPCGSelectGrammarElement : public FPCGSlicingBaseElement
+class FPCGSelectGrammarElement : public FPCGSubdivisionBaseElement
 {
 protected:
 	virtual bool ExecuteInternal(FPCGContext* InContext) const override;
diff --git a/Engine/Plugins/PCG/Source/PCG/Public/Elements/Grammar/PCGSlicingBase.h b/Engine/Plugins/PCG/Source/PCG/Public/Elements/Grammar/PCGSlicingBase.h
deleted file mode 100644
diff --git a/Engine/Plugins/PCG/Source/PCG/Public/Elements/Grammar/PCGSplineSlicer.h b/Engine/Plugins/PCG/Source/PCG/Public/Elements/Grammar/PCGSplineSlicer.h
deleted file mode 100644
diff --git a/Engine/Plugins/PCG/Source/PCG/Public/Elements/Grammar/PCGSubdivideSegment.h b/Engine/Plugins/PCG/Source/PCG/Public/Elements/Grammar/PCGSubdivideSegment.h
new file mode 100644
--- /dev/null
+++ b/Engine/Plugins/PCG/Source/PCG/Public/Elements/Grammar/PCGSubdivideSegment.h
@@ -0,0 +1,89 @@
+// Copyright Epic Games, Inc. All Rights Reserved.
+
+#pragma once
+
+#include "PCGSettings.h"
+
+#include "Elements/Grammar/PCGSubdivisionBase.h"
+
+#include "Elements/PCGSplitPoints.h"
+#include "Metadata/PCGAttributePropertySelector.h"
+
+#include "PCGSubdivideSegment.generated.h"
+
+UCLASS(MinimalAPI, BlueprintType, ClassGroup = (Procedural))
+class UPCGSubdivideSegmentSettings : public UPCGSubdivisionBaseSettings
+{
+	GENERATED_BODY()
+
+public:
+	//~Begin UPCGSettings interface
+#if WITH_EDITOR
+	virtual FName GetDefaultNodeName() const override;
+	virtual FText GetDefaultNodeTitle() const override;
+#endif
+
+protected:
+	virtual TArray<FPCGPinProperties> InputPinProperties() const override;
+	virtual TArray<FPCGPinProperties> OutputPinProperties() const override;
+	virtual FPCGElementPtr CreateElement() const override;
+	//~End UPCGSettings interface
+
+public:
+	/** Subdivision direction in point local space. */
+	UPROPERTY(BlueprintReadWrite, EditAnywhere, Category = Settings)
+	EPCGSplitAxis SubdivisionAxis = EPCGSplitAxis::X;
+
+	/** Use an attribute to determine whether we should flip axis. */
+	UPROPERTY(BlueprintReadWrite, EditAnywhere, Category = Settings)
+	bool bFlipAxisAsAttribute = false;
+
+	/** If we need to flip axis. */
+	UPROPERTY(BlueprintReadWrite, EditAnywhere, Category = Settings, meta = (PCG_Overridable, EditCondition = "!bFlipAxisAsAttribute", EditConditionHides, ShowAfter = "bFlipAxisAsAttribute"))
+	bool bShouldFlipAxis = false;
+
+	/** Name of the attribute to know if we need to flip axis. */
+	UPROPERTY(BlueprintReadWrite, EditAnywhere, Category = Settings, meta = (PCG_Overridable, EditCondition = "bFlipAxisAsAttribute", EditConditionHides, ShowAfter="bFlipAxisAsAttribute"))
+	FPCGAttributePropertyInputSelector FlipAxisAttribute;
+
+	/** If the subdivision with a given grammar doesn't fill the entire segment, setting it to true makes it a valid case. */
+	UPROPERTY(BlueprintReadWrite, EditAnywhere, Category = Settings)
+	bool bAcceptIncompleteSubdivision = false;
+
+	UPROPERTY(BlueprintReadWrite, EditAnywhere, Category = "Settings|Extra Output Attributes", meta = (PCG_Overridable, InlineEditConditionToggle))
+	bool bOutputModuleIndexAttribute = false;
+
+	/** Name of the module index output attribute name. */
+	UPROPERTY(BlueprintReadWrite, EditAnywhere, Category = "Settings|Extra Output Attributes", meta = (PCG_Overridable, EditCondition = "bOutputModuleIndexAttribute"))
+	FName ModuleIndexAttributeName = TEXT("ModuleIndex");
+
+	/** Output attributes labeling the first and final points after subdivision. */
+	UPROPERTY(BlueprintReadWrite, EditAnywhere, Category = "Settings|Extra Output Attributes", meta = (PCG_Overridable))
+	bool bOutputExtremityAttributes = false;
+
+	/** Name of the attribute labeling the first output point from the first module. */
+	UPROPERTY(BlueprintReadWrite, EditAnywhere, Category = "Settings|Extra Output Attributes", meta = (PCG_Overridable, EditCondition = "bOutputExtremityAttributes", EditConditionHides))
+	FName IsFirstAttributeName = TEXT("IsFirst");
+
+	/** Name of the attribute labeling the final output point from the final module. */
+	UPROPERTY(BlueprintReadWrite, EditAnywhere, Category = "Settings|Extra Output Attributes", meta = (PCG_Overridable, EditCondition = "bOutputExtremityAttributes", EditConditionHides))
+	FName IsFinalAttributeName = TEXT("IsFinal");
+
+	UPROPERTY(BlueprintReadWrite, EditAnywhere, Category = "Settings|Extra Output Attributes", meta = (PCG_Overridable, InlineEditConditionToggle))
+	bool bOutputExtremityNeighborIndexAttribute = false;
+
+	/** Name of the extremity neighbor index output attribute name. */
+	UPROPERTY(BlueprintReadWrite, EditAnywhere, Category = "Settings|Extra Output Attributes", meta = (PCG_Overridable, EditCondition = "bOutputExtremityNeighborIndexAttribute"))
+	FName ExtremityNeighborIndexAttributeName = TEXT("ExtremityNeighborIndex");
+
+};
+
+class PCGSubdivideSegmentHelpers;
+
+class FPCGSegmentSubdivisionElement : public FPCGSubdivisionBaseElement
+{
+public:
+	friend class PCGSubdivideSegmentHelpers;
+protected:
+	virtual bool ExecuteInternal(FPCGContext* InContext) const override;
+};
diff --git a/Engine/Plugins/PCG/Source/PCG/Public/Elements/Grammar/PCGSubdivideSpline.h b/Engine/Plugins/PCG/Source/PCG/Public/Elements/Grammar/PCGSubdivideSpline.h
new file mode 100644
--- /dev/null
+++ b/Engine/Plugins/PCG/Source/PCG/Public/Elements/Grammar/PCGSubdivideSpline.h
@@ -0,0 +1,70 @@
+// Copyright Epic Games, Inc. All Rights Reserved.
+
+#pragma once
+
+#include "PCGSettings.h"
+
+#include "Elements/Grammar/PCGSubdivisionBase.h"
+
+#include "PCGSubdivideSpline.generated.h"
+
+UCLASS(MinimalAPI, BlueprintType, ClassGroup = (Procedural))
+class UPCGSubdivideSplineSettings : public UPCGSubdivisionBaseSettings
+{
+	GENERATED_BODY()
+
+public:
+	//~Begin UPCGSettings interface
+#if WITH_EDITOR
+	virtual FName GetDefaultNodeName() const override { return FName(TEXT("SubdivideSpline")); }
+	virtual FText GetDefaultNodeTitle() const override { return NSLOCTEXT("PCGSubdivideSplineElement", "NodeTitle", "Subdivide Spline"); }
+#endif // WITH_EDITOR
+
+protected:
+	virtual TArray<FPCGPinProperties> InputPinProperties() const override;
+	virtual TArray<FPCGPinProperties> OutputPinProperties() const override;
+	virtual FPCGElementPtr CreateElement() const override;
+	//~End UPCGSettings interface
+
+public:
+	/** If the subdivision with a given grammar doesn't fill the entire spline, setting it to true makes it a valid case. */
+	UPROPERTY(BlueprintReadWrite, EditAnywhere, Category = Settings)
+	bool bAcceptIncompleteSubdivision = false;
+
+	/** Select the module height from an attribute. */
+	UPROPERTY(BlueprintReadWrite, EditAnywhere, Category = Settings)
+	bool bModuleHeightAsAttribute = false;
+
+	/** The height of each placed module. */
+	UPROPERTY(BlueprintReadWrite, EditAnywhere, Category = Settings, meta = (EditCondition = "!bModuleHeightAsAttribute", EditConditionHides))
+	double ModuleHeight = 100.0;
+
+	/** Selection will be used as the module height for placed modules. */
+	UPROPERTY(BlueprintReadWrite, EditAnywhere, Category = Settings, meta = (EditCondition = "bModuleHeightAsAttribute", EditConditionHides))
+	FPCGAttributePropertyInputSelector ModuleHeightAttribute;
+
+	UPROPERTY(BlueprintReadWrite, EditAnywhere, Category = "Settings|Extra Output Attributes", meta = (PCG_Overridable, InlineEditConditionToggle))
+	bool bOutputModuleIndexAttribute = false;
+
+	/** Name of the module index output attribute name. */
+	UPROPERTY(BlueprintReadWrite, EditAnywhere, Category = "Settings|Extra Output Attributes", meta = (PCG_Overridable, EditCondition = "bOutputModuleIndexAttribute"))
+	FName ModuleIndexAttributeName = TEXT("ModuleIndex");
+
+	/** Output attributes labeling the first and final points after subdivision. */
+	UPROPERTY(BlueprintReadWrite, EditAnywhere, Category = "Settings|Extra Output Attributes", meta = (PCG_Overridable))
+	bool bOutputExtremityAttributes = false;
+
+	/** Name of the attribute labeling the first output point from the first module. */
+	UPROPERTY(BlueprintReadWrite, EditAnywhere, Category = "Settings|Extra Output Attributes", meta = (PCG_Overridable, EditCondition = "bOutputExtremityAttributes", EditConditionHides))
+	FName IsFirstAttributeName = TEXT("IsFirst");
+
+	/** Name of the attribute labeling the final output point from the final module. */
+	UPROPERTY(BlueprintReadWrite, EditAnywhere, Category = "Settings|Extra Output Attributes", meta = (PCG_Overridable, EditCondition = "bOutputExtremityAttributes", EditConditionHides))
+	FName IsFinalAttributeName = TEXT("IsFinal");
+};
+
+class FPCGSubdivideSplineElement : public FPCGSubdivisionBaseElement
+{
+protected:
+	virtual bool ExecuteInternal(FPCGContext* InContext) const override;
+};
diff --git a/Engine/Plugins/PCG/Source/PCG/Public/Elements/Grammar/PCGSubdivisionBase.h b/Engine/Plugins/PCG/Source/PCG/Public/Elements/Grammar/PCGSubdivisionBase.h
new file mode 100644
--- /dev/null
+++ b/Engine/Plugins/PCG/Source/PCG/Public/Elements/Grammar/PCGSubdivisionBase.h
@@ -0,0 +1,481 @@
+// Copyright Epic Games, Inc. All Rights Reserved.
+
+#pragma once
+
+#include "PCGContext.h"
+#include "PCGSettings.h"
+#include "Grammar/PCGGrammar.h"
+#include "Utils/PCGLogErrors.h"
+
+#include "PCGSubdivisionBase.generated.h"
+
+class UPCGData;
+class UPCGPointData;
+
+USTRUCT(BlueprintType)
+struct FPCGSubdivisionSubmodule
+{
+	GENERATED_BODY()
+
+	/** Symbol for the grammar. */
+	UPROPERTY(BlueprintReadWrite, EditAnywhere, Category = "")
+	FName Symbol = NAME_None;
+
+	/** Size of the block, aligned on the segment direction. */
+	UPROPERTY(BlueprintReadWrite, EditAnywhere, Category = "")
+	double Size = 100.0;
+
+	/** If the volume can be scaled to fit the remaining space or not. */
+	UPROPERTY(BlueprintReadWrite, EditAnywhere, Category = "")
+	bool bScalable = false;
+
+	/** For easier debugging, using Point color in conjunction with PCG Debug Color Material. */
+	UPROPERTY(BlueprintReadWrite, EditAnywhere, Category = "Debug")
+	FVector4 DebugColor = FVector4::One();
+};
+
+namespace PCGSubdivisionBase::Constants
+{
+	const FName ModulesInfoPinLabel = TEXT("ModulesInfo");
+	const FName SymbolAttributeName = TEXT("Symbol");
+	const FName SizeAttributeName = TEXT("Size");
+	const FName ScalableAttributeName = TEXT("Scalable");
+	const FName DebugColorAttributeName = TEXT("DebugColor");
+}
+
+USTRUCT(BlueprintType)
+struct FPCGSubdivisionModuleAttributeNames
+{
+	GENERATED_BODY()
+
+public:
+	/** Mandatory. Expected type: FName. */
+	UPROPERTY(BlueprintReadWrite, EditAnywhere, Category = "")
+	FName SymbolAttributeName = PCGSubdivisionBase::Constants::SymbolAttributeName;
+
+	/** Mandatory. Expected type: double. */
+	UPROPERTY(BlueprintReadWrite, EditAnywhere, Category = "")
+	FName SizeAttributeName = PCGSubdivisionBase::Constants::SizeAttributeName;
+
+	UPROPERTY(BlueprintReadWrite, EditAnywhere, Category = "", meta = (InlineEditConditionToggle))
+	bool bProvideScalable = false;
+
+	/** Optional. Expected type: bool. If disabled, default value will be false. */
+	UPROPERTY(BlueprintReadWrite, EditAnywhere, Category = "", meta = (EditCondition = "bProvideScalable"))
+	FName ScalableAttributeName = PCGSubdivisionBase::Constants::ScalableAttributeName;
+
+	UPROPERTY(BlueprintReadWrite, EditAnywhere, Category = "", meta = (InlineEditConditionToggle))
+	bool bProvideDebugColor = false;
+
+	/** Optional. Expected type: Vector4. If disabled, default value will be (1.0, 1.0, 1.0, 1.0). */
+	UPROPERTY(BlueprintReadWrite, EditAnywhere, Category = "", meta = (EditCondition = "bProvideDebugColor"))
+	FName DebugColorAttributeName = PCGSubdivisionBase::Constants::DebugColorAttributeName;
+};
+
+UCLASS(MinimalAPI, Abstract, ClassGroup = (Procedural))
+class UPCGSubdivisionBaseSettings : public UPCGSettings
+{
+	GENERATED_BODY()
+
+public:
+	//~Begin UPCGSettings interface
+#if WITH_EDITOR
+	virtual EPCGSettingsType GetType() const override { return EPCGSettingsType::Spatial; }
+#endif
+
+	virtual bool UseSeed() const override { return true; }
+	//~End UPCGSettings interface
+
+	virtual void PostLoad() override;
+
+public:
+	/** Set it to true to pass the info as attribute set. */
+	UPROPERTY(BlueprintReadWrite, EditAnywhere, Category = Settings)
+	bool bModuleInfoAsInput = false;
+
+	/** Fixed array of modules used for the subdivision. */
+	UPROPERTY(BlueprintReadWrite, EditAnywhere, Category = Settings, meta = (EditCondition = "!bModuleInfoAsInput", EditConditionHides, DisplayAfter = bModuleInfoAsInput))
+	TArray<FPCGSubdivisionSubmodule> ModulesInfo;
+
+	/** Fixed array of modules used for the subdivision. */
+	UPROPERTY(BlueprintReadWrite, EditAnywhere, Category = Settings, meta = (EditCondition = "bModuleInfoAsInput", EditConditionHides, DisplayAfter = bModuleInfoAsInput, DisplayName = "Attribute Names for Module Info"))
+	FPCGSubdivisionModuleAttributeNames ModulesInfoAttributeNames;
+
+	/** An encoded string that represents how to apply a set of rules to a series of defined modules. */
+	UPROPERTY(BlueprintReadWrite, EditAnywhere, Category = Settings, meta = (ShowOnlyInnerProperties))
+	FPCGGrammarSelection GrammarSelection;
+
+	/** Attribute to be taken from the input spline containing the grammar to use for the subdivision. */
+	UPROPERTY(BlueprintReadWrite, EditAnywhere, Category = Settings, meta = (EditCondition = "bGrammarAsAttribute", EditConditionHides, DisplayAfter = bGrammarAsAttribute, PCG_Overridable))
+	FPCGAttributePropertyInputSelector GrammarAttribute;
+
+	/** Controls whether we'll use an attribute to drive random seeding for stochastic processes in the subdivision. */
+	UPROPERTY(BlueprintReadWrite, EditAnywhere, Category = Settings, meta = (PCG_Overridable, InlineEditConditionToggle))
+	bool bUseSeedAttribute = false;
+
+	/** Attribute to use to drive seed selection. It should be convertible to an integer. */
+	UPROPERTY(BlueprintReadWrite, EditAnywhere, Category = Settings, meta = (EditCondition = "bUseSeedAttribute", PCG_Overridable))
+	FPCGAttributePropertyInputSelector SeedAttribute;
+
+	/** Do a match and set with the incoming modules info, only if the modules info is passed as input. */
+	UPROPERTY(BlueprintReadWrite, EditAnywhere, Category = "Settings|Extra Output Attributes", meta = (EditCondition = "bModuleInfoAsInput", EditConditionHides, PCG_Overridable))
+	bool bForwardAttributesFromModulesInfo = false;
+
+	/** Name of the Symbol output attribute name. */
+	UPROPERTY(BlueprintReadWrite, EditAnywhere, Category = "Settings|Extra Output Attributes", meta = (PCG_Overridable))
+	FName SymbolAttributeName = PCGSubdivisionBase::Constants::SymbolAttributeName;
+
+	UPROPERTY(BlueprintReadWrite, EditAnywhere, Category = "Settings|Extra Output Attributes", meta = (PCG_Overridable, InlineEditConditionToggle))
+	bool bOutputSizeAttribute = true;
+
+	/** Name of the Size output attribute name, ignored if match and set from module info is true. */
+	UPROPERTY(BlueprintReadWrite, EditAnywhere, Category = "Settings|Extra Output Attributes", meta = (PCG_Overridable, EditCondition = "bOutputSizeAttribute"))
+	FName SizeAttributeName = PCGSubdivisionBase::Constants::SizeAttributeName;
+
+	UPROPERTY(BlueprintReadWrite, EditAnywhere, Category = "Settings|Extra Output Attributes", meta = (PCG_Overridable, InlineEditConditionToggle))
+	bool bOutputScalableAttribute = true;
+
+	/** Name of the Scalable output attribute name, ignored if match and set from module info is true. */
+	UPROPERTY(BlueprintReadWrite, EditAnywhere, Category = "Settings|Extra Output Attributes", meta = (PCG_Overridable, EditCondition = "bOutputScalableAttribute"))
+	FName ScalableAttributeName = PCGSubdivisionBase::Constants::ScalableAttributeName;
+
+	UPROPERTY(BlueprintReadWrite, EditAnywhere, Category = "Settings|Extra Output Attributes", meta = (PCG_Overridable, InlineEditConditionToggle))
+	bool bOutputDebugColorAttribute = false;
+
+	/** Name of the Debug Color output attribute name, ignored if match and set from module info is true. */
+	UPROPERTY(BlueprintReadWrite, EditAnywhere, Category = "Settings|Extra Output Attributes", meta = (PCG_Overridable, EditCondition = "bOutputDebugColorAttribute"))
+	FName DebugColorAttributeName = PCGSubdivisionBase::Constants::DebugColorAttributeName;
+
+private:
+#if WITH_EDITORONLY_DATA
+	UE_DEPRECATED(5.5, "Use 'GrammarSelection' instead.")
+	UPROPERTY()
+	bool bGrammarAsAttribute_DEPRECATED = false;
+
+	UE_DEPRECATED(5.5, "Use 'GrammarSelection' instead.")
+	UPROPERTY()
+	FString Grammar_DEPRECATED;
+#endif // WITH_EDITORONLY_DATA
+};
+
+namespace PCGSubdivisionBase
+{
+	using FModuleInfoMap = TMap<FName, FPCGSubdivisionSubmodule>;
+
+	// Materialized modules created from a tokenized grammar.
+	template <typename T>
+	struct TModuleInstance
+	{
+		TModuleInstance() = default;
+		explicit TModuleInstance(const T* InModule) : Module(InModule) {}
+
+		const T* Module = nullptr;
+		double ExtraScale = 0.0;
+		bool bIsValid = true;
+		bool bSkipExpansion = false;
+	};
+
+	PCGGrammar::FTokenizedGrammar GetTokenizedGrammar(FPCGContext* InContext, const FString& InGrammar, const FModuleInfoMap& InModulesInfo, double& OutMinSize);
+
+	template <typename T>
+	bool Subdivide(const T& Root, double Length, TArray<TModuleInstance<T>>& OutModuleInstances, double& RemainingLength, FPCGContext* InOptionalContext = nullptr, int32 InOptionalAdditionalSeed = 0)
+	{
+		OutModuleInstances.Reset();
+		RemainingLength = Length;
+
+		if (!Root.IsValid() || FMath::IsNearlyZero(Length))
+		{
+			return true;
+		}
+
+		FRandomStream RandomStream((InOptionalContext ? InOptionalContext->GetSeed() : 42) + InOptionalAdditionalSeed);
+		TArray<TModuleInstance<T>> CurrentModules;
+
+		// Start with root
+		CurrentModules.Emplace(&Root);
+		RemainingLength -= Root.GetMinSize(); // Here we use the min size because we'll consume concrete size only during expansion
+
+		if (RemainingLength < 0)
+		{
+			PCGLog::LogErrorOnGraph(NSLOCTEXT("PCGSubdivisionBase", "SegmentCutFail", "Grammar doesn't fit for this segment."), InOptionalContext);
+			return false;
+		}
+
+		// Working data sets
+		TArray<TModuleInstance<T>*> PriorityModules;
+		TArray<TModuleInstance<T>*> StochasticModules;
+		TArray<TModuleInstance<T>> ExpandedModules;
+		TArray<TModuleInstance<T>> PreviousModules;
+
+		bool bNotDone = true;
+		while (bNotDone)
+		{
+			bNotDone = false;
+			// 1. Expand/Visit sequences
+			// Implementation note; since we've already consumed the min size, there's no need to update anything when expanding nodes, only when replacing with another choice.
+			// We'll keep track of whether we have certain types of nodes at this point so we can short-circuit some passes
+			ExpandedModules.Reset();
+			ExpandedModules.Reserve(CurrentModules.Num());
+			bool bHasPriorityOrStochasticModules = false;
+
+			for (TModuleInstance<T>& CurrentModule : CurrentModules)
+			{
+				if (!CurrentModule.bIsValid)
+				{
+					continue; // discard this module (reclaim?)
+				}
+
+				if ((CurrentModule.Module->GetType() == PCGGrammar::EModuleType::Root || CurrentModule.Module->GetType() == PCGGrammar::EModuleType::Sequence))
+				{
+					// When we duplicate sequences, we can't expand them directly otherwise they won't be repeatable
+					if (CurrentModule.bSkipExpansion)
+					{
+						// Copy module as-is, but remove the expansion limitation so we break it down next iteration
+						TModuleInstance<T>& ExpandedModule = ExpandedModules.Add_GetRef(CurrentModule);
+						ExpandedModule.bSkipExpansion = false;
+					}
+					else
+					{
+						for (const T& Submodule : CurrentModule.Module->Submodules)
+						{
+							if (Submodule.IsValid())
+							{
+								// Expand module if it has a positive number of iterations (>= 1) or if it the unit size would fit.
+								// This is where we do a difference also between the * (0+) and + cases; the + will have a force one repetition here, while the * will not.
+								int NumRepeats = Submodule.GetNumRepeat();
+								if (Submodule.GetNumRepeat() == PCGGrammar::AtLeastOneRepetition)
+								{
+									NumRepeats = 1;
+								}
+								else if (NumRepeats == PCGGrammar::InfiniteRepetition && RemainingLength >= Submodule.GetMinConcreteSize())
+								{
+									NumRepeats = 1;
+									RemainingLength -= Submodule.GetMinConcreteSize();
+								}
+
+								for (int R = 0; R < NumRepeats; ++R)
+								{
+									TModuleInstance<T>& ExpandedModule = ExpandedModules.Emplace_GetRef(&Submodule);
+								}
+
+								if (NumRepeats > 0)
+								{
+									bNotDone = true;
+									bHasPriorityOrStochasticModules |= (Submodule.GetType() == PCGGrammar::EModuleType::Priority || Submodule.GetType() == PCGGrammar::EModuleType::Stochastic);
+								}
+							}
+						}
+					}
+				}
+				else if (CurrentModule.Module->GetType() == PCGGrammar::EModuleType::Priority || CurrentModule.Module->GetType() == PCGGrammar::EModuleType::Stochastic)
+				{
+					bHasPriorityOrStochasticModules = true;
+
+					// Copy module as-is, but remove the expansion limitation so we break it down next iteration
+					TModuleInstance<T>& ExpandedModule = ExpandedModules.Add_GetRef(CurrentModule);
+					ExpandedModule.bSkipExpansion = false;
+				}
+				else
+				{
+					// Copy module as-is, but do NOT remove the expansion limitation, otherwise we will have issues with literals with repetitions.
+					TModuleInstance<T>& ExpandedModule = ExpandedModules.Add_GetRef(CurrentModule);
+				}
+			}
+
+			CurrentModules = MoveTemp(ExpandedModules);
+
+			// 2- Make a copy of the current set of modules, because we'll use this for repeatable modules
+			PreviousModules = CurrentModules;
+
+			// 3- Perform choices
+			if (bHasPriorityOrStochasticModules)
+			{
+				// 3.1 - Gather priority & stochastic modules (pointers here)
+				PriorityModules.Reset();
+				StochasticModules.Reset();
+
+				for (TModuleInstance<T>& CurrentModule : CurrentModules)
+				{
+					check(CurrentModule.Module);
+					if (CurrentModule.Module->GetType() == PCGGrammar::EModuleType::Priority)
+					{
+						PriorityModules.Add(&CurrentModule);
+					}
+					else if (CurrentModule.Module->GetType() == PCGGrammar::EModuleType::Stochastic)
+					{
+						StochasticModules.Add(&CurrentModule);
+					}
+				}
+
+				bNotDone |= (!PriorityModules.IsEmpty() || !StochasticModules.IsEmpty());
+
+				// 3.2 - Apply priority choices first (from left to right)
+				for (TModuleInstance<T>* PriorityModule : PriorityModules)
+				{
+					bool bWasReplaced = false;
+					// Replace current module by the first of its childen that has a min size that fits in with the remaining length
+					for (const T& PrioritySubmodule : PriorityModule->Module->Submodules)
+					{
+						if (!PrioritySubmodule.IsValid())
+						{
+							continue;
+						}
+
+						const double DeltaMinSize = (PrioritySubmodule.GetMinConcreteSize() - PriorityModule->Module->GetConcreteUnitSize());
+						check(DeltaMinSize >= 0);
+						if (RemainingLength >= DeltaMinSize)
+						{
+							PriorityModule->Module = &PrioritySubmodule;
+							RemainingLength -= DeltaMinSize;
+							bWasReplaced = true;
+							break;
+						}
+					}
+
+					if (!bWasReplaced)
+					{
+						// TODO Kill this module + reclaim size(?)
+						PriorityModule->bIsValid = false;
+					}
+				}
+
+				// 3.3 Process stochastic modules
+				// TODO random pick order of stochastic modules
+				for (TModuleInstance<T>* StochasticModule : StochasticModules)
+				{
+					// Replace current module by a random pick according to total weights of "valid" choices we can still make
+					int TotalWeight = 0;
+					for (const T& StochasticSubmodule : StochasticModule->Module->Submodules)
+					{
+						if (!StochasticSubmodule.IsValid())
+						{
+							continue;
+						}
+
+						const double DeltaMinSize = (StochasticSubmodule.GetMinConcreteSize() - StochasticModule->Module->GetConcreteUnitSize());
+						check(DeltaMinSize >= 0);
+						if (DeltaMinSize <= RemainingLength)
+						{
+							TotalWeight += StochasticSubmodule.GetWeight();
+						}
+					}
+
+					if (TotalWeight == 0)
+					{
+						// Kill module + reclaim size(?)
+						StochasticModule->bIsValid = false;
+						continue;
+					}
+
+					int WeightPick = RandomStream.RandRange(0, TotalWeight - 1);
+					bool bWasReplaced = false;
+
+					for (const T& StochasticSubmodule : StochasticModule->Module->Submodules)
+					{
+						if (!StochasticSubmodule.IsValid())
+						{
+							continue;
+						}
+
+						const double DeltaMinSize = (StochasticSubmodule.GetMinConcreteSize() - StochasticModule->Module->GetConcreteUnitSize());
+						check(DeltaMinSize >= 0);
+						if (DeltaMinSize <= RemainingLength)
+						{
+							if (StochasticSubmodule.GetWeight() > WeightPick)
+							{
+								bWasReplaced = true;
+								RemainingLength -= DeltaMinSize;
+								StochasticModule->Module = &StochasticSubmodule;
+								break;
+							}
+							else
+							{
+								WeightPick -= StochasticSubmodule.GetWeight();
+							}
+						}
+					}
+
+					check(bWasReplaced);
+				}
+			}
+
+			// 4. Add arbitrary repetitions, requires to keep data from previous iteration
+			// Implementation note: since we add one module only in each iteration, it will definitely be inefficient when there is a long chain of things to duplicate.
+			check(PreviousModules.Num() == CurrentModules.Num());
+			int NumAddedRepetitions = 0;
+
+			for (int ModuleIndex = 0; ModuleIndex < PreviousModules.Num(); ++ModuleIndex)
+			{
+				TModuleInstance<T>& PreviousModule = PreviousModules[ModuleIndex];
+				if (!PreviousModule.bSkipExpansion && PreviousModule.Module->GetNumRepeat() < 0 && PreviousModule.Module->GetMinConcreteSize() <= RemainingLength)
+				{
+					// Mark the module to skip expansion, otherwise we can't chain for repeats
+					PreviousModule.bSkipExpansion = true;
+					CurrentModules.Insert(PreviousModule, ModuleIndex + NumAddedRepetitions + 1);
+					++NumAddedRepetitions;
+					RemainingLength -= PreviousModule.Module->GetMinConcreteSize();
+				}
+			}
+
+			bNotDone |= (NumAddedRepetitions > 0);
+		}
+
+#if WITH_EDITOR
+		// Perform some early validation and see if there's a mismatch on the reported size and the one actually placed
+		double CountedLength = RemainingLength;
+		for (TModuleInstance<T>& CurrentModule : CurrentModules)
+		{
+			CountedLength += CurrentModule.Module->GetUnitSize();
+		}
+
+		ensure(FMath::Abs(CountedLength-Length) < 1.0);
+#endif
+
+		// 7. Finally, apply adjusted scales to modules that support it.
+		check(RemainingLength >= 0);
+		if (!FMath::IsNearlyZero(RemainingLength))
+		{
+			int32 NumScalableModules = 0;
+			double ScalableLength = 0;
+			for (TModuleInstance<T>& CurrentModule : CurrentModules)
+			{
+				if (CurrentModule.Module->IsScalable())
+				{
+					++NumScalableModules;
+					// implementation note: at this point we have only unit-literals, so we need to ignore repetitions, if any here, hence using the unit size
+					ScalableLength += CurrentModule.Module->GetUnitSize();
+				}
+			}
+
+			if (ScalableLength > 0)
+			{
+				for (TModuleInstance<T>& CurrentModule : CurrentModules)
+				{
+					if (CurrentModule.Module->IsScalable())
+					{
+						CurrentModule.ExtraScale = (RemainingLength / ScalableLength);
+					}
+				}
+
+				RemainingLength = 0;
+			}
+		}
+
+		OutModuleInstances = MoveTemp(CurrentModules);
+
+		return true;
+	}
+}
+
+class FPCGSubdivisionBaseElement : public IPCGElement
+{
+protected:
+	using FModuleInfoMap = PCGSubdivisionBase::FModuleInfoMap;
+
+	FModuleInfoMap GetModulesInfoMap(FPCGContext* InContext, const TArray<FPCGSubdivisionSubmodule>& SubmodulesInfo, const UPCGParamData*& OutModuleInfoParamData) const;
+	FModuleInfoMap GetModulesInfoMap(FPCGContext* InContext, const FPCGSubdivisionModuleAttributeNames& InSubdivisionModuleAttributeNames, const UPCGParamData*& OutModuleInfoParamData) const;
+	FModuleInfoMap GetModulesInfoMap(FPCGContext* InContext, const UPCGSubdivisionBaseSettings* InSettings, const UPCGParamData*& OutModuleInfoParamData) const;
+	PCGGrammar::FTokenizedGrammar GetTokenizedGrammar(FPCGContext* InContext, const UPCGData* InputData, const UPCGSubdivisionBaseSettings* InSettings, const FModuleInfoMap& InModulesInfo, double& OutMinSize) const;
+	TMap<FString, PCGGrammar::FTokenizedGrammar> GetTokenizedGrammarForPoints(FPCGContext* InContext, const UPCGPointData* InputData, const UPCGSubdivisionBaseSettings* InSettings, const FModuleInfoMap& InModulesInfo, double& OutMinSize) const;
+	bool MatchAndSetAttributes(const TArray<FPCGTaggedData>& InputData, TArray<FPCGTaggedData>& OutputData, const UPCGParamData* InModuleInfoParamData, const UPCGSubdivisionBaseSettings* InSettings) const;
+};
diff --git a/Engine/Plugins/PCG/Source/PCG/Public/Elements/Grammar/PCGVolumeSlicer.h b/Engine/Plugins/PCG/Source/PCG/Public/Elements/Grammar/PCGVolumeSlicer.h
deleted file mode 100644
