Added the collapse node that does decimation and averaging.

diff --git a/Engine/Plugins/PCG/Source/PCG/Private/Elements/PCGCollapsePoints.cpp b/Engine/Plugins/PCG/Source/PCG/Private/Elements/PCGCollapsePoints.cpp
new file mode 100644
--- /dev/null
+++ b/Engine/Plugins/PCG/Source/PCG/Private/Elements/PCGCollapsePoints.cpp
@@ -0,0 +1,509 @@
+// Copyright Epic Games, Inc. All Rights Reserved.
+
+#include "Elements/PCGCollapsePoints.h"
+
+#include "PCGContext.h"
+#include "Data/PCGPointData.h"
+#include "Helpers/PCGAsync.h"
+#include "Helpers/PCGHelpers.h"
+#include "Metadata/Accessors/PCGAttributeAccessorHelpers.h"
+#include "SpatialAlgo/PCGOctreeQueries.h"
+
+#include "Algo/MaxElement.h"
+
+#define LOCTEXT_NAMESPACE "PCGCollapsePointsElement"
+
+namespace PCGCollapsePoints
+{
+	namespace Algo
+	{
+		void MergePairs(FPCGContext* InContext, const FCollapsePointsSettings& Settings, FCollapsePointsState& OutState)
+		{
+			TRACE_CPUPROFILER_EVENT_SCOPE(PCGCollapsePointsElement::Algo::MergePairs);
+
+			check(!OutState.Selections.IsEmpty());
+
+			for (const TPair<int32, int32>& MergePair : OutState.Selections)
+			{
+				int PrimaryPointIndex = MergePair.Key;
+				int SecondaryPointIndex = MergePair.Value;
+
+				check(OutState.Merged[PrimaryPointIndex] == INDEX_NONE && OutState.Merged[SecondaryPointIndex] == INDEX_NONE);
+
+				const double& PrimaryWeight = OutState.Weights[PrimaryPointIndex];
+				const double& SecondaryWeight = OutState.Weights[SecondaryPointIndex];
+				const double WeightSum = PrimaryWeight + SecondaryWeight;
+
+				double Alpha = FMath::IsNearlyZero(WeightSum) ? 0.5 : (SecondaryWeight / WeightSum);
+				const FVector DeltaPosition = OutState.Positions[SecondaryPointIndex] - OutState.Positions[PrimaryPointIndex];
+
+				OutState.Positions[PrimaryPointIndex] = OutState.Positions[PrimaryPointIndex] + Alpha * DeltaPosition;
+				OutState.Weights[PrimaryPointIndex] = WeightSum;
+				OutState.Merged[SecondaryPointIndex] = PrimaryPointIndex;
+			}
+		}
+	}
+
+	namespace Modes
+	{
+		// Creates exclusive pairs of points to merge, in the visit order.
+		bool PairwiseSelection(FPCGContext* InContext, const FCollapsePointsSettings& Settings, FCollapsePointsState& OutState)
+		{
+			TRACE_CPUPROFILER_EVENT_SCOPE(PCGCollapsePointsElement::Modes::PairwiseSelection);
+			OutState.Selections.Reset();
+
+			// Note: can't do async trivially here as we will block off things downstream
+			// 1. Reset visited state
+			OutState.Visited.SetNumUninitialized(OutState.VisitOrder.Num(), EAllowShrinking::No);
+			for (bool& Visited : OutState.Visited)
+			{
+				Visited = false;
+			}
+
+			for (int VisitIndex = 0; VisitIndex < OutState.VisitOrder.Num(); ++VisitIndex)
+			{
+				int PointIndex = OutState.VisitOrder[VisitIndex];
+
+				if (OutState.Merged[PointIndex] != INDEX_NONE || OutState.Visited[PointIndex])
+				{
+					continue;
+				}
+
+				// this isn't strictly needed as we will never read from this value
+				OutState.Visited[PointIndex] = true;
+
+				// Find closest unvisited point after current point that's within the distance threshold.
+				double MinSqrDistance = Settings.DistanceThreshold * Settings.DistanceThreshold;
+				int ClosestUnvisitedIndex = INDEX_NONE;
+
+				for (int NeighborVisitIndex = VisitIndex + 1; NeighborVisitIndex < OutState.VisitOrder.Num(); ++NeighborVisitIndex)
+				{
+					int NeighborIndex = OutState.VisitOrder[NeighborVisitIndex];
+
+					if (OutState.Merged[NeighborIndex] != INDEX_NONE || OutState.Visited[NeighborIndex])
+					{
+						continue;
+					}
+
+					double SqrDistance = (OutState.Positions[PointIndex] - OutState.Positions[NeighborIndex]).SquaredLength();
+
+					if (SqrDistance < MinSqrDistance)
+					{
+						MinSqrDistance = SqrDistance;
+						ClosestUnvisitedIndex = NeighborIndex;
+
+						// Since we're visiting in order, this is the best pick we'll have.
+						if (FMath::IsNearlyZero(MinSqrDistance))
+						{
+							break;
+						}
+					}
+				}
+
+				if (ClosestUnvisitedIndex != INDEX_NONE)
+				{
+					OutState.Visited[ClosestUnvisitedIndex] = true;
+					OutState.Selections.Emplace(PointIndex, ClosestUnvisitedIndex);
+				}
+			}
+
+			return !OutState.Selections.IsEmpty();
+		}
+
+		bool ClosestPairSelection(FPCGContext* InContext, const FCollapsePointsSettings& Settings, FCollapsePointsState& OutState)
+		{
+			// TODO - not implemented. Needs to find the absolute closest pair and insert it in the selections array.
+			return true;
+		}
+	} // namespace Modes
+} // namespace PCGCollapsePoints
+
+UPCGCollapsePointsSettings::UPCGCollapsePointsSettings()
+{
+	FPCGAttributePropertyOutputNoSourceSelector& DefaultAttribute = AttributesToMerge.Emplace_GetRef();
+	DefaultAttribute.SetPointProperty(EPCGPointProperties::Position);
+}
+
+#if WITH_EDITOR
+FText UPCGCollapsePointsSettings::GetNodeTooltipText() const
+{
+	return LOCTEXT("NodeTooltip", "Collapses points with their closest neighbors until all points are farther than the search distance.");
+}
+#endif // WITH_EDITOR
+
+TArray<FPCGPinProperties> UPCGCollapsePointsSettings::InputPinProperties() const
+{
+	TArray<FPCGPinProperties> PinProperties;
+	PinProperties.Emplace_GetRef(PCGPinConstants::DefaultInputLabel, EPCGDataType::Point).SetRequiredPin();
+	return PinProperties;
+}
+
+bool UPCGCollapsePointsSettings::UseSeed() const
+{
+	return VisitOrder == EPCGCollapseVisitOrder::Random;
+}
+
+FPCGElementPtr UPCGCollapsePointsSettings::CreateElement() const
+{
+	return MakeShared<FPCGCollapsePointsElement>();
+}
+
+bool FPCGCollapsePointsElement::PrepareDataInternal(FPCGContext* InContext) const
+{
+	TRACE_CPUPROFILER_EVENT_SCOPE(FPCGCollapsePointsElement::PrepareData);
+
+	const UPCGCollapsePointsSettings* Settings = InContext->GetInputSettings<UPCGCollapsePointsSettings>();
+	check(Settings);
+
+	ContextType* Context = static_cast<ContextType*>(InContext);
+	check(Context);
+
+	const TArray<FPCGTaggedData> Inputs = Context->InputData.GetInputsByPin(PCGPinConstants::DefaultInputLabel);
+
+	if (Inputs.IsEmpty())
+	{
+		return true;
+	}
+
+	Context->InitializePerExecutionState([Settings](const ContextType*, ExecStateType& OutState)
+	{
+		PCGCollapsePoints::FCollapsePointsSettings::PairSelectionFuncType PairSelectionFunc = nullptr;
+
+		if (Settings->Mode == EPCGCollapseMode::PairwiseClosest)
+		{
+			PairSelectionFunc = PCGCollapsePoints::Modes::PairwiseSelection;
+		}
+		else if(Settings->Mode == EPCGCollapseMode::AbsoluteClosest)
+		{
+			PairSelectionFunc = PCGCollapsePoints::Modes::ClosestPairSelection;
+		}
+		else
+		{
+			return EPCGTimeSliceInitResult::AbortExecution;
+		}
+
+		check(PairSelectionFunc);
+		OutState = PCGCollapsePoints::FCollapsePointsSettings
+		{
+			.Settings = Settings,
+			.PairSelectionFunc = PairSelectionFunc,
+			.MergeSelectionFunc = PCGCollapsePoints::Algo::MergePairs,
+			.DistanceThreshold = Settings->DistanceThreshold
+		};
+
+		return EPCGTimeSliceInitResult::Success;
+	});
+
+	Context->InitializePerIterationStates(Inputs.Num(), [&Inputs, Settings, Context](IterStateType& OutState, const ExecStateType&, const uint32 IterationIndex)
+	{
+		const UPCGPointData* PointData = CastChecked<UPCGPointData>(Inputs[IterationIndex].Data);
+
+		if (!PointData || PointData->IsEmpty())
+		{
+			return EPCGTimeSliceInitResult::NoOperation;
+		}
+
+		OutState.SourceData = PointData;
+
+		const TArray<FPCGPoint>& Points = PointData->GetPoints();
+		const int32 NumPoints = Points.Num();
+
+		// Get positions
+		OutState.Positions.Reserve(Points.Num());
+
+		for (const FPCGPoint& Point : Points)
+		{
+			OutState.Positions.Add(Point.Transform.GetLocation());
+		}
+
+		// Get weights
+		OutState.Weights.SetNumUninitialized(NumPoints);
+
+		if (Settings->bUseMergeWeightAttribute)
+		{
+			const FPCGAttributePropertyInputSelector Selector = Settings->MergeWeightAttribute.CopyAndFixLast(PointData);
+			const TUniquePtr<const IPCGAttributeAccessor> Accessor = PCGAttributeAccessorHelpers::CreateConstAccessor(PointData, Selector);
+			const TUniquePtr<const IPCGAttributeAccessorKeys> Keys = PCGAttributeAccessorHelpers::CreateConstKeys(PointData, Selector);
+			if (!Accessor || !Keys)
+			{
+				PCGLog::Metadata::LogFailToCreateAccessorError(Selector, Context);
+				return EPCGTimeSliceInitResult::NoOperation;
+			}
+			
+			if (!Accessor->GetRange(TArrayView<double>(OutState.Weights), 0, *Keys))
+			{
+				PCGLog::Metadata::LogFailToGetAttributeError(Selector, Context);
+				return EPCGTimeSliceInitResult::NoOperation;
+			}
+
+			// Fixup weights so that they are not negative
+			for (double& Weight : OutState.Weights)
+			{
+				Weight = FMath::Max(Weight, 0.0);
+			}
+		}
+		else
+		{
+			// Default uniform weights
+			for (double& Weight : OutState.Weights)
+			{
+				Weight = 1.0;
+			}
+		}
+
+		// We need a copy of the weights because we'll change it over the process of merging, but need the original weights at the end of the process
+		OutState.OriginalWeights = OutState.Weights;
+
+		// Build visit order
+		if (Settings->Mode == EPCGCollapseMode::PairwiseClosest)
+		{
+			OutState.VisitOrder.SetNumUninitialized(NumPoints);
+
+			for (int Index = 0; Index < NumPoints; ++Index)
+			{
+				OutState.VisitOrder[Index] = Index;
+			}
+
+			if (Settings->VisitOrder == EPCGCollapseVisitOrder::Ordered)
+			{
+				// Nothing to do, initial assignment is fine
+			} 
+			else if(Settings->VisitOrder == EPCGCollapseVisitOrder::Random)
+			{
+				FRandomStream RandomStream(Context->GetSeed());
+				PCGHelpers::ShuffleArray(RandomStream, OutState.VisitOrder);
+			}
+			else if (Settings->VisitOrder == EPCGCollapseVisitOrder::MinAttribute || Settings->VisitOrder == EPCGCollapseVisitOrder::MaxAttribute)
+			{
+				const FPCGAttributePropertyInputSelector VisitOrderSelector = Settings->VisitOrderAttribute.CopyAndFixLast(PointData);
+				TUniquePtr<const IPCGAttributeAccessor> VisitOrderAccessor = PCGAttributeAccessorHelpers::CreateConstAccessor(PointData, VisitOrderSelector);
+				TUniquePtr<const IPCGAttributeAccessorKeys> VisitOrderKeys = PCGAttributeAccessorHelpers::CreateConstKeys(PointData, VisitOrderSelector);
+				if (!VisitOrderAccessor || !VisitOrderKeys)
+				{
+					PCGLog::Metadata::LogFailToCreateAccessorError(VisitOrderSelector, Context);
+					return EPCGTimeSliceInitResult::NoOperation;
+				}
+
+				PCGAttributeAccessorHelpers::SortByAttribute(*VisitOrderAccessor, *VisitOrderKeys, OutState.VisitOrder, /*bAscending=*/Settings->VisitOrder == EPCGCollapseVisitOrder::MinAttribute);
+			}
+			else
+			{
+				checkNoEntry();
+				return EPCGTimeSliceInitResult::NoOperation;
+			}
+		}
+
+		// Initialize merged state to "not merged"
+		OutState.Merged.SetNumUninitialized(NumPoints);
+		for (int& MergeIndex : OutState.Merged)
+		{
+			MergeIndex = INDEX_NONE;
+		}
+
+		// Create the output point data now. Operate on a copy of the input point data.
+		UPCGPointData* OutPointData = FPCGContext::NewObject_AnyThread<UPCGPointData>(Context);
+		OutPointData->InitializeFromData(PointData);
+		OutState.OutData = OutPointData;
+
+		// Initialize merge accessors
+		for(const FPCGAttributePropertyOutputNoSourceSelector& InAttributeToMerge : Settings->AttributesToMerge)
+		{
+			FPCGAttributePropertyInputSelector SourceAttributeToMerge;
+			SourceAttributeToMerge.ImportFromOtherSelector(InAttributeToMerge);
+			
+			const FPCGAttributePropertyInputSelector AttributeToMerge = SourceAttributeToMerge.CopyAndFixLast(PointData);
+			TUniquePtr<const IPCGAttributeAccessor> AttributeAccessor = PCGAttributeAccessorHelpers::CreateConstAccessor(PointData, AttributeToMerge);
+
+			if (!AttributeAccessor)
+			{
+				PCGLog::Metadata::LogFailToCreateAccessorError(AttributeToMerge, Context);
+				continue;
+			}
+
+			FPCGAttributePropertyOutputSelector OutputAttributeSelector;
+			OutputAttributeSelector.ImportFromOtherSelector(AttributeToMerge); 
+
+			TUniquePtr<IPCGAttributeAccessor> OutputAttributeAccessor = PCGAttributeAccessorHelpers::CreateAccessor(OutPointData, OutputAttributeSelector);
+			if (!OutputAttributeAccessor)
+			{
+				PCGLog::Metadata::LogFailToCreateAccessorError(OutputAttributeSelector, Context);
+				continue;
+			}
+
+			OutState.SourceMergeAccessors.Add(std::move(AttributeAccessor));
+			OutState.OutputMergeAccessors.Add(std::move(OutputAttributeAccessor));
+		}
+
+		OutState.SourceMergeKeys = PCGAttributeAccessorHelpers::CreateConstKeys(PointData, FPCGAttributePropertyInputSelector());
+		if (!OutState.SourceMergeKeys)
+		{
+			return EPCGTimeSliceInitResult::NoOperation;
+		}
+
+		Context->OutputData.TaggedData.Emplace_GetRef().Data = OutPointData;
+		return EPCGTimeSliceInitResult::Success;
+	});
+
+	return true;
+}
+
+bool FPCGCollapsePointsElement::ExecuteInternal(FPCGContext* InContext) const
+{
+	TRACE_CPUPROFILER_EVENT_SCOPE(FPCGCollapsePointsElement::Execute);
+
+	ContextType* TimeSlicedContext = static_cast<ContextType*>(InContext);
+	check(TimeSlicedContext);
+
+	if (!TimeSlicedContext->DataIsPreparedForExecution())
+	{
+		return true;
+	}
+
+	if (TimeSlicedContext->GetExecutionStateResult() == EPCGTimeSliceInitResult::NoOperation)
+	{
+		TimeSlicedContext->OutputData = TimeSlicedContext->InputData;
+		return true;
+	}
+
+	const UPCGCollapsePointsSettings* Settings = TimeSlicedContext->GetInputSettings<UPCGCollapsePointsSettings>();
+	check(Settings);
+
+	return ExecuteSlice(TimeSlicedContext, [InContext](const ContextType* Context, const ExecStateType& CollapseSettings, IterStateType& CollapseState, const uint32 IterIndex)
+	{
+		if (Context->GetIterationStateResult(IterIndex) == EPCGTimeSliceInitResult::NoOperation)
+		{
+			return true;
+		}
+
+		// A single iteration is:
+		// - build list of pairs (all OR single closest pair)
+		// - merge them, update weights
+		// continue until nothing is merged
+		if (CollapseSettings.PairSelectionFunc(InContext, CollapseSettings, CollapseState))
+		{
+			// Merge based on selection
+			CollapseSettings.MergeSelectionFunc(InContext, CollapseSettings, CollapseState);
+			return false;
+		}
+		else
+		{
+			TArray<int>& Merged = CollapseState.Merged;
+
+			// We're done - compute final results
+			// Go up the merge hierarchy to point to the root
+			for (int PointIndex = 0; PointIndex < Merged.Num(); ++PointIndex)
+			{
+				while (Merged[PointIndex] != INDEX_NONE && Merged[Merged[PointIndex]] != INDEX_NONE)
+				{
+					Merged[PointIndex] = Merged[Merged[PointIndex]];
+				}
+			}
+
+			// Finally, make unmerged points/roots point to themselves
+			for (int PointIndex = 0; PointIndex < Merged.Num(); ++PointIndex)
+			{
+				if (Merged[PointIndex] == INDEX_NONE)
+				{
+					Merged[PointIndex] = PointIndex;
+				}
+			}
+
+			// Second, partition by merge target
+			TMap<int, TArray<int, TInlineAllocator<8>>> Partition;
+
+			for (int PointIndex = 0; PointIndex < Merged.Num(); ++PointIndex)
+			{
+				Partition.FindOrAdd(Merged[PointIndex]).Add(PointIndex);
+			}
+
+			// Create one point per partition
+			check(CollapseState.SourceData);
+			const TArray<FPCGPoint>& SourcePoints = CollapseState.SourceData->GetPoints();
+			check(CollapseState.OutData);
+			TArray<FPCGPoint>& OutPoints = CollapseState.OutData->GetMutablePoints();
+
+			// Need to allocate prior to key creation
+			OutPoints.SetNum(Partition.Num());
+
+			// Build accessor list(s) + keys, based on settings
+			TArray<double> PartitionWeights;
+			TUniquePtr<IPCGAttributeAccessorKeys> OutputKeys = PCGAttributeAccessorHelpers::CreateKeys(CollapseState.OutData, FPCGAttributePropertyOutputSelector());
+			check(OutputKeys);
+			TUniquePtr<const IPCGAttributeAccessorKeys>& SourceKeys = CollapseState.SourceMergeKeys;
+			check(SourceKeys);
+			
+			// TODO: do this in parallel
+			int PartitionIndex = 0;
+			for(const TPair<int, TArray<int, TInlineAllocator<8>>>& PartitionEntry : Partition)
+			{
+				int PrimaryPointIndex = PartitionEntry.Key;
+				const TArray<int, TInlineAllocator<8>>& SecondaryPointIndices = PartitionEntry.Value;
+
+				// First, copy the primary point as-is.
+				FPCGPoint& MergedPoint = OutPoints[PartitionIndex];
+				MergedPoint = SourcePoints[PrimaryPointIndex];
+
+				if (SecondaryPointIndices.Num() == 1)
+				{
+					++PartitionIndex;
+					continue;
+				}
+
+				// Prepare respective weights
+				double TotalWeight = 0;
+				PartitionWeights.SetNumUninitialized(SecondaryPointIndices.Num());
+				for(int SecondaryIndex = 0; SecondaryIndex < SecondaryPointIndices.Num(); ++SecondaryIndex)
+				{
+					const double& OriginalWeight = CollapseState.OriginalWeights[SecondaryPointIndices[SecondaryIndex]];
+					TotalWeight += OriginalWeight;
+					PartitionWeights[SecondaryIndex] = OriginalWeight;
+				}
+
+				for (double& PartitionWeight : PartitionWeights)
+				{
+					PartitionWeight = (FMath::IsNearlyZero(TotalWeight) ? 1.0 / PartitionWeights.Num() : PartitionWeight / TotalWeight);
+				}
+
+				for (int AccessorIndex = 0; AccessorIndex < CollapseState.SourceMergeAccessors.Num(); ++AccessorIndex)
+				{
+					TUniquePtr<const IPCGAttributeAccessor>& SourceAccessor = CollapseState.SourceMergeAccessors[AccessorIndex];
+					TUniquePtr<IPCGAttributeAccessor>& OutputAccessor = CollapseState.OutputMergeAccessors[AccessorIndex];
+
+					auto Callback = [PartitionIndex, &SecondaryPointIndices, &PartitionWeights, &SourceAccessor, &SourceKeys, &OutputAccessor, &OutputKeys](auto Dummy)
+					{
+						using OutputType = decltype(Dummy);
+
+						if constexpr (PCG::Private::MetadataTraits<OutputType>::CanInterpolate)
+						{
+							OutputType Value = PCG::Private::MetadataTraits<OutputType>::ZeroValueForWeightedSum();
+
+							for (int Index = 0; Index < SecondaryPointIndices.Num(); ++Index)
+							{
+								int PointIndex = SecondaryPointIndices[Index];
+
+								OutputType SecondaryValue{};
+								SourceAccessor->Get(SecondaryValue, PointIndex, *SourceKeys);
+
+								Value = PCG::Private::MetadataTraits<OutputType>::WeightedSum(Value, SecondaryValue, PartitionWeights[Index]);
+							}
+
+							OutputAccessor->Set(Value, PartitionIndex, *OutputKeys);
+						}
+						else
+						{
+							// Nothing to do because the default point will already have the value it would end up with.
+						}
+					};
+
+					PCGMetadataAttribute::CallbackWithRightType(OutputAccessor->GetUnderlyingType(), Callback);
+				}
+
+				++PartitionIndex;
+			}
+
+			return true;
+		}
+	});
+}
+
+#undef LOCTEXT_NAMESPACE
diff --git a/Engine/Plugins/PCG/Source/PCG/Public/Elements/PCGCollapsePoints.h b/Engine/Plugins/PCG/Source/PCG/Public/Elements/PCGCollapsePoints.h
new file mode 100644
--- /dev/null
+++ b/Engine/Plugins/PCG/Source/PCG/Public/Elements/PCGCollapsePoints.h
@@ -0,0 +1,129 @@
+// Copyright Epic Games, Inc. All Rights Reserved.
+
+#pragma once
+
+#include "PCGSettings.h"
+
+#include "Elements/PCGTimeSlicedElementBase.h"
+#include "Metadata/PCGMetadataAttributeTpl.h"
+#include "Metadata/Accessors/PCGPropertyAccessor.h"
+
+#include "PCGCollapsePoints.generated.h"
+
+class UPCGPointData;
+
+UENUM(BlueprintType)
+enum class EPCGCollapseMode : uint8
+{
+	PairwiseClosest UMETA(Tooltip="At every iteration, will select exclusive pairs that are closest."),
+	AbsoluteClosest UMETA(Hidden, Tooltip="At every iteration, will select the closest pair to merge. Performance warning.")
+};
+
+UENUM(BlueprintType)
+enum class EPCGCollapseVisitOrder : uint8
+{
+	Ordered UMETA(Tooltip="Will create pairs based on original point order."),
+	Random UMETA(Tooltip="Will generate a random ordering that will drive pair order creation."),
+	MinAttribute UMETA(Tooltip="Will create pairs according to attribute value order (minimum value first)."),
+	MaxAttribute UMETA(Tooltip = "Will create pairs according to attribute value order (maximum value first).")
+};
+
+/** Collapses (decimates) points on a closest-point basis, allowing weighted averages to be computed as required. */
+UCLASS(MinimalAPI, BlueprintType, ClassGroup = (Procedural))
+class UPCGCollapsePointsSettings : public UPCGSettings
+{
+	GENERATED_BODY()
+
+public:
+	UPCGCollapsePointsSettings();
+
+	//~Begin UPCGSettings interface
+#if WITH_EDITOR
+	virtual FName GetDefaultNodeName() const override { return FName(TEXT("CollapsePoints")); }
+	virtual FText GetDefaultNodeTitle() const override { return NSLOCTEXT("PCGCollapsePointsElement", "NodeTitle", "Collapse Points"); }
+	virtual EPCGSettingsType GetType() const override { return EPCGSettingsType::PointOps; }
+	virtual FText GetNodeTooltipText() const override;
+#endif
+	virtual bool UseSeed() const override;
+
+protected:
+	virtual TArray<FPCGPinProperties> InputPinProperties() const override;
+	virtual TArray<FPCGPinProperties> OutputPinProperties() const override { return Super::DefaultPointOutputPinProperties(); }
+	virtual FPCGElementPtr CreateElement() const override;
+	//~End UPCGSettings interface
+
+public:
+	/** Distance at which we will stop collapsing points. E.g. Points will continue collapsing until every point is at least this distance from each other. */
+	UPROPERTY(EditAnywhere, BlueprintReadWrite, Category = Settings, meta = (ClampMin = "0.01", PCG_Overridable))
+	double DistanceThreshold = 100.0;
+
+	UPROPERTY(EditAnywhere, BlueprintReadWrite, Category = Settings, meta = (PCG_Overridable))
+	EPCGCollapseMode Mode = EPCGCollapseMode::PairwiseClosest;
+
+	/** Determines order in which we will collapse points pair-wise. */
+	UPROPERTY(EditAnywhere, BlueprintReadWrite, Category = Settings, meta = (EditCondition="Mode == EPCGCollapseMode::PairwiseClosest", EditConditionHides, PCG_Overridable))
+	EPCGCollapseVisitOrder VisitOrder = EPCGCollapseVisitOrder::Ordered;
+
+	/** Attribute to drive visit order. */
+	UPROPERTY(EditAnywhere, BlueprintReadWrite, Category = Settings, meta = (EditCondition="Mode == EPCGCollapseMode::PairwiseClosest && (VisitOrder == EPCGCollapseVisitOrder::MinAttribute || VisitOrder == EPCGCollapseVisitOrder::MaxAttribute)", EditConditionHides, PCG_Overridable))
+	FPCGAttributePropertyInputSelector VisitOrderAttribute;
+
+	/** Controls whether input points will use a weight driven by an attribute */
+	UPROPERTY(EditAnywhere, BlueprintReadWrite, Category = Settings, meta = (InlineEditConditionToggle, PCG_Overridable))
+	bool bUseMergeWeightAttribute = false;
+
+	/** Attribute that will drive relative weight when merging points (only in the Weighted mode). */
+	UPROPERTY(EditAnywhere, BlueprintReadWrite, Category = Settings, meta = (EditCondition = "bUseMergeWeightAttribute", PCG_Overridable))
+	FPCGAttributePropertyInputSelector MergeWeightAttribute;
+
+	/** List of attributes to merge on the final points, based on the weights. */
+	UPROPERTY(EditAnywhere, BlueprintReadWrite, Category = Settings)
+	TArray<FPCGAttributePropertyOutputNoSourceSelector> AttributesToMerge;
+};
+
+namespace PCGCollapsePoints
+{
+	struct FCollapsePointsState
+	{
+		// Current iterative state during merge
+		TArray<FVector> Positions;
+		TArray<double> Weights;
+		TArray<int> Merged;
+		
+		// Selection state
+		TArray<TPair<int32, int32>> Selections;
+		TArray<bool> Visited;
+
+		// Computed once in the preparation stage
+		TArray<int> VisitOrder;
+		TArray<double> OriginalWeights;
+
+		// Output data
+		const UPCGPointData* SourceData = nullptr;
+		UPCGPointData* OutData = nullptr;
+
+		TArray<TUniquePtr<const IPCGAttributeAccessor>> SourceMergeAccessors;
+		TUniquePtr<const IPCGAttributeAccessorKeys> SourceMergeKeys;
+		TArray<TUniquePtr<IPCGAttributeAccessor>> OutputMergeAccessors;
+	};
+
+	struct FCollapsePointsSettings
+	{
+		const UPCGCollapsePointsSettings* Settings = nullptr;
+
+		using PairSelectionFuncType = bool(*)(FPCGContext* /*InContext*/, const FCollapsePointsSettings& /*Settings*/, FCollapsePointsState& /*OutState*/);
+		PairSelectionFuncType PairSelectionFunc = nullptr;
+
+		using MergeSelectionFuncType = void(*)(FPCGContext* /*InContext*/, const FCollapsePointsSettings& /*Settings*/, FCollapsePointsState& /*OutState*/);
+		MergeSelectionFuncType MergeSelectionFunc = nullptr;
+
+		double DistanceThreshold = 0.0;
+	};
+} // namespace PCGCollapsePoints
+
+class FPCGCollapsePointsElement : public TPCGTimeSlicedElementBase<PCGCollapsePoints::FCollapsePointsSettings, PCGCollapsePoints::FCollapsePointsState>
+{
+protected:
+	virtual bool PrepareDataInternal(FPCGContext* InContext) const override;
+	virtual bool ExecuteInternal(FPCGContext* InContext) const override;
+};
