[PCG] Add Graph Execution Id and use it to move PCG Cached Data  out of PCG Components
- Add a PreGraphElement when executing top level graphs which garantees the availability of GetActorPCGData and GetOriginalActorPCGData for other elements GetDependenciesCRC
- This allows to restore MT execution of PreExecute on a bunch of Elements that relied on this data to compute their dependencies

#rnx

diff --git a/Engine/Plugins/PCG/Source/PCG/Private/Data/PCGGetDataFunctionRegistry.cpp b/Engine/Plugins/PCG/Source/PCG/Private/Data/PCGGetDataFunctionRegistry.cpp
--- a/Engine/Plugins/PCG/Source/PCG/Private/Data/PCGGetDataFunctionRegistry.cpp
+++ b/Engine/Plugins/PCG/Source/PCG/Private/Data/PCGGetDataFunctionRegistry.cpp
@@ -2,6 +2,7 @@
 
 #include "Data/PCGGetDataFunctionRegistry.h"
 
+#include "PCGCommon.h"
 #include "PCGComponent.h"
 #include "PCGContext.h"
 #include "PCGData.h"
@@ -137,13 +138,16 @@ int FPCGGetDataFunctionRegistry::DefaultDataFromActor(FPCGContext* InContext, co
 		{
 			VolumeData->Initialize(PartitionActor->GetFixedBounds());
 
-			UPCGComponent* OriginalComponent = Component ? PartitionActor->GetOriginalComponent(Component) : nullptr;
-			// Important note: we do NOT call the collection version here, as we want to have a union if that's the case
-			const UPCGSpatialData* OriginalComponentSpatialData = OriginalComponent ? Cast<const UPCGSpatialData>(OriginalComponent->GetActorPCGData()) : nullptr;
-
-			if (OriginalComponentSpatialData)
+			if (Component && PartitionActor->GetOriginalComponent(Component))
 			{
-				Result = Result->IntersectWith(nullptr, OriginalComponentSpatialData);
+				PCG_EXECUTION_CACHE_VALIDATION_CREATE_ORIGINAL_SCOPE(const_cast<UPCGComponent*>(Component));
+				// Important note: we do NOT call the collection version here, as we want to have a union if that's the case
+				const UPCGSpatialData* OriginalComponentSpatialData = Cast<const UPCGSpatialData>(Component->GetOriginalActorPCGData());
+
+				if (OriginalComponentSpatialData)
+				{
+					Result = Result->IntersectWith(nullptr, OriginalComponentSpatialData);
+				}
 			}
 		}
 
diff --git a/Engine/Plugins/PCG/Source/PCG/Private/Elements/PCGCreatePoints.cpp b/Engine/Plugins/PCG/Source/PCG/Private/Elements/PCGCreatePoints.cpp
--- a/Engine/Plugins/PCG/Source/PCG/Private/Elements/PCGCreatePoints.cpp
+++ b/Engine/Plugins/PCG/Source/PCG/Private/Elements/PCGCreatePoints.cpp
@@ -139,18 +139,6 @@ bool FPCGCreatePointsElement::ExecuteInternal(FPCGContext* Context) const
 	return true;
 }
 
-bool FPCGCreatePointsElement::CanExecuteOnlyOnMainThread(FPCGContext* Context) const 
-{
-	if (Context && Context->CurrentPhase == EPCGExecutionPhase::Execute)
-	{
-		const UPCGCreatePointsSettings* Settings = Context->GetInputSettings<UPCGCreatePointsSettings>();
-		check(Settings);
-		return Settings->bCullPointsOutsideVolume;
-	}
-
-	return false;
-}
-
 bool FPCGCreatePointsElement::IsCacheable(const UPCGSettings* InSettings) const
 {
 	const UPCGCreatePointsSettings* Settings = Cast<const UPCGCreatePointsSettings>(InSettings);
diff --git a/Engine/Plugins/PCG/Source/PCG/Private/Elements/PCGElevationIsolines.cpp b/Engine/Plugins/PCG/Source/PCG/Private/Elements/PCGElevationIsolines.cpp
--- a/Engine/Plugins/PCG/Source/PCG/Private/Elements/PCGElevationIsolines.cpp
+++ b/Engine/Plugins/PCG/Source/PCG/Private/Elements/PCGElevationIsolines.cpp
@@ -49,17 +49,6 @@ TArray<FPCGPinProperties> UPCGElevationIsolinesSettings::OutputPinProperties() c
 	return Properties;
 }
 
-bool FPCGElevationIsolinesElement::CanExecuteOnlyOnMainThread(FPCGContext* Context) const
-{
-	// Dependency on GetActorPCGData requires to run on Main Thread
-	if (Context && Context->CurrentPhase == EPCGExecutionPhase::Execute)
-	{
-		return true;
-	}
-
-	return false;
-}
-
 bool FPCGElevationIsolinesElement::ExecuteInternal(FPCGContext* InContext) const
 {
 	TRACE_CPUPROFILER_EVENT_SCOPE(FPCGElevationIsolinesElement::Execute);
diff --git a/Engine/Plugins/PCG/Source/PCG/Private/Graph/PCGGraphExecutor.cpp b/Engine/Plugins/PCG/Source/PCG/Private/Graph/PCGGraphExecutor.cpp
--- a/Engine/Plugins/PCG/Source/PCG/Private/Graph/PCGGraphExecutor.cpp
+++ b/Engine/Plugins/PCG/Source/PCG/Private/Graph/PCGGraphExecutor.cpp
@@ -198,6 +198,11 @@ const FPCGStack* FPCGGraphTask::GetStack() const
 	return (StackContext && StackIndex != INDEX_NONE) ? StackContext->GetStack(StackIndex) : nullptr;
 }
 
+FPCGTaskId FPCGGraphTask::GetGraphExecutionTaskId() const
+{
+	return StackContext ? StackContext->GetGraphExecutionTaskId() : InvalidPCGTaskId;
+}
+
 #if WITH_EDITOR
 void FPCGGraphTask::LogVisual(ELogVerbosity::Type InVerbosity, const FText& InMessage) const
 {
@@ -333,6 +338,8 @@ FPCGGraphExecutor::FPCGGraphExecutor()
 FPCGGraphExecutor::FPCGGraphExecutor(UWorld* InWorld)
 	: World(InWorld)
 {
+	static_assert((int8)(EExecutionCacheDataType::Count) == ExecutionCacheSize);
+
 	GameThreadHandler = MakeShared<FGameThreadHandler>(this);
 }
 
@@ -361,7 +368,7 @@ FPCGTaskId FPCGGraphExecutor::Schedule(UPCGComponent* Component, const TArray<FP
 	check(Component);
 	UPCGGraph* Graph = Component->GetGraph();
 
-	return Schedule(Graph, Component, FPCGElementPtr(), GetFetchInputElement(), ExternalDependencies, InFromStack, /*bAllowHierarchicalGeneration=*/true);
+	return Schedule(Graph, Component, GetPreGraphElement(), GetFetchInputElement(), ExternalDependencies, InFromStack, /*bAllowHierarchicalGeneration=*/true);
 }
 
 FPCGTaskId FPCGGraphExecutor::Schedule(
@@ -437,6 +444,9 @@ FPCGTaskId FPCGGraphExecutor::Schedule(
 		NextTaskId += ScheduledTask.Tasks.Num();
 		ScheduledId = NextTaskId - 1; // This is true because the last task is from the output node or is the post-execute task
 
+		const FPCGTaskId GraphExecutionTaskId = InFromStack ? InFromStack->GetGraphExecutionTaskId() : ScheduledId;
+		StackContextPtr->SetGraphExecutionTaskId(GraphExecutionTaskId);
+
 		// Push task (not data) dependencies on the pre-execute task
 		// Note must be done after the offset ids, otherwise we'll break the dependencies
 		check(ScheduledTask.Tasks.Num() >= 2 && ScheduledTask.Tasks[ScheduledTask.Tasks.Num() - 2].Node == nullptr);
@@ -764,6 +774,31 @@ int32 FPCGGraphExecutor::GetNonScheduledRemainingTaskCount() const
 	return Tasks.Num() + ReadyTasks.Num() + ActiveTasks.Num() + ActiveTasksGameThreadOnly.Num() + SleepingTasks.Num();
 }
 
+UPCGData* FPCGGraphExecutor::GetExecutionCacheData(FPCGTaskId InGraphExecutionTaskId, EExecutionCacheDataType InExecutionCacheDataType)
+{
+	if (InGraphExecutionTaskId != InvalidPCGTaskId)
+	{
+		PCGGraphExecutor::TScopeLock ScopeLock(GraphExecutionCachesLock);
+		if (FGraphExecutionCache* GraphExecutionCache = GraphExecutionCaches.Find(InGraphExecutionTaskId))
+		{
+			return GraphExecutionCache->Data[(int8)InExecutionCacheDataType];
+		}
+	}
+
+	return nullptr;
+}
+
+void FPCGGraphExecutor::SetExecutionCacheData(FPCGTaskId InGraphExecutionTaskId, EExecutionCacheDataType InExecutionCacheDataType, UPCGData* InData)
+{
+	if (InGraphExecutionTaskId != InvalidPCGTaskId)
+	{
+		PCGGraphExecutor::TScopeLock ScopeLock(GraphExecutionCachesLock);
+		FGraphExecutionCache& GraphExecutionCache = GraphExecutionCaches.FindOrAdd(InGraphExecutionTaskId);
+		check(!GraphExecutionCache.Data[(int8)InExecutionCacheDataType]);
+		GraphExecutionCache.Data[(int8)InExecutionCacheDataType] = InData;
+	}
+}
+
 FPCGTaskId FPCGGraphExecutor::ScheduleGeneric(TFunction<bool()> InOperation, UPCGComponent* InSourceComponent, const TArray<FPCGTaskId>& TaskExecutionDependencies)
 {
 	return ScheduleGeneric(
@@ -2353,6 +2388,9 @@ void FPCGGraphExecutor::ClearResults()
 		if (ScheduledTasks.IsEmpty())
 		{
 			NextTaskId = 0;
+
+			PCGGraphExecutor::TScopeLock ChildScopeLock(GraphExecutionCachesLock);
+			GraphExecutionCaches.Empty();
 		}
 	}
 
@@ -2585,16 +2623,50 @@ void FPCGGraphExecutor::AddReferencedObjects(FReferenceCollector& Collector)
 			}
 		});
 	}
+
+	{
+		PCGGraphExecutor::TScopeLock ScopeLock(GraphExecutionCachesLock);
+		for (TMap<FPCGTaskId, FGraphExecutionCache>::TIterator It = GraphExecutionCaches.CreateIterator(); It; ++It)
+		{
+			It.Value().AddStructReferencedObjects(Collector);
+		}
+	}
+}
+
+FPCGGraphExecutor::FGraphExecutionCache::FGraphExecutionCache()
+{
+	for (int32 i = 0; i < (int32)EExecutionCacheDataType::Count; ++i)
+	{
+		Data[i] = nullptr;
+	}
+}
+
+void FPCGGraphExecutor::FGraphExecutionCache::AddStructReferencedObjects(FReferenceCollector& Collector)
+{
+	for (int32 i = 0; i < (int32)EExecutionCacheDataType::Count; ++i)
+	{
+		Collector.AddReferencedObject(Data[i]);
+	}
 }
 
 FPCGElementPtr FPCGGraphExecutor::GetFetchInputElement()
 {
-	if (!FetchInputElement)
+	if (!FetchInputElementPtr)
 	{
-		FetchInputElement = MakeShared<FPCGFetchInputElement>();
+		FetchInputElementPtr = MakeShared<FPCGFetchInputElement>();
 	}
 
-	return FetchInputElement;
+	return FetchInputElementPtr;
+}
+
+FPCGElementPtr FPCGGraphExecutor::GetPreGraphElement()
+{
+	if (!PreGraphElementPtr)
+	{
+		PreGraphElementPtr = MakeShared<FPCGPreGraphElement>();
+	}
+
+	return PreGraphElementPtr;
 }
 
 void FPCGGraphExecutor::LogTaskState() const
@@ -2753,6 +2825,37 @@ FTextFormat FPCGGraphExecutor::GetNotificationTextFormat()
 
 #endif // WITH_EDITOR
 
+bool FPCGPreGraphElement::ExecuteInternal(FPCGContext* Context) const
+{
+	check(Context);
+
+	check(IsInGameThread());
+	UPCGComponent* Component = Context->SourceComponent.Get();
+
+	// Early out if the component has been deleted/is invalid
+	if (!Component)
+	{
+		// If the component should exist but it doesn't (which is all the time here, previously we checked for it), then this should be cancelled
+		Context->OutputData.bCancelExecution = true;
+		return true;
+	}
+
+#if WITH_EDITOR
+	Component->StartGenerationInProgress();
+#endif
+
+	check(Component->GetGenerationTaskId() != InvalidPCGTaskId);
+
+	{
+		PCG_EXECUTION_CACHE_VALIDATION_CREATE_SCOPE(Component);
+		
+		// Call getters which will create the data and cache it
+		Component->GetActorPCGData();
+		Component->GetOriginalActorPCGData();
+	}
+	return true;
+}
+
 bool FPCGFetchInputElement::ExecuteInternal(FPCGContext* Context) const
 {
 	check(Context);
@@ -2771,7 +2874,7 @@ bool FPCGFetchInputElement::ExecuteInternal(FPCGContext* Context) const
 	}
 
 #if WITH_EDITOR
-	Component->StartGenerationInProgress();
+	check(Component->IsGenerationInProgress());
 #endif
 
 	check(Context->Node);
diff --git a/Engine/Plugins/PCG/Source/PCG/Private/Graph/PCGGraphExecutor.h b/Engine/Plugins/PCG/Source/PCG/Private/Graph/PCGGraphExecutor.h
--- a/Engine/Plugins/PCG/Source/PCG/Private/Graph/PCGGraphExecutor.h
+++ b/Engine/Plugins/PCG/Source/PCG/Private/Graph/PCGGraphExecutor.h
@@ -81,6 +81,7 @@ struct FPCGGraphTask
 #endif
 
 	const FPCGStack* GetStack() const;
+	FPCGTaskId GetGraphExecutionTaskId() const;
 
 	TArray<FPCGGraphTaskInput> Inputs;
 	const UPCGNode* Node = nullptr;
@@ -259,6 +260,39 @@ public:
 	//~End FGCObject interface
 
 private:
+	friend UPCGSubsystem;
+
+	enum class EExecutionCacheDataType : int8
+	{
+		PCGData = 0,
+		InputData,
+		ActorData,
+		LandscapeData,
+		LandscapeHeightData,
+		OriginalActorData,
+		Count
+	};
+
+	static constexpr int8 ExecutionCacheSize = 6;
+
+	UPCGData* GetExecutionCacheData(FPCGTaskId InGraphExecutionTaskId, EExecutionCacheDataType InExecutionCacheDataType);
+
+	UPCGData* GetPCGData(FPCGTaskId InGraphExecutionTaskId) { return GetExecutionCacheData(InGraphExecutionTaskId, EExecutionCacheDataType::PCGData); }
+	UPCGData* GetInputPCGData(FPCGTaskId InGraphExecutionTaskId) { return GetExecutionCacheData(InGraphExecutionTaskId, EExecutionCacheDataType::InputData); }
+	UPCGData* GetActorPCGData(FPCGTaskId InGraphExecutionTaskId) { return GetExecutionCacheData(InGraphExecutionTaskId, EExecutionCacheDataType::ActorData); }
+	UPCGData* GetLandscapePCGData(FPCGTaskId InGraphExecutionTaskId) { return GetExecutionCacheData(InGraphExecutionTaskId, EExecutionCacheDataType::LandscapeData); }
+	UPCGData* GetLandscapeHeightPCGData(FPCGTaskId InGraphExecutionTaskId) { return GetExecutionCacheData(InGraphExecutionTaskId, EExecutionCacheDataType::LandscapeHeightData); }
+	UPCGData* GetOriginalActorPCGData(FPCGTaskId InGraphExecutionTaskId) { return GetExecutionCacheData(InGraphExecutionTaskId, EExecutionCacheDataType::OriginalActorData); }
+
+	void SetExecutionCacheData(FPCGTaskId InGraphExecutionTaskId, EExecutionCacheDataType InExecutionCacheDataType, UPCGData* InData);
+
+	void SetPCGData(FPCGTaskId InGraphExecutionTaskId, UPCGData* InData) { SetExecutionCacheData(InGraphExecutionTaskId, EExecutionCacheDataType::PCGData, InData); }
+	void SetInputPCGData(FPCGTaskId InGraphExecutionTaskId, UPCGData* InData) { SetExecutionCacheData(InGraphExecutionTaskId, EExecutionCacheDataType::InputData, InData); }
+	void SetActorPCGData(FPCGTaskId InGraphExecutionTaskId, UPCGData* InData) { SetExecutionCacheData(InGraphExecutionTaskId, EExecutionCacheDataType::ActorData, InData); }
+	void SetLandscapePCGData(FPCGTaskId InGraphExecutionTaskId, UPCGData* InData) { SetExecutionCacheData(InGraphExecutionTaskId, EExecutionCacheDataType::LandscapeData, InData); }
+	void SetLandscapeHeightPCGData(FPCGTaskId InGraphExecutionTaskId, UPCGData* InData) { SetExecutionCacheData(InGraphExecutionTaskId, EExecutionCacheDataType::LandscapeHeightData, InData); }
+	void SetOriginalActorPCGData(FPCGTaskId InGraphExecutionTaskId, UPCGData* InData) { SetExecutionCacheData(InGraphExecutionTaskId, EExecutionCacheDataType::OriginalActorData, InData); }
+	
 	void ExecuteV1();
 	void ExecuteV2();
 	double GetTickBudgetInSeconds() const;
@@ -313,6 +347,7 @@ private:
 	static void GetPinIdsToDeactivate(FPCGTaskId TaskId, uint64 InactiveOutputPinBitmask, TArray<FPCGPinId>& InOutPinIds);
 
 	FPCGElementPtr GetFetchInputElement();
+	FPCGElementPtr GetPreGraphElement();
 
 	void LogTaskState() const;
 
@@ -335,7 +370,10 @@ private:
 	FPCGGraphCache GraphCache;
 
 	/** Input fetch element, stored here so we have only one */
-	FPCGElementPtr FetchInputElement;
+	FPCGElementPtr FetchInputElementPtr;
+
+	/** PreGraph element, stored here so we have only one */
+	FPCGElementPtr PreGraphElementPtr;
 
 	/** 
 	 * Define a Lock level for future reference. Rule is when we have a lock, we can't lock a lower or equal level lock to prevent deadlocks.
@@ -388,6 +426,21 @@ private:
 	/** Monotonically increasing id. Should be reset once all tasks are executed, should be protected by the ScheduleLock */
 	FPCGTaskId NextTaskId = 0;
 	
+	/** Struct holding different UPCGData caches that we want to compute only once per graph execution */
+	struct FGraphExecutionCache
+	{
+		FGraphExecutionCache();
+
+		void AddStructReferencedObjects(FReferenceCollector& Collector);
+
+		// Must be equal to 
+		TObjectPtr<UPCGData> Data[ExecutionCacheSize];
+	};
+
+	/** Per graph execution cache, gets emptied when executor has no more work to do */
+	UE::FSpinLock GraphExecutionCachesLock;
+	TMap<FPCGTaskId, FGraphExecutionCache> GraphExecutionCaches;
+
 	std::atomic<bool> bNeedToExecuteTasksEnded = false;
 
 	/** Runtime information */
@@ -439,6 +492,16 @@ protected:
 	virtual bool IsPassthrough(const UPCGSettings* InSettings) const override { return true; }
 };
 
+class FPCGPreGraphElement : public IPCGElement
+{
+public:
+	virtual bool IsCacheable(const UPCGSettings* InSettings) const override { return false; }
+	virtual bool CanExecuteOnlyOnMainThread(FPCGContext* Context) const override { return true; }
+protected:
+	virtual bool ExecuteInternal(FPCGContext* Context) const override;
+	virtual bool IsPassthrough(const UPCGSettings* InSettings) const override { return true; }
+};
+
 class FPCGGenericElement : public IPCGElement
 {
 public:
diff --git a/Engine/Plugins/PCG/Source/PCG/Private/Graph/PCGStackContext.cpp b/Engine/Plugins/PCG/Source/PCG/Private/Graph/PCGStackContext.cpp
--- a/Engine/Plugins/PCG/Source/PCG/Private/Graph/PCGStackContext.cpp
+++ b/Engine/Plugins/PCG/Source/PCG/Private/Graph/PCGStackContext.cpp
@@ -267,6 +267,7 @@ void FPCGStackContext::AppendStacks(const FPCGStackContext& InStacks)
 	for (const FPCGStack& SubgraphStack : InStacks.Stacks)
 	{
 		FPCGStack& NewStack = Stacks.Emplace_GetRef();
+		NewStack.GraphExecutionTaskId = GraphExecutionTaskId;
 		NewStack.StackFrames.Reserve(Stacks[CurrentStackIndex].StackFrames.Num() + SubgraphStack.StackFrames.Num());
 		
 		NewStack.StackFrames.Append(Stacks[CurrentStackIndex].StackFrames);
@@ -291,3 +292,14 @@ bool FPCGStackContext::operator==(const FPCGStackContext& Other) const
 {
 	return (CurrentStackIndex == Other.CurrentStackIndex) && (Stacks == Other.Stacks);
 }
+
+void FPCGStackContext::SetGraphExecutionTaskId(FPCGTaskId InGraphExecutionTaskId)
+{
+	check(GraphExecutionTaskId == InvalidPCGTaskId);
+	GraphExecutionTaskId = InGraphExecutionTaskId;
+
+	for (FPCGStack& Stack : Stacks)
+	{
+		Stack.GraphExecutionTaskId = GraphExecutionTaskId;
+	}
+}
\ No newline at end of file
diff --git a/Engine/Plugins/PCG/Source/PCG/Private/PCGComponent.cpp b/Engine/Plugins/PCG/Source/PCG/Private/PCGComponent.cpp
--- a/Engine/Plugins/PCG/Source/PCG/Private/PCGComponent.cpp
+++ b/Engine/Plugins/PCG/Source/PCG/Private/PCGComponent.cpp
@@ -2279,40 +2279,7 @@ void UPCGComponent::DirtyGenerated(EPCGComponentDirtyFlag DirtyFlag, const bool
 	bDirtyGenerated = true;
 
 	ClearPerPinGeneratedOutput();
-
-	// Dirty data as a waterfall from basic values
-	if (!!(DirtyFlag & EPCGComponentDirtyFlag::Actor))
-	{
-		CachedActorData = nullptr;
-		// Since landscape data is related on the bounds of the current actor, when we dirty the actor data, we need to dirty the landscape data as well
-		CachedLandscapeData = nullptr;
-		CachedLandscapeHeightData = nullptr;
-		CachedInputData = nullptr;
-		CachedPCGData = nullptr;
-	}
 	
-	if (!!(DirtyFlag & EPCGComponentDirtyFlag::Landscape))
-	{
-		CachedLandscapeData = nullptr;
-		CachedLandscapeHeightData = nullptr;
-		if (InputType == EPCGComponentInput::Landscape)
-		{
-			CachedInputData = nullptr;
-			CachedPCGData = nullptr;
-		}
-	}
-
-	if (!!(DirtyFlag & EPCGComponentDirtyFlag::Input))
-	{
-		CachedInputData = nullptr;
-		CachedPCGData = nullptr;
-	}
-
-	if (!!(DirtyFlag & EPCGComponentDirtyFlag::Data))
-	{
-		CachedPCGData = nullptr;
-	}
-
 	// For partitioned graph, we must forward the call to the partition actor, if we need to
 	// TODO: Don't forward for None for now, as it could break some stuff
 	if (bActivated && IsPartitioned() && bDispatchToLocalComponents)
@@ -2759,97 +2726,139 @@ void UPCGComponent::OnRefresh(bool bForceRefresh)
 }
 #endif // WITH_EDITOR
 
-UPCGData* UPCGComponent::GetPCGData()
+UPCGData* UPCGComponent::GetPCGData() const
 {
-	if (!CachedPCGData)
+	UPCGData* Data = nullptr;
+	if (UPCGSubsystem* PCGSubsystem = GetSubsystem())
 	{
-		CachedPCGData = CreatePCGData();
-
-		if (GetSubsystem() && GetSubsystem()->IsGraphCacheDebuggingEnabled())
+		Data = PCGSubsystem->GetPCGData(CurrentGenerationTask);
+		if (!Data)
 		{
-			UE_LOG(LogPCG, Log, TEXT("         [%s] CACHE REFRESH CachedPCGData"), *GetOwner()->GetName());
+			Data = CreatePCGData();
+			PCGSubsystem->SetPCGData(CurrentGenerationTask, Data);
+			if (PCGSubsystem->IsGraphCacheDebuggingEnabled() && CurrentGenerationTask != InvalidPCGTaskId)
+			{
+				UE_LOG(LogPCG, Log, TEXT("         [%s] CACHE REFRESH CachedPCGData"), *GetOwner()->GetName());
+			}
 		}
 	}
 
-	return CachedPCGData;
+	return Data;
 }
 
-UPCGData* UPCGComponent::GetInputPCGData()
+UPCGData* UPCGComponent::GetInputPCGData() const
 {
-	if (!CachedInputData)
+	UPCGData* Data = nullptr;
+	if (UPCGSubsystem* PCGSubsystem = GetSubsystem())
 	{
-		CachedInputData = CreateInputPCGData();
-
-		if (GetSubsystem() && GetSubsystem()->IsGraphCacheDebuggingEnabled())
+		Data = PCGSubsystem->GetInputPCGData(CurrentGenerationTask);
+		if (!Data)
 		{
-			UE_LOG(LogPCG, Log, TEXT("         [%s] CACHE REFRESH CachedInputData"), *GetOwner()->GetName());
+			Data = CreateInputPCGData();
+			PCGSubsystem->SetInputPCGData(CurrentGenerationTask, Data);
+			if (PCGSubsystem->IsGraphCacheDebuggingEnabled() && CurrentGenerationTask != InvalidPCGTaskId)
+			{
+				UE_LOG(LogPCG, Log, TEXT("         [%s] CACHE REFRESH CachedInputData"), *GetOwner()->GetName());
+			}
 		}
 	}
 
-	return CachedInputData;
+	return Data;
 }
 
-UPCGData* UPCGComponent::GetActorPCGData()
+UPCGData* UPCGComponent::GetActorPCGData() const
 {
-	// Actor PCG Data can be a Landscape data too
-	if (!CachedActorData || IsLandscapeCachedDataDirty(CachedActorData))
+	UPCGData* Data = nullptr;
+	if (UPCGSubsystem* PCGSubsystem = GetSubsystem())
 	{
-		CachedActorData = CreateActorPCGData();
-
-		if (GetSubsystem() && GetSubsystem()->IsGraphCacheDebuggingEnabled())
+		Data = PCGSubsystem->GetActorPCGData(CurrentGenerationTask);
+		if (!Data)
 		{
-			UE_LOG(LogPCG, Log, TEXT("         [%s] CACHE REFRESH CachedActorData"), *GetOwner()->GetName());
+			PCG_EXECUTION_CACHE_VALIDATION_CHECK(this);
+			Data = CreateActorPCGData();
+			PCGSubsystem->SetActorPCGData(CurrentGenerationTask, Data);
+			if (PCGSubsystem->IsGraphCacheDebuggingEnabled() && CurrentGenerationTask != InvalidPCGTaskId)
+			{
+				UE_LOG(LogPCG, Log, TEXT("         [%s] CACHE REFRESH CachedActorData"), *GetOwner()->GetName());
+			}
 		}
 	}
 
-	return CachedActorData;
+	return Data;
 }
 
-UPCGData* UPCGComponent::GetLandscapePCGData()
+UPCGData* UPCGComponent::GetLandscapePCGData() const
 {
-	if (!CachedLandscapeData || IsLandscapeCachedDataDirty(CachedLandscapeData))
+	UPCGData* Data = nullptr;
+	if (UPCGSubsystem* PCGSubsystem = GetSubsystem())
 	{
-		CachedLandscapeData = CreateLandscapePCGData(/*bHeightOnly=*/false);
-
-		if (GetSubsystem() && GetSubsystem()->IsGraphCacheDebuggingEnabled())
+		Data = PCGSubsystem->GetLandscapePCGData(CurrentGenerationTask);
+		if (!Data)
 		{
-			UE_LOG(LogPCG, Log, TEXT("         [%s] CACHE REFRESH CachedLandscapeData"), *GetOwner()->GetName());
+			Data = CreateLandscapePCGData(/*bHeightOnly=*/false);
+			PCGSubsystem->SetLandscapePCGData(CurrentGenerationTask, Data);
+			if (PCGSubsystem->IsGraphCacheDebuggingEnabled() && CurrentGenerationTask != InvalidPCGTaskId)
+			{
+				UE_LOG(LogPCG, Log, TEXT("         [%s] CACHE REFRESH CachedLandscapeData"), *GetOwner()->GetName());
+			}
 		}
 	}
 
-	return CachedLandscapeData;
+	return Data;
 }
 
-UPCGData* UPCGComponent::GetLandscapeHeightPCGData()
+UPCGData* UPCGComponent::GetLandscapeHeightPCGData() const
 {
-	if (!CachedLandscapeHeightData || IsLandscapeCachedDataDirty(CachedLandscapeHeightData))
+	UPCGData* Data = nullptr;
+	if (UPCGSubsystem* PCGSubsystem = GetSubsystem())
 	{
-		CachedLandscapeHeightData = CreateLandscapePCGData(/*bHeightOnly=*/true);
-
-		if (GetSubsystem() && GetSubsystem()->IsGraphCacheDebuggingEnabled())
+		Data = PCGSubsystem->GetLandscapeHeightPCGData(CurrentGenerationTask);
+		if (!Data)
 		{
-			UE_LOG(LogPCG, Log, TEXT("         [%s] CACHE REFRESH CachedLandscapeHeightData"), *GetOwner()->GetName());
+			Data = CreateLandscapePCGData(/*bHeightOnly=*/true);
+			PCGSubsystem->SetLandscapeHeightPCGData(CurrentGenerationTask, Data);
+			if (PCGSubsystem->IsGraphCacheDebuggingEnabled() && CurrentGenerationTask != InvalidPCGTaskId)
+			{
+				UE_LOG(LogPCG, Log, TEXT("         [%s] CACHE REFRESH CachedLandscapeHeightData"), *GetOwner()->GetName());
+			}
 		}
 	}
 
-	return CachedLandscapeHeightData;
+	return Data;
 }
 
-UPCGData* UPCGComponent::GetOriginalActorPCGData()
+UPCGData* UPCGComponent::GetOriginalActorPCGData() const
 {
-	if (APCGPartitionActor* PartitionActor = Cast<APCGPartitionActor>(GetOwner()))
+	APCGPartitionActor* PartitionActor = Cast<APCGPartitionActor>(GetOwner());
+	if (!PartitionActor)
 	{
-		if (UPCGComponent* OriginalComponent = PartitionActor->GetOriginalComponent(this))
-		{
-			return OriginalComponent->GetActorPCGData();
-		}
+		return GetActorPCGData();
 	}
-	else
+
+	UPCGData* Data = nullptr;
+	if (UPCGSubsystem* PCGSubsystem = GetSubsystem())
 	{
-		return GetActorPCGData();
+		Data = PCGSubsystem->GetOriginalActorPCGData(CurrentGenerationTask);
+		if (!Data)
+		{
+			if (UPCGComponent* OriginalComponent = PartitionActor->GetOriginalComponent(this))
+			{
+				Data = OriginalComponent->GetActorPCGData();
+				if (!Data)
+				{
+					PCG_EXECUTION_CACHE_VALIDATION_CHECK(this);
+					Data = OriginalComponent->CreateActorPCGData();
+				}
+				PCGSubsystem->SetOriginalActorPCGData(CurrentGenerationTask, Data);
+				if (PCGSubsystem->IsGraphCacheDebuggingEnabled() && CurrentGenerationTask != InvalidPCGTaskId)
+				{
+					UE_LOG(LogPCG, Log, TEXT("         [%s] CACHE REFRESH CachedOriginalActorData"), *GetOwner()->GetName());
+				}
+			}
+		}
 	}
 
-	return nullptr;
+	return Data;
 }
 
 UPCGComponent* UPCGComponent::GetOriginalComponent()
@@ -2863,12 +2872,12 @@ UPCGComponent* UPCGComponent::GetOriginalComponent()
 	return ensure(PartitionActor) ? PartitionActor->GetOriginalComponent(this) : this;
 }
 
-UPCGData* UPCGComponent::CreateActorPCGData()
+UPCGData* UPCGComponent::CreateActorPCGData() const
 {
 	return CreateActorPCGData(GetOwner(), bParseActorComponents);
 }
 
-UPCGData* UPCGComponent::CreateActorPCGData(AActor* Actor, bool bParseActor)
+UPCGData* UPCGComponent::CreateActorPCGData(AActor* Actor, bool bParseActor) const
 {
 	return CreateActorPCGData(Actor, this, bParseActor);
 }
@@ -2939,13 +2948,13 @@ void UPCGComponent::RefreshSchedulingPolicy()
 	}
 }
 
-UPCGData* UPCGComponent::CreatePCGData()
+UPCGData* UPCGComponent::CreatePCGData() const
 {
 	TRACE_CPUPROFILER_EVENT_SCOPE(UPCGComponent::CreatePCGData);
 	return GetInputPCGData();
 }
 
-UPCGData* UPCGComponent::CreateLandscapePCGData(bool bHeightOnly)
+UPCGData* UPCGComponent::CreateLandscapePCGData(bool bHeightOnly) const
 {
 	TRACE_CPUPROFILER_EVENT_SCOPE(UPCGComponent::CreateLandscapePCGData);
 	AActor* Actor = GetOwner();
@@ -3009,7 +3018,7 @@ UPCGData* UPCGComponent::CreateLandscapePCGData(bool bHeightOnly)
 	return LandscapeData;
 }
 
-UPCGData* UPCGComponent::CreateInputPCGData()
+UPCGData* UPCGComponent::CreateInputPCGData() const
 {
 	TRACE_CPUPROFILER_EVENT_SCOPE(UPCGComponent::CreateInputPCGData);
 	AActor* Actor = GetOwner();
diff --git a/Engine/Plugins/PCG/Source/PCG/Private/PCGContext.cpp b/Engine/Plugins/PCG/Source/PCG/Private/PCGContext.cpp
--- a/Engine/Plugins/PCG/Source/PCG/Private/PCGContext.cpp
+++ b/Engine/Plugins/PCG/Source/PCG/Private/PCGContext.cpp
@@ -23,6 +23,11 @@
 
 #define LOCTEXT_NAMESPACE "PCGContext"
 
+FPCGTaskId FPCGContext::GetGraphExecutionTaskId() const
+{
+	return ensure(Stack) ? Stack->GetGraphExecutionTaskId() : InvalidPCGTaskId;
+}
+
 FString FPCGContext::GetTaskName() const
 {
 	if (Node)
diff --git a/Engine/Plugins/PCG/Source/PCG/Private/PCGSubsystem.cpp b/Engine/Plugins/PCG/Source/PCG/Private/PCGSubsystem.cpp
--- a/Engine/Plugins/PCG/Source/PCG/Private/PCGSubsystem.cpp
+++ b/Engine/Plugins/PCG/Source/PCG/Private/PCGSubsystem.cpp
@@ -1828,6 +1828,84 @@ TSet<TObjectPtr<APCGPartitionActor>> UPCGSubsystem::GetPCGComponentPartitionActo
 
 #endif // WITH_EDITOR
 
+UPCGData* UPCGSubsystem::GetPCGData(FPCGTaskId InGraphExecutionTaskId)
+{
+	return GraphExecutor ? GraphExecutor->GetPCGData(InGraphExecutionTaskId) : nullptr;
+}
+
+UPCGData* UPCGSubsystem::GetInputPCGData(FPCGTaskId InGraphExecutionTaskId)
+{
+	return GraphExecutor ? GraphExecutor->GetInputPCGData(InGraphExecutionTaskId) : nullptr;
+}
+
+UPCGData* UPCGSubsystem::GetActorPCGData(FPCGTaskId InGraphExecutionTaskId)
+{
+	return GraphExecutor ? GraphExecutor->GetActorPCGData(InGraphExecutionTaskId) : nullptr;
+}
+
+UPCGData* UPCGSubsystem::GetLandscapePCGData(FPCGTaskId InGraphExecutionTaskId)
+{
+	return GraphExecutor ? GraphExecutor->GetLandscapePCGData(InGraphExecutionTaskId) : nullptr;
+}
+
+UPCGData* UPCGSubsystem::GetLandscapeHeightPCGData(FPCGTaskId InGraphExecutionTaskId)
+{
+	return GraphExecutor ? GraphExecutor->GetLandscapeHeightPCGData(InGraphExecutionTaskId) : nullptr;
+}
+
+UPCGData* UPCGSubsystem::GetOriginalActorPCGData(FPCGTaskId InGraphExecutionTaskId)
+{
+	return GraphExecutor ? GraphExecutor->GetOriginalActorPCGData(InGraphExecutionTaskId) : nullptr;
+}
+
+void UPCGSubsystem::SetPCGData(FPCGTaskId InGraphExecutionTaskId, UPCGData* InData)
+{
+	if (GraphExecutor)
+	{
+		GraphExecutor->SetPCGData(InGraphExecutionTaskId, InData);
+	}
+}
+
+void UPCGSubsystem::SetInputPCGData(FPCGTaskId InGraphExecutionTaskId, UPCGData* InData)
+{
+	if (GraphExecutor)
+	{
+		GraphExecutor->SetInputPCGData(InGraphExecutionTaskId, InData);
+	}
+}
+
+void UPCGSubsystem::SetActorPCGData(FPCGTaskId InGraphExecutionTaskId, UPCGData* InData)
+{
+	if (GraphExecutor)
+	{
+		GraphExecutor->SetActorPCGData(InGraphExecutionTaskId, InData);
+	}
+}
+
+void UPCGSubsystem::SetLandscapePCGData(FPCGTaskId InGraphExecutionTaskId, UPCGData* InData)
+{
+	if (GraphExecutor)
+	{
+		GraphExecutor->SetLandscapePCGData(InGraphExecutionTaskId, InData);
+	}
+}
+
+void UPCGSubsystem::SetLandscapeHeightPCGData(FPCGTaskId InGraphExecutionTaskId, UPCGData* InData)
+{
+	if (GraphExecutor)
+	{
+		GraphExecutor->SetLandscapeHeightPCGData(InGraphExecutionTaskId, InData);
+	}
+}
+
+void UPCGSubsystem::SetOriginalActorPCGData(FPCGTaskId InGraphExecutionTaskId, UPCGData* InData)
+{
+	if (GraphExecutor)
+	{
+		GraphExecutor->SetOriginalActorPCGData(InGraphExecutionTaskId, InData);
+	}
+}
+
 void UPCGSubsystem::ExecuteBeginTickActions()
 {
 	TArray<FTickAction> Actions = MoveTemp(BeginTickActions);
diff --git a/Engine/Plugins/PCG/Source/PCG/Public/Elements/PCGCreatePoints.h b/Engine/Plugins/PCG/Source/PCG/Public/Elements/PCGCreatePoints.h
--- a/Engine/Plugins/PCG/Source/PCG/Public/Elements/PCGCreatePoints.h
+++ b/Engine/Plugins/PCG/Source/PCG/Public/Elements/PCGCreatePoints.h
@@ -69,7 +69,6 @@ PRAGMA_ENABLE_DEPRECATION_WARNINGS
 class FPCGCreatePointsElement : public IPCGElement
 {
 public:
-	virtual bool CanExecuteOnlyOnMainThread(FPCGContext* Context) const override;
 	virtual bool IsCacheable(const UPCGSettings* InSettings) const override;
 	virtual bool ShouldComputeFullOutputDataCrc(FPCGContext* Context) const override { return true; }
 
diff --git a/Engine/Plugins/PCG/Source/PCG/Public/Elements/PCGCreatePointsGrid.h b/Engine/Plugins/PCG/Source/PCG/Public/Elements/PCGCreatePointsGrid.h
--- a/Engine/Plugins/PCG/Source/PCG/Public/Elements/PCGCreatePointsGrid.h
+++ b/Engine/Plugins/PCG/Source/PCG/Public/Elements/PCGCreatePointsGrid.h
@@ -90,9 +90,6 @@ PRAGMA_ENABLE_DEPRECATION_WARNINGS
 class FPCGCreatePointsGridElement : public IPCGElement
 {
 public:
-	// Returns true when Context is nullptr which happens in FPCGGraphExecutor::OnTaskInputsReady
-	// We can't know for sure at that point if we need to call UPCGComponent::GetActorPCGData() in our GetDependenciesCrc so we need to make sure it happens on MainThread
-	virtual bool CanExecuteOnlyOnMainThread(FPCGContext* Context) const override { return Context == nullptr; }
 	virtual void GetDependenciesCrc(const FPCGDataCollection& InInput, const UPCGSettings* InSettings, UPCGComponent* InComponent, FPCGCrc& OutCrc) const override;
 
 protected:
diff --git a/Engine/Plugins/PCG/Source/PCG/Public/Elements/PCGCreatePointsSphere.h b/Engine/Plugins/PCG/Source/PCG/Public/Elements/PCGCreatePointsSphere.h
--- a/Engine/Plugins/PCG/Source/PCG/Public/Elements/PCGCreatePointsSphere.h
+++ b/Engine/Plugins/PCG/Source/PCG/Public/Elements/PCGCreatePointsSphere.h
@@ -141,9 +141,6 @@ public:
 class FPCGCreatePointsSphereElement : public IPCGElement
 {
 public:
-	// Returns true when Context is nullptr which happens in FPCGGraphExecutor::OnTaskInputsReady
-	// We can't know for sure at that point if we need to call UPCGComponent::GetActorPCGData() in our GetDependenciesCrc so we need to make sure it happens on MainThread
-	virtual bool CanExecuteOnlyOnMainThread(FPCGContext* Context) const override { return Context == nullptr; }
 	virtual void GetDependenciesCrc(const FPCGDataCollection& InInput, const UPCGSettings* InSettings, UPCGComponent* InComponent, FPCGCrc& OutCrc) const override;
 
 protected:
diff --git a/Engine/Plugins/PCG/Source/PCG/Public/Elements/PCGCullPointsOutsideActorBounds.h b/Engine/Plugins/PCG/Source/PCG/Public/Elements/PCGCullPointsOutsideActorBounds.h
--- a/Engine/Plugins/PCG/Source/PCG/Public/Elements/PCGCullPointsOutsideActorBounds.h
+++ b/Engine/Plugins/PCG/Source/PCG/Public/Elements/PCGCullPointsOutsideActorBounds.h
@@ -38,9 +38,6 @@ public:
 class FPCGCullPointsOutsideActorBoundsElement : public IPCGElement
 {
 public:
-	// Returns true when Context is nullptr which happens in FPCGGraphExecutor::OnTaskInputsReady
-	// We can't know for sure at that point if we need to call UPCGComponent::GetActorPCGData() in our GetDependenciesCrc so we need to make sure it happens on MainThread
-	virtual bool CanExecuteOnlyOnMainThread(FPCGContext* Context) const override { return Context == nullptr; }
 	virtual void GetDependenciesCrc(const FPCGDataCollection& InInput, const UPCGSettings* InSettings, UPCGComponent* InComponent, FPCGCrc& OutCrc) const override;
 
 protected:
diff --git a/Engine/Plugins/PCG/Source/PCG/Public/Elements/PCGElevationIsolines.h b/Engine/Plugins/PCG/Source/PCG/Public/Elements/PCGElevationIsolines.h
--- a/Engine/Plugins/PCG/Source/PCG/Public/Elements/PCGElevationIsolines.h
+++ b/Engine/Plugins/PCG/Source/PCG/Public/Elements/PCGElevationIsolines.h
@@ -65,9 +65,6 @@ public:
 
 class FPCGElevationIsolinesElement : public IPCGElement
 {
-public:
-	virtual bool CanExecuteOnlyOnMainThread(FPCGContext* Context) const override;
-
 protected:
 	virtual bool ExecuteInternal(FPCGContext* InContext) const override;
 	virtual EPCGElementExecutionLoopMode ExecutionLoopMode(const UPCGSettings* Settings) const override { return EPCGElementExecutionLoopMode::SinglePrimaryPin; }
diff --git a/Engine/Plugins/PCG/Source/PCG/Public/Elements/PCGHiGenGridSize.h b/Engine/Plugins/PCG/Source/PCG/Public/Elements/PCGHiGenGridSize.h
--- a/Engine/Plugins/PCG/Source/PCG/Public/Elements/PCGHiGenGridSize.h
+++ b/Engine/Plugins/PCG/Source/PCG/Public/Elements/PCGHiGenGridSize.h
@@ -44,10 +44,6 @@ protected:
 class FPCGHiGenGridSizeElement : public IPCGElement
 {
 public:
-	// Returns true when Context is nullptr which happens in FPCGGraphExecutor::OnTaskInputsReady
-	// We can't know for sure at that point if we need to call UPCGComponent::GetActorPCGData() in our GetDependenciesCrc so we need to make sure it happens on MainThread
-	virtual bool CanExecuteOnlyOnMainThread(FPCGContext* Context) const override { return Context == nullptr; }
-
 	virtual void GetDependenciesCrc(const FPCGDataCollection& InInput, const UPCGSettings* InSettings, UPCGComponent* InComponent, FPCGCrc& OutCrc) const override;
 
 protected:
diff --git a/Engine/Plugins/PCG/Source/PCG/Public/Elements/PCGSplineSampler.h b/Engine/Plugins/PCG/Source/PCG/Public/Elements/PCGSplineSampler.h
--- a/Engine/Plugins/PCG/Source/PCG/Public/Elements/PCGSplineSampler.h
+++ b/Engine/Plugins/PCG/Source/PCG/Public/Elements/PCGSplineSampler.h
@@ -280,10 +280,6 @@ public:
 class FPCGSplineSamplerElement : public IPCGElement
 {
 public:
-	// Returns true when Context is nullptr which happens in FPCGGraphExecutor::OnTaskInputsReady
-	// We can't know for sure at that point if we need to call UPCGComponent::GetActorPCGData() in our GetDependenciesCrc so we need to make sure it happens on MainThread
-	virtual bool CanExecuteOnlyOnMainThread(FPCGContext* Context) const override { return Context == nullptr; }
-
 	virtual void GetDependenciesCrc(const FPCGDataCollection& InInput, const UPCGSettings* InSettings, UPCGComponent* InComponent, FPCGCrc& OutCrc) const override;
 	// Worth computing a full CRC in case we can halt change propagation/re-executions
 	virtual bool ShouldComputeFullOutputDataCrc(FPCGContext* Context) const override { return true; }
diff --git a/Engine/Plugins/PCG/Source/PCG/Public/Elements/PCGSurfaceSampler.h b/Engine/Plugins/PCG/Source/PCG/Public/Elements/PCGSurfaceSampler.h
--- a/Engine/Plugins/PCG/Source/PCG/Public/Elements/PCGSurfaceSampler.h
+++ b/Engine/Plugins/PCG/Source/PCG/Public/Elements/PCGSurfaceSampler.h
@@ -156,10 +156,6 @@ public:
 class FPCGSurfaceSamplerElement : public TPCGTimeSlicedElementBase<PCGSurfaceSampler::FSurfaceSamplerExecutionState, PCGSurfaceSampler::FSurfaceSamplerIterationState>
 {
 public:
-	// Returns true when Context is nullptr which happens in FPCGGraphExecutor::OnTaskInputsReady
-	// We can't know for sure at that point if we need to call UPCGComponent::GetActorPCGData() in our GetDependenciesCrc so we need to make sure it happens on MainThread
-	virtual bool CanExecuteOnlyOnMainThread(FPCGContext* Context) const override { return Context == nullptr; }
-
 	virtual void GetDependenciesCrc(const FPCGDataCollection& InInput, const UPCGSettings* InSettings, UPCGComponent* InComponent, FPCGCrc& OutCrc) const override;
 	// Might be sampling landscape or other external data, worth computing a full CRC in case we can halt change propagation/re-executions
 	virtual bool ShouldComputeFullOutputDataCrc(FPCGContext* Context) const override { return true; }
diff --git a/Engine/Plugins/PCG/Source/PCG/Public/Elements/PCGVolumeSampler.h b/Engine/Plugins/PCG/Source/PCG/Public/Elements/PCGVolumeSampler.h
--- a/Engine/Plugins/PCG/Source/PCG/Public/Elements/PCGVolumeSampler.h
+++ b/Engine/Plugins/PCG/Source/PCG/Public/Elements/PCGVolumeSampler.h
@@ -94,10 +94,6 @@ protected:
 class FPCGVolumeSamplerElement : public TPCGTimeSlicedElementBase<PCGVolumeSampler::FVolumeSamplerExecutionState, PCGVolumeSampler::FVolumeSamplerIterationState>
 {
 public:
-	// Returns true when Context is nullptr which happens in FPCGGraphExecutor::OnTaskInputsReady
-	// We can't know for sure at that point if we need to call UPCGComponent::GetActorPCGData() in our GetDependenciesCrc so we need to make sure it happens on MainThread
-	virtual bool CanExecuteOnlyOnMainThread(FPCGContext* Context) const override { return Context == nullptr; }
-
 	virtual void GetDependenciesCrc(const FPCGDataCollection& InInput, const UPCGSettings* InSettings, UPCGComponent* InComponent, FPCGCrc& OutCrc) const override;
 	// Might be sampling external data like brush, worth computing a full CRC in case we can halt change propagation/re-executions
 	virtual bool ShouldComputeFullOutputDataCrc(FPCGContext* Context) const override { return true; }
diff --git a/Engine/Plugins/PCG/Source/PCG/Public/Graph/PCGStackContext.h b/Engine/Plugins/PCG/Source/PCG/Public/Graph/PCGStackContext.h
--- a/Engine/Plugins/PCG/Source/PCG/Public/Graph/PCGStackContext.h
+++ b/Engine/Plugins/PCG/Source/PCG/Public/Graph/PCGStackContext.h
@@ -86,6 +86,9 @@ public:
 	PRAGMA_ENABLE_DEPRECATION_WARNINGS
 #endif // WITH_EDITOR
 
+	/** Returns root graph execution id */
+	FPCGTaskId GetGraphExecutionTaskId() const { return GraphExecutionTaskId; }
+
 	/** Push frame onto top of stack. */
 	void PushFrame(const FPCGStackFrame& Frame) { StackFrames.Add(Frame); }
 	void PushFrame(const UObject* InFrameObject) { StackFrames.Emplace(InFrameObject); }
@@ -140,6 +143,7 @@ public:
 	}
 
 private:
+	FPCGTaskId GraphExecutionTaskId = InvalidPCGTaskId;
 	TArray<FPCGStackFrame> StackFrames;
 
 public:
@@ -153,6 +157,7 @@ public:
 /** A collection of call stacks. */
 class PCG_API FPCGStackContext
 {
+	friend class FPCGGraphExecutor;
 public:
 	int32 GetNumStacks() const { return Stacks.Num(); }
 	int32 GetCurrentStackIndex() const { return CurrentStackIndex; }
@@ -175,7 +180,13 @@ public:
 
 	bool operator==(const FPCGStackContext& Other) const;
 
+	FPCGTaskId GetGraphExecutionTaskId() const { return GraphExecutionTaskId; }
 private:
+	void SetGraphExecutionTaskId(FPCGTaskId InGraphExecutionTaskId);
+
+	/** Unique graph execution id of this stack */
+	FPCGTaskId GraphExecutionTaskId = InvalidPCGTaskId;
+
 	/** List of all stacks encountered top graph and all (nested) subgraphs. Order is simply order of encountering during compilation. */
 	TArray<FPCGStack> Stacks;
 
diff --git a/Engine/Plugins/PCG/Source/PCG/Public/PCGCommon.h b/Engine/Plugins/PCG/Source/PCG/Public/PCGCommon.h
--- a/Engine/Plugins/PCG/Source/PCG/Public/PCGCommon.h
+++ b/Engine/Plugins/PCG/Source/PCG/Public/PCGCommon.h
@@ -427,4 +427,21 @@ enum class EPCGGenerationStatus : uint8
 {
 	Completed,
 	Aborted
-};
\ No newline at end of file
+};
+
+// Enable to debug if some PCG Elements are creating some data that should be pre-created by the graph PreGraph element
+#define PCG_EXECUTION_CACHE_VALIDATION_ENABLED 0
+
+#if PCG_EXECUTION_CACHE_VALIDATION_ENABLED
+
+#define PCG_EXECUTION_CACHE_VALIDATION_CREATE_SCOPE(PCGComponent) TGuardValue<bool> ValidationCreateScope(PCGComponent->bCanCreateExecutionCache, true);
+#define PCG_EXECUTION_CACHE_VALIDATION_CREATE_ORIGINAL_SCOPE(PCGComponent) TGuardValue<bool> ValidationCreateOriginalScope(PCGComponent->GetOriginalComponent()->bCanCreateExecutionCache, PCGComponent->bCanCreateExecutionCache);
+#define PCG_EXECUTION_CACHE_VALIDATION_CHECK(PCGComponent) ensureAlways(PCGComponent->bCanCreateExecutionCache || PCGComponent->CurrentGenerationTask == InvalidPCGTaskId);
+
+#else
+
+#define PCG_EXECUTION_CACHE_VALIDATION_CREATE_SCOPE(PCGComponent)
+#define PCG_EXECUTION_CACHE_VALIDATION_CREATE_ORIGINAL_SCOPE(PCGComponent)
+#define PCG_EXECUTION_CACHE_VALIDATION_CHECK(PCGComponent)
+
+#endif // PCG_EXECUTION_CACHE_VALIDATION_ENABLED
\ No newline at end of file
diff --git a/Engine/Plugins/PCG/Source/PCG/Public/PCGComponent.h b/Engine/Plugins/PCG/Source/PCG/Public/PCGComponent.h
--- a/Engine/Plugins/PCG/Source/PCG/Public/PCGComponent.h
+++ b/Engine/Plugins/PCG/Source/PCG/Public/PCGComponent.h
@@ -117,12 +117,12 @@ protected:
 	//~End UActorComponent Interface
 
 public:
-	UPCGData* GetPCGData();
-	UPCGData* GetInputPCGData();
-	UPCGData* GetActorPCGData();
-	UPCGData* GetLandscapePCGData();
-	UPCGData* GetLandscapeHeightPCGData();
-	UPCGData* GetOriginalActorPCGData();
+	UPCGData* GetPCGData() const;
+	UPCGData* GetInputPCGData() const;
+	UPCGData* GetActorPCGData() const;
+	UPCGData* GetLandscapePCGData() const;
+	UPCGData* GetLandscapeHeightPCGData() const;
+	UPCGData* GetOriginalActorPCGData() const;
 
 	/** If this is a local component returns self, otherwise returns the original component. */
 	UPCGComponent* GetOriginalComponent();
@@ -501,11 +501,11 @@ private:
 	uint32 RuntimeGridDescriptorHash = 0;
 
 private:
-	UPCGData* CreatePCGData();
-	UPCGData* CreateInputPCGData();
-	UPCGData* CreateActorPCGData();
-	UPCGData* CreateActorPCGData(AActor* Actor, bool bParseActor = true);
-	UPCGData* CreateLandscapePCGData(bool bHeightOnly);
+	UPCGData* CreatePCGData() const;
+	UPCGData* CreateInputPCGData() const;
+	UPCGData* CreateActorPCGData() const;
+	UPCGData* CreateActorPCGData(AActor* Actor, bool bParseActor = true) const;
+	UPCGData* CreateLandscapePCGData(bool bHeightOnly) const;
 	bool IsLandscapeCachedDataDirty(const UPCGData* Data) const;
 
 	bool ShouldGenerate(bool bForce, EPCGComponentGenerationTrigger RequestedGenerationTrigger) const;
@@ -570,21 +570,6 @@ private:
 
 	FBox GetGridBounds(const AActor* InActor) const;
 
-	UPROPERTY(Transient, NonPIEDuplicateTransient)
-	TObjectPtr<UPCGData> CachedPCGData = nullptr;
-
-	UPROPERTY(Transient, NonPIEDuplicateTransient)
-	TObjectPtr<UPCGData> CachedInputData = nullptr;
-
-	UPROPERTY(Transient, NonPIEDuplicateTransient)
-	TObjectPtr<UPCGData> CachedActorData = nullptr;
-
-	UPROPERTY(Transient, NonPIEDuplicateTransient)
-	TObjectPtr<UPCGData> CachedLandscapeData = nullptr;
-
-	UPROPERTY(Transient, NonPIEDuplicateTransient)
-	TObjectPtr<UPCGData> CachedLandscapeHeightData = nullptr;
-
 #if WITH_EDITORONLY_DATA
 	UPROPERTY()
 	TSet<TSoftObjectPtr<AActor>> GeneratedActors_DEPRECATED;
@@ -684,6 +669,11 @@ private:
 public:
 	mutable PCGUtils::FExtraCapture ExtraCapture;
 #endif // WITH_EDITOR
+
+public:
+#if PCG_EXECUTION_CACHE_VALIDATION_ENABLED
+	bool bCanCreateExecutionCache = false;
+#endif
 };
 
 /** Used to store generated resources data during RerunConstructionScripts */
diff --git a/Engine/Plugins/PCG/Source/PCG/Public/PCGContext.h b/Engine/Plugins/PCG/Source/PCG/Public/PCGContext.h
--- a/Engine/Plugins/PCG/Source/PCG/Public/PCGContext.h
+++ b/Engine/Plugins/PCG/Source/PCG/Public/PCGContext.h
@@ -109,6 +109,8 @@ struct PCG_API FPCGContext
 		return SettingsWithOverride ? Cast<SettingsType>(SettingsWithOverride) : GetOriginalSettings<SettingsType>();
 	}
 
+	FPCGTaskId GetGraphExecutionTaskId() const;
+
 	FString GetTaskName() const;
 	FString GetComponentName() const;
 	bool ShouldStop() const { return AsyncState.ShouldStop(); }
diff --git a/Engine/Plugins/PCG/Source/PCG/Public/PCGSubsystem.h b/Engine/Plugins/PCG/Source/PCG/Public/PCGSubsystem.h
--- a/Engine/Plugins/PCG/Source/PCG/Public/PCGSubsystem.h
+++ b/Engine/Plugins/PCG/Source/PCG/Public/PCGSubsystem.h
@@ -54,6 +54,7 @@ class PCG_API UPCGSubsystem : public UTickableWorldSubsystem
 	GENERATED_BODY()
 
 public:
+	friend class UPCGComponent;
 	friend FPCGActorAndComponentMapping;
 	friend struct FPCGWorldPartitionBuilder;
 
@@ -359,6 +360,20 @@ private:
 #endif // WITH_EDITOR
 	
 private:
+	UPCGData* GetPCGData(FPCGTaskId InGraphExecutionTaskId);
+	UPCGData* GetInputPCGData(FPCGTaskId InGraphExecutionTaskId);
+	UPCGData* GetActorPCGData(FPCGTaskId InGraphExecutionTaskId);
+	UPCGData* GetLandscapePCGData(FPCGTaskId InGraphExecutionTaskId);
+	UPCGData* GetLandscapeHeightPCGData(FPCGTaskId InGraphExecutionTaskId);
+	UPCGData* GetOriginalActorPCGData(FPCGTaskId InGraphExecutionTaskId);
+
+	void SetPCGData(FPCGTaskId InGraphExecutionTaskId, UPCGData* InData);
+	void SetInputPCGData(FPCGTaskId InGraphExecutionTaskId, UPCGData* InData);
+	void SetActorPCGData(FPCGTaskId InGraphExecutionTaskId, UPCGData* InData);
+	void SetLandscapePCGData(FPCGTaskId InGraphExecutionTaskId, UPCGData* InData);
+	void SetLandscapeHeightPCGData(FPCGTaskId InGraphExecutionTaskId, UPCGData* InData);
+	void SetOriginalActorPCGData(FPCGTaskId InGraphExecutionTaskId, UPCGData* InData);
+
 	void ExecuteBeginTickActions();
 
 	APCGWorldActor* PCGWorldActor = nullptr;
