[PCG] Builder: Add support for Included/Excluded DataLayers

#rnx

diff --git a/Engine/Plugins/PCG/Source/PCGEditor/Private/WorldPartitionBuilder/PCGBuilderSettings.h b/Engine/Plugins/PCG/Source/PCGEditor/Private/WorldPartitionBuilder/PCGBuilderSettings.h
--- a/Engine/Plugins/PCG/Source/PCGEditor/Private/WorldPartitionBuilder/PCGBuilderSettings.h
+++ b/Engine/Plugins/PCG/Source/PCGEditor/Private/WorldPartitionBuilder/PCGBuilderSettings.h
@@ -8,6 +8,7 @@
 #include "PCGBuilderSettings.generated.h"
 
 class UPCGGraphInterface;
+class UDataLayerAsset;
 
 UCLASS(hidecategories = Object)
 class UPCGBuilderSettings : public UObject
@@ -28,10 +29,22 @@ public:
 	/** Filter generated components by owner actor name. */
 	UPROPERTY(EditAnywhere, Category = Filter)
 	TArray<FString>  FilterByActorNames;
+		
+	/** Will load editor only data layers without having to specify them as included */
+	UPROPERTY(EditAnywhere, Category = DataLayers)
+	bool bLoadEditorOnlyDataLayers = true;
 
-	/** Call generate on each component and wait until completion and any async processes before generating the next. */
-	UPROPERTY(EditAnywhere, Category = Build, meta = (DisplayName = "One Component at a Time"))
-	bool bOneComponentAtATime = false;
+	/** Will load runtime data layers that are default activated without having to specify them as included */
+	UPROPERTY(EditAnywhere, Category = DataLayers)
+	bool bLoadActivatedRuntimeDataLayers = true;
+
+	/** Data layers that need to be loaded when running builder. */
+	UPROPERTY(EditAnywhere, Category = DataLayers)
+	TArray<TObjectPtr<UDataLayerAsset>> IncludedDataLayers;
+
+	/** Data layers that need to be unloaded when running builder. */
+	UPROPERTY(EditAnywhere, Category = DataLayers)
+	TArray<TObjectPtr<UDataLayerAsset>> ExcludedDataLayers;
 
 	/** Run builder iteratively loading cell regions instead of the whole map at once. Useful for large maps. */
 	UPROPERTY(EditAnywhere, Category = Advanced)
@@ -44,4 +57,8 @@ public:
 	/** Submit dirty files even if errors occurred during generation. */
 	UPROPERTY(EditAnywhere, Category = Advanced)
 	bool bIgnoreGenerationErrors = false;
+
+	/** Call generate on each component and wait until completion and any async processes before generating the next. */
+	UPROPERTY(EditAnywhere, Category = Advanced, meta = (DisplayName = "One Component at a Time"))
+	bool bOneComponentAtATime = false;
 };
\ No newline at end of file
diff --git a/Engine/Plugins/PCG/Source/PCGEditor/Private/WorldPartitionBuilder/PCGWorldPartitionBuilder.cpp b/Engine/Plugins/PCG/Source/PCGEditor/Private/WorldPartitionBuilder/PCGWorldPartitionBuilder.cpp
--- a/Engine/Plugins/PCG/Source/PCGEditor/Private/WorldPartitionBuilder/PCGWorldPartitionBuilder.cpp
+++ b/Engine/Plugins/PCG/Source/PCGEditor/Private/WorldPartitionBuilder/PCGWorldPartitionBuilder.cpp
@@ -30,6 +30,8 @@
 #include "WorldPartition/IWorldPartitionEditorModule.h"
 #include "WorldPartition/WorldPartition.h"
 #include "WorldPartition/WorldPartitionHelpers.h"
+#include "WorldPartition/DataLayer/DataLayerAsset.h"
+#include "WorldPartition/DataLayer/DataLayerManager.h"
 #include "WorldPartition/LoaderAdapter/LoaderAdapterShape.h"
 
 DEFINE_LOG_CATEGORY_STATIC(LogPCGWorldPartitionBuilder, All, All);
@@ -241,6 +243,10 @@ FPCGWorldPartitionBuilderArgs FPCGWorldPartitionBuilderArgs::InitializeFrom(cons
 			BuilderArgs.IncludeGraphs = BuilderSettings->Graphs;
 			BuilderArgs.bIterativeCellLoading = BuilderSettings->bIterativeCellLoading;
 			BuilderArgs.IterativeCellSize = BuilderSettings->IterativeCellSize;
+			BuilderArgs.bLoadEditorOnlyDataLayers = BuilderSettings->bLoadEditorOnlyDataLayers;
+			BuilderArgs.bLoadActivatedRuntimeDataLayers = BuilderSettings->bLoadActivatedRuntimeDataLayers;
+			BuilderArgs.IncludedDataLayers = BuilderSettings->IncludedDataLayers;
+			BuilderArgs.ExcludedDataLayers = BuilderSettings->ExcludedDataLayers;
 
 			UE_LOG(LogPCGWorldPartitionBuilder, Display, TEXT("Loaded PCGBuilderSettings '%s'"), *SettingsPath.ToString());
 			bUsingBuilderSettings = true;
@@ -471,6 +477,37 @@ bool UPCGWorldPartitionBuilder::PreRun(UWorld* World, FPackageSourceControlHelpe
 	// If World doesn't have a World Partition do not use iterative cell loading
 	Args.bIterativeCellLoading &= World->IsPartitionedWorld();
 
+	// Setup DataLayers
+	if (UDataLayerManager* DataLayerManager = UDataLayerManager::GetDataLayerManager(World))
+	{
+		auto FillDataLayerShortNames = [DataLayerManager](const TArray<TObjectPtr<UDataLayerAsset>>& InDataLayerAssets, TSet<FName>& OutDataLayerShortNames)
+		{
+			for(const TObjectPtr<UDataLayerAsset>& DataLayer : InDataLayerAssets)
+			{
+				if (DataLayer != nullptr)
+				{
+					const UDataLayerInstance* DataLayerInstance = DataLayerManager->GetDataLayerInstance(DataLayer);
+					if (DataLayerInstance == nullptr)
+					{
+						UE_LOG(LogPCGWorldPartitionBuilder, Error, TEXT("Missing UDataLayerInstance for %s."), *DataLayer->GetName());
+					}
+					else
+					{
+						OutDataLayerShortNames.Add(FName(DataLayerInstance->GetDataLayerShortName()));
+					}
+				}
+			}
+		};
+
+		// Set Base Builder Class fields
+		bLoadNonDynamicDataLayers = Args.bLoadEditorOnlyDataLayers;
+		bLoadInitiallyActiveDataLayers = Args.bLoadActivatedRuntimeDataLayers;
+
+		FillDataLayerShortNames(Args.IncludedDataLayers, DataLayerShortNames);
+		FillDataLayerShortNames(Args.ExcludedDataLayers, ExcludedDataLayerShortNames);
+	}
+
+
 	if (GetLoadingMode() == ELoadingMode::IterativeCells2D)
 	{
 		// Disable Source Control until Post Run
diff --git a/Engine/Plugins/PCG/Source/PCGEditor/Private/WorldPartitionBuilder/PCGWorldPartitionBuilder.h b/Engine/Plugins/PCG/Source/PCGEditor/Private/WorldPartitionBuilder/PCGWorldPartitionBuilder.h
--- a/Engine/Plugins/PCG/Source/PCGEditor/Private/WorldPartitionBuilder/PCGWorldPartitionBuilder.h
+++ b/Engine/Plugins/PCG/Source/PCGEditor/Private/WorldPartitionBuilder/PCGWorldPartitionBuilder.h
@@ -79,6 +79,18 @@ struct FPCGWorldPartitionBuilderArgs
 	/** Cell size for iterative cell loading */
 	int32 IterativeCellSize = 25600;
 
+	/** Will load editor only data layers without having to specify them as included */
+	bool bLoadEditorOnlyDataLayers = true;
+
+	/** Will load runtime data layers that are default activated without having to specify them as included */
+	bool bLoadActivatedRuntimeDataLayers = true;
+
+	/** Data layers that need to be loaded when running builder. */
+	TArray<TObjectPtr<UDataLayerAsset>> IncludedDataLayers;
+
+	/** Data layers that need to be unloaded when running builder. */
+	TArray<TObjectPtr<UDataLayerAsset>> ExcludedDataLayers;
+
 	static FPCGWorldPartitionBuilderArgs InitializeFrom(const FPCGWorldPartitionCommandlineArgs& CommandlineArgs);
 };
 
