[PCG] GPU - finalize data collection buffer format.

Move all headers to front of buffer, meaning the buffer preparation on CPU only needs to upload a handful of bytes which is a big memory and CPU time saving.

Changed some packing/unpacking code from ParallelFor { SwitchType { ... } } to SwitchType { ParallelFor { ... } } as cost of this is a perf pain point.

#rnx

diff --git a/Engine/Plugins/PCG/Shaders/Private/Elements/PCGCopyPoints.usf b/Engine/Plugins/PCG/Shaders/Private/Elements/PCGCopyPoints.usf
--- a/Engine/Plugins/PCG/Shaders/Private/Elements/PCGCopyPoints.usf
+++ b/Engine/Plugins/PCG/Shaders/Private/Elements/PCGCopyPoints.usf
@@ -21,38 +21,25 @@
 #define NUM_RESERVED_ATTRS 32
 
 // TODO pass in attribute IDs to copy.
-void CopyAttributesFromSource(uint SourceDataIndex, uint SourceElemIndex, uint OutDataIndex, uint OutElemIndex)
+void CopyAttributesFromSource(uint Source_DataIndex, uint Source_ElementIndex, uint Out_DataIndex, uint Out_ElementIndex)
 {
-	const uint SourceDataAddress = Source_GetDataAddress(SourceDataIndex);
-	const uint SourceHeadersAddress = Source_GetDataAttributeHeadersAddress(SourceDataAddress);
-
-	const uint OutDataAddress = Out_GetDataAddress(OutDataIndex);
-	const uint OutHeadersAddress = Out_GetDataAttributeHeadersAddress(OutDataAddress);
-
-	const uint NumAttributes = Source_GetDataNumAttributes(SourceDataAddress);
+	const uint NumAttributes = Source_GetDataNumAttributesInternal(Source_DataIndex);
 	uint NumAttributesProcessed = 0;
 	
 	// Loop over all attribute headers, if the address is non-zero, then copy it from pin Source to pin Out.
 	for (int AttributeIndex = NUM_RESERVED_ATTRS; AttributeIndex < MAX_NUM_ATTRS; ++AttributeIndex)
 	{
-		const uint SourceHeaderAddress = Source_GetAttributeHeaderAddress(SourceHeadersAddress, AttributeIndex);
-		const uint SourceAddress = Source_GetAttributeAddress(SourceHeaderAddress);
-
-		const uint OutHeaderAddress = Out_GetAttributeHeaderAddress(OutHeadersAddress, AttributeIndex);
-		const uint OutAddress = Out_GetAttributeAddress(OutHeaderAddress);
+		const uint Source_Stride = Source_GetAttributeStrideInternal(Source_DataIndex, AttributeIndex);
+		const uint Out_Stride = Target_GetAttributeStrideInternal(Out_DataIndex, AttributeIndex);
 
-		const uint Stride = Source_GetAttributeStride(SourceHeaderAddress);
-
-		if (SourceAddress != 0 && OutAddress != 0)
+		if (Source_Stride != 0 && Source_Stride == Out_Stride)
 		{
-			const uint SourceBaseElementAddress = SourceAddress + SourceElemIndex * Stride;
-			const uint OutBaseElementAddress = OutAddress + OutElemIndex * Stride;
+			const uint Source_ElementAddress = Source_GetElementAddressInternal(Source_DataIndex, Source_ElementIndex, AttributeIndex);
+			const uint Out_ElementAddress = Out_GetElementAddressInternal(Out_DataIndex, Out_ElementIndex, AttributeIndex);
 
-			for (int I = 0; I < Stride; I += 4)
+			for (int Offset = 0; Offset < Out_Stride; Offset += 4)
 			{
-				const uint SourceElementAddress = SourceBaseElementAddress + I;
-				const uint OutElementAddress = OutBaseElementAddress + I;
-				Out_StoreBufferInternal(OutElementAddress, Source_LoadBufferInternal(SourceElementAddress));
+				Out_StoreBufferInternal(Out_ElementAddress + Offset, Source_LoadBufferInternal(Source_ElementAddress + Offset));
 			}
 
 			// We can early-out when we've looked at all the possible attributes
@@ -65,38 +52,25 @@ void CopyAttributesFromSource(uint SourceDataIndex, uint SourceElemIndex, uint O
 }
 
 // TODO pass in attribute IDs to copy.
-void CopyAttributesFromTarget(uint TargetDataIndex, uint TargetElemIndex, uint OutDataIndex, uint OutElemIndex)
+void CopyAttributesFromTarget(uint Target_DataIndex, uint Target_ElementIndex, uint Out_DataIndex, uint Out_ElementIndex)
 {
-	const uint TargetDataAddress = Target_GetDataAddress(TargetDataIndex);
-	const uint TargetHeadersAddress = Target_GetDataAttributeHeadersAddress(TargetDataAddress);
-
-	const uint OutDataAddress = Out_GetDataAddress(OutDataIndex);
-	const uint OutHeadersAddress = Out_GetDataAttributeHeadersAddress(OutDataAddress);
-
-	const uint NumAttributes = Target_GetDataNumAttributes(TargetDataAddress);
+	const uint NumAttributes = Target_GetDataNumAttributesInternal(Target_DataIndex);
 	uint NumAttributesProcessed = 0;
 	
 	// Loop over all attribute headers, if the address is non-zero, then copy it from pin Target to pin Out.
 	for (int AttributeIndex = NUM_RESERVED_ATTRS; AttributeIndex < MAX_NUM_ATTRS; ++AttributeIndex)
 	{
-		const uint TargetHeaderAddress = Target_GetAttributeHeaderAddress(TargetHeadersAddress, AttributeIndex);
-		const uint TargetAddress = Target_GetAttributeAddress(TargetHeaderAddress);
-
-		const uint OutHeaderAddress = Out_GetAttributeHeaderAddress(OutHeadersAddress, AttributeIndex);
-		const uint OutAddress = Out_GetAttributeAddress(OutHeaderAddress);
+		const uint Target_Stride = Target_GetAttributeStrideInternal(Target_DataIndex, AttributeIndex);
+		const uint Out_Stride = Out_GetAttributeStrideInternal(Out_DataIndex, AttributeIndex);
 
-		const uint Stride = Target_GetAttributeStride(TargetHeaderAddress);
-
-		if (TargetAddress != 0 && OutAddress != 0)
+		if (Target_Stride != 0 && Target_Stride == Out_Stride)
 		{
-			const uint TargetBaseElementAddress = TargetAddress + TargetElemIndex * Stride;
-			const uint OutBaseElementAddress = OutAddress + OutElemIndex * Stride;
+			const uint Target_ElementAddress = Target_GetElementAddressInternal(Target_DataIndex, Target_ElementIndex, AttributeIndex);
+			const uint Out_ElementAddress = Out_GetElementAddressInternal(Out_DataIndex, Out_ElementIndex, AttributeIndex);
 
-			for (int I = 0; I < Stride; I += 4)
+			for (int Offset = 0; Offset < Out_Stride; Offset += 4)
 			{
-				const uint TargetElementAddress = TargetBaseElementAddress + I;
-				const uint OutElementAddress = OutBaseElementAddress + I;
-				Out_StoreBufferInternal(OutElementAddress, Target_LoadBufferInternal(TargetElementAddress));
+				Out_StoreBufferInternal(Out_ElementAddress + Offset, Target_LoadBufferInternal(Target_ElementAddress + Offset));
 			}
 
 			// We can early-out when we've looked at all the possible attributes
@@ -163,11 +137,8 @@ void Main(uint3 GroupId : SV_GroupID, uint GroupIndex : SV_GroupIndex)
 			SourceDataIndex = min(DataIndex, SourceNumData - 1);
 			TargetDataIndex = min(DataIndex, TargetNumData - 1);
 
-			const uint SourceDataAddress = Source_GetDataAddress(SourceDataIndex);
-			const uint TargetDataAddress = Target_GetDataAddress(TargetDataIndex);
-
-			const uint SourceElemCount = Source_GetDataInfo(SourceDataAddress);
-			const uint TargetElemCount = Target_GetDataInfo(TargetDataAddress);
+			const uint SourceElemCount = Source_GetDataNumElements(SourceDataIndex);
+			const uint TargetElemCount = Target_GetDataNumElements(TargetDataIndex);
 
 			const uint ElemCount = SourceElemCount * TargetElemCount;
 
diff --git a/Engine/Plugins/PCG/Shaders/Private/Elements/PCGStaticMeshSpawner.usf b/Engine/Plugins/PCG/Shaders/Private/Elements/PCGStaticMeshSpawner.usf
--- a/Engine/Plugins/PCG/Shaders/Private/Elements/PCGStaticMeshSpawner.usf
+++ b/Engine/Plugins/PCG/Shaders/Private/Elements/PCGStaticMeshSpawner.usf
@@ -115,8 +115,6 @@ void Main(uint3 GroupId : SV_GroupID, uint GroupIndex : SV_GroupIndex)
 	InstanceData_WriteInstanceLocalToWorld(InstanceIndex, transpose(Transform));
 
 	// Write instance custom float data.
-	const uint DataAddress = In_GetDataAddress(In_DataIndex);
-	const uint AttributeHeadersAddress = In_GetDataAttributeHeadersAddress(DataAddress);
 	const uint NumAttributes = SMSpawner_GetNumAttributes();
 	for (uint AttributeIndex = 0; AttributeIndex < NumAttributes; ++AttributeIndex)
 	{
@@ -125,15 +123,11 @@ void Main(uint3 GroupId : SV_GroupID, uint GroupIndex : SV_GroupIndex)
 		const uint AttributeOffset = AttributeIdOffsetStride[1];
 		const uint AttributeStride = AttributeIdOffsetStride[2];
 
-		const uint AttributeHeaderAddress = In_GetAttributeHeaderAddress(AttributeHeadersAddress, AttributeId);
-		const uint AttributeAddress = In_GetAttributeAddress(AttributeHeaderAddress);
-
 		for (uint FloatIndex = 0; FloatIndex < AttributeStride; ++FloatIndex)
 		{
-			const uint AttributeValueAddress = AttributeAddress + In_ElementIndex * AttributeStride * 4;
-			
+			const uint AttributeValueAddress = In_GetElementAddressInternal(In_DataIndex, In_ElementIndex, AttributeId);
 			// Assumes floating point data. TODO add type information and deal with different types in different ways?
-			InstanceData_WriteCustomFloatRaw(InstanceIndex, AttributeOffset + FloatIndex, In_LoadBufferInternal(AttributeValueAddress));
+			InstanceData_WriteCustomFloatRaw(InstanceIndex, AttributeOffset + FloatIndex, In_LoadBufferInternal(AttributeValueAddress + 4 * FloatIndex));
 		}
 	}
 
@@ -158,36 +152,23 @@ void Main(uint3 GroupId : SV_GroupID, uint GroupIndex : SV_GroupIndex)
 // TODO pass in attribute IDs to copy.
 void CopyAttributesToOuput(uint In_DataIndex, uint In_ElementIndex, uint Out_DataIndex, uint Out_ElementIndex)
 {
-	const uint In_DataAddress = In_GetDataAddress(In_DataIndex);
-	const uint In_HeadersAddress = In_GetDataAttributeHeadersAddress(In_DataAddress);
-
-	const uint Out_DataAddress = Out_GetDataAddress(Out_DataIndex);
-	const uint Out_HeadersAddress = Out_GetDataAttributeHeadersAddress(Out_DataAddress);
-
-	const uint NumAttributes = In_GetDataNumAttributes(In_DataAddress);
+	const uint NumAttributes = In_GetDataNumAttributesInternal(In_DataIndex);
 	uint NumAttributesProcessed = 0;
 	
 	// Loop over all attribute headers, if the address is non-zero, then copy it from pin Source to pin Out.
 	for (int AttributeIndex = NUM_RESERVED_ATTRS; AttributeIndex < MAX_NUM_ATTRS; ++AttributeIndex)
 	{
-		const uint In_HeaderAddress = In_GetAttributeHeaderAddress(In_HeadersAddress, AttributeIndex);
-		const uint In_Address = In_GetAttributeAddress(In_HeaderAddress);
-
-		const uint OutHeaderAddress = Out_GetAttributeHeaderAddress(Out_HeadersAddress, AttributeIndex);
-		const uint Out_Address = Out_GetAttributeAddress(OutHeaderAddress);
-
-		const uint Stride = In_GetAttributeStride(In_HeaderAddress);
+		const uint In_Stride = In_GetAttributeStrideInternal(In_DataIndex, AttributeIndex);
+		const uint Out_Stride = In_GetAttributeStrideInternal(In_DataIndex, AttributeIndex);
 
-		if (In_Address != 0 && Out_Address != 0)
+		if (In_Stride != 0 && In_Stride == Out_Stride)
 		{
-			const uint In_BaseElementAddress = In_Address + In_ElementIndex * Stride;
-			const uint Out_BaseElementAddress = Out_Address + Out_ElementIndex * Stride;
+			const uint In_ElementAddress = In_GetElementAddressInternal(In_DataIndex, In_ElementIndex, AttributeIndex);
+			const uint Out_ElementAddress = Out_GetElementAddressInternal(Out_DataIndex, Out_ElementIndex, AttributeIndex);
 
-			for (int I = 0; I < Stride; I += 4)
+			for (int Address = 0; Address < In_Stride; Address += 4)
 			{
-				const uint In_ElementAddress = In_BaseElementAddress + I;
-				const uint Out_ElementAddress = Out_BaseElementAddress + I;
-				Out_StoreBufferInternal(Out_ElementAddress, In_LoadBufferInternal(In_ElementAddress));
+				Out_StoreBufferInternal(Out_ElementAddress + Address, In_LoadBufferInternal(In_ElementAddress + Address));
 			}
 
 			// We can early-out when we've looked at all the possible attributes
diff --git a/Engine/Plugins/PCG/Shaders/Private/PCGDataCollectionDataInterface.ush b/Engine/Plugins/PCG/Shaders/Private/PCGDataCollectionDataInterface.ush
--- a/Engine/Plugins/PCG/Shaders/Private/PCGDataCollectionDataInterface.ush
+++ b/Engine/Plugins/PCG/Shaders/Private/PCGDataCollectionDataInterface.ush
@@ -1,5 +1,21 @@
 // Copyright Epic Games, Inc. All Rights Reserved.
 
+// Buffer format (sizes in bytes):
+
+// -- START HEADER --
+// PackedExecutionFlagAndNumData (4)
+// For each data:
+//   TypeId (4)
+//   Attribute Count (4)
+//   Element Count (4)
+//   For attribute ID 0 to MAX_NUM_ATTR:
+//     PackedIdAndStride (4)
+//     AttributeElementStartAddress (4)
+// -- END HEADER --
+//
+//   For all present attributes across all data:
+//     <Packed element data> (element count * stride)
+
 #include "/Engine/Private/Quaternion.ush"
 
 RWByteAddressBuffer {DataInterfaceName}_DataCollectionBuffer;
@@ -18,17 +34,9 @@ RWByteAddressBuffer {DataInterfaceName}_DataCollectionBuffer;
 
 #define INVALID_DENSITY 1.#INF
 
-// #################### DIRECT BUFFER OPERATIONS ##########################
-
-uint LoadBufferInternal_{DataInterfaceName}(uint Address)
-{
-	return {DataInterfaceName}_DataCollectionBuffer.Load(Address);
-}
-
-void StoreBufferInternal_{DataInterfaceName}(uint Address, uint Value)
-{
-	return {DataInterfaceName}_DataCollectionBuffer.Store(Address, Value);
-}
+#define DATA_COLLECTION_HEADER_SIZE_BYTES 4
+#define DATA_HEADER_SIZE_BYTES 1036
+#define ATTRIBUTE_HEADER_SIZE_BYTES 8
 
 // #################### HEADER READERS ##########################
 
@@ -38,48 +46,29 @@ uint GetNumData_{DataInterfaceName}()
 	return {DataInterfaceName}_DataCollectionBuffer.Load(0) & ~(1 << 31); // TODO: Code-gen to PCGComputeConstants::KernelExecutedFlag
 }
 
-uint GetDataAddress_{DataInterfaceName}(uint DataIndex)
-{
-	return {DataInterfaceName}_DataCollectionBuffer.Load((DataIndex + 1) * 4);
-}
-
-uint GetDataId_{DataInterfaceName}(uint DataAddress)
-{
-	return {DataInterfaceName}_DataCollectionBuffer.Load(DataAddress + 0);
-}
-
-uint GetDataNumAttributes_{DataInterfaceName}(uint DataAddress)
+uint GetDataNumElements_{DataInterfaceName}(uint InDataIndex)
 {
-	return {DataInterfaceName}_DataCollectionBuffer.Load(DataAddress + 4);
-}
+	const uint ReadAddress =
+		DATA_COLLECTION_HEADER_SIZE_BYTES +
+		InDataIndex * DATA_HEADER_SIZE_BYTES +
+		/*TypeId*/4 + /*Attribute Count*/4;
 
-uint GetDataPreambleSize_{DataInterfaceName}(uint DataAddress)
-{
-	return {DataInterfaceName}_DataCollectionBuffer.Load(DataAddress + 8);
-}
-
-// For PointData, this is just NumPoints
-uint GetDataInfo_{DataInterfaceName}(uint DataAddress)
-{
-	return {DataInterfaceName}_DataCollectionBuffer.Load(DataAddress + 12);
+	return {DataInterfaceName}_DataCollectionBuffer.Load(ReadAddress);
 }
 
-bool GetThreadData_{DataInterfaceName}(uint InThreadIndex, out uint OutDataIndex, out uint OutDataAddress, out uint OutElementIndex)
+bool GetThreadData_{DataInterfaceName}(uint InThreadIndex, out uint OutDataIndex, out uint OutElementIndex)
 {
 	int ElementIndex = InThreadIndex;
 
 	const uint NumData = GetNumData_{DataInterfaceName}();
 	for (uint DataIndex = 0; DataIndex < NumData; ++DataIndex)
 	{
-		const uint DataAddress = GetDataAddress_{DataInterfaceName}(DataIndex);
-		const uint ElemCount = GetDataInfo_{DataInterfaceName}(DataAddress);
+		const uint ElemCount = GetDataNumElements_{DataInterfaceName}(DataIndex);
 
 		if (ElementIndex < ElemCount)
 		{
 			OutDataIndex = DataIndex;
-			OutDataAddress = DataAddress;
 			OutElementIndex = ElementIndex;
-
 			return true;
 		}
 
@@ -87,18 +76,10 @@ bool GetThreadData_{DataInterfaceName}(uint InThreadIndex, out uint OutDataIndex
 	}
 
 	OutDataIndex = -1;
-	OutDataAddress = -1;
 	OutElementIndex = -1;
-
 	return false;
 }
 
-bool GetThreadData_{DataInterfaceName}(uint InThreadIndex, out uint OutDataIndex, out uint OutElementIndex)
-{
-	uint DataAddressDummy = -1;
-	return GetThreadData_{DataInterfaceName}(InThreadIndex, OutDataIndex, DataAddressDummy, OutElementIndex);
-}
-
 uint GetNumElements_{DataInterfaceName}()
 {
 	uint NumElements = 0;
@@ -106,63 +87,65 @@ uint GetNumElements_{DataInterfaceName}()
 	const uint NumData = GetNumData_{DataInterfaceName}();
 	for (uint DataIndex = 0; DataIndex < NumData; ++DataIndex)
 	{
-		const uint DataAddress = GetDataAddress_{DataInterfaceName}(DataIndex);
-		NumElements += GetDataInfo_{DataInterfaceName}(DataAddress);
+		NumElements += GetDataNumElements_{DataInterfaceName}(DataIndex);
 	}
 
 	return NumElements;
 }
 
-uint GetDataAttributeHeadersAddress_{DataInterfaceName}(uint DataAddress)
-{
-	return DataAddress + GetDataPreambleSize_{DataInterfaceName}(DataAddress);
-}
+// #################### INTERNAL HELPERS ##########################
 
-uint GetAttributeHeaderAddress_{DataInterfaceName}(uint AttributeHeadersAddress, uint AttributeId)
+uint LoadBufferInternal_{DataInterfaceName}(uint Address)
 {
-	return AttributeHeadersAddress + AttributeId * 8;
+	return {DataInterfaceName}_DataCollectionBuffer.Load(Address);
 }
 
-uint GetAttributeIdAndStride_{DataInterfaceName}(uint AttributeHeaderAddress)
+void StoreBufferInternal_{DataInterfaceName}(uint Address, uint Value)
 {
-	return {DataInterfaceName}_DataCollectionBuffer.Load(AttributeHeaderAddress);
+	return {DataInterfaceName}_DataCollectionBuffer.Store(Address, Value);
 }
 
-uint GetAttributeId_{DataInterfaceName}(uint AttributeHeaderAddress)
+uint GetFirstElementAddressInternal_{DataInterfaceName}(uint InDataIndex, uint InAttributeId)
 {
-	return GetAttributeIdAndStride_{DataInterfaceName}(AttributeHeaderAddress) >> 8;
+	const uint ReadAddress =
+		DATA_COLLECTION_HEADER_SIZE_BYTES +
+		InDataIndex * DATA_HEADER_SIZE_BYTES +
+		/*TypeId*/4 + /*Attribute Count*/4 + /*Element Count*/4 +
+		InAttributeId * ATTRIBUTE_HEADER_SIZE_BYTES +
+		/*PackedAttributeIdAndStride*/4;
+
+	return {DataInterfaceName}_DataCollectionBuffer.Load(ReadAddress);
 }
 
-uint GetAttributeStride_{DataInterfaceName}(uint AttributeHeaderAddress)
+uint GetAttributeStrideInternal_{DataInterfaceName}(uint InDataIndex, uint InAttributeId)
 {
-	return GetAttributeIdAndStride_{DataInterfaceName}(AttributeHeaderAddress) & (0xFFFFFFFF >> 24);
+	const uint ReadAddress =
+		DATA_COLLECTION_HEADER_SIZE_BYTES +
+		InDataIndex * DATA_HEADER_SIZE_BYTES +
+		/*TypeId*/4 + /*Attribute Count*/4 + /*Element Count*/4 +
+		InAttributeId * ATTRIBUTE_HEADER_SIZE_BYTES;
+
+	const uint PackedAttributeIdAndStride = {DataInterfaceName}_DataCollectionBuffer.Load(ReadAddress);
+	return PackedAttributeIdAndStride & (0xFFFFFFFF >> 24);
 }
 
-uint GetAttributeAddress_{DataInterfaceName}(uint AttributeHeaderAddress)
+uint GetElementAddressInternal_{DataInterfaceName}(uint InDataIndex, uint InElementIndex, uint InAttributeId)
 {
-	return {DataInterfaceName}_DataCollectionBuffer.Load(AttributeHeaderAddress + 4);
+	const uint FirstElementAddress = GetFirstElementAddressInternal_{DataInterfaceName}(InDataIndex, InAttributeId);
+	const uint AttributeStride = GetAttributeStrideInternal_{DataInterfaceName}(InDataIndex, InAttributeId);
+	return InElementIndex * AttributeStride + FirstElementAddress;
 }
 
-uint GetAttributeAddress_{DataInterfaceName}(uint InDataAddress, uint InAttributeId)
+uint GetDataNumAttributesInternal_{DataInterfaceName}(uint InDataIndex)
 {
-	const uint AttributeHeadersAddress = GetDataAttributeHeadersAddress_{DataInterfaceName}(InDataAddress);
-	const uint NumAttributes = GetDataNumAttributes_{DataInterfaceName}(InDataAddress);
-
-	for (uint AttributeIndex = 0; AttributeIndex < NumAttributes; ++AttributeIndex)
-	{
-		const uint AttributeHeaderAddress = GetAttributeHeaderAddress_{DataInterfaceName}(AttributeHeadersAddress, AttributeIndex);
-
-		if (GetAttributeId_{DataInterfaceName}(AttributeHeaderAddress) == InAttributeId)
-		{
-			return GetAttributeAddress_{DataInterfaceName}(AttributeHeaderAddress);
-		}
-	}
+	const uint ReadAddress =
+		DATA_COLLECTION_HEADER_SIZE_BYTES +
+		InDataIndex * DATA_HEADER_SIZE_BYTES +
+		/*TypeId*/4;
 
-	return -1;
+	return {DataInterfaceName}_DataCollectionBuffer.Load(ReadAddress);
 }
 
-// #################### HEADER WRITERS ##########################
-
 void SetAsExecutedInternal_{DataInterfaceName}()
 {
 	// We use the most significant bit of NumData to store the KernelExecuted flag
@@ -170,103 +153,59 @@ void SetAsExecutedInternal_{DataInterfaceName}()
 	{DataInterfaceName}_DataCollectionBuffer.Store(0, NumData | (1 << 31)); // TODO: Code-gen to PCGComputeConstants::KernelExecutedFlag
 }
 
-void SetDataAddressInternal_{DataInterfaceName}(uint DataIndex, uint DataAddress)
-{
-	{DataInterfaceName}_DataCollectionBuffer.Store((DataIndex + 1) * 4, DataAddress);
-}
-
-void SetDataIdInternal_{DataInterfaceName}(uint DataAddress, uint Id)
-{
-	{DataInterfaceName}_DataCollectionBuffer.Store(DataAddress + 0, Id);
-}
-
-void SetDataNumAttributesInternal_{DataInterfaceName}(uint DataAddress, uint NumAttributes)
-{
-	{DataInterfaceName}_DataCollectionBuffer.Store(DataAddress + 4, NumAttributes);
-}
-
-void SetDataPreambleSizeInternal_{DataInterfaceName}(uint DataAddress, uint PreambleSize)
-{
-	{DataInterfaceName}_DataCollectionBuffer.Store(DataAddress + 8, PreambleSize);
-}
-
-// For PointData, this is just NumPoints
-void SetDataInfoInternal_{DataInterfaceName}(uint DataAddress, uint Info)
-{
-	{DataInterfaceName}_DataCollectionBuffer.Store(DataAddress + 12, Info);
-}
-
-void SetAttributeIdAndStrideInternal_{DataInterfaceName}(uint AttributeHeaderAddress, uint IdAndStride)
-{
-	{DataInterfaceName}_DataCollectionBuffer.Store(AttributeHeaderAddress, IdAndStride);
-}
-
-void SetAttributeAddressInternal_{DataInterfaceName}(uint AttributeHeaderAddress, uint AttributeAddress)
-{
-	{DataInterfaceName}_DataCollectionBuffer.Store(AttributeHeaderAddress + 4, AttributeAddress);
-}
-
 // #################### ATTRIBUTE GETTERS ##########################
 
-uint ComputeAddressHelper_{DataInterfaceName}(uint DataIndex, uint ElementIndex, uint AttributeId)
-{
-	const uint DataAddress = GetDataAddress_{DataInterfaceName}(DataIndex);
-	const uint AttributeHeadersAddress = GetDataAttributeHeadersAddress_{DataInterfaceName}(DataAddress);
-	const uint AttributeHeaderAddress = GetAttributeHeaderAddress_{DataInterfaceName}(AttributeHeadersAddress, AttributeId);
-	return GetAttributeAddress_{DataInterfaceName}(AttributeHeaderAddress);
-}
-
 bool GetBool_{DataInterfaceName}(uint DataIndex, uint ElementIndex, uint AttributeId)
 {
-	const uint AttributeAddress = ComputeAddressHelper_{DataInterfaceName}(DataIndex, ElementIndex, AttributeId);
+	const uint AttributeAddress = GetFirstElementAddressInternal_{DataInterfaceName}(DataIndex, AttributeId);
 	return asbool({DataInterfaceName}_DataCollectionBuffer.Load(AttributeAddress + ElementIndex * 1 * 4));
 }
 
 int GetInt_{DataInterfaceName}(uint DataIndex, uint ElementIndex, uint AttributeId)
 {
-	const uint AttributeAddress = ComputeAddressHelper_{DataInterfaceName}(DataIndex, ElementIndex, AttributeId);
+	const uint AttributeAddress = GetFirstElementAddressInternal_{DataInterfaceName}(DataIndex, AttributeId);
 	return asint({DataInterfaceName}_DataCollectionBuffer.Load(AttributeAddress + ElementIndex * 1 * 4));
 }
 
 int GetUint_{DataInterfaceName}(uint DataIndex, uint ElementIndex, uint AttributeId)
 {
-	const uint AttributeAddress = ComputeAddressHelper_{DataInterfaceName}(DataIndex, ElementIndex, AttributeId);
+	const uint AttributeAddress = GetFirstElementAddressInternal_{DataInterfaceName}(DataIndex, AttributeId);
 	return {DataInterfaceName}_DataCollectionBuffer.Load(AttributeAddress + ElementIndex * 1 * 4);
 }
 
 float GetFloat_{DataInterfaceName}(uint DataIndex, uint ElementIndex, uint AttributeId)
 {
-	const uint AttributeAddress = ComputeAddressHelper_{DataInterfaceName}(DataIndex, ElementIndex, AttributeId);
+	const uint AttributeAddress = GetFirstElementAddressInternal_{DataInterfaceName}(DataIndex, AttributeId);
 	return asfloat({DataInterfaceName}_DataCollectionBuffer.Load(AttributeAddress + ElementIndex * 1 * 4));
 }
 
 float2 GetFloat2_{DataInterfaceName}(uint DataIndex, uint ElementIndex, uint AttributeId)
 {
-	const uint AttributeAddress = ComputeAddressHelper_{DataInterfaceName}(DataIndex, ElementIndex, AttributeId);
+	const uint AttributeAddress = GetFirstElementAddressInternal_{DataInterfaceName}(DataIndex, AttributeId);
 	return asfloat({DataInterfaceName}_DataCollectionBuffer.Load2(AttributeAddress + ElementIndex * 2 * 4));
 }
 
 float3 GetFloat3_{DataInterfaceName}(uint DataIndex, uint ElementIndex, uint AttributeId)
 {
-	const uint AttributeAddress = ComputeAddressHelper_{DataInterfaceName}(DataIndex, ElementIndex, AttributeId);
+	const uint AttributeAddress = GetFirstElementAddressInternal_{DataInterfaceName}(DataIndex, AttributeId);
 	return asfloat({DataInterfaceName}_DataCollectionBuffer.Load3(AttributeAddress + ElementIndex * 3 * 4));
 }
 
 float4 GetFloat4_{DataInterfaceName}(uint DataIndex, uint ElementIndex, uint AttributeId)
 {
-	const uint AttributeAddress = ComputeAddressHelper_{DataInterfaceName}(DataIndex, ElementIndex, AttributeId);
+	const uint AttributeAddress = GetFirstElementAddressInternal_{DataInterfaceName}(DataIndex, AttributeId);
 	return asfloat({DataInterfaceName}_DataCollectionBuffer.Load4(AttributeAddress + ElementIndex * 4 * 4));
 }
 
 float3 GetRotator_{DataInterfaceName}(uint DataIndex, uint ElementIndex, uint AttributeId)
 {
-	const uint AttributeAddress = ComputeAddressHelper_{DataInterfaceName}(DataIndex, ElementIndex, AttributeId);
+	const uint AttributeAddress = GetFirstElementAddressInternal_{DataInterfaceName}(DataIndex, AttributeId);
 	return asfloat({DataInterfaceName}_DataCollectionBuffer.Load3(AttributeAddress + ElementIndex * 3 * 4));
 }
 
 float4 GetQuat_{DataInterfaceName}(uint DataIndex, uint ElementIndex, uint AttributeId)
 {
-	const uint AttributeAddress = ComputeAddressHelper_{DataInterfaceName}(DataIndex, ElementIndex, AttributeId);
+	const uint AttributeAddress = GetFirstElementAddressInternal_{DataInterfaceName}(DataIndex, AttributeId);
 	return asfloat({DataInterfaceName}_DataCollectionBuffer.Load4(AttributeAddress + ElementIndex * 4 * 4));
 }
 
@@ -280,55 +219,55 @@ int GetStringKey_{DataInterfaceName}(uint DataIndex, uint ElementIndex, uint Att
 
 void SetBool_{DataInterfaceName}(uint DataIndex, uint ElementIndex, uint AttributeId, bool Value)
 {
-	const uint AttributeAddress = ComputeAddressHelper_{DataInterfaceName}(DataIndex, ElementIndex, AttributeId);
+	const uint AttributeAddress = GetFirstElementAddressInternal_{DataInterfaceName}(DataIndex, AttributeId);
 	{DataInterfaceName}_DataCollectionBuffer.Store(AttributeAddress + ElementIndex * 1 * 4, asuint(Value));
 }
 
 void SetInt_{DataInterfaceName}(uint DataIndex, uint ElementIndex, uint AttributeId, int Value)
 {
-	const uint AttributeAddress = ComputeAddressHelper_{DataInterfaceName}(DataIndex, ElementIndex, AttributeId);
+	const uint AttributeAddress = GetFirstElementAddressInternal_{DataInterfaceName}(DataIndex, AttributeId);
 	{DataInterfaceName}_DataCollectionBuffer.Store(AttributeAddress + ElementIndex * 1 * 4, asuint(Value));
 }
 
 void SetUint_{DataInterfaceName}(uint DataIndex, uint ElementIndex, uint AttributeId, uint Value)
 {
-	const uint AttributeAddress = ComputeAddressHelper_{DataInterfaceName}(DataIndex, ElementIndex, AttributeId);
+	const uint AttributeAddress = GetFirstElementAddressInternal_{DataInterfaceName}(DataIndex, AttributeId);
 	{DataInterfaceName}_DataCollectionBuffer.Store(AttributeAddress + ElementIndex * 1 * 4, Value);
 }
 
 void SetFloat_{DataInterfaceName}(uint DataIndex, uint ElementIndex, uint AttributeId, float Value)
 {
-	const uint AttributeAddress = ComputeAddressHelper_{DataInterfaceName}(DataIndex, ElementIndex, AttributeId);
+	const uint AttributeAddress = GetFirstElementAddressInternal_{DataInterfaceName}(DataIndex, AttributeId);
 	{DataInterfaceName}_DataCollectionBuffer.Store(AttributeAddress + ElementIndex * 1 * 4, asuint(Value));
 }
 
 void SetFloat2_{DataInterfaceName}(uint DataIndex, uint ElementIndex, uint AttributeId, float2 Value)
 {
-	const uint AttributeAddress = ComputeAddressHelper_{DataInterfaceName}(DataIndex, ElementIndex, AttributeId);
+	const uint AttributeAddress = GetFirstElementAddressInternal_{DataInterfaceName}(DataIndex, AttributeId);
 	{DataInterfaceName}_DataCollectionBuffer.Store2(AttributeAddress + ElementIndex * 2 * 4, asuint(Value));
 }
 
 void SetFloat3_{DataInterfaceName}(uint DataIndex, uint ElementIndex, uint AttributeId, float3 Value)
 {
-	const uint AttributeAddress = ComputeAddressHelper_{DataInterfaceName}(DataIndex, ElementIndex, AttributeId);
+	const uint AttributeAddress = GetFirstElementAddressInternal_{DataInterfaceName}(DataIndex, AttributeId);
 	{DataInterfaceName}_DataCollectionBuffer.Store3(AttributeAddress + ElementIndex * 3 * 4, asuint(Value));
 }
 
 void SetFloat4_{DataInterfaceName}(uint DataIndex, uint ElementIndex, uint AttributeId, float4 Value)
 {
-	const uint AttributeAddress = ComputeAddressHelper_{DataInterfaceName}(DataIndex, ElementIndex, AttributeId);
+	const uint AttributeAddress = GetFirstElementAddressInternal_{DataInterfaceName}(DataIndex, AttributeId);
 	{DataInterfaceName}_DataCollectionBuffer.Store4(AttributeAddress + ElementIndex * 4 * 4, asuint(Value));
 }
 
 void SetRotator_{DataInterfaceName}(uint DataIndex, uint ElementIndex, uint AttributeId, float3 Value)
 {
-	const uint AttributeAddress = ComputeAddressHelper_{DataInterfaceName}(DataIndex, ElementIndex, AttributeId);
+	const uint AttributeAddress = GetFirstElementAddressInternal_{DataInterfaceName}(DataIndex, AttributeId);
 	{DataInterfaceName}_DataCollectionBuffer.Store3(AttributeAddress + ElementIndex * 3 * 4, asuint(Value));
 }
 
 void SetQuat_{DataInterfaceName}(uint DataIndex, uint ElementIndex, uint AttributeId, float4 Value)
 {
-	const uint AttributeAddress = ComputeAddressHelper_{DataInterfaceName}(DataIndex, ElementIndex, AttributeId);
+	const uint AttributeAddress = GetFirstElementAddressInternal_{DataInterfaceName}(DataIndex, AttributeId);
 	{DataInterfaceName}_DataCollectionBuffer.Store4(AttributeAddress + ElementIndex * 4 * 4, asuint(Value));
 }
 
@@ -491,5 +430,5 @@ void InitializePoint_{DataInterfaceName}(uint InDataIndex, uint InElementIndex)
 	SetSeed_{DataInterfaceName}(InDataIndex, InElementIndex, 42);
 	SetSteepness_{DataInterfaceName}(InDataIndex, InElementIndex, 0.5f);
 	
-	// TODO - code gen initialization of all attributes, or remove this function.
+	// TODO - pass in relevant attribute IDs and init them
 }
diff --git a/Engine/Plugins/PCG/Source/PCG/Private/Compute/DataInterfaces/PCGDataCollectionDataInterface.cpp b/Engine/Plugins/PCG/Source/PCG/Private/Compute/DataInterfaces/PCGDataCollectionDataInterface.cpp
--- a/Engine/Plugins/PCG/Source/PCG/Private/Compute/DataInterfaces/PCGDataCollectionDataInterface.cpp
+++ b/Engine/Plugins/PCG/Source/PCG/Private/Compute/DataInterfaces/PCGDataCollectionDataInterface.cpp
@@ -21,59 +21,34 @@
 
 void UPCGDataCollectionDataInterface::GetSupportedInputs(TArray<FShaderFunctionDefinition>& OutFunctions) const
 {
-	OutFunctions.AddDefaulted_GetRef()
-		.SetName(TEXT("LoadBufferInternal"))
-		.AddReturnType(EShaderFundamentalType::Uint)
-		.AddParam(EShaderFundamentalType::Uint);
-
-	OutFunctions.AddDefaulted_GetRef()
-		.SetName(TEXT("StoreBufferInternal"))
-		.AddParam(EShaderFundamentalType::Uint)
-		.AddParam(EShaderFundamentalType::Uint);
-
-	// Header Readers
+	// Internal direct access to buffer, used to optimize access.
 	{
 		OutFunctions.AddDefaulted_GetRef()
-			.SetName(TEXT("GetNumData"))
-			.AddReturnType(EShaderFundamentalType::Uint);
-
-		OutFunctions.AddDefaulted_GetRef()
-			.SetName(TEXT("GetDataAddress"))
-			.AddReturnType(EShaderFundamentalType::Uint)
-			.AddParam(EShaderFundamentalType::Uint);
-
-		OutFunctions.AddDefaulted_GetRef()
-			.SetName(TEXT("GetDataId"))
+			.SetName(TEXT("LoadBufferInternal"))
 			.AddReturnType(EShaderFundamentalType::Uint)
 			.AddParam(EShaderFundamentalType::Uint);
 
 		OutFunctions.AddDefaulted_GetRef()
-			.SetName(TEXT("GetDataNumAttributes"))
-			.AddReturnType(EShaderFundamentalType::Uint)
+			.SetName(TEXT("StoreBufferInternal"))
+			.AddParam(EShaderFundamentalType::Uint)
 			.AddParam(EShaderFundamentalType::Uint);
+	}
 
+	// Header Readers
+	{
 		OutFunctions.AddDefaulted_GetRef()
-			.SetName(TEXT("GetDataPreambleSize"))
-			.AddReturnType(EShaderFundamentalType::Uint)
-			.AddParam(EShaderFundamentalType::Uint);
+			.SetName(TEXT("GetNumData"))
+			.AddReturnType(EShaderFundamentalType::Uint);
 
 		OutFunctions.AddDefaulted_GetRef()
-			.SetName(TEXT("GetDataInfo"))
+			.SetName(TEXT("GetDataNumAttributesInternal"))
 			.AddReturnType(EShaderFundamentalType::Uint)
-			.AddParam(EShaderFundamentalType::Uint);
+			.AddParam(EShaderFundamentalType::Uint); // InDataIndex
 
 		OutFunctions.AddDefaulted_GetRef()
-			.SetName(TEXT("GetDataAttributeHeadersAddress"))
+			.SetName(TEXT("GetDataNumElements"))
 			.AddReturnType(EShaderFundamentalType::Uint)
-			.AddParam(EShaderFundamentalType::Uint);
-
-		OutFunctions.AddDefaulted_GetRef()
-			.SetName(TEXT("GetThreadData"))
-			.AddParam(EShaderFundamentalType::Uint) // InThreadIndex
-			.AddParam(EShaderFundamentalType::Uint, 0, 0, EShaderParamModifier::Out) // OutDataIndex
-			.AddParam(EShaderFundamentalType::Uint, 0, 0, EShaderParamModifier::Out) // OutDataAddress
-			.AddParam(EShaderFundamentalType::Uint, 0, 0, EShaderParamModifier::Out) // OutElementIndex
-			.AddReturnType(EShaderFundamentalType::Bool);
+			.AddParam(EShaderFundamentalType::Uint); // InDataIndex
 
 		OutFunctions.AddDefaulted_GetRef()
 			.SetName(TEXT("GetThreadData"))
@@ -87,80 +62,23 @@ void UPCGDataCollectionDataInterface::GetSupportedInputs(TArray<FShaderFunctionD
 			.AddReturnType(EShaderFundamentalType::Uint);
 
 		OutFunctions.AddDefaulted_GetRef()
-			.SetName(TEXT("GetAttributeHeaderAddress"))
+			.SetName(TEXT("GetElementAddressInternal"))
 			.AddReturnType(EShaderFundamentalType::Uint)
-			.AddParam(EShaderFundamentalType::Uint)
-			.AddParam(EShaderFundamentalType::Uint);
-
-		OutFunctions.AddDefaulted_GetRef()
-			.SetName(TEXT("GetAttributeIdAndStride"))
-			.AddReturnType(EShaderFundamentalType::Uint)
-			.AddParam(EShaderFundamentalType::Uint);
-
-		OutFunctions.AddDefaulted_GetRef()
-			.SetName(TEXT("GetAttributeId"))
-			.AddReturnType(EShaderFundamentalType::Uint)
-			.AddParam(EShaderFundamentalType::Uint);
-
-		OutFunctions.AddDefaulted_GetRef()
-			.SetName(TEXT("GetAttributeStride"))
-			.AddReturnType(EShaderFundamentalType::Uint)
-			.AddParam(EShaderFundamentalType::Uint);
+			.AddParam(EShaderFundamentalType::Uint) // InDataIndex
+			.AddParam(EShaderFundamentalType::Uint) // InElementIndex
+			.AddParam(EShaderFundamentalType::Uint); // InAttributeId
 
 		OutFunctions.AddDefaulted_GetRef()
-			.SetName(TEXT("GetAttributeAddress"))
+			.SetName(TEXT("GetAttributeStrideInternal"))
 			.AddReturnType(EShaderFundamentalType::Uint)
-			.AddParam(EShaderFundamentalType::Uint);
-
-		OutFunctions.AddDefaulted_GetRef()
-			.SetName(TEXT("GetAttributeAddress"))
-			.AddReturnType(EShaderFundamentalType::Uint)
-			.AddParam(EShaderFundamentalType::Uint)
-			.AddParam(EShaderFundamentalType::Uint);
+			.AddParam(EShaderFundamentalType::Uint) // InDataIndex
+			.AddParam(EShaderFundamentalType::Uint); // InAttributeId
 	}
 
 	// Header Writers
 	{
 		OutFunctions.AddDefaulted_GetRef()
 			.SetName(TEXT("SetAsExecutedInternal"));
-
-		OutFunctions.AddDefaulted_GetRef()
-			.SetName(TEXT("SetDataAddressInternal"))
-			.AddParam(EShaderFundamentalType::Uint)
-			.AddParam(EShaderFundamentalType::Uint);
-
-		OutFunctions.AddDefaulted_GetRef()
-			.SetName(TEXT("SetDataIdInternal"))
-			.AddParam(EShaderFundamentalType::Uint)
-			.AddParam(EShaderFundamentalType::Uint);
-
-		OutFunctions.AddDefaulted_GetRef()
-			.SetName(TEXT("SetDataNumAttributesInternal"))
-			.AddParam(EShaderFundamentalType::Uint)
-			.AddParam(EShaderFundamentalType::Uint);
-
-		OutFunctions.AddDefaulted_GetRef()
-			.SetName(TEXT("SetDataPreambleSizeInternal"))
-			.AddParam(EShaderFundamentalType::Uint)
-			.AddParam(EShaderFundamentalType::Uint);
-
-		OutFunctions.AddDefaulted_GetRef()
-			.SetName(TEXT("SetDataInfoInternal"))
-			.AddParam(EShaderFundamentalType::Uint)
-			.AddParam(EShaderFundamentalType::Uint);
-
-		OutFunctions.AddDefaulted_GetRef()
-			.SetName(TEXT("SetAttributeIdAndStrideInternal"))
-			.AddParam(EShaderFundamentalType::Uint)
-			.AddParam(EShaderFundamentalType::Uint);
-
-		OutFunctions.AddDefaulted_GetRef()
-			.SetName(TEXT("SetAttributeAddressInternal"))
-			.AddParam(EShaderFundamentalType::Uint)
-			.AddParam(EShaderFundamentalType::Uint);
-
-		OutFunctions.AddDefaulted_GetRef()
-			.SetName(TEXT("WriteHeader"));
 	}
 
 	// Attribute Getters
@@ -661,7 +579,7 @@ FPCGDataCollectionDataProviderProxy::FPCGDataCollectionDataProviderProxy(const F
 	: ReadbackMode(InReadbackMode)
 	, PinDesc(InPinDesc)
 {
-	SizeBytes = PinDesc.ComputePackedSize();
+	SizeBytes = PinDesc.ComputePackedSizeBytes();
 }
 
 bool FPCGDataCollectionDataProviderProxy::IsValid(FValidationData const& InValidationData) const
@@ -697,15 +615,14 @@ void FPCGDataCollectionDataProviderProxy::AllocateResources(FRDGBuilder& GraphBu
 	TRACE_CPUPROFILER_EVENT_SCOPE(FPCGDataCollectionDataProviderProxy::AllocateResources);
 	check(SizeBytes > 0);
 
+	FRDGBufferDesc Desc = FRDGBufferDesc::CreateStructuredDesc(sizeof(uint32), SizeBytes >> 2);
+	Buffer = GraphBuilder.CreateBuffer(Desc, TEXT("PCGDataCollectionBuffer"));
+	BufferUAV = GraphBuilder.CreateUAV(Buffer);
+
 	// Initialize with an empty data collection. The kernel may not run, for example if indirect dispatch args end up being 0. Ensure
 	// there is something meaningful to readback.
-	// TODO could have a statically-allocated resource rather than allocating & uploading here.
 	TArray<uint32> PackedDataCollection;
-	PinDesc.PrepareBufferForKernelOutput(PackedDataCollection);
-
-	FRDGBufferDesc Desc = FRDGBufferDesc::CreateStructuredDesc(PackedDataCollection.GetTypeSize(), PackedDataCollection.Num());
-	Buffer = GraphBuilder.CreateBuffer(Desc, TEXT("PCGDataCollectionBuffer"));
-	BufferUAV = GraphBuilder.CreateUAV(Buffer);
+	PinDesc.WriteHeader(PackedDataCollection);
 
 	GraphBuilder.QueueBufferUpload(Buffer, PackedDataCollection.GetData(), PackedDataCollection.Num() * PackedDataCollection.GetTypeSize(), ERDGInitialDataFlags::None);
 }
diff --git a/Engine/Plugins/PCG/Source/PCG/Private/Compute/Elements/PCGCustomHLSL.cpp b/Engine/Plugins/PCG/Source/PCG/Private/Compute/Elements/PCGCustomHLSL.cpp
--- a/Engine/Plugins/PCG/Source/PCG/Private/Compute/Elements/PCGCustomHLSL.cpp
+++ b/Engine/Plugins/PCG/Source/PCG/Private/Compute/Elements/PCGCustomHLSL.cpp
@@ -921,9 +921,8 @@ void UPCGCustomHLSLSettings::UpdateHelperDeclarations()
 					HelperDeclarations += TEXT("// Valid pins: ") + FString::Join(PinNames, TEXT(", ")) + TEXT("\n");
 				}
 
-				HelperDeclarations += FString::Format(TEXT(
-					"bool {0}_GetThreadData(uint ThreadIndex, out uint OutDataIndex, out uint OutElementIndex);\n"
-					"bool {0}_GetThreadData(uint ThreadIndex, out uint OutDataIndex, out uint OutDataAddress, out uint OutElementIndex);\n"),
+				HelperDeclarations += FString::Format(
+					TEXT("bool {0}_GetThreadData(uint ThreadIndex, out uint OutDataIndex, out uint OutElementIndex);\n"),
 					{ bMultiPin ? PCGHLSLElement::PinDeclTemplateStr : PinNames[0] });
 
 				HelperDeclarations += TEXT("\n");
@@ -1614,8 +1613,7 @@ FString UPCGCustomHLSLSettings::GetCookedKernelSource(const TMap<FName, FPCGKern
 	{
 		KernelSpecificPreamble += FString::Format(TEXT(
 			"    uint {0}_DataIndex;\n"
-			"    uint {0}_DataAddress;\n"
-			"    if (!{0}_GetThreadData(ThreadIndex, {0}_DataIndex, {0}_DataAddress, ElementIndex)) return;\n"),
+			"    if (!{0}_GetThreadData(ThreadIndex, {0}_DataIndex, ElementIndex)) return;\n"),
 			{ PinLabel.ToString() });
 	};
 
@@ -1641,28 +1639,26 @@ FString UPCGCustomHLSLSettings::GetCookedKernelSource(const TMap<FName, FPCGKern
 				{ InputPin->Properties.Label.ToString(), OutputPin->Properties.Label.ToString() });
 
 			// Automatically copy value of all attributes for this element.
+			// TODO pass in IDs of attributes that are actually present.
 			KernelSpecificPreamble += FString::Format(TEXT(
 				"\n"
-				"    // Loop over all attribute headers, if the address is non-zero, then copy it from pin {0} to pin {1}.\n"
+				"    // Initialize output data elements. Loop over all attribute headers, if the address is non-zero, then copy it from pin {0} to pin {1}.\n"
 				"    {\n"
-				"        const uint HeadersAddress = {0}_GetDataAttributeHeadersAddress({0}_DataAddress);\n"
-				"        const uint NumAttributes = {0}_GetDataNumAttributes({0}_DataAddress);\n"
+				"        const uint NumAttributes = {0}_GetDataNumAttributesInternal({0}_DataIndex);\n"
 				"        uint NumAttributesProcessed = 0;\n"
 				"\n"
 				"        for (int AttributeIndex = 0; AttributeIndex < 128; ++AttributeIndex)\n"
 				"        {\n"
-				"            const uint HeaderAddress = {0}_GetAttributeHeaderAddress(HeadersAddress, AttributeIndex);\n"
-				"            const uint Stride = {0}_GetAttributeStride(HeaderAddress);\n"
-				"            const uint Address = {0}_GetAttributeAddress(HeaderAddress);\n"
+				"            const uint Stride = {0}_GetAttributeStrideInternal({0}_DataIndex, AttributeIndex);\n"
 				"\n"
-				"            if (Address != 0)\n"
+				"            if (Stride != 0)\n"
 				"            {\n"
-				"                const uint BaseElementAddress = Address + ElementIndex * Stride;\n"
+				"                const uint {0}_ElementAddress = {0}_GetElementAddressInternal({0}_DataIndex, ElementIndex, AttributeIndex);\n"
+				"                const uint {1}_ElementAddress = {1}_GetElementAddressInternal({1}_DataIndex, ElementIndex, AttributeIndex);\n"
 				"\n"
-				"                for (int I = 0; I < Stride; I += 4)\n"
+				"                for (uint Address = 0; Address < Stride; Address += 4)\n"
 				"                {\n"
-				"                    const uint ElementAddress = BaseElementAddress + I;\n"
-				"                    {1}_StoreBufferInternal(ElementAddress, {0}_LoadBufferInternal(ElementAddress));\n"
+				"                    {1}_StoreBufferInternal({1}_ElementAddress + Address, {0}_LoadBufferInternal({0}_ElementAddress + Address));\n"
 				"                }\n"
 				"\n"
 				"                if (++NumAttributesProcessed >= NumAttributes) break; // We can early-out when we've looked at all the possible attributes\n"
diff --git a/Engine/Plugins/PCG/Source/PCG/Private/Compute/PCGComputeCommon.h b/Engine/Plugins/PCG/Source/PCG/Private/Compute/PCGComputeCommon.h
--- a/Engine/Plugins/PCG/Source/PCG/Private/Compute/PCGComputeCommon.h
+++ b/Engine/Plugins/PCG/Source/PCG/Private/Compute/PCGComputeCommon.h
@@ -27,13 +27,16 @@ namespace PCGComputeConstants
 	constexpr int NUM_RESERVED_ATTRS = 32; // Reserved for point properties, spline accessors, etc.
 	constexpr int MAX_NUM_CUSTOM_ATTRS = MAX_NUM_ATTRS - NUM_RESERVED_ATTRS; // Reserved for custom attributes
 
+	constexpr int DATA_COLLECTION_HEADER_SIZE_BYTES = 4;
+	constexpr int DATA_HEADER_SIZE_BYTES = 1036;
 	constexpr int ATTRIBUTE_HEADER_SIZE_BYTES = 8;
 
 	constexpr int POINT_DATA_TYPE_ID = 0;
+	constexpr int PARAM_DATA_TYPE_ID = 1;
+
+	// TODO remove
 	constexpr int POINT_DATA_HEADER_PREAMBLE_SIZE_BYTES = 16; // 4 bytes for Type, 4 bytes for NumAttrs, 4 bytes for the address, 4 bytes for TypeInfo
 	constexpr int POINT_DATA_HEADER_SIZE_BYTES = POINT_DATA_HEADER_PREAMBLE_SIZE_BYTES + ATTRIBUTE_HEADER_SIZE_BYTES * MAX_NUM_ATTRS;
-
-	constexpr int PARAM_DATA_TYPE_ID = 1;
 	constexpr int PARAM_DATA_HEADER_PREAMBLE_SIZE_BYTES = 16; // 4 bytes for Type, 4 bytes for NumAttrs, 4 bytes for the address, 4 bytes for TypeInfo
 	constexpr int PARAM_DATA_HEADER_SIZE_BYTES = PARAM_DATA_HEADER_PREAMBLE_SIZE_BYTES + ATTRIBUTE_HEADER_SIZE_BYTES * MAX_NUM_ATTRS;
 
diff --git a/Engine/Plugins/PCG/Source/PCG/Private/Compute/PCGDataForGPU.cpp b/Engine/Plugins/PCG/Source/PCG/Private/Compute/PCGDataForGPU.cpp
--- a/Engine/Plugins/PCG/Source/PCG/Private/Compute/PCGDataForGPU.cpp
+++ b/Engine/Plugins/PCG/Source/PCG/Private/Compute/PCGDataForGPU.cpp
@@ -92,7 +92,7 @@ namespace PCGDataForGPUHelpers
 		}
 	}
 
-	bool PackAttributeHelper(const FPCGMetadataAttributeBase* InAttributeBase, const FPCGKernelAttributeDesc& InAttributeDesc, PCGMetadataEntryKey InEntryKey, const TArray<FString>& InStringTable, TArray<uint32>& OutPackedDataCollection, uint32 ElementIndex)
+	bool PackAttributeHelper(const FPCGMetadataAttributeBase* InAttributeBase, const FPCGKernelAttributeDesc& InAttributeDesc, PCGMetadataEntryKey InEntryKey, const TArray<FString>& InStringTable, TArray<uint32>& OutPackedDataCollection, uint32& InOutAddressUints)
 	{
 		check(InAttributeBase);
 
@@ -107,7 +107,7 @@ namespace PCGDataForGPUHelpers
 			const FPCGMetadataAttribute<bool>* Attribute = static_cast<const FPCGMetadataAttribute<bool>*>(InAttributeBase);
 			const bool Value = Attribute->GetValue(ValueKey);
 			check(StrideBytes == 4);
-			OutPackedDataCollection[ElementIndex + 0] = Value;
+			OutPackedDataCollection[InOutAddressUints++] = Value;
 			break;
 		}
 		case PCG::Private::MetadataTypes<float>::Id:
@@ -115,7 +115,7 @@ namespace PCGDataForGPUHelpers
 			const FPCGMetadataAttribute<float>* Attribute = static_cast<const FPCGMetadataAttribute<float>*>(InAttributeBase);
 			const float Value = Attribute->GetValue(ValueKey);
 			check(StrideBytes == 4);
-			OutPackedDataCollection[ElementIndex + 0] = FMath::AsUInt(Value);
+			OutPackedDataCollection[InOutAddressUints++] = FMath::AsUInt(Value);
 			break;
 		}
 		case PCG::Private::MetadataTypes<double>::Id:
@@ -123,7 +123,7 @@ namespace PCGDataForGPUHelpers
 			const FPCGMetadataAttribute<double>* Attribute = static_cast<const FPCGMetadataAttribute<double>*>(InAttributeBase);
 			const double Value = Attribute->GetValue(ValueKey);
 			check(StrideBytes == 4);
-			OutPackedDataCollection[ElementIndex + 0] = FMath::AsUInt(static_cast<float>(Value));
+			OutPackedDataCollection[InOutAddressUints++] = FMath::AsUInt(static_cast<float>(Value));
 			break;
 		}
 		case PCG::Private::MetadataTypes<int32>::Id:
@@ -131,7 +131,7 @@ namespace PCGDataForGPUHelpers
 			const FPCGMetadataAttribute<int32>* Attribute = static_cast<const FPCGMetadataAttribute<int32>*>(InAttributeBase);
 			const int32 Value = Attribute->GetValue(ValueKey);
 			check(StrideBytes == 4);
-			OutPackedDataCollection[ElementIndex + 0] = Value;
+			OutPackedDataCollection[InOutAddressUints++] = Value;
 			break;
 		}
 		case PCG::Private::MetadataTypes<int64>::Id:
@@ -139,7 +139,7 @@ namespace PCGDataForGPUHelpers
 			const FPCGMetadataAttribute<int64>* Attribute = static_cast<const FPCGMetadataAttribute<int64>*>(InAttributeBase);
 			const int64 Value = Attribute->GetValue(ValueKey);
 			check(StrideBytes == 4);
-			OutPackedDataCollection[ElementIndex + 0] = Value;
+			OutPackedDataCollection[InOutAddressUints++] = Value;
 			break;
 		}
 		case PCG::Private::MetadataTypes<FVector2D>::Id:
@@ -147,8 +147,8 @@ namespace PCGDataForGPUHelpers
 			const FPCGMetadataAttribute<FVector2D>* Attribute = static_cast<const FPCGMetadataAttribute<FVector2D>*>(InAttributeBase);
 			const FVector2D Value = Attribute->GetValue(ValueKey);
 			check(StrideBytes == 8);
-			OutPackedDataCollection[ElementIndex + 0] = FMath::AsUInt(static_cast<float>(Value.X));
-			OutPackedDataCollection[ElementIndex + 1] = FMath::AsUInt(static_cast<float>(Value.Y));
+			OutPackedDataCollection[InOutAddressUints++] = FMath::AsUInt(static_cast<float>(Value.X));
+			OutPackedDataCollection[InOutAddressUints++] = FMath::AsUInt(static_cast<float>(Value.Y));
 			break;
 		}
 		case PCG::Private::MetadataTypes<FRotator>::Id:
@@ -156,9 +156,9 @@ namespace PCGDataForGPUHelpers
 			const FPCGMetadataAttribute<FRotator>* Attribute = static_cast<const FPCGMetadataAttribute<FRotator>*>(InAttributeBase);
 			const FRotator Value = Attribute->GetValue(ValueKey);
 			check(StrideBytes == 12);
-			OutPackedDataCollection[ElementIndex + 0] = FMath::AsUInt(static_cast<float>(Value.Pitch));
-			OutPackedDataCollection[ElementIndex + 1] = FMath::AsUInt(static_cast<float>(Value.Yaw));
-			OutPackedDataCollection[ElementIndex + 2] = FMath::AsUInt(static_cast<float>(Value.Roll));
+			OutPackedDataCollection[InOutAddressUints++] = FMath::AsUInt(static_cast<float>(Value.Pitch));
+			OutPackedDataCollection[InOutAddressUints++] = FMath::AsUInt(static_cast<float>(Value.Yaw));
+			OutPackedDataCollection[InOutAddressUints++] = FMath::AsUInt(static_cast<float>(Value.Roll));
 			break;
 		}
 		case PCG::Private::MetadataTypes<FVector>::Id:
@@ -166,9 +166,9 @@ namespace PCGDataForGPUHelpers
 			const FPCGMetadataAttribute<FVector>* Attribute = static_cast<const FPCGMetadataAttribute<FVector>*>(InAttributeBase);
 			const FVector Value = Attribute->GetValue(ValueKey);
 			check(StrideBytes == 12);
-			OutPackedDataCollection[ElementIndex + 0] = FMath::AsUInt(static_cast<float>(Value.X));
-			OutPackedDataCollection[ElementIndex + 1] = FMath::AsUInt(static_cast<float>(Value.Y));
-			OutPackedDataCollection[ElementIndex + 2] = FMath::AsUInt(static_cast<float>(Value.Z));
+			OutPackedDataCollection[InOutAddressUints++] = FMath::AsUInt(static_cast<float>(Value.X));
+			OutPackedDataCollection[InOutAddressUints++] = FMath::AsUInt(static_cast<float>(Value.Y));
+			OutPackedDataCollection[InOutAddressUints++] = FMath::AsUInt(static_cast<float>(Value.Z));
 			break;
 		}
 		case PCG::Private::MetadataTypes<FVector4>::Id:
@@ -176,10 +176,10 @@ namespace PCGDataForGPUHelpers
 			const FPCGMetadataAttribute<FVector4>* Attribute = static_cast<const FPCGMetadataAttribute<FVector4>*>(InAttributeBase);
 			const FVector4 Value = Attribute->GetValue(ValueKey);
 			check(StrideBytes == 16);
-			OutPackedDataCollection[ElementIndex + 0] = FMath::AsUInt(static_cast<float>(Value.X));
-			OutPackedDataCollection[ElementIndex + 1] = FMath::AsUInt(static_cast<float>(Value.Y));
-			OutPackedDataCollection[ElementIndex + 2] = FMath::AsUInt(static_cast<float>(Value.Z));
-			OutPackedDataCollection[ElementIndex + 3] = FMath::AsUInt(static_cast<float>(Value.W));
+			OutPackedDataCollection[InOutAddressUints++] = FMath::AsUInt(static_cast<float>(Value.X));
+			OutPackedDataCollection[InOutAddressUints++] = FMath::AsUInt(static_cast<float>(Value.Y));
+			OutPackedDataCollection[InOutAddressUints++] = FMath::AsUInt(static_cast<float>(Value.Z));
+			OutPackedDataCollection[InOutAddressUints++] = FMath::AsUInt(static_cast<float>(Value.W));
 			break;
 		}
 		case PCG::Private::MetadataTypes<FQuat>::Id:
@@ -187,10 +187,10 @@ namespace PCGDataForGPUHelpers
 			const FPCGMetadataAttribute<FQuat>* Attribute = static_cast<const FPCGMetadataAttribute<FQuat>*>(InAttributeBase);
 			const FQuat Value = Attribute->GetValue(ValueKey);
 			check(StrideBytes == 16);
-			OutPackedDataCollection[ElementIndex + 0] = FMath::AsUInt(static_cast<float>(Value.X));
-			OutPackedDataCollection[ElementIndex + 1] = FMath::AsUInt(static_cast<float>(Value.Y));
-			OutPackedDataCollection[ElementIndex + 2] = FMath::AsUInt(static_cast<float>(Value.Z));
-			OutPackedDataCollection[ElementIndex + 3] = FMath::AsUInt(static_cast<float>(Value.W));
+			OutPackedDataCollection[InOutAddressUints++] = FMath::AsUInt(static_cast<float>(Value.X));
+			OutPackedDataCollection[InOutAddressUints++] = FMath::AsUInt(static_cast<float>(Value.Y));
+			OutPackedDataCollection[InOutAddressUints++] = FMath::AsUInt(static_cast<float>(Value.Z));
+			OutPackedDataCollection[InOutAddressUints++] = FMath::AsUInt(static_cast<float>(Value.W));
 			break;
 		}
 		case PCG::Private::MetadataTypes<FTransform>::Id:
@@ -199,22 +199,22 @@ namespace PCGDataForGPUHelpers
 			const FTransform Value = Attribute->GetValue(ValueKey);
 			const FMatrix Matrix = Value.ToMatrixWithScale();
 			check(StrideBytes == 64);
-			OutPackedDataCollection[ElementIndex + 0] = FMath::AsUInt(static_cast<float>(Matrix.M[0][0]));
-			OutPackedDataCollection[ElementIndex + 1] = FMath::AsUInt(static_cast<float>(Matrix.M[0][1]));
-			OutPackedDataCollection[ElementIndex + 2] = FMath::AsUInt(static_cast<float>(Matrix.M[0][2]));
-			OutPackedDataCollection[ElementIndex + 3] = FMath::AsUInt(static_cast<float>(Matrix.M[0][3]));
-			OutPackedDataCollection[ElementIndex + 4] = FMath::AsUInt(static_cast<float>(Matrix.M[1][0]));
-			OutPackedDataCollection[ElementIndex + 5] = FMath::AsUInt(static_cast<float>(Matrix.M[1][1]));
-			OutPackedDataCollection[ElementIndex + 6] = FMath::AsUInt(static_cast<float>(Matrix.M[1][2]));
-			OutPackedDataCollection[ElementIndex + 7] = FMath::AsUInt(static_cast<float>(Matrix.M[1][3]));
-			OutPackedDataCollection[ElementIndex + 8] = FMath::AsUInt(static_cast<float>(Matrix.M[2][0]));
-			OutPackedDataCollection[ElementIndex + 9] = FMath::AsUInt(static_cast<float>(Matrix.M[2][1]));
-			OutPackedDataCollection[ElementIndex + 10] = FMath::AsUInt(static_cast<float>(Matrix.M[2][2]));
-			OutPackedDataCollection[ElementIndex + 11] = FMath::AsUInt(static_cast<float>(Matrix.M[2][3]));
-			OutPackedDataCollection[ElementIndex + 12] = FMath::AsUInt(static_cast<float>(Matrix.M[3][0]));
-			OutPackedDataCollection[ElementIndex + 13] = FMath::AsUInt(static_cast<float>(Matrix.M[3][1]));
-			OutPackedDataCollection[ElementIndex + 14] = FMath::AsUInt(static_cast<float>(Matrix.M[3][2]));
-			OutPackedDataCollection[ElementIndex + 15] = FMath::AsUInt(static_cast<float>(Matrix.M[3][3]));
+			OutPackedDataCollection[InOutAddressUints++] = FMath::AsUInt(static_cast<float>(Matrix.M[0][0]));
+			OutPackedDataCollection[InOutAddressUints++] = FMath::AsUInt(static_cast<float>(Matrix.M[0][1]));
+			OutPackedDataCollection[InOutAddressUints++] = FMath::AsUInt(static_cast<float>(Matrix.M[0][2]));
+			OutPackedDataCollection[InOutAddressUints++] = FMath::AsUInt(static_cast<float>(Matrix.M[0][3]));
+			OutPackedDataCollection[InOutAddressUints++] = FMath::AsUInt(static_cast<float>(Matrix.M[1][0]));
+			OutPackedDataCollection[InOutAddressUints++] = FMath::AsUInt(static_cast<float>(Matrix.M[1][1]));
+			OutPackedDataCollection[InOutAddressUints++] = FMath::AsUInt(static_cast<float>(Matrix.M[1][2]));
+			OutPackedDataCollection[InOutAddressUints++] = FMath::AsUInt(static_cast<float>(Matrix.M[1][3]));
+			OutPackedDataCollection[InOutAddressUints++] = FMath::AsUInt(static_cast<float>(Matrix.M[2][0]));
+			OutPackedDataCollection[InOutAddressUints++] = FMath::AsUInt(static_cast<float>(Matrix.M[2][1]));
+			OutPackedDataCollection[InOutAddressUints++] = FMath::AsUInt(static_cast<float>(Matrix.M[2][2]));
+			OutPackedDataCollection[InOutAddressUints++] = FMath::AsUInt(static_cast<float>(Matrix.M[2][3]));
+			OutPackedDataCollection[InOutAddressUints++] = FMath::AsUInt(static_cast<float>(Matrix.M[3][0]));
+			OutPackedDataCollection[InOutAddressUints++] = FMath::AsUInt(static_cast<float>(Matrix.M[3][1]));
+			OutPackedDataCollection[InOutAddressUints++] = FMath::AsUInt(static_cast<float>(Matrix.M[3][2]));
+			OutPackedDataCollection[InOutAddressUints++] = FMath::AsUInt(static_cast<float>(Matrix.M[3][3]));
 			break;
 		}
 		case PCG::Private::MetadataTypes<FString>::Id:
@@ -223,7 +223,7 @@ namespace PCGDataForGPUHelpers
 			const FPCGMetadataAttribute<FString>* Attribute = static_cast<const FPCGMetadataAttribute<FString>*>(InAttributeBase);
 			const int32 Value = InStringTable.IndexOfByKey(Attribute->GetValue(ValueKey));
 			check(StrideBytes == 4);
-			OutPackedDataCollection[ElementIndex] = Value;
+			OutPackedDataCollection[InOutAddressUints++] = Value;
 			break;
 		}
 		case PCG::Private::MetadataTypes<FSoftObjectPath>::Id:
@@ -232,7 +232,7 @@ namespace PCGDataForGPUHelpers
 			const FPCGMetadataAttribute<FSoftObjectPath>* Attribute = static_cast<const FPCGMetadataAttribute<FSoftObjectPath>*>(InAttributeBase);
 			const int32 Value = InStringTable.IndexOfByKey(Attribute->GetValue(ValueKey).ToString());
 			check(StrideBytes == 4);
-			OutPackedDataCollection[ElementIndex] = Value;
+			OutPackedDataCollection[InOutAddressUints++] = Value;
 			break;
 		}
 		default:
@@ -690,209 +690,273 @@ FPCGDataCollectionDesc FPCGDataCollectionDesc::BuildFromInputDataCollectionAndIn
 	return CollectionDesc;
 }
 
-uint32 FPCGDataCollectionDesc::ComputePackedSize(TArray<uint32>* OutDataAddresses) const
+uint32 FPCGDataCollectionDesc::ComputePackedHeaderSizeBytes() const
+{
+	return PCGComputeConstants::DATA_COLLECTION_HEADER_SIZE_BYTES + PCGComputeConstants::DATA_HEADER_SIZE_BYTES * DataDescs.Num();
+}
+
+uint32 FPCGDataCollectionDesc::ComputePackedSizeBytes() const
 {
 	TRACE_CPUPROFILER_EVENT_SCOPE(FPCGDataCollectionDesc::ComputePackedSize);
 
-	const int NumData = DataDescs.Num();
-	// Calculation: sizeof(NumDatas) + (sizeof(DataAddress) * NumData)
-	const uint32 CollectionHeaderSizeBytes = sizeof(uint32) + (sizeof(uint32) * NumData);
-	uint32 TotalCollectionSizeBytes = CollectionHeaderSizeBytes;
+	uint32 TotalCollectionSizeBytes = ComputePackedHeaderSizeBytes();
 
-	if (OutDataAddresses)
+	for (const FPCGDataDesc& DataDesc : DataDescs)
 	{
-		OutDataAddresses->SetNumUninitialized(NumData);
+		TotalCollectionSizeBytes += DataDesc.ComputePackedSize();
 	}
 
-	for (int DataIndex = 0; DataIndex < NumData; ++DataIndex)
+	return TotalCollectionSizeBytes;
+}
+
+void FPCGDataCollectionDesc::WriteHeader(TArray<uint32>& OutPackedDataCollectionHeader) const
+{
+	const uint32 HeaderSizeBytes = ComputePackedHeaderSizeBytes();
+	const uint32 HeaderSizeUints = HeaderSizeBytes >> 2;
+
+	if (OutPackedDataCollectionHeader.Num() < static_cast<int32>(HeaderSizeUints))
+	{
+		OutPackedDataCollectionHeader.SetNumUninitialized(HeaderSizeUints);
+	}
+
+	// Zero-initialize header portion. We detect absent attributes using 0s.
+	for (uint32 Index = 0; Index < HeaderSizeUints; ++Index)
+	{
+		OutPackedDataCollectionHeader[Index] = 0;
+	}
+
+	uint32 WriteAddressUints = 0;
+
+	// Num data
+	OutPackedDataCollectionHeader[WriteAddressUints++] = DataDescs.Num();
+
+	for (int32 DataIndex = 0; DataIndex < DataDescs.Num(); ++DataIndex)
 	{
 		const FPCGDataDesc& DataDesc = DataDescs[DataIndex];
-		const uint32 DataSize = DataDesc.ComputePackedSize();
 
-		if (OutDataAddresses)
+		// Data i: type ID
+		if (DataDesc.Type == EPCGDataType::Param)
+		{
+			OutPackedDataCollectionHeader[WriteAddressUints++] = PARAM_DATA_TYPE_ID;
+		}
+		else
+		{
+			ensure(DataDesc.Type == EPCGDataType::Point);
+			OutPackedDataCollectionHeader[WriteAddressUints++] = POINT_DATA_TYPE_ID;
+		}
+
+		// Data i: attribute count (including intrinsic point properties)
+		OutPackedDataCollectionHeader[WriteAddressUints++] = DataDesc.AttributeDescs.Num();
+
+		// Data i: element count
+		OutPackedDataCollectionHeader[WriteAddressUints++] = DataDesc.ElementCount;
+
+		const uint32 DataAttributesHeaderStartAddressBytes = WriteAddressUints << 2;
+
+		for (int32 AttrIndex = 0; AttrIndex < DataDesc.AttributeDescs.Num(); ++AttrIndex)
 		{
-			(*OutDataAddresses)[DataIndex] = TotalCollectionSizeBytes;
+			const FPCGKernelAttributeDesc& AttributeDesc = DataDesc.AttributeDescs[AttrIndex];
+
+			// Scatter from attributes that are present into header which has slots for all possible attributes.
+			WriteAddressUints = (AttributeDesc.Index * ATTRIBUTE_HEADER_SIZE_BYTES + DataAttributesHeaderStartAddressBytes) >> 2;
+
+			// Data i element j: packed ID and stride
+			const uint32 AttributeId = AttributeDesc.Index;
+			const uint32 AttributeStride = PCGDataForGPUHelpers::GetAttributeTypeStrideBytes(AttributeDesc.Type);
+			const uint32 PackedIdAndStride = (AttributeId << 8) + AttributeStride;
+			OutPackedDataCollectionHeader[WriteAddressUints++] = PackedIdAndStride;
+
+			// Data i element j: data start address bytes
+			// TODO: Accumulate rather than calculate from scratch.
+			uint32 DataStartAddressBytes = HeaderSizeBytes; // Start at end of header
+			for (int32 PreviousDataIndex = 0; PreviousDataIndex < DataIndex; ++PreviousDataIndex) // Fast forward past previous data
+			{
+				for (const FPCGKernelAttributeDesc& AttrDesc : DataDescs[PreviousDataIndex].AttributeDescs)
+				{
+					DataStartAddressBytes += DataDescs[PreviousDataIndex].ElementCount * PCGDataForGPUHelpers::GetAttributeTypeStrideBytes(AttrDesc.Type);
+				}
+			}
+			for (int32 PreviousAttrIndex = 0; PreviousAttrIndex < AttrIndex; ++PreviousAttrIndex) // Fast forward past previous attributes
+			{
+				DataStartAddressBytes += DataDesc.ElementCount * PCGDataForGPUHelpers::GetAttributeTypeStrideBytes(DataDesc.AttributeDescs[PreviousAttrIndex].Type);
+			}
+			OutPackedDataCollectionHeader[WriteAddressUints++] = DataStartAddressBytes;
 		}
 
-		TotalCollectionSizeBytes += DataSize;
+		// After scattering in attribute headers, fast forward to end of section.
+		WriteAddressUints = (PCGComputeConstants::MAX_NUM_ATTRS * PCGComputeConstants::ATTRIBUTE_HEADER_SIZE_BYTES + DataAttributesHeaderStartAddressBytes) >> 2;
 	}
 
-	return TotalCollectionSizeBytes;
+	check(WriteAddressUints * 4 == HeaderSizeBytes);
+}
+
+static uint32 GetElementDataStartAddressUints(const uint32* InPackedDataCollection, uint32 InDataIndex, uint32 InAttributeId)
+{
+	uint32 ReadAddressBytes = PCGComputeConstants::DATA_COLLECTION_HEADER_SIZE_BYTES + InDataIndex * PCGComputeConstants::DATA_HEADER_SIZE_BYTES;
+	ReadAddressBytes += /*TypeId*/4 + /*Attribute Count*/4 + /*Element Count*/4;
+
+	ReadAddressBytes += InAttributeId * PCGComputeConstants::ATTRIBUTE_HEADER_SIZE_BYTES;
+	ReadAddressBytes += /*PackedIdAndStride*/4;
+
+	return InPackedDataCollection[ReadAddressBytes >> 2] >> 2;
 }
 
 void FPCGDataCollectionDesc::PackDataCollection(const FPCGDataCollection& InDataCollection, FName InPin, const TArray<FString>& InStringTable, TArray<uint32>& OutPackedDataCollection) const
 {
+	TRACE_CPUPROFILER_EVENT_SCOPE(FPCGDataCollectionDesc::PackDataCollection);
+
 	const TArray<FPCGTaggedData> InputData = InDataCollection.GetInputsByPin(InPin);
-	const uint32 NumData = InputData.Num();
 
-	TArray<uint32> DataAddresses;
-	const uint32 PackedDataCollectionSizeBytes = ComputePackedSize(&DataAddresses);
+	const uint32 PackedDataCollectionSizeBytes = ComputePackedSizeBytes();
 
-	OutPackedDataCollection.SetNumZeroed(PackedDataCollectionSizeBytes / sizeof(uint32));
-	OutPackedDataCollection[0] = NumData;
+	// Uninitialized is fine, all data is initialized explicitly.
+	OutPackedDataCollection.SetNumUninitialized(PackedDataCollectionSizeBytes >> 2);
 
-	for (uint32 DataIndex = 0; DataIndex < NumData; ++DataIndex)
-	{
-		TRACE_CPUPROFILER_EVENT_SCOPE(FPCGDataCollectionDesc::PackDataItem);
+	// Data addresses are written to the header and will be used during packing below.
+	WriteHeader(OutPackedDataCollection);
 
-		// Write the data addresses
-		const uint32 CurrentDataAddress = DataAddresses[DataIndex];
-		const uint32 CurrentDataIndex = CurrentDataAddress / sizeof(uint32);
-		OutPackedDataCollection[DataIndex + 1] = CurrentDataAddress;
+	for (int32 DataIndex = 0; DataIndex < InputData.Num(); ++DataIndex)
+	{
+		const FPCGDataDesc& DataDesc = DataDescs[DataIndex];
+		const UPCGMetadata* Metadata = InputData[DataIndex].Data ? InputData[DataIndex].Data->ConstMetadata() : nullptr;
+		if (!ensure(Metadata))
+		{
+			continue;
+		}
 
-		// DataHeader: (TypeId, NumAttrs, AttrHeaderStartOffset, TypeInfo), Attr0 Header, Attr1 Header, ..., Attr255 Header
-		// Data: Attr0, Attr1, ...
 		if (const UPCGPointData* PointData = Cast<UPCGPointData>(InputData[DataIndex].Data))
 		{
-			const UPCGMetadata* Metadata = PointData->ConstMetadata();
 			const TArray<FPCGPoint>& Points = PointData->GetPoints();
 			const uint32 NumElements = Points.Num();
 
-			const TArray<FPCGKernelAttributeDesc>& AttributeDescs = DataDescs[DataIndex].AttributeDescs;
-			const uint32 NumAttributes = AttributeDescs.Num();
-
-			OutPackedDataCollection[CurrentDataIndex + 0] = /*PointDataTypeId=*/POINT_DATA_TYPE_ID;
-			OutPackedDataCollection[CurrentDataIndex + 1] = NumAttributes;
-			OutPackedDataCollection[CurrentDataIndex + 2] = POINT_DATA_HEADER_PREAMBLE_SIZE_BYTES;
-			OutPackedDataCollection[CurrentDataIndex + 3] = NumElements; // TypeInfo for PointData is just NumPoints
-
-			const uint32 BaseAttributeHeaderAddress = CurrentDataAddress + POINT_DATA_HEADER_PREAMBLE_SIZE_BYTES;
-			uint32 CurrentAttributeAddress = CurrentDataAddress + POINT_DATA_HEADER_SIZE_BYTES;
-
-			for (const FPCGKernelAttributeDesc& AttributeDesc : AttributeDescs)
+			for (const FPCGKernelAttributeDesc& AttributeDesc : DataDesc.AttributeDescs)
 			{
 				const uint32 AttributeId = AttributeDesc.Index;
 				const uint32 AttributeStrideBytes = PCGDataForGPUHelpers::GetAttributeTypeStrideBytes(AttributeDesc.Type);
-				const uint32 AttributeNumComponents = AttributeStrideBytes / sizeof(uint32); // E.g. float3 has 3 components
-
-				// Pack Position (24 bits for AttributeId, 8 bits for Stride)
-				const uint32 PackedIdAndStride = (AttributeId << 8) + AttributeStrideBytes;
-				const uint32 AttributeIndex = CurrentAttributeAddress / sizeof(uint32);
-
-				const uint32 AttributeHeaderIndex = (BaseAttributeHeaderAddress + AttributeId * ATTRIBUTE_HEADER_SIZE_BYTES) / sizeof(uint32);
-				OutPackedDataCollection[AttributeHeaderIndex + 0] = PackedIdAndStride;
-				OutPackedDataCollection[AttributeHeaderIndex + 1] = CurrentAttributeAddress;
 
 				const FPCGMetadataAttributeBase* AttributeBase = (AttributeId >= NUM_RESERVED_ATTRS) ? Metadata->GetConstAttribute(AttributeDesc.Name) : nullptr;
 
-				for (uint32 ElementIndex = 0; ElementIndex < NumElements; ++ElementIndex)
-				{
-					const uint32 PackedDataElementIndex = AttributeIndex + (ElementIndex * AttributeNumComponents);
+				uint32 AddressUints = GetElementDataStartAddressUints(OutPackedDataCollection.GetData(), DataIndex, AttributeId);
 
-					if (AttributeBase) // Pack attribute
+				if (AttributeId < NUM_RESERVED_ATTRS)
+				{
+					// Point property.
+					switch (AttributeId)
 					{
-						ensure(PCGDataForGPUHelpers::PackAttributeHelper(AttributeBase, AttributeDesc, Points[ElementIndex].MetadataEntry, InStringTable, OutPackedDataCollection, PackedDataElementIndex));
-					}
-					else // Pack property
+					case POINT_POSITION_ATTRIBUTE_ID:
 					{
-						switch (AttributeId)
-						{
-						case POINT_POSITION_ATTRIBUTE_ID:
+						for (uint32 ElementIndex = 0; ElementIndex < NumElements; ++ElementIndex)
 						{
 							const FVector Position = Points[ElementIndex].Transform.GetLocation();
-
-							OutPackedDataCollection[PackedDataElementIndex + 0] = FMath::AsUInt(static_cast<float>(Position.X));
-							OutPackedDataCollection[PackedDataElementIndex + 1] = FMath::AsUInt(static_cast<float>(Position.Y));
-							OutPackedDataCollection[PackedDataElementIndex + 2] = FMath::AsUInt(static_cast<float>(Position.Z));
-							break;
+							OutPackedDataCollection[AddressUints++] = FMath::AsUInt(static_cast<float>(Position.X));
+							OutPackedDataCollection[AddressUints++] = FMath::AsUInt(static_cast<float>(Position.Y));
+							OutPackedDataCollection[AddressUints++] = FMath::AsUInt(static_cast<float>(Position.Z));
 						}
-						case POINT_ROTATION_ATTRIBUTE_ID:
+						break;
+					}
+					case POINT_ROTATION_ATTRIBUTE_ID:
+					{
+						for (uint32 ElementIndex = 0; ElementIndex < NumElements; ++ElementIndex)
 						{
 							const FQuat Rotation = Points[ElementIndex].Transform.GetRotation();
-
-							OutPackedDataCollection[PackedDataElementIndex + 0] = FMath::AsUInt(static_cast<float>(Rotation.X));
-							OutPackedDataCollection[PackedDataElementIndex + 1] = FMath::AsUInt(static_cast<float>(Rotation.Y));
-							OutPackedDataCollection[PackedDataElementIndex + 2] = FMath::AsUInt(static_cast<float>(Rotation.Z));
-							OutPackedDataCollection[PackedDataElementIndex + 3] = FMath::AsUInt(static_cast<float>(Rotation.W));
-							break;
+							OutPackedDataCollection[AddressUints++] = FMath::AsUInt(static_cast<float>(Rotation.X));
+							OutPackedDataCollection[AddressUints++] = FMath::AsUInt(static_cast<float>(Rotation.Y));
+							OutPackedDataCollection[AddressUints++] = FMath::AsUInt(static_cast<float>(Rotation.Z));
+							OutPackedDataCollection[AddressUints++] = FMath::AsUInt(static_cast<float>(Rotation.W));
 						}
-						case POINT_SCALE_ATTRIBUTE_ID:
+						break;
+					}
+					case POINT_SCALE_ATTRIBUTE_ID:
+					{
+						for (uint32 ElementIndex = 0; ElementIndex < NumElements; ++ElementIndex)
 						{
 							const FVector Scale = Points[ElementIndex].Transform.GetScale3D();
-
-							OutPackedDataCollection[PackedDataElementIndex + 0] = FMath::AsUInt(static_cast<float>(Scale.X));
-							OutPackedDataCollection[PackedDataElementIndex + 1] = FMath::AsUInt(static_cast<float>(Scale.Y));
-							OutPackedDataCollection[PackedDataElementIndex + 2] = FMath::AsUInt(static_cast<float>(Scale.Z));
-							break;
+							OutPackedDataCollection[AddressUints++] = FMath::AsUInt(static_cast<float>(Scale.X));
+							OutPackedDataCollection[AddressUints++] = FMath::AsUInt(static_cast<float>(Scale.Y));
+							OutPackedDataCollection[AddressUints++] = FMath::AsUInt(static_cast<float>(Scale.Z));
 						}
-						case POINT_BOUNDS_MIN_ATTRIBUTE_ID:
+						break;
+					}
+					case POINT_BOUNDS_MIN_ATTRIBUTE_ID:
+					{
+						for (uint32 ElementIndex = 0; ElementIndex < NumElements; ++ElementIndex)
 						{
 							const FVector& BoundsMin = Points[ElementIndex].BoundsMin;
-
-							OutPackedDataCollection[PackedDataElementIndex + 0] = FMath::AsUInt(static_cast<float>(BoundsMin.X));
-							OutPackedDataCollection[PackedDataElementIndex + 1] = FMath::AsUInt(static_cast<float>(BoundsMin.Y));
-							OutPackedDataCollection[PackedDataElementIndex + 2] = FMath::AsUInt(static_cast<float>(BoundsMin.Z));
-							break;
+							OutPackedDataCollection[AddressUints++] = FMath::AsUInt(static_cast<float>(BoundsMin.X));
+							OutPackedDataCollection[AddressUints++] = FMath::AsUInt(static_cast<float>(BoundsMin.Y));
+							OutPackedDataCollection[AddressUints++] = FMath::AsUInt(static_cast<float>(BoundsMin.Z));
 						}
-						case POINT_BOUNDS_MAX_ATTRIBUTE_ID:
+						break;
+					}
+					case POINT_BOUNDS_MAX_ATTRIBUTE_ID:
+					{
+						for (uint32 ElementIndex = 0; ElementIndex < NumElements; ++ElementIndex)
 						{
 							const FVector& BoundsMax = Points[ElementIndex].BoundsMax;
-
-							OutPackedDataCollection[PackedDataElementIndex + 0] = FMath::AsUInt(static_cast<float>(BoundsMax.X));
-							OutPackedDataCollection[PackedDataElementIndex + 1] = FMath::AsUInt(static_cast<float>(BoundsMax.Y));
-							OutPackedDataCollection[PackedDataElementIndex + 2] = FMath::AsUInt(static_cast<float>(BoundsMax.Z));
-							break;
+							OutPackedDataCollection[AddressUints++] = FMath::AsUInt(static_cast<float>(BoundsMax.X));
+							OutPackedDataCollection[AddressUints++] = FMath::AsUInt(static_cast<float>(BoundsMax.Y));
+							OutPackedDataCollection[AddressUints++] = FMath::AsUInt(static_cast<float>(BoundsMax.Z));
 						}
-						case POINT_COLOR_ATTRIBUTE_ID:
+						break;
+					}
+					case POINT_COLOR_ATTRIBUTE_ID:
+					{
+						for (uint32 ElementIndex = 0; ElementIndex < NumElements; ++ElementIndex)
 						{
 							const FVector4& Color = Points[ElementIndex].Color;
-
-							OutPackedDataCollection[PackedDataElementIndex + 0] = FMath::AsUInt(static_cast<float>(Color.X));
-							OutPackedDataCollection[PackedDataElementIndex + 1] = FMath::AsUInt(static_cast<float>(Color.Y));
-							OutPackedDataCollection[PackedDataElementIndex + 2] = FMath::AsUInt(static_cast<float>(Color.Z));
-							OutPackedDataCollection[PackedDataElementIndex + 3] = FMath::AsUInt(static_cast<float>(Color.W));
-							break;
+							OutPackedDataCollection[AddressUints++] = FMath::AsUInt(static_cast<float>(Color.X));
+							OutPackedDataCollection[AddressUints++] = FMath::AsUInt(static_cast<float>(Color.Y));
+							OutPackedDataCollection[AddressUints++] = FMath::AsUInt(static_cast<float>(Color.Z));
+							OutPackedDataCollection[AddressUints++] = FMath::AsUInt(static_cast<float>(Color.W));
 						}
-						case POINT_DENSITY_ATTRIBUTE_ID:
+						break;
+					}
+					case POINT_DENSITY_ATTRIBUTE_ID:
+					{
+						for (uint32 ElementIndex = 0; ElementIndex < NumElements; ++ElementIndex)
 						{
 							const float Density = Points[ElementIndex].Density;
-
-							OutPackedDataCollection[PackedDataElementIndex + 0] = FMath::AsUInt(Density);
-							break;
+							OutPackedDataCollection[AddressUints++] = FMath::AsUInt(Density);
 						}
-						case POINT_SEED_ATTRIBUTE_ID:
+						break;
+					}
+					case POINT_SEED_ATTRIBUTE_ID:
+					{
+						for (uint32 ElementIndex = 0; ElementIndex < NumElements; ++ElementIndex)
 						{
 							const int Seed = Points[ElementIndex].Seed;
-
-							OutPackedDataCollection[PackedDataElementIndex + 0] = Seed;
-							break;
+							OutPackedDataCollection[AddressUints++] = Seed;
 						}
-						case POINT_STEEPNESS_ATTRIBUTE_ID:
+						break;
+					}
+					case POINT_STEEPNESS_ATTRIBUTE_ID:
+					{
+						for (uint32 ElementIndex = 0; ElementIndex < NumElements; ++ElementIndex)
 						{
 							const float Steepness = Points[ElementIndex].Steepness;
-
-							OutPackedDataCollection[PackedDataElementIndex + 0] = FMath::AsUInt(Steepness);
-							break;
-						}
-						default:
-							checkNoEntry();
-							break;
+							OutPackedDataCollection[AddressUints++] = FMath::AsUInt(Steepness);
 						}
+						break;
+					}
+					default:
+						checkNoEntry();
+						break;
+					}
+				}
+				else
+				{
+					// Pack attribute. Validate first element only for perf.
+					ensure(PCGDataForGPUHelpers::PackAttributeHelper(AttributeBase, AttributeDesc, Points[0].MetadataEntry, InStringTable, OutPackedDataCollection, AddressUints));
+					for (uint32 ElementIndex = 1; ElementIndex < NumElements; ++ElementIndex)
+					{
+						PCGDataForGPUHelpers::PackAttributeHelper(AttributeBase, AttributeDesc, Points[ElementIndex].MetadataEntry, InStringTable, OutPackedDataCollection, AddressUints);
 					}
 				}
-
-				CurrentAttributeAddress += NumElements * AttributeNumComponents * 4;
 			}
 		}
 		else if (const UPCGParamData* ParamData = Cast<UPCGParamData>(InputData[DataIndex].Data))
 		{
-			const UPCGMetadata* Metadata = ParamData->ConstMetadata();
-
-			const FPCGDataDesc& DataDesc = DataDescs[DataIndex];
-			const uint32 NumElements = DataDesc.ElementCount;
-
-			const TArray<FPCGKernelAttributeDesc>& AttributeDescs = DataDesc.AttributeDescs;
-			const uint32 NumAttributes = AttributeDescs.Num();
-
-			OutPackedDataCollection[CurrentDataIndex + 0] = /*ParamDataTypeId=*/PARAM_DATA_TYPE_ID;
-			OutPackedDataCollection[CurrentDataIndex + 1] = NumAttributes;
-			OutPackedDataCollection[CurrentDataIndex + 2] = PARAM_DATA_HEADER_PREAMBLE_SIZE_BYTES;
-			OutPackedDataCollection[CurrentDataIndex + 3] = NumElements; // TypeInfo for ParamData is # of elements
-
-			const uint32 BaseAttributeHeaderAddress = CurrentDataAddress + PARAM_DATA_HEADER_PREAMBLE_SIZE_BYTES;
-			uint32 CurrentAttributeAddress = CurrentDataAddress + PARAM_DATA_HEADER_SIZE_BYTES;
-
-			for (const FPCGKernelAttributeDesc& AttributeDesc : AttributeDescs)
+			for (const FPCGKernelAttributeDesc& AttributeDesc : DataDesc.AttributeDescs)
 			{
 				const FPCGMetadataAttributeBase* AttributeBase = Metadata->GetConstAttribute(AttributeDesc.Name);
 				if (!AttributeBase)
@@ -900,116 +964,17 @@ void FPCGDataCollectionDesc::PackDataCollection(const FPCGDataCollection& InData
 					continue;
 				}
 
-				const uint32 AttributeId = AttributeDesc.Index;
-				const uint32 AttributeStrideBytes = PCGDataForGPUHelpers::GetAttributeTypeStrideBytes(AttributeDesc.Type);
-				const uint32 AttributeNumComponents = AttributeStrideBytes / sizeof(uint32); // E.g. float3 has 3 components
+				uint32 AddressUints = GetElementDataStartAddressUints(OutPackedDataCollection.GetData(), DataIndex, AttributeDesc.Index);
 
-				// Pack Position (24 bits for AttributeId, 8 bits for Stride)
-				const uint32 PackedIdAndStride = (AttributeId << 8) + AttributeStrideBytes;
-				const uint32 AttributeIndex = CurrentAttributeAddress / sizeof(uint32);
-
-				const uint32 AttributeHeaderIndex = (BaseAttributeHeaderAddress + AttributeId * ATTRIBUTE_HEADER_SIZE_BYTES) / sizeof(uint32);
-				OutPackedDataCollection[AttributeHeaderIndex + 0] = PackedIdAndStride;
-				OutPackedDataCollection[AttributeHeaderIndex + 1] = CurrentAttributeAddress;
-
-				for (uint32 ElementIndex = 0; ElementIndex < NumElements; ++ElementIndex)
+				// Pack attribute. Validate first element only for perf.
+				ensure(PCGDataForGPUHelpers::PackAttributeHelper(AttributeBase, AttributeDesc, /*InEntryKey*/0, InStringTable, OutPackedDataCollection, AddressUints));
+				for (int32 ElementIndex = 1; ElementIndex < DataDesc.ElementCount; ++ElementIndex)
 				{
-					const uint32 PackedDataElementIndex = AttributeIndex + (ElementIndex * AttributeNumComponents);
-					const int64 MetadataKey = ElementIndex;
-
-					// Pack attribute
-					ensure(PCGDataForGPUHelpers::PackAttributeHelper(AttributeBase, AttributeDesc, MetadataKey, InStringTable, OutPackedDataCollection, PackedDataElementIndex));
+					PCGDataForGPUHelpers::PackAttributeHelper(AttributeBase, AttributeDesc, /*InEntryKey*/ElementIndex, InStringTable, OutPackedDataCollection, AddressUints);
 				}
-
-				CurrentAttributeAddress += NumElements * AttributeNumComponents * 4;
 			}
 		}
-		else { /* TODO: Support non-point data. */ }
-	}
-}
-
-void FPCGDataCollectionDesc::PrepareBufferForKernelOutput(TArray<uint32>& OutPackedDataCollection)
-{
-	const uint32 NumData = DataDescs.Num();
-
-	TArray<uint32> DataAddresses;
-	const uint32 PackedDataCollectionSizeBytes = ComputePackedSize(&DataAddresses);
-
-	OutPackedDataCollection.SetNumZeroed(PackedDataCollectionSizeBytes / sizeof(uint32));
-	OutPackedDataCollection[0] = NumData;
-
-	for (uint32 DataIndex = 0; DataIndex < NumData; ++DataIndex)
-	{
-		TRACE_CPUPROFILER_EVENT_SCOPE(FPCGDataCollectionDesc::PackDataItem);
-
-		// Write the data addresses
-		const uint32 CurrentDataAddress = DataAddresses[DataIndex];
-		const uint32 CurrentDataIndex = CurrentDataAddress / sizeof(uint32);
-		OutPackedDataCollection[DataIndex + 1] = CurrentDataAddress;
-
-		// DataHeader: (TypeId, NumAttrs, AttrHeaderStartOffset, TypeInfo), Attr0 Header, Attr1 Header, ..., Attr255 Header
-		// Data: Attr0, Attr1, ...
-		if (DataDescs[DataIndex].Type == EPCGDataType::Point)
-		{
-			const uint32 NumElements = DataDescs[DataIndex].ElementCount;
-
-			const TArray<FPCGKernelAttributeDesc>& AttributeDescs = DataDescs[DataIndex].AttributeDescs;
-			const uint32 NumAttributes = AttributeDescs.Num();
-
-			OutPackedDataCollection[CurrentDataIndex + 0] = /*PointDataTypeId=*/POINT_DATA_TYPE_ID;
-			OutPackedDataCollection[CurrentDataIndex + 1] = NumAttributes;
-			OutPackedDataCollection[CurrentDataIndex + 2] = POINT_DATA_HEADER_PREAMBLE_SIZE_BYTES;
-			OutPackedDataCollection[CurrentDataIndex + 3] = NumElements; // TypeInfo for PointData is just NumPoints
-
-			const uint32 BaseAttributeHeaderAddress = CurrentDataAddress + POINT_DATA_HEADER_PREAMBLE_SIZE_BYTES;
-			uint32 CurrentAttributeAddress = CurrentDataAddress + POINT_DATA_HEADER_SIZE_BYTES;
-
-			for (const FPCGKernelAttributeDesc& AttributeDesc : AttributeDescs)
-			{
-				const uint32 AttributeId = AttributeDesc.Index;
-				const uint32 AttributeStrideBytes = PCGDataForGPUHelpers::GetAttributeTypeStrideBytes(AttributeDesc.Type);
-				const uint32 AttributeNumComponents = AttributeStrideBytes / sizeof(uint32); // E.g. float3 has 3 components
-				const uint32 AttributeHeaderIndex = (BaseAttributeHeaderAddress + AttributeId * ATTRIBUTE_HEADER_SIZE_BYTES) / sizeof(uint32);
-
-				// Pack Position (24 bits for AttributeId, 8 bits for Stride)
-				const uint32 PackedIdAndStride = (AttributeId << 8) + AttributeStrideBytes;
-				OutPackedDataCollection[AttributeHeaderIndex + 0] = PackedIdAndStride;
-
-				OutPackedDataCollection[AttributeHeaderIndex + 1] = CurrentAttributeAddress;
-				CurrentAttributeAddress += NumElements * AttributeNumComponents * 4;
-			}
-		}
-		if (DataDescs[DataIndex].Type == EPCGDataType::Param)
-		{
-			const uint32 NumElements = DataDescs[DataIndex].ElementCount;
-
-			const TArray<FPCGKernelAttributeDesc>& AttributeDescs = DataDescs[DataIndex].AttributeDescs;
-			const uint32 NumAttributes = AttributeDescs.Num();
-
-			OutPackedDataCollection[CurrentDataIndex + 0] = /*ParamDataTypeId=*/PARAM_DATA_TYPE_ID;
-			OutPackedDataCollection[CurrentDataIndex + 1] = NumAttributes;
-			OutPackedDataCollection[CurrentDataIndex + 2] = PARAM_DATA_HEADER_PREAMBLE_SIZE_BYTES;
-			OutPackedDataCollection[CurrentDataIndex + 3] = NumElements; // TypeInfo for ParamData is # of elems
-
-			const uint32 BaseAttributeHeaderAddress = CurrentDataAddress + PARAM_DATA_HEADER_PREAMBLE_SIZE_BYTES;
-			uint32 CurrentAttributeAddress = CurrentDataAddress + PARAM_DATA_HEADER_SIZE_BYTES;
-
-			for (const FPCGKernelAttributeDesc& AttributeDesc : AttributeDescs)
-			{
-				const uint32 AttributeId = AttributeDesc.Index;
-				const uint32 AttributeStrideBytes = PCGDataForGPUHelpers::GetAttributeTypeStrideBytes(AttributeDesc.Type);
-				const uint32 AttributeNumComponents = AttributeStrideBytes / sizeof(uint32); // E.g. float3 has 3 components
-				const uint32 AttributeHeaderIndex = (BaseAttributeHeaderAddress + AttributeId * ATTRIBUTE_HEADER_SIZE_BYTES) / sizeof(uint32);
-
-				// Pack Position (24 bits for AttributeId, 8 bits for Stride)
-				const uint32 PackedIdAndStride = (AttributeId << 8) + AttributeStrideBytes;
-				OutPackedDataCollection[AttributeHeaderIndex + 0] = PackedIdAndStride;
-
-				OutPackedDataCollection[AttributeHeaderIndex + 1] = CurrentAttributeAddress;
-				CurrentAttributeAddress += NumElements * AttributeNumComponents * 4;
-			}
-		}
-		else { /* TODO: Support non-point data. */ }
+		else { /* TODO: Support additional data types. */ }
 	}
 }
 
@@ -1020,11 +985,13 @@ EPCGUnpackDataCollectionResult FPCGDataCollectionDesc::UnpackDataCollection(cons
 	const uint32* DataAsUint = static_cast<const uint32*>(PackedData);
 	const int32* DataAsInt = static_cast<const int32*>(PackedData);
 
-	const uint32 NumPackedFloats = InPackedData.Num() / 4;
+	uint32 ReadAddress = 0;
+
+	const uint32 PackedExecutionFlagAndNumData = DataAsUint[ReadAddress++];
 
 	// Most significant bit of NumData is reserved to flag whether or not the kernel executed.
-	ensureMsgf(DataAsUint[0] & PCGComputeConstants::KernelExecutedFlag, TEXT("Tried to unpack a GPU data collection, but the compute shader did not execute."));
-	const uint32 NumData = DataAsUint[0] & ~PCGComputeConstants::KernelExecutedFlag;
+	ensureMsgf(PackedExecutionFlagAndNumData & PCGComputeConstants::KernelExecutedFlag, TEXT("Tried to unpack a GPU data collection, but the compute shader did not execute."));
+	const uint32 NumData = PackedExecutionFlagAndNumData & ~PCGComputeConstants::KernelExecutedFlag;
 
 	if (NumData != DataDescs.Num())
 	{
@@ -1037,12 +1004,9 @@ EPCGUnpackDataCollectionResult FPCGDataCollectionDesc::UnpackDataCollection(cons
 	{
 		TRACE_CPUPROFILER_EVENT_SCOPE(FPCGDataCollectionDesc::UnpackDataItem);
 
-		const uint32 CurrentDataAddress = DataAsUint[DataIndex + 1];
-		const uint32 CurrentDataIndex = CurrentDataAddress / sizeof(uint32);
-		const uint32 TypeId =                      DataAsUint[CurrentDataIndex + 0];
-		const uint32 NumAttributes =               DataAsUint[CurrentDataIndex + 1];
-		const uint32 DataHeaderPreambleSizeBytes = DataAsUint[CurrentDataIndex + 2];
-		const uint32 NumElements =                 DataAsUint[CurrentDataIndex + 3];
+		const uint32 TypeId = DataAsUint[ReadAddress++];
+		const uint32 NumAttributes = DataAsUint[ReadAddress++];
+		const uint32 NumElements = DataAsUint[ReadAddress++];
 
 		const TArray<FPCGKernelAttributeDesc>& AttributeDescs = DataDescs[DataIndex].AttributeDescs;
 		check(NumAttributes == AttributeDescs.Num());
@@ -1072,139 +1036,143 @@ EPCGUnpackDataCollectionResult FPCGDataCollectionDesc::UnpackDataCollection(cons
 			OutTaggedData.Data = OutPointData;
 			OutTaggedData.Pin = InPin;
 
-			const uint32 AttributeHeadersIndex = CurrentDataIndex + POINT_DATA_HEADER_PREAMBLE_SIZE_BYTES / sizeof(uint32);
-
 			// Loop over attributes.
 			for (const FPCGKernelAttributeDesc& AttributeDesc : AttributeDescs)
 			{
-				TRACE_CPUPROFILER_EVENT_SCOPE(WriteAttribute);
+				TRACE_CPUPROFILER_EVENT_SCOPE(UnpackAttribute);
 
 				const uint32 AttributeId = AttributeDesc.Index;
-				const uint32 AttributeNumComponents = PCGDataForGPUHelpers::GetAttributeTypeStrideBytes(AttributeDesc.Type) / sizeof(uint32);
-				const uint32 AttributeHeaderIndex = AttributeHeadersIndex + AttributeDesc.Index * ATTRIBUTE_HEADER_SIZE_BYTES / sizeof(uint32);
-				const uint32 AttributeIndex = DataAsUint[AttributeHeaderIndex + 1] / sizeof(uint32);
-
-				FPCGMetadataAttributeBase* AttributeBase = (AttributeId >= NUM_RESERVED_ATTRS) ? PCGDataForGPUHelpers::CreateAttributeFromAttributeDesc(Metadata, AttributeDesc) : nullptr;
+				const uint32 AttributeStrideUints = PCGDataForGPUHelpers::GetAttributeTypeStrideBytes(AttributeDesc.Type) >> 2;
 
-				// 2. Parse each element in the attribute.
+				const uint32 AddressUints = GetElementDataStartAddressUints(DataAsUint, DataIndex, AttributeId);
 
-				ParallelFor(NumElements, [&](int32 ElementIndex)
+				if (AttributeId < NUM_RESERVED_ATTRS)
 				{
-					const uint32 PackedDataElementIndex = AttributeIndex + ElementIndex * AttributeNumComponents;
-					check(PackedDataElementIndex + AttributeNumComponents <= NumPackedFloats);
-
-					if (AttributeBase) // Unpack attribute
+					// We tried hoisting this decision to a lambda but it didn't appear to help.
+					switch (AttributeId)
 					{
-						Metadata->InitializeOnSet(OutPoints[ElementIndex].MetadataEntry);
-						ensure(PCGDataForGPUHelpers::UnpackAttributeHelper(PackedData, InStringTable, PackedDataElementIndex, AttributeBase, AttributeDesc, OutPoints[ElementIndex].MetadataEntry));
-					}
-					else // Unpack property
+					case POINT_POSITION_ATTRIBUTE_ID:
 					{
-						// We tried hoisting this decision to a lambda but it didn't appear to help.
-						switch (AttributeId)
-						{
-						case POINT_POSITION_ATTRIBUTE_ID:
+						ParallelFor(NumElements, [DataAsFloat, AddressUints, &OutPoints](int32 ElementIndex)
 						{
-							const FVector Location = FVector
-							(
-								DataAsFloat[PackedDataElementIndex + 0],
-								DataAsFloat[PackedDataElementIndex + 1],
-								DataAsFloat[PackedDataElementIndex + 2]
-							);
+							const FVector Location = FVector(
+								DataAsFloat[AddressUints + ElementIndex * 3 + 0],
+								DataAsFloat[AddressUints + ElementIndex * 3 + 1],
+								DataAsFloat[AddressUints + ElementIndex * 3 + 2]);
 
 							OutPoints[ElementIndex].Transform.SetLocation(Location);
-							break;
-						}
-						case POINT_ROTATION_ATTRIBUTE_ID:
+						});
+						break;
+					}
+					case POINT_ROTATION_ATTRIBUTE_ID:
+					{
+						ParallelFor(NumElements, [DataAsFloat, AddressUints, &OutPoints](int32 ElementIndex)
 						{
-							FQuat Rotation = FQuat
-							(
-								DataAsFloat[PackedDataElementIndex + 0],
-								DataAsFloat[PackedDataElementIndex + 1],
-								DataAsFloat[PackedDataElementIndex + 2],
-								DataAsFloat[PackedDataElementIndex + 3]
-							);
+							const FQuat Rotation = FQuat(
+								DataAsFloat[AddressUints + ElementIndex * 4 + 0],
+								DataAsFloat[AddressUints + ElementIndex * 4 + 1],
+								DataAsFloat[AddressUints + ElementIndex * 4 + 2],
+								DataAsFloat[AddressUints + ElementIndex * 4 + 3]);
 
 							// Normalize here with default tolerance (zero quat will return identity).
 							OutPoints[ElementIndex].Transform.SetRotation(Rotation.GetNormalized());
-							break;
-						}
-						case POINT_SCALE_ATTRIBUTE_ID:
+						});
+						break;
+					}
+					case POINT_SCALE_ATTRIBUTE_ID:
+					{
+						ParallelFor(NumElements, [DataAsFloat, AddressUints, &OutPoints](int32 ElementIndex)
 						{
 							const FVector Scale = FVector
 							(
-								DataAsFloat[PackedDataElementIndex + 0],
-								DataAsFloat[PackedDataElementIndex + 1],
-								DataAsFloat[PackedDataElementIndex + 2]
-							);
+								DataAsFloat[AddressUints + ElementIndex * 3 + 0],
+								DataAsFloat[AddressUints + ElementIndex * 3 + 1],
+								DataAsFloat[AddressUints + ElementIndex * 3 + 2]);
 
 							OutPoints[ElementIndex].Transform.SetScale3D(Scale);
-							break;
-						}
-						case POINT_BOUNDS_MIN_ATTRIBUTE_ID:
+						});
+						break;
+					}
+					case POINT_BOUNDS_MIN_ATTRIBUTE_ID:
+					{
+						ParallelFor(NumElements, [DataAsFloat, AddressUints, &OutPoints](int32 ElementIndex)
 						{
-							const FVector BoundsMin = FVector
-							(
-								DataAsFloat[PackedDataElementIndex + 0],
-								DataAsFloat[PackedDataElementIndex + 1],
-								DataAsFloat[PackedDataElementIndex + 2]
-							);
+							const FVector BoundsMin = FVector(
+								DataAsFloat[AddressUints + ElementIndex * 3 + 0],
+								DataAsFloat[AddressUints + ElementIndex * 3 + 1],
+								DataAsFloat[AddressUints + ElementIndex * 3 + 2]);
 
 							OutPoints[ElementIndex].BoundsMin = BoundsMin;
-							break;
-						}
-						case POINT_BOUNDS_MAX_ATTRIBUTE_ID:
+						});
+						break;
+					}
+					case POINT_BOUNDS_MAX_ATTRIBUTE_ID:
+					{
+						ParallelFor(NumElements, [DataAsFloat, AddressUints, &OutPoints](int32 ElementIndex)
 						{
-							const FVector BoundsMax = FVector
-							(
-								DataAsFloat[PackedDataElementIndex + 0],
-								DataAsFloat[PackedDataElementIndex + 1],
-								DataAsFloat[PackedDataElementIndex + 2]
-							);
+							const FVector BoundsMax = FVector(
+								DataAsFloat[AddressUints + ElementIndex * 3 + 0],
+								DataAsFloat[AddressUints + ElementIndex * 3 + 1],
+								DataAsFloat[AddressUints + ElementIndex * 3 + 2]);
 
 							OutPoints[ElementIndex].BoundsMax = BoundsMax;
-							break;
-						}
-						case POINT_COLOR_ATTRIBUTE_ID:
+						});
+						break;
+					}
+					case POINT_COLOR_ATTRIBUTE_ID:
+					{
+						ParallelFor(NumElements, [DataAsFloat, AddressUints, &OutPoints](int32 ElementIndex)
 						{
-							const FVector4 Color = FVector4
-							(
-								DataAsFloat[PackedDataElementIndex + 0],
-								DataAsFloat[PackedDataElementIndex + 1],
-								DataAsFloat[PackedDataElementIndex + 2],
-								DataAsFloat[PackedDataElementIndex + 3]
-							);
+							const FVector4 Color = FVector4(
+								DataAsFloat[AddressUints + ElementIndex * 4 + 0],
+								DataAsFloat[AddressUints + ElementIndex * 4 + 1],
+								DataAsFloat[AddressUints + ElementIndex * 4 + 2],
+								DataAsFloat[AddressUints + ElementIndex * 4 + 3]);
 
 							OutPoints[ElementIndex].Color = Color;
-							break;
-						}
-						case POINT_DENSITY_ATTRIBUTE_ID:
+						});
+						break;
+					}
+					case POINT_DENSITY_ATTRIBUTE_ID:
+					{
+						ParallelFor(NumElements, [DataAsFloat, AddressUints, &OutPoints](int32 ElementIndex)
 						{
-							const float Density = DataAsFloat[PackedDataElementIndex];
-
-							OutPoints[ElementIndex].Density = Density;
-							break;
-						}
-						case POINT_SEED_ATTRIBUTE_ID:
+							OutPoints[ElementIndex].Density = DataAsFloat[AddressUints + ElementIndex];
+						});
+						break;
+					}
+					case POINT_SEED_ATTRIBUTE_ID:
+					{
+						ParallelFor(NumElements, [DataAsInt, AddressUints, &OutPoints](int32 ElementIndex)
 						{
-							const int32 Seed = DataAsInt[PackedDataElementIndex];
-
-							OutPoints[ElementIndex].Seed = Seed;
-							break;
-						}
-						case POINT_STEEPNESS_ATTRIBUTE_ID:
+							OutPoints[ElementIndex].Seed = DataAsInt[AddressUints + ElementIndex];
+						});
+						break;
+					}
+					case POINT_STEEPNESS_ATTRIBUTE_ID:
+					{
+						ParallelFor(NumElements, [DataAsFloat, AddressUints, &OutPoints](int32 ElementIndex)
 						{
-							const float Steepness = DataAsFloat[PackedDataElementIndex];
-
-							OutPoints[ElementIndex].Steepness = Steepness;
-							break;
-						}
-						default:
-							checkNoEntry();
-							break;
-						}
+							OutPoints[ElementIndex].Steepness = DataAsFloat[AddressUints + ElementIndex];
+						});
+						break;
 					}
-				});
+					default:
+						checkNoEntry();
+						break;
+					}
+				}
+				else
+				{
+					FPCGMetadataAttributeBase* AttributeBase = PCGDataForGPUHelpers::CreateAttributeFromAttributeDesc(Metadata, AttributeDesc);
+					check(AttributeBase);
+
+					ParallelFor(NumElements, [PackedData, &InStringTable, &OutPoints, Metadata, AttributeBase, &AttributeDesc, AddressUints, AttributeStrideUints](int32 ElementIndex)
+					{
+						Metadata->InitializeOnSet(OutPoints[ElementIndex].MetadataEntry);
+						ensure(PCGDataForGPUHelpers::UnpackAttributeHelper(PackedData, InStringTable, AddressUints + ElementIndex * AttributeStrideUints, AttributeBase, AttributeDesc, OutPoints[ElementIndex].MetadataEntry));
+					});
+				}
 			}
 
 			// TODO: It may be more efficient to create a mapping from input point index to final output point index and do everything in one pass.
@@ -1230,8 +1198,6 @@ EPCGUnpackDataCollectionResult FPCGDataCollectionDesc::UnpackDataCollection(cons
 			OutTaggedData.Data = OutParamData;
 			OutTaggedData.Pin = InPin;
 
-			const uint32 AttributeHeadersIndex = CurrentDataIndex + PARAM_DATA_HEADER_PREAMBLE_SIZE_BYTES / sizeof(uint32);
-
 			TArray<TTuple</*EntryKey=*/int64, /*ParentEntryKey=*/int64>> AllMetadataEntries;
 			AllMetadataEntries.SetNumUninitialized(NumElements);
 
@@ -1247,25 +1213,24 @@ EPCGUnpackDataCollectionResult FPCGDataCollectionDesc::UnpackDataCollection(cons
 			{
 				TRACE_CPUPROFILER_EVENT_SCOPE(WriteAttribute);
 
-				const uint32 AttributeNumComponents = PCGDataForGPUHelpers::GetAttributeTypeStrideBytes(AttributeDesc.Type) / sizeof(uint32);
-				const uint32 AttributeHeaderIndex = AttributeHeadersIndex + AttributeDesc.Index * ATTRIBUTE_HEADER_SIZE_BYTES / sizeof(uint32);
-				const uint32 AttributeIndex = DataAsUint[AttributeHeaderIndex + 1] / sizeof(uint32);
-
-				FPCGMetadataAttributeBase* AttributeBase = PCGDataForGPUHelpers::CreateAttributeFromAttributeDesc(Metadata, AttributeDesc);
+				const uint32 AddressUints = GetElementDataStartAddressUints(DataAsUint, DataIndex, AttributeDesc.Index);
+				const uint32 AttributeStrideUints = PCGDataForGPUHelpers::GetAttributeTypeStrideBytes(AttributeDesc.Type) / sizeof(uint32);
 
-				ParallelFor(NumElements, [&](int32 ElementIndex)
+				if (FPCGMetadataAttributeBase* AttributeBase = PCGDataForGPUHelpers::CreateAttributeFromAttributeDesc(Metadata, AttributeDesc))
 				{
-					if (AttributeBase)
+					const uint32 NumPackedUints = InPackedData.Num() / 4;
+
+					ParallelFor(NumElements, [PackedData, NumPackedUints, AddressUints, AttributeStrideUints, &AttributeDesc, AttributeBase, &InStringTable](int32 ElementIndex)
 					{
-						const uint32 PackedDataElementIndex = AttributeIndex + ElementIndex * AttributeNumComponents;
-						check(PackedDataElementIndex + AttributeNumComponents <= NumPackedFloats);
+						const uint32 PackedDataElementIndex = AddressUints + ElementIndex * AttributeStrideUints;
+						check(PackedDataElementIndex + AttributeStrideUints <= NumPackedUints);
 
 						ensure(PCGDataForGPUHelpers::UnpackAttributeHelper(PackedData, InStringTable, PackedDataElementIndex, AttributeBase, AttributeDesc, ElementIndex));
-					}
-				});
+					});
+				}
 			}
 		}
-		else { /* TODO: Support non-point data. */ }
+		else { /* TODO: Support additional data types. */ }
 	}
 
 	return EPCGUnpackDataCollectionResult::Success;
diff --git a/Engine/Plugins/PCG/Source/PCG/Public/Compute/PCGDataForGPU.h b/Engine/Plugins/PCG/Source/PCG/Public/Compute/PCGDataForGPU.h
--- a/Engine/Plugins/PCG/Source/PCG/Public/Compute/PCGDataForGPU.h
+++ b/Engine/Plugins/PCG/Source/PCG/Public/Compute/PCGDataForGPU.h
@@ -95,7 +95,9 @@ struct FPCGKernelAttributeDesc
 	{
 	}
 
+	// Attribute Id
 	int32 Index = INDEX_NONE;
+
 	EPCGKernelAttributeType Type = EPCGKernelAttributeType::Float;
 	FName Name = NAME_None;
 
@@ -128,15 +130,17 @@ struct FPCGDataCollectionDesc
 		const TMap<FName, FPCGKernelAttributeIDAndType>& InAttributeLookupTable,
 		const TArray<FString>& InStringTable);
 
-	/** Computes the size (in bytes) of the data collection after packing. Also produces the offset (in bytes) for each data in the packed collection. */
-	uint32 ComputePackedSize(TArray<uint32>* OutDataAddresses = nullptr) const;
+	/** Computes the size (in bytes) of the header portion of the packed data collection buffer. */
+	uint32 ComputePackedHeaderSizeBytes() const;
+
+	/** Computes the size (in bytes) of the data collection after packing. */
+	uint32 ComputePackedSizeBytes() const;
+
+	void WriteHeader(TArray<uint32>& OutPackedDataCollectionHeader) const;
 
 	/** Pack a data collection into the GPU data format. DataDescs defines which attributes are packed. */
 	void PackDataCollection(const FPCGDataCollection& InDataCollection, FName InPin, const TArray<FString>& InStringTable, TArray<uint32>& OutPackedDataCollection) const;
 
-	/** Allocates the correct size and sets up header. Initializes data count to 0, which kernel will then overwrite if it executes at least one thread. */
-	void PrepareBufferForKernelOutput(TArray<uint32>& OutPackedDataCollection);
-
 	/** Unpack a buffer of 8-bit uints to a data collection. */
 	EPCGUnpackDataCollectionResult UnpackDataCollection(const TArray<uint8>& InPackedData, FName InPin, const TArray<FString>& InStringTable, FPCGDataCollection& OutDataCollection) const;
 
