[PCG] Updated Pathfinding Node

+ Start/Goal as input
+ Add Tag to Partial Paths
+ Support multiple paths to one goal
+ Support multiple goals from one start
+ Support N:N pathfinding

#rnx

diff --git a/Engine/Plugins/PCG/Source/PCG/Private/Elements/PCGPathfindingElement.cpp b/Engine/Plugins/PCG/Source/PCG/Private/Elements/PCGPathfindingElement.cpp
--- a/Engine/Plugins/PCG/Source/PCG/Private/Elements/PCGPathfindingElement.cpp
+++ b/Engine/Plugins/PCG/Source/PCG/Private/Elements/PCGPathfindingElement.cpp
@@ -4,6 +4,7 @@
 
 #include "PCGContext.h"
 #include "PCGComponent.h"
+#include "PCGParamData.h"
 #include "Data/PCGPointData.h"
 #include "Data/PCGSplineData.h"
 #include "Helpers/PCGHelpers.h"
@@ -18,44 +19,58 @@
 
 #define LOCTEXT_NAMESPACE "PCGPathfindingElement"
 
-namespace PCGPathfindingElement::Helpers
+namespace PCGPathfindingElement
 {
-	TArray<FSplinePoint> ConvertPathToSplinePoints(const TArrayView<const FPCGPoint> Path, const EPCGPathfindingSplineMode SplineMode)
+	namespace Constants
 	{
-		ESplinePointType::Type SplineCurveMode = ESplinePointType::Type::Constant;
-		switch (SplineMode)
+		const FName StartLocationsInputPinLabel = TEXT("Start");
+		const FName GoalLocationsInputPinLabel = TEXT("Goal");
+		static const TCHAR* CompletePathTag = TEXT("CompletePath");
+		static const TCHAR* PartialPathTag = TEXT("PartialPath");
+	}
+
+	namespace Helpers
+	{
+		TArray<FSplinePoint> ConvertPathToSplinePoints(const TArrayView<const FPCGPoint> Path, const EPCGPathfindingSplineMode SplineMode)
 		{
-			case EPCGPathfindingSplineMode::Curve:
-				SplineCurveMode = ESplinePointType::Curve;
-				break;
-			case EPCGPathfindingSplineMode::Linear:
-				SplineCurveMode = ESplinePointType::Linear;
-			default:
-				break;
-		}
+			ESplinePointType::Type SplineCurveMode = ESplinePointType::Type::Constant;
+			switch (SplineMode)
+			{
+				case EPCGPathfindingSplineMode::Curve:
+					SplineCurveMode = ESplinePointType::Curve;
+					break;
+				case EPCGPathfindingSplineMode::Linear:
+					SplineCurveMode = ESplinePointType::Linear;
+				default:
+					break;
+			}
 
-		TArray<FSplinePoint> SplinePoints;
-		SplinePoints.Reserve(Path.Num());
+			TArray<FSplinePoint> SplinePoints;
+			SplinePoints.Reserve(Path.Num());
 
-		int Index = 0;
-		Algo::Transform(Path, SplinePoints, [&Index, SplineCurveMode](const FPCGPoint& Point)
-		{
-			return FSplinePoint(
-					Index++, // Spline points must be indexed in ascending order
-					Point.Transform.GetLocation(),
-					FVector::ZeroVector,
-					FVector::ZeroVector,
-					FRotator::ZeroRotator,
-					FVector::OneVector,
-					SplineCurveMode);
-		});
-
-		return SplinePoints;
+			int Index = 0;
+			Algo::Transform(Path, SplinePoints, [&Index, SplineCurveMode](const FPCGPoint& Point)
+			{
+				return FSplinePoint(
+						Index++, // Spline points must be indexed in ascending order
+						Point.Transform.GetLocation(),
+						FVector::ZeroVector,
+						FVector::ZeroVector,
+						FRotator::ZeroRotator,
+						FVector::OneVector,
+						SplineCurveMode);
+			});
+
+			return SplinePoints;
+		}
 	}
-} // namespace PCGPathfindingElement::Helpers
+} // namespace PCGPathfindingElement
 
 UPCGPathfindingSettings::UPCGPathfindingSettings()
 {
+	StartLocationAttribute.SetPointProperty(EPCGPointProperties::Position);
+	GoalLocationAttribute.SetPointProperty(EPCGPointProperties::Position);
+
 	// In most cases, we're not going to be interested in checking for occlusion by the landscape itself, as we'll be pathfinding on the landscape.
 	PathTraceParams.SelectLandscapeHits = EPCGWorldQuerySelectLandscapeHits::Exclude;
 }
@@ -64,7 +79,17 @@ TArray<FPCGPinProperties> UPCGPathfindingSettings::InputPinProperties() const
 {
 	TArray<FPCGPinProperties> Properties;
 	Properties.Emplace_GetRef(PCGPinConstants::DefaultInputLabel, EPCGDataType::Point).SetRequiredPin();
-	// TODO: PointOrParam input for start/goal.
+
+	if (bStartLocationsAsInput)
+	{
+		Properties.Emplace_GetRef(PCGPathfindingElement::Constants::StartLocationsInputPinLabel, EPCGDataType::PointOrParam, /*bInAllowMultipleConnections=*/false, /*bAllowMultipleData=*/false).SetRequiredPin();
+	}
+
+	if (bGoalLocationsAsInput)
+	{
+		Properties.Emplace_GetRef(PCGPathfindingElement::Constants::GoalLocationsInputPinLabel, EPCGDataType::PointOrParam, /*bInAllowMultipleConnections=*/false, /*bAllowMultipleData=*/false).SetRequiredPin();
+	}
+
 	return Properties;
 }
 
@@ -92,7 +117,10 @@ FPCGElementPtr UPCGPathfindingSettings::CreateElement() const
 EPCGChangeType UPCGPathfindingSettings::GetChangeTypeForProperty(const FName& InPropertyName) const
 {
 	EPCGChangeType LocalChangeType = EPCGChangeType::Cosmetic;
-	if (InPropertyName == GET_MEMBER_NAME_CHECKED(UPCGPathfindingSettings, bOutputAsSpline))
+
+	if (InPropertyName == GET_MEMBER_NAME_CHECKED(UPCGPathfindingSettings, bOutputAsSpline)
+		|| InPropertyName == GET_MEMBER_NAME_CHECKED(UPCGPathfindingSettings, bStartLocationsAsInput)
+		|| InPropertyName == GET_MEMBER_NAME_CHECKED(UPCGPathfindingSettings, bGoalLocationsAsInput))
 	{
 		LocalChangeType |= EPCGChangeType::Structural;
 	}
@@ -124,7 +152,7 @@ bool FPCGPathfindingElement::PrepareDataInternal(FPCGContext* InContext) const
 		return true;
 	}
 
-	Context->InitializePerExecutionState([&PointInputs, Settings](ContextType*, ExecStateType& OutState)
+	const EPCGTimeSliceInitResult ExecResult = Context->InitializePerExecutionState([&PointInputs, Settings](const ContextType* SlicedContext, ExecStateType& OutState)
 	{
 		TRACE_CPUPROFILER_EVENT_SCOPE(FPCGClusterElement::PrepareData::InitializePerExecutionState);
 
@@ -133,48 +161,143 @@ bool FPCGPathfindingElement::PrepareDataInternal(FPCGContext* InContext) const
 			return EPCGTimeSliceInitResult::NoOperation;
 		}
 
-		// TODO: The start and goal points will eventually come from PointOrParam.
-		OutState = PCGSpatialAlgo::AStar::FSearchSettings
+		OutState.Settings = PCGSpatialAlgo::AStar::FSearchSettings
 		{
-			.StartPoint = FPCGPoint(FTransform(Settings->Start), /*InDensity=*/1, PCGHelpers::ComputeSeedFromPosition(Settings->Start)),
-			.GoalPoint = FPCGPoint(FTransform(Settings->Goal), /*InDensity=*/1.0, PCGHelpers::ComputeSeedFromPosition(Settings->Goal)),
 			.SearchDistance = Settings->SearchDistance,
-			.Start = Settings->Start,
-			.Goal = Settings->Goal,
 			.HeuristicWeight = Settings->HeuristicWeight,
 			.bAcceptPartialPath = Settings->bAcceptPartialPath,
 			.bCopyOriginatingPoints = Settings->bCopyOriginatingPoints
 		};
 
+		auto GetLocationsFromInput = [SlicedContext, Settings](FName Pin) -> const UPCGData*
+		{
+			const TArray<FPCGTaggedData> LocationsInputData = SlicedContext->InputData.GetInputsByPin(Pin);
+			if (LocationsInputData.IsEmpty())
+			{
+				return nullptr;
+			}
+
+			if (LocationsInputData.Num() > 1)
+			{
+				PCGLog::InputOutput::LogFirstInputOnlyWarning(Pin, SlicedContext);
+			}
+
+			return LocationsInputData[0].Data;
+		};
+
+		auto PopulateArray = [SlicedContext](const UPCGData* Data, const FPCGAttributePropertyInputSelector& InSelector, TArray<FPCGPoint>& OutPointArray)
+		{
+			const FPCGAttributePropertyInputSelector Selector = InSelector.CopyAndFixLast(Data);
+			const TUniquePtr<const IPCGAttributeAccessor> Accessor = PCGAttributeAccessorHelpers::CreateConstAccessor(Data, Selector);
+			const TUniquePtr<const IPCGAttributeAccessorKeys> Keys = PCGAttributeAccessorHelpers::CreateConstKeys(Data, Selector);
+			if (!Accessor || !Keys)
+			{
+				PCGLog::Metadata::LogFailToCreateAccessorError(Selector, SlicedContext);
+				return;
+			}
+
+			if (const UPCGPointData* PointData = Cast<UPCGPointData>(Data))
+			{
+				OutPointArray = PointData->GetPointsCopy();
+			}
+			else
+			{
+				OutPointArray.SetNum(Keys->GetNum());
+			}
+
+			// For Points with $Position as the selection, nothing is needed. For all other conditions, set the position based on the selected attribute/property.
+			if (OutPointArray.IsEmpty() || (Selector.GetSelection() == EPCGAttributePropertySelection::PointProperty && Selector.GetPointProperty() != EPCGPointProperties::Position))
+			{
+				return;
+			}
+
+			TArray<FVector> Locations;
+			Locations.SetNumUninitialized(Keys->GetNum());
+			Accessor->GetRange(TArrayView<FVector>(Locations), 0, *Keys, EPCGAttributeAccessorFlags::AllowBroadcastAndConstructible);
+
+			check(Locations.Num() == OutPointArray.Num());
+			for (int i = 0; i < OutPointArray.Num(); ++i)
+			{
+				OutPointArray[i].Transform.SetLocation(Locations[i]);
+				OutPointArray[i].Seed = PCGHelpers::ComputeSeedFromPosition(Locations[i]);
+			}
+		};
+
+		if (Settings->bStartLocationsAsInput)
+		{
+			if (const UPCGData* Data = GetLocationsFromInput(PCGPathfindingElement::Constants::StartLocationsInputPinLabel))
+			{
+				PopulateArray(Data, Settings->StartLocationAttribute, OutState.StartPoints);
+			}
+		}
+		else
+		{
+			OutState.StartPoints = {FPCGPoint(FTransform(Settings->Start), /*InDensity=*/1, PCGHelpers::ComputeSeedFromPosition(Settings->Start))};
+		}
+
+		if (OutState.StartPoints.IsEmpty())
+		{
+			return EPCGTimeSliceInitResult::NoOperation;
+		}
+
+		if (Settings->bGoalLocationsAsInput)
+		{
+			if (const UPCGData* Data = GetLocationsFromInput(PCGPathfindingElement::Constants::GoalLocationsInputPinLabel))
+			{
+				PopulateArray(Data, Settings->GoalLocationAttribute, OutState.GoalPoints);
+			}
+		}
+		else
+		{
+			OutState.GoalPoints = {FPCGPoint(FTransform(Settings->Goal), /*InDensity=*/1, PCGHelpers::ComputeSeedFromPosition(Settings->Goal))};
+		}
+
+		if (OutState.GoalPoints.IsEmpty())
+		{
+			return EPCGTimeSliceInitResult::NoOperation;
+		}
+
+		if (OutState.StartPoints.Num() != OutState.GoalPoints.Num() && OutState.StartPoints.Num() != 1 && OutState.GoalPoints.Num() != 1)
+		{
+			PCGLog::InputOutput::LogInvalidCardinalityError(PCGPathfindingElement::Constants::StartLocationsInputPinLabel, PCGPathfindingElement::Constants::GoalLocationsInputPinLabel, SlicedContext);
+			return EPCGTimeSliceInitResult::AbortExecution;
+		}
+
 		return EPCGTimeSliceInitResult::Success;
 	});
 
-	Context->InitializePerIterationStates(PointInputs.Num(), [&PointInputs, Settings, Context](IterStateType& OutSearchState, const ExecStateType& SearchSettings, const uint32 IterationIndex)
+	if (ExecResult != EPCGTimeSliceInitResult::Success)
+	{
+		return true;
+	}
+
+	Context->InitializePerIterationStates(PointInputs.Num(), [&PointInputs, Settings, Context](IterStateType& OutState, const ExecStateType& ExecutionState, const uint32 IterationIndex)
 	{
 		TRACE_CPUPROFILER_EVENT_SCOPE(FPCGClusterElement::PrepareData::InitializePerIterationStates);
 
-		OutSearchState.OriginatingPointData = Cast<UPCGPointData>(PointInputs[IterationIndex].Data);
-		if (OutSearchState.OriginatingPointData->IsEmpty())
+		const UPCGPointData* PointData = Cast<UPCGPointData>(PointInputs[IterationIndex].Data);
+		if (!PointData || PointData->IsEmpty())
 		{
-			// Already confirmed we can't make it from start->goal in the execution state check.
 			return EPCGTimeSliceInitResult::NoOperation;
 		}
 
+		OutState.SearchState.OriginatingPointData = PointData;
+
 		// Build cost attribute accessor if required
 		TSharedPtr<const IPCGAttributeAccessor> CostAccessor = nullptr;
 
 		if (Settings->CostFunctionMode != EPCGPathfindingCostFunctionMode::Distance)
 		{
-			const FPCGAttributePropertyInputSelector Selector = Settings->CostAttribute.CopyAndFixLast(OutSearchState.OriginatingPointData);
-			CostAccessor = MakeShareable(PCGAttributeAccessorHelpers::CreateConstAccessor(OutSearchState.OriginatingPointData, Selector).Release());
+			const FPCGAttributePropertyInputSelector Selector = Settings->CostAttribute.CopyAndFixLast(PointData);
+			CostAccessor = MakeShareable(PCGAttributeAccessorHelpers::CreateConstAccessor(PointData, Selector).Release());
 
-			FPCGAttributeAccessorKeysPoints Keys(OutSearchState.OriginatingPointData->GetPoints());
+			FPCGAttributeAccessorKeysPoints Keys(PointData->GetPoints());
 
 			if (!CostAccessor)
 			{
 				PCGLog::Metadata::LogFailToCreateAccessorError(Selector, Context);
 			}
-			else if(!PCG::Private::IsBroadcastableOrConstructible(CostAccessor->GetUnderlyingType(), PCG::Private::MetadataTypes<double>::Id))
+			else if (!PCG::Private::IsBroadcastableOrConstructible(CostAccessor->GetUnderlyingType(), PCG::Private::MetadataTypes<double>::Id))
 			{
 				PCGLog::Metadata::LogFailToGetAttributeError<double>(Selector, CostAccessor.Get(), Context);
 				CostAccessor = nullptr;
@@ -217,9 +340,9 @@ bool FPCGPathfindingElement::PrepareDataInternal(FPCGContext* InContext) const
 			{
 				const double MaxFitnessPenaltyFactor = FMath::Max(Settings->MaximumFitnessPenaltyFactor, 1.0);
 
-				OutSearchState.CostFunction = [FitnessAccessor = CostAccessor, MaxFitnessPenaltyFactor, PathTraceTest = MoveTemp(LineTraceTest)](const double PreviousNodeCost, const FPCGPoint* PreviousNodePoint, const double DistanceToCurrentSquared, const FPCGPoint* CurrentNodePoint)
+				OutState.SearchState.CostFunction = [FitnessAccessor = CostAccessor, MaxFitnessPenaltyFactor, PathTraceTest = MoveTemp(LineTraceTest)](const double PreviousNodeCost, const FPCGPoint* PreviousNodePoint, const double DistanceToCurrentSquared, const FPCGPoint* CurrentNodePoint)
 				{
-					if(!PathTraceTest(PreviousNodePoint->Transform.GetLocation(), CurrentNodePoint->Transform.GetLocation()))
+					if (!PathTraceTest(PreviousNodePoint->Transform.GetLocation(), CurrentNodePoint->Transform.GetLocation()))
 					{
 						return std::numeric_limits<double>::max();
 					}
@@ -233,9 +356,9 @@ bool FPCGPathfindingElement::PrepareDataInternal(FPCGContext* InContext) const
 					return PreviousNodeCost + (1.0 - FitnessScore) * MaxFitnessPenaltyFactor * FMath::Sqrt(DistanceToCurrentSquared);
 				};
 			}
-			else if (Settings->CostFunctionMode == EPCGPathfindingCostFunctionMode::CostMultipler)
+			else if (Settings->CostFunctionMode == EPCGPathfindingCostFunctionMode::CostMultiplier)
 			{
-				OutSearchState.CostFunction = [MultiplierAccessor = CostAccessor, PathTraceTest = MoveTemp(LineTraceTest)](const double PreviousNodeCost, const FPCGPoint* PreviousNodePoint, const double DistanceToCurrentSquared, const FPCGPoint* CurrentNodePoint)
+				OutState.SearchState.CostFunction = [MultiplierAccessor = CostAccessor, PathTraceTest = MoveTemp(LineTraceTest)](const double PreviousNodeCost, const FPCGPoint* PreviousNodePoint, const double DistanceToCurrentSquared, const FPCGPoint* CurrentNodePoint)
 				{
 					if (!PathTraceTest(PreviousNodePoint->Transform.GetLocation(), CurrentNodePoint->Transform.GetLocation()))
 					{
@@ -243,7 +366,7 @@ bool FPCGPathfindingElement::PrepareDataInternal(FPCGContext* InContext) const
 					}
 
 					double Multiplier = 1.0;
-					FPCGAttributeAccessorKeysPoints Key(*CurrentNodePoint);
+					const FPCGAttributeAccessorKeysPoints Key(*CurrentNodePoint);
 
 					MultiplierAccessor->Get(Multiplier, Key, EPCGAttributeAccessorFlags::AllowBroadcastAndConstructible);
 					Multiplier = FMath::Max(Multiplier, 1.0);
@@ -259,7 +382,7 @@ bool FPCGPathfindingElement::PrepareDataInternal(FPCGContext* InContext) const
 		else if (Settings->bUsePathTraces)
 		{
 			// Use distance but with line trace
-			OutSearchState.CostFunction = [PathTraceTest = MoveTemp(LineTraceTest)](const double PreviousNodeCost, const FPCGPoint* PreviousNodePoint, const double DistanceToCurrentSquared, const FPCGPoint* CurrentNodePoint)
+			OutState.SearchState.CostFunction = [PathTraceTest = MoveTemp(LineTraceTest)](const double PreviousNodeCost, const FPCGPoint* PreviousNodePoint, const double DistanceToCurrentSquared, const FPCGPoint* CurrentNodePoint)
 			{
 				if (!PathTraceTest(PreviousNodePoint->Transform.GetLocation(), CurrentNodePoint->Transform.GetLocation()))
 				{
@@ -272,8 +395,6 @@ bool FPCGPathfindingElement::PrepareDataInternal(FPCGContext* InContext) const
 			};
 		}
 
-		PCGSpatialAlgo::AStar::Initialize(OutSearchState.OriginatingPointData, SearchSettings, OutSearchState);
-
 		return EPCGTimeSliceInitResult::Success;
 	});
 
@@ -295,58 +416,87 @@ bool FPCGPathfindingElement::ExecuteInternal(FPCGContext* InContext) const
 	const UPCGPathfindingSettings* Settings = TimeSlicedContext->GetInputSettings<UPCGPathfindingSettings>();
 	check(Settings);
 
-	return ExecuteSlice(TimeSlicedContext, [Settings](ContextType* Context, const ExecStateType& SearchSettings, IterStateType& PathData, const uint32 IterIndex)
+	return ExecuteSlice(TimeSlicedContext, [Settings](ContextType* Context, const ExecStateType& ExecutionState, IterStateType& IterationState, const uint32 IterIndex)
 	{
 		if (Context->GetIterationStateResult(IterIndex) == EPCGTimeSliceInitResult::NoOperation)
 		{
 			return true;
 		}
 
-		TArray<FPCGPoint> FinalPath;
-		bool bFinished = false;
-		do
+		while (IterationState.PathIterationIndex < ExecutionState.StartPoints.Num())
 		{
-			bFinished = PCGSpatialAlgo::AStar::ExecuteSearchIteration(SearchSettings, PathData, FinalPath);
+			using PCGSpatialAlgo::AStar::ESearchResult;
+			TArray<FPCGPoint> FinalPath;
 
-			if (!bFinished && Context->ShouldStop())
+			// Initialize the next iteration.
+			if (IterationState.PathIterationIndex != IterationState.LastPathIterationIndex)
 			{
-				return false;
+				PCGSpatialAlgo::AStar::Initialize(&ExecutionState.StartPoints[IterationState.PathIterationIndex], IterationState.SearchState);
+				IterationState.LastPathIterationIndex = IterationState.PathIterationIndex;
 			}
-		}
-		while (!bFinished);
 
-		// No path was found and partial paths were not enabled.
-		if (FinalPath.IsEmpty())
-		{
-			check(!SearchSettings.bAcceptPartialPath);
-			return true;
-		}
+			ESearchResult SearchResult;
+			do
+			{
+				SearchResult = ExecuteSearchIteration(ExecutionState.Settings, IterationState.SearchState, ExecutionState.GoalPoints, FinalPath);
 
-		// Finally, output the path as either a spline or points
-		if (Settings->bOutputAsSpline)
-		{
-			const TArray<FSplinePoint> SplinePoints = PCGPathfindingElement::Helpers::ConvertPathToSplinePoints(FinalPath, Settings->SplineMode);
-			UPCGSplineData* SplineData = FPCGContext::NewObject_AnyThread<UPCGSplineData>(Context);
-			SplineData->Initialize(SplinePoints, /*bInClosedLoop=*/false, FTransform::Identity);
+				if (SearchResult == ESearchResult::Processing && Context->ShouldStop())
+				{
+					return false;
+				}
+			}
+			while (SearchResult == ESearchResult::Processing);
 
-			FPCGTaggedData& OutputData = Context->OutputData.TaggedData.Emplace_GetRef();
-			OutputData.Data = SplineData;
-		}
-		else
-		{
-			UPCGPointData* OutputPointData = FPCGContext::NewObject_AnyThread<UPCGPointData>(Context);
-			if (SearchSettings.bCopyOriginatingPoints)
+			if (SearchResult == ESearchResult::Invalid)
+			{
+				PCGLog::LogErrorOnGraph(LOCTEXT("InvalidSearch", "The search could not be completed. Check for invalid search settings or input."));
+				return true;
+			}
+
+			// No path was found and partial paths were not enabled.
+			if (FinalPath.IsEmpty())
 			{
-				OutputPointData->InitializeFromData(PathData.OriginatingPointData);
+				check(!ExecutionState.Settings.bAcceptPartialPath && SearchResult != ESearchResult::Partial);
+				return true;
 			}
-			OutputPointData->GetMutablePoints() = std::move(FinalPath);
 
 			FPCGTaggedData& OutputData = Context->OutputData.TaggedData.Emplace_GetRef();
-			OutputData.Data = OutputPointData;
+
+			// Output the path as either a spline or points.
+			if (Settings->bOutputAsSpline)
+			{
+				const TArray<FSplinePoint> SplinePoints = PCGPathfindingElement::Helpers::ConvertPathToSplinePoints(FinalPath, Settings->SplineMode);
+				UPCGSplineData* SplineData = FPCGContext::NewObject_AnyThread<UPCGSplineData>(Context);
+				SplineData->Initialize(SplinePoints, /*bInClosedLoop=*/false, FTransform::Identity);
+				OutputData.Data = SplineData;
+			}
+			else
+			{
+				UPCGPointData* OutputPointData = FPCGContext::NewObject_AnyThread<UPCGPointData>(Context);
+				if (ExecutionState.Settings.bCopyOriginatingPoints)
+				{
+					OutputPointData->InitializeFromData(IterationState.SearchState.OriginatingPointData);
+				}
+
+				OutputPointData->GetMutablePoints() = std::move(FinalPath);
+				OutputData.Data = OutputPointData;
+			}
+
+			if (ExecutionState.Settings.bAcceptPartialPath)
+			{
+				OutputData.Tags.Emplace(SearchResult == ESearchResult::Complete ? PCGPathfindingElement::Constants::CompletePathTag : PCGPathfindingElement::Constants::PartialPathTag);
+			}
+
+			++IterationState.PathIterationIndex;
+
+			if (Context->ShouldStop())
+			{
+				return false;
+			}
 		}
 
 		return true;
 	});
 }
 
-#undef LOCTEXT_NAMESPACE
+#undef LOCTEXT_NAMESPACE
\ No newline at end of file
diff --git a/Engine/Plugins/PCG/Source/PCG/Private/SpatialAlgo/PCGAStar.cpp b/Engine/Plugins/PCG/Source/PCG/Private/SpatialAlgo/PCGAStar.cpp
--- a/Engine/Plugins/PCG/Source/PCG/Private/SpatialAlgo/PCGAStar.cpp
+++ b/Engine/Plugins/PCG/Source/PCG/Private/SpatialAlgo/PCGAStar.cpp
@@ -40,6 +40,26 @@ namespace Helpers
 		// Reverse to get it in the correct travel order.
 		Algo::Reverse(OutPath);
 	}
+
+	TTuple<const FPCGPoint*, double> GetNearestGoalToLocation(const FVector& Location, const TArray<FPCGPoint>& GoalPoints)
+	{
+		check(!GoalPoints.IsEmpty());
+
+		double MinSquaredDistanceToGoal = std::numeric_limits<double>::max();
+		int32 CurrentGoalIndex = INDEX_NONE;
+		for (int32 GoalIndex = 0; GoalIndex < GoalPoints.Num(); ++GoalIndex)
+		{
+			const double SquareDistanceToGoal = FVector::DistSquared(Location, GoalPoints[GoalIndex].Transform.GetLocation());
+			if (SquareDistanceToGoal < MinSquaredDistanceToGoal)
+			{
+				MinSquaredDistanceToGoal = SquareDistanceToGoal;
+				CurrentGoalIndex = GoalIndex;
+			}
+		}
+		check(CurrentGoalIndex != INDEX_NONE);
+
+		return MakeTuple(&GoalPoints[CurrentGoalIndex], MinSquaredDistanceToGoal);
+	}
 }
 
 namespace Cost
@@ -60,40 +80,44 @@ namespace Heuristic
 }
 
 /** Initialize the search state for AStar. Must be called before ExecuteSearchIteration. */
-void Initialize(const UPCGPointData* const PointData, const FSearchSettings& Settings, FSearchState& OutSearchState)
+void Initialize(const FPCGPoint* const StartPoint, FSearchState& OutSearchState)
 {
-	check(PointData && !PointData->IsEmpty());
+	check(StartPoint);
 
 	// Emplace the starting node on the open list. The costs will be 0 at the starting point.
 	OutSearchState.NodeList.Reset();
-	OutSearchState.NodeList.Emplace(&Settings.StartPoint, /*InParent=*/INDEX_NONE, /*InCost=*/0.0, /*HeuristicCost=*/0.0);
+	OutSearchState.NodeList.Emplace(StartPoint, /*InParent=*/INDEX_NONE, /*InCost=*/0.0, /*HeuristicCost=*/0.0);
 	OutSearchState.OpenIndexList.Reset();
 	// Okay to add the first index this way, since it will be "heapified" by default with one element.
 	OutSearchState.OpenIndexList.Add(0);
 	OutSearchState.ClosedIndexList.Reset();
+	OutSearchState.PointToNodeIndexMap.Reset();
 }
 
 /** Runs a single iteration of the A* algorithm. Intended to be called multiple times, whereupon it will return true when the algorithm is finished. */
-bool ExecuteSearchIteration(const FSearchSettings& SearchSettings, FSearchState& SearchState, TArray<FPCGPoint>& OutPath)
+ESearchResult ExecuteSearchIteration(const FSearchSettings& SearchSettings, FSearchState& SearchState, const TArray<FPCGPoint>& GoalPoints, TArray<FPCGPoint>& OutPath)
 {
-	TRACE_CPUPROFILER_EVENT_SCOPE(FPCGPathfindingElement::Algorithm::AStar);
+	TRACE_CPUPROFILER_EVENT_SCOPE(FPCGPathfindingElement::Algorithm::ExecuteSearchIteration);
 	check(SearchState.OriginatingPointData);
 
-	// The starting node should be guaranteed by the caller.
-	if (!ensure(!SearchState.OpenIndexList.IsEmpty()))
+	// The goal points should be guaranteed by the caller. Must call Initialize first to create the open index list.
+	if (GoalPoints.IsEmpty() || !ensure(!SearchState.OpenIndexList.IsEmpty()))
 	{
-		return true;
+		return ESearchResult::Invalid;
 	}
 
 	TArray<FNode, TInlineAllocator<FSearchState::PreAllocNodeCount>>& NodeList = SearchState.NodeList;
+
+	auto HeapPredicate = [&SearchState](int32 Index1, int32 Index2) { return Helpers::CompareNodes(SearchState.NodeList[Index1], SearchState.NodeList[Index2]); };
+
 	// Get the lowest cost point on the list, which has been binary heap sorted.
 	int32 CurrentNodeIndex;
-	SearchState.OpenIndexList.HeapPop(CurrentNodeIndex, [&SearchState](int32 Index1, int32 Index2) { return Helpers::CompareNodes(SearchState.NodeList[Index1], SearchState.NodeList[Index2]); }, EAllowShrinking::No);
+	SearchState.OpenIndexList.HeapPop(CurrentNodeIndex, HeapPredicate, EAllowShrinking::No);
 	check(NodeList[CurrentNodeIndex].PCGPoint);
 
-	const double SquareDistanceToGoal = FVector::DistSquared(NodeList[CurrentNodeIndex].PCGPoint->Transform.GetLocation(), SearchSettings.Goal);
+	auto [CurrentGoal, MinSquaredDistanceToGoal] = Helpers::GetNearestGoalToLocation(NodeList[CurrentNodeIndex].PCGPoint->Transform.GetLocation(), GoalPoints);
 
-	auto PerPointProcessing = [&NodeList, CurrentNodeIndex, &SearchSettings, &SearchState](const FPCGPoint* Point, const double DistanceToPointSquared)
+	auto PerPointProcessing = [&NodeList, CurrentNodeIndex, CurrentGoal, &HeapPredicate, &SearchSettings, &SearchState, &GoalPoints](const FPCGPoint* Point, const double DistanceToPointSquared)
 	{
 		const bool bTrackedPoint = SearchState.PointToNodeIndexMap.Contains(Point);
 		// Node has already been ruled out.
@@ -110,13 +134,22 @@ bool ExecuteSearchIteration(const FSearchSettings& SearchSettings, FSearchState&
 			return;
 		}
 
+		/** Implementation Note: If there are multiple goals, a heuristic to a specific goal is not valid. The goal
+		 * can not be changed during the pathfinding process, as neighbor evaluation is tied to the heuristic to the
+		 * current goal at any point. Nodes may be ruled out and added to the Closed List
+		 */
+		auto GetHeuristicCost = [&SearchSettings, &SearchState, &GoalPoints](const FPCGPoint* CurrentPoint, const FPCGPoint* GoalPoint) -> double
+		{
+			return (GoalPoints.Num() > 1) ? 0.0 : SearchSettings.HeuristicWeight * SearchState.HeuristicFunction(CurrentPoint->Transform.GetLocation(), GoalPoint->Transform.GetLocation());
+		};
+
 		// Not tracking this point yet--add it to the node list and map it.
 		if (!bTrackedPoint)
 		{
-			const double HeuristicCost = SearchSettings.HeuristicWeight * SearchState.HeuristicFunction(Point->Transform.GetLocation(), SearchSettings.Goal);
+			const double HeuristicCost = GetHeuristicCost(Point, CurrentGoal);
 			NodeList.Emplace(Point, CurrentNodeIndex, TentativeNewLocalCost, TentativeNewLocalCost + HeuristicCost);
 			const int32 NewNodeIndex = NodeList.Num() - 1;
-			SearchState.OpenIndexList.HeapPush(NewNodeIndex, [&SearchState](int32 Index1, int32 Index2) { return Helpers::CompareNodes(SearchState.NodeList[Index1], SearchState.NodeList[Index2]); });
+			SearchState.OpenIndexList.HeapPush(NewNodeIndex, HeapPredicate);
 			SearchState.PointToNodeIndexMap.Emplace(Point, NewNodeIndex);
 
 			return;
@@ -127,25 +160,31 @@ bool ExecuteSearchIteration(const FSearchSettings& SearchSettings, FSearchState&
 		FNode& Neighbor = SearchState.NodeList[NeighborIndex];
 		if (TentativeNewLocalCost < Neighbor.LocalCost)
 		{
-			check(SearchState.OpenIndexList.Contains(NeighborIndex));
+			const int32 CurrentOpenListIndex = SearchState.OpenIndexList.Find(NeighborIndex);
+			check(CurrentOpenListIndex != INDEX_NONE);
+
+			// The priority may have changed, so remove the index and add it again.
+			SearchState.OpenIndexList.HeapRemoveAt(CurrentOpenListIndex, HeapPredicate, EAllowShrinking::No);
 
-			const double HeuristicCost = SearchSettings.HeuristicWeight * SearchState.HeuristicFunction(Point->Transform.GetLocation(), SearchSettings.Goal);
+			const double HeuristicCost = GetHeuristicCost(Point, CurrentGoal);
 			Neighbor.PreviousNodeIndex = CurrentNodeIndex;
 			Neighbor.LocalCost = TentativeNewLocalCost;
 			Neighbor.EstimatedGoalCost = TentativeNewLocalCost + HeuristicCost;
+
+			SearchState.OpenIndexList.HeapPush(NeighborIndex, HeapPredicate);
 		}
 	};
 
 	// Arrived at the goal.
-	if (FMath::IsNearlyZero(SquareDistanceToGoal))
+	if (FMath::IsNearlyZero(MinSquaredDistanceToGoal))
 	{
 		Helpers::BuildFinalPath(SearchState, SearchState.NodeList[CurrentNodeIndex], SearchSettings.bCopyOriginatingPoints, OutPath);
-		return true;
+		return ESearchResult::Complete;
 	}
 	// Close enough to the goal we need to add it as a possibility.
-	else if (SquareDistanceToGoal <= SearchSettings.SearchDistance * SearchSettings.SearchDistance)
+	else if (MinSquaredDistanceToGoal <= SearchSettings.SearchDistance * SearchSettings.SearchDistance)
 	{
-		PerPointProcessing(&SearchSettings.GoalPoint, SquareDistanceToGoal);
+		PerPointProcessing(CurrentGoal, MinSquaredDistanceToGoal);
 	}
 
 	// Gather neighbors within the search radius.
@@ -168,20 +207,29 @@ bool ExecuteSearchIteration(const FSearchSettings& SearchSettings, FSearchState&
 
 		if (SearchSettings.bAcceptPartialPath)
 		{
-			auto DistanceSquaredToGoal = [&NodeList = SearchState.NodeList, &Goal = SearchSettings.Goal](const int32 NodeIndex)
-			{
-				return FVector::DistSquared(NodeList[NodeIndex].PCGPoint->Transform.GetLocation(), Goal);
-			};
-
-			if (const int32* MaxNodeIndex = Algo::MinElementBy(SearchState.ClosedIndexList, DistanceSquaredToGoal))
+			// Find the closest point to any goal and walk back the path from there.
+			int32 ClosestNodeIndex = INDEX_NONE;
+			double ClosestSquaredDistance = std::numeric_limits<double>::max();
+			for (const int32 NodeIndex : SearchState.ClosedIndexList)
 			{
-				Helpers::BuildFinalPath(SearchState, SearchState.NodeList[*MaxNodeIndex], SearchSettings.bCopyOriginatingPoints, OutPath);
+				auto [Point, SquaredDistance] = Helpers::GetNearestGoalToLocation(SearchState.NodeList[NodeIndex].PCGPoint->Transform.GetLocation(), GoalPoints);
+				if (SquaredDistance < ClosestSquaredDistance)
+				{
+					ClosestSquaredDistance = SquaredDistance;
+					ClosestNodeIndex = NodeIndex;
+				}
 			}
-		}
 
-		return true;
+			Helpers::BuildFinalPath(SearchState, SearchState.NodeList[ClosestNodeIndex], SearchSettings.bCopyOriginatingPoints, OutPath);
+
+			return ESearchResult::Partial;
+		}
+		else
+		{
+			return ESearchResult::Invalid;
+		}
 	}
 
-	return false;
+	return ESearchResult::Processing;
 }
-} // namespace PCGSpatialAlgo::AStar
+} // namespace PCGSpatialAlgo::AStar
\ No newline at end of file
diff --git a/Engine/Plugins/PCG/Source/PCG/Public/Elements/PCGPathfindingElement.h b/Engine/Plugins/PCG/Source/PCG/Public/Elements/PCGPathfindingElement.h
--- a/Engine/Plugins/PCG/Source/PCG/Public/Elements/PCGPathfindingElement.h
+++ b/Engine/Plugins/PCG/Source/PCG/Public/Elements/PCGPathfindingElement.h
@@ -24,7 +24,7 @@ enum class EPCGPathfindingCostFunctionMode : uint8
 {
 	Distance UMETA(Tooltip = "Pathfinding cost will be the distance only."),
 	FitnessScore UMETA(Tooltip = "Pathfinding cost will be driven by a fitness score (0-1 range), with a maximum penalty applied at fitness = 0."),
-	CostMultipler UMETA(Tooltip = "Pathfinding cost will be the distance multiplied by the provided factor. Note that multipliers below 1 will be clamped to 1.")
+	CostMultiplier UMETA(Tooltip = "Pathfinding cost will be the distance multiplied by the provided factor. Note that multipliers below 1 will be clamped to 1.")
 };
 
 /** Finds the optimal path across the points of a given point cloud--should one exist--when provided a start and goal
@@ -60,16 +60,28 @@ public:
 	UPROPERTY(EditAnywhere, BlueprintReadWrite, Category = Settings, meta = (ClampMin = "0.01", PCG_Overridable))
 	double SearchDistance = 1000;
 
-	/** The path's starting location. */
-	UPROPERTY(EditAnywhere, BlueprintReadWrite, Category = Settings, meta = (PCG_Overridable))
+	UPROPERTY(EditAnywhere, BlueprintReadWrite, Category = Settings)
+	bool bStartLocationsAsInput = false;
+
+	UPROPERTY(EditAnywhere, BlueprintReadWrite, Category = Settings, meta = (EditCondition = "bStartLocationsAsInput", EditConditionHides, PCG_Overridable))
+	FPCGAttributePropertyInputSelector StartLocationAttribute;
+
+	/** The location the pathfinding should attempt to reach. Not used when using start locations from an input. */
+	UPROPERTY(EditAnywhere, BlueprintReadWrite, Category = Settings, meta = (EditCondition = "!bStartLocationsAsInput", EditConditionHides, PCG_Overridable))
 	FVector Start = FVector::ZeroVector;
 
-	/** The location the pathfinding should attempt to reach. */
-	UPROPERTY(EditAnywhere, BlueprintReadWrite, Category = Settings, meta = (PCG_Overridable))
+	UPROPERTY(EditAnywhere, BlueprintReadWrite, Category = Settings)
+	bool bGoalLocationsAsInput = false;
+
+	UPROPERTY(EditAnywhere, BlueprintReadWrite, Category = Settings, meta = (EditCondition = "bGoalLocationsAsInput", EditConditionHides, PCG_Overridable))
+	FPCGAttributePropertyInputSelector GoalLocationAttribute;
+
+	/** The location the pathfinding should attempt to reach. Not used when using goal locations from an input. */
+	UPROPERTY(EditAnywhere, BlueprintReadWrite, Category = Settings, meta = (EditCondition = "!bGoalLocationsAsInput", EditConditionHides, PCG_Overridable))
 	FVector Goal = FVector::ZeroVector;
 
-	/** The heuristic estimates a faster path to speed up processing. A lower heuristic weight can be faster, but it may cease being the optimal path. A weight of 0 is essentially flood fill. */
-	UPROPERTY(EditAnywhere, BlueprintReadWrite, Category = Settings, meta = (EditCondition = "Algorithm == EPCGPathfindingAlgorithm::AStar", PCG_Overridable))
+	/** The heuristic estimates a faster path to speed up processing. A higher than 1 heuristic weight can be faster, but it may cease being the optimal path. A weight of 0 is essentially flood fill. */
+	UPROPERTY(EditAnywhere, BlueprintReadWrite, Category = Settings, meta = (ClampMin = "0.0", EditCondition = "Algorithm == EPCGPathfindingAlgorithm::AStar", PCG_Overridable))
 	double HeuristicWeight = 1.0;
 
 	/** Controls whether the cost function will use a given attribute as a scalar wrt to the distance. */
@@ -91,7 +103,7 @@ public:
 	UPROPERTY(EditAnywhere, BlueprintReadWrite, Category = Settings, meta = (EditCondition = "bUsePathTraces", EditConditionHides, PCG_Overridable, ShowOnlyInnerProperties))
 	FPCGWorldRaycastQueryParams PathTraceParams;
 
-	/** Even if the path is not complete, return the most optimal and viable partial path to the goal. */
+	/** Even if the path is not complete, return a viable partial path to the point closest to the goal. Output data will be tagged with "CompletePath" or "PartialPath", depending on the result, if enabled. */
 	UPROPERTY(EditAnywhere, BlueprintReadWrite, Category = Settings, meta = (PCG_Overridable))
 	bool bAcceptPartialPath = true;
 
@@ -108,10 +120,32 @@ public:
 	bool bCopyOriginatingPoints = false;
 };
 
-class FPCGPathfindingElement : public TPCGTimeSlicedElementBase<PCGSpatialAlgo::AStar::FSearchSettings, PCGSpatialAlgo::AStar::FSearchState>
+namespace PCGPathfindingElement
+{
+	struct FExecutionState
+	{
+		/** The starting points of the paths to search. One search iteration will happen for each starting point. */
+		TArray<FPCGPoint> StartPoints;
+		/** The goal points of the paths to search for. If more than one goal exists, the heuristic will be ignored. */
+		TArray<FPCGPoint> GoalPoints;
+		PCGSpatialAlgo::AStar::FSearchSettings Settings;
+	};
+
+	struct FIterationState
+	{
+		/** Tracks which path this iteration should follow based on the start points. */
+		int32 PathIterationIndex = 0;
+		/** To compare against PathIterationIndex to know if the iteration has changed, for re-initialization of the search. */
+		int32 LastPathIterationIndex = INDEX_NONE;
+		PCGSpatialAlgo::AStar::FSearchState SearchState;
+	};
+}
+
+class FPCGPathfindingElement : public TPCGTimeSlicedElementBase<PCGPathfindingElement::FExecutionState, PCGPathfindingElement::FIterationState>
 {
 public:
 	virtual bool IsCacheable(const UPCGSettings* InSettings) const override;
+
 protected:
 	virtual bool PrepareDataInternal(FPCGContext* InContext) const override;
 	virtual bool ExecuteInternal(FPCGContext* InContext) const override;
diff --git a/Engine/Plugins/PCG/Source/PCG/Public/SpatialAlgo/PCGAStar.h b/Engine/Plugins/PCG/Source/PCG/Public/SpatialAlgo/PCGAStar.h
--- a/Engine/Plugins/PCG/Source/PCG/Public/SpatialAlgo/PCGAStar.h
+++ b/Engine/Plugins/PCG/Source/PCG/Public/SpatialAlgo/PCGAStar.h
@@ -39,15 +39,8 @@ namespace Heuristic
 
 struct FSearchSettings
 {
-	/** Cached start/goal points as initial points in an otherwise const point data. */
-	FPCGPoint StartPoint;
-	FPCGPoint GoalPoint;
 	/** The max distance from each point to search for the next viable point in the path. */
 	double SearchDistance = 1000;
-	/** The path's starting location. */
-	FVector Start = FVector::ZeroVector;
-	/** The location the pathfinding should attempt to reach. */
-	FVector Goal = FVector::ZeroVector;
 	/** The heuristic estimates a faster path to speed up processing. A lower heuristic weight can be faster, but it may cease being the optimal path. A weight of 0 is essentially flood fill. */
 	double HeuristicWeight = 1.0;
 	/** Even if the path is not complete, return the most optimal and viable partial path to the goal. */
@@ -59,20 +52,20 @@ struct FSearchSettings
 struct FSearchState
 {
 	static uint32 constexpr PreAllocNodeCount = 1024u;
-	// Keep track of the point data for the octree search.
+	/** Keep track of the point data for the octree search. */
 	const UPCGPointData* OriginatingPointData = nullptr;
 	/** Cost function */
 	TFunction<double(const double /*PreviousNodeCost*/, const FPCGPoint* /*PreviousNodePoint*/, const double /*DistanceToPreviousNodeSquared*/, const FPCGPoint* /*CurrentPoint*/)> CostFunction = Cost::CalculateCost_EuclideanDistance;
 	/** Heuristic function */
 	TFunction<double(const FVector& /*CurrentLocation*/, const FVector& /*GoalLocation*/)> HeuristicFunction = Heuristic::CalculateHeuristic_EuclideanDistance;
 
-	// Storage of the current known node. Source of truth for the nodes, so it should not be resized after its initial reserve.
+	/** Storage of the current known node. Source of truth for the nodes, so it should not be resized after its initial reserve. */
 	TArray<FNode, TInlineAllocator<PreAllocNodeCount>> NodeList;
-	// Used as a priority list of nodes, sorted by cost.
+	/** Used as a priority list of nodes, sorted by cost. */
 	TArray<int32, TInlineAllocator<PreAllocNodeCount>> OpenIndexList;
-	// Nodes that are no longer viable.
+	/** Nodes that are no longer viable. */
 	TArray<int32, TInlineAllocator<PreAllocNodeCount>> ClosedIndexList;
-	// To track nodes to their corresponding points.
+	/** To track nodes to their corresponding points. */
 	TMap<const FPCGPoint*, int32> PointToNodeIndexMap;
 };
 
@@ -82,9 +75,17 @@ namespace Helpers
 	void BuildFinalPath(const FSearchState& SearchState, const FNode& FinalNode, bool bCopyOriginatingPoints, TArray<FPCGPoint>& OutPath);
 }
 
+enum class ESearchResult
+{
+	Invalid = 0,
+	Processing,
+	Partial,
+	Complete
+};
+
 /** Initialize the search state for AStar. Must be called before ExecuteSearchIteration. */
-void Initialize(const UPCGPointData* const PointData, const FSearchSettings& Settings, FSearchState& OutSearchState);
+void Initialize(const FPCGPoint* const StartPoint, FSearchState& OutSearchState);
 
 /** Runs a single iteration of the A* algorithm. Intended to be called multiple times, whereupon it will return true when the algorithm is finished. */
-bool ExecuteSearchIteration(const FSearchSettings& SearchSettings, FSearchState& SearchState, TArray<FPCGPoint>& OutPath);
+ESearchResult ExecuteSearchIteration(const FSearchSettings& SearchSettings, FSearchState& SearchState, const TArray<FPCGPoint>& GoalPoints, TArray<FPCGPoint>& OutPath);
 }
