[PCG] Fix static analysis error in PCGAstar for Android

#rnx

diff --git a/Engine/Plugins/PCG/Source/PCG/Private/SpatialAlgo/PCGAStar.cpp b/Engine/Plugins/PCG/Source/PCG/Private/SpatialAlgo/PCGAStar.cpp
--- a/Engine/Plugins/PCG/Source/PCG/Private/SpatialAlgo/PCGAStar.cpp
+++ b/Engine/Plugins/PCG/Source/PCG/Private/SpatialAlgo/PCGAStar.cpp
@@ -115,7 +115,9 @@ ESearchResult ExecuteSearchIteration(const FSearchSettings& SearchSettings, FSea
 	SearchState.OpenIndexList.HeapPop(CurrentNodeIndex, HeapPredicate, EAllowShrinking::No);
 	check(NodeList[CurrentNodeIndex].PCGPoint);
 
-	auto [CurrentGoal, MinSquaredDistanceToGoal] = Helpers::GetNearestGoalToLocation(NodeList[CurrentNodeIndex].PCGPoint->Transform.GetLocation(), GoalPoints);
+	TTuple<const FPCGPoint*, double> GoalInfo = Helpers::GetNearestGoalToLocation(NodeList[CurrentNodeIndex].PCGPoint->Transform.GetLocation(), GoalPoints);
+	const FPCGPoint* CurrentGoal = GoalInfo.Get<0>();
+	const double MinSquaredDistanceToGoal = GoalInfo.Get<1>();
 
 	auto PerPointProcessing = [&NodeList, CurrentNodeIndex, CurrentGoal, &HeapPredicate, &SearchSettings, &SearchState, &GoalPoints](const FPCGPoint* Point, const double DistanceToPointSquared)
 	{
@@ -212,7 +214,8 @@ ESearchResult ExecuteSearchIteration(const FSearchSettings& SearchSettings, FSea
 			double ClosestSquaredDistance = std::numeric_limits<double>::max();
 			for (const int32 NodeIndex : SearchState.ClosedIndexList)
 			{
-				auto [Point, SquaredDistance] = Helpers::GetNearestGoalToLocation(SearchState.NodeList[NodeIndex].PCGPoint->Transform.GetLocation(), GoalPoints);
+				TTuple<const FPCGPoint*, double> ClosestGoalInfo = Helpers::GetNearestGoalToLocation(SearchState.NodeList[NodeIndex].PCGPoint->Transform.GetLocation(), GoalPoints);
+				const double SquaredDistance = ClosestGoalInfo.Get<1>();
 				if (SquaredDistance < ClosestSquaredDistance)
 				{
 					ClosestSquaredDistance = SquaredDistance;
