[PCG] Add support for TSet extraction (now that it is supported by StructUtils)

#rnx

diff --git a/Engine/Plugins/PCG/Source/PCG/Private/Helpers/PCGPropertyHelpers.cpp b/Engine/Plugins/PCG/Source/PCG/Private/Helpers/PCGPropertyHelpers.cpp
--- a/Engine/Plugins/PCG/Source/PCG/Private/Helpers/PCGPropertyHelpers.cpp
+++ b/Engine/Plugins/PCG/Source/PCG/Private/Helpers/PCGPropertyHelpers.cpp
@@ -34,36 +34,32 @@ namespace PCGPropertyHelpers
 	}
 
 	/**
-	* Expands container locations to their contents when the property passed in is an array.
-	* This is useful to allow extraction downstream of properties inside of arrays and also to generate the list of addresses/values to look at
+	* Expands container locations to their contents when the property passed in is an array or a set.
+	* This is useful to allow extraction downstream of properties inside of arrays/sets and also to generate the list of addresses/values to look at
 	* when extracting the values to the attribute set.
 	* 
-	* @param InArrayProperty   Property that drives the container expansion. Can be null, in which case we'll copy the container locations directly.
-	* @param InContainers      Container locations to expand
-	* @param OutContainers     Expanded container locations. Expected to be a different array than InContainers.
+	* @param InContainerProperty Property that drives the container expansion.
+	* @param InContainers        Container locations to expand
+	* @param OutContainers       Expanded container locations. Expected to be a different array than InContainers.
 	*/
-	template<typename FirstArrayType, typename SecondArrayType>
-	void ExpandContainers(const FArrayProperty* InArrayProperty, const FirstArrayType& InContainers, SecondArrayType& OutContainers)
+	template<typename ContainerProperty, typename FirstArrayType, typename SecondArrayType>
+	void ExpandContainers(const ContainerProperty* InContainerProperty, const FirstArrayType& InContainers, SecondArrayType& OutContainers)
 	{
-		check(OutContainers.IsEmpty());
+		check(OutContainers.IsEmpty() && InContainerProperty);
 
-		if (InArrayProperty)
+		static_assert(std::is_same_v<ContainerProperty, FArrayProperty> || std::is_same_v<ContainerProperty, FSetProperty>);
+		using FScriptContainerHelper = std::conditional_t<std::is_same_v<ContainerProperty, FArrayProperty>, FScriptArrayHelper_InContainer, FScriptSetHelper_InContainer>;
+
+		for (const void* Container : InContainers)
 		{
-			for (const void* Container : InContainers)
+			FScriptContainerHelper Helper(InContainerProperty, Container);
+			int32 Offset = OutContainers.Num();
+			OutContainers.SetNumUninitialized(OutContainers.Num() + Helper.Num());
+			for (int32 DynamicIndex = 0; DynamicIndex < Helper.Num(); ++DynamicIndex)
 			{
-				FScriptArrayHelper_InContainer Helper(InArrayProperty, Container);
-				int32 Offset = OutContainers.Num();
-				OutContainers.SetNumUninitialized(OutContainers.Num() + Helper.Num());
-				for (int32 DynamicIndex = 0; DynamicIndex < Helper.Num(); ++DynamicIndex)
-				{
-					OutContainers[Offset + DynamicIndex] = Helper.GetRawPtr(DynamicIndex);
-				}
+				OutContainers[Offset + DynamicIndex] = Helper.GetElementPtr(DynamicIndex);
 			}
 		}
-		else
-		{
-			OutContainers = InContainers;
-		}
 	}
 
 	/**
@@ -121,6 +117,34 @@ namespace PCGPropertyHelpers
 			return nullptr;
 		}
 
+		auto ExtractContainers = [&OutContainers](UStruct*& NextClass, const auto* ContainerProperty, const FProperty* InnerProperty) -> bool
+		{
+			bool bPropertyNotExtractable = false;
+			
+			if (const FStructProperty* InnerStructProperty = CastField<FStructProperty>(InnerProperty))
+			{
+				NextClass = InnerStructProperty->Struct;
+			}
+			else if (const FObjectProperty* InnerObjectProperty = CastField<FObjectProperty>(InnerProperty))
+			{
+				NextClass = InnerObjectProperty->PropertyClass;
+			}
+			else
+			{
+				bPropertyNotExtractable = true;
+			}
+
+			// If contents of the container are extractable, do so now by replacing the container entry (e.g. the array/set) with the pointer to its contents
+			if (!bPropertyNotExtractable)
+			{
+				TArray<const void*> Subcontainers;
+				PCGPropertyHelpers::ExpandContainers(ContainerProperty, OutContainers, Subcontainers);
+				OutContainers = MoveTemp(Subcontainers);
+			}
+			
+			return bPropertyNotExtractable;
+		};
+
 		if (!NextNames.IsEmpty())
 		{
 			UStruct* NextClass = nullptr;
@@ -144,27 +168,11 @@ namespace PCGPropertyHelpers
 			}
 			else if (const FArrayProperty* ArrayProperty = CastField<FArrayProperty>(Property))
 			{
-				const FProperty* InnerProperty = ArrayProperty->Inner;
-				if (const FStructProperty* InnerStructProperty = CastField<FStructProperty>(InnerProperty))
-				{
-					NextClass = InnerStructProperty->Struct;
-				}
-				else if (const FObjectProperty* InnerObjectProperty = CastField<FObjectProperty>(InnerProperty))
-				{
-					NextClass = InnerObjectProperty->PropertyClass;
-				}
-				else
-				{
-					bPropertyNotExtractable = true;
-				}
-
-				// If contents of the array are extractable, do so now by replacing the container entry (e.g. the array) with the pointer to its contents
-				if (!bPropertyNotExtractable)
-				{
-					TArray<const void*> Subcontainers;
-					ExpandContainers(ArrayProperty, OutContainers, Subcontainers);
-					OutContainers = MoveTemp(Subcontainers);
-				}
+				bPropertyNotExtractable = ExtractContainers(NextClass, ArrayProperty, ArrayProperty->Inner);
+			}
+			else if (const FSetProperty* SetProperty = CastField<FSetProperty>(Property))
+			{
+				bPropertyNotExtractable = ExtractContainers(NextClass, SetProperty, SetProperty->ElementProp);
 			}
 			else
 			{
@@ -270,12 +278,17 @@ UPCGParamData* PCGPropertyHelpers::ExtractPropertyAsAttributeSet(const PCGProper
 			}
 		}
 
-		// If the property is an array, we will work on the underlying property, and extract each element as an entry in the param data
+		// If the property is an array/set, we will work on the underlying property, and extract each element as an entry in the param data
 		const FArrayProperty* ArrayProperty = CastField<FArrayProperty>(Property);
+		const FSetProperty* SetProperty = CastField<FSetProperty>(Property);
 		if (ArrayProperty)
 		{
 			Property = ArrayProperty->Inner;
 		}
+		else if (SetProperty)
+		{
+			Property = SetProperty->ElementProp;
+		}
 
 		using ExtractablePropertyTuple = TTuple<FString, const FProperty*>;
 		TArray<ExtractablePropertyTuple> ExtractableProperties;
@@ -362,9 +375,23 @@ UPCGParamData* PCGPropertyHelpers::ExtractPropertyAsAttributeSet(const PCGProper
 			return nullptr;
 		}
 
-		// Before we need to compute all the addresses for each entry in our array (or just a single entry if there is no array)
-		TArray<const void*, TInlineAllocator<16>> ElementAddresses;
-		ExpandContainers(ArrayProperty, Containers, ElementAddresses);
+		// Before we need to compute all the addresses for each entry in our array/set (or just a single entry if there is no array/set)
+		TArray<const void*, TInlineAllocator<16>> ExpandedContainers;
+		TArrayView<const void*> ElementAddresses;
+		if (SetProperty)
+		{
+			ExpandContainers(SetProperty, Containers, ExpandedContainers);
+			ElementAddresses = MakeArrayView(ExpandedContainers);
+		}
+		else if (ArrayProperty)
+		{
+			ExpandContainers(ArrayProperty, Containers, ExpandedContainers);
+			ElementAddresses = MakeArrayView(ExpandedContainers);
+		}
+		else
+		{
+			ElementAddresses = MakeArrayView(Containers);
+		}
 
 		if (!ParamData)
 		{
