Added an octree in the collapse points to make it more efficient in finding closest neighbors

#rnx

diff --git a/Engine/Plugins/PCG/Source/PCG/Private/Elements/PCGCollapsePoints.cpp b/Engine/Plugins/PCG/Source/PCG/Private/Elements/PCGCollapsePoints.cpp
--- a/Engine/Plugins/PCG/Source/PCG/Private/Elements/PCGCollapsePoints.cpp
+++ b/Engine/Plugins/PCG/Source/PCG/Private/Elements/PCGCollapsePoints.cpp
@@ -35,13 +35,30 @@ namespace PCGCollapsePoints
 				const double WeightSum = PrimaryWeight + SecondaryWeight;
 
 				double Alpha = FMath::IsNearlyZero(WeightSum) ? 0.5 : (SecondaryWeight / WeightSum);
-				const FVector DeltaPosition = OutState.Positions[SecondaryPointIndex] - OutState.Positions[PrimaryPointIndex];
+				const FVector DeltaPosition = OutState.Points[SecondaryPointIndex].Transform.GetLocation() - OutState.Points[PrimaryPointIndex].Transform.GetLocation();
 
-				OutState.Positions[PrimaryPointIndex] = OutState.Positions[PrimaryPointIndex] + Alpha * DeltaPosition;
+				OutState.Points[PrimaryPointIndex].Transform.AddToTranslation(Alpha * DeltaPosition);
 				OutState.Weights[PrimaryPointIndex] = WeightSum;
 				OutState.Merged[SecondaryPointIndex] = PrimaryPointIndex;
 			}
 		}
+
+		void RebuildOctree(FPCGContext* InContext, FCollapsePointsState& OutState)
+		{
+			TRACE_CPUPROFILER_EVENT_SCOPE(PCGCollapsePointsElement::Algo::RebuildOctree);
+			check(OutState.SourceData && OutState.Merged.Num() == OutState.Points.Num());
+
+			UPCGPointData::PointOctree NewOctree(OutState.SourceData->GetBounds().GetCenter(), OutState.SourceData->GetBounds().GetExtent().Length());
+			for (int PointIndex = 0; PointIndex < OutState.Points.Num(); ++PointIndex)
+			{
+				if (OutState.Merged[PointIndex] == INDEX_NONE)
+				{
+					NewOctree.AddElement(FPCGPointRef(OutState.Points[PointIndex]));
+				}
+			}
+
+			OutState.Octree = MoveTemp(NewOctree);
+		}
 	}
 
 	namespace Modes
@@ -75,30 +92,48 @@ namespace PCGCollapsePoints
 				// Find closest unvisited point after current point that's within the distance threshold.
 				double MinSqrDistance = Settings.DistanceThreshold * Settings.DistanceThreshold;
 				int ClosestUnvisitedIndex = INDEX_NONE;
+				bool bHasColocatedPoint = false;
 
-				for (int NeighborVisitIndex = VisitIndex + 1; NeighborVisitIndex < OutState.VisitOrder.Num(); ++NeighborVisitIndex)
+				const double Extents = UE_DOUBLE_SQRT_2 * Settings.DistanceThreshold;
+				FBoxCenterAndExtent SearchBounds(OutState.Points[PointIndex].Transform.GetLocation(), FVector(Extents, Extents, Extents));
+				OutState.Octree.FindElementsWithBoundsTest(SearchBounds, [PointIndex, &MinSqrDistance, &ClosestUnvisitedIndex, &bHasColocatedPoint, &OutState](const FPCGPointRef& PointRef)
 				{
-					int NeighborIndex = OutState.VisitOrder[NeighborVisitIndex];
+					int NeighborIndex = PointRef.Point - OutState.Points.GetData();
 
 					if (OutState.Merged[NeighborIndex] != INDEX_NONE || OutState.Visited[NeighborIndex])
 					{
-						continue;
+						return;
 					}
 
-					double SqrDistance = (OutState.Positions[PointIndex] - OutState.Positions[NeighborIndex]).SquaredLength();
+					double SqrDistance = (OutState.Points[PointIndex].Transform.GetLocation() - OutState.Points[NeighborIndex].Transform.GetLocation()).SquaredLength();
 
-					if (SqrDistance < MinSqrDistance)
+					if (FMath::IsNearlyZero(SqrDistance))
 					{
-						MinSqrDistance = SqrDistance;
-						ClosestUnvisitedIndex = NeighborIndex;
+						if (bHasColocatedPoint)
+						{
+							check(ClosestUnvisitedIndex != INDEX_NONE);
+							// Prioritize visit order then
+							const int32 NeighborVisitOrder = OutState.VisitOrder.IndexOfByKey(NeighborIndex);
+							const int32 ClosestVisitOrder = OutState.VisitOrder.IndexOfByKey(ClosestUnvisitedIndex);
 
-						// Since we're visiting in order, this is the best pick we'll have.
-						if (FMath::IsNearlyZero(MinSqrDistance))
+							if (NeighborVisitOrder < ClosestVisitOrder)
+							{
+								ClosestUnvisitedIndex = NeighborIndex;
+							}
+						}
+						else
 						{
-							break;
+							bHasColocatedPoint = true;
+							ClosestUnvisitedIndex = NeighborIndex;
+							MinSqrDistance = SqrDistance;
 						}
 					}
-				}
+					else if (SqrDistance < MinSqrDistance)
+					{
+						MinSqrDistance = SqrDistance;
+						ClosestUnvisitedIndex = NeighborIndex;
+					}
+				});
 
 				if (ClosestUnvisitedIndex != INDEX_NONE)
 				{
@@ -208,13 +243,8 @@ bool FPCGCollapsePointsElement::PrepareDataInternal(FPCGContext* InContext) cons
 		const TArray<FPCGPoint>& Points = PointData->GetPoints();
 		const int32 NumPoints = Points.Num();
 
-		// Get positions
-		OutState.Positions.Reserve(Points.Num());
-
-		for (const FPCGPoint& Point : Points)
-		{
-			OutState.Positions.Add(Point.Transform.GetLocation());
-		}
+		// Get points, octree will be rebuilt later
+		OutState.Points = Points;
 
 		// Get weights
 		OutState.Weights.SetNumUninitialized(NumPoints);
@@ -340,6 +370,9 @@ bool FPCGCollapsePointsElement::PrepareDataInternal(FPCGContext* InContext) cons
 			return EPCGTimeSliceInitResult::NoOperation;
 		}
 
+		// Rebuild octree
+		PCGCollapsePoints::Algo::RebuildOctree(Context, OutState);
+
 		Context->OutputData.TaggedData.Emplace_GetRef().Data = OutPointData;
 		return EPCGTimeSliceInitResult::Success;
 	});
@@ -383,10 +416,14 @@ bool FPCGCollapsePointsElement::ExecuteInternal(FPCGContext* InContext) const
 		{
 			// Merge based on selection
 			CollapseSettings.MergeSelectionFunc(InContext, CollapseSettings, CollapseState);
+			// Rebuild octree for next iteration
+			PCGCollapsePoints::Algo::RebuildOctree(InContext, CollapseState);
+
 			return false;
 		}
 		else
 		{
+			TRACE_CPUPROFILER_EVENT_SCOPE(FPCGCollapsePointsElement::Execute::ComputeFinalResults);
 			TArray<int>& Merged = CollapseState.Merged;
 
 			// We're done - compute final results
diff --git a/Engine/Plugins/PCG/Source/PCG/Public/Elements/PCGCollapsePoints.h b/Engine/Plugins/PCG/Source/PCG/Public/Elements/PCGCollapsePoints.h
--- a/Engine/Plugins/PCG/Source/PCG/Public/Elements/PCGCollapsePoints.h
+++ b/Engine/Plugins/PCG/Source/PCG/Public/Elements/PCGCollapsePoints.h
@@ -4,6 +4,7 @@
 
 #include "PCGSettings.h"
 
+#include "Data/PCGPointData.h"
 #include "Elements/PCGTimeSlicedElementBase.h"
 #include "Metadata/PCGMetadataAttributeTpl.h"
 #include "Metadata/Accessors/PCGPropertyAccessor.h"
@@ -86,7 +87,9 @@ namespace PCGCollapsePoints
 	struct FCollapsePointsState
 	{
 		// Current iterative state during merge
-		TArray<FVector> Positions;
+		TArray<FPCGPoint> Points;
+		UPCGPointData::PointOctree Octree;
+
 		TArray<double> Weights;
 		TArray<int> Merged;
 		
