[PCG] Changed WorldQueryHelpers to not overwrite attributes that already exist.

This was overwriting landscape and other metadata, and other minor cases resulting in a unneeded warning.

#rnx

diff --git a/Engine/Plugins/PCG/Source/PCG/Private/Helpers/PCGWorldQueryHelpers.cpp b/Engine/Plugins/PCG/Source/PCG/Private/Helpers/PCGWorldQueryHelpers.cpp
--- a/Engine/Plugins/PCG/Source/PCG/Private/Helpers/PCGWorldQueryHelpers.cpp
+++ b/Engine/Plugins/PCG/Source/PCG/Private/Helpers/PCGWorldQueryHelpers.cpp
@@ -181,7 +181,7 @@ namespace PCGWorldQueryHelpers
 
 			if (!OutMetadata->HasAttribute(AttributeName))
 			{
-				if (OutMetadata->CreateAttribute<Type>(AttributeName, DefaultValue, /*bAllowsInterpolation=*/true, /*bOverrideParent=*/false))
+				if (OutMetadata->FindOrCreateAttribute<Type>(AttributeName, DefaultValue, /*bAllowsInterpolation=*/true, /*bOverrideParent=*/false, /*bOverwriteIfTypeMismatch=*/true))
 				{
 					return true;
 				}
@@ -349,4 +349,4 @@ namespace PCGWorldQueryHelpers
 
 		return bResult;
 	}
-}
+}
\ No newline at end of file
