[PCG] Fix PCGCollapseSubgraphTest failed check because no CurrentGenerationTask was assigned to component

#rnx

diff --git a/Engine/Plugins/PCG/Source/PCG/Private/Tests/PCGTestsCommon.cpp b/Engine/Plugins/PCG/Source/PCG/Private/Tests/PCGTestsCommon.cpp
--- a/Engine/Plugins/PCG/Source/PCG/Private/Tests/PCGTestsCommon.cpp
+++ b/Engine/Plugins/PCG/Source/PCG/Private/Tests/PCGTestsCommon.cpp
@@ -113,6 +113,14 @@ namespace PCGTestsCommon
 		return PCGTestsCommon::InitializeTestContext(Settings->GetElement().Get(), InputData, TestPCGComponent, InNode);
 	}
 
+	void FTestData::SetCurrentGenerationTask(FPCGTaskId InTaskId)
+	{
+		if (TestPCGComponent)
+		{
+			TestPCGComponent->CurrentGenerationTask = InTaskId;
+		}
+	}
+
 	AActor* CreateTemporaryActor()
 	{
 		return NewObject<AActor>();
diff --git a/Engine/Plugins/PCG/Source/PCG/Private/Tests/Subgraphs/PCGCollapseSubgraphTest.cpp b/Engine/Plugins/PCG/Source/PCG/Private/Tests/Subgraphs/PCGCollapseSubgraphTest.cpp
--- a/Engine/Plugins/PCG/Source/PCG/Private/Tests/Subgraphs/PCGCollapseSubgraphTest.cpp
+++ b/Engine/Plugins/PCG/Source/PCG/Private/Tests/Subgraphs/PCGCollapseSubgraphTest.cpp
@@ -280,6 +280,10 @@ bool FPCGCollapseSubgraphCombinationOfCollapses::RunTest(const FString& Paramete
 		TestData.TestPCGComponent->GetGraphInstance()->SetGraphParameter<double>(TEXT("AnotherConstant"), SecondConstant);
 
 		const FPCGTaskId TaskId = GraphExecutor.Schedule(TestData.TestPCGComponent);
+
+		// TODO: Post-release GraphExecutor should be setting this on the Component
+		TestData.SetCurrentGenerationTask(TaskId);
+
 		const FPCGTaskId FinalTaskId = GraphExecutor.ScheduleGeneric([&OutputPoints, &GraphExecutor, TaskId, &bDone]() -> bool
 		{
 			FPCGDataCollection OutputData{};
diff --git a/Engine/Plugins/PCG/Source/PCG/Public/PCGComponent.h b/Engine/Plugins/PCG/Source/PCG/Public/PCGComponent.h
--- a/Engine/Plugins/PCG/Source/PCG/Public/PCGComponent.h
+++ b/Engine/Plugins/PCG/Source/PCG/Public/PCGComponent.h
@@ -77,6 +77,11 @@ enum class EPCGComponentDirtyFlag : uint8
 };
 ENUM_CLASS_FLAGS(EPCGComponentDirtyFlag);
 
+namespace PCGTestsCommon
+{
+	struct FTestData;
+}
+
 UCLASS(BlueprintType, ClassGroup = (Procedural), meta = (BlueprintSpawnableComponent, PrioritizeCategories = "PCG"))
 class PCG_API UPCGComponent : public UActorComponent
 {
@@ -113,6 +118,8 @@ public:
 
 protected:
 	friend struct FPCGComponentInstanceData;
+	friend struct PCGTestsCommon::FTestData;
+
 	virtual TStructOnScope<FActorComponentInstanceData> GetComponentInstanceData() const override;
 	//~End UActorComponent Interface
 
diff --git a/Engine/Plugins/PCG/Source/PCG/Public/Tests/PCGTestsCommon.h b/Engine/Plugins/PCG/Source/PCG/Public/Tests/PCGTestsCommon.h
--- a/Engine/Plugins/PCG/Source/PCG/Public/Tests/PCGTestsCommon.h
+++ b/Engine/Plugins/PCG/Source/PCG/Public/Tests/PCGTestsCommon.h
@@ -41,6 +41,8 @@ namespace PCGTestsCommon
 		/** Initialize a context and set the number of tasks available to 1. */
 		TUniquePtr<FPCGContext> InitializeTestContext(const UPCGNode* InNode = nullptr) const;
 
+		void SetCurrentGenerationTask(FPCGTaskId InTaskId);
+
 		AActor* TestActor;
 		UPCGComponent* TestPCGComponent;
 		FPCGDataCollection InputData;
