[PCG] ExecutorV2: Fix timeslicing of async processing by providing a max task

#rnx

diff --git a/Engine/Plugins/PCG/Source/PCG/Private/Graph/PCGGraphExecutor.cpp b/Engine/Plugins/PCG/Source/PCG/Private/Graph/PCGGraphExecutor.cpp
--- a/Engine/Plugins/PCG/Source/PCG/Private/Graph/PCGGraphExecutor.cpp
+++ b/Engine/Plugins/PCG/Source/PCG/Private/Graph/PCGGraphExecutor.cpp
@@ -1412,7 +1412,10 @@ bool FPCGGraphExecutor::ExecuteScheduling(double EndTime, TSharedPtr<FPCGGraphAc
 	bool bStateChanged = false;
 	const bool bIsInGameThread = IsInGameThread();
 	TRACE_CPUPROFILER_EVENT_SCOPE(FPCGGraphExecutor::ExecuteScheduling);
-		
+	
+	const float MaxPercentageOfThreadsToUse = FMath::Clamp(CVarMaxPercentageOfThreadsToUse.GetValueOnAnyThread(), 0.0f, 1.0f);
+	const int32 MaxNumThreads = FMath::Max(0, FMath::Min((int32)(FPlatformMisc::NumberOfCoresIncludingHyperthreads() * MaxPercentageOfThreadsToUse), CVarMaxNumTasks.GetValueOnAnyThread() - 1));
+
 	TArray<FCachedResult*> CachedResults;
 	{
 		TRACE_CPUPROFILER_EVENT_SCOPE(FPCGGraphExecutor::ExecuteScheduling::ReadyTasks);
@@ -1530,7 +1533,7 @@ bool FPCGGraphExecutor::ExecuteScheduling(double EndTime, TSharedPtr<FPCGGraphAc
 						check (!ActiveTask->bIsExecutingTask && !ActiveTask->Context->bIsPaused)
 						
 						ActiveTask->StartExecuting();
-						ActiveTask->Context->AsyncState.NumAvailableTasks = -1;
+						ActiveTask->Context->AsyncState.NumAvailableTasks = FMath::Max(1, MaxNumThreads);
 						ActiveTask->Context->AsyncState.bIsRunningOnMainThread = true;
 						ActiveTask->Context->AsyncState.bIsRunningOutOfTick = false;
 						ActiveTask->Context->AsyncState.EndTime = EndTime;
@@ -1541,7 +1544,7 @@ bool FPCGGraphExecutor::ExecuteScheduling(double EndTime, TSharedPtr<FPCGGraphAc
 			}
 
 			// PCGGraphExecutor::CVarMaxWorkerTasks.GetValueOnAnyThread()+1 because the mainthread doesn't count
-			if (PCGGraphExecutor::CVarGraphMultithreading.GetValueOnAnyThread() && PCGGraphExecutor::CVarMaxWorkerTasks.GetValueOnAnyThread()+1 > FPCGGraphActiveTask::NumExecuting)
+			if (PCGGraphExecutor::CVarGraphMultithreading.GetValueOnAnyThread() && PCGGraphExecutor::CVarMaxWorkerTasks.GetValueOnAnyThread()+1 > FPCGGraphActiveTask::NumExecuting && MaxNumThreads > 0)
 			{
 				TRACE_CPUPROFILER_EVENT_SCOPE(FPCGGraphExecutor::ExecuteScheduling::LaunchTasks);
 				for (int32 ExecutionIndex = 0; ExecutionIndex < ActiveTasks.Num(); ++ExecutionIndex)
@@ -1569,7 +1572,8 @@ bool FPCGGraphExecutor::ExecuteScheduling(double EndTime, TSharedPtr<FPCGGraphAc
 #endif
 					{
 						check(ActiveTask->Context->CurrentPhase != EPCGExecutionPhase::Done);
-						ActiveTask->Context->AsyncState.NumAvailableTasks = -1;
+						check(MaxNumThreads > 0);
+						ActiveTask->Context->AsyncState.NumAvailableTasks = MaxNumThreads;
 						ActiveTask->Context->AsyncState.EndTime = EndTime;
 						ActiveTask->Context->AsyncState.bIsRunningOnMainThread = false;
 
