[PCG] GPU - fix spurious warning and ensures occurring in demo level.

#rnx

diff --git a/Engine/Plugins/PCG/Source/PCG/Private/Compute/Elements/PCGCustomHLSL.cpp b/Engine/Plugins/PCG/Source/PCG/Private/Compute/Elements/PCGCustomHLSL.cpp
--- a/Engine/Plugins/PCG/Source/PCG/Private/Compute/Elements/PCGCustomHLSL.cpp
+++ b/Engine/Plugins/PCG/Source/PCG/Private/Compute/Elements/PCGCustomHLSL.cpp
@@ -335,7 +335,19 @@ FPCGDataCollectionDesc UPCGCustomHLSLSettings::ComputeOutputPinDataDesc(const UP
 	// Try to propagate string keys across node. Not trivial because there could be one or more string key attributes on input pins and on output pins,
 	// and it is in general hard to determine from source which string keys from input are being written to outputs. Try first collecting all string keys
 	// from matching attribute names (across all input pins), and then fall back to collecting keys from all string key attributes across all inputs.
-	if (const UPCGNode* Node = Cast<UPCGNode>(GetOuter()))
+	bool bOutputHasStringKeys = false;
+
+	for (const FPCGDataDesc& DataDesc : PinDesc.DataDescs)
+	{
+		if (DataDesc.AttributeDescs.FindByPredicate([](const FPCGKernelAttributeDesc& InAttributeDesc) { return InAttributeDesc.Type == EPCGKernelAttributeType::StringKey; }))
+		{
+			bOutputHasStringKeys = true;
+			break;
+		}
+	}
+
+	const UPCGNode* Node = Cast<UPCGNode>(GetOuter());
+	if (bOutputHasStringKeys && Node)
 	{
 		TArray<FPCGDataCollectionDesc> RelevantInputDataDescs;
 
diff --git a/Engine/Plugins/PCG/Source/PCG/Private/Compute/PCGDataBinding.cpp b/Engine/Plugins/PCG/Source/PCG/Private/Compute/PCGDataBinding.cpp
--- a/Engine/Plugins/PCG/Source/PCG/Private/Compute/PCGDataBinding.cpp
+++ b/Engine/Plugins/PCG/Source/PCG/Private/Compute/PCGDataBinding.cpp
@@ -171,7 +171,7 @@ void UPCGDataBinding::DebugLogDataDescriptions()
 	UE_LOG(LogPCG, Display, TEXT("\n### ATTRIBUTE TABLE ###"));
 	for (TPair<FName, FPCGKernelAttributeIDAndType>& AttributeEntry : GlobalAttributeLookupTable)
 	{
-		UE_LOG(LogPCG, Display, TEXT("\tName: %s\t\tID: %d\t\tType: %d)"), *AttributeEntry.Get<0>().ToString(), AttributeEntry.Get<1>().Id, AttributeEntry.Get<1>().Type);
+		UE_LOG(LogPCG, Display, TEXT("\tName: %s\t\tID: %d\t\tType: %d"), *AttributeEntry.Get<0>().ToString(), AttributeEntry.Get<1>().Id, AttributeEntry.Get<1>().Type);
 	}
 
 	UE_LOG(LogPCG, Display, TEXT("\n### STRING TABLE ###"));
diff --git a/Engine/Plugins/PCG/Source/PCG/Private/Compute/PCGDataForGPU.cpp b/Engine/Plugins/PCG/Source/PCG/Private/Compute/PCGDataForGPU.cpp
--- a/Engine/Plugins/PCG/Source/PCG/Private/Compute/PCGDataForGPU.cpp
+++ b/Engine/Plugins/PCG/Source/PCG/Private/Compute/PCGDataForGPU.cpp
@@ -415,11 +415,6 @@ namespace PCGDataForGPUHelpers
 			{
 				Attribute->SetValue(EntryKey, InStringTable[StringKey]);
 			}
-			else
-			{
-				UE_LOG(LogPCG, Error, TEXT("String retrieval failed for string key %d. String table size is %d."), StringKey, InStringTable.Num());
-				ensure(false);
-			}
 			break;
 		}
 		default:
diff --git a/Engine/Plugins/PCG/Source/PCG/Private/Graph/PCGGraphCompilerGPU.cpp b/Engine/Plugins/PCG/Source/PCG/Private/Graph/PCGGraphCompilerGPU.cpp
--- a/Engine/Plugins/PCG/Source/PCG/Private/Graph/PCGGraphCompilerGPU.cpp
+++ b/Engine/Plugins/PCG/Source/PCG/Private/Graph/PCGGraphCompilerGPU.cpp
@@ -608,9 +608,11 @@ void FPCGGraphCompilerGPU::BuildGPUGraphTask(
 								{
 									FoundEntry->Type = Key.Type;
 								}
-								else
+								else if (Key.Type != FoundEntry->Type)
 								{
-									ensure(Key.Type == FoundEntry->Type);
+									// TODO: This is more a limitation of the current system rather than an issue with the user's data. Log for bug reports etc but remove this
+									// later when we properly flow attribute table and string table down graph branches.
+									UE_LOG(LogPCG, Warning, TEXT("Attribute '%s' multiply defined with different types (existing: %d, new: %d)."), *Key.Name.ToString(), FoundEntry->Type, Key.Type);
 								}
 							}
 						}
