[PCG] GPU - fix change propagation for GPU SM Spawner so results update when disabling nodes.

#rnx

diff --git a/Engine/Plugins/PCG/Source/PCG/Private/Compute/Elements/PCGComputeGraphElement.cpp b/Engine/Plugins/PCG/Source/PCG/Private/Compute/Elements/PCGComputeGraphElement.cpp
--- a/Engine/Plugins/PCG/Source/PCG/Private/Compute/Elements/PCGComputeGraphElement.cpp
+++ b/Engine/Plugins/PCG/Source/PCG/Private/Compute/Elements/PCGComputeGraphElement.cpp
@@ -251,7 +251,7 @@ bool FPCGComputeGraphElement::SetupProceduralISMComponents(FPCGContext* InContex
 	for (const UPCGSettings* Settings : Graph->StaticMeshSpawners)
 	{
 		const UPCGStaticMeshSpawnerSettings* SpawnerSettings = Cast<UPCGStaticMeshSpawnerSettings>(Settings);
-		if (!ensure(SpawnerSettings))
+		if (!ensure(SpawnerSettings) || !SpawnerSettings->bEnabled)
 		{
 			continue;
 		}
diff --git a/Engine/Plugins/PCG/Source/PCG/Private/Elements/PCGStaticMeshSpawner.cpp b/Engine/Plugins/PCG/Source/PCG/Private/Elements/PCGStaticMeshSpawner.cpp
--- a/Engine/Plugins/PCG/Source/PCG/Private/Elements/PCGStaticMeshSpawner.cpp
+++ b/Engine/Plugins/PCG/Source/PCG/Private/Elements/PCGStaticMeshSpawner.cpp
@@ -192,6 +192,19 @@ void UPCGStaticMeshSpawnerSettings::CreateAdditionalOutputDataInterfaces(TArray<
 }
 
 #if WITH_EDITOR
+EPCGChangeType UPCGStaticMeshSpawnerSettings::GetChangeTypeForProperty(const FName& InPropertyName) const
+{
+	EPCGChangeType ChangeType = Super::GetChangeTypeForProperty(InPropertyName);
+
+	if (ChangeType != EPCGChangeType::Cosmetic && ShouldExecuteOnGPU())
+	{
+		// GPU SM Spawners are compiled into compute graph so trigger recompilation.
+		ChangeType |= EPCGChangeType::Structural;
+	}
+
+	return ChangeType;
+}
+
 FText UPCGStaticMeshSpawnerSettings::GetDefaultNodeTitle() const
 {
 	return LOCTEXT("NodeTitle", "Static Mesh Spawner");
diff --git a/Engine/Plugins/PCG/Source/PCG/Public/Elements/PCGStaticMeshSpawner.h b/Engine/Plugins/PCG/Source/PCG/Public/Elements/PCGStaticMeshSpawner.h
--- a/Engine/Plugins/PCG/Source/PCG/Public/Elements/PCGStaticMeshSpawner.h
+++ b/Engine/Plugins/PCG/Source/PCG/Public/Elements/PCGStaticMeshSpawner.h
@@ -37,7 +37,7 @@ public:
 	virtual void CreateAdditionalInputDataInterfaces(TArray<TObjectPtr<UComputeDataInterface>>& OutDataInterfaces) const override;
 	virtual void CreateAdditionalOutputDataInterfaces(TArray<TObjectPtr<UComputeDataInterface>>& OutDataInterfaces) const override;
 #if WITH_EDITOR
-
+	virtual EPCGChangeType GetChangeTypeForProperty(const FName& InPropertyName) const override;
 	virtual bool DisplayExecuteOnGPUSetting() const override { return true; }
 	virtual FName GetDefaultNodeName() const override { return FName(TEXT("StaticMeshSpawner")); }
 	virtual FText GetDefaultNodeTitle() const override;
