[PCG] GPU - make SM Spawner gracefully handle executions where 0 points arrive.

#rnx

diff --git a/Engine/Plugins/PCG/Source/PCG/Private/Compute/DataInterfaces/Elements/PCGStaticMeshSpawnerDataInterface.cpp b/Engine/Plugins/PCG/Source/PCG/Private/Compute/DataInterfaces/Elements/PCGStaticMeshSpawnerDataInterface.cpp
--- a/Engine/Plugins/PCG/Source/PCG/Private/Compute/DataInterfaces/Elements/PCGStaticMeshSpawnerDataInterface.cpp
+++ b/Engine/Plugins/PCG/Source/PCG/Private/Compute/DataInterfaces/Elements/PCGStaticMeshSpawnerDataInterface.cpp
@@ -27,10 +27,6 @@ void UPCGStaticMeshSpawnerDataInterface::GetSupportedInputs(TArray<FShaderFuncti
 		.SetName(TEXT("SMSpawner_GetNumPrimitives"))
 		.AddReturnType(EShaderFundamentalType::Uint); // Num primitives
 
-	OutFunctions.AddDefaulted_GetRef()
-		.SetName(TEXT("SMSpawner_GetNumInputPoints"))
-		.AddReturnType(EShaderFundamentalType::Uint); // Num input points
-
 	OutFunctions.AddDefaulted_GetRef()
 		.SetName(TEXT("SMSpawner_GetAttributeIdOffsetStride"))
 		.AddReturnType(EShaderFundamentalType::Uint, 4)
@@ -54,7 +50,6 @@ BEGIN_SHADER_PARAMETER_STRUCT(FPCGStaticMeshSpawnerDataInterfaceParameters,)
 	SHADER_PARAMETER(uint32, NumAttributes)
 	SHADER_PARAMETER(uint32, NumPrimitives)
 	SHADER_PARAMETER(uint32, SelectorAttributeId)
-	SHADER_PARAMETER(uint32, NumInputPoints)
 END_SHADER_PARAMETER_STRUCT()
 
 void UPCGStaticMeshSpawnerDataInterface::GetShaderParameters(TCHAR const* UID, FShaderParametersMetadataBuilder& InOutBuilder, FShaderParametersMetadataAllocations& InOutAllocations) const
@@ -75,9 +70,6 @@ void UPCGStaticMeshSpawnerDataInterface::GetHLSL(FString& OutHLSL, FString const
 		"uint {DataInterfaceName}_SelectorAttributeId;\n"
 		"uint SMSpawner_GetSelectorAttributeId_{DataInterfaceName}() { return {DataInterfaceName}_SelectorAttributeId; }\n"
 		"\n"
-		"uint {DataInterfaceName}_NumInputPoints;\n"
-		"uint SMSpawner_GetNumInputPoints_{DataInterfaceName}() { return {DataInterfaceName}_NumInputPoints; }\n"
-		"\n"
 		"uint {DataInterfaceName}_NumAttributes;\n"
 		"uint4 {DataInterfaceName}_AttributeIdOffsetStrides[{MaxAttributes}];\n"
 		"uint SMSpawner_GetNumAttributes_{DataInterfaceName}() { return {DataInterfaceName}_NumAttributes; }\n"
@@ -110,28 +102,25 @@ UComputeDataProvider* UPCGStaticMeshSpawnerDataInterface::CreateDataProvider(TOb
 {
 	TRACE_CPUPROFILER_EVENT_SCOPE(UPCGStaticMeshSpawnerDataInterface::CreateDataProvider);
 	UPCGDataBinding* Binding = CastChecked<UPCGDataBinding>(InBinding);
-	const FPCGSpawnerPrimitives* Primitives = Binding->MeshSpawnersToPrimitives.Find(Settings);
-	if (!ensure(Primitives))
-	{
-		return nullptr;
-	}
-	if (!ensure(!Primitives->Primitives.IsEmpty()))
-	{
-		return nullptr;
-	}
 
 	TObjectPtr<UPCGStaticMeshSpawnerDataProvider> DataProvider = NewObject<UPCGStaticMeshSpawnerDataProvider>();
 	DataProvider->Settings = Cast<UPCGStaticMeshSpawnerSettings>(Settings);
-	DataProvider->AttributeIdOffsetStrides = Primitives->AttributeIdOffsetStrides;
-	DataProvider->SelectionCDF = Primitives->SelectionCDF;
-	DataProvider->SelectorAttributeId = Primitives->SelectorAttributeId;
-	DataProvider->PrimitiveStringKeys = Primitives->PrimitiveStringKeys;
+	
+	// If there were 0 input points for this execution, we will not have created any primitives, so check for null.
+	if (const FPCGSpawnerPrimitives* Primitives = Binding->MeshSpawnersToPrimitives.Find(Settings))
+	{
+		DataProvider->AttributeIdOffsetStrides = Primitives->AttributeIdOffsetStrides;
+		DataProvider->SelectionCDF = Primitives->SelectionCDF;
+		DataProvider->SelectorAttributeId = Primitives->SelectorAttributeId;
+		DataProvider->PrimitiveStringKeys = Primitives->PrimitiveStringKeys;
+	}
+
 	return DataProvider;
 }
 
 FComputeDataProviderRenderProxy* UPCGStaticMeshSpawnerDataProvider::GetRenderProxy()
 {
-	return new FPCGStaticMeshSpawnerDataProviderProxy(NumInputPoints, AttributeIdOffsetStrides, SelectorAttributeId, PrimitiveStringKeys, SelectionCDF);
+	return new FPCGStaticMeshSpawnerDataProviderProxy(AttributeIdOffsetStrides, SelectorAttributeId, PrimitiveStringKeys, SelectionCDF);
 }
 
 bool FPCGStaticMeshSpawnerDataProviderProxy::IsValid(FValidationData const& InValidationData) const
diff --git a/Engine/Plugins/PCG/Source/PCG/Private/Compute/DataInterfaces/Elements/PCGStaticMeshSpawnerDataInterface.h b/Engine/Plugins/PCG/Source/PCG/Private/Compute/DataInterfaces/Elements/PCGStaticMeshSpawnerDataInterface.h
--- a/Engine/Plugins/PCG/Source/PCG/Private/Compute/DataInterfaces/Elements/PCGStaticMeshSpawnerDataInterface.h
+++ b/Engine/Plugins/PCG/Source/PCG/Private/Compute/DataInterfaces/Elements/PCGStaticMeshSpawnerDataInterface.h
@@ -64,9 +64,8 @@ public:
 class FPCGStaticMeshSpawnerDataProviderProxy : public FComputeDataProviderRenderProxy
 {
 public:
-	FPCGStaticMeshSpawnerDataProviderProxy(int32 InNumInputPoints, TArray<FUintVector4>& InAttributeIdOffsetStrides, int32 InSelectorAttributeId, TArray<uint32>& InPrimitiveStringKeys, TArray<float> InSelectionCDF)
-		: NumInputPoints(InNumInputPoints)
-		, AttributeIdOffsetStrides(InAttributeIdOffsetStrides)
+	FPCGStaticMeshSpawnerDataProviderProxy(TArray<FUintVector4>& InAttributeIdOffsetStrides, int32 InSelectorAttributeId, TArray<uint32>& InPrimitiveStringKeys, TArray<float> InSelectionCDF)
+		: AttributeIdOffsetStrides(InAttributeIdOffsetStrides)
 		, SelectionCDF(InSelectionCDF)
 		, SelectorAttributeId(InSelectorAttributeId)
 		, PrimitiveStringKeys(InPrimitiveStringKeys)
@@ -80,8 +79,6 @@ public:
 protected:
 	using FParameters = FPCGStaticMeshSpawnerDataInterfaceParameters;
 
-	int32 NumInputPoints = 0;
-
 	TArray<FUintVector4> AttributeIdOffsetStrides;
 	TArray<float> SelectionCDF;
 
diff --git a/Engine/Plugins/PCG/Source/PCG/Private/Compute/DataInterfaces/PCGInstanceDataInterface.cpp b/Engine/Plugins/PCG/Source/PCG/Private/Compute/DataInterfaces/PCGInstanceDataInterface.cpp
--- a/Engine/Plugins/PCG/Source/PCG/Private/Compute/DataInterfaces/PCGInstanceDataInterface.cpp
+++ b/Engine/Plugins/PCG/Source/PCG/Private/Compute/DataInterfaces/PCGInstanceDataInterface.cpp
@@ -114,30 +114,35 @@ UComputeDataProvider* UPCGInstanceDataInterface::CreateDataProvider(TObjectPtr<U
 	TRACE_CPUPROFILER_EVENT_SCOPE(UPCGInstanceDataInterface::CreateDataProvider);
 	UPCGDataBinding* Binding = CastChecked<UPCGDataBinding>(InBinding);
 
-	FPCGSpawnerPrimitives* FoundPrimitives = Binding->MeshSpawnersToPrimitives.Find(ProducerSettings);
-	if (!FoundPrimitives || FoundPrimitives->Primitives.IsEmpty())
-	{
-		return nullptr;
-	}
-
-	if (!ensure(FoundPrimitives->Primitives.Num() <= PCGComputeConstants::MAX_PRIMITIVE_COMPONENTS_PER_SPAWNER))
-	{
-		// Last resort - should be clamped earlier during setup and not come to this.
-		FoundPrimitives->Primitives.SetNum(PCGComputeConstants::MAX_PRIMITIVE_COMPONENTS_PER_SPAWNER);
-	}
-
 	TObjectPtr<UPCGInstanceDataProvider> DataProvider = NewObject<UPCGInstanceDataProvider>();
-	
-	DataProvider->Primitives = FoundPrimitives->Primitives;
-	DataProvider->NumCustomFloatsPerInstance = FoundPrimitives->NumCustomFloats;
-	
+
 	check(ProducerSettings);
 	const FPCGDataCollectionDesc InputDataDesc = ProducerSettings->ComputeInputPinDataDesc(InputPinProvidingData, Binding);
 	DataProvider->NumInstancesAllPrimitives = InputDataDesc.ComputeDataElementCount(EPCGDataType::Point);
-	if (FoundPrimitives->SelectorAttributeId != -1)
+
+	FPCGSpawnerPrimitives* FoundPrimitives = Binding->MeshSpawnersToPrimitives.Find(ProducerSettings);
+
+	// If there were 0 input points for this execution, we will not have created any primitives, so check for null.
+	if (FoundPrimitives)
 	{
-		// When selecting primitives dynamically (by attribute), we don't know statically how many instances will end up in each primitive, so run worst case.
-		DataProvider->NumInstancesAllPrimitives *= FoundPrimitives->Primitives.Num();
+		if (!ensure(FoundPrimitives->Primitives.Num() <= PCGComputeConstants::MAX_PRIMITIVE_COMPONENTS_PER_SPAWNER))
+		{
+			// Last resort - should be clamped earlier during setup and not come to this.
+			FoundPrimitives->Primitives.SetNum(PCGComputeConstants::MAX_PRIMITIVE_COMPONENTS_PER_SPAWNER);
+		}
+
+		DataProvider->Primitives = FoundPrimitives->Primitives;
+		DataProvider->NumCustomFloatsPerInstance = FoundPrimitives->NumCustomFloats;
+
+		if (FoundPrimitives->SelectorAttributeId != -1)
+		{
+			// When selecting primitives dynamically (by attribute), we don't know statically how many instances will end up in each primitive, so run worst case.
+			DataProvider->NumInstancesAllPrimitives *= FoundPrimitives->Primitives.Num();
+		}
+	}
+	else
+	{
+		ensure(DataProvider->NumInstancesAllPrimitives == 0);
 	}
 
 	return DataProvider;
@@ -171,31 +176,13 @@ FPCGInstanceDataProviderProxy::FPCGInstanceDataProviderProxy(TArray<FPrimitiveSc
 
 bool FPCGInstanceDataProviderProxy::IsValid(FValidationData const& InValidationData) const
 {
-	if (InValidationData.ParameterStructSize != sizeof(FParameters))
-	{
-		return false;
-	}
-
-	if (PrimitiveSceneProxies.IsEmpty())
-	{
-		UE_LOG(LogPCG, Warning, TEXT("Proxy invalid due to missing primitive scene proxy."));
-		return false;
-	}
-
-	if (NumInstancesAllPrimitives == 0)
-	{
-		UE_LOG(LogPCG, Warning, TEXT("Proxy invalid due to 0 instance count."));
-		return false;
-	}
-
-	return true;
+	return InValidationData.ParameterStructSize == sizeof(FParameters);
 }
 
 void FPCGInstanceDataProviderProxy::AllocateResources(FRDGBuilder& GraphBuilder, FAllocationData const& InAllocationData)
 {
-	ensure(NumInstancesAllPrimitives > 0);
 	const uint32 StrideUint4s = 3;
-	InstanceData = GraphBuilder.CreateBuffer(FRDGBufferDesc::CreateStructuredDesc(sizeof(FUintVector4) * StrideUint4s, NumInstancesAllPrimitives), TEXT("PCGInstanceDataBuffer"));
+	InstanceData = GraphBuilder.CreateBuffer(FRDGBufferDesc::CreateStructuredDesc(sizeof(FUintVector4) * StrideUint4s, FMath::Max(1u, NumInstancesAllPrimitives)), TEXT("PCGInstanceDataBuffer"));
 	InstanceDataSRV = GraphBuilder.CreateSRV(InstanceData);
 	InstanceDataUAV = GraphBuilder.CreateUAV(InstanceData);
 
@@ -204,11 +191,12 @@ void FPCGInstanceDataProviderProxy::AllocateResources(FRDGBuilder& GraphBuilder,
 	InstanceCustomFloatDataSRV = GraphBuilder.CreateSRV(InstanceCustomFloatData);
 	InstanceCustomFloatDataUAV = GraphBuilder.CreateUAV(InstanceCustomFloatData);
 
-	WriteCounters = GraphBuilder.CreateBuffer(FRDGBufferDesc::CreateStructuredDesc(sizeof(uint32), PrimitiveSceneProxies.Num()), TEXT("PCGWriteCounters"));
+	const int32 NumCountersRequired = FMath::Max(1, PrimitiveSceneProxies.Num());
+	WriteCounters = GraphBuilder.CreateBuffer(FRDGBufferDesc::CreateStructuredDesc(sizeof(uint32), NumCountersRequired), TEXT("PCGWriteCounters"));
 	WriteCountersSRV = GraphBuilder.CreateSRV(WriteCounters);
 	WriteCountersUAV = GraphBuilder.CreateUAV(WriteCounters);
 	TArray<uint32> Zeros;
-	Zeros.SetNumZeroed(PrimitiveSceneProxies.Num());
+	Zeros.SetNumZeroed(NumCountersRequired);
 	GraphBuilder.QueueBufferUpload(WriteCounters, Zeros.GetData(), Zeros.Num() * Zeros.GetTypeSize());
 }
 
diff --git a/Engine/Plugins/PCG/Source/PCG/Private/Compute/DataInterfaces/PCGInstanceDataInterface.h b/Engine/Plugins/PCG/Source/PCG/Private/Compute/DataInterfaces/PCGInstanceDataInterface.h
--- a/Engine/Plugins/PCG/Source/PCG/Private/Compute/DataInterfaces/PCGInstanceDataInterface.h
+++ b/Engine/Plugins/PCG/Source/PCG/Private/Compute/DataInterfaces/PCGInstanceDataInterface.h
@@ -58,7 +58,10 @@ public:
 	UPROPERTY()
 	TArray<TObjectPtr<UPrimitiveComponent>> Primitives;
 
+	UPROPERTY()
 	uint32 NumInstancesAllPrimitives = 0;
+
+	UPROPERTY()
 	uint32 NumCustomFloatsPerInstance = 0;
 };
 
diff --git a/Engine/Plugins/PCG/Source/PCG/Private/Compute/Elements/PCGComputeGraphElement.cpp b/Engine/Plugins/PCG/Source/PCG/Private/Compute/Elements/PCGComputeGraphElement.cpp
--- a/Engine/Plugins/PCG/Source/PCG/Private/Compute/Elements/PCGComputeGraphElement.cpp
+++ b/Engine/Plugins/PCG/Source/PCG/Private/Compute/Elements/PCGComputeGraphElement.cpp
@@ -265,6 +265,11 @@ bool FPCGComputeGraphElement::SetupProceduralISMComponents(FPCGContext* InContex
 		const FPCGDataCollectionDesc InputDataDesc = SpawnerSettings->ComputeInputPinDataDesc(PCGPinConstants::DefaultInputLabel, InBinding);
 
 		const uint32 InputPointCount = InputDataDesc.ComputeDataElementCount(EPCGDataType::Point);
+		if (InputPointCount == 0)
+		{
+			continue;
+		}
+
 		const FBox LocalBounds = InContext->SourceComponent->GetGridBounds().ShiftBy(-InContext->SourceComponent->GetOwner()->GetActorLocation());
 
 		uint32 CustomFloatCount = 0;
