Tentative optimization for cull outside bounds when order is not important

diff --git a/Engine/Plugins/PCG/Source/PCG/Private/Elements/PCGCullPointsOutsideActorBounds.cpp b/Engine/Plugins/PCG/Source/PCG/Private/Elements/PCGCullPointsOutsideActorBounds.cpp
--- a/Engine/Plugins/PCG/Source/PCG/Private/Elements/PCGCullPointsOutsideActorBounds.cpp
+++ b/Engine/Plugins/PCG/Source/PCG/Private/Elements/PCGCullPointsOutsideActorBounds.cpp
@@ -47,15 +47,20 @@ bool FPCGCullPointsOutsideActorBoundsElement::ExecuteInternal(FPCGContext* Conte
 		return true;
 	}
 
-	UPCGVolumeData* VolumeData = FPCGContext::NewObject_AnyThread<UPCGVolumeData>(Context);
-	check(VolumeData);
-
 	const UPCGCullPointsOutsideActorBoundsSettings* Settings = Context->GetInputSettings<UPCGCullPointsOutsideActorBoundsSettings>();
 	check(Settings);
 
 	// Initialize directly. Could also have gone through PCGComponent::CreateActorPCGDataCollection.
-	VolumeData->Initialize(Context->SourceComponent->GetGridBounds().ExpandBy(Settings->BoundsExpansion));
+	const FBox BoundsBox = Context->SourceComponent->GetGridBounds().ExpandBy(Settings->BoundsExpansion);
 
+	UPCGVolumeData* VolumeData = nullptr;
+	if (Settings->Mode == EPCGCullPointsMode::Ordered)
+	{
+		VolumeData = FPCGContext::NewObject_AnyThread<UPCGVolumeData>(Context);
+		check(VolumeData);
+		VolumeData->Initialize(BoundsBox);
+	}
+	
 	const TArray<FPCGTaggedData>& Inputs = Context->InputData.TaggedData;
 	for (const FPCGTaggedData& Input : Inputs)
 	{
@@ -67,20 +72,55 @@ bool FPCGCullPointsOutsideActorBoundsElement::ExecuteInternal(FPCGContext* Conte
 			continue;
 		}
 
-		UPCGIntersectionData* Intersection = InputPointData->IntersectWith(Context, VolumeData);
-		check(Intersection);
+		const UPCGPointData* OutputData = nullptr;
+
+		if (Settings->Mode == EPCGCullPointsMode::Ordered)
+		{
+			check(VolumeData);
+			UPCGIntersectionData* Intersection = InputPointData->IntersectWith(Context, VolumeData);
+			check(Intersection);
+
+			const UPCGPointData* IntersectedPoints = Intersection->ToPointData(Context);
+			check(IntersectedPoints);
 
-		const UPCGPointData* IntersectedPoints = Intersection->ToPointData(Context);
-		check(IntersectedPoints);
+			// Skip empty result
+			if (IntersectedPoints->GetPoints().IsEmpty())
+			{
+				continue;
+			}
 
-		// Skip empty result
-		if (IntersectedPoints->GetPoints().IsEmpty())
+			OutputData = IntersectedPoints;
+		}
+		else if (Settings->Mode == EPCGCullPointsMode::Unordered)
 		{
-			continue;
+			UPCGPointData* CulledPointsData = FPCGContext::NewObject_AnyThread<UPCGPointData>(Context);
+			CulledPointsData->InitializeFromData(InputPointData);
+			TArray<FPCGPoint>& CulledPoints = CulledPointsData->GetMutablePoints();
+
+			// TODO: split up this query across mutiple threads, but will make gathering the point a bit harder
+			const UPCGPointData::PointOctree& InputOctree = InputPointData->GetOctree();
+			FBoxCenterAndExtent BoxCenterAndExtents(BoundsBox);
+
+			InputOctree.FindElementsWithBoundsTest(BoxCenterAndExtents, [&CulledPoints](const FPCGPointRef& PointRef)
+			{
+				CulledPoints.Add(*PointRef.Point);
+			});
+
+			if (!CulledPoints.IsEmpty())
+			{
+				OutputData = CulledPointsData;
+			}
+		}
+		else
+		{
+			PCGLog::LogErrorOnGraph(LOCTEXT("InvalidMode", "Invalid culling mode."), Context);
 		}
 
-		FPCGTaggedData& Output = Context->OutputData.TaggedData.Emplace_GetRef(Input);
-		Output.Data = IntersectedPoints;
+		if (OutputData)
+		{
+			FPCGTaggedData& Output = Context->OutputData.TaggedData.Emplace_GetRef(Input);
+			Output.Data = OutputData;
+		}
 	}
 
 	return true;
diff --git a/Engine/Plugins/PCG/Source/PCG/Public/Elements/PCGCullPointsOutsideActorBounds.h b/Engine/Plugins/PCG/Source/PCG/Public/Elements/PCGCullPointsOutsideActorBounds.h
--- a/Engine/Plugins/PCG/Source/PCG/Public/Elements/PCGCullPointsOutsideActorBounds.h
+++ b/Engine/Plugins/PCG/Source/PCG/Public/Elements/PCGCullPointsOutsideActorBounds.h
@@ -6,6 +6,13 @@
 
 #include "PCGCullPointsOutsideActorBounds.generated.h"
 
+UENUM(BlueprintType)
+enum class EPCGCullPointsMode : uint8
+{
+	Ordered = 0,
+	Unordered
+};
+
 /**
  * Removes points that lie outside the current actor bounds.
  */
@@ -33,6 +40,9 @@ protected:
 public:
 	UPROPERTY(BlueprintReadWrite, EditAnywhere, Category = Settings, meta = (PCG_Overridable))
 	float BoundsExpansion = 0.0;
+
+	UPROPERTY(BlueprintReadWrite, EditAnywhere, Category = Settings, meta = (PCG_Overridable))
+	EPCGCullPointsMode Mode = EPCGCullPointsMode::Ordered;
 };
 
 class FPCGCullPointsOutsideActorBoundsElement : public IPCGElement
