[PCG] GPU - fix GPU Copy Points & GPU Static Mesh Spawner combination failing due to incorrect usage of string attributes.

#rnx

diff --git a/Engine/Plugins/PCG/Shaders/Private/Elements/PCGCopyPoints.usf b/Engine/Plugins/PCG/Shaders/Private/Elements/PCGCopyPoints.usf
--- a/Engine/Plugins/PCG/Shaders/Private/Elements/PCGCopyPoints.usf
+++ b/Engine/Plugins/PCG/Shaders/Private/Elements/PCGCopyPoints.usf
@@ -23,63 +23,89 @@
 // TODO pass in attribute IDs to copy.
 void CopyAttributesFromSource(uint Source_DataIndex, uint Source_ElementIndex, uint Out_DataIndex, uint Out_ElementIndex)
 {
-	const uint NumAttributes = Source_GetDataNumAttributesInternal(Source_DataIndex);
-	uint NumAttributesProcessed = 0;
-	
-	// Loop over all attribute headers, if the address is non-zero, then copy it from pin Source to pin Out.
-	for (int AttributeIndex = NUM_RESERVED_ATTRS; AttributeIndex < MAX_NUM_ATTRS; ++AttributeIndex)
-	{
-		const uint Source_Stride = Source_GetAttributeStrideInternal(Source_DataIndex, AttributeIndex);
-		const uint Out_Stride = Target_GetAttributeStrideInternal(Out_DataIndex, AttributeIndex);
-
-		if (Source_Stride != 0 && Source_Stride == Out_Stride)
-		{
-			const uint Source_ElementAddress = Source_GetElementAddressInternal(Source_DataIndex, Source_ElementIndex, AttributeIndex);
-			const uint Out_ElementAddress = Out_GetElementAddressInternal(Out_DataIndex, Out_ElementIndex, AttributeIndex);
-
-			for (int Offset = 0; Offset < Out_Stride; Offset += 4)
-			{
-				Out_StoreBufferInternal(Out_ElementAddress + Offset, Source_LoadBufferInternal(Source_ElementAddress + Offset));
-			}
-
-			// We can early-out when we've looked at all the possible attributes
-			if (++NumAttributesProcessed >= NumAttributes)
-			{
-				break;
-			}
-		}
-	}
+    const uint NumAttributes = Out_GetDataNumAttributesInternal(Out_DataIndex);
+    uint NumAttributesProcessed = 0;
+
+    for (int AttributeId = NUM_RESERVED_ATTRS; AttributeId < MAX_NUM_ATTRS; ++AttributeId)
+    {
+        const uint Out_Stride = Out_GetAttributeStrideInternal(Out_DataIndex, AttributeId);
+
+        if (Out_Stride == 0)
+        {
+            // No output attribute to write to.
+            continue;
+        }
+
+        const uint Source_Stride = Source_GetAttributeStrideInternal(Source_DataIndex, AttributeId);
+        const uint Out_ElementAddress = Out_GetElementAddressInternal(Out_DataIndex, Out_ElementIndex, AttributeId);
+
+        if (Source_Stride == Out_Stride)
+        {
+            const uint Source_ElementAddress = Source_GetElementAddressInternal(Source_DataIndex, Source_ElementIndex, AttributeId);
+
+            for (uint Offset = 0; Offset < Out_Stride; Offset += 4)
+            {
+                Out_StoreBufferInternal(Out_ElementAddress + Offset, Source_LoadBufferInternal(Source_ElementAddress + Offset));
+            }
+        }
+        else
+        {
+            for (uint Offset = 0; Offset < Out_Stride; Offset += 4)
+            {
+				// Initialize output data (only header part of buffer will be initialized).
+                Out_StoreBufferInternal(Out_ElementAddress + Offset, 0u);
+            }
+        }
+
+        if (++NumAttributesProcessed >= NumAttributes)
+        {
+            break; // We can early-out when we've looked at all the possible attributes
+        }
+    }
 }
 
 // TODO pass in attribute IDs to copy.
 void CopyAttributesFromTarget(uint Target_DataIndex, uint Target_ElementIndex, uint Out_DataIndex, uint Out_ElementIndex)
 {
-	const uint NumAttributes = Target_GetDataNumAttributesInternal(Target_DataIndex);
-	uint NumAttributesProcessed = 0;
-	
-	// Loop over all attribute headers, if the address is non-zero, then copy it from pin Target to pin Out.
-	for (int AttributeIndex = NUM_RESERVED_ATTRS; AttributeIndex < MAX_NUM_ATTRS; ++AttributeIndex)
-	{
-		const uint Target_Stride = Target_GetAttributeStrideInternal(Target_DataIndex, AttributeIndex);
-		const uint Out_Stride = Out_GetAttributeStrideInternal(Out_DataIndex, AttributeIndex);
-
-		if (Target_Stride != 0 && Target_Stride == Out_Stride)
-		{
-			const uint Target_ElementAddress = Target_GetElementAddressInternal(Target_DataIndex, Target_ElementIndex, AttributeIndex);
-			const uint Out_ElementAddress = Out_GetElementAddressInternal(Out_DataIndex, Out_ElementIndex, AttributeIndex);
-
-			for (int Offset = 0; Offset < Out_Stride; Offset += 4)
-			{
-				Out_StoreBufferInternal(Out_ElementAddress + Offset, Target_LoadBufferInternal(Target_ElementAddress + Offset));
-			}
-
-			// We can early-out when we've looked at all the possible attributes
-			if (++NumAttributesProcessed >= NumAttributes)
-			{
-				break;
-			}
-		}
-	}
+    const uint NumAttributes = Out_GetDataNumAttributesInternal(Out_DataIndex);
+    uint NumAttributesProcessed = 0;
+
+    for (int AttributeId = NUM_RESERVED_ATTRS; AttributeId < MAX_NUM_ATTRS; ++AttributeId)
+    {
+        const uint Out_Stride = Out_GetAttributeStrideInternal(Out_DataIndex, AttributeId);
+
+        if (Out_Stride == 0)
+        {
+            // No output attribute to write to.
+            continue;
+        }
+
+        const uint Target_Stride = Target_GetAttributeStrideInternal(Target_DataIndex, AttributeId);
+        const uint Out_ElementAddress = Out_GetElementAddressInternal(Out_DataIndex, Out_ElementIndex, AttributeId);
+
+        if (Target_Stride == Out_Stride)
+        {
+            const uint Target_ElementAddress = Target_GetElementAddressInternal(Target_DataIndex, Target_ElementIndex, AttributeId);
+
+            for (uint Offset = 0; Offset < Out_Stride; Offset += 4)
+            {
+                Out_StoreBufferInternal(Out_ElementAddress + Offset, Target_LoadBufferInternal(Target_ElementAddress + Offset));
+            }
+        }
+        else
+        {
+            for (uint Offset = 0; Offset < Out_Stride; Offset += 4)
+            {
+				// Initialize output data (only header part of buffer will be initialized).
+                Out_StoreBufferInternal(Out_ElementAddress + Offset, 0u);
+            }
+        }
+
+        if (++NumAttributesProcessed >= NumAttributes)
+        {
+            break; // We can early-out when we've looked at all the possible attributes
+        }
+    }
 }
 
 [numthreads(64, 1, 1)]
diff --git a/Engine/Plugins/PCG/Source/PCG/Private/Compute/Elements/PCGCustomHLSL.cpp b/Engine/Plugins/PCG/Source/PCG/Private/Compute/Elements/PCGCustomHLSL.cpp
--- a/Engine/Plugins/PCG/Source/PCG/Private/Compute/Elements/PCGCustomHLSL.cpp
+++ b/Engine/Plugins/PCG/Source/PCG/Private/Compute/Elements/PCGCustomHLSL.cpp
@@ -1656,24 +1656,44 @@ FString UPCGCustomHLSLSettings::GetCookedKernelSource(const TMap<FName, FPCGKern
 				"\n"
 				"    // Initialize output data elements. Loop over all attribute headers, if the address is non-zero, then copy it from pin {0} to pin {1}.\n"
 				"    {\n"
-				"        const uint NumAttributes = {0}_GetDataNumAttributesInternal({0}_DataIndex);\n"
+				"        const uint NumAttributes = {1}_GetDataNumAttributesInternal({1}_DataIndex);\n"
 				"        uint NumAttributesProcessed = 0;\n"
 				"\n"
-				"        for (int AttributeIndex = 0; AttributeIndex < 128; ++AttributeIndex)\n"
+				"        for (int AttributeId = 0; AttributeId < 128; ++AttributeId)\n"
 				"        {\n"
-				"            const uint Stride = {0}_GetAttributeStrideInternal({0}_DataIndex, AttributeIndex);\n"
+				"            const uint {1}_Stride = {1}_GetAttributeStrideInternal({1}_DataIndex, AttributeId);\n"
 				"\n"
-				"            if (Stride != 0)\n"
+				"            if ({1}_Stride == 0)\n"
 				"            {\n"
-				"                const uint {0}_ElementAddress = {0}_GetElementAddressInternal({0}_DataIndex, ElementIndex, AttributeIndex);\n"
-				"                const uint {1}_ElementAddress = {1}_GetElementAddressInternal({1}_DataIndex, ElementIndex, AttributeIndex);\n"
+				"                // No output attribute to write to.\n"
+				"                continue;\n"
+				"            }\n"
+				"\n"
+				"            const uint {0}_Stride = {0}_GetAttributeStrideInternal({0}_DataIndex, AttributeId);\n"
+				"            const uint {1}_ElementAddress = {1}_GetElementAddressInternal({1}_DataIndex, ElementIndex, AttributeId);\n"
+				"\n"
+				"            if ({0}_Stride == {1}_Stride)\n"
+				"            {\n"
+				"                const uint {0}_ElementAddress = {0}_GetElementAddressInternal({0}_DataIndex, ElementIndex, AttributeId);\n"
 				"\n"
-				"                for (uint Address = 0; Address < Stride; Address += 4)\n"
+				"                for (uint Offset = 0; Offset < {1}_Stride; Offset += 4)\n"
 				"                {\n"
-				"                    {1}_StoreBufferInternal({1}_ElementAddress + Address, {0}_LoadBufferInternal({0}_ElementAddress + Address));\n"
+				"                    {1}_StoreBufferInternal({1}_ElementAddress + Offset, {0}_LoadBufferInternal({0}_ElementAddress + Offset));\n"
 				"                }\n"
 				"\n"
-				"                if (++NumAttributesProcessed >= NumAttributes) break; // We can early-out when we've looked at all the possible attributes\n"
+				"            }\n"
+				"            else\n"
+				"            {\n"
+				"                for (uint Offset = 0; Offset < {1}_Stride; Offset += 4)\n"
+				"                {\n"
+				"                    // Initialize output data (only header part of buffer will be initialized).\n"
+				"                    {1}_StoreBufferInternal({1}_ElementAddress + Offset, 0u);\n"
+				"                }\n"
+				"            }\n"
+				"\n"
+				"            if (++NumAttributesProcessed >= NumAttributes)\n"
+				"            {\n"
+				"                break; // We can early-out when we've looked at all the possible attributes\n"
 				"            }\n"
 				"        }\n"
 				"    }\n"),
diff --git a/Engine/Plugins/PCG/Source/PCG/Private/Compute/PCGDataBinding.cpp b/Engine/Plugins/PCG/Source/PCG/Private/Compute/PCGDataBinding.cpp
--- a/Engine/Plugins/PCG/Source/PCG/Private/Compute/PCGDataBinding.cpp
+++ b/Engine/Plugins/PCG/Source/PCG/Private/Compute/PCGDataBinding.cpp
@@ -105,7 +105,7 @@ void UPCGDataBinding::BuildStringTable()
 				{
 					for (const FPCGPoint& Point : PointData->GetPoints())
 					{
-						AddEntryToStringTable(Point.MetadataEntry);
+						AddEntryToStringTable(AttributeBase->GetValueKey(Point.MetadataEntry));
 					}
 				}
 				else if (const UPCGParamData* ParamData = Cast<UPCGParamData>(Data.Data))
@@ -114,7 +114,7 @@ void UPCGDataBinding::BuildStringTable()
 
 					for (int64 MetadataKey = 0; MetadataKey < NumElements; ++MetadataKey)
 					{
-						AddEntryToStringTable(MetadataKey);
+						AddEntryToStringTable(AttributeBase->GetValueKey(MetadataKey));
 					}
 				}
 				else { /** TODO: More attribute types! */ }
diff --git a/Engine/Plugins/PCG/Source/PCG/Private/Compute/PCGDataForGPU.cpp b/Engine/Plugins/PCG/Source/PCG/Private/Compute/PCGDataForGPU.cpp
--- a/Engine/Plugins/PCG/Source/PCG/Private/Compute/PCGDataForGPU.cpp
+++ b/Engine/Plugins/PCG/Source/PCG/Private/Compute/PCGDataForGPU.cpp
@@ -658,7 +658,7 @@ void FPCGDataDesc::InitializeAttributeDescs(const UPCGData* InData, const TMap<F
 		for (int DelayedAttributeIndex = 0; DelayedAttributeIndex < DelayedAttributeKeysAndStringKeys.Num(); ++DelayedAttributeIndex)
 		{
 			const FPCGKernelAttributeKey& AttributeKey = DelayedAttributeKeysAndStringKeys[DelayedAttributeIndex].Get<0>();
-			AttributeDescs.Emplace(NUM_RESERVED_ATTRS + DelayedAttributeIndex + NumAttributesFromLUT, AttributeKey.Type, AttributeKey.Name, MoveTemp(DelayedAttributeKeysAndStringKeys[DelayedAttributeIndex].Get<1>()));
+			AttributeDescs.Emplace(NUM_RESERVED_ATTRS + DelayedAttributeIndex + NumAttributesFromLUT + InGlobalAttributeLookupTable.Num(), AttributeKey.Type, AttributeKey.Name, MoveTemp(DelayedAttributeKeysAndStringKeys[DelayedAttributeIndex].Get<1>()));
 		}
 	}
 }
