[PCG] Fix failed check when going from V2 to V1 execution path

#rnx

diff --git a/Engine/Plugins/PCG/Source/PCG/Private/Graph/PCGGraphExecutor.cpp b/Engine/Plugins/PCG/Source/PCG/Private/Graph/PCGGraphExecutor.cpp
--- a/Engine/Plugins/PCG/Source/PCG/Private/Graph/PCGGraphExecutor.cpp
+++ b/Engine/Plugins/PCG/Source/PCG/Private/Graph/PCGGraphExecutor.cpp
@@ -1381,6 +1381,10 @@ void FPCGGraphExecutor::Execute()
 	if (GetNonScheduledRemainingTaskCount() == 0)
 	{
 		ExecuteVersion = EExecuteVersion::None;
+
+		// TODO: Remove when removing V1 code path
+		// Reset as V2 doesn't properly maintain that value (doesn't need to) and it will cause a failed check when going from V2 to V1
+		CurrentlyUsedThreads = 0;
 	}
 }
 
