[PCG] Fix broken change propagation for changes to GPU nodes.

Node changes often affect compute graph and a recompilation is required. Without this fix you can disable a GPU node and it won't update the compute graph and the node will still execute.

Also fixes ensure when incorrectly computing GPU data descriptions across nodes that won't execute on GPU - such as disabled nodes.

#rnx

diff --git a/Engine/Plugins/PCG/Source/PCG/Private/Compute/Elements/PCGCustomHLSL.cpp b/Engine/Plugins/PCG/Source/PCG/Private/Compute/Elements/PCGCustomHLSL.cpp
--- a/Engine/Plugins/PCG/Source/PCG/Private/Compute/Elements/PCGCustomHLSL.cpp
+++ b/Engine/Plugins/PCG/Source/PCG/Private/Compute/Elements/PCGCustomHLSL.cpp
@@ -283,6 +283,13 @@ FPCGDataCollectionDesc UPCGCustomHLSLSettings::ComputeOutputPinDataDesc(const UP
 {
 	check(OutputPin);
 	check(Binding);
+
+	// If node will not execute on GPU then its data comes straight from CPU.
+	if (!ShouldExecuteOnGPU() || !bEnabled)
+	{
+		return Super::ComputeOutputPinDataDesc(OutputPin, Binding);
+	}
+
 	FPCGDataCollectionDesc PinDesc;
 
 	const FPCGPinPropertiesGPU* Properties = GetOutputPinPropertiesGPU(OutputPin->Properties.Label);
diff --git a/Engine/Plugins/PCG/Source/PCG/Private/Elements/PCGCopyPoints.cpp b/Engine/Plugins/PCG/Source/PCG/Private/Elements/PCGCopyPoints.cpp
--- a/Engine/Plugins/PCG/Source/PCG/Private/Elements/PCGCopyPoints.cpp
+++ b/Engine/Plugins/PCG/Source/PCG/Private/Elements/PCGCopyPoints.cpp
@@ -34,13 +34,15 @@ FString UPCGCopyPointsSettings::GetCookedKernelSource(const TMap<FName, FPCGKern
 
 FPCGDataCollectionDesc UPCGCopyPointsSettings::ComputeOutputPinDataDesc(const UPCGPin* OutputPin, const UPCGDataBinding* Binding) const
 {
-	if (!ShouldExecuteOnGPU())
+	check(OutputPin);
+	check(Binding);
+
+	// If node will not execute on GPU then its data comes straight from CPU.
+	if (!ShouldExecuteOnGPU() || !bEnabled)
 	{
 		return Super::ComputeOutputPinDataDesc(OutputPin, Binding);
 	}
 
-	check(OutputPin);
-
 	const UPCGNode* Node = CastChecked<UPCGNode>(GetOuter());
 	const UPCGPin* SourcePin = Node->GetInputPin(PCGCopyPointsConstants::SourcePointsLabel);
 	const UPCGPin* TargetPin = Node->GetInputPin(PCGCopyPointsConstants::TargetPointsLabel);
diff --git a/Engine/Plugins/PCG/Source/PCG/Private/Elements/PCGStaticMeshSpawner.cpp b/Engine/Plugins/PCG/Source/PCG/Private/Elements/PCGStaticMeshSpawner.cpp
--- a/Engine/Plugins/PCG/Source/PCG/Private/Elements/PCGStaticMeshSpawner.cpp
+++ b/Engine/Plugins/PCG/Source/PCG/Private/Elements/PCGStaticMeshSpawner.cpp
@@ -116,12 +116,15 @@ FString UPCGStaticMeshSpawnerSettings::GetCookedKernelSource(const TMap<FName, F
 
 FPCGDataCollectionDesc UPCGStaticMeshSpawnerSettings::ComputeOutputPinDataDesc(const UPCGPin* OutputPin, const UPCGDataBinding* Binding) const
 {
-	if (!ShouldExecuteOnGPU())
+	check(OutputPin);
+	check(Binding);
+
+	// If node will not execute on GPU then its data comes straight from CPU.
+	if (!ShouldExecuteOnGPU() || !bEnabled)
 	{
 		return Super::ComputeOutputPinDataDesc(OutputPin, Binding);
 	}
 
-	check(OutputPin);
 
 	// Passthrough - output description matches input description.
 	return ComputeInputPinDataDesc(PCGPinConstants::DefaultInputLabel, Binding);
@@ -192,19 +195,6 @@ void UPCGStaticMeshSpawnerSettings::CreateAdditionalOutputDataInterfaces(TArray<
 }
 
 #if WITH_EDITOR
-EPCGChangeType UPCGStaticMeshSpawnerSettings::GetChangeTypeForProperty(const FName& InPropertyName) const
-{
-	EPCGChangeType ChangeType = Super::GetChangeTypeForProperty(InPropertyName);
-
-	if (ChangeType != EPCGChangeType::Cosmetic && ShouldExecuteOnGPU())
-	{
-		// GPU SM Spawners are compiled into compute graph so trigger recompilation.
-		ChangeType |= EPCGChangeType::Structural;
-	}
-
-	return ChangeType;
-}
-
 FText UPCGStaticMeshSpawnerSettings::GetDefaultNodeTitle() const
 {
 	return LOCTEXT("NodeTitle", "Static Mesh Spawner");
diff --git a/Engine/Plugins/PCG/Source/PCG/Private/PCGSettings.cpp b/Engine/Plugins/PCG/Source/PCG/Private/PCGSettings.cpp
--- a/Engine/Plugins/PCG/Source/PCG/Private/PCGSettings.cpp
+++ b/Engine/Plugins/PCG/Source/PCG/Private/PCGSettings.cpp
@@ -712,6 +712,12 @@ PRAGMA_DISABLE_DEPRECATION_WARNINGS
 	ChangeType |= IsStructuralProperty(InPropertyName) ? EPCGChangeType::Structural : EPCGChangeType::None;
 PRAGMA_ENABLE_DEPRECATION_WARNINGS
 
+	// Changes to GPU nodes can deeply affect the compiled compute graph so trigger recompiles on any non-trivial changes.
+	if (ChangeType != EPCGChangeType::Cosmetic && ShouldExecuteOnGPU())
+	{
+		ChangeType |= EPCGChangeType::Structural;
+	}
+
 	return ChangeType;
 }
 
diff --git a/Engine/Plugins/PCG/Source/PCG/Public/Elements/PCGStaticMeshSpawner.h b/Engine/Plugins/PCG/Source/PCG/Public/Elements/PCGStaticMeshSpawner.h
--- a/Engine/Plugins/PCG/Source/PCG/Public/Elements/PCGStaticMeshSpawner.h
+++ b/Engine/Plugins/PCG/Source/PCG/Public/Elements/PCGStaticMeshSpawner.h
@@ -37,7 +37,6 @@ public:
 	virtual void CreateAdditionalInputDataInterfaces(TArray<TObjectPtr<UComputeDataInterface>>& OutDataInterfaces) const override;
 	virtual void CreateAdditionalOutputDataInterfaces(TArray<TObjectPtr<UComputeDataInterface>>& OutDataInterfaces) const override;
 #if WITH_EDITOR
-	virtual EPCGChangeType GetChangeTypeForProperty(const FName& InPropertyName) const override;
 	virtual bool DisplayExecuteOnGPUSetting() const override { return true; }
 	virtual FName GetDefaultNodeName() const override { return FName(TEXT("StaticMeshSpawner")); }
 	virtual FText GetDefaultNodeTitle() const override;
