[PCG] Support output selected mesh attribute from GPU SM Spawner node.

Straightforward in theory but when using weighted selection, the mesh path strings are not already in the data - the SM Spawner needs to create strings for the meshes in the node settings. To this end, nodes are given the chance to add strings to the string table during compilation, so that they then can be referenced using string keys at execution time.

#rnx

diff --git a/Engine/Plugins/PCG/Shaders/Private/Elements/PCGStaticMeshSpawner.usf b/Engine/Plugins/PCG/Shaders/Private/Elements/PCGStaticMeshSpawner.usf
--- a/Engine/Plugins/PCG/Shaders/Private/Elements/PCGStaticMeshSpawner.usf
+++ b/Engine/Plugins/PCG/Shaders/Private/Elements/PCGStaticMeshSpawner.usf
@@ -45,6 +45,7 @@ void Main(uint3 GroupId : SV_GroupID, uint GroupIndex : SV_GroupIndex)
 	const int NumPrimitives = InstanceData_GetNumPrimitives();
 
 	uint PrimitiveIndex = 0;
+	uint OutputMeshStringKey = (uint)-1;
 
 	if (NumPrimitives > 1)
 	{
@@ -53,8 +54,8 @@ void Main(uint3 GroupId : SV_GroupID, uint GroupIndex : SV_GroupIndex)
 		if (SelectorAttributeId != (uint)-1)
 		{
 			// Select using string key
-			const int MeshPathStringKey = In_GetStringKey(In_DataIndex, In_ElementIndex, SelectorAttributeId);
-			PrimitiveIndex = SMSpawner_GetPrimitiveIndexFromStringKey(MeshPathStringKey);
+			OutputMeshStringKey = In_GetStringKey(In_DataIndex, In_ElementIndex, SelectorAttributeId);
+			PrimitiveIndex = SMSpawner_GetPrimitiveIndexFromStringKey(OutputMeshStringKey);
 		}
 		else
 		{
@@ -69,6 +70,8 @@ void Main(uint3 GroupId : SV_GroupID, uint GroupIndex : SV_GroupIndex)
 					break;
 				}
 			}
+
+			OutputMeshStringKey = SMSpawner_GetPrimitiveStringKey(PrimitiveIndex);
 		}
 	}
 
@@ -146,36 +149,54 @@ void Main(uint3 GroupId : SV_GroupID, uint GroupIndex : SV_GroupIndex)
 		Out_SetSteepness(Out_DataIndex, Out_ElementIndex, In_GetSteepness(In_DataIndex, In_ElementIndex));
 
 		CopyAttributesToOuput(In_DataIndex, In_ElementIndex, Out_DataIndex, Out_ElementIndex);
+
+		if (SMSpawner_GetSelectedMeshAttributeId() != (uint)-1)
+		{
+			Out_SetStringKey(Out_DataIndex, Out_ElementIndex, SMSpawner_GetSelectedMeshAttributeId(), OutputMeshStringKey);
+		}
 	}
 }
 
 // TODO pass in attribute IDs to copy.
 void CopyAttributesToOuput(uint In_DataIndex, uint In_ElementIndex, uint Out_DataIndex, uint Out_ElementIndex)
 {
-	const uint NumAttributes = In_GetDataNumAttributesInternal(In_DataIndex);
-	uint NumAttributesProcessed = 0;
-	
-	// Loop over all attribute headers, if the address is non-zero, then copy it from pin Source to pin Out.
-	for (int AttributeIndex = NUM_RESERVED_ATTRS; AttributeIndex < MAX_NUM_ATTRS; ++AttributeIndex)
-	{
-		const uint In_Stride = In_GetAttributeStrideInternal(In_DataIndex, AttributeIndex);
-		const uint Out_Stride = In_GetAttributeStrideInternal(In_DataIndex, AttributeIndex);
-
-		if (In_Stride != 0 && In_Stride == Out_Stride)
-		{
-			const uint In_ElementAddress = In_GetElementAddressInternal(In_DataIndex, In_ElementIndex, AttributeIndex);
-			const uint Out_ElementAddress = Out_GetElementAddressInternal(Out_DataIndex, Out_ElementIndex, AttributeIndex);
-
-			for (int Address = 0; Address < In_Stride; Address += 4)
-			{
-				Out_StoreBufferInternal(Out_ElementAddress + Address, In_LoadBufferInternal(In_ElementAddress + Address));
-			}
-
-			// We can early-out when we've looked at all the possible attributes
-			if (++NumAttributesProcessed >= NumAttributes)
-			{
-				break;
-			}
-		}
-	}
+    const uint NumAttributes = Out_GetDataNumAttributesInternal(Out_DataIndex);
+    uint NumAttributesProcessed = 0;
+
+    for (int AttributeId = NUM_RESERVED_ATTRS; AttributeId < MAX_NUM_ATTRS; ++AttributeId)
+    {
+        const uint Out_Stride = Out_GetAttributeStrideInternal(Out_DataIndex, AttributeId);
+
+        if (Out_Stride == 0)
+        {
+            // No output attribute to write to.
+            continue;
+        }
+
+        const uint In_Stride = In_GetAttributeStrideInternal(In_DataIndex, AttributeId);
+        const uint Out_ElementAddress = Out_GetElementAddressInternal(Out_DataIndex, Out_ElementIndex, AttributeId);
+
+        if (In_Stride == Out_Stride)
+        {
+            const uint In_ElementAddress = In_GetElementAddressInternal(In_DataIndex, In_ElementIndex, AttributeId);
+
+            for (uint Offset = 0; Offset < Out_Stride; Offset += 4)
+            {
+                Out_StoreBufferInternal(Out_ElementAddress + Offset, In_LoadBufferInternal(In_ElementAddress + Offset));
+            }
+        }
+        else
+        {
+            for (uint Offset = 0; Offset < Out_Stride; Offset += 4)
+            {
+				// Initialize output data (only header part of buffer will be initialized).
+                Out_StoreBufferInternal(Out_ElementAddress + Offset, 0u);
+            }
+        }
+
+        if (++NumAttributesProcessed >= NumAttributes)
+        {
+            break; // We can early-out when we've looked at all the possible attributes
+        }
+    }
 }
diff --git a/Engine/Plugins/PCG/Source/PCG/Private/Compute/DataInterfaces/Elements/PCGStaticMeshSpawnerDataInterface.cpp b/Engine/Plugins/PCG/Source/PCG/Private/Compute/DataInterfaces/Elements/PCGStaticMeshSpawnerDataInterface.cpp
--- a/Engine/Plugins/PCG/Source/PCG/Private/Compute/DataInterfaces/Elements/PCGStaticMeshSpawnerDataInterface.cpp
+++ b/Engine/Plugins/PCG/Source/PCG/Private/Compute/DataInterfaces/Elements/PCGStaticMeshSpawnerDataInterface.cpp
@@ -32,24 +32,34 @@ void UPCGStaticMeshSpawnerDataInterface::GetSupportedInputs(TArray<FShaderFuncti
 		.AddReturnType(EShaderFundamentalType::Uint, 4)
 		.AddParam(EShaderFundamentalType::Uint); // InAttributeIndex
 
+	OutFunctions.AddDefaulted_GetRef()
+		.SetName(TEXT("SMSpawner_GetPrimitiveStringKey"))
+		.AddReturnType(EShaderFundamentalType::Int) // String key
+		.AddParam(EShaderFundamentalType::Uint); // InPrimitiveIndex
+
 	OutFunctions.AddDefaulted_GetRef()
 		.SetName(TEXT("SMSpawner_GetPrimitiveIndexFromStringKey"))
 		.AddReturnType(EShaderFundamentalType::Uint) // Primitive index
-		.AddParam(EShaderFundamentalType::Uint); // InMeshPathStringKey
+		.AddParam(EShaderFundamentalType::Int); // InMeshPathStringKey
 
 	OutFunctions.AddDefaulted_GetRef()
 		.SetName(TEXT("SMSpawner_GetPrimitiveSelectionCDF"))
 		.AddReturnType(EShaderFundamentalType::Float) // CDF value
 		.AddParam(EShaderFundamentalType::Uint); // InPrimitiveIndex
+
+	OutFunctions.AddDefaulted_GetRef()
+		.SetName(TEXT("SMSpawner_GetSelectedMeshAttributeId"))
+		.AddReturnType(EShaderFundamentalType::Uint); // Attribute id to output mesh path string key to
 }
 
 BEGIN_SHADER_PARAMETER_STRUCT(FPCGStaticMeshSpawnerDataInterfaceParameters,)
 	SHADER_PARAMETER_ARRAY(FUintVector4, AttributeIdOffsetStrides, [UPCGStaticMeshSpawnerDataInterface::MAX_ATTRIBUTES])
-	SHADER_PARAMETER_ARRAY(FUintVector4, PrimitiveStringKeys, [PCGComputeConstants::MAX_PRIMITIVE_COMPONENTS_PER_SPAWNER])
+	SHADER_PARAMETER_ARRAY(FIntVector4, PrimitiveStringKeys, [PCGComputeConstants::MAX_PRIMITIVE_COMPONENTS_PER_SPAWNER])
 	SHADER_PARAMETER_SCALAR_ARRAY(float, SelectionCDF, [PCGComputeConstants::MAX_PRIMITIVE_COMPONENTS_PER_SPAWNER])
 	SHADER_PARAMETER(uint32, NumAttributes)
 	SHADER_PARAMETER(uint32, NumPrimitives)
 	SHADER_PARAMETER(uint32, SelectorAttributeId)
+	SHADER_PARAMETER(uint32, SelectedMeshAttributeId)
 END_SHADER_PARAMETER_STRUCT()
 
 void UPCGStaticMeshSpawnerDataInterface::GetShaderParameters(TCHAR const* UID, FShaderParametersMetadataBuilder& InOutBuilder, FShaderParametersMetadataAllocations& InOutAllocations) const
@@ -76,11 +86,16 @@ void UPCGStaticMeshSpawnerDataInterface::GetHLSL(FString& OutHLSL, FString const
 		"uint4 SMSpawner_GetAttributeIdOffsetStride_{DataInterfaceName}(uint InAttributeIndex) { return {DataInterfaceName}_AttributeIdOffsetStrides[InAttributeIndex]; }\n"
 		"\n"
 		"uint {DataInterfaceName}_NumPrimitives;\n"
-		"uint4 {DataInterfaceName}_PrimitiveStringKeys[{MaxPrimitives}];"
+		"int4 {DataInterfaceName}_PrimitiveStringKeys[{MaxPrimitives}];"
 		"DECLARE_SCALAR_ARRAY(float, {DataInterfaceName}_SelectionCDF, {MaxPrimitives});\n"
 		"uint SMSpawner_GetNumPrimitives_{DataInterfaceName}() { return {DataInterfaceName}_NumPrimitives; }\n"
 		"\n"
-		"uint SMSpawner_GetPrimitiveIndexFromStringKey_{DataInterfaceName}(uint InMeshPathStringKey)\n"
+		"int SMSpawner_GetPrimitiveStringKey_{DataInterfaceName}(uint InPrimitiveIndex)\n"
+		"{\n"
+		"	return {DataInterfaceName}_PrimitiveStringKeys[InPrimitiveIndex][0];\n"
+		"}\n"
+		"\n"
+		"uint SMSpawner_GetPrimitiveIndexFromStringKey_{DataInterfaceName}(int InMeshPathStringKey)\n"
 		"{\n"
 		"	for (uint Index = 0; Index < {DataInterfaceName}_NumPrimitives; ++Index)\n"
 		"	{\n"
@@ -95,6 +110,8 @@ void UPCGStaticMeshSpawnerDataInterface::GetHLSL(FString& OutHLSL, FString const
 		"\n"
 		"float SMSpawner_GetPrimitiveSelectionCDF_{DataInterfaceName}(uint InPrimitiveIndex) { return GET_SCALAR_ARRAY_ELEMENT({DataInterfaceName}_SelectionCDF, InPrimitiveIndex); }\n"
 		"\n"
+		"uint {DataInterfaceName}_SelectedMeshAttributeId;\n"
+		"uint SMSpawner_GetSelectedMeshAttributeId_{DataInterfaceName}() { return {DataInterfaceName}_SelectedMeshAttributeId; }\n"
 		), TemplateArgs);
 }
 
@@ -104,7 +121,7 @@ UComputeDataProvider* UPCGStaticMeshSpawnerDataInterface::CreateDataProvider(TOb
 	UPCGDataBinding* Binding = CastChecked<UPCGDataBinding>(InBinding);
 
 	TObjectPtr<UPCGStaticMeshSpawnerDataProvider> DataProvider = NewObject<UPCGStaticMeshSpawnerDataProvider>();
-	DataProvider->Settings = Cast<UPCGStaticMeshSpawnerSettings>(Settings);
+	DataProvider->Settings = CastChecked<UPCGStaticMeshSpawnerSettings>(Settings);
 	
 	// If there were 0 input points for this execution, we will not have created any primitives, so check for null.
 	if (const FPCGSpawnerPrimitives* Primitives = Binding->MeshSpawnersToPrimitives.Find(Settings))
@@ -113,6 +130,9 @@ UComputeDataProvider* UPCGStaticMeshSpawnerDataInterface::CreateDataProvider(TOb
 		DataProvider->SelectionCDF = Primitives->SelectionCDF;
 		DataProvider->SelectorAttributeId = Primitives->SelectorAttributeId;
 		DataProvider->PrimitiveStringKeys = Primitives->PrimitiveStringKeys;
+
+		const FPCGKernelAttributeIDAndType* FoundAttribute = Binding->GetAttributeLookupTable().Find(DataProvider->Settings->OutAttributeName);
+		DataProvider->SelectedMeshAttributeId = ensure(FoundAttribute) ? FoundAttribute->Id : -1;
 	}
 
 	return DataProvider;
@@ -120,7 +140,7 @@ UComputeDataProvider* UPCGStaticMeshSpawnerDataInterface::CreateDataProvider(TOb
 
 FComputeDataProviderRenderProxy* UPCGStaticMeshSpawnerDataProvider::GetRenderProxy()
 {
-	return new FPCGStaticMeshSpawnerDataProviderProxy(AttributeIdOffsetStrides, SelectorAttributeId, PrimitiveStringKeys, SelectionCDF);
+	return new FPCGStaticMeshSpawnerDataProviderProxy(AttributeIdOffsetStrides, SelectorAttributeId, PrimitiveStringKeys, SelectionCDF, SelectedMeshAttributeId);
 }
 
 bool FPCGStaticMeshSpawnerDataProviderProxy::IsValid(FValidationData const& InValidationData) const
@@ -142,7 +162,8 @@ void FPCGStaticMeshSpawnerDataProviderProxy::GatherDispatchData(FDispatchData co
 
 		Parameters.NumAttributes = AttributeIdOffsetStrides.Num();
 		Parameters.NumPrimitives = SelectionCDF.Num();
-		Parameters.SelectorAttributeId = (uint32)SelectorAttributeId;
+		Parameters.SelectorAttributeId = static_cast<uint32>(SelectorAttributeId);
+		Parameters.SelectedMeshAttributeId = static_cast<uint32>(SelectedMeshAttributeId);
 
 		for (int32 Index = 0; Index < AttributeIdOffsetStrides.Num(); ++Index)
 		{
@@ -151,7 +172,7 @@ void FPCGStaticMeshSpawnerDataProviderProxy::GatherDispatchData(FDispatchData co
 
 		for (int32 Index = 0; Index < PrimitiveStringKeys.Num(); ++Index)
 		{
-			Parameters.PrimitiveStringKeys[Index] = FUintVector4(PrimitiveStringKeys[Index], /*Unused*/0, /*Unused*/0, /*Unused*/0);
+			Parameters.PrimitiveStringKeys[Index] = FIntVector4(PrimitiveStringKeys[Index], /*Unused*/0, /*Unused*/0, /*Unused*/0);
 		}
 
 		for (int32 Index = 0; Index < SelectionCDF.Num(); ++Index)
diff --git a/Engine/Plugins/PCG/Source/PCG/Private/Compute/DataInterfaces/Elements/PCGStaticMeshSpawnerDataInterface.h b/Engine/Plugins/PCG/Source/PCG/Private/Compute/DataInterfaces/Elements/PCGStaticMeshSpawnerDataInterface.h
--- a/Engine/Plugins/PCG/Source/PCG/Private/Compute/DataInterfaces/Elements/PCGStaticMeshSpawnerDataInterface.h
+++ b/Engine/Plugins/PCG/Source/PCG/Private/Compute/DataInterfaces/Elements/PCGStaticMeshSpawnerDataInterface.h
@@ -49,7 +49,7 @@ public:
 	TArray<FUintVector4> AttributeIdOffsetStrides;
 
 	UPROPERTY()
-	TArray<uint32> PrimitiveStringKeys;
+	TArray<int32> PrimitiveStringKeys;
 
 	UPROPERTY()
 	TArray<float> SelectionCDF;
@@ -59,16 +59,20 @@ public:
 
 	UPROPERTY()
 	int32 NumInputPoints = 0;
+
+	UPROPERTY()
+	int32 SelectedMeshAttributeId = -1;
 };
 
 class FPCGStaticMeshSpawnerDataProviderProxy : public FComputeDataProviderRenderProxy
 {
 public:
-	FPCGStaticMeshSpawnerDataProviderProxy(TArray<FUintVector4>& InAttributeIdOffsetStrides, int32 InSelectorAttributeId, TArray<uint32>& InPrimitiveStringKeys, TArray<float> InSelectionCDF)
+	FPCGStaticMeshSpawnerDataProviderProxy(TArray<FUintVector4>& InAttributeIdOffsetStrides, int32 InSelectorAttributeId, TArray<int32>& InPrimitiveStringKeys, TArray<float> InSelectionCDF, int32 InSelectedMeshAttributeId)
 		: AttributeIdOffsetStrides(InAttributeIdOffsetStrides)
 		, SelectionCDF(InSelectionCDF)
 		, SelectorAttributeId(InSelectorAttributeId)
 		, PrimitiveStringKeys(InPrimitiveStringKeys)
+		, SelectedMeshAttributeId(InSelectedMeshAttributeId)
 	{}
 
 	//~ Begin FComputeDataProviderRenderProxy Interface
@@ -82,6 +86,8 @@ protected:
 	TArray<FUintVector4> AttributeIdOffsetStrides;
 	TArray<float> SelectionCDF;
 
-	int32 SelectorAttributeId = -1;
-	TArray<uint32> PrimitiveStringKeys;
+	int32 SelectorAttributeId = INDEX_NONE;
+	TArray<int32> PrimitiveStringKeys;
+
+	int32 SelectedMeshAttributeId = INDEX_NONE;
 };
diff --git a/Engine/Plugins/PCG/Source/PCG/Private/Compute/Elements/PCGComputeGraphElement.cpp b/Engine/Plugins/PCG/Source/PCG/Private/Compute/Elements/PCGComputeGraphElement.cpp
--- a/Engine/Plugins/PCG/Source/PCG/Private/Compute/Elements/PCGComputeGraphElement.cpp
+++ b/Engine/Plugins/PCG/Source/PCG/Private/Compute/Elements/PCGComputeGraphElement.cpp
@@ -285,7 +285,7 @@ bool FPCGComputeGraphElement::SetupProceduralISMComponents(FPCGContext* InContex
 
 		TArray<float> PrimitiveSelectionCDF;
 		int32 SelectorAttributeId = -1;
-		TArray<uint32> PrimitiveStringKeys; // ok index is just array index! TODO
+		TArray<int32> PrimitiveStringKeys;
 		TArray<FPCGProceduralISMComponentDescriptor> ComponentsToCreate;
 
 		if (const UPCGMeshSelectorByAttribute* SelectorByAttribute = Cast<UPCGMeshSelectorByAttribute>(SpawnerSettings->MeshSelectorParameters))
@@ -341,7 +341,7 @@ bool FPCGComputeGraphElement::SetupProceduralISMComponents(FPCGContext* InContex
 					continue;
 				}
 
-				PrimitiveStringKeys.Emplace(static_cast<uint32>(StringKey));
+				PrimitiveStringKeys.Emplace(StringKey);
 				ComponentsToCreate.Add(MoveTemp(Descriptor));
 			}
 
@@ -377,6 +377,7 @@ bool FPCGComputeGraphElement::SetupProceduralISMComponents(FPCGContext* InContex
 				const float Weight = float(Entry.Weight) / TotalWeight;
 				CumulativeWeight += Weight;
 				PrimitiveSelectionCDF.Add(CumulativeWeight);
+				PrimitiveStringKeys.Add(InBinding->GetStringTable().IndexOfByKey(Entry.Descriptor.StaticMesh.ToString()));
 
 				FPCGProceduralISMComponentDescriptor Descriptor;
 				Descriptor = Entry.Descriptor;
diff --git a/Engine/Plugins/PCG/Source/PCG/Private/Compute/PCGComputeGraph.h b/Engine/Plugins/PCG/Source/PCG/Private/Compute/PCGComputeGraph.h
--- a/Engine/Plugins/PCG/Source/PCG/Private/Compute/PCGComputeGraph.h
+++ b/Engine/Plugins/PCG/Source/PCG/Private/Compute/PCGComputeGraph.h
@@ -27,6 +27,8 @@ public:
 	/** Get the global attribute information collated during compilation. */
 	const TMap<FName, FPCGKernelAttributeIDAndType>& GetAttributeLookupTable() const { return GlobalAttributeLookupTable; }
 
+	const TArray<FString>& GetStringTable() const { return StringTable; }
+
 public:
 	TMap<TObjectKey<const UPCGNode>, TArray<FComputeKernelCompileMessage>> KernelToCompileMessages;
 	
@@ -54,8 +56,13 @@ public:
 	bool bLogDataDescriptions = false;
 
 protected:
+	/** Global attribute information collated during compilation. */
 	UPROPERTY()
 	TMap<FName /* Attribute name */, FPCGKernelAttributeIDAndType> GlobalAttributeLookupTable;
 
+	/** String table collated during compilation. */
+	UPROPERTY()
+	TArray<FString> StringTable;
+
 	friend class FPCGGraphCompilerGPU;
 };
diff --git a/Engine/Plugins/PCG/Source/PCG/Private/Compute/PCGDataBinding.cpp b/Engine/Plugins/PCG/Source/PCG/Private/Compute/PCGDataBinding.cpp
--- a/Engine/Plugins/PCG/Source/PCG/Private/Compute/PCGDataBinding.cpp
+++ b/Engine/Plugins/PCG/Source/PCG/Private/Compute/PCGDataBinding.cpp
@@ -58,6 +58,10 @@ void UPCGDataBinding::InitializeInputData(const FPCGDataCollection& InComputeGra
 void UPCGDataBinding::BuildStringTable()
 {
 	TRACE_CPUPROFILER_EVENT_SCOPE(UPCGDataBinding::BuildStringTable);
+	check(Graph);
+
+	// Start with any strings known statically at compilation time.
+	StringTable = Graph->GetStringTable();
 
 	for (const FPCGTaggedData& Data : DataForGPU.InputDataCollection.TaggedData)
 	{
diff --git a/Engine/Plugins/PCG/Source/PCG/Private/Compute/PCGDataBinding.h b/Engine/Plugins/PCG/Source/PCG/Private/Compute/PCGDataBinding.h
--- a/Engine/Plugins/PCG/Source/PCG/Private/Compute/PCGDataBinding.h
+++ b/Engine/Plugins/PCG/Source/PCG/Private/Compute/PCGDataBinding.h
@@ -34,7 +34,7 @@ struct FPCGSpawnerPrimitives
 	int32 SelectorAttributeId = -1;
 
 	UPROPERTY()
-	TArray<uint32> PrimitiveStringKeys;
+	TArray<int32> PrimitiveStringKeys;
 };
 
 UCLASS(Transient, Category = PCG)
diff --git a/Engine/Plugins/PCG/Source/PCG/Private/Elements/PCGStaticMeshSpawner.cpp b/Engine/Plugins/PCG/Source/PCG/Private/Elements/PCGStaticMeshSpawner.cpp
--- a/Engine/Plugins/PCG/Source/PCG/Private/Elements/PCGStaticMeshSpawner.cpp
+++ b/Engine/Plugins/PCG/Source/PCG/Private/Elements/PCGStaticMeshSpawner.cpp
@@ -125,9 +125,65 @@ FPCGDataCollectionDesc UPCGStaticMeshSpawnerSettings::ComputeOutputPinDataDesc(c
 		return Super::ComputeOutputPinDataDesc(OutputPin, Binding);
 	}
 
+	const FPCGDataCollectionDesc InputPinDesc = ComputeInputPinDataDesc(PCGPinConstants::DefaultInputLabel, Binding);
+	FPCGDataCollectionDesc OutputPinDesc = InputPinDesc;
 
-	// Passthrough - output description matches input description.
-	return ComputeInputPinDataDesc(PCGPinConstants::DefaultInputLabel, Binding);
+	if (ensure(OutputPin->Properties.Label == PCGPinConstants::DefaultOutputLabel))
+	{
+		const TMap<FName, FPCGKernelAttributeIDAndType>& GlobalAttributeLookupTable = Binding->GetAttributeLookupTable();
+
+		if (const FPCGKernelAttributeIDAndType* FoundAttributeIDAndType = GlobalAttributeLookupTable.Find(OutAttributeName))
+		{
+			// Create attribute description.
+			FPCGKernelAttributeDesc AttributeDesc(FoundAttributeIDAndType->Id, FoundAttributeIDAndType->Type, OutAttributeName);
+
+			// Create unique value keys for the output string.
+			if (UPCGMeshSelectorWeighted* SelectorWeighted = Cast<UPCGMeshSelectorWeighted>(MeshSelectorParameters))
+			{
+				// Weighted selection - add explicit strings from settings.
+				for (const FPCGMeshSelectorWeightedEntry& Entry : SelectorWeighted->MeshEntries)
+				{
+					const FString Value = Entry.Descriptor.StaticMesh.ToString();
+					const int32 StringIndex = Binding->GetStringTable().IndexOfByKey(Value);
+					if (ensureAlways(StringIndex != INDEX_NONE))
+					{
+						AttributeDesc.UniqueStringKeys.Add(StringIndex);
+					}
+				}
+			}
+			else if (UPCGMeshSelectorByAttribute* SelectorByAttribute = Cast<UPCGMeshSelectorByAttribute>(MeshSelectorParameters))
+			{
+				// By-attribute selection - pass on strings from input attribute.
+				if (SelectorByAttribute->AttributeName != NAME_None)
+				{
+					for (const FPCGDataDesc& InputDataDesc : InputPinDesc.DataDescs)
+					{
+						for (const FPCGKernelAttributeDesc& InputAttrDesc : InputDataDesc.AttributeDescs)
+						{
+							if (InputAttrDesc.Name == SelectorByAttribute->AttributeName)
+							{
+								AttributeDesc.UniqueStringKeys.Append(InputAttrDesc.UniqueStringKeys);
+							}
+						}
+					}
+				}
+			}
+			else if (MeshSelectorParameters)
+			{
+				UE_LOG(LogPCG, Warning, TEXT("Mesh selector not supported by GPU Static Mesh Spawner: %s"), *MeshSelectorParameters->GetName());
+			}
+
+			for (FPCGDataDesc& DataDesc : OutputPinDesc.DataDescs)
+			{
+				if (const FPCGKernelAttributeIDAndType* IDAndType = GlobalAttributeLookupTable.Find(OutAttributeName))
+				{
+					DataDesc.AttributeDescs.AddUnique(AttributeDesc);
+				}
+			}
+		}
+	}
+
+	return OutputPinDesc;
 }
 
 const TArray<FPCGKernelAttributeKey> UPCGStaticMeshSpawnerSettings::GetKernelAttributeKeys() const
@@ -142,7 +198,7 @@ const TArray<FPCGKernelAttributeKey> UPCGStaticMeshSpawnerSettings::GetKernelAtt
 			for (const FName& AttributeName : AttributeNames)
 			{
 				// We don't know the type statically before execution, leave unset.
-				AttributeKeys.Emplace(AttributeName, EPCGKernelAttributeType::None);
+				AttributeKeys.AddUnique(FPCGKernelAttributeKey(AttributeName, EPCGKernelAttributeType::None));
 			}
 		}
 	}
@@ -150,12 +206,26 @@ const TArray<FPCGKernelAttributeKey> UPCGStaticMeshSpawnerSettings::GetKernelAtt
 	if (const UPCGMeshSelectorByAttribute* Selector = Cast<UPCGMeshSelectorByAttribute>(MeshSelectorParameters))
 	{
 		// Add an attribute key for the given attribute so we register it as being read.
-		AttributeKeys.Emplace(Selector->AttributeName, EPCGKernelAttributeType::None);
+		AttributeKeys.AddUnique(FPCGKernelAttributeKey(Selector->AttributeName, EPCGKernelAttributeType::None));
 	}
 
+	// Write out attribute.
+	AttributeKeys.AddUnique(FPCGKernelAttributeKey(OutAttributeName, EPCGKernelAttributeType::StringKey));
+
 	return AttributeKeys;
 }
 
+void UPCGStaticMeshSpawnerSettings::AddStaticCreatedStrings(TArray<FString>& InOutStringTable) const
+{
+	if (UPCGMeshSelectorWeighted* SelectorWeighted = Cast<UPCGMeshSelectorWeighted>(MeshSelectorParameters))
+	{
+		for (const FPCGMeshSelectorWeightedEntry& Entry : SelectorWeighted->MeshEntries)
+		{
+			InOutStringTable.AddUnique(Entry.Descriptor.StaticMesh.ToString());
+		}
+	}
+}
+
 int UPCGStaticMeshSpawnerSettings::ComputeKernelThreadCount(const UPCGDataBinding* Binding) const
 {
 	const UPCGNode* Node = CastChecked<UPCGNode>(GetOuter());
diff --git a/Engine/Plugins/PCG/Source/PCG/Private/Graph/PCGGraphCompilerGPU.cpp b/Engine/Plugins/PCG/Source/PCG/Private/Graph/PCGGraphCompilerGPU.cpp
--- a/Engine/Plugins/PCG/Source/PCG/Private/Graph/PCGGraphCompilerGPU.cpp
+++ b/Engine/Plugins/PCG/Source/PCG/Private/Graph/PCGGraphCompilerGPU.cpp
@@ -596,6 +596,7 @@ void FPCGGraphCompilerGPU::BuildGPUGraphTask(
 					ComputeGraph->DataInterfaces.Add(OutputDI);
 					OutputPinDataInterfaces.Add({ TaskId, OutputPinProperties.Label }, OutputDI);
 
+					// TODO: This seems like a once-per-settings process rather than repeating for each output pin?
 					for (const FPCGKernelAttributeKey& Key : Settings->GetKernelAttributeKeys())
 					{
 						FPCGKernelAttributeIDAndType* FoundEntry = ComputeGraph->GlobalAttributeLookupTable.Find(Key.Name);
@@ -622,6 +623,8 @@ void FPCGGraphCompilerGPU::BuildGPUGraphTask(
 							ComputeGraph->GlobalAttributeLookupTable.Add(Key.Name, IDAndType);
 						}
 					}
+
+					Settings->AddStaticCreatedStrings(ComputeGraph->StringTable);
 				}
 			}
 		}
diff --git a/Engine/Plugins/PCG/Source/PCG/Public/Elements/PCGStaticMeshSpawner.h b/Engine/Plugins/PCG/Source/PCG/Public/Elements/PCGStaticMeshSpawner.h
--- a/Engine/Plugins/PCG/Source/PCG/Public/Elements/PCGStaticMeshSpawner.h
+++ b/Engine/Plugins/PCG/Source/PCG/Public/Elements/PCGStaticMeshSpawner.h
@@ -32,6 +32,7 @@ public:
 	virtual bool IsKernelValid(FPCGContext* InContext = nullptr, bool bQuiet = true) const;
 	virtual FString GetCookedKernelSource(const TMap<FName, FPCGKernelAttributeIDAndType>& GlobalAttributeLookupTable) const override;
 	virtual const TArray<FPCGKernelAttributeKey> GetKernelAttributeKeys() const override;
+	virtual void AddStaticCreatedStrings(TArray<FString>& InOutStringTable) const override;
 	virtual int ComputeKernelThreadCount(const UPCGDataBinding* Binding) const override;
 	virtual FPCGDataCollectionDesc ComputeOutputPinDataDesc(const UPCGPin* OutputPin, const UPCGDataBinding* Binding) const override;
 	virtual void CreateAdditionalInputDataInterfaces(TArray<TObjectPtr<UComputeDataInterface>>& OutDataInterfaces) const override;
diff --git a/Engine/Plugins/PCG/Source/PCG/Public/PCGSettings.h b/Engine/Plugins/PCG/Source/PCG/Public/PCGSettings.h
--- a/Engine/Plugins/PCG/Source/PCG/Public/PCGSettings.h
+++ b/Engine/Plugins/PCG/Source/PCG/Public/PCGSettings.h
@@ -524,6 +524,9 @@ public:
 	/** Get a list of the attributes read or written by this node. */
 	virtual const TArray<FPCGKernelAttributeKey> GetKernelAttributeKeys() const { return {}; }
 
+	/** Add any strings emitted by this node that are known statically at compile time. */
+	virtual void AddStaticCreatedStrings(TArray<FString>& InOutStringTable) const {};
+
 	/** Compute how many threads should be dispatched to execute this node on the GPU. */
 	virtual int ComputeKernelThreadCount(const UPCGDataBinding* Binding) const { return 0; };
 
