[PCG] Add some logic to make sure overridable properties that are hard  references are always loaded on the main thread

#rnx

diff --git a/Engine/Plugins/PCG/Source/PCG/Private/Graph/PCGGraphExecutor.cpp b/Engine/Plugins/PCG/Source/PCG/Private/Graph/PCGGraphExecutor.cpp
--- a/Engine/Plugins/PCG/Source/PCG/Private/Graph/PCGGraphExecutor.cpp
+++ b/Engine/Plugins/PCG/Source/PCG/Private/Graph/PCGGraphExecutor.cpp
@@ -203,6 +203,11 @@ FPCGTaskId FPCGGraphTask::GetGraphExecutionTaskId() const
 	return StackContext ? StackContext->GetGraphExecutionTaskId() : InvalidPCGTaskId;
 }
 
+bool FPCGGraphTask::CanExecuteOnlyOnMainThread() const
+{
+	return (Element && Element->CanExecuteOnlyOnMainThread(Context)) || (Context && Context->CanExecuteOnlyOnMainThread());
+}
+
 #if WITH_EDITOR
 void FPCGGraphTask::LogVisual(ELogVerbosity::Type InVerbosity, const FText& InMessage) const
 {
@@ -330,6 +335,11 @@ void FPCGGraphActiveTask::StopExecuting()
 	}
 }
 
+bool FPCGGraphActiveTask::CanExecuteOnlyOnMainThread() const
+{
+	return (Element && Element->CanExecuteOnlyOnMainThread(Context.Get())) || (Context && Context->CanExecuteOnlyOnMainThread());
+}
+
 FPCGGraphExecutor::FPCGGraphExecutor()
 	: FPCGGraphExecutor(nullptr)
 {
@@ -1062,7 +1072,7 @@ void FPCGGraphExecutor::PostTaskExecute(TSharedPtr<FPCGGraphActiveTask> ActiveTa
 				PCGGraphExecutor::AddToActiveTaskArrayNoLock(SleepingTasks, ActiveTaskPtr);
 				PCGGraphExecutor::RemoveAtFromActiveTaskArrayNoLock(ActiveTaskArrayToRemove, TaskRemoveIndex);
 			}
-			else if (ActiveTask.bIsGameThreadOnly != ActiveTask.Element->CanExecuteOnlyOnMainThread(ActiveTask.Context.Get()))
+			else if (ActiveTask.bIsGameThreadOnly != ActiveTask.CanExecuteOnlyOnMainThread())
 			{
 				ActiveTask.bIsGameThreadOnly = !ActiveTask.bIsGameThreadOnly;
 				TArray<TSharedPtr<FPCGGraphActiveTask>>& ActiveTaskArrayToAdd = ActiveTask.bIsGameThreadOnly ? ActiveTasksGameThreadOnly : ActiveTasks;
@@ -1182,7 +1192,7 @@ void FPCGGraphExecutor::OnTaskInputsReady(FPCGGraphTask& Task, TArray<FCachedRes
 		FCachedResult* CachedResult = nullptr;
 
 		// Query Cache if we can here, if not the Task will be pushed to ReadyTasks and PrepareForExecute will be called on MainThread if needed
-		if (bIsInGameThread || !Task.Element->CanExecuteOnlyOnMainThread(nullptr))
+		if (bIsInGameThread || !Task.CanExecuteOnlyOnMainThread())
 		{
 			PrepareForExecute(Task, CachedResult, /*bLiveTasksLockAlreadyLocked=*/false);
 		}
@@ -1440,7 +1450,7 @@ bool FPCGGraphExecutor::ExecuteScheduling(double EndTime, TSharedPtr<FPCGGraphAc
 						continue; // still sleeping
 					}
 
-					const bool bIsGameThreadOnly = SleepingTask->Element->CanExecuteOnlyOnMainThread(SleepingTask->Context.Get());
+					const bool bIsGameThreadOnly = SleepingTask->CanExecuteOnlyOnMainThread();
 					TArray<TSharedPtr<FPCGGraphActiveTask>>& ActiveTaskArray = bIsGameThreadOnly ? ActiveTasksGameThreadOnly : ActiveTasks;
 					SleepingTask->bIsGameThreadOnly = bIsGameThreadOnly;
 
@@ -1457,7 +1467,7 @@ bool FPCGGraphExecutor::ExecuteScheduling(double EndTime, TSharedPtr<FPCGGraphAc
 
 				if (!Task.bHasDonePrepareForExecute)
 				{
-					if (!bIsInGameThread && Task.Element->CanExecuteOnlyOnMainThread(nullptr))
+					if (!bIsInGameThread && Task.CanExecuteOnlyOnMainThread())
 					{
 						continue;
 					}
@@ -1482,7 +1492,7 @@ bool FPCGGraphExecutor::ExecuteScheduling(double EndTime, TSharedPtr<FPCGGraphAc
 				}
 
 				// Validate that we can start this task now
-				const bool bIsGameThreadOnly = Task.Element->CanExecuteOnlyOnMainThread(Task.Context);
+				const bool bIsGameThreadOnly = Task.CanExecuteOnlyOnMainThread();
 				TArray<TSharedPtr<FPCGGraphActiveTask>>& ActiveTaskArray = bIsGameThreadOnly ? ActiveTasksGameThreadOnly : ActiveTasks;
 
 				FPCGGraphActiveTask& ActiveTask = *ActiveTaskArray.Emplace_GetRef(MakeShared<FPCGGraphActiveTask>());
@@ -1567,7 +1577,7 @@ bool FPCGGraphExecutor::ExecuteScheduling(double EndTime, TSharedPtr<FPCGGraphAc
 					}
 
 					// Main thread only
-					if (ActiveTask->Element->CanExecuteOnlyOnMainThread(ActiveTask->Context.Get()))
+					if (ActiveTask->CanExecuteOnlyOnMainThread())
 					{
 						continue;
 					}
@@ -1725,7 +1735,7 @@ void FPCGGraphExecutor::ExecuteV1()
 		bHasAlreadyCheckedSleepingTasks = ReadyTasks.Num() == 0 && ActiveTasks.Num() == 0 && SleepingTasks.Num() > 0;
 
 		// First: if we have free resources, move ready tasks to the active tasks
-		bool bMainThreadAvailable = (ActiveTasks.Num() == 0 || !ActiveTasks[0]->Element->CanExecuteOnlyOnMainThread(ActiveTasks[0]->Context.Get()));
+		bool bMainThreadAvailable = (ActiveTasks.Num() == 0 || !ActiveTasks[0]->CanExecuteOnlyOnMainThread());
 		int32 NumAvailableThreads = FMath::Max(0, MaxNumThreads - CurrentlyUsedThreads);
 
 		const bool bMainThreadWasAvailable = bMainThreadAvailable;
@@ -1773,7 +1783,7 @@ void FPCGGraphExecutor::ExecuteV1()
 				}
 
 				// Validate that we can start this task now
-				const bool bIsMainThreadTask = SleepingTask->Element->CanExecuteOnlyOnMainThread(SleepingTask->Context.Get());
+				const bool bIsMainThreadTask = SleepingTask->CanExecuteOnlyOnMainThread();
 
 				if (!bIsMainThreadTask || bMainThreadAvailable)
 				{
@@ -1804,7 +1814,7 @@ void FPCGGraphExecutor::ExecuteV1()
 				check(Task.bHasDoneSetup && Task.bHasDonePrepareForExecute);
 
 				// Validate that we can start this task now
-				const bool bIsMainThreadTask = Task.Element->CanExecuteOnlyOnMainThread(Task.Context);
+				const bool bIsMainThreadTask = Task.CanExecuteOnlyOnMainThread();
 
 				if (!bIsMainThreadTask || bMainThreadAvailable)
 				{
@@ -1870,7 +1880,7 @@ void FPCGGraphExecutor::ExecuteV1()
 				if (!ActiveTask.bIsBypassed)
 	#endif
 				{
-					check(!ActiveTask.Element->CanExecuteOnlyOnMainThread(ActiveTask.Context.Get()));
+					check(!ActiveTask.CanExecuteOnlyOnMainThread());
 					ActiveTask.Context->AsyncState.EndTime = EndTime;
 					ActiveTask.Context->AsyncState.bIsRunningOnMainThread = false;
 					ActiveTask.StartExecuting();
@@ -1949,7 +1959,7 @@ void FPCGGraphExecutor::ExecuteV1()
 #endif
 
 				// Any task that asks to be paused or now needs to run on the main thread but doesn't have that slot currently will be moved to the sleeping queue
-				const bool bTaskShouldBePutAside = (ActiveTask->Context->bIsPaused || (ActiveTaskIndex > 0 && ActiveTask->Element->CanExecuteOnlyOnMainThread(ActiveTask->Context.Get())));
+				const bool bTaskShouldBePutAside = (ActiveTask->Context->bIsPaused || (ActiveTaskIndex > 0 && ActiveTask->CanExecuteOnlyOnMainThread()));
 				if (bTaskShouldBePutAside)
 				{
 					if (ActiveTask->Context->AsyncState.NumAvailableTasks > 0)
diff --git a/Engine/Plugins/PCG/Source/PCG/Private/Graph/PCGGraphExecutor.h b/Engine/Plugins/PCG/Source/PCG/Private/Graph/PCGGraphExecutor.h
--- a/Engine/Plugins/PCG/Source/PCG/Private/Graph/PCGGraphExecutor.h
+++ b/Engine/Plugins/PCG/Source/PCG/Private/Graph/PCGGraphExecutor.h
@@ -82,6 +82,8 @@ struct FPCGGraphTask
 
 	const FPCGStack* GetStack() const;
 	FPCGTaskId GetGraphExecutionTaskId() const;
+	
+	bool CanExecuteOnlyOnMainThread() const;
 
 	TArray<FPCGGraphTaskInput> Inputs;
 	const UPCGNode* Node = nullptr;
@@ -139,6 +141,8 @@ struct FPCGGraphActiveTask : TSharedFromThis<FPCGGraphActiveTask>
 	void StartExecuting();
 	void StopExecuting();
 
+	bool CanExecuteOnlyOnMainThread() const;
+
 	TArray<FPCGGraphTaskInput> Inputs;
 	FPCGElementPtr Element;
 	TUniquePtr<FPCGContext> Context;
diff --git a/Engine/Plugins/PCG/Source/PCG/Private/PCGContext.cpp b/Engine/Plugins/PCG/Source/PCG/Private/PCGContext.cpp
--- a/Engine/Plugins/PCG/Source/PCG/Private/PCGContext.cpp
+++ b/Engine/Plugins/PCG/Source/PCG/Private/PCGContext.cpp
@@ -23,6 +23,17 @@
 
 #define LOCTEXT_NAMESPACE "PCGContext"
 
+namespace PCGContextHelpers
+{
+	template <typename T>
+	bool GetOverrideParamValue(const IPCGAttributeAccessor& InAccessor, T& OutValue)
+	{
+		// Override were using the first entry (0) by default.
+		FPCGAttributeAccessorKeysEntries FirstEntry(PCGMetadataEntryKey(0));
+		return InAccessor.Get<T>(OutValue, FirstEntry, EPCGAttributeAccessorFlags::AllowBroadcastAndConstructible);
+	}
+}
+
 FPCGTaskId FPCGContext::GetGraphExecutionTaskId() const
 {
 	return ensure(Stack) ? Stack->GetGraphExecutionTaskId() : InvalidPCGTaskId;
@@ -131,9 +142,40 @@ void FPCGContext::InitializeSettings(bool bSkipPostLoad)
 			}
 
 			if (bHasParamConnected)
-			{				
-				FObjectDuplicationParameters DuplicateParams(const_cast<UPCGSettings*>(NodeSettings), GetTransientPackage());
+			{
+				// If we have an override, we are not on the main thread and there is a hard ref override, we need to make sure that all objects
+				// are loaded. If not, we need to schedule the task on the main thread.
 				const bool bIsInGameThread = IsInGameThread();
+				if (!bIsInGameThread && NodeSettings->HasAnyOverridableHardReferences())
+				{
+					for (const FPCGSettingsOverridableParam& Param : OverridableParams)
+					{
+						// If the param is not a hard ref, ignore.
+						if (!Param.IsHardReferenceOverride())
+						{
+							continue;
+						}
+
+						PCGAttributeAccessorHelpers::AccessorParamResult AccessorResult{};
+						TUniquePtr<const IPCGAttributeAccessor> AttributeAccessor = PCGAttributeAccessorHelpers::CreateConstAccessorForOverrideParamWithResult(InputData, Param, &AccessorResult);
+
+						// If the accessor failed to be created, ignore
+						if (!AttributeAccessor.IsValid())
+						{
+							continue;
+						}
+
+						FSoftObjectPath ObjectPath;
+						if (PCGContextHelpers::GetOverrideParamValue(*AttributeAccessor, ObjectPath) && !ObjectPath.ResolveObject())
+						{
+							// We have an override value that is not loaded, and we are not on the main thread. We need to schedule the task on the main thread.
+							bOverrideSettingsOnMainThread = true;
+							break;
+						}
+					}
+				}
+			
+				FObjectDuplicationParameters DuplicateParams(const_cast<UPCGSettings*>(NodeSettings), GetTransientPackage());
 				DuplicateParams.bSkipPostLoad = bSkipPostLoad;
 				DuplicateParams.ApplyFlags = RF_Transient;
 			
@@ -185,6 +227,10 @@ void FPCGContext::OverrideSettings()
 {
 	TRACE_CPUPROFILER_EVENT_SCOPE(FPCGContext::OverrideSettings);
 
+	// If we have to be on the main thread check it there. It is necessary for overrides of hard references that are not yet loaded. We can reset the flag afterwards.
+	check(!bOverrideSettingsOnMainThread || IsInGameThread());
+	bOverrideSettingsOnMainThread = false;
+
 	// Use original settings to avoid recomputing OverridableParams() everytime
 	const UPCGSettings* OriginalSettings = GetOriginalSettings<UPCGSettings>();
 
@@ -287,7 +333,7 @@ void FPCGContext::OverrideSettings()
 			FPCGAttributeAccessorKeysEntries FirstEntry(PCGMetadataEntryKey(0));
 					
 			PropertyType Value{};
-			if (!AttributeAccessor->Get<PropertyType>(Value, FirstEntry, EPCGAttributeAccessorFlags::AllowBroadcast | EPCGAttributeAccessorFlags::AllowConstructible))
+			if (!PCGContextHelpers::GetOverrideParamValue(*AttributeAccessor, Value))
 			{
 				PCGE_LOG_C(Warning, GraphAndLog, this, FText::Format(LOCTEXT("ConversionFailed", "Parameter '{0}' cannot be converted from attribute '{1}'"), FText::FromName(Param.Label), FText::FromName(AttributeName)));
 				return false;
diff --git a/Engine/Plugins/PCG/Source/PCG/Private/PCGSettings.cpp b/Engine/Plugins/PCG/Source/PCG/Private/PCGSettings.cpp
--- a/Engine/Plugins/PCG/Source/PCG/Private/PCGSettings.cpp
+++ b/Engine/Plugins/PCG/Source/PCG/Private/PCGSettings.cpp
@@ -92,6 +92,11 @@ FString FPCGSettingsOverridableParam::GetPropertyPath() const
 	return FString::JoinBy(Properties, PCGSettings::PropertyPathSeparator, [](const FProperty* InProperty) { return InProperty ? InProperty->GetAuthoredName() : FString(); });
 }
 
+bool FPCGSettingsOverridableParam::IsHardReferenceOverride() const
+{
+	return !Properties.IsEmpty() && CastField<const FObjectProperty>(Properties.Last());
+}
+
 TArray<FName> FPCGSettingsOverridableParam::GenerateAllPossibleAliases() const
 {
 	if (Properties.IsEmpty() || !HasAliases())
@@ -1042,6 +1047,9 @@ void UPCGSettings::InitializeCachedOverridableParams(bool bReset)
 	}
 #endif // WITH_EDITOR
 
+	// Reset the value
+	bHasAnyOverridableHardReferences = false;
+	
 	for (int32 i = 0; i < CachedOverridableParams.Num(); ++i)
 	{
 		FPCGSettingsOverridableParam& Param = CachedOverridableParams[i];
@@ -1103,6 +1111,12 @@ void UPCGSettings::InitializeCachedOverridableParams(bool bReset)
 				break;
 			}
 		}
+
+		// Keep the information if the last property is a hard ref
+		if (Param.IsHardReferenceOverride())
+		{
+			bHasAnyOverridableHardReferences = true;
+		}
 	}
 }
 
diff --git a/Engine/Plugins/PCG/Source/PCG/Public/PCGContext.h b/Engine/Plugins/PCG/Source/PCG/Public/PCGContext.h
--- a/Engine/Plugins/PCG/Source/PCG/Public/PCGContext.h
+++ b/Engine/Plugins/PCG/Source/PCG/Public/PCGContext.h
@@ -75,6 +75,8 @@ struct PCG_API FPCGContext
 	/** The current call stack. */
 	const FPCGStack* Stack = nullptr;
 
+	bool CanExecuteOnlyOnMainThread() const { return bOverrideSettingsOnMainThread; } 
+
 	const UPCGSettingsInterface* GetInputSettingsInterface() const;
 	
 	// After initializing the context, we can call this method to prepare for parameter override
@@ -174,6 +176,9 @@ private:
 	// List of params that were in effect overriden
 	TArray<const FPCGSettingsOverridableParam*> OverriddenParams;
 
+	// If the settings need to be overridden on the main thread (because we have to load objects)
+	bool bOverrideSettingsOnMainThread = false;
+
 	// List of objects created by the PCG Elements, we need to track them so we can remove their Async flags when storing results on main thread
 	// so that they can be considered as existing on the main thread (and get properly GCed)
 	TSet<TObjectPtr<UObject>> AsyncObjects;
diff --git a/Engine/Plugins/PCG/Source/PCG/Public/PCGSettings.h b/Engine/Plugins/PCG/Source/PCG/Public/PCGSettings.h
--- a/Engine/Plugins/PCG/Source/PCG/Public/PCGSettings.h
+++ b/Engine/Plugins/PCG/Source/PCG/Public/PCGSettings.h
@@ -116,6 +116,9 @@ struct FPCGSettingsOverridableParam
 
 	TArray<FName> GenerateAllPossibleAliases() const;
 
+	/** Returns true if the last property is an ObjectProperty. */
+	PCG_API bool IsHardReferenceOverride() const;
+
 #if WITH_EDITOR
 	FString GetDisplayPropertyPath() const;
 	PCG_API FText GetDisplayPropertyPathText() const;
@@ -511,6 +514,13 @@ protected:
 	UPROPERTY()
 	TArray<FPCGSettingsOverridableParam> CachedOverridableParams;
 
+	// We need to make sure that if we have hard references that are overridable, and they are overriden by paths
+	// on objects that are not yet loaded, that we are loading it on the main thread.
+	bool bHasAnyOverridableHardReferences = false;
+	
+public:
+	bool HasAnyOverridableHardReferences() const { return bHasAnyOverridableHardReferences; }
+
 	// GPU section
 public:
 	/** Whether this node should be executed on the GPU. */
