[PCG] Fix higen grid links losing data when upstream pin label is not 'Out'.

In a previous change (35826801), we needed to propagate pin type downstream across higen grid links, so that GPU compilation could access the pin type. This exposed issues with missing labels and we assigned virtual labels to the inputs and outputs of a grid link - 'In' and 'Out'. However this only worked when the upstream pin label is also 'Out'.

This change forwards the upstream pin properties all the way through the grid link and doesn't artificially mangle pin labels at any point of the link.

#rnx

diff --git a/Engine/Plugins/PCG/Source/PCG/Private/Graph/PCGGraphCompiler.cpp b/Engine/Plugins/PCG/Source/PCG/Private/Graph/PCGGraphCompiler.cpp
--- a/Engine/Plugins/PCG/Source/PCG/Private/Graph/PCGGraphCompiler.cpp
+++ b/Engine/Plugins/PCG/Source/PCG/Private/Graph/PCGGraphCompiler.cpp
@@ -470,11 +470,9 @@ void FPCGGraphCompiler::CreateGridLinkages(
 					UpstreamPinProps.AllowedTypes = UpstreamSettings->GetCurrentPinTypes(UpstreamPin);
 				}
 
-				// Create pin properties for the input of the LinkTask to convey type information. Since we also need to give a pin label, use the default input pin label.
-				// So, effectively, Link Elements have a default input pin and output pin. These will forward the type of the upstream pin, which is required to preserve the
-				// type of the upstream node, which will be culled from the graph.
+				// Create pin properties for the input of the LinkTask to convey type information. The grid linkage use the upstream output pin label on both input and
+				// output, so effectively the upstream output pin label is forwarded all the way through to the final downstream edge (on the lower grid).
 				FPCGPinProperties DownstreamPinProps = UpstreamPinProps;
-				DownstreamPinProps.Label = PCGPinConstants::DefaultInputLabel;
 
 				LinkTask.Inputs.Emplace(TaskInput.TaskId, UpstreamPinProps, DownstreamPinProps);
 
@@ -482,13 +480,14 @@ void FPCGGraphCompiler::CreateGridLinkages(
 				const EPCGHiGenGrid ToGrid = InOutTaskGenerationGrid[TaskId];
 
 				// This lambda runs at execution time and attempts to retrieve the data from a larger grid. Capture by value is intentional.
-				auto GridLinkageOperation = [FromGrid, ToGrid, ResourceKey, InGenerationGrid](FPCGContext* InContext)
+				auto GridLinkageOperation = [FromGrid, ToGrid, ResourceKey, InGenerationGrid, UpstreamPinLabel=UpstreamPinProps.Label](FPCGContext* InContext)
 				{
 					return PCGGraphExecutor::ExecuteGridLinkage(
 						InGenerationGrid,
 						FromGrid,
 						ToGrid,
 						ResourceKey,
+						UpstreamPinLabel,
 						static_cast<FPCGGridLinkageContext*>(InContext));
 				};
 
diff --git a/Engine/Plugins/PCG/Source/PCG/Private/Graph/PCGGraphExecutor.cpp b/Engine/Plugins/PCG/Source/PCG/Private/Graph/PCGGraphExecutor.cpp
--- a/Engine/Plugins/PCG/Source/PCG/Private/Graph/PCGGraphExecutor.cpp
+++ b/Engine/Plugins/PCG/Source/PCG/Private/Graph/PCGGraphExecutor.cpp
@@ -3013,6 +3013,7 @@ namespace PCGGraphExecutor
 		EPCGHiGenGrid InFromGrid,
 		EPCGHiGenGrid InToGrid,
 		const FString& InResourceKey,
+		FName InUpstreamPinLabel,
 		FPCGGridLinkageContext* InContext)
 	{
 		check(InContext);
@@ -3073,8 +3074,8 @@ namespace PCGGraphExecutor
 		{
 			FPCGDataCollection Data;
 
-			// Grid Linkages are assumed to have a single input pin with the default input label.
-			Data.TaggedData = InContext->InputData.GetInputsByPin(PCGPinConstants::DefaultInputLabel);
+			// Grid linkages have a single virtual input pin with the label matching the upstream pin. A grid linkage is created for each edge.
+			Data.TaggedData = InContext->InputData.GetInputsByPin(InUpstreamPinLabel);
 
 			PCGGraphExecutionLogging::LogGridLinkageTaskExecuteStore(InContext, InGenerationGrid, FromGridSize, ToGridSize, InResourceKey, Data.TaggedData.Num());
 			
@@ -3132,12 +3133,12 @@ namespace PCGGraphExecutor
 				PCGGraphExecutionLogging::LogGridLinkageTaskExecuteRetrieveSuccess(InContext, ComponentWithData, InResourceKey, Data->TaggedData.Num());
 				InContext->OutputData = *Data;
 
-				// Grid Linkages are assumed to have a single output pin with the default output label.
+				// Grid Linkages have a single virtual output pin which forwards the same label as the upstream output pin. Apply the label now.
 				// The normal path in IPCGElement::CleanupAndValidateOutput() to fix up the output label does not run because Grid Linkages
-				// have no Settings. Therefore, we have to fix-up by hand here.
+				// have no Settings.
 				for (FPCGTaggedData& TaggedData : InContext->OutputData.TaggedData)
 				{
-					TaggedData.Pin = PCGPinConstants::DefaultOutputLabel;
+					TaggedData.Pin = InUpstreamPinLabel;
 				}
 
 				return true;
diff --git a/Engine/Plugins/PCG/Source/PCG/Private/Graph/PCGGraphExecutor.h b/Engine/Plugins/PCG/Source/PCG/Private/Graph/PCGGraphExecutor.h
--- a/Engine/Plugins/PCG/Source/PCG/Private/Graph/PCGGraphExecutor.h
+++ b/Engine/Plugins/PCG/Source/PCG/Private/Graph/PCGGraphExecutor.h
@@ -600,5 +600,6 @@ namespace PCGGraphExecutor
 		EPCGHiGenGrid InFromGrid,
 		EPCGHiGenGrid InToGrid,
 		const FString& InResourceKey,
+		FName InUpstreamPinLabel,
 		FPCGGridLinkageContext* InContext);
 }
