[PCG] Fix race condition in SM Spawner when using by-attribute selection.

For N instances and P primitives, we allocate N*P primitives in scene, and scene writer runs N*P threads. The SM Spawner kernel however just needs to bucket N instances into P primitives and should only run N threads, not N*P threads as is currently done. The current code is inefficient and unnecessarily complex, and also has a race condition where threads race to exhaust each bucket. Also since only threads with index < N are actually allowed to write output data, if they don't win the race then no output gets written.

#rnx

diff --git a/Engine/Plugins/PCG/Shaders/Private/Elements/PCGStaticMeshSpawner.usf b/Engine/Plugins/PCG/Shaders/Private/Elements/PCGStaticMeshSpawner.usf
--- a/Engine/Plugins/PCG/Shaders/Private/Elements/PCGStaticMeshSpawner.usf
+++ b/Engine/Plugins/PCG/Shaders/Private/Elements/PCGStaticMeshSpawner.usf
@@ -3,8 +3,7 @@
 // Static mesh spawner kernel. Selects a primitive either using a primitive cumulative distribution function (CDF) or by an attribute, and packs attributes in custom floats.
 // Passes through points to output.
 
-// When spawning by attribute, we don't know in advance how many instances of each mesh will be spawned, so the worst case NumPoints * NumPrims is allocated in the GPU Scene,
-// and NumPoints * NumPrims threads are dispatched here, hence the adjustments to ThreadIndex below and usage of bWriteOutput.
+// When spawning by attribute, we don't know in advance how many instances of each mesh will be spawned, so the worst case NumPoints * NumPrims is allocated in the GPU Scene.
 
 #include "/Engine/Private/ComputeShaderUtils.ush"
 #include "/Engine/Private/Quaternion.ush"
@@ -32,12 +31,8 @@ void Main(uint3 GroupId : SV_GroupID, uint GroupIndex : SV_GroupIndex)
 		return;
 	}
 
-	// Break out thread index into DataIndex and In_ElementIndex. For spawning by attribute we spawn NumPoints * NumPrims threads.
-	const uint NumInputPoints = In_GetNumElements();
-	const uint PointIndex = ThreadIndex % NumInputPoints;
-
 	uint In_DataIndex, In_ElementIndex;
-	if (!In_GetThreadData(PointIndex, In_DataIndex, In_ElementIndex))
+	if (!In_GetThreadData(ThreadIndex, In_DataIndex, In_ElementIndex))
 	{
 		return;
 	}
@@ -75,12 +70,6 @@ void Main(uint3 GroupId : SV_GroupID, uint GroupIndex : SV_GroupIndex)
 		}
 	}
 
-	if (PrimitiveIndex >= NumPrimitives)
-	{
-		// Should never happen.
-		return;
-	}
-
 	const int InstanceIndexInt = InstanceData_GetIndexToWriteTo(PrimitiveIndex);
 	if (InstanceIndexInt < 0)
 	{
@@ -88,10 +77,8 @@ void Main(uint3 GroupId : SV_GroupID, uint GroupIndex : SV_GroupIndex)
 		return;
 	}
 
-	// We run NumPoints x NumPrimitives threads. The output points only have to be written for the first NumPoints threads, not the rest.
-	const bool bWriteOutput = ThreadIndex < NumInputPoints;
 	uint Out_DataIndex, Out_ElementIndex;
-	if (bWriteOutput && !Out_GetThreadData(ThreadIndex, Out_DataIndex, Out_ElementIndex))
+	if (!Out_GetThreadData(ThreadIndex, Out_DataIndex, Out_ElementIndex))
 	{
 		// Should never happen.
 		return;
@@ -104,11 +91,8 @@ void Main(uint3 GroupId : SV_GroupID, uint GroupIndex : SV_GroupIndex)
 		// Cull the instance by setting the transform to 0.
 		InstanceData_WriteInstanceLocalToWorld(InstanceIndex, (float4x4)0);
 		
-		// Propagate the 'Removed' status to the output point if either source or target point is removed. Otherwise, this point will not be culled.
-		if (bWriteOutput)
-		{
-			Out_RemovePoint(Out_DataIndex, Out_ElementIndex);
-		}
+		// Propagate the 'Removed' status to the output point if input point is removed. Otherwise, this point will not be culled.
+		Out_RemovePoint(Out_DataIndex, Out_ElementIndex);
 
 		return;
 	}
@@ -136,24 +120,21 @@ void Main(uint3 GroupId : SV_GroupID, uint GroupIndex : SV_GroupIndex)
 
 	// Write output point data (pass through).
 	// TODO - support writing mesh bounds.
-	if (bWriteOutput)
+	Out_SetPosition(Out_DataIndex, Out_ElementIndex, In_GetPosition(In_DataIndex, In_ElementIndex));
+	Out_SetRotation(Out_DataIndex, Out_ElementIndex, In_GetRotation(In_DataIndex, In_ElementIndex));
+	Out_SetScale(Out_DataIndex, Out_ElementIndex, In_GetScale(In_DataIndex, In_ElementIndex));
+	Out_SetBoundsMin(Out_DataIndex, Out_ElementIndex, In_GetBoundsMin(In_DataIndex, In_ElementIndex));
+	Out_SetBoundsMax(Out_DataIndex, Out_ElementIndex, In_GetBoundsMax(In_DataIndex, In_ElementIndex));
+	Out_SetColor(Out_DataIndex, Out_ElementIndex, In_GetColor(In_DataIndex, In_ElementIndex));
+	Out_SetSeed(Out_DataIndex, Out_ElementIndex, In_GetSeed(In_DataIndex, In_ElementIndex));
+	Out_SetDensity(Out_DataIndex, Out_ElementIndex, In_GetDensity(In_DataIndex, In_ElementIndex));
+	Out_SetSteepness(Out_DataIndex, Out_ElementIndex, In_GetSteepness(In_DataIndex, In_ElementIndex));
+
+	CopyAttributesToOuput(In_DataIndex, In_ElementIndex, Out_DataIndex, Out_ElementIndex);
+
+	if (SMSpawner_GetSelectedMeshAttributeId() != (uint)-1)
 	{
-		Out_SetPosition(Out_DataIndex, Out_ElementIndex, In_GetPosition(In_DataIndex, In_ElementIndex));
-		Out_SetRotation(Out_DataIndex, Out_ElementIndex, In_GetRotation(In_DataIndex, In_ElementIndex));
-		Out_SetScale(Out_DataIndex, Out_ElementIndex, In_GetScale(In_DataIndex, In_ElementIndex));
-		Out_SetBoundsMin(Out_DataIndex, Out_ElementIndex, In_GetBoundsMin(In_DataIndex, In_ElementIndex));
-		Out_SetBoundsMax(Out_DataIndex, Out_ElementIndex, In_GetBoundsMax(In_DataIndex, In_ElementIndex));
-		Out_SetColor(Out_DataIndex, Out_ElementIndex, In_GetColor(In_DataIndex, In_ElementIndex));
-		Out_SetSeed(Out_DataIndex, Out_ElementIndex, In_GetSeed(In_DataIndex, In_ElementIndex));
-		Out_SetDensity(Out_DataIndex, Out_ElementIndex, In_GetDensity(In_DataIndex, In_ElementIndex));
-		Out_SetSteepness(Out_DataIndex, Out_ElementIndex, In_GetSteepness(In_DataIndex, In_ElementIndex));
-
-		CopyAttributesToOuput(In_DataIndex, In_ElementIndex, Out_DataIndex, Out_ElementIndex);
-
-		if (SMSpawner_GetSelectedMeshAttributeId() != (uint)-1)
-		{
-			Out_SetStringKey(Out_DataIndex, Out_ElementIndex, SMSpawner_GetSelectedMeshAttributeId(), OutputMeshStringKey);
-		}
+		Out_SetStringKey(Out_DataIndex, Out_ElementIndex, SMSpawner_GetSelectedMeshAttributeId(), OutputMeshStringKey);
 	}
 }
 
diff --git a/Engine/Plugins/PCG/Source/PCG/Private/Elements/PCGStaticMeshSpawner.cpp b/Engine/Plugins/PCG/Source/PCG/Private/Elements/PCGStaticMeshSpawner.cpp
--- a/Engine/Plugins/PCG/Source/PCG/Private/Elements/PCGStaticMeshSpawner.cpp
+++ b/Engine/Plugins/PCG/Source/PCG/Private/Elements/PCGStaticMeshSpawner.cpp
@@ -231,18 +231,7 @@ int UPCGStaticMeshSpawnerSettings::ComputeKernelThreadCount(const UPCGDataBindin
 	const UPCGNode* Node = CastChecked<UPCGNode>(GetOuter());
 	const UPCGPin* Pin = Node->GetInputPin(PCGPinConstants::DefaultInputLabel);
 
-	uint32 ElementCount = ComputeInputPinDataDesc(Pin, Binding).ComputeDataElementCount(EPCGDataType::Point);
-
-	if (const FPCGSpawnerPrimitives* FoundPrimitives = Binding->MeshSpawnersToPrimitives.Find(this))
-	{
-		// When selecting primitives dynamically (by attribute), we don't know statically how many instances will end up in each primitive, so run worst case.
-		if (FoundPrimitives->SelectorAttributeId != -1)
-		{
-			ElementCount *= FoundPrimitives->Primitives.Num();
-		}
-	}
-
-	return ElementCount;
+	return ComputeInputPinDataDesc(Pin, Binding).ComputeDataElementCount(EPCGDataType::Point);
 }
 
 void UPCGStaticMeshSpawnerSettings::CreateAdditionalInputDataInterfaces(TArray<TObjectPtr<UComputeDataInterface>>& OutDataInterfaces) const
