[PCG] Fixed spline interior sampling not respecting the Bounding Shape input correctly

[RN] Fixed spline interior sampling not respecting the Bounding Shape input correctly

diff --git a/Engine/Plugins/PCG/Source/PCG/Private/Elements/PCGSplineSampler.cpp b/Engine/Plugins/PCG/Source/PCG/Private/Elements/PCGSplineSampler.cpp
--- a/Engine/Plugins/PCG/Source/PCG/Private/Elements/PCGSplineSampler.cpp
+++ b/Engine/Plugins/PCG/Source/PCG/Private/Elements/PCGSplineSampler.cpp
@@ -1215,6 +1215,7 @@ namespace PCGSplineSamplerHelpers
 		NumDispatch = FMath::Max(1, NumDispatch);
 
 		const int32 NumIterationsPerDispatch = NumIterations / NumDispatch;
+		const FBox GeneratedPointBounds = FBox(-FVector::OneVector * Params.InteriorSampleSpacing / 2.0f, FVector::OneVector * Params.InteriorSampleSpacing / 2.0f);
 
 		TArray<TArray<TTuple<FTransform, FVector, float>>> InteriorSplinePointData;
 		InteriorSplinePointData.SetNum(NumDispatch);
@@ -1369,7 +1370,7 @@ namespace PCGSplineSamplerHelpers
 
 						// Prune points outside of bounds
 						FPCGPoint BoundsTestPoint;
-						if (InBoundingShape && !InBoundingShape->SamplePoint(TransformWS, SplineLocalBounds.GetBox(), BoundsTestPoint, nullptr))
+						if (InBoundingShape && !InBoundingShape->SamplePoint(TransformWS, GeneratedPointBounds, BoundsTestPoint, nullptr))
 						{
 							continue;
 						}
