[PCG] GPU SM Spawner - fix output mesh attribute not being assigned when there's only a single primitive.

#rnx

diff --git a/Engine/Plugins/PCG/Shaders/Private/Elements/PCGStaticMeshSpawner.usf b/Engine/Plugins/PCG/Shaders/Private/Elements/PCGStaticMeshSpawner.usf
--- a/Engine/Plugins/PCG/Shaders/Private/Elements/PCGStaticMeshSpawner.usf
+++ b/Engine/Plugins/PCG/Shaders/Private/Elements/PCGStaticMeshSpawner.usf
@@ -40,7 +40,6 @@ void Main(uint3 GroupId : SV_GroupID, uint GroupIndex : SV_GroupIndex)
 	const int NumPrimitives = InstanceData_GetNumPrimitives();
 
 	uint PrimitiveIndex = 0;
-	uint OutputMeshStringKey = (uint)-1;
 
 	if (NumPrimitives > 1)
 	{
@@ -49,8 +48,8 @@ void Main(uint3 GroupId : SV_GroupID, uint GroupIndex : SV_GroupIndex)
 		if (SelectorAttributeId != (uint)-1)
 		{
 			// Select using string key
-			OutputMeshStringKey = In_GetStringKey(In_DataIndex, In_ElementIndex, SelectorAttributeId);
-			PrimitiveIndex = SMSpawner_GetPrimitiveIndexFromStringKey(OutputMeshStringKey);
+			const int MeshStringKey = In_GetStringKey(In_DataIndex, In_ElementIndex, SelectorAttributeId);
+			PrimitiveIndex = SMSpawner_GetPrimitiveIndexFromStringKey(MeshStringKey);
 		}
 		else
 		{
@@ -65,8 +64,6 @@ void Main(uint3 GroupId : SV_GroupID, uint GroupIndex : SV_GroupIndex)
 					break;
 				}
 			}
-
-			OutputMeshStringKey = SMSpawner_GetPrimitiveStringKey(PrimitiveIndex);
 		}
 	}
 
@@ -134,6 +131,7 @@ void Main(uint3 GroupId : SV_GroupID, uint GroupIndex : SV_GroupIndex)
 
 	if (SMSpawner_GetSelectedMeshAttributeId() != (uint)-1)
 	{
+		const int OutputMeshStringKey = SMSpawner_GetPrimitiveStringKey(PrimitiveIndex);
 		Out_SetStringKey(Out_DataIndex, Out_ElementIndex, SMSpawner_GetSelectedMeshAttributeId(), OutputMeshStringKey);
 	}
 }
