[PCG] Fix 'implicit conversion from 'literal int' to 'unsigned int'' warning when kernel fails to compile

#rnx

diff --git a/Engine/Plugins/PCG/Shaders/Private/PCGDataCollectionDataInterface.ush b/Engine/Plugins/PCG/Shaders/Private/PCGDataCollectionDataInterface.ush
--- a/Engine/Plugins/PCG/Shaders/Private/PCGDataCollectionDataInterface.ush
+++ b/Engine/Plugins/PCG/Shaders/Private/PCGDataCollectionDataInterface.ush
@@ -37,13 +37,14 @@ RWByteAddressBuffer {DataInterfaceName}_DataCollectionBuffer;
 #define DATA_COLLECTION_HEADER_SIZE_BYTES 4
 #define DATA_HEADER_SIZE_BYTES 1036
 #define ATTRIBUTE_HEADER_SIZE_BYTES 8
+#define KERNEL_EXECUTED_FLAG (1u << 31)
 
 // #################### HEADER READERS ##########################
 
 uint GetNumData_{DataInterfaceName}()
 {
 	// We use the most significant bit of NumData to store the KernelExecuted flag
-	return {DataInterfaceName}_DataCollectionBuffer.Load(0) & ~(1 << 31); // TODO: Code-gen to PCGComputeConstants::KernelExecutedFlag
+	return {DataInterfaceName}_DataCollectionBuffer.Load(0) & ~KERNEL_EXECUTED_FLAG;
 }
 
 uint GetDataNumElements_{DataInterfaceName}(uint InDataIndex)
@@ -150,7 +151,7 @@ void SetAsExecutedInternal_{DataInterfaceName}()
 {
 	// We use the most significant bit of NumData to store the KernelExecuted flag
 	const uint NumData = GetNumData_{DataInterfaceName}();
-	{DataInterfaceName}_DataCollectionBuffer.Store(0, NumData | (1 << 31)); // TODO: Code-gen to PCGComputeConstants::KernelExecutedFlag
+	{DataInterfaceName}_DataCollectionBuffer.Store(0, NumData | KERNEL_EXECUTED_FLAG);
 }
 
 // #################### ATTRIBUTE GETTERS ##########################
