[PCG] Keep PCGComponent local cache
- Brings back a local cache on the PCGComponent which can be reused between executions
- Always fetch from execution cache first and fill it with PCGComponent local cache if it doesn't exist

#rnx

diff --git a/Engine/Plugins/PCG/Source/PCG/Private/PCGComponent.cpp b/Engine/Plugins/PCG/Source/PCG/Private/PCGComponent.cpp
--- a/Engine/Plugins/PCG/Source/PCG/Private/PCGComponent.cpp
+++ b/Engine/Plugins/PCG/Source/PCG/Private/PCGComponent.cpp
@@ -2279,6 +2279,39 @@ void UPCGComponent::DirtyGenerated(EPCGComponentDirtyFlag DirtyFlag, const bool
 	bDirtyGenerated = true;
 
 	ClearPerPinGeneratedOutput();
+
+	// Dirty data as a waterfall from basic values
+	if (!!(DirtyFlag & EPCGComponentDirtyFlag::Actor))
+	{
+		CachedActorData = nullptr;
+		// Since landscape data is related on the bounds of the current actor, when we dirty the actor data, we need to dirty the landscape data as well
+		CachedLandscapeData = nullptr;
+		CachedLandscapeHeightData = nullptr;
+		CachedInputData = nullptr;
+		CachedPCGData = nullptr;
+	}
+
+	if (!!(DirtyFlag & EPCGComponentDirtyFlag::Landscape))
+	{
+		CachedLandscapeData = nullptr;
+		CachedLandscapeHeightData = nullptr;
+		if (InputType == EPCGComponentInput::Landscape)
+		{
+			CachedInputData = nullptr;
+			CachedPCGData = nullptr;
+		}
+	}
+
+	if (!!(DirtyFlag & EPCGComponentDirtyFlag::Input))
+	{
+		CachedInputData = nullptr;
+		CachedPCGData = nullptr;
+	}
+
+	if (!!(DirtyFlag & EPCGComponentDirtyFlag::Data))
+	{
+		CachedPCGData = nullptr;
+	}
 	
 	// For partitioned graph, we must forward the call to the partition actor, if we need to
 	// TODO: Don't forward for None for now, as it could break some stuff
@@ -2726,6 +2759,16 @@ void UPCGComponent::OnRefresh(bool bForceRefresh)
 }
 #endif // WITH_EDITOR
 
+// The Actor Data Cache is a double buffered cache
+// - When executing a graph we assign a CurrentGenerationTask to the Component.
+// - When calling GetActorPCGData(), we will first try and find the Data inside the Execution Cache using the CurrentGenerationTask id
+// - This guarantees that once that execution cache is primed, we will always return the same Data for the whole execution
+// - If the Execution Cache doesn't contain the Data, we will first check the Component local cache to see if it is still valid (Can get invalidated by UPCGComponent::DirtyGenerated call or in some case if the landscape cache is dirty)
+// - If the local Component Cache is valid, we will store the data in the Execution Cache for the following calls and return the data
+// - If the local Component Cache isn't valid, we will create the cache Data, store the data in the Execution Cache and in the Component local cache
+// - If CurrentGenerationTask is InvalidTaskId then only the Local Component Cache will be used
+//
+// This applies to GetActorPCGData/GetLandscapePCGData/GetLandscapeHeightData/GetInputPCGData/GetPCGData
 UPCGData* UPCGComponent::GetPCGData() const
 {
 	UPCGData* Data = nullptr;
@@ -2734,7 +2777,9 @@ UPCGData* UPCGComponent::GetPCGData() const
 		Data = PCGSubsystem->GetPCGData(CurrentGenerationTask);
 		if (!Data)
 		{
-			Data = CreatePCGData();
+			Data = CachedPCGData ? CachedPCGData.Get() : CreatePCGData();
+			CachedPCGData = Data;
+
 			PCGSubsystem->SetPCGData(CurrentGenerationTask, Data);
 			if (PCGSubsystem->IsGraphCacheDebuggingEnabled() && CurrentGenerationTask != InvalidPCGTaskId)
 			{
@@ -2754,7 +2799,9 @@ UPCGData* UPCGComponent::GetInputPCGData() const
 		Data = PCGSubsystem->GetInputPCGData(CurrentGenerationTask);
 		if (!Data)
 		{
-			Data = CreateInputPCGData();
+			Data = CachedInputData ? CachedInputData.Get() : CreateInputPCGData();
+			CachedInputData = Data;
+
 			PCGSubsystem->SetInputPCGData(CurrentGenerationTask, Data);
 			if (PCGSubsystem->IsGraphCacheDebuggingEnabled() && CurrentGenerationTask != InvalidPCGTaskId)
 			{
@@ -2775,7 +2822,9 @@ UPCGData* UPCGComponent::GetActorPCGData() const
 		if (!Data)
 		{
 			PCG_EXECUTION_CACHE_VALIDATION_CHECK(this);
-			Data = CreateActorPCGData();
+			Data = CachedActorData && !IsLandscapeCachedDataDirty(CachedActorData) ? CachedActorData.Get() : CreateActorPCGData();
+			CachedActorData = Data;
+
 			PCGSubsystem->SetActorPCGData(CurrentGenerationTask, Data);
 			if (PCGSubsystem->IsGraphCacheDebuggingEnabled() && CurrentGenerationTask != InvalidPCGTaskId)
 			{
@@ -2795,7 +2844,9 @@ UPCGData* UPCGComponent::GetLandscapePCGData() const
 		Data = PCGSubsystem->GetLandscapePCGData(CurrentGenerationTask);
 		if (!Data)
 		{
-			Data = CreateLandscapePCGData(/*bHeightOnly=*/false);
+			Data = CachedLandscapeData && !IsLandscapeCachedDataDirty(CachedLandscapeData) ? CachedLandscapeData.Get() : CreateLandscapePCGData(/*bHeightOnly=*/false);
+			CachedLandscapeData = Data;
+
 			PCGSubsystem->SetLandscapePCGData(CurrentGenerationTask, Data);
 			if (PCGSubsystem->IsGraphCacheDebuggingEnabled() && CurrentGenerationTask != InvalidPCGTaskId)
 			{
@@ -2815,7 +2866,9 @@ UPCGData* UPCGComponent::GetLandscapeHeightPCGData() const
 		Data = PCGSubsystem->GetLandscapeHeightPCGData(CurrentGenerationTask);
 		if (!Data)
 		{
-			Data = CreateLandscapePCGData(/*bHeightOnly=*/true);
+			Data = CachedLandscapeHeightData && !IsLandscapeCachedDataDirty(CachedLandscapeHeightData) ? CachedLandscapeHeightData.Get() : CreateLandscapePCGData(/*bHeightOnly=*/true);
+			CachedLandscapeHeightData = Data;
+
 			PCGSubsystem->SetLandscapeHeightPCGData(CurrentGenerationTask, Data);
 			if (PCGSubsystem->IsGraphCacheDebuggingEnabled() && CurrentGenerationTask != InvalidPCGTaskId)
 			{
diff --git a/Engine/Plugins/PCG/Source/PCG/Public/PCGComponent.h b/Engine/Plugins/PCG/Source/PCG/Public/PCGComponent.h
--- a/Engine/Plugins/PCG/Source/PCG/Public/PCGComponent.h
+++ b/Engine/Plugins/PCG/Source/PCG/Public/PCGComponent.h
@@ -577,6 +577,21 @@ private:
 
 	FBox GetGridBounds(const AActor* InActor) const;
 
+	UPROPERTY(Transient, NonPIEDuplicateTransient)
+	mutable TObjectPtr<UPCGData> CachedPCGData = nullptr;
+
+	UPROPERTY(Transient, NonPIEDuplicateTransient)
+	mutable TObjectPtr<UPCGData> CachedInputData = nullptr;
+
+	UPROPERTY(Transient, NonPIEDuplicateTransient)
+	mutable TObjectPtr<UPCGData> CachedActorData = nullptr;
+
+	UPROPERTY(Transient, NonPIEDuplicateTransient)
+	mutable TObjectPtr<UPCGData> CachedLandscapeData = nullptr;
+
+	UPROPERTY(Transient, NonPIEDuplicateTransient)
+	mutable TObjectPtr<UPCGData> CachedLandscapeHeightData = nullptr;
+
 #if WITH_EDITORONLY_DATA
 	UPROPERTY()
 	TSet<TSoftObjectPtr<AActor>> GeneratedActors_DEPRECATED;
