[PCG] Remove old compute constants and fixup inconsistent data header size calculations

#rnx

diff --git a/Engine/Plugins/PCG/Source/PCG/Private/Compute/PCGComputeCommon.h b/Engine/Plugins/PCG/Source/PCG/Private/Compute/PCGComputeCommon.h
--- a/Engine/Plugins/PCG/Source/PCG/Private/Compute/PCGComputeCommon.h
+++ b/Engine/Plugins/PCG/Source/PCG/Private/Compute/PCGComputeCommon.h
@@ -27,19 +27,14 @@ namespace PCGComputeConstants
 	constexpr int NUM_RESERVED_ATTRS = 32; // Reserved for point properties, spline accessors, etc.
 	constexpr int MAX_NUM_CUSTOM_ATTRS = MAX_NUM_ATTRS - NUM_RESERVED_ATTRS; // Reserved for custom attributes
 
-	constexpr int DATA_COLLECTION_HEADER_SIZE_BYTES = 4;
-	constexpr int DATA_HEADER_SIZE_BYTES = 1036;
-	constexpr int ATTRIBUTE_HEADER_SIZE_BYTES = 8;
+	constexpr int DATA_COLLECTION_HEADER_SIZE_BYTES = 4; // 4 bytes for NumData
+	constexpr int DATA_HEADER_PREAMBLE_SIZE_BYTES = 12; // 4 bytes for DataType, 4 bytes for NumAttrs, 4 bytes for NumElements
+	constexpr int ATTRIBUTE_HEADER_SIZE_BYTES = 8; // 4 bytes for PackedIdAndStride, 4 bytes for data start address
+	constexpr int DATA_HEADER_SIZE_BYTES = DATA_HEADER_PREAMBLE_SIZE_BYTES + MAX_NUM_ATTRS * ATTRIBUTE_HEADER_SIZE_BYTES;
 
 	constexpr int POINT_DATA_TYPE_ID = 0;
 	constexpr int PARAM_DATA_TYPE_ID = 1;
 
-	// TODO remove
-	constexpr int POINT_DATA_HEADER_PREAMBLE_SIZE_BYTES = 16; // 4 bytes for Type, 4 bytes for NumAttrs, 4 bytes for the address, 4 bytes for TypeInfo
-	constexpr int POINT_DATA_HEADER_SIZE_BYTES = POINT_DATA_HEADER_PREAMBLE_SIZE_BYTES + ATTRIBUTE_HEADER_SIZE_BYTES * MAX_NUM_ATTRS;
-	constexpr int PARAM_DATA_HEADER_PREAMBLE_SIZE_BYTES = 16; // 4 bytes for Type, 4 bytes for NumAttrs, 4 bytes for the address, 4 bytes for TypeInfo
-	constexpr int PARAM_DATA_HEADER_SIZE_BYTES = PARAM_DATA_HEADER_PREAMBLE_SIZE_BYTES + ATTRIBUTE_HEADER_SIZE_BYTES * MAX_NUM_ATTRS;
-
 	constexpr int NUM_POINT_PROPERTIES = 9;
 	constexpr int POINT_POSITION_ATTRIBUTE_ID = 0;
 	constexpr int POINT_ROTATION_ATTRIBUTE_ID = 1;
diff --git a/Engine/Plugins/PCG/Source/PCG/Private/Compute/PCGDataForGPU.cpp b/Engine/Plugins/PCG/Source/PCG/Private/Compute/PCGDataForGPU.cpp
--- a/Engine/Plugins/PCG/Source/PCG/Private/Compute/PCGDataForGPU.cpp
+++ b/Engine/Plugins/PCG/Source/PCG/Private/Compute/PCGDataForGPU.cpp
@@ -498,21 +498,9 @@ FPCGDataDesc::FPCGDataDesc(const UPCGData* InData, const TMap<FName, FPCGKernelA
 
 uint32 FPCGDataDesc::ComputePackedSize() const
 {
-	uint32 DataSizeBytes = 0;
+	check(PCGComputeHelpers::IsTypeAllowedInDataCollection(Type));
 
-	if (Type == EPCGDataType::Point)
-	{
-		DataSizeBytes += POINT_DATA_HEADER_SIZE_BYTES;
-	}
-	else if (Type == EPCGDataType::Param)
-	{
-		DataSizeBytes += PARAM_DATA_HEADER_SIZE_BYTES;
-	}
-	else
-	{
-		// TODO: Support more types
-		checkNoEntry();
-	}
+	uint32 DataSizeBytes = DATA_HEADER_SIZE_BYTES;
 
 	for (const FPCGKernelAttributeDesc& AttributeDesc : AttributeDescs)
 	{
