[PCG] Minor GPU UI/UX improvements
* Remove default input pin on Point Generators
* Remove Get/Set attr declarations for types that don't support them (Landscape/Texture)
* Prevent multiple connection on GPU Static Mesh Spawner
* Make all calls to LoadShaderSourceFile consistent (add ensure)

#rnx

diff --git a/Engine/Plugins/PCG/Source/PCG/Private/Compute/DataInterfaces/PCGDataCollectionDataInterface.cpp b/Engine/Plugins/PCG/Source/PCG/Private/Compute/DataInterfaces/PCGDataCollectionDataInterface.cpp
--- a/Engine/Plugins/PCG/Source/PCG/Private/Compute/DataInterfaces/PCGDataCollectionDataInterface.cpp
+++ b/Engine/Plugins/PCG/Source/PCG/Private/Compute/DataInterfaces/PCGDataCollectionDataInterface.cpp
@@ -402,8 +402,10 @@ void UPCGDataCollectionDataInterface::GetHLSL(FString& OutHLSL, FString const& I
 	};
 
 	FString TemplateFile;
-	LoadShaderSourceFile(TemplateFilePath, EShaderPlatform::SP_PCD3D_SM5, &TemplateFile, nullptr);
-	OutHLSL += FString::Format(*TemplateFile, TemplateArgs);
+	if (ensure(LoadShaderSourceFile(TemplateFilePath, EShaderPlatform::SP_PCD3D_SM5, &TemplateFile, nullptr)))
+	{
+		OutHLSL += FString::Format(*TemplateFile, TemplateArgs);
+	}
 }
 
 UComputeDataProvider* UPCGDataCollectionDataInterface::CreateDataProvider(TObjectPtr<UObject> InBinding, uint64 InInputMask, uint64 InOutputMask) const
diff --git a/Engine/Plugins/PCG/Source/PCG/Private/Compute/DataInterfaces/PCGLandscapeDataInterface.cpp b/Engine/Plugins/PCG/Source/PCG/Private/Compute/DataInterfaces/PCGLandscapeDataInterface.cpp
--- a/Engine/Plugins/PCG/Source/PCG/Private/Compute/DataInterfaces/PCGLandscapeDataInterface.cpp
+++ b/Engine/Plugins/PCG/Source/PCG/Private/Compute/DataInterfaces/PCGLandscapeDataInterface.cpp
@@ -248,8 +248,10 @@ void UPCGLandscapeDataInterface::GetHLSL(FString& OutHLSL, FString const& InData
 	};
 
 	FString TemplateFile;
-	LoadShaderSourceFile(TemplateFilePath, EShaderPlatform::SP_PCD3D_SM5, &TemplateFile, nullptr);
-	OutHLSL += FString::Format(*TemplateFile, TemplateArgs);
+	if (ensure(LoadShaderSourceFile(TemplateFilePath, EShaderPlatform::SP_PCD3D_SM5, &TemplateFile, nullptr)))
+	{
+		OutHLSL += FString::Format(*TemplateFile, TemplateArgs);
+	}
 }
 
 UComputeDataProvider* UPCGLandscapeDataInterface::CreateDataProvider(TObjectPtr<UObject> InBinding, uint64 InInputMask, uint64 InOutputMask) const
diff --git a/Engine/Plugins/PCG/Source/PCG/Private/Compute/DataInterfaces/PCGTextureDataInterface.cpp b/Engine/Plugins/PCG/Source/PCG/Private/Compute/DataInterfaces/PCGTextureDataInterface.cpp
--- a/Engine/Plugins/PCG/Source/PCG/Private/Compute/DataInterfaces/PCGTextureDataInterface.cpp
+++ b/Engine/Plugins/PCG/Source/PCG/Private/Compute/DataInterfaces/PCGTextureDataInterface.cpp
@@ -63,8 +63,10 @@ void UPCGTextureDataInterface::GetHLSL(FString& OutHLSL, FString const& InDataIn
 	};
 
 	FString TemplateFile;
-	LoadShaderSourceFile(TemplateFilePath, EShaderPlatform::SP_PCD3D_SM5, &TemplateFile, nullptr);
-	OutHLSL += FString::Format(*TemplateFile, TemplateArgs);
+	if (ensure(LoadShaderSourceFile(TemplateFilePath, EShaderPlatform::SP_PCD3D_SM5, &TemplateFile, nullptr)))
+	{
+		OutHLSL += FString::Format(*TemplateFile, TemplateArgs);
+	}
 }
 
 UComputeDataProvider* UPCGTextureDataInterface::CreateDataProvider(TObjectPtr<UObject> InBinding, uint64 InInputMask, uint64 InOutputMask) const
diff --git a/Engine/Plugins/PCG/Source/PCG/Private/Compute/Elements/PCGCustomHLSL.cpp b/Engine/Plugins/PCG/Source/PCG/Private/Compute/Elements/PCGCustomHLSL.cpp
--- a/Engine/Plugins/PCG/Source/PCG/Private/Compute/Elements/PCGCustomHLSL.cpp
+++ b/Engine/Plugins/PCG/Source/PCG/Private/Compute/Elements/PCGCustomHLSL.cpp
@@ -475,6 +475,12 @@ void UPCGCustomHLSLSettings::ApplyPreconfiguredSettings(const FPCGPreConfiguredS
 		{
 			KernelType = EPCGKernelType(PreconfiguredInfo.PreconfiguredIndex);
 
+			// Point Generator doesn't utilize the default input pin, so let's not add it by default.
+			if (KernelType == EPCGKernelType::PointGenerator)
+			{
+				InputPins.Empty();
+			}
+
 #if WITH_EDITOR
 			UpdatePinSettings();
 #endif
@@ -665,7 +671,7 @@ void UPCGCustomHLSLSettings::UpdateInputDeclarations()
 		InputDeclarations += TEXT("const uint ElementIndex;\n\n");
 	}
 
-	TArray<FString> DataPins;
+	TArray<FString> DataCollectionDataPins;
 	TArray<FString> PointDataPins;
 	TArray<FString> SplineDataPins;
 	TArray<FString> LandscapeDataPins;
@@ -674,7 +680,10 @@ void UPCGCustomHLSLSettings::UpdateInputDeclarations()
 
 	for (const FPCGPinProperties& Pin : InputPinProperties())
 	{
-		DataPins.Add(Pin.Label.ToString());
+		if (PCGComputeHelpers::IsTypeAllowedInDataCollection(Pin.AllowedTypes))
+		{
+			DataCollectionDataPins.Add(Pin.Label.ToString());
+		}
 
 		if (!!(Pin.AllowedTypes & EPCGDataType::Point))
 		{
@@ -697,15 +706,15 @@ void UPCGCustomHLSLSettings::UpdateInputDeclarations()
 		}
 	}
 
-	if (!DataPins.IsEmpty())
+	if (!DataCollectionDataPins.IsEmpty())
 	{
 		InputDeclarations += TEXT("/*** INPUT DATA FUNCTIONS ***/\n\n");
 
-		const bool bMultiPin = DataPins.Num() > 1;
+		const bool bMultiPin = DataCollectionDataPins.Num() > 1;
 
 		if (bMultiPin)
 		{
-			InputDeclarations += TEXT("// Valid pins: ") + FString::Join(DataPins, TEXT(", ")) + TEXT("\n");
+			InputDeclarations += TEXT("// Valid pins: ") + FString::Join(DataCollectionDataPins, TEXT(", ")) + TEXT("\n");
 		}
 
 		InputDeclarations += FString::Format(TEXT(
@@ -716,7 +725,7 @@ void UPCGCustomHLSLSettings::UpdateInputDeclarations()
 			"\n"
 			"<type> {0}_Get<type>(uint DataIndex, uint ElementIndex, uint AttributeId);\n"
 			"<type> {0}_Get<type>(uint DataIndex, uint ElementIndex, 'AttributeName');\n"),
-			{ bMultiPin ? PCGHLSLElement::PinDeclTemplateStr : DataPins[0] });
+			{ bMultiPin ? PCGHLSLElement::PinDeclTemplateStr : DataCollectionDataPins[0] });
 
 		InputDeclarations += TEXT("\n");
 	}
@@ -814,14 +823,17 @@ void UPCGCustomHLSLSettings::UpdateOutputDeclarations()
 {
 	OutputDeclarations.Reset();
 
-	TArray<FString> DataPins;
+	TArray<FString> DataCollectionDataPins;
 	TArray<FString> PointDataPins;
 	TArray<FString> SplineDataPins;
 	TArray<FString> RawBufferDataPins;
 
 	for (const FPCGPinProperties& Pin : OutputPinProperties())
 	{
-		DataPins.Add(Pin.Label.ToString());
+		if (PCGComputeHelpers::IsTypeAllowedInDataCollection(Pin.AllowedTypes))
+		{
+			DataCollectionDataPins.Add(Pin.Label.ToString());
+		}
 
 		if (!!(Pin.AllowedTypes & EPCGDataType::Point))
 		{
@@ -834,15 +846,15 @@ void UPCGCustomHLSLSettings::UpdateOutputDeclarations()
 		}
 	}
 
-	if (!DataPins.IsEmpty())
+	if (!DataCollectionDataPins.IsEmpty())
 	{
 		OutputDeclarations += TEXT("/*** OUTPUT DATA FUNCTIONS ***/\n\n");
 
-		const bool bMultiPin = DataPins.Num() > 1;
+		const bool bMultiPin = DataCollectionDataPins.Num() > 1;
 
 		if (bMultiPin)
 		{
-			OutputDeclarations += TEXT("// Valid pins: ") + FString::Join(DataPins, TEXT(", ")) + TEXT("\n");
+			OutputDeclarations += TEXT("// Valid pins: ") + FString::Join(DataCollectionDataPins, TEXT(", ")) + TEXT("\n");
 		}
 
 		OutputDeclarations += FString::Format(TEXT(
@@ -850,7 +862,7 @@ void UPCGCustomHLSLSettings::UpdateOutputDeclarations()
 			"\n"
 			"void {0}_Set<type>(uint DataIndex, uint ElementIndex, uint AttributeId, <type> Value);\n"
 			"void {0}_Set<type>(uint DataIndex, uint ElementIndex, 'AttributeName', <type> Value);\n"),
-			{ bMultiPin ? PCGHLSLElement::PinDeclTemplateStr : DataPins[0] });
+			{ bMultiPin ? PCGHLSLElement::PinDeclTemplateStr : DataCollectionDataPins[0] });
 
 		OutputDeclarations += TEXT("\n");
 	}
diff --git a/Engine/Plugins/PCG/Source/PCG/Private/Elements/PCGCopyPoints.cpp b/Engine/Plugins/PCG/Source/PCG/Private/Elements/PCGCopyPoints.cpp
--- a/Engine/Plugins/PCG/Source/PCG/Private/Elements/PCGCopyPoints.cpp
+++ b/Engine/Plugins/PCG/Source/PCG/Private/Elements/PCGCopyPoints.cpp
@@ -28,7 +28,7 @@ FText UPCGCopyPointsSettings::GetNodeTooltipText() const
 FString UPCGCopyPointsSettings::GetCookedKernelSource(const TMap<FName, FPCGKernelAttributeIDAndType>& GlobalAttributeLookupTable) const
 {
 	FString TemplateFile;
-	LoadShaderSourceFile(TEXT("/Plugin/PCG/Private/Elements/PCGCopyPoints.usf"), EShaderPlatform::SP_PCD3D_SM5, &TemplateFile, nullptr);
+	ensure(LoadShaderSourceFile(TEXT("/Plugin/PCG/Private/Elements/PCGCopyPoints.usf"), EShaderPlatform::SP_PCD3D_SM5, &TemplateFile, nullptr));
 	return TemplateFile;
 }
 
@@ -131,6 +131,7 @@ void UPCGCopyPointsSettings::CreateAdditionalInputDataInterfaces(TArray<TObjectP
 TArray<FPCGPinProperties> UPCGCopyPointsSettings::InputPinProperties() const
 {
 	// Note: If executing on the GPU, we need to prevent multiple connections on inputs, since it is not supported at this time.
+	// Also note: Since the ShouldExecuteOnGPU() is already tied to structural changes, we don't need to implement any logic for this in GetChangeTypeForProperty()
 	const bool bAllowMultipleConnections = !ShouldExecuteOnGPU();
 
 	TArray<FPCGPinProperties> PinProperties;
diff --git a/Engine/Plugins/PCG/Source/PCG/Private/Elements/PCGStaticMeshSpawner.cpp b/Engine/Plugins/PCG/Source/PCG/Private/Elements/PCGStaticMeshSpawner.cpp
--- a/Engine/Plugins/PCG/Source/PCG/Private/Elements/PCGStaticMeshSpawner.cpp
+++ b/Engine/Plugins/PCG/Source/PCG/Private/Elements/PCGStaticMeshSpawner.cpp
@@ -273,6 +273,19 @@ void UPCGStaticMeshSpawnerSettings::ApplyDeprecation(UPCGNode* InOutNode)
 }
 #endif
 
+TArray<FPCGPinProperties> UPCGStaticMeshSpawnerSettings::InputPinProperties() const
+{
+	// Note: If executing on the GPU, we need to prevent multiple connections on inputs, since it is not supported at this time.
+	// Also note: Since the ShouldExecuteOnGPU() is already tied to structural changes, we don't need to implement any logic for this in GetChangeTypeForProperty()
+	const bool bAllowMultipleConnections = !ShouldExecuteOnGPU();
+
+	TArray<FPCGPinProperties> Properties;
+	FPCGPinProperties& InputPinProperty = Properties.Emplace_GetRef(PCGPinConstants::DefaultInputLabel, EPCGDataType::Point, bAllowMultipleConnections);
+	InputPinProperty.SetRequiredPin();
+
+	return Properties;
+}
+
 FPCGElementPtr UPCGStaticMeshSpawnerSettings::CreateElement() const
 {
 	return MakeShared<FPCGStaticMeshSpawnerElement>();
diff --git a/Engine/Plugins/PCG/Source/PCG/Public/Elements/PCGStaticMeshSpawner.h b/Engine/Plugins/PCG/Source/PCG/Public/Elements/PCGStaticMeshSpawner.h
--- a/Engine/Plugins/PCG/Source/PCG/Public/Elements/PCGStaticMeshSpawner.h
+++ b/Engine/Plugins/PCG/Source/PCG/Public/Elements/PCGStaticMeshSpawner.h
@@ -47,7 +47,7 @@ public:
 	virtual bool UseSeed() const override { return true; }
 
 protected:
-	virtual TArray<FPCGPinProperties> InputPinProperties() const override { return Super::DefaultPointInputPinProperties(); }
+	virtual TArray<FPCGPinProperties> InputPinProperties() const override;
 	virtual TArray<FPCGPinProperties> OutputPinProperties() const override { return Super::DefaultPointOutputPinProperties(); }
 	virtual FPCGElementPtr CreateElement() const override;
 	// ~End UPCGSettings interface
