[PCG] GPU - highest priority UI/UX fixes
* Add Open Node Source Editor button to Custom HLSL details to summon editor tab.
* Rename "Shader Text" section of source editor to "Shader Source"
* Mark CustomHLSL node as experimental in tooltip
* Fix some out of date declarations

#rnx

diff --git a/Engine/Plugins/PCG/Source/PCG/Private/Compute/Elements/PCGCustomHLSL.cpp b/Engine/Plugins/PCG/Source/PCG/Private/Compute/Elements/PCGCustomHLSL.cpp
--- a/Engine/Plugins/PCG/Source/PCG/Private/Compute/Elements/PCGCustomHLSL.cpp
+++ b/Engine/Plugins/PCG/Source/PCG/Private/Compute/Elements/PCGCustomHLSL.cpp
@@ -720,6 +720,7 @@ void UPCGCustomHLSLSettings::UpdateInputDeclarations()
 		InputDeclarations += FString::Format(TEXT(
 			"uint {0}_GetNumData();\n"
 			"uint {0}_GetNumElements();\n"
+			"uint {0}_GetDataNumElements(uint DataIndex);\n"
 			"\n"
 			"// Valid types: bool, int, float, float2, float3, float4, Rotator (float3), Quat (float4), Transform (float4x4), StringKey (int)\n"
 			"\n"
@@ -742,7 +743,6 @@ void UPCGCustomHLSLSettings::UpdateInputDeclarations()
 		}
 
 		InputDeclarations += FString::Format(TEXT(
-			"uint {0}_GetNumPoints(uint DataIndex);\n"
 			"float3 {0}_GetPosition(uint DataIndex, uint ElementIndex);\n"
 			"float4 {0}_GetRotation(uint DataIndex, uint ElementIndex);\n"
 			"float3 {0}_GetScale(uint DataIndex, uint ElementIndex);\n"
diff --git a/Engine/Plugins/PCG/Source/PCG/Private/Graph/PCGGraphCompilerGPU.cpp b/Engine/Plugins/PCG/Source/PCG/Private/Graph/PCGGraphCompilerGPU.cpp
--- a/Engine/Plugins/PCG/Source/PCG/Private/Graph/PCGGraphCompilerGPU.cpp
+++ b/Engine/Plugins/PCG/Source/PCG/Private/Graph/PCGGraphCompilerGPU.cpp
@@ -546,9 +546,15 @@ void FPCGGraphCompilerGPU::BuildGPUGraphTask(
 			break;
 		}
 		default:
-			ensure(false);
+		{
+			UE_LOG(LogPCG, Error, TEXT("Unsupported connected upstream pin '%s' on node '%s' with type %s. Consider adding a conversion to a supported type such as Point."),
+				*InOutputPinProperties.Label.ToString(),
+				InOutCompiledTasks[InTaskId].Node ? *InOutCompiledTasks[InTaskId].Node->GetNodeTitle(EPCGNodeTitleType::ListView).ToString() : TEXT("MISSING"),
+				StaticEnum<EPCGDataType>() ? *StaticEnum<EPCGDataType>()->GetNameStringByValue(static_cast<int64>(PinType)) : TEXT("MISSING")
+			);
 			break;
 		}
+		}
 
 		if (DataInterface)
 		{
diff --git a/Engine/Plugins/PCG/Source/PCG/Public/Compute/Elements/PCGCustomHLSL.h b/Engine/Plugins/PCG/Source/PCG/Public/Compute/Elements/PCGCustomHLSL.h
--- a/Engine/Plugins/PCG/Source/PCG/Public/Compute/Elements/PCGCustomHLSL.h
+++ b/Engine/Plugins/PCG/Source/PCG/Public/Compute/Elements/PCGCustomHLSL.h
@@ -56,7 +56,7 @@ public:
 #if WITH_EDITOR
 	virtual FName GetDefaultNodeName() const override { return FName(TEXT("CustomHLSL")); }
 	virtual FText GetDefaultNodeTitle() const override { return NSLOCTEXT("PCGCustomHLSLElement", "NodeTitle", "Custom HLSL"); }
-	virtual FText GetNodeTooltipText() const override { return NSLOCTEXT("PCGCustomHLSLElement", "NodeTooltip", "Produces a HLSL compute shader which will be executed on the GPU."); }
+	virtual FText GetNodeTooltipText() const override { return NSLOCTEXT("PCGCustomHLSLElement", "NodeTooltip", "[EXPERIMENTAL] Produces a HLSL compute shader which will be executed on the GPU."); }
 	virtual EPCGSettingsType GetType() const override { return EPCGSettingsType::GPU; }
 
 	virtual TArray<FPCGPreConfiguredSettingsInfo> GetPreconfiguredInfo() const override;
@@ -158,23 +158,23 @@ public:
 
 protected:
 	/** Optional functions that can be called from the source. Intended to be edited using the Node Source Editor window. */
-	UPROPERTY(BlueprintReadWrite, EditAnywhere, Category = "Source", meta = (MultiLine = true))
-	FString ShaderFunctions;
+	UPROPERTY(BlueprintReadWrite, EditAnywhere, Category = "Settings", AdvancedDisplay, meta = (MultiLine = true))
+	FString ShaderFunctions = "/** CUSTOM SHADER FUNCTIONS **/\n";
 
 	/** Shader code that forms the body of the kernel. Intended to be edited using the Node Source Editor window. */
-	UPROPERTY(BlueprintReadWrite, EditAnywhere, Category = "Source", meta = (MultiLine = true))
+	UPROPERTY(BlueprintReadWrite, EditAnywhere, Category = "Settings", AdvancedDisplay, meta = (MultiLine = true))
 	FString ShaderSource;
 
 	/** Inputs data accessors that can be used from the shader code. Intended to be viewed using the Node Source Editor window. */
-	UPROPERTY(Transient, VisibleAnywhere, Category = "Declarations|Inputs", meta = (MultiLine = true))
+	UPROPERTY(Transient, VisibleAnywhere, Category = "Settings", AdvancedDisplay, meta = (MultiLine = true))
 	FString InputDeclarations;
 
 	/** Output data accessors that can be used from the shader code. Intended to be viewed using the Node Source Editor window. */
-	UPROPERTY(Transient, VisibleAnywhere, Category = "Declarations|Outputs", meta = (MultiLine = true))
+	UPROPERTY(Transient, VisibleAnywhere, Category = "Settings", AdvancedDisplay, meta = (MultiLine = true))
 	FString OutputDeclarations;
 
 	/** Helper data and functions that can be used from the shader code. Intended to be viewed using the Node Source Editor window. */
-	UPROPERTY(Transient, VisibleAnywhere, Category = "Declarations|Helpers", meta = (MultiLine = true))
+	UPROPERTY(Transient, VisibleAnywhere, Category = "Settings", AdvancedDisplay, meta = (MultiLine = true))
 	FString HelperDeclarations;
 
 	/** Attributes statically detected as being read, written, or created by this node. */
diff --git a/Engine/Plugins/PCG/Source/PCG/Public/PCGSettings.h b/Engine/Plugins/PCG/Source/PCG/Public/PCGSettings.h
--- a/Engine/Plugins/PCG/Source/PCG/Public/PCGSettings.h
+++ b/Engine/Plugins/PCG/Source/PCG/Public/PCGSettings.h
@@ -560,7 +560,7 @@ public:
 
 public:
 	/** Whether this node should be executed on the GPU. */
-	UPROPERTY(BlueprintReadWrite, EditAnywhere, Category = GPU, meta = (EditCondition = "DisplayExecuteOnGPUSetting()", EditConditionHides, HideEditConditionToggle))
+	UPROPERTY(BlueprintReadWrite, EditAnywhere, Category = "GPU", meta = (EditCondition = "DisplayExecuteOnGPUSetting()", EditConditionHides, HideEditConditionToggle))
 	bool bExecuteOnGPU = false;
 
 	/** Dump the cooked HLSL into the log after it is generated. */
diff --git a/Engine/Plugins/PCG/Source/PCGEditor/Private/Details/PCGCustomHLSLSettingsDetails.cpp b/Engine/Plugins/PCG/Source/PCGEditor/Private/Details/PCGCustomHLSLSettingsDetails.cpp
new file mode 100644
--- /dev/null
+++ b/Engine/Plugins/PCG/Source/PCGEditor/Private/Details/PCGCustomHLSLSettingsDetails.cpp
@@ -0,0 +1,85 @@
+// Copyright Epic Games, Inc. All Rights Reserved.
+
+#include "Details/PCGCustomHLSLSettingsDetails.h"
+
+#include "Compute/Elements/PCGCustomHLSL.h"
+
+#include "PCGEditor.h"
+#include "PCGEditorGraph.h"
+
+#include "Containers/StringConv.h"
+#include "DetailWidgetRow.h"
+#include "DetailLayoutBuilder.h"
+#include "DetailCategoryBuilder.h"
+#include "PropertyHandle.h"
+#include "Widgets/Input/SButton.h"
+
+#define LOCTEXT_NAMESPACE "PCGCustomHLSLSettingsDetails"
+
+void FPCGCustomHLSLSettingsDetails::CustomizeDetails(IDetailLayoutBuilder& DetailBuilder)
+{
+	TArray<TWeakObjectPtr<UObject>> ObjectsBeingCustomized;
+	DetailBuilder.GetObjectsBeingCustomized(ObjectsBeingCustomized);
+
+	IDetailCategoryBuilder& SettingsCategory = DetailBuilder.EditCategory(TEXT("Settings"));
+
+	TWeakPtr<FPCGEditor> PCGEditor = GetPCGEditor(ObjectsBeingCustomized);
+	
+	// Add a "open source editor" button if the source editor tab is not already visible.
+	if (PCGEditor.IsValid() && PCGEditor.Pin()->IsNodeSourceEditorTabClosed())
+	{
+		SettingsCategory.AddCustomRow(FText::GetEmpty()).ValueContent()
+			.MaxDesiredWidth(120.f)
+			[
+				SNew(SHorizontalBox)
+				+SHorizontalBox::Slot()
+				.AutoWidth()
+				.Padding(2.0f, 0.0f)
+				.VAlign(VAlign_Fill)
+				[
+					SNew(SButton)
+					.OnClicked_Lambda([PCGEditor]() -> FReply
+					{
+						if (PCGEditor.IsValid())
+						{
+							PCGEditor.Pin()->SpawnNodeSourceEditorTab();
+						}
+
+						return FReply::Handled();
+					})
+					.ToolTipText(FText::FromString("Opens Node Source Editor tab."))
+					[
+						SNew(STextBlock)
+						.Font(IDetailLayoutBuilder::GetDetailFont())
+						.Text(LOCTEXT("ButtonOpenSourceEditor", "Open Source Editor"))
+					]
+				]
+			];
+	}
+
+	// Add default properties in the category in order
+	TArray<TSharedRef<IPropertyHandle>> AllProperties;
+	SettingsCategory.GetDefaultProperties(AllProperties);
+
+	for (auto& Property : AllProperties)
+	{
+		SettingsCategory.AddProperty(Property);
+	}
+}
+
+TWeakPtr<FPCGEditor> FPCGCustomHLSLSettingsDetails::GetPCGEditor(const TArray<TWeakObjectPtr<UObject>>& InObjectsBeingCustomized)
+{
+	if (InObjectsBeingCustomized.IsEmpty() || !InObjectsBeingCustomized[0].IsValid())
+	{
+		return nullptr;
+	}
+
+	const UPCGSettings* FirstSettings = Cast<UPCGSettings>(InObjectsBeingCustomized[0].Get());
+	UPCGNode* Node = Cast<UPCGNode>(FirstSettings ? FirstSettings->GetOuter() : nullptr);
+	UPCGGraph* Graph = Cast<UPCGGraph>(Node ? Node->GetOuter() : nullptr);
+	const UPCGEditorGraph* EditorGraph = Graph ? FPCGEditor::GetPCGEditorGraph(Graph) : nullptr;
+
+	return EditorGraph ? EditorGraph->GetEditor() : nullptr;
+}
+
+#undef LOCTEXT_NAMESPACE
diff --git a/Engine/Plugins/PCG/Source/PCGEditor/Private/PCGEditor.cpp b/Engine/Plugins/PCG/Source/PCGEditor/Private/PCGEditor.cpp
--- a/Engine/Plugins/PCG/Source/PCGEditor/Private/PCGEditor.cpp
+++ b/Engine/Plugins/PCG/Source/PCGEditor/Private/PCGEditor.cpp
@@ -1677,6 +1677,19 @@ bool FPCGEditor::CanConvertRerouteToNamedReroute() const
 	return false;
 }
 
+bool FPCGEditor::IsNodeSourceEditorTabClosed() const
+{
+	return !TabManager.IsValid() || !TabManager->FindExistingLiveTab(FPCGEditor_private::NodeSourceID).IsValid();
+}
+
+void FPCGEditor::SpawnNodeSourceEditorTab()
+{
+	if (TabManager.IsValid() && !TabManager->FindExistingLiveTab(FPCGEditor_private::NodeSourceID))
+	{
+		TabManager->TryInvokeTab(FPCGEditor_private::NodeSourceID);
+	}
+}
+
 void FPCGEditor::OnCollapseNodesInSubgraph()
 {
 	if (!InternalValidationOnAction())
diff --git a/Engine/Plugins/PCG/Source/PCGEditor/Private/PCGEditor.h b/Engine/Plugins/PCG/Source/PCGEditor/Private/PCGEditor.h
--- a/Engine/Plugins/PCG/Source/PCGEditor/Private/PCGEditor.h
+++ b/Engine/Plugins/PCG/Source/PCGEditor/Private/PCGEditor.h
@@ -97,6 +97,12 @@ public:
 	virtual void InitToolMenuContext(FToolMenuContext& MenuContext) override;
 	// ~End FAssetEditorToolkit interface
 
+	/** Returns true if the Node Source Editor tab is not currently open. */
+	bool IsNodeSourceEditorTabClosed() const;
+
+	/** Opens the Node Source Editor tab if it is not already visible. */
+	void SpawnNodeSourceEditorTab();
+
 	FOnInspectedStackChanged OnInspectedStackChangedDelegate;
 
 protected:
diff --git a/Engine/Plugins/PCG/Source/PCGEditor/Private/PCGEditorModule.cpp b/Engine/Plugins/PCG/Source/PCGEditor/Private/PCGEditorModule.cpp
--- a/Engine/Plugins/PCG/Source/PCGEditor/Private/PCGEditorModule.cpp
+++ b/Engine/Plugins/PCG/Source/PCGEditor/Private/PCGEditorModule.cpp
@@ -8,11 +8,9 @@
 #include "PCGParamData.h"
 #include "PCGSubsystem.h"
 #include "PCGVolumeFactory.h"
+#include "Compute/Elements/PCGCustomHLSL.h"
 #include "Data/PCGSpatialData.h"
 #include "Data/PCGSplineData.h"
-#include "DataVisualizations/PCGParamDataVisualization.h"
-#include "DataVisualizations/PCGSpatialDataVisualization.h"
-#include "DataVisualizations/PCGSplineDataVisualization.h"
 #include "Grid/PCGPartitionActor.h"
 #include "WorldPartitionBuilder/PCGWorldPartitionBuilder.h"
 
@@ -23,7 +21,18 @@
 #include "PCGEditorSettings.h"
 #include "PCGEditorStyle.h"
 #include "PCGEditorUtils.h"
-#include "Widgets/SPCGNodeSourceTextBox.h"
+#include "DataVisualizations/PCGParamDataVisualization.h"
+#include "DataVisualizations/PCGSpatialDataVisualization.h"
+#include "DataVisualizations/PCGSplineDataVisualization.h"
+#include "Details/PCGAttributePropertySelectorDetails.h"
+#include "Details/PCGBlueprintSettingsDetails.h"
+#include "Details/PCGComponentDetails.h"
+#include "Details/PCGCustomHLSLSettingsDetails.h"
+#include "Details/PCGEditableUserParameterDetails.h"
+#include "Details/PCGGraphDetails.h"
+#include "Details/PCGGraphInstanceDetails.h"
+#include "Details/PCGInstancedPropertyBagOverrideDetails.h"
+#include "Details/PCGVolumeDetails.h"
 
 #include "ContentBrowserMenuContexts.h"
 #include "ContentBrowserModule.h"
@@ -33,23 +42,16 @@
 #include "EditorModes.h"
 #include "IContentBrowserSingleton.h"
 #include "ISettingsModule.h"
+#include "LevelEditor.h"
 #include "LevelEditorMenuContext.h"
 #include "PropertyEditorModule.h"
 #include "ScopedTransaction.h"
 #include "ToolMenus.h"
-
 #include "Details/EnumSelectorDetails.h"
-#include "Details/PCGAttributePropertySelectorDetails.h"
-#include "Details/PCGBlueprintSettingsDetails.h"
-#include "Details/PCGEditableUserParameterDetails.h"
-#include "Details/PCGGraphDetails.h"
-#include "Details/PCGGraphInstanceDetails.h"
-#include "Details/PCGInstancedPropertyBagOverrideDetails.h"
-#include "Details/PCGVolumeDetails.h"
-#include "Framework/MultiBox/MultiBoxBuilder.h"
-#include "LevelEditor.h"
 #include "Editor/SceneOutliner/Public/ISceneOutliner.h"
+#include "Framework/MultiBox/MultiBoxBuilder.h"
 #include "Framework/Notifications/NotificationManager.h"
+#include "Widgets/SPCGNodeSourceTextBox.h"
 
 #define LOCTEXT_NAMESPACE "FPCGEditorModule"
 
@@ -232,6 +234,7 @@ void FPCGEditorModule::RegisterDetailsCustomizations()
 	PropertyEditor.RegisterCustomClassLayout("PCGGraphInstance", FOnGetDetailCustomizationInstance::CreateStatic(&FPCGGraphInstanceDetails::MakeInstance));
 	PropertyEditor.RegisterCustomClassLayout("PCGVolume", FOnGetDetailCustomizationInstance::CreateStatic(&FPCGVolumeDetails::MakeInstance));
 	PropertyEditor.RegisterCustomClassLayout("PCGUserParameterGetSettings", FOnGetDetailCustomizationInstance::CreateStatic(&FPCGEditableUserParameterDetails::MakeInstance));
+	PropertyEditor.RegisterCustomClassLayout(UPCGCustomHLSLSettings::StaticClass()->GetFName(), FOnGetDetailCustomizationInstance::CreateStatic(&FPCGCustomHLSLSettingsDetails::MakeInstance));
 
 	PropertyEditor.RegisterCustomPropertyTypeLayout("PCGAttributePropertySelector", FOnGetPropertyTypeCustomizationInstance::CreateStatic(&FPCGAttributePropertySelectorDetails::MakeInstance));
 	PropertyEditor.RegisterCustomPropertyTypeLayout("PCGAttributePropertyInputSelector", FOnGetPropertyTypeCustomizationInstance::CreateStatic(&FPCGAttributePropertySelectorDetails::MakeInstance));
diff --git a/Engine/Plugins/PCG/Source/PCGEditor/Private/Widgets/SPCGEditorNodeSource.cpp b/Engine/Plugins/PCG/Source/PCGEditor/Private/Widgets/SPCGEditorNodeSource.cpp
--- a/Engine/Plugins/PCG/Source/PCGEditor/Private/Widgets/SPCGEditorNodeSource.cpp
+++ b/Engine/Plugins/PCG/Source/PCGEditor/Private/Widgets/SPCGEditorNodeSource.cpp
@@ -71,7 +71,7 @@ void SPCGEditorNodeSource::Construct(const FArguments& InArgs, TSharedPtr<FPCGEd
 			+SVerticalBox::Slot()
 			.AutoHeight()
 			[
-				ConstructNonExpandableHeaderWidget(ExpandableAreaArgs.AreaTitle(LOCTEXT("PCGNodeSource_ShaderText_Title", "Shader Text")))
+				ConstructNonExpandableHeaderWidget(ExpandableAreaArgs.AreaTitle(LOCTEXT("PCGNodeSource_ShaderText_Title", "Shader Source")))
 			]
 			+SVerticalBox::Slot()
 			.FillHeight(1.0f)
diff --git a/Engine/Plugins/PCG/Source/PCGEditor/Public/Details/PCGCustomHLSLSettingsDetails.h b/Engine/Plugins/PCG/Source/PCGEditor/Public/Details/PCGCustomHLSLSettingsDetails.h
new file mode 100644
--- /dev/null
+++ b/Engine/Plugins/PCG/Source/PCGEditor/Public/Details/PCGCustomHLSLSettingsDetails.h
@@ -0,0 +1,25 @@
+// Copyright Epic Games, Inc. All Rights Reserved.
+
+#pragma once
+
+#include "CoreMinimal.h"
+#include "IDetailCustomization.h"
+
+class FPCGEditor;
+
+/**
+ * Implements a details view customization for the Custom HLSL node settings.
+ */
+class FPCGCustomHLSLSettingsDetails : public IDetailCustomization
+{
+public:
+	static TSharedRef<IDetailCustomization> MakeInstance()
+	{
+		return MakeShared<FPCGCustomHLSLSettingsDetails>();
+	}
+
+	virtual void CustomizeDetails(IDetailLayoutBuilder& DetailBuilder) override;
+
+private:
+	TWeakPtr<FPCGEditor> GetPCGEditor(const TArray<TWeakObjectPtr<UObject>>& InObjectsBeingCustomized);
+};
