[PCG] Finalize shader API

* Add GetNumElements(DataIndex) to declarations (and finalize its name)
* Change type of AttributeId everywhere to int which matches CPU

#rnx

diff --git a/Engine/Plugins/PCG/Shaders/Private/PCGDataCollectionDataInterface.ush b/Engine/Plugins/PCG/Shaders/Private/PCGDataCollectionDataInterface.ush
--- a/Engine/Plugins/PCG/Shaders/Private/PCGDataCollectionDataInterface.ush
+++ b/Engine/Plugins/PCG/Shaders/Private/PCGDataCollectionDataInterface.ush
@@ -47,7 +47,7 @@ uint GetNumData_{DataInterfaceName}()
 	return {DataInterfaceName}_DataCollectionBuffer.Load(0) & ~KERNEL_EXECUTED_FLAG;
 }
 
-uint GetDataNumElements_{DataInterfaceName}(uint InDataIndex)
+uint GetNumElements_{DataInterfaceName}(uint InDataIndex)
 {
 	const uint ReadAddress =
 		DATA_COLLECTION_HEADER_SIZE_BYTES +
@@ -106,7 +106,7 @@ void StoreBufferInternal_{DataInterfaceName}(uint Address, uint Value)
 	return {DataInterfaceName}_DataCollectionBuffer.Store(Address, Value);
 }
 
-uint GetFirstElementAddressInternal_{DataInterfaceName}(uint InDataIndex, uint InAttributeId)
+uint GetFirstElementAddressInternal_{DataInterfaceName}(uint InDataIndex, int InAttributeId)
 {
 	const uint ReadAddress =
 		DATA_COLLECTION_HEADER_SIZE_BYTES +
@@ -118,7 +118,7 @@ uint GetFirstElementAddressInternal_{DataInterfaceName}(uint InDataIndex, uint I
 	return {DataInterfaceName}_DataCollectionBuffer.Load(ReadAddress);
 }
 
-uint GetAttributeStrideInternal_{DataInterfaceName}(uint InDataIndex, uint InAttributeId)
+uint GetAttributeStrideInternal_{DataInterfaceName}(uint InDataIndex, int InAttributeId)
 {
 	const uint ReadAddress =
 		DATA_COLLECTION_HEADER_SIZE_BYTES +
@@ -130,7 +130,7 @@ uint GetAttributeStrideInternal_{DataInterfaceName}(uint InDataIndex, uint InAtt
 	return PackedAttributeIdAndStride & (0xFFFFFFFF >> 24);
 }
 
-uint GetElementAddressInternal_{DataInterfaceName}(uint InDataIndex, uint InElementIndex, uint InAttributeId)
+uint GetElementAddressInternal_{DataInterfaceName}(uint InDataIndex, uint InElementIndex, int InAttributeId)
 {
 	const uint FirstElementAddress = GetFirstElementAddressInternal_{DataInterfaceName}(InDataIndex, InAttributeId);
 	const uint AttributeStride = GetAttributeStrideInternal_{DataInterfaceName}(InDataIndex, InAttributeId);
@@ -156,61 +156,61 @@ void SetAsExecutedInternal_{DataInterfaceName}()
 
 // #################### ATTRIBUTE GETTERS ##########################
 
-bool GetBool_{DataInterfaceName}(uint DataIndex, uint ElementIndex, uint AttributeId)
+bool GetBool_{DataInterfaceName}(uint DataIndex, uint ElementIndex, int AttributeId)
 {
 	const uint AttributeAddress = GetFirstElementAddressInternal_{DataInterfaceName}(DataIndex, AttributeId);
 	return asbool({DataInterfaceName}_DataCollectionBuffer.Load(AttributeAddress + ElementIndex * 1 * 4));
 }
 
-int GetInt_{DataInterfaceName}(uint DataIndex, uint ElementIndex, uint AttributeId)
+int GetInt_{DataInterfaceName}(uint DataIndex, uint ElementIndex, int AttributeId)
 {
 	const uint AttributeAddress = GetFirstElementAddressInternal_{DataInterfaceName}(DataIndex, AttributeId);
 	return asint({DataInterfaceName}_DataCollectionBuffer.Load(AttributeAddress + ElementIndex * 1 * 4));
 }
 
-int GetUint_{DataInterfaceName}(uint DataIndex, uint ElementIndex, uint AttributeId)
+int GetUint_{DataInterfaceName}(uint DataIndex, uint ElementIndex, int AttributeId)
 {
 	const uint AttributeAddress = GetFirstElementAddressInternal_{DataInterfaceName}(DataIndex, AttributeId);
 	return {DataInterfaceName}_DataCollectionBuffer.Load(AttributeAddress + ElementIndex * 1 * 4);
 }
 
-float GetFloat_{DataInterfaceName}(uint DataIndex, uint ElementIndex, uint AttributeId)
+float GetFloat_{DataInterfaceName}(uint DataIndex, uint ElementIndex, int AttributeId)
 {
 	const uint AttributeAddress = GetFirstElementAddressInternal_{DataInterfaceName}(DataIndex, AttributeId);
 	return asfloat({DataInterfaceName}_DataCollectionBuffer.Load(AttributeAddress + ElementIndex * 1 * 4));
 }
 
-float2 GetFloat2_{DataInterfaceName}(uint DataIndex, uint ElementIndex, uint AttributeId)
+float2 GetFloat2_{DataInterfaceName}(uint DataIndex, uint ElementIndex, int AttributeId)
 {
 	const uint AttributeAddress = GetFirstElementAddressInternal_{DataInterfaceName}(DataIndex, AttributeId);
 	return asfloat({DataInterfaceName}_DataCollectionBuffer.Load2(AttributeAddress + ElementIndex * 2 * 4));
 }
 
-float3 GetFloat3_{DataInterfaceName}(uint DataIndex, uint ElementIndex, uint AttributeId)
+float3 GetFloat3_{DataInterfaceName}(uint DataIndex, uint ElementIndex, int AttributeId)
 {
 	const uint AttributeAddress = GetFirstElementAddressInternal_{DataInterfaceName}(DataIndex, AttributeId);
 	return asfloat({DataInterfaceName}_DataCollectionBuffer.Load3(AttributeAddress + ElementIndex * 3 * 4));
 }
 
-float4 GetFloat4_{DataInterfaceName}(uint DataIndex, uint ElementIndex, uint AttributeId)
+float4 GetFloat4_{DataInterfaceName}(uint DataIndex, uint ElementIndex, int AttributeId)
 {
 	const uint AttributeAddress = GetFirstElementAddressInternal_{DataInterfaceName}(DataIndex, AttributeId);
 	return asfloat({DataInterfaceName}_DataCollectionBuffer.Load4(AttributeAddress + ElementIndex * 4 * 4));
 }
 
-float3 GetRotator_{DataInterfaceName}(uint DataIndex, uint ElementIndex, uint AttributeId)
+float3 GetRotator_{DataInterfaceName}(uint DataIndex, uint ElementIndex, int AttributeId)
 {
 	const uint AttributeAddress = GetFirstElementAddressInternal_{DataInterfaceName}(DataIndex, AttributeId);
 	return asfloat({DataInterfaceName}_DataCollectionBuffer.Load3(AttributeAddress + ElementIndex * 3 * 4));
 }
 
-float4 GetQuat_{DataInterfaceName}(uint DataIndex, uint ElementIndex, uint AttributeId)
+float4 GetQuat_{DataInterfaceName}(uint DataIndex, uint ElementIndex, int AttributeId)
 {
 	const uint AttributeAddress = GetFirstElementAddressInternal_{DataInterfaceName}(DataIndex, AttributeId);
 	return asfloat({DataInterfaceName}_DataCollectionBuffer.Load4(AttributeAddress + ElementIndex * 4 * 4));
 }
 
-int GetStringKey_{DataInterfaceName}(uint DataIndex, uint ElementIndex, uint AttributeId)
+int GetStringKey_{DataInterfaceName}(uint DataIndex, uint ElementIndex, int AttributeId)
 {
 	// String keys are represented as ints.
 	return GetInt_{DataInterfaceName}(DataIndex, ElementIndex, AttributeId);
@@ -218,61 +218,61 @@ int GetStringKey_{DataInterfaceName}(uint DataIndex, uint ElementIndex, uint Att
 
 // #################### ATTRIBUTE SETTERS ##########################
 
-void SetBool_{DataInterfaceName}(uint DataIndex, uint ElementIndex, uint AttributeId, bool Value)
+void SetBool_{DataInterfaceName}(uint DataIndex, uint ElementIndex, int AttributeId, bool Value)
 {
 	const uint AttributeAddress = GetFirstElementAddressInternal_{DataInterfaceName}(DataIndex, AttributeId);
 	{DataInterfaceName}_DataCollectionBuffer.Store(AttributeAddress + ElementIndex * 1 * 4, asuint(Value));
 }
 
-void SetInt_{DataInterfaceName}(uint DataIndex, uint ElementIndex, uint AttributeId, int Value)
+void SetInt_{DataInterfaceName}(uint DataIndex, uint ElementIndex, int AttributeId, int Value)
 {
 	const uint AttributeAddress = GetFirstElementAddressInternal_{DataInterfaceName}(DataIndex, AttributeId);
 	{DataInterfaceName}_DataCollectionBuffer.Store(AttributeAddress + ElementIndex * 1 * 4, asuint(Value));
 }
 
-void SetUint_{DataInterfaceName}(uint DataIndex, uint ElementIndex, uint AttributeId, uint Value)
+void SetUint_{DataInterfaceName}(uint DataIndex, uint ElementIndex, int AttributeId, uint Value)
 {
 	const uint AttributeAddress = GetFirstElementAddressInternal_{DataInterfaceName}(DataIndex, AttributeId);
 	{DataInterfaceName}_DataCollectionBuffer.Store(AttributeAddress + ElementIndex * 1 * 4, Value);
 }
 
-void SetFloat_{DataInterfaceName}(uint DataIndex, uint ElementIndex, uint AttributeId, float Value)
+void SetFloat_{DataInterfaceName}(uint DataIndex, uint ElementIndex, int AttributeId, float Value)
 {
 	const uint AttributeAddress = GetFirstElementAddressInternal_{DataInterfaceName}(DataIndex, AttributeId);
 	{DataInterfaceName}_DataCollectionBuffer.Store(AttributeAddress + ElementIndex * 1 * 4, asuint(Value));
 }
 
-void SetFloat2_{DataInterfaceName}(uint DataIndex, uint ElementIndex, uint AttributeId, float2 Value)
+void SetFloat2_{DataInterfaceName}(uint DataIndex, uint ElementIndex, int AttributeId, float2 Value)
 {
 	const uint AttributeAddress = GetFirstElementAddressInternal_{DataInterfaceName}(DataIndex, AttributeId);
 	{DataInterfaceName}_DataCollectionBuffer.Store2(AttributeAddress + ElementIndex * 2 * 4, asuint(Value));
 }
 
-void SetFloat3_{DataInterfaceName}(uint DataIndex, uint ElementIndex, uint AttributeId, float3 Value)
+void SetFloat3_{DataInterfaceName}(uint DataIndex, uint ElementIndex, int AttributeId, float3 Value)
 {
 	const uint AttributeAddress = GetFirstElementAddressInternal_{DataInterfaceName}(DataIndex, AttributeId);
 	{DataInterfaceName}_DataCollectionBuffer.Store3(AttributeAddress + ElementIndex * 3 * 4, asuint(Value));
 }
 
-void SetFloat4_{DataInterfaceName}(uint DataIndex, uint ElementIndex, uint AttributeId, float4 Value)
+void SetFloat4_{DataInterfaceName}(uint DataIndex, uint ElementIndex, int AttributeId, float4 Value)
 {
 	const uint AttributeAddress = GetFirstElementAddressInternal_{DataInterfaceName}(DataIndex, AttributeId);
 	{DataInterfaceName}_DataCollectionBuffer.Store4(AttributeAddress + ElementIndex * 4 * 4, asuint(Value));
 }
 
-void SetRotator_{DataInterfaceName}(uint DataIndex, uint ElementIndex, uint AttributeId, float3 Value)
+void SetRotator_{DataInterfaceName}(uint DataIndex, uint ElementIndex, int AttributeId, float3 Value)
 {
 	const uint AttributeAddress = GetFirstElementAddressInternal_{DataInterfaceName}(DataIndex, AttributeId);
 	{DataInterfaceName}_DataCollectionBuffer.Store3(AttributeAddress + ElementIndex * 3 * 4, asuint(Value));
 }
 
-void SetQuat_{DataInterfaceName}(uint DataIndex, uint ElementIndex, uint AttributeId, float4 Value)
+void SetQuat_{DataInterfaceName}(uint DataIndex, uint ElementIndex, int AttributeId, float4 Value)
 {
 	const uint AttributeAddress = GetFirstElementAddressInternal_{DataInterfaceName}(DataIndex, AttributeId);
 	{DataInterfaceName}_DataCollectionBuffer.Store4(AttributeAddress + ElementIndex * 4 * 4, asuint(Value));
 }
 
-void SetStringKey_{DataInterfaceName}(uint DataIndex, uint ElementIndex, uint AttributeId, int Value)
+void SetStringKey_{DataInterfaceName}(uint DataIndex, uint ElementIndex, int AttributeId, int Value)
 {
 	// String keys are represented as ints.
 	SetInt_{DataInterfaceName}(DataIndex, ElementIndex, AttributeId, Value);
diff --git a/Engine/Plugins/PCG/Source/PCG/Private/Compute/DataInterfaces/Elements/PCGStaticMeshSpawnerDataInterface.cpp b/Engine/Plugins/PCG/Source/PCG/Private/Compute/DataInterfaces/Elements/PCGStaticMeshSpawnerDataInterface.cpp
--- a/Engine/Plugins/PCG/Source/PCG/Private/Compute/DataInterfaces/Elements/PCGStaticMeshSpawnerDataInterface.cpp
+++ b/Engine/Plugins/PCG/Source/PCG/Private/Compute/DataInterfaces/Elements/PCGStaticMeshSpawnerDataInterface.cpp
@@ -58,8 +58,8 @@ BEGIN_SHADER_PARAMETER_STRUCT(FPCGStaticMeshSpawnerDataInterfaceParameters,)
 	SHADER_PARAMETER_SCALAR_ARRAY(float, SelectionCDF, [PCGComputeConstants::MAX_PRIMITIVE_COMPONENTS_PER_SPAWNER])
 	SHADER_PARAMETER(uint32, NumAttributes)
 	SHADER_PARAMETER(uint32, NumPrimitives)
-	SHADER_PARAMETER(uint32, SelectorAttributeId)
-	SHADER_PARAMETER(uint32, SelectedMeshAttributeId)
+	SHADER_PARAMETER(int32, SelectorAttributeId)
+	SHADER_PARAMETER(int32, SelectedMeshAttributeId)
 END_SHADER_PARAMETER_STRUCT()
 
 void UPCGStaticMeshSpawnerDataInterface::GetShaderParameters(TCHAR const* UID, FShaderParametersMetadataBuilder& InOutBuilder, FShaderParametersMetadataAllocations& InOutAllocations) const
@@ -77,8 +77,8 @@ void UPCGStaticMeshSpawnerDataInterface::GetHLSL(FString& OutHLSL, FString const
 	};
 
 	OutHLSL += FString::Format(TEXT(
-		"uint {DataInterfaceName}_SelectorAttributeId;\n"
-		"uint SMSpawner_GetSelectorAttributeId_{DataInterfaceName}() { return {DataInterfaceName}_SelectorAttributeId; }\n"
+		"int {DataInterfaceName}_SelectorAttributeId;\n"
+		"int SMSpawner_GetSelectorAttributeId_{DataInterfaceName}() { return {DataInterfaceName}_SelectorAttributeId; }\n"
 		"\n"
 		"uint {DataInterfaceName}_NumAttributes;\n"
 		"uint4 {DataInterfaceName}_AttributeIdOffsetStrides[{MaxAttributes}];\n"
@@ -110,8 +110,8 @@ void UPCGStaticMeshSpawnerDataInterface::GetHLSL(FString& OutHLSL, FString const
 		"\n"
 		"float SMSpawner_GetPrimitiveSelectionCDF_{DataInterfaceName}(uint InPrimitiveIndex) { return GET_SCALAR_ARRAY_ELEMENT({DataInterfaceName}_SelectionCDF, InPrimitiveIndex); }\n"
 		"\n"
-		"uint {DataInterfaceName}_SelectedMeshAttributeId;\n"
-		"uint SMSpawner_GetSelectedMeshAttributeId_{DataInterfaceName}() { return {DataInterfaceName}_SelectedMeshAttributeId; }\n"
+		"int {DataInterfaceName}_SelectedMeshAttributeId;\n"
+		"int SMSpawner_GetSelectedMeshAttributeId_{DataInterfaceName}() { return {DataInterfaceName}_SelectedMeshAttributeId; }\n"
 		), TemplateArgs);
 }
 
@@ -162,8 +162,8 @@ void FPCGStaticMeshSpawnerDataProviderProxy::GatherDispatchData(FDispatchData co
 
 		Parameters.NumAttributes = AttributeIdOffsetStrides.Num();
 		Parameters.NumPrimitives = SelectionCDF.Num();
-		Parameters.SelectorAttributeId = static_cast<uint32>(SelectorAttributeId);
-		Parameters.SelectedMeshAttributeId = static_cast<uint32>(SelectedMeshAttributeId);
+		Parameters.SelectorAttributeId = SelectorAttributeId;
+		Parameters.SelectedMeshAttributeId = SelectedMeshAttributeId;
 
 		for (int32 Index = 0; Index < AttributeIdOffsetStrides.Num(); ++Index)
 		{
diff --git a/Engine/Plugins/PCG/Source/PCG/Private/Compute/Elements/PCGCustomHLSL.cpp b/Engine/Plugins/PCG/Source/PCG/Private/Compute/Elements/PCGCustomHLSL.cpp
--- a/Engine/Plugins/PCG/Source/PCG/Private/Compute/Elements/PCGCustomHLSL.cpp
+++ b/Engine/Plugins/PCG/Source/PCG/Private/Compute/Elements/PCGCustomHLSL.cpp
@@ -720,11 +720,11 @@ void UPCGCustomHLSLSettings::UpdateInputDeclarations()
 		InputDeclarations += FString::Format(TEXT(
 			"uint {0}_GetNumData();\n"
 			"uint {0}_GetNumElements();\n"
-			"uint {0}_GetDataNumElements(uint DataIndex);\n"
+			"uint {0}_GetNumElements(uint DataIndex);\n"
 			"\n"
 			"// Valid types: bool, int, float, float2, float3, float4, Rotator (float3), Quat (float4), Transform (float4x4), StringKey (int)\n"
 			"\n"
-			"<type> {0}_Get<type>(uint DataIndex, uint ElementIndex, uint AttributeId);\n"
+			"<type> {0}_Get<type>(uint DataIndex, uint ElementIndex, int AttributeId);\n"
 			"<type> {0}_Get<type>(uint DataIndex, uint ElementIndex, 'AttributeName');\n"),
 			{ bMultiPin ? PCGHLSLElement::PinDeclTemplateStr : DataCollectionDataPins[0] });
 
@@ -860,7 +860,7 @@ void UPCGCustomHLSLSettings::UpdateOutputDeclarations()
 		OutputDeclarations += FString::Format(TEXT(
 			"// Valid types: bool, int, float, float2, float3, float4, Rotator (float3), Quat (float4), Transform (float4x4), StringKey (int)\n"
 			"\n"
-			"void {0}_Set<type>(uint DataIndex, uint ElementIndex, uint AttributeId, <type> Value);\n"
+			"void {0}_Set<type>(uint DataIndex, uint ElementIndex, int AttributeId, <type> Value);\n"
 			"void {0}_Set<type>(uint DataIndex, uint ElementIndex, 'AttributeName', <type> Value);\n"),
 			{ bMultiPin ? PCGHLSLElement::PinDeclTemplateStr : DataCollectionDataPins[0] });
 
