[PCG] GPU SM Spawner - fix error from asset loading code when there is a validation error and the validation logic runs on a worker thread and attempts to load the static mesh.

#rnx

diff --git a/Engine/Plugins/PCG/Source/PCG/Private/Elements/PCGStaticMeshSpawner.cpp b/Engine/Plugins/PCG/Source/PCG/Private/Elements/PCGStaticMeshSpawner.cpp
--- a/Engine/Plugins/PCG/Source/PCG/Private/Elements/PCGStaticMeshSpawner.cpp
+++ b/Engine/Plugins/PCG/Source/PCG/Private/Elements/PCGStaticMeshSpawner.cpp
@@ -593,7 +593,18 @@ bool FPCGStaticMeshSpawnerElement::CanExecuteOnlyOnMainThread(FPCGContext* Conte
 	// PrepareData can call UPCGManagedComponent::MarkAsReused which registers the ISMC, which can go into Chaos code that asserts if not on main thread.
 	// TODO: We can likely re-enable multi-threading for PrepareData if we move the call to MarkAsReused to Execute. There should hopefully not be
 	// wider contention on resources resources are not shared across nodes and are also per-component.
-	return !Context || Context->CurrentPhase == EPCGExecutionPhase::Execute || Context->CurrentPhase == EPCGExecutionPhase::PrepareData;
+	bool bRequiresMainThread = !Context || Context->CurrentPhase == EPCGExecutionPhase::Execute || Context->CurrentPhase == EPCGExecutionPhase::PrepareData;
+
+	if (!bRequiresMainThread && Context->GetInputSettings<UPCGStaticMeshSpawnerSettings>())
+	{
+		// TODO: Unfortunate result of convoluted logic. If the spawner is set to GPU, but the validation failed (e.g. static mesh is not compatible),
+		// then the compute graph is not compiled and the CPU spawner task is left in the graph, and the spawner validation will be executed from the element
+		// so that graph warnings and errors can be logged. This validation loads the static mesh to check its settings and this must be done on the main
+		// thread. In the future we may want to invite nodes to log warnings/errors at graph schedule time, and always remove CPU tasks for GPU nodes.
+		bRequiresMainThread |= Context->GetInputSettings<UPCGStaticMeshSpawnerSettings>()->ShouldExecuteOnGPU();
+	}
+
+	return bRequiresMainThread;
 }
 
 void FPCGStaticMeshSpawnerElement::SpawnStaticMeshInstances(FPCGStaticMeshSpawnerContext* Context, const FPCGMeshInstanceList& InstanceList, AActor* TargetActor, const FPCGPackedCustomData* InPackedCustomData) const
