[PCG] Fix uninitialized members

#rnx

diff --git a/Engine/Plugins/PCG/Source/PCG/Private/Components/PCGProceduralISMComponent.cpp b/Engine/Plugins/PCG/Source/PCG/Private/Components/PCGProceduralISMComponent.cpp
--- a/Engine/Plugins/PCG/Source/PCG/Private/Components/PCGProceduralISMComponent.cpp
+++ b/Engine/Plugins/PCG/Source/PCG/Private/Components/PCGProceduralISMComponent.cpp
@@ -105,6 +105,11 @@ UPCGManagedProceduralISMComponent* PCGManagedProceduralISMComponent::GetOrCreate
 	return Resource;
 }
 
+FPCGProceduralISMComponentDescriptor::FPCGProceduralISMComponentDescriptor()
+{
+	*this = FPCGSoftISMComponentDescriptor();
+}
+
 FPCGProceduralISMComponentDescriptor& FPCGProceduralISMComponentDescriptor::operator=(const FPCGSoftISMComponentDescriptor& Other)
 {
 	InstanceStartCullDistance = Other.InstanceStartCullDistance;
diff --git a/Engine/Plugins/PCG/Source/PCG/Private/Components/PCGProceduralISMComponent.h b/Engine/Plugins/PCG/Source/PCG/Private/Components/PCGProceduralISMComponent.h
--- a/Engine/Plugins/PCG/Source/PCG/Private/Components/PCGProceduralISMComponent.h
+++ b/Engine/Plugins/PCG/Source/PCG/Private/Components/PCGProceduralISMComponent.h
@@ -134,6 +134,7 @@ struct FPCGProceduralISMComponentDescriptor
 	GENERATED_BODY()
 
 public:
+	FPCGProceduralISMComponentDescriptor();
 	virtual ~FPCGProceduralISMComponentDescriptor() = default;
 
 	FPCGProceduralISMComponentDescriptor& operator=(const FPCGSoftISMComponentDescriptor& Other);
@@ -178,19 +179,19 @@ public:
 	TArray<TObjectPtr<URuntimeVirtualTexture>> RuntimeVirtualTextures;
 
 	UPROPERTY()
-	int32 NumInstances;
+	int32 NumInstances = 0;
 
 	UPROPERTY()
-	int32 NumCustomFloats;
+	int32 NumCustomFloats = 0;
 
 	UPROPERTY()
 	FBox LocalBounds;
 
 	UPROPERTY()
-	int32 InstanceStartCullDistance;
+	int32 InstanceStartCullDistance = 0;
 
 	UPROPERTY()
-	int32 InstanceEndCullDistance;
+	int32 InstanceEndCullDistance = 0;
 
 	UPROPERTY()
 	TArray<FName> ComponentTags;
