[PCG] Add support for properties in Instance Data Packer

* Added new accessor keys for points subset

[RN] Deprecated FPCGMeshInstanceList::InstancesMetadateEntry, you'll have to update your BP/code if you used it.

diff --git a/Engine/Plugins/PCG/Source/PCG/Private/InstanceDataPackers/PCGInstanceDataPackerBase.cpp b/Engine/Plugins/PCG/Source/PCG/Private/InstanceDataPackers/PCGInstanceDataPackerBase.cpp
--- a/Engine/Plugins/PCG/Source/PCG/Private/InstanceDataPackers/PCGInstanceDataPackerBase.cpp
+++ b/Engine/Plugins/PCG/Source/PCG/Private/InstanceDataPackers/PCGInstanceDataPackerBase.cpp
@@ -4,6 +4,7 @@
 
 #include "PCGElement.h"
 #include "Elements/Metadata/PCGMetadataElementCommon.h"
+#include "MeshSelectors/PCGMeshSelectorBase.h"
 #include "Metadata/PCGMetadata.h"
 #include "Metadata/Accessors/IPCGAttributeAccessor.h"
 #include "Metadata/Accessors/PCGAttributeAccessorKeys.h"
@@ -63,8 +64,10 @@ void UPCGInstanceDataPackerBase::PackCustomDataFromAttributes(const FPCGMeshInst
 
 void UPCGInstanceDataPackerBase::PackCustomDataFromAttributes(const FPCGMeshInstanceList& InstanceList, const TArray<const FPCGMetadataAttributeBase*>& Attributes, FPCGPackedCustomData& OutPackedCustomData) const
 {
-	for(uint64 PointMetadataEntry : InstanceList.InstancesMetadataEntry)
+	const UPCGPointData* PointData = InstanceList.PointData.Get();
+	for (const int32 Index : InstanceList.InstancesIndices)
 	{
+		const PCGMetadataEntryKey EntryKey = PointData ? PointData->GetPoints()[Index].MetadataEntry : Index;
 		for (const FPCGMetadataAttributeBase* AttributeBase : Attributes)
 		{
 			check(AttributeBase);
@@ -76,7 +79,7 @@ void UPCGInstanceDataPackerBase::PackCustomDataFromAttributes(const FPCGMeshInst
 				const FPCGMetadataAttribute<bool>* Attribute = static_cast<const FPCGMetadataAttribute<bool>*>(AttributeBase);
 				check(Attribute);
 
-				const bool Value = Attribute->GetValueFromItemKey(PointMetadataEntry);
+				const bool Value = Attribute->GetValueFromItemKey(EntryKey);
 				OutPackedCustomData.CustomData.Add(Value);
 				break;
 			}
@@ -85,7 +88,7 @@ void UPCGInstanceDataPackerBase::PackCustomDataFromAttributes(const FPCGMeshInst
 				const FPCGMetadataAttribute<float>* Attribute = static_cast<const FPCGMetadataAttribute<float>*>(AttributeBase);
 				check(Attribute);
 
-				const float Value = Attribute->GetValueFromItemKey(PointMetadataEntry);
+				const float Value = Attribute->GetValueFromItemKey(EntryKey);
 				OutPackedCustomData.CustomData.Add(Value);
 				break;
 			}
@@ -94,7 +97,7 @@ void UPCGInstanceDataPackerBase::PackCustomDataFromAttributes(const FPCGMeshInst
 				const FPCGMetadataAttribute<double>* Attribute = static_cast<const FPCGMetadataAttribute<double>*>(AttributeBase);
 				check(Attribute);
 
-				const double Value = Attribute->GetValueFromItemKey(PointMetadataEntry);
+				const double Value = Attribute->GetValueFromItemKey(EntryKey);
 				OutPackedCustomData.CustomData.Add(Value);
 				break;
 			}
@@ -103,7 +106,7 @@ void UPCGInstanceDataPackerBase::PackCustomDataFromAttributes(const FPCGMeshInst
 				const FPCGMetadataAttribute<int32>* Attribute = static_cast<const FPCGMetadataAttribute<int32>*>(AttributeBase);
 				check(Attribute);
 
-				const float Value = static_cast<float>(Attribute->GetValueFromItemKey(PointMetadataEntry));
+				const float Value = static_cast<float>(Attribute->GetValueFromItemKey(EntryKey));
 				OutPackedCustomData.CustomData.Add(Value);
 				break;
 			}
@@ -112,7 +115,7 @@ void UPCGInstanceDataPackerBase::PackCustomDataFromAttributes(const FPCGMeshInst
 				const FPCGMetadataAttribute<int64>* Attribute = static_cast<const FPCGMetadataAttribute<int64>*>(AttributeBase);
 				check(Attribute);
 
-				const float Value = static_cast<float>(Attribute->GetValueFromItemKey(PointMetadataEntry));
+				const float Value = static_cast<float>(Attribute->GetValueFromItemKey(EntryKey));
 				OutPackedCustomData.CustomData.Add(Value);
 				break;
 			}
@@ -121,7 +124,7 @@ void UPCGInstanceDataPackerBase::PackCustomDataFromAttributes(const FPCGMeshInst
 				const FPCGMetadataAttribute<FRotator>* Attribute = static_cast<const FPCGMetadataAttribute<FRotator>*>(AttributeBase);
 				check(Attribute);
 
-				const FRotator Value = Attribute->GetValueFromItemKey(PointMetadataEntry);
+				const FRotator Value = Attribute->GetValueFromItemKey(EntryKey);
 				OutPackedCustomData.CustomData.Add(Value.Roll);
 				OutPackedCustomData.CustomData.Add(Value.Pitch);
 				OutPackedCustomData.CustomData.Add(Value.Yaw);
@@ -132,7 +135,7 @@ void UPCGInstanceDataPackerBase::PackCustomDataFromAttributes(const FPCGMeshInst
 				const FPCGMetadataAttribute<FVector2D>* Attribute = static_cast<const FPCGMetadataAttribute<FVector2D>*>(AttributeBase);
 				check(Attribute);
 
-				const FVector2D Value = Attribute->GetValueFromItemKey(PointMetadataEntry);
+				const FVector2D Value = Attribute->GetValueFromItemKey(EntryKey);
 				OutPackedCustomData.CustomData.Add(Value.X);
 				OutPackedCustomData.CustomData.Add(Value.Y);
 				break;
@@ -142,7 +145,7 @@ void UPCGInstanceDataPackerBase::PackCustomDataFromAttributes(const FPCGMeshInst
 				const FPCGMetadataAttribute<FVector>* Attribute = static_cast<const FPCGMetadataAttribute<FVector>*>(AttributeBase);
 				check(Attribute);
 
-				const FVector Value = Attribute->GetValueFromItemKey(PointMetadataEntry);
+				const FVector Value = Attribute->GetValueFromItemKey(EntryKey);
 				OutPackedCustomData.CustomData.Add(Value.X);
 				OutPackedCustomData.CustomData.Add(Value.Y);
 				OutPackedCustomData.CustomData.Add(Value.Z);
@@ -153,7 +156,7 @@ void UPCGInstanceDataPackerBase::PackCustomDataFromAttributes(const FPCGMeshInst
 				const FPCGMetadataAttribute<FVector4>* Attribute = static_cast<const FPCGMetadataAttribute<FVector4>*>(AttributeBase);
 				check(Attribute);
 
-				const FVector4 Value = Attribute->GetValueFromItemKey(PointMetadataEntry);
+				const FVector4 Value = Attribute->GetValueFromItemKey(EntryKey);
 				OutPackedCustomData.CustomData.Add(Value.X);
 				OutPackedCustomData.CustomData.Add(Value.Y);
 				OutPackedCustomData.CustomData.Add(Value.Z);
@@ -165,7 +168,7 @@ void UPCGInstanceDataPackerBase::PackCustomDataFromAttributes(const FPCGMeshInst
 				const FPCGMetadataAttribute<FQuat>* Attribute = static_cast<const FPCGMetadataAttribute<FQuat>*>(AttributeBase);
 				check(Attribute);
 
-				const FQuat Value = Attribute->GetValueFromItemKey(PointMetadataEntry);
+				const FQuat Value = Attribute->GetValueFromItemKey(EntryKey);
 				OutPackedCustomData.CustomData.Add(Value.X);
 				OutPackedCustomData.CustomData.Add(Value.Y);
 				OutPackedCustomData.CustomData.Add(Value.Z);
diff --git a/Engine/Plugins/PCG/Source/PCG/Private/InstanceDataPackers/PCGInstanceDataPackerByAttribute.cpp b/Engine/Plugins/PCG/Source/PCG/Private/InstanceDataPackers/PCGInstanceDataPackerByAttribute.cpp
--- a/Engine/Plugins/PCG/Source/PCG/Private/InstanceDataPackers/PCGInstanceDataPackerByAttribute.cpp
+++ b/Engine/Plugins/PCG/Source/PCG/Private/InstanceDataPackers/PCGInstanceDataPackerByAttribute.cpp
@@ -48,7 +48,21 @@ void UPCGInstanceDataPackerByAttribute::PackInstances_Implementation(FPCGContext
 	for (const FPCGAttributePropertyInputSelector& Selector : AttributeSelectors)
 	{
 		TUniquePtr<const IPCGAttributeAccessor> Accessor = PCGAttributeAccessorHelpers::CreateConstAccessor(InSpatialData, Selector);
-		TUniquePtr<const IPCGAttributeAccessorKeys> Keys = MakeUnique<FPCGAttributeAccessorKeysEntries>(InstanceList.InstancesMetadataEntry);
+		TUniquePtr<const IPCGAttributeAccessorKeys> Keys;
+		const UPCGPointData* PointData = InstanceList.PointData.Get();
+		TArray<const PCGMetadataEntryKey> ExtractedKeys;
+		
+		if (InSpatialData == PointData)
+		{
+			Keys = MakeUnique<const FPCGAttributeAccessorKeysPointsSubset>(PointData->GetPoints(), InstanceList.InstancesIndices);
+		}
+		else
+		{
+			// Convert indices to entry keys
+			ExtractedKeys.Reserve(InstanceList.InstancesIndices.Num());
+			Algo::Transform(InstanceList.InstancesIndices, ExtractedKeys, [](const int32 Index) -> PCGMetadataEntryKey{ return Index; });
+			Keys = MakeUnique<const FPCGAttributeAccessorKeysEntries>(ExtractedKeys);
+		}
 
 		if (!Accessor.IsValid() || !Keys.IsValid())
 		{
diff --git a/Engine/Plugins/PCG/Source/PCG/Private/MeshSelectors/PCGMeshSelectorByAttribute.cpp b/Engine/Plugins/PCG/Source/PCG/Private/MeshSelectors/PCGMeshSelectorByAttribute.cpp
--- a/Engine/Plugins/PCG/Source/PCG/Private/MeshSelectors/PCGMeshSelectorByAttribute.cpp
+++ b/Engine/Plugins/PCG/Source/PCG/Private/MeshSelectors/PCGMeshSelectorByAttribute.cpp
@@ -28,6 +28,7 @@ namespace PCGMeshSelectorAttribute
 		TSoftObjectPtr<UStaticMesh> Mesh,
 		const TArray<TSoftObjectPtr<UMaterialInterface>>& MaterialOverrides,
 		bool bReverseCulling,
+		const UPCGPointData* InPointData,
 		const int AttributePartitionIndex = INDEX_NONE)
 	{
 		for (FPCGMeshInstanceList& InstanceList : InstanceLists)
@@ -46,6 +47,7 @@ namespace PCGMeshSelectorAttribute
 		NewInstanceList.Descriptor.OverrideMaterials = MaterialOverrides;
 		NewInstanceList.Descriptor.bReverseCulling = bReverseCulling;
 		NewInstanceList.AttributePartitionIndex = AttributePartitionIndex;
+		NewInstanceList.PointData = InPointData;
 
 		return NewInstanceList;
 	}
@@ -283,33 +285,32 @@ bool UPCGMeshSelectorByAttribute::SelectInstances(
 				Partition.SetNum(WriteIndex);
 			}
 
-			auto AddPointsToInstanceList = [&OutMeshInstances, &CurrentPartitionDescriptor, &MaterialOverrideHelper, &Points, PartitionIndex](const TArray<int32>& PointIndices, bool bReverseTransform)
+			auto AddPointsToInstanceList = [&OutMeshInstances, &CurrentPartitionDescriptor, &MaterialOverrideHelper, &Points, PartitionIndex, InPointData](const TArray<int32>& PointIndices, bool bReverseTransform)
 			{
 				if (MaterialOverrideHelper.OverridesMaterials())
 				{
 					for (int32 PointIndex = 0; PointIndex < PointIndices.Num(); ++PointIndex)
 					{
 						const FPCGPoint& Point = Points[PointIndices[PointIndex]];
-						FPCGMeshInstanceList& InstanceList = PCGMeshSelectorAttribute::GetInstanceList(OutMeshInstances, CurrentPartitionDescriptor, CurrentPartitionDescriptor.StaticMesh, MaterialOverrideHelper.GetMaterialOverrides(Point.MetadataEntry), bReverseTransform, PartitionIndex);
+						FPCGMeshInstanceList& InstanceList = PCGMeshSelectorAttribute::GetInstanceList(OutMeshInstances, CurrentPartitionDescriptor, CurrentPartitionDescriptor.StaticMesh, MaterialOverrideHelper.GetMaterialOverrides(Point.MetadataEntry), bReverseTransform, InPointData, PartitionIndex);
 						InstanceList.Instances.Emplace(Point.Transform);
-						InstanceList.InstancesMetadataEntry.Emplace(Point.MetadataEntry);
+						InstanceList.InstancesIndices.Emplace(PointIndex);
 					}
 				}
 				else
 				{
 					TArray<TSoftObjectPtr<UMaterialInterface>> DummyMaterialList;
-					FPCGMeshInstanceList& InstanceList = PCGMeshSelectorAttribute::GetInstanceList(OutMeshInstances, CurrentPartitionDescriptor, CurrentPartitionDescriptor.StaticMesh, DummyMaterialList, bReverseTransform, PartitionIndex);
+					FPCGMeshInstanceList& InstanceList = PCGMeshSelectorAttribute::GetInstanceList(OutMeshInstances, CurrentPartitionDescriptor, CurrentPartitionDescriptor.StaticMesh, DummyMaterialList, bReverseTransform, InPointData, PartitionIndex);
 
-					check(InstanceList.Instances.Num() == InstanceList.InstancesMetadataEntry.Num());
+					check(InstanceList.Instances.Num() == InstanceList.InstancesIndices.Num());
 					const int32 InstanceOffset = InstanceList.Instances.Num();
 					InstanceList.Instances.SetNum(InstanceOffset + PointIndices.Num());
-					InstanceList.InstancesMetadataEntry.SetNum(InstanceOffset + PointIndices.Num());
+					InstanceList.InstancesIndices.Append(PointIndices);
 
 					for(int32 PointIndex = 0; PointIndex < PointIndices.Num(); ++PointIndex)
 					{
 						const FPCGPoint& Point = Points[PointIndices[PointIndex]];
 						InstanceList.Instances[InstanceOffset + PointIndex] = Point.Transform;
-						InstanceList.InstancesMetadataEntry[InstanceOffset + PointIndex] = Point.MetadataEntry;
 					}
 				}
 			};
diff --git a/Engine/Plugins/PCG/Source/PCG/Private/MeshSelectors/PCGMeshSelectorWeighted.cpp b/Engine/Plugins/PCG/Source/PCG/Private/MeshSelectors/PCGMeshSelectorWeighted.cpp
--- a/Engine/Plugins/PCG/Source/PCG/Private/MeshSelectors/PCGMeshSelectorWeighted.cpp
+++ b/Engine/Plugins/PCG/Source/PCG/Private/MeshSelectors/PCGMeshSelectorWeighted.cpp
@@ -22,7 +22,8 @@ namespace PCGMeshSelectorWeighted
 		TArray<FPCGMeshInstanceList>& InstanceLists,
 		bool bUseMaterialOverrides,
 		const TArray<TSoftObjectPtr<UMaterialInterface>>& InMaterialOverrides,
-		bool bInReverseCulling)
+		bool bInReverseCulling,
+		const UPCGPointData* InPointData)
 	{
 		check(InstanceLists.Num() > 0);
 
@@ -42,6 +43,7 @@ namespace PCGMeshSelectorWeighted
 		FPCGMeshInstanceList& NewInstanceList = InstanceLists.Emplace_GetRef();
 		NewInstanceList.Descriptor = InstanceLists[0].Descriptor;
 		NewInstanceList.Descriptor.bReverseCulling = bInReverseCulling;
+		NewInstanceList.PointData = InPointData;
 
 		if (bUseMaterialOverrides)
 		{
@@ -143,7 +145,8 @@ bool UPCGMeshSelectorWeighted::SelectInstances(
 			}
 
 			TArray<FPCGMeshInstanceList>& PickEntry = MeshInstances.Emplace_GetRef();
-			PickEntry.Emplace_GetRef(Entry.Descriptor);
+			FPCGMeshInstanceList& MeshInstanceList = PickEntry.Emplace_GetRef(Entry.Descriptor);
+			MeshInstanceList.PointData = InPointData;
 
 			TotalWeight += Entry.Weight;
 			CumulativeWeights.Add(TotalWeight);
@@ -221,9 +224,9 @@ bool UPCGMeshSelectorWeighted::SelectInstances(
 			if(RandomPick < MeshInstances.Num())
 			{
 				const bool bNeedsReverseCulling = (Point.Transform.GetDeterminant() < 0);
-				FPCGMeshInstanceList& InstanceList = PCGMeshSelectorWeighted::GetInstanceList(MeshInstances[RandomPick], bUseAttributeMaterialOverrides, MaterialOverrideHelper.GetMaterialOverrides(Point.MetadataEntry), bNeedsReverseCulling);
+				FPCGMeshInstanceList& InstanceList = PCGMeshSelectorWeighted::GetInstanceList(MeshInstances[RandomPick], bUseAttributeMaterialOverrides, MaterialOverrideHelper.GetMaterialOverrides(Point.MetadataEntry), bNeedsReverseCulling, InPointData);
 				InstanceList.Instances.Emplace(Point.Transform);
-				InstanceList.InstancesMetadataEntry.Emplace(Point.MetadataEntry);
+				InstanceList.InstancesIndices.Emplace(CurrentPointIndex - 1); // - 1 because it is already incremented.
 
 				const TSoftObjectPtr<UStaticMesh>& Mesh = InstanceList.Descriptor.StaticMesh;
 
diff --git a/Engine/Plugins/PCG/Source/PCG/Private/MeshSelectors/PCGMeshSelectorWeightedByCategory.cpp b/Engine/Plugins/PCG/Source/PCG/Private/MeshSelectors/PCGMeshSelectorWeightedByCategory.cpp
--- a/Engine/Plugins/PCG/Source/PCG/Private/MeshSelectors/PCGMeshSelectorWeightedByCategory.cpp
+++ b/Engine/Plugins/PCG/Source/PCG/Private/MeshSelectors/PCGMeshSelectorWeightedByCategory.cpp
@@ -129,7 +129,8 @@ bool UPCGMeshSelectorWeightedByCategory::SelectInstances(
 				}
 
 				TArray<FPCGMeshInstanceList>& PickEntry = InstancesAndWeights->MeshInstances.Emplace_GetRef();
-				PickEntry.Emplace_GetRef(WeightedEntry.Descriptor);
+				FPCGMeshInstanceList& MeshInstanceList = PickEntry.Emplace_GetRef(WeightedEntry.Descriptor);
+				MeshInstanceList.PointData = InPointData;
 
 				// precompute the weights
 				TotalWeight += WeightedEntry.Weight;
@@ -227,9 +228,9 @@ bool UPCGMeshSelectorWeightedByCategory::SelectInstances(
 		if(RandomPick < InstancesAndWeights->MeshInstances.Num())
 		{
 			const bool bNeedsReverseCulling = (Point.Transform.GetDeterminant() < 0);
-			FPCGMeshInstanceList& InstanceList = PCGMeshSelectorWeighted::GetInstanceList(InstancesAndWeights->MeshInstances[RandomPick], bUseAttributeMaterialOverrides, MaterialOverrideHelper.GetMaterialOverrides(Point.MetadataEntry), bNeedsReverseCulling);
+			FPCGMeshInstanceList& InstanceList = PCGMeshSelectorWeighted::GetInstanceList(InstancesAndWeights->MeshInstances[RandomPick], bUseAttributeMaterialOverrides, MaterialOverrideHelper.GetMaterialOverrides(Point.MetadataEntry), bNeedsReverseCulling, InPointData);
 			InstanceList.Instances.Emplace(Point.Transform);
-			InstanceList.InstancesMetadataEntry.Emplace(Point.MetadataEntry);
+			InstanceList.InstancesIndices.Emplace(CurrentPointIndex - 1); // - 1 because it is already incremented.
 
 			const TSoftObjectPtr<UStaticMesh>& Mesh = InstanceList.Descriptor.StaticMesh;
 
diff --git a/Engine/Plugins/PCG/Source/PCG/Private/Metadata/Accessors/PCGAttributeAccessorKeys.cpp b/Engine/Plugins/PCG/Source/PCG/Private/Metadata/Accessors/PCGAttributeAccessorKeys.cpp
--- a/Engine/Plugins/PCG/Source/PCG/Private/Metadata/Accessors/PCGAttributeAccessorKeys.cpp
+++ b/Engine/Plugins/PCG/Source/PCG/Private/Metadata/Accessors/PCGAttributeAccessorKeys.cpp
@@ -136,4 +136,62 @@ bool FPCGAttributeAccessorKeysPoints::GetMetadataEntryKeys(int32 InStart, TArray
 bool FPCGAttributeAccessorKeysPoints::GetMetadataEntryKeys(int32 InStart, TArrayView<const PCGMetadataEntryKey*> OutEntryKeys) const
 {
 	return PCGAttributeAccessorKeys::GetKeys(Points, InStart, OutEntryKeys, [](const FPCGPoint& Point) -> const PCGMetadataEntryKey* { return &(Point.MetadataEntry); });
+}
+
+////////////////////////////////////////////////////////////////////
+
+FPCGAttributeAccessorKeysPointsSubset::FPCGAttributeAccessorKeysPointsSubset(const TArrayView<FPCGPoint>& InPoints, const TArrayView<const int32>& InPointIndices)
+	: IPCGAttributeAccessorKeys(/*bInReadOnly=*/ false)
+{
+	check(InPoints.Num() == InPointIndices.Num());
+	Points.Reserve(InPointIndices.Num());
+	Algo::Transform(InPointIndices, Points, [&InPoints](const int32 Index) -> FPCGPoint* { return &InPoints[Index]; });
+}
+
+FPCGAttributeAccessorKeysPointsSubset::FPCGAttributeAccessorKeysPointsSubset(const TArrayView<const FPCGPoint>& InPoints, const TArrayView<const int32>& InPointIndices)
+	: IPCGAttributeAccessorKeys(/*bInReadOnly=*/ true)
+{
+	check(InPoints.Num() == InPointIndices.Num());
+	Points.Reserve(InPointIndices.Num());
+	Algo::Transform(InPointIndices, Points, [&InPoints](const int32 Index) -> FPCGPoint* { return const_cast<FPCGPoint*>(&InPoints[Index]); });
+}
+
+FPCGAttributeAccessorKeysPointsSubset::FPCGAttributeAccessorKeysPointsSubset(TArray<FPCGPoint*> InPointPtrs)
+	: IPCGAttributeAccessorKeys(/*bInReadOnly=*/ false)
+	, Points(std::move(InPointPtrs))
+{}
+
+FPCGAttributeAccessorKeysPointsSubset::FPCGAttributeAccessorKeysPointsSubset(TArray<const FPCGPoint*> InPointPtrs)
+	: IPCGAttributeAccessorKeys(/*bInReadOnly=*/ true)
+	, Points(std::move(*reinterpret_cast<TArray<FPCGPoint*>*>(&InPointPtrs)))
+{}
+
+bool FPCGAttributeAccessorKeysPointsSubset::GetPointKeys(int32 InStart, TArrayView<FPCGPoint*> OutPoints)
+{
+	return PCGAttributeAccessorKeys::GetKeys(Points, InStart, OutPoints, [](FPCGPoint* Point) -> FPCGPoint* { return Point; });
+}
+
+bool FPCGAttributeAccessorKeysPointsSubset::GetPointKeys(int32 InStart, TArrayView<const FPCGPoint*> OutPoints) const
+{
+	return PCGAttributeAccessorKeys::GetKeys(Points, InStart, OutPoints, [](const FPCGPoint* Point) -> const FPCGPoint* { return Point; });
+}
+
+bool FPCGAttributeAccessorKeysPointsSubset::GetGenericObjectKeys(int32 InStart, TArrayView<void*> OutObjects)
+{
+	return PCGAttributeAccessorKeys::GetKeys(Points, InStart, OutObjects, [](FPCGPoint* Point) -> void* { return Point; });
+}
+
+bool FPCGAttributeAccessorKeysPointsSubset::GetGenericObjectKeys(int32 InStart, TArrayView<const void*> OutObjects) const
+{
+	return PCGAttributeAccessorKeys::GetKeys(Points, InStart, OutObjects, [](const FPCGPoint* Point) -> const void* { return Point; });
+}
+
+bool FPCGAttributeAccessorKeysPointsSubset::GetMetadataEntryKeys(int32 InStart, TArrayView<PCGMetadataEntryKey*> OutEntryKeys)
+{
+	return PCGAttributeAccessorKeys::GetKeys(Points, InStart, OutEntryKeys, [](FPCGPoint* Point) -> PCGMetadataEntryKey* { return &(Point->MetadataEntry); });
+}
+
+bool FPCGAttributeAccessorKeysPointsSubset::GetMetadataEntryKeys(int32 InStart, TArrayView<const PCGMetadataEntryKey*> OutEntryKeys) const
+{
+	return PCGAttributeAccessorKeys::GetKeys(Points, InStart, OutEntryKeys, [](const FPCGPoint* Point) -> const PCGMetadataEntryKey* { return &(Point->MetadataEntry); });
 }
\ No newline at end of file
diff --git a/Engine/Plugins/PCG/Source/PCG/Public/InstanceDataPackers/PCGInstanceDataPackerByAttribute.h b/Engine/Plugins/PCG/Source/PCG/Public/InstanceDataPackers/PCGInstanceDataPackerByAttribute.h
--- a/Engine/Plugins/PCG/Source/PCG/Public/InstanceDataPackers/PCGInstanceDataPackerByAttribute.h
+++ b/Engine/Plugins/PCG/Source/PCG/Public/InstanceDataPackers/PCGInstanceDataPackerByAttribute.h
@@ -22,7 +22,7 @@ public:
 	virtual bool GetAttributeNames(TArray<FName>* OutNames) override;
 
 public:
-	UPROPERTY(BlueprintReadWrite, EditAnywhere, Category = InstanceDataPacker, meta = (PCG_DiscardPropertySelection))
+	UPROPERTY(BlueprintReadWrite, EditAnywhere, Category = InstanceDataPacker)
 	TArray<FPCGAttributePropertyInputSelector> AttributeSelectors;
 
 #if WITH_EDITORONLY_DATA
diff --git a/Engine/Plugins/PCG/Source/PCG/Public/MeshSelectors/PCGMeshSelectorBase.h b/Engine/Plugins/PCG/Source/PCG/Public/MeshSelectors/PCGMeshSelectorBase.h
--- a/Engine/Plugins/PCG/Source/PCG/Public/MeshSelectors/PCGMeshSelectorBase.h
+++ b/Engine/Plugins/PCG/Source/PCG/Public/MeshSelectors/PCGMeshSelectorBase.h
@@ -3,6 +3,7 @@
 #pragma once
 
 #include "PCGPoint.h"
+#include "Data/PCGPointData.h"
 #include "MeshSelectors/PCGISMDescriptor.h"
 #include "Metadata/PCGMetadata.h"
 
@@ -24,12 +25,20 @@ struct FPCGMeshInstanceList
 {
 	GENERATED_BODY()
 
-	FPCGMeshInstanceList() = default;
-
+	// Note: We need to explicitly disable warnings on these constructors/operators for clang to be happy with deprecated variables
+	PRAGMA_DISABLE_DEPRECATION_WARNINGS
 	explicit FPCGMeshInstanceList(const FPCGSoftISMComponentDescriptor& InDescriptor)
 		: Descriptor(InDescriptor)
 		, AttributePartitionIndex(INDEX_NONE)
 	{}
+	
+	FPCGMeshInstanceList() = default;
+	~FPCGMeshInstanceList() = default;
+	FPCGMeshInstanceList(const FPCGMeshInstanceList&) = default;
+	FPCGMeshInstanceList(FPCGMeshInstanceList&&) = default;
+	FPCGMeshInstanceList& operator=(const FPCGMeshInstanceList&) = default;
+	FPCGMeshInstanceList& operator=(FPCGMeshInstanceList&&) = default;
+	PRAGMA_ENABLE_DEPRECATION_WARNINGS
 
 	UPROPERTY(EditAnywhere, Category = Settings)
 	FPCGSoftISMComponentDescriptor Descriptor;
@@ -37,11 +46,20 @@ struct FPCGMeshInstanceList
 	UPROPERTY(BlueprintReadWrite, EditAnywhere, Category = Settings)
 	TArray<FTransform> Instances;
 
-	UPROPERTY(BlueprintReadWrite, EditAnywhere, Category = Settings)
-	TArray<int64> InstancesMetadataEntry;
-
 	/** Tracks which partition the instance list belongs to. */
 	int64 AttributePartitionIndex;
+
+	UPROPERTY(BlueprintReadWrite, EditAnywhere, Category = Settings)
+	TWeakObjectPtr<const UPCGPointData> PointData;
+	
+	UPROPERTY(BlueprintReadWrite, EditAnywhere, Category = Settings)
+	TArray<int32> InstancesIndices;
+
+#if WITH_EDITORONLY_DATA
+	UE_DEPRECATED(5.5, "Use PointData + InstanceIndices instead.")
+	UPROPERTY(meta = (DeprecatedProperty))
+	TArray<int64> InstancesMetadataEntry;
+#endif // WITH_EDITORONLY_DATA
 };
 
 UENUM()
diff --git a/Engine/Plugins/PCG/Source/PCG/Public/MeshSelectors/PCGMeshSelectorWeighted.h b/Engine/Plugins/PCG/Source/PCG/Public/MeshSelectors/PCGMeshSelectorWeighted.h
--- a/Engine/Plugins/PCG/Source/PCG/Public/MeshSelectors/PCGMeshSelectorWeighted.h
+++ b/Engine/Plugins/PCG/Source/PCG/Public/MeshSelectors/PCGMeshSelectorWeighted.h
@@ -14,7 +14,8 @@ namespace PCGMeshSelectorWeighted
 		TArray<FPCGMeshInstanceList>& InstanceLists,
 		bool bUseMaterialOverrides,
 		const TArray<TSoftObjectPtr<UMaterialInterface>>& InMaterialOverrides,
-		bool bInIsLocalToWorldDeterminantNegative);
+		bool bInIsLocalToWorldDeterminantNegative,
+		const UPCGPointData* InPointData);
 }
 
 USTRUCT(BlueprintType)
diff --git a/Engine/Plugins/PCG/Source/PCG/Public/Metadata/Accessors/PCGAttributeAccessorKeys.h b/Engine/Plugins/PCG/Source/PCG/Public/Metadata/Accessors/PCGAttributeAccessorKeys.h
--- a/Engine/Plugins/PCG/Source/PCG/Public/Metadata/Accessors/PCGAttributeAccessorKeys.h
+++ b/Engine/Plugins/PCG/Source/PCG/Public/Metadata/Accessors/PCGAttributeAccessorKeys.h
@@ -190,6 +190,35 @@ protected:
 	TArrayView<FPCGPoint> Points;
 };
 
+///////////////////////////////////////////////////////////////////////
+
+/**
+* Key around subset of points
+*/
+class PCG_API FPCGAttributeAccessorKeysPointsSubset : public IPCGAttributeAccessorKeys
+{
+public:
+	FPCGAttributeAccessorKeysPointsSubset(const TArrayView<FPCGPoint>& InPoints, const TArrayView<const int32>& InPointIndices);
+	FPCGAttributeAccessorKeysPointsSubset(const TArrayView<const FPCGPoint>& InPoints, const TArrayView<const int32>& InPointIndices);
+
+	FPCGAttributeAccessorKeysPointsSubset(TArray<FPCGPoint*> InPointPtrs);
+	FPCGAttributeAccessorKeysPointsSubset(TArray<const FPCGPoint*> InPointPtrs);
+
+	virtual int32 GetNum() const override { return Points.Num(); }
+
+protected:
+	virtual bool GetPointKeys(int32 InStart, TArrayView<FPCGPoint*> OutPoints) override;
+	virtual bool GetPointKeys(int32 InStart, TArrayView<const FPCGPoint*> OutPoints) const override;
+
+	virtual bool GetGenericObjectKeys(int32 InStart, TArrayView<void*> OutObjects) override;
+	virtual bool GetGenericObjectKeys(int32 InStart, TArrayView<const void*> OutObjects) const override;
+
+	virtual bool GetMetadataEntryKeys(int32 InStart, TArrayView<PCGMetadataEntryKey*> OutEntryKeys) override;
+	virtual bool GetMetadataEntryKeys(int32 InStart, TArrayView<const PCGMetadataEntryKey*> OutEntryKeys) const override;
+
+	TArray<FPCGPoint*> Points;
+};
+
 /////////////////////////////////////////////////////////////////
 
 /**
