[PCG] Support for cooking GPU graphs
* Now we serialize compiled graph tasks on cook so that compilation is unnecessary outside of editor. This is required for GPU graphs since compute graphs cannot be compiled outside of editor.

- Changes -
* Made FPCGGraphTask a serialiable U-Type
* Refactored FPCGGraphCompiler. Compilation results are now stored in FPCGGraphCompilerCache.
* Introduced an enum EPCGElementSource, used to indicate how an element should be recreated when it is missing (e.g. in standalone).
* Added GridLinkageSettings and ComputeGraphSettings, serializable alternatives to GridLinkageElement and ComputeGraphElement
* On UPCGGraph::PreSave(), we compile the graph for all grid sizes and move the compiled results into a new class UPCGGraphCompilationData, which is serialized for cook.

diff --git a/Engine/Plugins/PCG/Source/PCG/Private/Compute/DataInterfaces/PCGComputeDataInterface.h b/Engine/Plugins/PCG/Source/PCG/Private/Compute/DataInterfaces/PCGComputeDataInterface.h
--- a/Engine/Plugins/PCG/Source/PCG/Private/Compute/DataInterfaces/PCGComputeDataInterface.h
+++ b/Engine/Plugins/PCG/Source/PCG/Private/Compute/DataInterfaces/PCGComputeDataInterface.h
@@ -27,9 +27,13 @@ public:
 
 protected:
 	/** Label of output pin that this data interface is associated with. */
+	UPROPERTY()
 	FName OutputPinLabel;
+
+	UPROPERTY()
 	FName OutputPinLabelAlias;
 
 	/** Generated PCG data will be assigned these labels. */
+	UPROPERTY()
 	TArray<FName> DownstreamInputPinLabelAliases;
 };
diff --git a/Engine/Plugins/PCG/Source/PCG/Private/Compute/DataInterfaces/PCGDataCollectionDataInterface.h b/Engine/Plugins/PCG/Source/PCG/Private/Compute/DataInterfaces/PCGDataCollectionDataInterface.h
--- a/Engine/Plugins/PCG/Source/PCG/Private/Compute/DataInterfaces/PCGDataCollectionDataInterface.h
+++ b/Engine/Plugins/PCG/Source/PCG/Private/Compute/DataInterfaces/PCGDataCollectionDataInterface.h
@@ -61,6 +61,7 @@ protected:
 	TObjectPtr<const UPCGSettings> ProducerSettings;
 
 	/** Whether this data should be copied from GPU to CPU. */
+	UPROPERTY()
 	bool bRequiresReadback = false;
 
 	static TCHAR const* TemplateFilePath;
diff --git a/Engine/Plugins/PCG/Source/PCG/Private/Compute/Elements/PCGComputeGraphElement.cpp b/Engine/Plugins/PCG/Source/PCG/Private/Compute/Elements/PCGComputeGraphElement.cpp
--- a/Engine/Plugins/PCG/Source/PCG/Private/Compute/Elements/PCGComputeGraphElement.cpp
+++ b/Engine/Plugins/PCG/Source/PCG/Private/Compute/Elements/PCGComputeGraphElement.cpp
@@ -3,6 +3,7 @@
 #include "Compute/Elements/PCGComputeGraphElement.h"
 
 #include "PCGComponent.h"
+#include "PCGGraph.h"
 #include "PCGManagedResource.h"
 #include "PCGModule.h"
 #include "PCGSubsystem.h"
@@ -11,6 +12,7 @@
 #include "Compute/DataInterfaces/PCGDataCollectionDataInterface.h"
 #include "Elements/PCGStaticMeshSpawner.h"
 #include "Engine/StaticMesh.h"
+#include "Graph/PCGGraphCompiler.h"
 #include "InstanceDataPackers/PCGInstanceDataPackerBase.h"
 #include "MeshSelectors/PCGMeshSelectorByAttribute.h"
 #include "MeshSelectors/PCGMeshSelectorWeighted.h"
@@ -21,6 +23,8 @@
 #include "GameFramework/Actor.h"
 #include "Logging/LogVerbosity.h"
 
+#include UE_INLINE_GENERATED_CPP_BY_NAME(PCGComputeGraphElement)
+
 #define LOCTEXT_NAMESPACE "PCGComputeGraphElement"
 
 void FPCGComputeGraphContext::AddExtraStructReferencedObjects(FReferenceCollector& Collector)
@@ -31,11 +35,7 @@ void FPCGComputeGraphContext::AddExtraStructReferencedObjects(FReferenceCollecto
 #if WITH_EDITOR
 bool FPCGComputeGraphElement::operator==(const FPCGComputeGraphElement& Other) const
 {
-	// Equivalence is same compute graph.
-	// TODO: A compute graph is currently generated for every compile, so the presence of GPU nodes breaks the current
-	// change detection. We could either cache compute graphs formed by subsets of GPU nodes that have not changed, or
-	// we could do a deep equality check for compute graphs here.
-	return Graph == Other.Graph;
+	return ComputeGraphIndex == Other.ComputeGraphIndex;
 }
 #endif
 
@@ -45,9 +45,32 @@ bool FPCGComputeGraphElement::ExecuteInternal(FPCGContext* InContext) const
 	check(InContext);
 	FPCGComputeGraphContext* Context = static_cast<FPCGComputeGraphContext*>(InContext);
 
-	if (!ensure(Graph))
+	if (!Context->ComputeGraph)
 	{
-		return true;
+		const UPCGComponent* Component = InContext->SourceComponent.Get();
+
+		if (const UPCGGraph* PCGGraph = Component ? Component->GetGraph() : nullptr)
+		{
+			uint32 GenerationGridSize = PCGHiGenGrid::UninitializedGridSize();
+
+			if (PCGGraph->IsHierarchicalGenerationEnabled())
+			{
+				if (Component->IsLocalComponent() || Component->IsPartitioned())
+				{
+					GenerationGridSize = Component->GetGenerationGridSize();
+				}
+			}
+
+			if (UPCGSubsystem* Subsystem = Context->SourceComponent->GetSubsystem())
+			{
+				Context->ComputeGraph.Reset(Subsystem->GetComputeGraph(PCGGraph, GenerationGridSize, ComputeGraphIndex));
+			}
+		}
+
+		if (!ensure(Context->ComputeGraph))
+		{
+			return true;
+		}
 	}
 
 	// Drive the execution of a compute graph. The stages are commented below and numbered by the sequence in which they are executed.
@@ -92,7 +115,7 @@ bool FPCGComputeGraphElement::ExecuteInternal(FPCGContext* InContext) const
 	};
 
 	// 3. If still compiling, try again next frame.
-	if (Graph->HasKernelResourcesPendingShaderCompilation())
+	if (Context->ComputeGraph->HasKernelResourcesPendingShaderCompilation())
 	{
 		UE_LOG(LogPCG, Log, TEXT("Deferring until next frame as the kernel has pending shader compilations."));
 		SleepUntilNextFrame();
@@ -113,10 +136,10 @@ bool FPCGComputeGraphElement::ExecuteInternal(FPCGContext* InContext) const
 		UPCGDataBinding* DataBinding = FPCGContext::NewObject_AnyThread<UPCGDataBinding>(Context);
 		Context->DataBinding.Reset(DataBinding);
 
-		Context->DataBinding->Initialize(Graph.Get(), Context->SourceComponent, Context->InputData, Graph->GetAttributeLookupTable());
+		Context->DataBinding->Initialize(Context->ComputeGraph.Get(), Context->SourceComponent, Context->InputData, Context->ComputeGraph->GetAttributeLookupTable());
 
 		// Perform validation after input data is initialized.
-		for (TWeakObjectPtr<const UPCGNode> Node : Graph->KernelToNode)
+		for (TWeakObjectPtr<const UPCGNode> Node : Context->ComputeGraph->KernelToNode)
 		{
 			const UPCGSettings* Settings = Node.Get() ? Node->GetSettings() : nullptr;
 
@@ -130,7 +153,7 @@ bool FPCGComputeGraphElement::ExecuteInternal(FPCGContext* InContext) const
 
 		{
 			TRACE_CPUPROFILER_EVENT_SCOPE(ComputeGraphInstance.CreateDataProviders);
-			Context->ComputeGraphInstance.CreateDataProviders(Graph.Get(), 0, Context->DataBinding.Get());
+			Context->ComputeGraphInstance.CreateDataProviders(Context->ComputeGraph.Get(), 0, Context->DataBinding.Get());
 		}
 
 		// Register all providers running async operations. TODO review if we should have a general API like "RunsAsyncOperations()"?
@@ -171,11 +194,11 @@ bool FPCGComputeGraphElement::ExecuteInternal(FPCGContext* InContext) const
 	check(Context->DataBinding && InContext->SourceComponent.Get());
 
 	// 1. Prepare render resources. In editor, this will trigger shader compile if not compiled already.
-	if (!Graph->GetRenderProxy())
+	if (!Context->ComputeGraph->GetRenderProxy())
 	{
-		TRACE_CPUPROFILER_EVENT_SCOPE(Graph->UpdateResources);
+		TRACE_CPUPROFILER_EVENT_SCOPE(Context->ComputeGraph->UpdateResources);
 
-		Graph->UpdateResources();
+		Context->ComputeGraph->UpdateResources();
 
 		SleepUntilNextFrame();
 		return false;
@@ -189,7 +212,8 @@ bool FPCGComputeGraphElement::ExecuteInternal(FPCGContext* InContext) const
 #endif
 
 		// If there was any error then we should abort.
-		for (const TPair<TObjectKey<const UPCGNode>, TArray<FComputeKernelCompileMessage>>& NodeAndCompileMessages : Graph->KernelToCompileMessages)
+		using FNodeAndCompileMessages = const TPair<TObjectKey<const UPCGNode>, TArray<FComputeKernelCompileMessage>>;
+		for (FNodeAndCompileMessages& NodeAndCompileMessages : Context->ComputeGraph->KernelToCompileMessages)
 		{
 			for (const FComputeKernelCompileMessage& Message : NodeAndCompileMessages.Get<1>())
 			{
@@ -208,7 +232,7 @@ bool FPCGComputeGraphElement::ExecuteInternal(FPCGContext* InContext) const
 		TRACE_CPUPROFILER_EVENT_SCOPE(ComputeGraphInstance.EnqueueWork);
 
 		Context->bGraphEnqueued = Context->ComputeGraphInstance.EnqueueWork(
-			Graph.Get(),
+			Context->ComputeGraph.Get(),
 			InContext->SourceComponent->GetScene(),
 			ComputeTaskExecutionGroup::EndOfFrameUpdate,
 			InContext->SourceComponent->GetOwner()->GetFName(),
@@ -246,9 +270,11 @@ bool FPCGComputeGraphElement::SetupProceduralISMComponents(FPCGContext* InContex
 	check(InContext);
 	check(InBinding);
 
+	FPCGComputeGraphContext* Context = static_cast<FPCGComputeGraphContext*>(InContext);
+
 	bool bAnyComponentsSetup = false;
 
-	for (const UPCGSettings* Settings : Graph->StaticMeshSpawners)
+	for (const UPCGSettings* Settings : Context->ComputeGraph->StaticMeshSpawners)
 	{
 		const UPCGStaticMeshSpawnerSettings* SpawnerSettings = Cast<UPCGStaticMeshSpawnerSettings>(Settings);
 		if (!ensure(SpawnerSettings) || !SpawnerSettings->bEnabled)
@@ -481,7 +507,8 @@ void FPCGComputeGraphElement::LogCompilationMessages(FPCGComputeGraphContext* In
 {
 	if (InContext->SourceComponent.IsValid() && InContext->Stack)
 	{
-		for (const TPair<TObjectKey<const UPCGNode>, TArray<FComputeKernelCompileMessage>>& NodeAndCompileMessages : Graph->KernelToCompileMessages)
+		using FNodeAndCompileMessages = const TPair<TObjectKey<const UPCGNode>, TArray<FComputeKernelCompileMessage>>;
+		for (FNodeAndCompileMessages& NodeAndCompileMessages : InContext->ComputeGraph->KernelToCompileMessages)
 		{
 			for (const FComputeKernelCompileMessage& Message : NodeAndCompileMessages.Get<1>())
 			{
@@ -538,4 +565,9 @@ void FPCGComputeGraphElement::LogCompilationMessages(FPCGComputeGraphContext* In
 }
 #endif
 
+FPCGElementPtr UPCGComputeGraphSettings::CreateElement() const
+{
+	return MakeShared<FPCGComputeGraphElement>(ComputeGraphIndex);
+}
+
 #undef LOCTEXT_NAMESPACE
diff --git a/Engine/Plugins/PCG/Source/PCG/Private/Compute/Elements/PCGComputeGraphElement.h b/Engine/Plugins/PCG/Source/PCG/Private/Compute/Elements/PCGComputeGraphElement.h
--- a/Engine/Plugins/PCG/Source/PCG/Private/Compute/Elements/PCGComputeGraphElement.h
+++ b/Engine/Plugins/PCG/Source/PCG/Private/Compute/Elements/PCGComputeGraphElement.h
@@ -2,16 +2,17 @@
 
 #pragma once
 
-#include "PCGElement.h"
-
 #include "PCGContext.h"
-
+#include "PCGElement.h"
+#include "PCGSettings.h"
 #include "Compute/PCGComputeGraph.h"
 #include "Compute/PCGDataBinding.h"
 #include "ComputeFramework/ComputeGraphInstance.h"
 
 #include "UObject/StrongObjectPtr.h"
 
+#include "PCGComputeGraphElement.generated.h"
+
 struct FPCGComputeGraphContext : public FPCGContext
 {
 public:
@@ -22,6 +23,7 @@ protected:
 
 public:
 	TStrongObjectPtr<UPCGDataBinding> DataBinding = nullptr;
+	TStrongObjectPtr<UPCGComputeGraph> ComputeGraph = nullptr;
 
 	/** Data providers created from data interfaces and data bindings. */
 	FComputeGraphInstance ComputeGraphInstance;
@@ -44,7 +46,10 @@ public:
 class FPCGComputeGraphElement : public IPCGElement
 {
 public:
-	FPCGComputeGraphElement() {}
+	FPCGComputeGraphElement() = default;
+	explicit FPCGComputeGraphElement(int32 InComputeGraphIndex)
+		: ComputeGraphIndex(InComputeGraphIndex)
+	{}
 
 #if WITH_EDITOR
 	//~Begin IPCGElement interface
@@ -55,7 +60,9 @@ public:
 	bool operator==(const FPCGComputeGraphElement& Other) const;
 #endif
 
-	TStrongObjectPtr<UPCGComputeGraph> Graph = nullptr;
+	// TODO: ComputeGraphIndex could be removed from the element if we properly hook up EPCGElementSource::FromSettings in the graph executor 
+	// to manufacture the context's settings from the FPCGGraphTask Settings member. Would also eliminate the need for IsComputeGraphElement().
+	int32 ComputeGraphIndex = INDEX_NONE;
 
 protected:
 	virtual FPCGContext* CreateContext() override { return new FPCGComputeGraphContext(); }
@@ -77,3 +84,16 @@ protected:
 	void LogCompilationMessages(FPCGComputeGraphContext* InContext) const;
 #endif
 };
+
+UCLASS(ClassGroup = (Procedural))
+class UPCGComputeGraphSettings : public UPCGSettings
+{
+	GENERATED_BODY()
+
+protected:
+	virtual FPCGElementPtr CreateElement() const override;
+
+public:
+	UPROPERTY()
+	int32 ComputeGraphIndex = INDEX_NONE;
+};
diff --git a/Engine/Plugins/PCG/Source/PCG/Private/Compute/PCGComputeGraph.cpp b/Engine/Plugins/PCG/Source/PCG/Private/Compute/PCGComputeGraph.cpp
--- a/Engine/Plugins/PCG/Source/PCG/Private/Compute/PCGComputeGraph.cpp
+++ b/Engine/Plugins/PCG/Source/PCG/Private/Compute/PCGComputeGraph.cpp
@@ -1,6 +1,6 @@
 // Copyright Epic Games, Inc. All Rights Reserved.
 
-#include "PCGComputeGraph.h"
+#include "Compute/PCGComputeGraph.h"
 
 #include "PCGGraph.h"
 #include "PCGModule.h"
diff --git a/Engine/Plugins/PCG/Source/PCG/Private/Graph/PCGGraphCompilationData.h b/Engine/Plugins/PCG/Source/PCG/Private/Graph/PCGGraphCompilationData.h
new file mode 100644
--- /dev/null
+++ b/Engine/Plugins/PCG/Source/PCG/Private/Graph/PCGGraphCompilationData.h
@@ -0,0 +1,47 @@
+// Copyright Epic Games, Inc. All Rights Reserved.
+
+#pragma once
+
+#include "Graph/PCGGraphTask.h"
+
+#include "PCGGraphCompilationData.generated.h"
+
+class UPCGComputeGraph;
+struct FPCGStackContext;
+
+/** Helper struct to allow serialization of GridSize to Compiled Graph Tasks map. */
+USTRUCT()
+struct FPCGGraphTasks
+{
+	GENERATED_BODY();
+
+	UPROPERTY()
+	TArray<FPCGGraphTask> GraphTasks;
+};
+
+/** Helper struct to allow serialization of GridSize to ComputeGraphs map. */
+USTRUCT()
+struct FPCGComputeGraphs
+{
+	GENERATED_BODY();
+
+	UPROPERTY()
+	TArray<TObjectPtr<UPCGComputeGraph>> ComputeGraphs;
+};
+
+/** Holds compiled data for serialization during cook. */
+UCLASS()
+class UPCGGraphCompilationData : public UObject
+{
+	GENERATED_BODY()
+
+public:
+	UPROPERTY()
+	TMap</*HiGen Grid Size*/uint32, FPCGGraphTasks> Tasks;
+
+	UPROPERTY()
+	TMap</*HiGen Grid Size*/uint32, FPCGStackContext> StackContexts;
+
+	UPROPERTY()
+	TMap</*HiGen Grid Size*/uint32, FPCGComputeGraphs> ComputeGraphs;
+};
diff --git a/Engine/Plugins/PCG/Source/PCG/Private/Graph/PCGGraphCompiler.cpp b/Engine/Plugins/PCG/Source/PCG/Private/Graph/PCGGraphCompiler.cpp
--- a/Engine/Plugins/PCG/Source/PCG/Private/Graph/PCGGraphCompiler.cpp
+++ b/Engine/Plugins/PCG/Source/PCG/Private/Graph/PCGGraphCompiler.cpp
@@ -12,6 +12,7 @@
 #include "Elements/PCGGather.h"
 #include "Elements/PCGHiGenGridSize.h"
 #include "Elements/PCGReroute.h"
+#include "Graph/PCGGraphCompilationData.h"
 #include "Graph/PCGGraphCompilerGPU.h"
 #include "Graph/PCGGraphExecutor.h"
 #include "Graph/PCGPinDependencyExpression.h"
@@ -19,9 +20,6 @@
 #include "HAL/IConsoleManager.h"
 #include "Misc/ScopeRWLock.h"
 
-FPCGElementPtr FPCGGraphCompiler::SharedGatherElement;
-FRWLock FPCGGraphCompiler::SharedGatherElementLock;
-
 namespace PCGGraphCompiler
 {
 	static TAutoConsoleVariable<bool> CVarEnableTaskStaticCulling(
@@ -77,6 +75,72 @@ namespace PCGGraphCompiler
 	};
 }
 
+UPCGComputeGraph* FPCGGraphCompilerCache::GetCompiledComputeGraph(const UPCGGraph* InGraph, uint32 GridSize, uint32 ComputeGraphIndex)
+{
+#if WITH_EDITOR
+	if (TMap<uint32, TArray<TObjectPtr<UPCGComputeGraph>>>* GridSizeToComputeGraphs = TopGraphToComputeGraphMap.Find(InGraph))
+	{
+		if (TArray<TObjectPtr<UPCGComputeGraph>>* ComputeGraphs = GridSizeToComputeGraphs->Find(GridSize))
+		{
+			if (ComputeGraphs->IsValidIndex(ComputeGraphIndex))
+			{
+				return (*ComputeGraphs)[ComputeGraphIndex];
+			}
+		}
+	}
+#else
+	// Cannot compile compute graphs outside of editor, so we always look for cooked compute graphs.
+	const UPCGGraphCompilationData* CookedData = InGraph ? InGraph->GetCookedCompilationData() : nullptr;
+
+	if (CookedData)
+	{
+		const FPCGComputeGraphs* CookedComputeGraphs = CookedData->ComputeGraphs.Find(GridSize);
+
+		if (CookedComputeGraphs && CookedComputeGraphs->ComputeGraphs.IsValidIndex(ComputeGraphIndex))
+		{
+			return CookedComputeGraphs->ComputeGraphs[ComputeGraphIndex];
+		}
+	}
+#endif
+
+	return nullptr;
+}
+
+#if WITH_EDITOR
+void FPCGGraphCompilerCache::RemoveFromCache(UPCGGraph* InGraph)
+{
+	UE_LOG(LogPCG, Verbose, TEXT("FPCGGraphCompilerCache::RemoveFromCache '%s'"), *InGraph->GetName());
+
+	check(InGraph);
+	RemoveFromCacheRecursive(InGraph);
+}
+
+void FPCGGraphCompilerCache::RemoveFromCacheRecursive(UPCGGraph* InGraph)
+{
+	{
+		FWriteScopeLock Lock(GraphToTaskMapLock);
+		GraphToTaskMap.Remove(InGraph);
+		GraphToStackContextMap.Remove(InGraph);
+		TopGraphToTaskMap.Remove(InGraph);
+		TopGraphToStackContextMap.Remove(InGraph);
+		TopGraphToComputeGraphMap.Remove(InGraph);
+	}
+
+	TArray<UPCGGraph*> ParentGraphs;
+
+	{
+		FScopeLock Lock(&GraphDependenciesLock);
+		GraphDependencies.MultiFind(InGraph, ParentGraphs);
+		GraphDependencies.Remove(InGraph);
+	}
+
+	for (UPCGGraph* Graph : ParentGraphs)
+	{
+		RemoveFromCacheRecursive(Graph);
+	}
+}
+#endif // WITH_EDITOR
+
 TArray<FPCGGraphTask> FPCGGraphCompiler::CompileGraph(UPCGGraph* InGraph, FPCGTaskId& NextId, FPCGStackContext& InOutStackContext)
 {
 	if (!InGraph)
@@ -128,9 +192,9 @@ TArray<FPCGGraphTask> FPCGGraphCompiler::CompileGraph(UPCGGraph* InGraph, FPCGTa
 			TArray<FPCGGraphTask> Subtasks = GetCompiledTasks(Subgraph, PCGHiGenGrid::UninitializedGridSize(), SubgraphStackContext, /*bIsTopGraph=*/false);
 
 #if WITH_EDITOR
-			GraphDependenciesLock.Lock();
-			GraphDependencies.AddUnique(Subgraph, InGraph);
-			GraphDependenciesLock.Unlock();
+			Cache.GraphDependenciesLock.Lock();
+			Cache.GraphDependencies.AddUnique(Subgraph, InGraph);
+			Cache.GraphDependenciesLock.Unlock();
 #endif // WITH_EDITOR
 
 			// Append all the stack frames from inside the subgraph to my current graph stack
@@ -229,7 +293,7 @@ TArray<FPCGGraphTask> FPCGGraphCompiler::CompileGraph(UPCGGraph* InGraph, FPCGTa
 			PostTask.StackIndex = InOutStackContext.GetCurrentStackIndex();
 			// Implementation note: since we`ve already executed the node once, we normally don`t need to execute it a second time
 			// especially since we cannot distinguish between the pre and post during execution so any data filtering related to pins is bound to fail.
-			PostTask.Element = GetSharedTrivialElement();
+			PostTask.ElementSource = EPCGElementSource::Trivial;
 
 			// Add execution-only dependency on pre-task, without this post task can be scheduled concurrently with pre-task, and concurrently
 			// with something that might become inactive and would then fail to dynamically cull this already-scheduled task.
@@ -325,9 +389,22 @@ TArray<FPCGGraphTask> FPCGGraphCompiler::CompileGraph(UPCGGraph* InGraph, FPCGTa
 
 void FPCGGraphCompiler::Compile(UPCGGraph* InGraph)
 {
-	GraphToTaskMapLock.ReadLock();
-	bool bAlreadyCached = GraphToTaskMap.Contains(InGraph);
-	GraphToTaskMapLock.ReadUnlock();
+	bool bAlreadyCached = false;
+
+	{
+		FReadScopeLock Lock(Cache.GraphToTaskMapLock);
+		bAlreadyCached = Cache.GraphToTaskMap.Contains(InGraph);
+	}
+
+#if !WITH_EDITOR
+	// Don't need to compile if the tasks are already cooked.
+	const UPCGGraphCompilationData* CookedData = InGraph ? InGraph->GetCookedCompilationData() : nullptr;
+
+	if (CookedData)
+	{
+		bAlreadyCached |= !CookedData->Tasks.IsEmpty();
+	}
+#endif
 
 	if (bAlreadyCached)
 	{
@@ -345,11 +422,11 @@ void FPCGGraphCompiler::Compile(UPCGGraph* InGraph)
 	// Store back the results in the cache if it's valid
 	if (!CompiledTasks.IsEmpty())
 	{
-		FWriteScopeLock Lock(GraphToTaskMapLock);
-		if (!GraphToTaskMap.Contains(InGraph))
+		FWriteScopeLock Lock(Cache.GraphToTaskMapLock);
+		if (!Cache.GraphToTaskMap.Contains(InGraph))
 		{
-			GraphToTaskMap.Add(InGraph, MoveTemp(CompiledTasks));
-			GraphToStackContext.Add(InGraph, StackContext);
+			Cache.GraphToTaskMap.Add(InGraph, MoveTemp(CompiledTasks));
+			Cache.GraphToStackContextMap.Add(InGraph, StackContext);
 		}
 	}
 }
@@ -357,16 +434,16 @@ void FPCGGraphCompiler::Compile(UPCGGraph* InGraph)
 TArray<FPCGGraphTask> FPCGGraphCompiler::GetPrecompiledTasks(const UPCGGraph* InGraph, uint32 GenerationGridSize, FPCGStackContext& OutStackContext, bool bIsTopGraph) const
 {
 	// Get compiled tasks in a threadsafe way
-	FReadScopeLock ReadLock(GraphToTaskMapLock);
+	FReadScopeLock ReadLock(Cache.GraphToTaskMapLock);
 
 	const TArray<FPCGGraphTask>* ExistingTasks = nullptr;
 	if (bIsTopGraph)
 	{
 		// Top graphs are optimized per grid size.
-		const TMap<uint32, TArray<FPCGGraphTask>>* GridSizeToCompiledGraph = TopGraphToTaskMap.Find(InGraph);
+		const TMap<uint32, TArray<FPCGGraphTask>>* GridSizeToCompiledGraph = Cache.TopGraphToTaskMap.Find(InGraph);
 		ExistingTasks = GridSizeToCompiledGraph ? GridSizeToCompiledGraph->Find(GenerationGridSize) : nullptr;
 
-		const TMap<uint32, FPCGStackContext>* GridSizeToStackContext = TopGraphToStackContextMap.Find(InGraph);
+		const TMap<uint32, FPCGStackContext>* GridSizeToStackContext = Cache.TopGraphToStackContextMap.Find(InGraph);
 		const FPCGStackContext* ExistingStackContext = GridSizeToStackContext ? GridSizeToStackContext->Find(GenerationGridSize) : nullptr;
 		OutStackContext = ExistingStackContext ? *ExistingStackContext : FPCGStackContext();
 
@@ -375,12 +452,17 @@ TArray<FPCGGraphTask> FPCGGraphCompiler::GetPrecompiledTasks(const UPCGGraph* In
 	}
 	else
 	{
-		ExistingTasks = GraphToTaskMap.Find(InGraph);
+		ExistingTasks = Cache.GraphToTaskMap.Find(InGraph);
 	}
 
 	return ExistingTasks ? *ExistingTasks : TArray<FPCGGraphTask>();
 }
 
+UPCGComputeGraph* FPCGGraphCompiler::GetComputeGraph(const UPCGGraph* InGraph, uint32 GridSize, uint32 ComputeGraphIndex)
+{
+	return Cache.GetCompiledComputeGraph(InGraph, GridSize, ComputeGraphIndex);
+}
+
 void FPCGGraphCompiler::ResolveGridSizes(
 	EPCGHiGenGrid GenerationGrid,
 	const TArray<FPCGGraphTask>& CompiledTasks,
@@ -407,10 +489,12 @@ void FPCGGraphCompiler::ResolveGridSizes(
 }
 
 void FPCGGraphCompiler::CreateGridLinkages(
+	UPCGGraph* InGraph,
 	EPCGHiGenGrid InGenerationGrid,
 	TArray<EPCGHiGenGrid>& InOutTaskGenerationGrid,
 	TArray<FPCGGraphTask>& InOutCompiledTasks,
-	const FPCGStackContext& InStackContext)
+	const FPCGStackContext& InStackContext,
+	bool bIsCooking)
 {
 	// Now add link tasks - if a Grid256 task depends on data from a Grid512 task, inject a link
 	// task that looks up the Grid512 component, schedules its execution if it does not have data, and
@@ -496,7 +580,24 @@ void FPCGGraphCompiler::CreateGridLinkages(
 					return new FPCGGridLinkageContext();
 				};
 
-				LinkTask.Element = MakeShared<PCGGraphExecutor::FPCGGridLinkageElement>(GridLinkageOperation, ContextAllocator, FromGrid, ToGrid, ResourceKey, UpstreamPin);
+#if WITH_EDITOR
+				if (bIsCooking)
+				{
+					TObjectPtr<UPCGGridLinkageSettings> Settings = NewObject<UPCGGridLinkageSettings>(InGraph);
+					Settings->FromGrid = FromGrid;
+					Settings->ToGrid = ToGrid;
+					Settings->GenerationGrid = InGenerationGrid;
+					Settings->ResourceKey = ResourceKey;
+					Settings->UpstreamPin = UpstreamPin;
+
+					LinkTask.ElementSource = EPCGElementSource::FromSettings;
+					LinkTask.Settings = Settings;
+				}
+				else
+#endif
+				{
+					LinkTask.Element = MakeShared<PCGGraphExecutor::FPCGGridLinkageElement>(GridLinkageOperation, ContextAllocator, FromGrid, ToGrid, InGenerationGrid, ResourceKey, UpstreamPin);
+				}
 
 				// Now splice in the new task - redirect the downstream task to grab its input from the link task.
 				TaskInput.TaskId = LinkTask.NodeId;
@@ -854,6 +955,25 @@ void FPCGGraphCompiler::CullTasks(TArray<FPCGGraphTask>& InOutCompiledTasks, boo
 	}
 }
 
+void FPCGGraphCompiler::AddReferencedObjects(FReferenceCollector& Collector)
+{
+	TRACE_CPUPROFILER_EVENT_SCOPE(FCPGGraphCompiler::AddReferencedObjects);
+
+	TMap<UPCGGraph*, TMap<uint32, TArray<TObjectPtr<UPCGComputeGraph>>>> TopGraphToComputeGraphMap;
+
+	for (const TPair<UPCGGraph*, TMap<uint32, TArray<TObjectPtr<UPCGComputeGraph>>>>& GraphToGridSizes : Cache.TopGraphToComputeGraphMap)
+	{
+		for (const TPair<uint32, TArray<TObjectPtr<UPCGComputeGraph>>>& GridSizeToComputeGraphs : GraphToGridSizes.Value)
+		{
+			for (TObjectPtr<UPCGComputeGraph> ComputeGraph : GridSizeToComputeGraphs.Value)
+			{
+				Collector.AddReferencedObject(ComputeGraph);
+			}
+		}
+	}
+}
+
+
 void FPCGGraphCompiler::PostCullStackCleanup(TArray<FPCGGraphTask>& InCompiledTasks, FPCGStackContext& InOutStackContext)
 {
 	// Build set of stack IDs used by tasks. Using array as set is likely small.
@@ -1007,17 +1127,35 @@ TArray<FPCGGraphTask> FPCGGraphCompiler::GetCompiledTasks(UPCGGraph* InGraph, ui
 {
 	TArray<FPCGGraphTask> CompiledTasks;
 
+#if !WITH_EDITOR
+	// In standalone builds, try using cooked tasks first.
+	if (InGraph)
+	{
+		if (const UPCGGraphCompilationData* CookedData = InGraph->GetCookedCompilationData())
+		{
+			const FPCGGraphTasks* CookedTasks = CookedData->Tasks.Find(GenerationGridSize);
+			const FPCGStackContext* CookedStackContext = CookedData->StackContexts.Find(GenerationGridSize);
+
+			if (CookedTasks && CookedStackContext)
+			{
+				OutStackContext = *CookedStackContext;
+				return CookedTasks->GraphTasks;
+			}
+		}
+	}
+#endif
+
 	if (bIsTopGraph)
 	{
 		// Always try to compile
 		CompileTopGraph(InGraph, GenerationGridSize);
 
 		// Get compiled tasks in a threadsafe way
-		FReadScopeLock Lock(GraphToTaskMapLock);
-		const TMap<uint32, TArray<FPCGGraphTask>>* GridSizeToCompiledGraph = TopGraphToTaskMap.Find(InGraph);
+		FReadScopeLock Lock(Cache.GraphToTaskMapLock);
+		const TMap<uint32, TArray<FPCGGraphTask>>* GridSizeToCompiledGraph = Cache.TopGraphToTaskMap.Find(InGraph);
 		const TArray<FPCGGraphTask>* Tasks = GridSizeToCompiledGraph ? GridSizeToCompiledGraph->Find(GenerationGridSize) : nullptr;
 
-		const TMap<uint32, FPCGStackContext>* GridSizeToStackContext = TopGraphToStackContextMap.Find(InGraph);
+		const TMap<uint32, FPCGStackContext>* GridSizeToStackContext = Cache.TopGraphToStackContextMap.Find(InGraph);
 		const FPCGStackContext* StackContext = GridSizeToStackContext ? GridSizeToStackContext->Find(GenerationGridSize) : nullptr;
 
 		// Should have either found both or neither.
@@ -1035,11 +1173,11 @@ TArray<FPCGGraphTask> FPCGGraphCompiler::GetCompiledTasks(UPCGGraph* InGraph, ui
 		Compile(InGraph);
 
 		// Get compiled tasks in a threadsafe way
-		FReadScopeLock Lock(GraphToTaskMapLock);
-		if (TArray<FPCGGraphTask>* Tasks = GraphToTaskMap.Find(InGraph))
+		FReadScopeLock Lock(Cache.GraphToTaskMapLock);
+		if (TArray<FPCGGraphTask>* Tasks = Cache.GraphToTaskMap.Find(InGraph))
 		{
 			CompiledTasks = *Tasks;
-			OutStackContext = GraphToStackContext[InGraph];
+			OutStackContext = Cache.GraphToStackContextMap[InGraph];
 		}
 	}
 
@@ -1072,10 +1210,13 @@ void FPCGGraphCompiler::OffsetNodeIds(TArray<FPCGGraphTask>& Tasks, FPCGTaskId O
 
 void FPCGGraphCompiler::CompileTopGraph(UPCGGraph* InGraph, uint32 GenerationGridSize)
 {
-	GraphToTaskMapLock.ReadLock();
-	const TMap<uint32, TArray<FPCGGraphTask>>* GridSizeToCompiledGraph = TopGraphToTaskMap.Find(InGraph);
-	const bool bAlreadyCached = GridSizeToCompiledGraph && GridSizeToCompiledGraph->Contains(GenerationGridSize);
-	GraphToTaskMapLock.ReadUnlock();
+	bool bAlreadyCached = false;
+
+	{
+		FReadScopeLock Lock(Cache.GraphToTaskMapLock);
+		const TMap<uint32, TArray<FPCGGraphTask>>* GridSizeToCompiledGraph = Cache.TopGraphToTaskMap.Find(InGraph);
+		bAlreadyCached = GridSizeToCompiledGraph && GridSizeToCompiledGraph->Contains(GenerationGridSize);
+	}
 
 	if (bAlreadyCached)
 	{
@@ -1115,7 +1256,7 @@ void FPCGGraphCompiler::CompileTopGraph(UPCGGraph* InGraph, uint32 GenerationGri
 			ResolveGridSizes(GenerationGrid, CompiledTasks, StackContext, DefaultGrid, TaskGenerationGrid);
 
 			// Create linkage tasks for edges that cross from large grid to small grid tasks.
-			CreateGridLinkages(GenerationGrid, TaskGenerationGrid, CompiledTasks, StackContext);
+			CreateGridLinkages(InGraph, GenerationGrid, TaskGenerationGrid, CompiledTasks, StackContext, bIsCooking);
 
 			// Cull any task that should not execute on the current grid.
 			CullTasks(CompiledTasks, /*bAddPassthroughWires=*/false, [GenerationGrid, &TaskGenerationGrid](const FPCGGraphTask& InTask)
@@ -1130,7 +1271,7 @@ void FPCGGraphCompiler::CompileTopGraph(UPCGGraph* InGraph, uint32 GenerationGri
 #if WITH_EDITOR
 	if (PCGGraphCompiler::CVarEnableGPUExecution.GetValueOnAnyThread())
 	{
-		FPCGGraphCompilerGPU::CreateGPUNodes(InGraph, CompiledTasks);
+		FPCGGraphCompilerGPU::CreateGPUNodes(*this, InGraph, GenerationGridSize, CompiledTasks);
 	}
 	else
 #endif
@@ -1200,7 +1341,7 @@ void FPCGGraphCompiler::CompileTopGraph(UPCGGraph* InGraph, uint32 GenerationGri
 	const FPCGTaskId PostExecuteTaskId = PreExecuteTaskId + 1;
 
 	FPCGGraphTask& PreExecuteTask = CompiledTasks.Emplace_GetRef();
-	PreExecuteTask.Element = GetSharedTrivialElement();
+	PreExecuteTask.ElementSource = EPCGElementSource::Trivial;
 	PreExecuteTask.NodeId = PreExecuteTaskId;
 
 	for (int TaskIndex = 0; TaskIndex < TaskNum; ++TaskIndex)
@@ -1215,7 +1356,7 @@ void FPCGGraphCompiler::CompileTopGraph(UPCGGraph* InGraph, uint32 GenerationGri
 	}
 
 	FPCGGraphTask& PostExecuteTask = CompiledTasks.Emplace_GetRef();
-	PostExecuteTask.Element = GetSharedTrivialPostGraphElement();
+	PostExecuteTask.ElementSource = EPCGElementSource::TrivialPostGraph;
 	PostExecuteTask.NodeId = PostExecuteTaskId;
 
 	// Find end nodes, e.g. all nodes that have no successors.
@@ -1244,19 +1385,19 @@ void FPCGGraphCompiler::CompileTopGraph(UPCGGraph* InGraph, uint32 GenerationGri
 	}
 
 	// Store back the results in the cache
-	GraphToTaskMapLock.WriteLock();
-	TMap<uint32, TArray<FPCGGraphTask>>& TasksPerGenerationGrid = TopGraphToTaskMap.FindOrAdd(InGraph);
+	Cache.GraphToTaskMapLock.WriteLock();
+	TMap<uint32, TArray<FPCGGraphTask>>& TasksPerGenerationGrid = Cache.TopGraphToTaskMap.FindOrAdd(InGraph);
 	if (!TasksPerGenerationGrid.Contains(GenerationGridSize))
 	{
 		TasksPerGenerationGrid.Add(GenerationGridSize, MoveTemp(CompiledTasks));
 	}
 
-	TMap<uint32, FPCGStackContext>& StackContextPerGenerationGrid = TopGraphToStackContextMap.FindOrAdd(InGraph);
+	TMap<uint32, FPCGStackContext>& StackContextPerGenerationGrid = Cache.TopGraphToStackContextMap.FindOrAdd(InGraph);
 	if (!StackContextPerGenerationGrid.Contains(GenerationGridSize))
 	{
 		StackContextPerGenerationGrid.Add(GenerationGridSize, MoveTemp(StackContext));
 	}
-	GraphToTaskMapLock.WriteUnlock();
+	Cache.GraphToTaskMapLock.WriteUnlock();
 }
 
 FPCGElementPtr FPCGGraphCompiler::GetSharedTrivialElement()
@@ -1304,7 +1445,7 @@ FPCGElementPtr FPCGGraphCompiler::GetSharedGatherElement()
 FPCGElementPtr FPCGGraphCompiler::GetSharedTrivialPostGraphElement()
 {
 	{
-		FReadScopeLock Lock(SharedTrivialElementLock);
+		FReadScopeLock Lock(SharedTrivialPostGraphElementLock);
 
 		if (SharedTrivialPostGraphElement)
 		{
@@ -1312,7 +1453,7 @@ FPCGElementPtr FPCGGraphCompiler::GetSharedTrivialPostGraphElement()
 		}
 	}
 
-	FWriteScopeLock Lock(SharedTrivialElementLock);
+	FWriteScopeLock Lock(SharedTrivialPostGraphElementLock);
 
 	if (!SharedTrivialPostGraphElement)
 	{
@@ -1324,11 +1465,12 @@ FPCGElementPtr FPCGGraphCompiler::GetSharedTrivialPostGraphElement()
 
 void FPCGGraphCompiler::ClearCache()
 {
-	FWriteScopeLock Lock(GraphToTaskMapLock);
-	GraphToTaskMap.Reset();
-	GraphToStackContext.Reset();
-	TopGraphToTaskMap.Reset();
-	TopGraphToStackContextMap.Reset();
+	FWriteScopeLock Lock(Cache.GraphToTaskMapLock);
+	Cache.GraphToTaskMap.Reset();
+	Cache.GraphToStackContextMap.Reset();
+	Cache.TopGraphToTaskMap.Reset();
+	Cache.TopGraphToStackContextMap.Reset();
+	Cache.TopGraphToComputeGraphMap.Reset();
 }
 
 #if WITH_EDITOR
@@ -1336,7 +1478,7 @@ void FPCGGraphCompiler::NotifyGraphChanged(UPCGGraph* InGraph, EPCGChangeType Ch
 {
 	if (InGraph && (ChangeType != EPCGChangeType::Cosmetic))
 	{
-		RemoveFromCache(InGraph);
+		Cache.RemoveFromCache(InGraph);
 	}
 }
 
@@ -1346,7 +1488,7 @@ bool FPCGGraphCompiler::Recompile(UPCGGraph* InGraph, uint32 GenerationGridSize,
 	const TArray<FPCGGraphTask> TasksBefore = GetPrecompiledTasks(InGraph, GenerationGridSize, StackContextBefore, bIsTopGraph);
 
 	// Need to manually purge as the graph compiler will not have gotten the change notification yet. Editor only.
-	RemoveFromCache(InGraph);
+	Cache.RemoveFromCache(InGraph);
 
 	FPCGStackContext StackContextAfter;
 	const TArray<FPCGGraphTask> TasksAfter = GetCompiledTasks(InGraph, GenerationGridSize, StackContextAfter, bIsTopGraph);
@@ -1363,33 +1505,4 @@ bool FPCGGraphCompiler::Recompile(UPCGGraph* InGraph, uint32 GenerationGridSize,
 	// Compiled result is compiled tasks + associated stacks. Compare both and return true if compiled result changes.
 	return !bAllTasksEqual || (StackContextBefore != StackContextAfter);
 }
-
-void FPCGGraphCompiler::RemoveFromCache(UPCGGraph* InGraph)
-{
-	UE_LOG(LogPCG, Verbose, TEXT("FPCGGraphCompiler::RemoveFromCache '%s'"), *InGraph->GetName());
-
-	check(InGraph);
-	RemoveFromCacheRecursive(InGraph);
-}
-
-void FPCGGraphCompiler::RemoveFromCacheRecursive(UPCGGraph* InGraph)
-{
-	GraphToTaskMapLock.WriteLock();
-	GraphToTaskMap.Remove(InGraph);
-	GraphToStackContext.Remove(InGraph);
-	TopGraphToTaskMap.Remove(InGraph);
-	TopGraphToStackContextMap.Remove(InGraph);
-	GraphToTaskMapLock.WriteUnlock();
-
-	GraphDependenciesLock.Lock();
-	TArray<UPCGGraph*> ParentGraphs;
-	GraphDependencies.MultiFind(InGraph, ParentGraphs);
-	GraphDependencies.Remove(InGraph);
-	GraphDependenciesLock.Unlock();
-
-	for (UPCGGraph* Graph : ParentGraphs)
-	{
-		RemoveFromCacheRecursive(Graph);
-	}
-}
 #endif // WITH_EDITOR
diff --git a/Engine/Plugins/PCG/Source/PCG/Private/Graph/PCGGraphCompiler.h b/Engine/Plugins/PCG/Source/PCG/Private/Graph/PCGGraphCompiler.h
--- a/Engine/Plugins/PCG/Source/PCG/Private/Graph/PCGGraphCompiler.h
+++ b/Engine/Plugins/PCG/Source/PCG/Private/Graph/PCGGraphCompiler.h
@@ -3,25 +3,74 @@
 #pragma once
 
 #include "PCGCommon.h"
+#include "Graph/PCGGraphTask.h"
 #include "Graph/PCGStackContext.h"
 
 class IPCGElement;
+class UPCGComputeGraph;
 class UPCGGraph;
-struct FPCGGraphTask;
 
 typedef TSharedPtr<IPCGElement, ESPMode::ThreadSafe> FPCGElementPtr;
 
+/** Holds compilation results from FPCGGraphCompiler. */
+struct FPCGGraphCompilerCache
+{
+	friend class FPCGGraphCompiler;
+	friend class UPCGGraph;
+#if WITH_EDITOR
+	friend class FPCGGraphCompilerGPU;
+#endif
+
+public:
+	UPCGComputeGraph* GetCompiledComputeGraph(const UPCGGraph* InGraph, uint32 GridSize, uint32 ComputeGraphIndex);
+
+#if WITH_EDITOR
+	void RemoveFromCache(UPCGGraph* InGraph);
+#endif // WITH_EDITOR
+
+private:
+#if WITH_EDITOR
+	void RemoveFromCacheRecursive(UPCGGraph* InGraph);
+#endif // WITH_EDITOR
+
+	mutable FRWLock GraphToTaskMapLock;
+	TMap<UPCGGraph*, TArray<FPCGGraphTask>> GraphToTaskMap;
+	TMap<UPCGGraph*, FPCGStackContext> GraphToStackContextMap;
+
+	// Top graphs are optimized for execution grid and store one set of compiled tasks per grid size.
+	TMap<UPCGGraph*, TMap<uint32, TArray<FPCGGraphTask>>> TopGraphToTaskMap;
+	TMap<UPCGGraph*, TMap<uint32, FPCGStackContext>> TopGraphToStackContextMap;
+
+	// A single graph can have multiple compute graphs at each grid size, so we store multiple and refer to them by index when querying them.
+	TMap<UPCGGraph*, TMap<uint32, TArray<TObjectPtr<UPCGComputeGraph>>>> TopGraphToComputeGraphMap;
+
+#if WITH_EDITOR
+	FCriticalSection GraphDependenciesLock;
+	TMultiMap<UPCGGraph*, UPCGGraph*> GraphDependencies;
+#endif // WITH_EDITOR
+};
+
 /** 
-* FPCGGraphCompiler
-* This class compiles a graph into tasks and keeps an internal cache
+* This class compiles a graph into tasks and populates an FPCGGraphCompilerCache.
 */
 class FPCGGraphCompiler
 {
 public:
+	explicit FPCGGraphCompiler(bool bInIsCooking = false)
+		: bIsCooking(bInIsCooking)
+	{}
+
 	void Compile(UPCGGraph* InGraph);
 	TArray<FPCGGraphTask> GetCompiledTasks(UPCGGraph* InGraph, uint32 GenerationGridSize, FPCGStackContext& OutStackContext, bool bIsTopGraph = true);
 	TArray<FPCGGraphTask> GetPrecompiledTasks(const UPCGGraph* InGraph, uint32 GenerationGridSize, FPCGStackContext& OutStackContext, bool bIsTopGraph = true) const;
 
+	FPCGGraphCompilerCache& GetCache() { return Cache; }
+	const FPCGGraphCompilerCache& GetCache() const { return Cache; }
+
+	UPCGComputeGraph* GetComputeGraph(const UPCGGraph* InGraph, uint32 GridSize, uint32 ComputeGraphIndex);
+
+	bool IsCooking() const { return bIsCooking; }
+
 	static void OffsetNodeIds(TArray<FPCGGraphTask>& Tasks, FPCGTaskId Offset, FPCGTaskId ParentId);
 
 #if WITH_EDITOR
@@ -29,33 +78,32 @@ public:
 	* The cache will be updated with the latest compiled result. Returns true if the compiled result changes;
 	*/
 	bool Recompile(UPCGGraph* InGraph, uint32 GenerationGridSize, bool bIsTopGraph = true);
-
-	void RemoveFromCache(UPCGGraph* InGraph);
-
 	void NotifyGraphChanged(UPCGGraph* InGraph, EPCGChangeType ChangeType);
 #endif
 
 	/** Flush all cached compiled graphs. */
 	void ClearCache();
 
+	/** Returns the trivial element object shared by all tasks that need it. */
+	FPCGElementPtr GetSharedTrivialElement();
+
 	/** Returns the post graph element, which is used to determine results caching behavior */
 	FPCGElementPtr GetSharedTrivialPostGraphElement();
 
 	/** Returns the gather element object shared by all tasks that need it. */
-	static FPCGElementPtr GetSharedGatherElement();
+	FPCGElementPtr GetSharedGatherElement();
 
 	/** Culls tasks based on a given lambda. Never culls the first (input) task in the array. */
 	static void CullTasks(TArray<FPCGGraphTask>& InOutCompiledTasks, bool bAddPassthroughWires, TFunctionRef<bool(const FPCGGraphTask&)> CullTask);
 
+	void AddReferencedObjects(FReferenceCollector& Collector);
+
 private:
 	TArray<FPCGGraphTask> CompileGraph(UPCGGraph* InGraph, FPCGTaskId& NextId, FPCGStackContext& InOutStackContext);
 
 	/** Compiles the top graph and applies culling optimizations if a non-unitialized grid size is provided. */
 	void CompileTopGraph(UPCGGraph* InGraph, uint32 GenerationGridSize);
 
-	/** Returns the trivial element object shared by all tasks that need it. */
-	FPCGElementPtr GetSharedTrivialElement();
-
 	/** Propagates grid sizes through a graph's compiled tasks. */
 	static void ResolveGridSizes(
 		EPCGHiGenGrid GenerationGrid,
@@ -74,10 +122,12 @@ private:
 
 	/** Create linkage tasks for edges that cross from large grid to small grid tasks. */
 	static void CreateGridLinkages(
+		UPCGGraph* InGraph,
 		EPCGHiGenGrid InGenerationGrid,
 		TArray<EPCGHiGenGrid>& TaskGenerationGrid,
 		TArray<FPCGGraphTask>& InOutCompiledTasks,
-		const FPCGStackContext& InStackContext);
+		const FPCGStackContext& InStackContext,
+		bool bIsCooking);
 
 	/** Discovers whether task is on a statically active branch (and needs to be passed to graph executor). */
 	static bool CalculateStaticallyActiveRecursive(FPCGTaskId InTaskId, const TArray<FPCGGraphTask>& InCompiledTasks, TMap<FPCGTaskId, bool>& InTaskIdToActiveFlag);
@@ -91,26 +141,16 @@ private:
 	/** Task will be active if *any* of the upstream pin IDs are active, or if the pin ID list is empty. */
 	static void CalculateDynamicActivePinDependencies(FPCGTaskId InTaskId, TArray<FPCGGraphTask>& InOutCompiledTasks);
 
-	mutable FRWLock GraphToTaskMapLock;
-	TMap<UPCGGraph*, TArray<FPCGGraphTask>> GraphToTaskMap;
-	TMap<UPCGGraph*, FPCGStackContext> GraphToStackContext;
-
-	// Top graphs are optimized for execution grid and store one set of compiled tasks per grid size.
-	TMap<UPCGGraph*, TMap<uint32, TArray<FPCGGraphTask>>> TopGraphToTaskMap;
-	TMap<UPCGGraph*, TMap<uint32, FPCGStackContext>> TopGraphToStackContextMap;
+private:
+	FPCGGraphCompilerCache Cache;
+	bool bIsCooking = false;
 
 	FPCGElementPtr SharedTrivialElement;
-	mutable FRWLock SharedTrivialElementLock;
+	FRWLock SharedTrivialElementLock;
 
 	FPCGElementPtr SharedTrivialPostGraphElement;
+	FRWLock SharedTrivialPostGraphElementLock;
 
-	static FPCGElementPtr SharedGatherElement;
-	static FRWLock SharedGatherElementLock;
-
-#if WITH_EDITOR
-	void RemoveFromCacheRecursive(UPCGGraph* InGraph);
-
-	FCriticalSection GraphDependenciesLock;
-	TMultiMap<UPCGGraph*, UPCGGraph*> GraphDependencies;
-#endif // WITH_EDITOR
+	FPCGElementPtr SharedGatherElement;
+	FRWLock SharedGatherElementLock;
 };
diff --git a/Engine/Plugins/PCG/Source/PCG/Private/Graph/PCGGraphCompilerGPU.cpp b/Engine/Plugins/PCG/Source/PCG/Private/Graph/PCGGraphCompilerGPU.cpp
--- a/Engine/Plugins/PCG/Source/PCG/Private/Graph/PCGGraphCompilerGPU.cpp
+++ b/Engine/Plugins/PCG/Source/PCG/Private/Graph/PCGGraphCompilerGPU.cpp
@@ -20,16 +20,17 @@
 
 #include "ComputeFramework/ComputeKernel.h"
 #include "HAL/IConsoleManager.h"
+#include "Misc/App.h"
 #include "Shader/ShaderTypes.h"
+#include "UObject/Package.h"
 
+#if WITH_EDITOR
 namespace PCGGraphCompilerGPU
 {
-#if WITH_EDITOR
 	static TAutoConsoleVariable<bool> CVarEnableDebugging(
 		TEXT("pcg.GraphExecution.GPU.EnableDebugging"),
 		false,
 		TEXT("Enable verbose logging of GPU compilation and execution."));
-#endif
 }
 
 void FPCGGraphCompilerGPU::LabelConnectedGPUNodeIslands(
@@ -253,7 +254,7 @@ void FPCGGraphCompilerGPU::CollectGPUNodeSubsets(
 	}
 }
 
-void FPCGGraphCompilerGPU::CreateGatherTasksAtGPUInputs(const TSet<FPCGTaskId>& InGPUCompatibleTaskIds, TArray<FPCGGraphTask>& InOutCompiledTasks)
+void FPCGGraphCompilerGPU::CreateGatherTasksAtGPUInputs(UPCGGraph* InGraph, const TSet<FPCGTaskId>& InGPUCompatibleTaskIds, TArray<FPCGGraphTask>& InOutCompiledTasks)
 {
 	using FOriginalInputPinKey = TPair<FPCGTaskId /* Original GPU task */, FName /* Input pin label */>;
 
@@ -304,7 +305,7 @@ void FPCGGraphCompilerGPU::CreateGatherTasksAtGPUInputs(const TSet<FPCGTaskId>&
 			FPCGGraphTask& GatherTask = InOutCompiledTasks.Emplace_GetRef();
 			GatherTask.NodeId = GatherTaskId;
 			GatherTask.ParentId = InOutCompiledTasks[GPUTaskId].ParentId;
-			GatherTask.Element = FPCGGraphCompiler::GetSharedGatherElement();
+			GatherTask.ElementSource = EPCGElementSource::Gather;
 
 			InputPinToGatherTask.Add(PinKey, GatherTaskId);
 		}
@@ -399,13 +400,23 @@ void FPCGGraphCompilerGPU::WireGPUGraphNode(
 				const UPCGPin* OutputPin = UpstreamTask.Node ? UpstreamTask.Node->GetOutputPin(AddedInput.UpstreamPin->Label) : nullptr;
 
 				// Grid Linkages don't have a node associated, so ask the GridLinkageElement for the output pin instead.
-				if (!OutputPin && UpstreamTask.Element && UpstreamTask.Element->IsGridLinkage())
+				if (!OutputPin)
 				{
-					const PCGGraphExecutor::FPCGGridLinkageElement* LinkageElement = static_cast<const PCGGraphExecutor::FPCGGridLinkageElement*>(UpstreamTask.Element.Get());
-					OutputPin = LinkageElement->GetUpstreamPin();
+					// TODO: In the future FPCGGridLinkageElement shoud be stateless (like all other elements), and we should always only consult the
+					// linkage settings for the upstream pin.
+					if (UpstreamTask.Element && UpstreamTask.Element->IsGridLinkage())
+					{
+						const PCGGraphExecutor::FPCGGridLinkageElement* LinkageElement = static_cast<const PCGGraphExecutor::FPCGGridLinkageElement*>(UpstreamTask.Element.Get());
+						OutputPin = LinkageElement->GetUpstreamPin();
+					}
+					else if (UpstreamTask.Settings && UpstreamTask.Settings->IsA<UPCGGridLinkageSettings>())
+					{
+						const UPCGGridLinkageSettings* LinkageSettings = CastChecked<UPCGGridLinkageSettings>(UpstreamTask.Settings);
+						OutputPin = LinkageSettings->UpstreamPin.Get();
+					}
 				}
 
-				if (OutputPin)
+				if (ensure(OutputPin))
 				{
 					OutOutputCPUPinToVirtualPin.Add(OutputPin, VirtualLabel);
 				}
@@ -471,7 +482,9 @@ void FPCGGraphCompilerGPU::WireGPUGraphNode(
 }
 
 void FPCGGraphCompilerGPU::BuildGPUGraphTask(
+	FPCGGraphCompiler& InOutCompiler,
 	UPCGGraph* InGraph,
+	uint32 InGridSize,
 	FPCGTaskId InGPUGraphTaskId,
 	const TSet<FPCGTaskId>& InCollapsedTasks,
 	const FTaskToSuccessors& InTaskSuccessors,
@@ -479,18 +492,16 @@ void FPCGGraphCompilerGPU::BuildGPUGraphTask(
 	const FOriginalToVirtualPin& InOriginalToVirtualPin,
 	const TMap<TObjectPtr<const UPCGPin>, FName>& InOutputCPUPinToVirtualPin)
 {
-	TSharedPtr<FPCGComputeGraphElement> Element = MakeShared<FPCGComputeGraphElement>();
-	InOutCompiledTasks[InGPUGraphTaskId].Element = Element;
-
+	check(InGraph);
 	const FName GraphName = MakeUniqueObjectName(InGraph, UPCGComputeGraph::StaticClass(), InGraph->GetFName());
 
-	UPCGComputeGraph* ComputeGraph = NewObject<UPCGComputeGraph>(InGraph, GraphName);
+	UObject* ComputeGraphOuter = InOutCompiler.IsCooking() ? InGraph : Cast<UObject>(GetTransientPackage());
+	UPCGComputeGraph* ComputeGraph = NewObject<UPCGComputeGraph>(ComputeGraphOuter, GraphName);
 	ComputeGraph->OutputCPUPinToInputGPUPinAlias = InOutputCPUPinToVirtualPin;
-	Element->Graph.Reset(ComputeGraph);
 
 	// Not incredibly useful for us - DG adds GetComponentSource()->GetComponentClass() object which allows it to bind at execution time by class.
 	// But execution code requires it currently.
-	Element->Graph->Bindings.Add(UPCGDataBinding::StaticClass());
+	ComputeGraph->Bindings.Add(UPCGDataBinding::StaticClass());
 
 	// Create data interfaces which allow kernels to read or write data. Each data interface is associated with a node output pin.
 	// For CPU->GPU edges, an upload data interface is created. For GPU->CPU edges, a readback data interface is created.
@@ -958,7 +969,6 @@ void FPCGGraphCompilerGPU::BuildGPUGraphTask(
 				UE_LOG(LogPCG, Log, TEXT("Cooked HLSL:\n%s\n"), *KernelSource->GetSource());
 			}
 
-#if WITH_EDITOR
 			if (PCGGraphCompilerGPU::CVarEnableDebugging.GetValueOnAnyThread())
 			{
 				UE_LOG(LogPCG, Warning, TEXT("ATTRIBUTE LOOK-UP TABLE [%s]"), *Settings->GetDefaultNodeTitle().ToString());
@@ -971,7 +981,6 @@ void FPCGGraphCompilerGPU::BuildGPUGraphTask(
 						*UEnum::GetValueAsString(Pair.Value.Type));
 				}
 			}
-#endif
 
 			// Add functions for external inputs/outputs which must be fulfilled by DIs
 			for (FInterfaceBinding& Binding : KernelWithBindings.InputDataBindings)
@@ -1073,14 +1082,41 @@ void FPCGGraphCompilerGPU::BuildGPUGraphTask(
 		}
 	}
 
-	// Kick off shader compilation (if needed).
-	ComputeGraph->UpdateResources();
+	TMap<uint32, TArray<TObjectPtr<UPCGComputeGraph>>>& GridToComputeGraphs = InOutCompiler.GetCache().TopGraphToComputeGraphMap.FindOrAdd(InGraph);
+	TArray<TObjectPtr<UPCGComputeGraph>>& ComputeGraphs = GridToComputeGraphs.FindOrAdd(InGridSize);
+	const uint32 ComputeGraphIndex = ComputeGraphs.Num();
+	ComputeGraphs.Add(ComputeGraph);
+
+	if (InOutCompiler.IsCooking())
+	{
+		TObjectPtr<UPCGComputeGraphSettings> Settings = NewObject<UPCGComputeGraphSettings>(InGraph);
+		Settings->ComputeGraphIndex = ComputeGraphIndex;
+		InOutCompiledTasks[InGPUGraphTaskId].ElementSource = EPCGElementSource::FromSettings;
+		InOutCompiledTasks[InGPUGraphTaskId].Settings = Settings;
+	}
+	else
+	{
+		InOutCompiledTasks[InGPUGraphTaskId].Element = MakeShared<FPCGComputeGraphElement>(ComputeGraphIndex);
+	}
+
+	if (FApp::CanEverRender())
+	{
+		// Compile shader resources and create render proxies.
+		ComputeGraph->UpdateResources();
+	}
 }
 
-void FPCGGraphCompilerGPU::CreateGPUNodes(UPCGGraph* InGraph, TArray<FPCGGraphTask>& InOutCompiledTasks)
+void FPCGGraphCompilerGPU::CreateGPUNodes(FPCGGraphCompiler& InOutCompiler, UPCGGraph* InGraph, uint32 InGridSize, TArray<FPCGGraphTask>& InOutCompiledTasks)
 {
 	TRACE_CPUPROFILER_EVENT_SCOPE(FPCGGraphCompilerGPU::CreateGPUNodes);
 
+	// Clear out any previous compute graphs for this grid level before adding new ones.
+	if (ensure(InGraph))
+	{
+		TMap<uint32, TArray<TObjectPtr<UPCGComputeGraph>>>& ComputeGraphs = InOutCompiler.GetCache().TopGraphToComputeGraphMap.FindOrAdd(InGraph);
+		ComputeGraphs.FindOrAdd(InGridSize).Reset();
+	}
+
 	TSet<FPCGTaskId> GPUCompatibleTaskIds;
 	GPUCompatibleTaskIds.Reserve(InOutCompiledTasks.Num());
 	for (FPCGTaskId TaskId = 0; TaskId < InOutCompiledTasks.Num(); ++TaskId)
@@ -1111,15 +1147,17 @@ void FPCGGraphCompilerGPU::CreateGPUNodes(UPCGGraph* InGraph, TArray<FPCGGraphTa
 
 	// For input pins at CPU -> GPU boundary, inject gather elements to pre-combine data on CPU side
 	// before passing to GPU.
-	CreateGatherTasksAtGPUInputs(GPUCompatibleTaskIds, InOutCompiledTasks);
+	CreateGatherTasksAtGPUInputs(InGraph, GPUCompatibleTaskIds, InOutCompiledTasks);
 
 	TArray<TSet<FPCGTaskId>> NodeSubsetsToConvertToCFGraph;
 	CollectGPUNodeSubsets(InOutCompiledTasks, TaskSuccessors, GPUCompatibleTaskIds, NodeSubsetsToConvertToCFGraph);
 
 	// Do actual collapsing now, one subset at a time. Each collapse will do all fixup of task ids? That will invalidate
 	// ids in NodeSubsetsToConvertToCFGraph, so may need remap table. But can ignore this for now.
-	for (TSet<FPCGTaskId>& NodeSubsetToConvertToCFGraph : NodeSubsetsToConvertToCFGraph)
+	for (int32 ComputeGraphIndex = 0; ComputeGraphIndex < NodeSubsetsToConvertToCFGraph.Num(); ++ComputeGraphIndex)
 	{
+		TSet<FPCGTaskId>& NodeSubsetToConvertToCFGraph = NodeSubsetsToConvertToCFGraph[ComputeGraphIndex];
+
 		if (NodeSubsetToConvertToCFGraph.IsEmpty())
 		{
 			ensure(false);
@@ -1158,7 +1196,9 @@ void FPCGGraphCompilerGPU::CreateGPUNodes(UPCGGraph* InGraph, TArray<FPCGGraphTa
 
 		// Generate a compute graph from all of the individual GPU tasks.
 		BuildGPUGraphTask(
+			InOutCompiler,
 			InGraph,
+			InGridSize,
 			ComputeGraphTaskId,
 			NodeSubsetToConvertToCFGraph,
 			TaskSuccessors,
@@ -1181,3 +1221,4 @@ void FPCGGraphCompilerGPU::CreateGPUNodes(UPCGGraph* InGraph, TArray<FPCGGraphTa
 		return false;
 	});
 }
+#endif // WITH_EDITOR
diff --git a/Engine/Plugins/PCG/Source/PCG/Private/Graph/PCGGraphCompilerGPU.h b/Engine/Plugins/PCG/Source/PCG/Private/Graph/PCGGraphCompilerGPU.h
--- a/Engine/Plugins/PCG/Source/PCG/Private/Graph/PCGGraphCompilerGPU.h
+++ b/Engine/Plugins/PCG/Source/PCG/Private/Graph/PCGGraphCompilerGPU.h
@@ -4,10 +4,12 @@
 
 #include "PCGCommon.h"
 
+class FPCGGraphCompiler;
 class UPCGGraph;
 class UPCGPin;
 struct FPCGGraphTask;
 
+#if WITH_EDITOR
 class FPCGGraphCompilerGPU
 {
 public:
@@ -33,7 +35,7 @@ public:
 	/** For GPU node inputs that have multiple incident edges, bundle them into a single edge. This is to avoid an inefficient
 	* gather operation on the GPU, and allows data interfaces to pick their data from the compute graph element input data collection
 	* using unique virtual input pin labels. */
-	static void CreateGatherTasksAtGPUInputs(const TSet<FPCGTaskId>& InGPUCompatibleTaskIds, TArray<FPCGGraphTask>& InOutCompiledTasks);
+	static void CreateGatherTasksAtGPUInputs(UPCGGraph* InGraph, const TSet<FPCGTaskId>& InGPUCompatibleTaskIds, TArray<FPCGGraphTask>& InOutCompiledTasks);
 	
 	/** Wires in a compute graph element alongside each set of GPU compatible nodes. The tasks for each node will be culled later. */
 	static void WireGPUGraphNode(
@@ -48,7 +50,9 @@ public:
 	
 	/** Compiles a compute graph. */
 	static void BuildGPUGraphTask(
+		FPCGGraphCompiler& InOutCompiler,
 		UPCGGraph* InGraph,
+		uint32 InGridSize,
 		FPCGTaskId InGPUGraphTaskId,
 		const TSet<FPCGTaskId>& InCollapsedTasks,
 		const FTaskToSuccessors& InTaskSuccessors,
@@ -57,5 +61,6 @@ public:
 		const TMap<TObjectPtr<const UPCGPin>, FName>& InOutputCPUPinToVirtualPin);
 	
 	/** Finds connected subgraphs of GPU - enabled nodes that can be dispatched together and replaces each one with a compute graph. */
-	static void CreateGPUNodes(UPCGGraph* InGraph, TArray<FPCGGraphTask>& InOutCompiledTasks);
+	static void CreateGPUNodes(FPCGGraphCompiler& InOutCompiler, UPCGGraph* InGraph, uint32 InGridSize, TArray<FPCGGraphTask>& InOutCompiledTasks);
 };
+#endif // WITH_EDITOR
diff --git a/Engine/Plugins/PCG/Source/PCG/Private/Graph/PCGGraphExecutor.cpp b/Engine/Plugins/PCG/Source/PCG/Private/Graph/PCGGraphExecutor.cpp
--- a/Engine/Plugins/PCG/Source/PCG/Private/Graph/PCGGraphExecutor.cpp
+++ b/Engine/Plugins/PCG/Source/PCG/Private/Graph/PCGGraphExecutor.cpp
@@ -22,7 +22,6 @@
 
 #include "Algo/AnyOf.h"
 #include "Algo/ForEach.h"
-#include "Containers/Ticker.h"
 #include "GameFramework/Actor.h"
 #include "Misc/ScopeExit.h"
 #include "Tasks/Task.h"
@@ -40,6 +39,8 @@
 #include "FileHelpers.h"
 #endif
 
+#include UE_INLINE_GENERATED_CPP_BY_NAME(PCGGraphExecutor)
+
 static TAutoConsoleVariable<int32> CVarMaxNumTasks(
 	TEXT("pcg.MaxConcurrentTasks"),
 	4096,
@@ -88,20 +89,6 @@ namespace PCGGraphExecutor
 		TEXT("pcg.Graph.MaxWorkerTasks"),
 		32,
 		TEXT("Max in flight PCG Element tasks allowed at the same time. Note that PCG Element execution can spawn async tasks themselves"));
-	
-
-	template<typename FunctorType>
-	void ExecuteOnGameThread(const TCHAR* DebugName, FunctorType&& Functor)
-	{
-		if (IsInGameThread())
-		{
-			Functor();
-		}
-		else
-		{
-			::ExecuteOnGameThread(DebugName, std::forward<FunctorType>(Functor));
-		}
-	}
 
 	void ClearAsyncFlags(TSet<TObjectPtr<UObject>>& AsyncObjects)
 	{
@@ -191,155 +178,6 @@ namespace PCGGraphExecutor
 	};
 }
 
-int32 FPCGGraphActiveTask::NumExecuting = 0;
-
-const FPCGStack* FPCGGraphTask::GetStack() const
-{
-	return (StackContext && StackIndex != INDEX_NONE) ? StackContext->GetStack(StackIndex) : nullptr;
-}
-
-FPCGTaskId FPCGGraphTask::GetGraphExecutionTaskId() const
-{
-	return StackContext ? StackContext->GetGraphExecutionTaskId() : InvalidPCGTaskId;
-}
-
-bool FPCGGraphTask::CanExecuteOnlyOnMainThread() const
-{
-	return (Element && Element->CanExecuteOnlyOnMainThread(Context)) || (Context && Context->CanExecuteOnlyOnMainThread());
-}
-
-#if WITH_EDITOR
-void FPCGGraphTask::LogVisual(ELogVerbosity::Type InVerbosity, const FText& InMessage) const
-{
-	if (!SourceComponent.IsValid())
-	{
-		return;
-	}
-
-	if (UPCGSubsystem* Subsystem = UPCGSubsystem::GetInstance(SourceComponent->GetWorld()))
-	{
-		const FPCGStack* TaskStack = GetStack();
-		FPCGStack StackWithNode = TaskStack ? FPCGStack(*TaskStack) : FPCGStack();
-		StackWithNode.PushFrame(Node);
-		Subsystem->GetNodeVisualLogsMutable().Log(StackWithNode, InVerbosity, InMessage);
-	}
-}
-
-bool FPCGGraphTaskInput::operator==(const FPCGGraphTaskInput& Other) const
-{
-	return (TaskId == Other.TaskId)
-		&& (UpstreamPin == Other.UpstreamPin)
-		&& (DownstreamPin == Other.DownstreamPin)
-		&& (bProvideData == Other.bProvideData);
-}
-
-bool FPCGGraphTask::IsApproximatelyEqual(const FPCGGraphTask& Other) const
-{
-	// Do trivial pointer comparisons first, then run == operator to determine equivalence.
-	bool bElementsMatch = (Element == Other.Element);
-	if (!bElementsMatch && Element && Other.Element)
-	{
-		if (Element->IsGridLinkage() && Other.Element->IsGridLinkage())
-		{
-			const PCGGraphExecutor::FPCGGridLinkageElement& LinkageElement = static_cast<const PCGGraphExecutor::FPCGGridLinkageElement&>(*Element);
-			const PCGGraphExecutor::FPCGGridLinkageElement& OtherLinkageElement = static_cast<const PCGGraphExecutor::FPCGGridLinkageElement&>(*Other.Element);
-			bElementsMatch = (LinkageElement == OtherLinkageElement);
-		}
-		else if (Element->IsComputeGraphElement() && Other.Element->IsComputeGraphElement())
-		{
-			const FPCGComputeGraphElement& ComputeGraphElement = static_cast<const FPCGComputeGraphElement&>(*Element);
-			const FPCGComputeGraphElement& OtherComputeGraphElement = static_cast<const FPCGComputeGraphElement&>(*Other.Element);
-			bElementsMatch = (ComputeGraphElement == OtherComputeGraphElement);
-		}
-		else
-		{
-			ensureMsgf(false, TEXT("Graph compilation emitted an element type that is not a trivial element or a grid linkage element. Element comparison will fail. ")
-				TEXT("Equivalence operator needs to be implemented for this new element."));
-		}
-	}
-
-	return (Inputs == Other.Inputs)
-		&& (Node == Other.Node)
-		&& (SourceComponent == Other.SourceComponent)
-		&& bElementsMatch
-		&& (Context == Other.Context)
-		&& (NodeId == Other.NodeId)
-		&& (CompiledTaskId == Other.CompiledTaskId)
-		&& (ParentId == Other.ParentId)
-		&& (PinDependency == Other.PinDependency)
-		&& (StackIndex == Other.StackIndex)
-		&& (StackContext == Other.StackContext);
-}
-#endif // WITH_EDITOR
-
-FPCGGraphActiveTask::~FPCGGraphActiveTask()
-{
-	if (Context)
-	{
-		PCGGraphExecutor::ExecuteOnGameThread(UE_SOURCE_LOCATION, [ContextPtr = Context.Release()]()
-		{
-			PCGGraphExecutor::ClearAsyncFlags(ContextPtr->AsyncObjects);
-			delete ContextPtr;
-		});	
-	}
-}
-
-void FPCGGraphActiveTask::StartExecuting()
-{
-	check(!bIsExecutingTask);
-	++NumExecuting;
-
-	auto GatherReferences = [this](const FPCGDataCollection& Collection)
-	{
-		for (const FPCGTaggedData& TaggedInputData : Collection.TaggedData)
-		{
-			if (TaggedInputData.Data)
-			{
-				ExecutingReferences.Add(TaggedInputData.Data);
-			}
-		}
-	};
-
-	// @todo_pcg this should be merged through some kind of visitor pattern with AddStructReferencedObjects
-	// We need to keep a immutable list of references prior to execution mainly so that TaggedData contents don't change while we are running GC
-	{
-		ExecutingReferences.Reserve(Context->InputData.TaggedData.Num() + Context->OutputData.TaggedData.Num() + 1);
-		GatherReferences(Context->InputData);
-		GatherReferences(Context->OutputData);
-
-		// Additional cached per-data
-		for (TPair<FPCGDataCollection, FPCGDataCollection>& CachedInputToOutput : Context->CachedInputToOutputInternalResults)
-		{
-			GatherReferences(CachedInputToOutput.Key);
-			GatherReferences(CachedInputToOutput.Value);
-		}
-
-		if (Context->SettingsWithOverride)
-		{
-			ExecutingReferences.Add(Context->SettingsWithOverride);
-		}
-	}
-
-	bIsExecutingTask = true;
-}
-
-void FPCGGraphActiveTask::StopExecuting()
-{
-	if (bIsExecutingTask)
-	{
-		bIsExecutingTask = false;
-		ExecutingTask = {};
-		--NumExecuting;
-
-		ExecutingReferences.Empty();
-	}
-}
-
-bool FPCGGraphActiveTask::CanExecuteOnlyOnMainThread() const
-{
-	return (Element && Element->CanExecuteOnlyOnMainThread(Context.Get())) || (Context && Context->CanExecuteOnlyOnMainThread());
-}
-
 FPCGGraphExecutor::FPCGGraphExecutor()
 	: FPCGGraphExecutor(nullptr)
 {
@@ -1228,13 +1066,38 @@ bool FPCGGraphExecutor::SetupTask(FPCGGraphTask& Task, TArray<FPCGTaskId>& Resul
 	// Initialize the element if needed (required to know whether it will run on the main thread or not)
 	if (!Task.Element)
 	{
-		// Get appropriate settings
-		check(Task.Node);
-		const UPCGSettings* Settings = Task.TaskInput.GetSettings(Task.Node->GetSettings());
+		// TODO: It may be more efficient to do this in UPCGGraph::PostLoad(), since these are only related to cooking (with the exception of EPCGElementSource::FromNode).
+		if (Task.ElementSource == EPCGElementSource::Trivial)
+		{
+			Task.Element = GraphCompiler.GetSharedTrivialElement();
+		}
+		else if (Task.ElementSource == EPCGElementSource::TrivialPostGraph)
+		{
+			Task.Element = GraphCompiler.GetSharedTrivialPostGraphElement();
+		}
+		else if (Task.ElementSource == EPCGElementSource::Gather)
+		{
+			Task.Element = GraphCompiler.GetSharedGatherElement();
+		}
+		else if (Task.ElementSource == EPCGElementSource::FromNode)
+		{
+			// Get appropriate settings
+			check(Task.Node);
+			const UPCGSettings* Settings = Task.TaskInput.GetSettings(Task.Node->GetSettings());
 
-		if (Settings)
+			if (Settings)
+			{
+				Task.Element = Settings->GetElement();
+			}
+		}
+		else if (Task.ElementSource == EPCGElementSource::FromSettings)
+		{
+			check(Task.Settings);
+			Task.Element = Task.Settings->GetElement();
+		}
+		else
 		{
-			Task.Element = Settings->GetElement();
+			checkNoEntry();
 		}
 	}
 
@@ -1451,6 +1314,7 @@ bool FPCGGraphExecutor::ExecuteScheduling(double EndTime, TSharedPtr<FPCGGraphAc
 					}
 
 					const bool bIsGameThreadOnly = SleepingTask->CanExecuteOnlyOnMainThread();
+
 					TArray<TSharedPtr<FPCGGraphActiveTask>>& ActiveTaskArray = bIsGameThreadOnly ? ActiveTasksGameThreadOnly : ActiveTasks;
 					SleepingTask->bIsGameThreadOnly = bIsGameThreadOnly;
 
@@ -2649,6 +2513,8 @@ void FPCGGraphExecutor::AddReferencedObjects(FReferenceCollector& Collector)
 			It.Value().AddStructReferencedObjects(Collector);
 		}
 	}
+
+	GraphCompiler.AddReferencedObjects(Collector);
 }
 
 FPCGGraphExecutor::FGraphExecutionCache::FGraphExecutionCache()
@@ -3209,3 +3075,24 @@ namespace PCGGraphExecutor
 		return true;
 	}
 }
+
+FPCGElementPtr UPCGGridLinkageSettings::CreateElement() const
+{
+	auto GridLinkageOperation = [this](FPCGContext* InContext)
+	{
+		return PCGGraphExecutor::ExecuteGridLinkage(
+			GenerationGrid,
+			FromGrid,
+			ToGrid,
+			ResourceKey,
+			UpstreamPin.Get() ? UpstreamPin->Properties.Label : NAME_None,
+			static_cast<FPCGGridLinkageContext*>(InContext));
+	};
+
+	FPCGGenericElement::FContextAllocator ContextAllocator = [](const FPCGDataCollection&, TWeakObjectPtr<UPCGComponent>, const UPCGNode*)
+	{
+		return new FPCGGridLinkageContext();
+	};
+
+	return MakeShared<PCGGraphExecutor::FPCGGridLinkageElement>(GridLinkageOperation, ContextAllocator, FromGrid, ToGrid, GenerationGrid, ResourceKey, UpstreamPin.Get());
+}
diff --git a/Engine/Plugins/PCG/Source/PCG/Private/Graph/PCGGraphExecutor.h b/Engine/Plugins/PCG/Source/PCG/Private/Graph/PCGGraphExecutor.h
--- a/Engine/Plugins/PCG/Source/PCG/Private/Graph/PCGGraphExecutor.h
+++ b/Engine/Plugins/PCG/Source/PCG/Private/Graph/PCGGraphExecutor.h
@@ -5,28 +5,30 @@
 #include "PCGContext.h"
 #include "PCGElement.h"
 #include "PCGPin.h"
+#include "PCGSettings.h"
 #include "PCGSubsystem.h"
 #include "Graph/PCGGraphCache.h"
 #include "Graph/PCGGraphCompiler.h"
-#include "Graph/PCGPinDependencyExpression.h"
+#include "Graph/PCGGraphTask.h"
 #include "Graph/PCGStackContext.h"
 
+#include "Containers/Ticker.h"
 #include "Misc/SpinLock.h"
-#include "UObject/GCObject.h"
 #include "Tasks/Task.h"
 #include "Templates/UniquePtr.h"
+#include "UObject/GCObject.h"
 
 #if WITH_EDITOR
 #include "Editor/IPCGEditorProgressNotification.h"
 #include "WorldPartition/WorldPartitionHandle.h" // Needed for FWorldPartitionReference
 #endif
 
+#include "PCGGraphExecutor.generated.h"
+
+class UPCGComponent;
 class UPCGGraph;
 class UPCGNode;
-class UPCGComponent;
-class FPCGGraphCompiler;
 struct FPCGStack;
-class FPCGStackContext;
 class FTextFormat;
 
 namespace PCGGraphExecutor
@@ -37,135 +39,22 @@ namespace PCGGraphExecutor
 #if WITH_EDITOR
 	extern PCG_API TAutoConsoleVariable<float> CVarEditorTimePerFrame;
 #endif
-}
 
-struct FPCGGraphTaskInput
-{
-	FPCGGraphTaskInput(FPCGTaskId InTaskId, const TOptional<FPCGPinProperties>& InUpstreamPin = NoPin, const TOptional<FPCGPinProperties>& InDownstreamPin = NoPin, bool bInProvideData = true)
-		: TaskId(InTaskId)
-		, UpstreamPin(InUpstreamPin)
-		, DownstreamPin(InDownstreamPin)
-		, bProvideData(bInProvideData)
+	template<typename FunctorType>
+	void ExecuteOnGameThread(const TCHAR* DebugName, FunctorType&& Functor)
 	{
+		if (IsInGameThread())
+		{
+			Functor();
+		}
+		else
+		{
+			::ExecuteOnGameThread(DebugName, std::forward<FunctorType>(Functor));
+		}
 	}
 
-#if WITH_EDITOR
-	bool operator==(const FPCGGraphTaskInput& Other) const;
-#endif
-
-	FPCGTaskId TaskId;
-
-	/** The upstream output pin from which the input data comes. */
-	TOptional<FPCGPinProperties> UpstreamPin;
-
-	/** The input pin on the task element. */
-	TOptional<FPCGPinProperties> DownstreamPin;
-
-	/** Whether the input provides any data. For the post execute task, only the output node will provide data. */
-	bool bProvideData;
-	
-	/** Whether the input is used multiple times (previous task has multiple successors). Useful to know if the data can be stolen or not. True by default.*/
-	bool bIsUsedMultipleTimes = true;
-
-	static inline const TOptional<FPCGPinProperties> NoPin = TOptional<FPCGPinProperties>();
-};
-
-struct FPCGGraphTask
-{
-#if WITH_EDITOR
-	/** Approximate equivalence. Does not deeply check node settings, nor does it do a deep comparison of the element. */
-	bool IsApproximatelyEqual(const FPCGGraphTask& Other) const;
-
-	/** Because we might not already have a context, but still want to attach some logs to the node, use this utility function */
-	void LogVisual(ELogVerbosity::Type InVerbosity, const FText& InMessage) const;
-#endif
-
-	const FPCGStack* GetStack() const;
-	FPCGTaskId GetGraphExecutionTaskId() const;
-	
-	bool CanExecuteOnlyOnMainThread() const;
-
-	TArray<FPCGGraphTaskInput> Inputs;
-	const UPCGNode* Node = nullptr;
-	TWeakObjectPtr<UPCGComponent> SourceComponent = nullptr;
-	FPCGElementPtr Element; // Added to have tasks that aren't node-bound
-	FPCGContext* Context = nullptr;
-	FPCGTaskId NodeId = InvalidPCGTaskId;
-	FPCGTaskId CompiledTaskId = InvalidPCGTaskId; // the task id as it exists when compiled
-	FPCGTaskId ParentId = InvalidPCGTaskId; // represents the parent sub object graph task, if we were called from one
-
-	/** Conjunction of disjunctions of pin IDs that are required to be active for this task to be active.
-	* Example - keep task if: UpstreamPin0Active && (UpstreamPin1Active || UpstreamPin2Active)
-	*/
-	FPCGPinDependencyExpression PinDependency;
-
-	int32 StackIndex = INDEX_NONE;
-	TSharedPtr<const FPCGStackContext> StackContext;
-
-	// Whether SetupTask has been called on this task
-	bool bHasDoneSetup = false;
-	// BuildTaskInput will initialize this Collection which will later be used by PrepareForExecute
-	FPCGDataCollection TaskInput;
-	// CombineParams call might have created AsyncObjects
-	TSet<TObjectPtr<UObject>> CombineParamsAsyncObjects;
-
-	// Whether PrepareForExecute as been called on this task
-	bool bHasDonePrepareForExecute = false;
-
-#if WITH_EDITOR
-	// Can be true when we want to have debug display on a task but have taken the results from the cache
-	bool bIsBypassed = false;
-#endif
-};
-
-struct FPCGGraphScheduleTask
-{
-	TArray<FPCGGraphTask> Tasks;
-	TWeakObjectPtr<UPCGComponent> SourceComponent = nullptr;
-	int32 FirstTaskIndex = 0;
-	int32 LastTaskIndex = 0;
-	bool bHasAbortCallbacks = false;
-};
-
-struct FPCGGraphActiveTask : TSharedFromThis<FPCGGraphActiveTask>
-{
-	FPCGGraphActiveTask() = default;
-	virtual ~FPCGGraphActiveTask();
-
-	FPCGGraphActiveTask(const FPCGGraphActiveTask&) = delete;
-	FPCGGraphActiveTask& operator=(const FPCGGraphActiveTask&) = delete;
-
-	FPCGGraphActiveTask(FPCGGraphActiveTask&&) = delete;
-	FPCGGraphActiveTask& operator=(FPCGGraphActiveTask&&) = delete;
-
-	void StartExecuting();
-	void StopExecuting();
-
-	bool CanExecuteOnlyOnMainThread() const;
-
-	TArray<FPCGGraphTaskInput> Inputs;
-	FPCGElementPtr Element;
-	TUniquePtr<FPCGContext> Context;
-	FPCGTaskId NodeId = InvalidPCGTaskId;
-	std::atomic<bool> bWasCancelled = false;
-#if WITH_EDITOR
-	bool bIsBypassed = false;
-#endif
-	int32 StackIndex = INDEX_NONE;
-	TSharedPtr<const FPCGStackContext> StackContext;
-		
-	// Those members need to be modified under the FPCGGraphExecutor::LiveTasksLock (unless we are running the old executor path)
-	UE::Tasks::TTask<bool> ExecutingTask;
-	bool bIsExecutingTask = false;
-
-	// Used to know if task should be in ActiveTasks or ActiveTasksGameThreadOnly
-	bool bIsGameThreadOnly = false;
-	// TaskIndex inside ActiveTasks/ActiveTasksGameThreadOnly/SleepingTasks for fast removal
-	int32 TaskIndex = INDEX_NONE;
-
-	static int32 NumExecuting;
-	TArray<TObjectPtr<const UObject>> ExecutingReferences;
-};
+	void ClearAsyncFlags(TSet<TObjectPtr<UObject>>& AsyncObjects);
+}
 
 class FPCGGraphExecutor : public FGCObject
 {
@@ -234,10 +123,10 @@ public:
 	void ClearOutputData(FPCGTaskId InTaskId);
 
 	/** Accessor so PCG tools (e.g. profiler) can easily decode graph task ids **/
-	FPCGGraphCompiler& GetCompiler() { return GraphCompiler; }
+	FPCGGraphCompiler* GetCompiler() { return &GraphCompiler; }
 
 	/** Accessor so PCG tools (e.g. profiler) can easily decode graph task ids **/
-	const FPCGGraphCompiler& GetCompiler() const { return GraphCompiler; }
+	const FPCGGraphCompiler* GetCompiler() const { return &GraphCompiler; }
 
 #if WITH_EDITOR
 	FPCGTaskId ScheduleDebugWithTaskCallback(UPCGComponent* InComponent, TFunction<void(FPCGTaskId, const UPCGNode*, const FPCGDataCollection&)> TaskCompleteCallback);
@@ -296,7 +185,7 @@ private:
 	void SetLandscapePCGData(FPCGTaskId InGraphExecutionTaskId, UPCGData* InData) { SetExecutionCacheData(InGraphExecutionTaskId, EExecutionCacheDataType::LandscapeData, InData); }
 	void SetLandscapeHeightPCGData(FPCGTaskId InGraphExecutionTaskId, UPCGData* InData) { SetExecutionCacheData(InGraphExecutionTaskId, EExecutionCacheDataType::LandscapeHeightData, InData); }
 	void SetOriginalActorPCGData(FPCGTaskId InGraphExecutionTaskId, UPCGData* InData) { SetExecutionCacheData(InGraphExecutionTaskId, EExecutionCacheDataType::OriginalActorData, InData); }
-	
+
 	void ExecuteV1();
 	void ExecuteV2();
 	double GetTickBudgetInSeconds() const;
@@ -429,7 +318,7 @@ private:
 	
 	/** Monotonically increasing id. Should be reset once all tasks are executed, should be protected by the ScheduleLock */
 	FPCGTaskId NextTaskId = 0;
-	
+
 	/** Struct holding different UPCGData caches that we want to compute only once per graph execution */
 	struct FGraphExecutionCache
 	{
@@ -563,10 +452,11 @@ namespace PCGGraphExecutor
 	class FPCGGridLinkageElement : public FPCGGenericElement
 	{
 	public:
-		FPCGGridLinkageElement(TFunction<bool(FPCGContext*)> InOperation, const FContextAllocator& InContextAllocator, EPCGHiGenGrid InFromGrid, EPCGHiGenGrid InToGrid, const FString& InResourceKey, const UPCGPin* InUpstreamPin)
+		FPCGGridLinkageElement(TFunction<bool(FPCGContext*)> InOperation, const FContextAllocator& InContextAllocator, EPCGHiGenGrid InFromGrid, EPCGHiGenGrid InToGrid, EPCGHiGenGrid InGenerationGrid, const FString& InResourceKey, const UPCGPin* InUpstreamPin)
 			: FPCGGenericElement(InOperation, InContextAllocator)
 			, FromGrid(InFromGrid)
 			, ToGrid(InToGrid)
+			, GenerationGrid(InGenerationGrid)
 			, ResourceKey(InResourceKey)
 		{
 			if (IsValid(InUpstreamPin))
@@ -590,6 +480,10 @@ namespace PCGGraphExecutor
 		// These values are stored here so that we can compare two grid linkage elements for equivalence.
 		EPCGHiGenGrid FromGrid = EPCGHiGenGrid::Uninitialized;
 		EPCGHiGenGrid ToGrid = EPCGHiGenGrid::Uninitialized;
+
+		// This tells us which side of the From/To relationship this grid linkage is on.
+		EPCGHiGenGrid GenerationGrid = EPCGHiGenGrid::Uninitialized; 
+
 		FString ResourceKey;
 		TWeakObjectPtr<const UPCGPin> UpstreamPin = nullptr;
 	};
@@ -603,3 +497,29 @@ namespace PCGGraphExecutor
 		FName InUpstreamPinLabel,
 		FPCGGridLinkageContext* InContext);
 }
+
+UCLASS(ClassGroup = (Procedural))
+class UPCGGridLinkageSettings : public UPCGSettings
+{
+	GENERATED_BODY()
+
+protected:
+	virtual FPCGElementPtr CreateElement() const override;
+
+public:
+	UPROPERTY()
+	EPCGHiGenGrid FromGrid = EPCGHiGenGrid::Uninitialized;
+
+	UPROPERTY()
+	EPCGHiGenGrid ToGrid = EPCGHiGenGrid::Uninitialized;
+
+	UPROPERTY()
+	EPCGHiGenGrid GenerationGrid = EPCGHiGenGrid::Uninitialized;
+
+	UPROPERTY()
+	FString ResourceKey;
+
+	UPROPERTY()
+	TSoftObjectPtr<const UPCGPin> UpstreamPin = nullptr;
+};
+
diff --git a/Engine/Plugins/PCG/Source/PCG/Private/Graph/PCGGraphTask.cpp b/Engine/Plugins/PCG/Source/PCG/Private/Graph/PCGGraphTask.cpp
new file mode 100644
--- /dev/null
+++ b/Engine/Plugins/PCG/Source/PCG/Private/Graph/PCGGraphTask.cpp
@@ -0,0 +1,159 @@
+// Copyright Epic Games, Inc. All Rights Reserved.
+
+#include "PCGGraphTask.h"
+
+#include "PCGComponent.h"
+#include "PCGSubsystem.h"
+#include "Compute/Elements/PCGComputeGraphElement.h"
+#include "Graph/PCGGraphExecutor.h"
+
+#include UE_INLINE_GENERATED_CPP_BY_NAME(PCGGraphTask)
+
+int32 FPCGGraphActiveTask::NumExecuting = 0;
+
+const FPCGStack* FPCGGraphTask::GetStack() const
+{
+	return (StackContext && StackIndex != INDEX_NONE) ? StackContext->GetStack(StackIndex) : nullptr;
+}
+
+FPCGTaskId FPCGGraphTask::GetGraphExecutionTaskId() const
+{
+	return StackContext ? StackContext->GetGraphExecutionTaskId() : InvalidPCGTaskId;
+}
+
+bool FPCGGraphTask::CanExecuteOnlyOnMainThread() const
+{
+	return (Element && Element->CanExecuteOnlyOnMainThread(Context)) || (Context && Context->CanExecuteOnlyOnMainThread());
+}
+
+#if WITH_EDITOR
+void FPCGGraphTask::LogVisual(ELogVerbosity::Type InVerbosity, const FText& InMessage) const
+{
+	if (!SourceComponent.IsValid())
+	{
+		return;
+	}
+
+	if (UPCGSubsystem* Subsystem = UPCGSubsystem::GetInstance(SourceComponent->GetWorld()))
+	{
+		const FPCGStack* TaskStack = GetStack();
+		FPCGStack StackWithNode = TaskStack ? FPCGStack(*TaskStack) : FPCGStack();
+		StackWithNode.PushFrame(Node);
+		Subsystem->GetNodeVisualLogsMutable().Log(StackWithNode, InVerbosity, InMessage);
+	}
+}
+
+bool FPCGGraphTaskInput::operator==(const FPCGGraphTaskInput& Other) const
+{
+	return (TaskId == Other.TaskId)
+		&& (UpstreamPin == Other.UpstreamPin)
+		&& (DownstreamPin == Other.DownstreamPin)
+		&& (bProvideData == Other.bProvideData);
+}
+
+bool FPCGGraphTask::IsApproximatelyEqual(const FPCGGraphTask& Other) const
+{
+	// Do trivial pointer comparisons first, then run == operator to determine equivalence.
+	bool bElementsMatch = (Element == Other.Element);
+	if (!bElementsMatch && Element && Other.Element)
+	{
+		if (Element->IsGridLinkage() && Other.Element->IsGridLinkage())
+		{
+			const PCGGraphExecutor::FPCGGridLinkageElement& LinkageElement = static_cast<const PCGGraphExecutor::FPCGGridLinkageElement&>(*Element);
+			const PCGGraphExecutor::FPCGGridLinkageElement& OtherLinkageElement = static_cast<const PCGGraphExecutor::FPCGGridLinkageElement&>(*Other.Element);
+			bElementsMatch = (LinkageElement == OtherLinkageElement);
+		}
+		else if (Element->IsComputeGraphElement() && Other.Element->IsComputeGraphElement())
+		{
+			const FPCGComputeGraphElement& ComputeGraphElement = static_cast<const FPCGComputeGraphElement&>(*Element);
+			const FPCGComputeGraphElement& OtherComputeGraphElement = static_cast<const FPCGComputeGraphElement&>(*Other.Element);
+			bElementsMatch = (ComputeGraphElement == OtherComputeGraphElement);
+		}
+		else
+		{
+			ensureMsgf(false, TEXT("Graph compilation emitted an element type that is not a trivial element or a grid linkage element. Element comparison will fail. ")
+				TEXT("Equivalence operator needs to be implemented for this new element."));
+		}
+	}
+
+	return (Inputs == Other.Inputs)
+		&& (Node == Other.Node)
+		&& (SourceComponent == Other.SourceComponent)
+		&& bElementsMatch
+		&& (Context == Other.Context)
+		&& (NodeId == Other.NodeId)
+		&& (CompiledTaskId == Other.CompiledTaskId)
+		&& (ParentId == Other.ParentId)
+		&& (PinDependency == Other.PinDependency)
+		&& (StackIndex == Other.StackIndex)
+		&& (StackContext == Other.StackContext);
+}
+#endif // WITH_EDITOR
+
+FPCGGraphActiveTask::~FPCGGraphActiveTask()
+{
+	if (Context)
+	{
+		PCGGraphExecutor::ExecuteOnGameThread(UE_SOURCE_LOCATION, [ContextPtr = Context.Release()]()
+		{
+			PCGGraphExecutor::ClearAsyncFlags(ContextPtr->AsyncObjects);
+			delete ContextPtr;
+		});	
+	}
+}
+
+void FPCGGraphActiveTask::StartExecuting()
+{
+	check(!bIsExecutingTask);
+	++NumExecuting;
+
+	auto GatherReferences = [this](const FPCGDataCollection& Collection)
+	{
+		for (const FPCGTaggedData& TaggedInputData : Collection.TaggedData)
+		{
+			if (TaggedInputData.Data)
+			{
+				ExecutingReferences.Add(TaggedInputData.Data);
+			}
+		}
+	};
+
+	// @todo_pcg this should be merged through some kind of visitor pattern with AddStructReferencedObjects
+	// We need to keep a immutable list of references prior to execution mainly so that TaggedData contents don't change while we are running GC
+	{
+		ExecutingReferences.Reserve(Context->InputData.TaggedData.Num() + Context->OutputData.TaggedData.Num() + 1);
+		GatherReferences(Context->InputData);
+		GatherReferences(Context->OutputData);
+
+		// Additional cached per-data
+		for (TPair<FPCGDataCollection, FPCGDataCollection>& CachedInputToOutput : Context->CachedInputToOutputInternalResults)
+		{
+			GatherReferences(CachedInputToOutput.Key);
+			GatherReferences(CachedInputToOutput.Value);
+		}
+
+		if (Context->SettingsWithOverride)
+		{
+			ExecutingReferences.Add(Context->SettingsWithOverride);
+		}
+	}
+
+	bIsExecutingTask = true;
+}
+
+void FPCGGraphActiveTask::StopExecuting()
+{
+	if (bIsExecutingTask)
+	{
+		bIsExecutingTask = false;
+		ExecutingTask = {};
+		--NumExecuting;
+
+		ExecutingReferences.Empty();
+	}
+}
+
+bool FPCGGraphActiveTask::CanExecuteOnlyOnMainThread() const
+{
+	return (Element && Element->CanExecuteOnlyOnMainThread(Context.Get())) || (Context && Context->CanExecuteOnlyOnMainThread());
+}
diff --git a/Engine/Plugins/PCG/Source/PCG/Private/Graph/PCGGraphTask.h b/Engine/Plugins/PCG/Source/PCG/Private/Graph/PCGGraphTask.h
new file mode 100644
--- /dev/null
+++ b/Engine/Plugins/PCG/Source/PCG/Private/Graph/PCGGraphTask.h
@@ -0,0 +1,194 @@
+// Copyright Epic Games, Inc. All Rights Reserved.
+
+#pragma once
+
+#include "PCGData.h"
+#include "PCGElement.h"
+#include "PCGPin.h"
+#include "Graph/PCGPinDependencyExpression.h"
+#include "Graph/PCGStackContext.h"
+
+#include "Tasks/Task.h"
+
+#include "PCGGraphTask.generated.h"
+
+/** Dictates how to recreate an element if it does not exist. */
+UENUM()
+enum class EPCGElementSource
+{
+	/** Element is a shared Trivial element. */
+	Trivial,
+
+	/** Element is a shared Trivial Post Graph element. */
+	TrivialPostGraph,
+
+	/** Element is a shared Gather element. */
+	Gather,
+
+	/** Element should be recreated from the node settings. */
+	FromNode,
+
+	/** Element does not have a node and should be recreated directly from settings. */
+	FromSettings,
+};
+
+USTRUCT()
+struct FPCGGraphTaskInput
+{
+	GENERATED_BODY()
+
+	FPCGGraphTaskInput() = default;
+
+	explicit FPCGGraphTaskInput(FPCGTaskId InTaskId, const TOptional<FPCGPinProperties>& InUpstreamPin = NoPin, const TOptional<FPCGPinProperties>& InDownstreamPin = NoPin, bool bInProvideData = true)
+		: TaskId(InTaskId)
+		, UpstreamPin(InUpstreamPin)
+		, DownstreamPin(InDownstreamPin)
+		, bProvideData(bInProvideData)
+	{
+	}
+
+#if WITH_EDITOR
+	bool operator==(const FPCGGraphTaskInput& Other) const;
+#endif
+
+	UPROPERTY()
+	uint64 TaskId = InvalidPCGTaskId;
+
+	/** The upstream output pin from which the input data comes. */
+	UPROPERTY()
+	TOptional<FPCGPinProperties> UpstreamPin;
+
+	/** The input pin on the task element. */
+	UPROPERTY()
+	TOptional<FPCGPinProperties> DownstreamPin;
+
+	/** Whether the input provides any data. For the post execute task, only the output node will provide data. */
+	UPROPERTY()
+	bool bProvideData = false;
+	
+	/** Whether the input is used multiple times (previous task has multiple successors). Useful to know if the data can be stolen or not. True by default.*/
+	UPROPERTY()
+	bool bIsUsedMultipleTimes = true;
+
+	static inline const TOptional<FPCGPinProperties> NoPin = TOptional<FPCGPinProperties>();
+};
+
+USTRUCT()
+struct FPCGGraphTask
+{
+	GENERATED_BODY()
+
+#if WITH_EDITOR
+	/** Approximate equivalence. Does not deeply check node settings, nor does it do a deep comparison of the element. */
+	bool IsApproximatelyEqual(const FPCGGraphTask& Other) const;
+
+	/** Because we might not already have a context, but still want to attach some logs to the node, use this utility function */
+	void LogVisual(ELogVerbosity::Type InVerbosity, const FText& InMessage) const;
+#endif
+
+	const FPCGStack* GetStack() const;
+	FPCGTaskId GetGraphExecutionTaskId() const;
+	bool CanExecuteOnlyOnMainThread() const;
+
+	UPROPERTY()
+	TArray<FPCGGraphTaskInput> Inputs;
+
+	/** Dictates how to recreate an element if it does not exist. */
+	UPROPERTY()
+	EPCGElementSource ElementSource = EPCGElementSource::FromNode;
+
+	UPROPERTY()
+	TObjectPtr<const UPCGNode> Node = nullptr;
+
+	/** Useful to recreate elements that don't have a node, such as ComputeGraph elements and GridLinkage elements. */
+	UPROPERTY()
+	TObjectPtr<const UPCGSettings> Settings = nullptr;
+
+	UPROPERTY()
+	uint64 NodeId = InvalidPCGTaskId;
+
+	/** The task id as it exists when compiled. */
+	UPROPERTY()
+	uint64 CompiledTaskId = InvalidPCGTaskId;
+
+	/** Represents the parent sub object graph task, if we were called from one. */
+	UPROPERTY()
+	uint64 ParentId = InvalidPCGTaskId;
+
+	/** Conjunction of disjunctions of pin IDs that are required to be active for this task to be active.
+	* Example - keep task if: UpstreamPin0Active && (UpstreamPin1Active || UpstreamPin2Active)
+	*/
+	UPROPERTY()
+	FPCGPinDependencyExpression PinDependency;
+
+	UPROPERTY()
+	int32 StackIndex = INDEX_NONE;
+
+	FPCGElementPtr Element; // Added to have tasks that aren't node-bound
+	TWeakObjectPtr<UPCGComponent> SourceComponent = nullptr;
+	FPCGContext* Context = nullptr;
+	TSharedPtr<const FPCGStackContext> StackContext;
+
+	// Whether SetupTask has been called on this task
+	bool bHasDoneSetup = false;
+	// BuildTaskInput will initialize this Collection which will later be used by PrepareForExecute
+	FPCGDataCollection TaskInput;
+	// CombineParams call might have created AsyncObjects
+	TSet<TObjectPtr<UObject>> CombineParamsAsyncObjects;
+
+	// Whether PrepareForExecute as been called on this task
+	bool bHasDonePrepareForExecute = false;
+
+#if WITH_EDITOR
+	// Can be true when we want to have debug display on a task but have taken the results from the cache
+	bool bIsBypassed = false;
+#endif
+};
+
+struct FPCGGraphScheduleTask
+{
+	TArray<FPCGGraphTask> Tasks;
+	TWeakObjectPtr<UPCGComponent> SourceComponent = nullptr;
+	int32 FirstTaskIndex = 0;
+	int32 LastTaskIndex = 0;
+	bool bHasAbortCallbacks = false;
+};
+
+struct FPCGGraphActiveTask : TSharedFromThis<FPCGGraphActiveTask>
+{
+	FPCGGraphActiveTask() = default;
+	virtual ~FPCGGraphActiveTask();
+
+	FPCGGraphActiveTask(const FPCGGraphActiveTask&) = delete;
+	FPCGGraphActiveTask& operator=(const FPCGGraphActiveTask&) = delete;
+
+	FPCGGraphActiveTask(FPCGGraphActiveTask&&) = delete;
+	FPCGGraphActiveTask& operator=(FPCGGraphActiveTask&&) = delete;
+
+	void StartExecuting();
+	void StopExecuting();
+	bool CanExecuteOnlyOnMainThread() const;
+
+	TArray<FPCGGraphTaskInput> Inputs;
+	FPCGElementPtr Element;
+	TUniquePtr<FPCGContext> Context;
+	FPCGTaskId NodeId = InvalidPCGTaskId;
+	std::atomic<bool> bWasCancelled = false;
+#if WITH_EDITOR
+	bool bIsBypassed = false;
+#endif
+	int32 StackIndex = INDEX_NONE;
+	TSharedPtr<const FPCGStackContext> StackContext;
+		
+	// Those members need to be modified under the FPCGGraphExecutor::LiveTasksLock (unless we are running the old executor path)
+	UE::Tasks::TTask<bool> ExecutingTask;
+	bool bIsExecutingTask = false;
+
+	// Used to know if task should be in ActiveTasks or ActiveTasksGameThreadOnly
+	bool bIsGameThreadOnly = false;
+	// TaskIndex inside ActiveTasks/ActiveTasksGameThreadOnly/SleepingTasks for fast removal
+	int32 TaskIndex = INDEX_NONE;
+
+	static int32 NumExecuting;
+	TArray<TObjectPtr<const UObject>> ExecutingReferences;
+};
diff --git a/Engine/Plugins/PCG/Source/PCG/Private/Graph/PCGPinDependencyExpression.h b/Engine/Plugins/PCG/Source/PCG/Private/Graph/PCGPinDependencyExpression.h
--- a/Engine/Plugins/PCG/Source/PCG/Private/Graph/PCGPinDependencyExpression.h
+++ b/Engine/Plugins/PCG/Source/PCG/Private/Graph/PCGPinDependencyExpression.h
@@ -4,6 +4,8 @@
 
 #include "PCGCommon.h"
 
+#include "PCGPinDependencyExpression.generated.h"
+
 /** Conjunction of disjunctions of pin IDs that are required to be active for this task to be active.
 * 
 * Example:
@@ -13,8 +15,11 @@
 *     If Pin1 is deactivated during execution: [ PinID0, ConjunctionMarker, RemovedMarker, RemovedMarker ]
 *     Expression is now false and the task using this expression can be culled.
 */
+USTRUCT()
 struct FPCGPinDependencyExpression
 {
+	GENERATED_BODY()
+
 	/** Expression building. Appends a pin ID to the current disjunction. */
 	void AddPinDependency(FPCGPinId PinId);
 
@@ -45,5 +50,6 @@ private:
 	/* Signifies a pin removed from the expression, used to set terms false (pin inactive). */
 	inline static const FPCGPinId RemovedTermMarker = std::numeric_limits<FPCGPinId>::max() - 1;
 
-	TArray<FPCGPinId> Expression;
+	UPROPERTY()
+	TArray<uint64> Expression;
 };
diff --git a/Engine/Plugins/PCG/Source/PCG/Private/PCGComponent.cpp b/Engine/Plugins/PCG/Source/PCG/Private/PCGComponent.cpp
--- a/Engine/Plugins/PCG/Source/PCG/Private/PCGComponent.cpp
+++ b/Engine/Plugins/PCG/Source/PCG/Private/PCGComponent.cpp
@@ -2501,7 +2501,7 @@ void UPCGComponent::StoreInspectionData(const FPCGStack* InStack, const UPCGNode
 
 			if (StackFrames.Last().Object.IsValid() && StackFrames.Last().Object->IsA<UPCGSubgraphNode>())
 			{
-				StorePinInspectionDataFromNode(&StackToSubgraphNode, Cast<const UPCGSubgraphNode>(StackFrames.Last().Object)->GetInputPins(), InInputData, InspectionCache);
+				StorePinInspectionDataFromNode(&StackToSubgraphNode, Cast<const UPCGSubgraphNode>(StackFrames.Last().Object.Get())->GetInputPins(), InInputData, InspectionCache);
 			}
 		}
 
diff --git a/Engine/Plugins/PCG/Source/PCG/Private/PCGGraph.cpp b/Engine/Plugins/PCG/Source/PCG/Private/PCGGraph.cpp
--- a/Engine/Plugins/PCG/Source/PCG/Private/PCGGraph.cpp
+++ b/Engine/Plugins/PCG/Source/PCG/Private/PCGGraph.cpp
@@ -9,12 +9,13 @@
 #include "PCGModule.h"
 #include "PCGNode.h"
 #include "PCGPin.h"
-#include "PCGSubsystem.h"
 #include "PCGSubgraph.h"
+#include "PCGSubsystem.h"
 #include "Elements/PCGHiGenGridSize.h"
 #include "Elements/PCGUserParameterGet.h"
 #include "Elements/ControlFlow/PCGQualityBranch.h"
 #include "Elements/ControlFlow/PCGQualitySelect.h"
+#include "Graph/PCGGraphCompilationData.h"
 #include "Graph/PCGGraphCompiler.h"
 #include "Graph/PCGGraphExecutor.h"
 
@@ -395,7 +396,7 @@ void UPCGGraph::PostLoad()
 	{
 		FixInvalidEdges();
 	}
-#endif
+#endif // WITH_EDITOR
 }
 
 bool UPCGGraph::IsEditorOnly() const
@@ -482,6 +483,82 @@ UPCGNode* UPCGGraph::FindNodeWithSettings(const UPCGSettingsInterface* InSetting
 	return NodeFound;
 }
 
+#if WITH_EDITOR
+void UPCGGraph::PreSave(FObjectPreSaveContext ObjectSaveContext)
+{
+	Super::PreSave(ObjectSaveContext);
+
+	if (ObjectSaveContext.IsCooking())
+	{
+		FPCGGraphCompiler GraphCompiler(/*bIsCooking=*/true);
+
+		// Compile graph for all grid sizes in preparation for cooking.
+		if (IsHierarchicalGenerationEnabled())
+		{
+			bool bHasUnbounded = true;
+			PCGHiGenGrid::FSizeArray GridSizes;
+			GetGridSizes(GridSizes, bHasUnbounded);
+
+			for (uint32 GridSize : GridSizes)
+			{
+				FPCGStackContext StackContext;
+
+				GraphCompiler.GetCompiledTasks(this, GridSize, StackContext);
+			}
+
+			if (bHasUnbounded)
+			{
+				FPCGStackContext StackContext;
+				GraphCompiler.GetCompiledTasks(this, PCGHiGenGrid::UnboundedGridSize(), StackContext);
+			}
+		}
+		else
+		{
+			FPCGStackContext StackContext;
+			GraphCompiler.GetCompiledTasks(this, PCGHiGenGrid::UninitializedGridSize(), StackContext);
+		}
+
+		// Move compiled results into cooked results.
+		FPCGGraphCompilerCache& Cache = GraphCompiler.GetCache();
+		TMap<uint32, TArray<FPCGGraphTask>>* CompiledTasks = Cache.TopGraphToTaskMap.Find(this);
+		TMap<uint32, FPCGStackContext>* CompiledStackContexts = Cache.TopGraphToStackContextMap.Find(this);
+		TMap<uint32, TArray<TObjectPtr<UPCGComputeGraph>>>* CompiledComputeGraphs = Cache.TopGraphToComputeGraphMap.Find(this);
+
+		CookedCompilationData = NewObject<UPCGGraphCompilationData>(this);
+
+		if (ensure(CompiledTasks))
+		{
+			CookedCompilationData->Tasks.Reserve(CompiledTasks->Num());
+
+			for (const TPair<uint32, TArray<FPCGGraphTask>>& Pair : *CompiledTasks)
+			{
+				CookedCompilationData->Tasks.Emplace(Pair.Key, std::move(Pair.Value));
+			}
+		}
+
+		if (ensure(CompiledStackContexts))
+		{
+			CookedCompilationData->StackContexts.Reserve(CompiledStackContexts->Num());
+
+			for (const TPair<uint32, FPCGStackContext>& Pair : *CompiledStackContexts)
+			{
+				CookedCompilationData->StackContexts.Emplace(Pair.Key, std::move(Pair.Value));
+			}
+		}
+
+		if (ensure(CompiledComputeGraphs))
+		{
+			CookedCompilationData->ComputeGraphs.Reserve(CompiledComputeGraphs->Num());
+
+			for (const TPair<uint32, TArray<TObjectPtr<UPCGComputeGraph>>>& Pair : *CompiledComputeGraphs)
+			{
+				CookedCompilationData->ComputeGraphs.Emplace(Pair.Key, std::move(Pair.Value));
+			}
+		}
+	}
+}
+#endif
+
 #if WITH_EDITOR
 void UPCGGraph::DeclareConstructClasses(TArray<FTopLevelAssetPath>& OutConstructClasses, const UClass* SpecificSubclass)
 {
@@ -1180,7 +1257,8 @@ bool UPCGGraph::PrimeGraphCompilationCache()
 {
 	UPCGSubsystem* Subsystem = UPCGSubsystem::GetActiveEditorInstance();
 	FPCGGraphCompiler* GraphCompiler = Subsystem ? Subsystem->GetGraphCompiler() : nullptr;
-	if (!Subsystem || !GraphCompiler)
+
+	if (!GraphCompiler)
 	{
 		return false;
 	}
@@ -1199,7 +1277,8 @@ bool UPCGGraph::Recompile()
 {
 	UPCGSubsystem* Subsystem = UPCGSubsystem::GetActiveEditorInstance();
 	FPCGGraphCompiler* GraphCompiler = Subsystem ? Subsystem->GetGraphCompiler() : nullptr;
-	if (!Subsystem || !GraphCompiler)
+
+	if (!GraphCompiler)
 	{
 		return true;
 	}
diff --git a/Engine/Plugins/PCG/Source/PCG/Private/PCGSubsystem.cpp b/Engine/Plugins/PCG/Source/PCG/Private/PCGSubsystem.cpp
--- a/Engine/Plugins/PCG/Source/PCG/Private/PCGSubsystem.cpp
+++ b/Engine/Plugins/PCG/Source/PCG/Private/PCGSubsystem.cpp
@@ -1292,6 +1292,21 @@ FPCGGenSourceManager* UPCGSubsystem::GetGenSourceManager() const
 	return RuntimeGenScheduler ? RuntimeGenScheduler->GenSourceManager : nullptr;
 }
 
+FPCGGraphCompiler* UPCGSubsystem::GetGraphCompiler()
+{
+	return GraphExecutor ? GraphExecutor->GetCompiler() : nullptr;
+}
+
+UPCGComputeGraph* UPCGSubsystem::GetComputeGraph(const UPCGGraph* InGraph, uint32 GridSize, uint32 ComputeGraphIndex)
+{
+	if (FPCGGraphCompiler* GraphCompiler = GetGraphCompiler())
+	{
+		return GraphCompiler->GetComputeGraph(InGraph, GridSize, ComputeGraphIndex);
+	}
+
+	return nullptr;
+}
+
 bool UPCGSubsystem::GetOutputData(FPCGTaskId TaskId, FPCGDataCollection& OutData)
 {
 	check(GraphExecutor);
@@ -1665,16 +1680,6 @@ void UPCGSubsystem::ClearLandscapeCache()
 	}
 }
 
-FPCGGraphCompiler* UPCGSubsystem::GetGraphCompiler()
-{
-	if (GraphExecutor)
-	{
-		return &(GraphExecutor->GetCompiler());
-	}
-
-	return nullptr;
-}
-
 bool UPCGSubsystem::GetStackContext(const UPCGComponent* InComponent, FPCGStackContext& OutStackContext)
 {
 	UPCGGraph* Graph = InComponent ? InComponent->GetGraph() : nullptr;
@@ -1695,14 +1700,15 @@ bool UPCGSubsystem::GetStackContext(const UPCGComponent* InComponent, FPCGStackC
 
 	if (bDoesComponentExecute)
 	{
-		GetGraphCompiler()->GetCompiledTasks(Graph, InComponent->GetGenerationGridSize(), OutStackContext, /*bIsCooking=*/false);
-		return true;
-	}
-	else
-	{
-		OutStackContext = FPCGStackContext();
-		return false;
+		if (FPCGGraphCompiler* GraphCompiler = GetGraphCompiler())
+		{
+			GraphCompiler->GetCompiledTasks(Graph, InComponent->GetGenerationGridSize(), OutStackContext, /*bIsTopGraph=*/false);
+			return true;
+		}
 	}
+
+	OutStackContext = FPCGStackContext();
+	return false;
 }
 
 uint32 UPCGSubsystem::GetGraphCacheEntryCount(IPCGElement* InElement) const
@@ -1924,10 +1930,10 @@ IPCGGraphCache* UPCGSubsystem::GetCache()
 
 void UPCGSubsystem::FlushCache()
 {
-	if (GraphExecutor)
+	if (GraphExecutor && GraphExecutor->GetCompiler())
 	{
 		GraphExecutor->GetCache().ClearCache();
-		GraphExecutor->GetCompiler().ClearCache();
+		GraphExecutor->GetCompiler()->ClearCache();
 	}
 
 #if WITH_EDITOR
diff --git a/Engine/Plugins/PCG/Source/PCG/Public/Graph/PCGStackContext.h b/Engine/Plugins/PCG/Source/PCG/Public/Graph/PCGStackContext.h
--- a/Engine/Plugins/PCG/Source/PCG/Public/Graph/PCGStackContext.h
+++ b/Engine/Plugins/PCG/Source/PCG/Public/Graph/PCGStackContext.h
@@ -24,7 +24,7 @@ struct PCG_API FPCGStackFrame
 
 	FPCGStackFrame() {}
 
-	explicit FPCGStackFrame(TWeakObjectPtr<const UObject> InObject)
+	explicit FPCGStackFrame(const UObject* InObject)
 	{
 		SetObject(InObject);
 	}
@@ -42,7 +42,7 @@ struct PCG_API FPCGStackFrame
 	// A valid frame should either point to an object or have a loop index >= 0.
 	bool IsValid() const { return LoopIndex != INDEX_NONE || Object.IsValid(); }
 
-	void SetObject(TWeakObjectPtr<const UObject> InObject)
+	void SetObject(const UObject* InObject)
 	{
 		Object = InObject;
 		LoopIndex = INDEX_NONE;
@@ -57,12 +57,15 @@ struct PCG_API FPCGStackFrame
 	}
 
 	// Stores object this frame refers to. Use SetObject to change this and properly update the hash.
-	TWeakObjectPtr<const UObject> Object;
+	UPROPERTY()
+	TSoftObjectPtr<const UObject> Object;
 
 	// Stores the loop index (if any) that this frame refers to. Use SetLoopIndex to change this and properly update the hash.
+	UPROPERTY()
 	int32 LoopIndex = INDEX_NONE;
 
 private:
+	UPROPERTY()
 	uint32 Hash = 0;
 };
 
@@ -72,7 +75,7 @@ struct PCG_API FPCGStack
 {
 	GENERATED_BODY()
 
-	friend class FPCGStackContext;
+	friend struct FPCGStackContext;
 
 public:
 #if WITH_EDITOR
@@ -144,6 +147,8 @@ public:
 
 private:
 	FPCGTaskId GraphExecutionTaskId = InvalidPCGTaskId;
+
+	UPROPERTY()
 	TArray<FPCGStackFrame> StackFrames;
 
 public:
@@ -155,9 +160,13 @@ public:
 };
 
 /** A collection of call stacks. */
-class PCG_API FPCGStackContext
+USTRUCT()
+struct PCG_API FPCGStackContext
 {
+	GENERATED_BODY()
+
 	friend class FPCGGraphExecutor;
+
 public:
 	int32 GetNumStacks() const { return Stacks.Num(); }
 	int32 GetCurrentStackIndex() const { return CurrentStackIndex; }
@@ -181,6 +190,7 @@ public:
 	bool operator==(const FPCGStackContext& Other) const;
 
 	FPCGTaskId GetGraphExecutionTaskId() const { return GraphExecutionTaskId; }
+
 private:
 	void SetGraphExecutionTaskId(FPCGTaskId InGraphExecutionTaskId);
 
@@ -188,8 +198,10 @@ private:
 	FPCGTaskId GraphExecutionTaskId = InvalidPCGTaskId;
 
 	/** List of all stacks encountered top graph and all (nested) subgraphs. Order is simply order of encountering during compilation. */
+	UPROPERTY()
 	TArray<FPCGStack> Stacks;
 
 	/** Index of element in Stacks that is the current stack. */
+	UPROPERTY()
 	int32 CurrentStackIndex = INDEX_NONE;
 };
diff --git a/Engine/Plugins/PCG/Source/PCG/Public/PCGComponent.h b/Engine/Plugins/PCG/Source/PCG/Public/PCGComponent.h
--- a/Engine/Plugins/PCG/Source/PCG/Public/PCGComponent.h
+++ b/Engine/Plugins/PCG/Source/PCG/Public/PCGComponent.h
@@ -20,7 +20,6 @@ namespace EEndPlayReason { enum Type : int; }
 
 class APCGPartitionActor;
 class FPCGActorAndComponentMapping;
-class FPCGStackContext;
 class UPCGComponent;
 class UPCGData;
 class IPCGGenSourceBase;
@@ -34,6 +33,7 @@ class ALandscapeProxy;
 class FLandscapeProxyComponentDataChangedParams;
 class UClass;
 struct FPCGContext;
+struct FPCGStackContext;
 
 #if WITH_EDITOR
 	DECLARE_MULTICAST_DELEGATE_OneParam(FOnPCGGraphStartGenerating, UPCGComponent*);
diff --git a/Engine/Plugins/PCG/Source/PCG/Public/PCGGraph.h b/Engine/Plugins/PCG/Source/PCG/Public/PCGGraph.h
--- a/Engine/Plugins/PCG/Source/PCG/Public/PCGGraph.h
+++ b/Engine/Plugins/PCG/Source/PCG/Public/PCGGraph.h
@@ -15,6 +15,7 @@
 #include "PCGGraph.generated.h"
 
 class UPCGGraphInterface;
+class UPCGGraphCompilationData;
 #if WITH_EDITOR
 class UPCGEditorGraph;
 struct FEdGraphPinType;
@@ -218,7 +219,9 @@ public:
 	/** ~Begin UObject interface */
 	virtual void PostLoad() override;
 	virtual bool IsEditorOnly() const override;
+
 #if WITH_EDITOR
+	virtual void PreSave(FObjectPreSaveContext ObjectSaveContext) override;
 	static void DeclareConstructClasses(TArray<FTopLevelAssetPath>& OutConstructClasses, const UClass* SpecificSubclass);
 #endif
 
@@ -377,6 +380,8 @@ public:
 	/** Size of grid on which this node should be executed. Nodes execute at the minimum of all input grid sizes. */
 	uint32 GetNodeGenerationGridSize(const UPCGNode* InNode, uint32 InDefaultGridSize) const;
 
+	const TObjectPtr<UPCGGraphCompilationData> GetCookedCompilationData() const { return CookedCompilationData; }
+
 protected:
 	/** Internal function to react to add/remove nodes. bNotify can be set to false to not notify the world. */
 	void OnNodeAdded(UPCGNode* InNode, bool bNotify = true);
@@ -456,6 +461,13 @@ protected:
 	bool bDebugFlagAppliesToIndividualComponents = true;
 #endif // WITH_EDITORONLY_DATA
 
+	/**
+	 * Populated during cook to prewarm graph compiler cache in standalone builds. Also necessary for GPU execution because compiling
+	 * compute graphs is not supported outside of editor.
+	 */
+	UPROPERTY()
+	TObjectPtr<UPCGGraphCompilationData> CookedCompilationData = nullptr;
+
 public:
 	virtual const FInstancedPropertyBag* GetUserParametersStruct() const override { return &UserParameters; }
 
diff --git a/Engine/Plugins/PCG/Source/PCG/Public/PCGSubsystem.h b/Engine/Plugins/PCG/Source/PCG/Public/PCGSubsystem.h
--- a/Engine/Plugins/PCG/Source/PCG/Public/PCGSubsystem.h
+++ b/Engine/Plugins/PCG/Source/PCG/Public/PCGSubsystem.h
@@ -15,9 +15,10 @@
 
 class APCGPartitionActor;
 class APCGWorldActor;
+class FPCGRuntimeGenScheduler;
+class UPCGComputeGraph;
 class UPCGGraph;
 class UPCGLandscapeCache;
-class FPCGRuntimeGenScheduler;
 
 enum class EPCGComponentDirtyFlag : uint8;
 enum class ETickableTickType : uint8;
@@ -273,6 +274,8 @@ public:
 	bool IsGraphCacheDebuggingEnabled() const;
 
 	FPCGGenSourceManager* GetGenSourceManager() const;
+	FPCGGraphCompiler* GetGraphCompiler();
+	UPCGComputeGraph* GetComputeGraph(const UPCGGraph* InGraph, uint32 GridSize, uint32 ComputeGraphIndex);
 
 #if WITH_EDITOR
 public:
@@ -309,9 +312,6 @@ public:
 	/** Clears the landscape data cache */
 	void ClearLandscapeCache();
 
-	/** Returns the graph compiler so we can figure out task info in the profiler view **/
-	FPCGGraphCompiler* GetGraphCompiler();
-
 	/** Get the execution stack information for the given component. */
 	bool GetStackContext(const UPCGComponent* InComponent, FPCGStackContext& OutStackContext);
 
diff --git a/Engine/Plugins/PCG/Source/PCGEditor/Private/PCGEditor.cpp b/Engine/Plugins/PCG/Source/PCGEditor/Private/PCGEditor.cpp
--- a/Engine/Plugins/PCG/Source/PCGEditor/Private/PCGEditor.cpp
+++ b/Engine/Plugins/PCG/Source/PCGEditor/Private/PCGEditor.cpp
@@ -255,7 +255,24 @@ void FPCGEditor::OnObjectsReplaced(const TMap<UObject*, UObject*>& ReplacementMa
 	TArray<FPCGStackFrame>& StackFrames = StackBeingInspected.GetStackFramesMutable();
 	if (!StackFrames.IsEmpty())
 	{
-		if (UObject* NewStackRoot = ReplacementMap.FindRef(StackFrames[0].Object.GetEvenIfUnreachable()))
+		UObject* NewStackRoot = ReplacementMap.FindRef(StackFrames[0].Object.Get());
+
+		// If the stack frame was marked as garbage, NewStackRoot will be nullptr, but we still match against the object path.
+		if (!NewStackRoot)
+		{
+			for(const TPair<UObject*, UObject*>& Pair : ReplacementMap)
+			{
+				if (Pair.Key && Pair.Value)
+				{
+					if (Pair.Key->GetPathName() == StackFrames[0].Object.ToString())
+					{
+						NewStackRoot = Pair.Value;
+					}
+				}
+			}
+		}
+
+		if (NewStackRoot)
 		{
 			StackFrames[0].SetObject(NewStackRoot);
 		}
diff --git a/Engine/Plugins/PCG/Source/PCGEditor/Private/SPCGEditorGraphDebugObjectTree.cpp b/Engine/Plugins/PCG/Source/PCGEditor/Private/SPCGEditorGraphDebugObjectTree.cpp
--- a/Engine/Plugins/PCG/Source/PCGEditor/Private/SPCGEditorGraphDebugObjectTree.cpp
+++ b/Engine/Plugins/PCG/Source/PCGEditor/Private/SPCGEditorGraphDebugObjectTree.cpp
@@ -141,7 +141,24 @@ void FPCGEditorGraphDebugObjectItem::OnObjectsReplaced(const TMap<UObject*, UObj
 		TArray<FPCGStackFrame>& StackFrames = Stack->GetStackFramesMutable();
 		if (!StackFrames.IsEmpty())
 		{
-			if (UObject* NewStackRoot = ReplacementMap.FindRef(StackFrames[0].Object.GetEvenIfUnreachable()))
+			UObject* NewStackRoot = ReplacementMap.FindRef(StackFrames[0].Object.Get());
+
+			// If the stack frame was marked as garbage, NewStackRoot will be nullptr, but we still match against the object path.
+			if (!NewStackRoot)
+			{
+				for(const TPair<UObject*, UObject*>& Pair : ReplacementMap)
+				{
+					if (Pair.Key && Pair.Value)
+					{
+						if (Pair.Key->GetPathName() == StackFrames[0].Object.ToString())
+						{
+							NewStackRoot = Pair.Value;
+						}
+					}
+				}
+			}
+
+			if (NewStackRoot)
 			{
 				StackFrames[0].SetObject(NewStackRoot);
 			}
@@ -695,7 +712,7 @@ void SPCGEditorGraphDebugObjectTree::AddStacksToTree(const TArray<FPCGStack>& St
 			FPCGEditorGraphDebugObjectItemPtr CurrentItem;
 
 			// When we encounter a graph, we look at the frame index and/or preceding frames to determine the graph type.
-			if (const UPCGGraph* StackGraph = Cast<const UPCGGraph>(StackFrame.Object))
+			if (const UPCGGraph* StackGraph = Cast<const UPCGGraph>(StackFrame.Object.Get()))
 			{
 				const bool bIsDebuggable = (GraphBeingEdited == StackGraph);
 
@@ -717,14 +734,14 @@ void SPCGEditorGraphDebugObjectTree::AddStacksToTree(const TArray<FPCGStack>& St
 					}
 				}
 				// Previous stack was node, therefore static subgraph.
-				else if (const UPCGNode* SubgraphNode = Cast<const UPCGNode>(PreviousStackFrame.Object))
+				else if (const UPCGNode* SubgraphNode = Cast<const UPCGNode>(PreviousStackFrame.Object.Get()))
 				{
 					AddSubgraphOrLoopItemToStack(StackGraph, SubgraphNode, FrameIndex + 1, bIsDebuggable, bDisplayGrayedOut);
 				}
 				// Previous stack was loop index, therefore loop subgraph.
 				else if (FrameIndex >= 2 && (PreviousStackFrame.LoopIndex != INDEX_NONE))
 				{
-					const UPCGNode* LoopSubgraphNode = Cast<const UPCGNode>(StackFrames[FrameIndex - 2].Object);
+					const UPCGNode* LoopSubgraphNode = Cast<const UPCGNode>(StackFrames[FrameIndex - 2].Object.Get());
 					if (ensure(LoopSubgraphNode))
 					{
 						// Take the stack up to the looped subgraph node, add a item for the node + graph
@@ -747,7 +764,7 @@ void SPCGEditorGraphDebugObjectTree::AddStacksToTree(const TArray<FPCGStack>& St
 				// Previous stack was invalid node / node with a INDEX_NONE loop, therefore most likely a dynamic subgraph
 				else if (FrameIndex >= 2 && !PreviousStackFrame.IsValid())
 				{
-					const UPCGNode* DynamicSubgraphNode = Cast<const UPCGNode>(StackFrames[FrameIndex - 2].Object);
+					const UPCGNode* DynamicSubgraphNode = Cast<const UPCGNode>(StackFrames[FrameIndex - 2].Object.Get());
 					if (DynamicSubgraphNode)
 					{
 						AddSubgraphOrLoopItemToStack(StackGraph, DynamicSubgraphNode, FrameIndex + 1, bIsDebuggable, bDisplayGrayedOut);
@@ -1215,8 +1232,8 @@ void SPCGEditorGraphDebugObjectTree::ContextMenu_JumpToGraphInTree()
 			{
 				if (SelectedStack.GetStackFrames()[i].Object != nullptr && SelectedStack.GetStackFrames()[i].Object->IsA<UPCGGraph>())
 				{
-					JumpToPCGGraph = const_cast<UPCGGraph*>(Cast<const UPCGGraph>(SelectedStack.GetStackFrames()[i].Object));
-					JumpToPCGNode = const_cast<UPCGNode*>(Cast<const UPCGNode>(SelectedStack.GetStackFrames()[i + 1].Object));
+					JumpToPCGGraph = const_cast<UPCGGraph*>(Cast<const UPCGGraph>(SelectedStack.GetStackFrames()[i].Object.Get()));
+					JumpToPCGNode = const_cast<UPCGNode*>(Cast<const UPCGNode>(SelectedStack.GetStackFrames()[i + 1].Object.Get()));
 					break;
 				}
 			}
@@ -1231,7 +1248,7 @@ void SPCGEditorGraphDebugObjectTree::ContextMenu_JumpToGraphInTree()
 			const int SubjectNodeFrameIndex = SelectedStack.GetStackFrames().Num();
 			if (PreviouslySelectedStack.GetStackFrames().IsValidIndex(SubjectNodeFrameIndex))
 			{
-				JumpToPCGNode = const_cast<UPCGNode*>(Cast<const UPCGNode>(PreviouslySelectedStack.GetStackFrames()[SubjectNodeFrameIndex].Object));
+				JumpToPCGNode = const_cast<UPCGNode*>(Cast<const UPCGNode>(PreviouslySelectedStack.GetStackFrames()[SubjectNodeFrameIndex].Object.Get()));
 			}
 		}
 
