[PCG] Fix mac compilation error in PCGGraph.cpp

#rnx

diff --git a/Engine/Plugins/PCG/Source/PCG/Private/Graph/PCGGraphCompilationData.h b/Engine/Plugins/PCG/Source/PCG/Private/Graph/PCGGraphCompilationData.h
--- a/Engine/Plugins/PCG/Source/PCG/Private/Graph/PCGGraphCompilationData.h
+++ b/Engine/Plugins/PCG/Source/PCG/Private/Graph/PCGGraphCompilationData.h
@@ -15,6 +15,12 @@ struct FPCGGraphTasks
 {
 	GENERATED_BODY();
 
+	FPCGGraphTasks() = default;
+
+	explicit FPCGGraphTasks(TArray<FPCGGraphTask> InGraphTasks)
+		: GraphTasks(std::move(InGraphTasks))
+	{}
+
 	UPROPERTY()
 	TArray<FPCGGraphTask> GraphTasks;
 };
@@ -25,6 +31,12 @@ struct FPCGComputeGraphs
 {
 	GENERATED_BODY();
 
+	FPCGComputeGraphs() = default;
+
+	explicit FPCGComputeGraphs(TArray<TObjectPtr<UPCGComputeGraph>> InComputeGraphs)
+		: ComputeGraphs(std::move(InComputeGraphs))
+	{}
+
 	UPROPERTY()
 	TArray<TObjectPtr<UPCGComputeGraph>> ComputeGraphs;
 };
diff --git a/Engine/Plugins/PCG/Source/PCG/Private/PCGGraph.cpp b/Engine/Plugins/PCG/Source/PCG/Private/PCGGraph.cpp
--- a/Engine/Plugins/PCG/Source/PCG/Private/PCGGraph.cpp
+++ b/Engine/Plugins/PCG/Source/PCG/Private/PCGGraph.cpp
@@ -532,7 +532,7 @@ void UPCGGraph::PreSave(FObjectPreSaveContext ObjectSaveContext)
 
 			for (const TPair<uint32, TArray<FPCGGraphTask>>& Pair : *CompiledTasks)
 			{
-				CookedCompilationData->Tasks.Emplace(Pair.Key, std::move(Pair.Value));
+				CookedCompilationData->Tasks.Emplace(Pair.Key, FPCGGraphTasks(std::move(Pair.Value)));
 			}
 		}
 
@@ -552,7 +552,7 @@ void UPCGGraph::PreSave(FObjectPreSaveContext ObjectSaveContext)
 
 			for (const TPair<uint32, TArray<TObjectPtr<UPCGComputeGraph>>>& Pair : *CompiledComputeGraphs)
 			{
-				CookedCompilationData->ComputeGraphs.Emplace(Pair.Key, std::move(Pair.Value));
+				CookedCompilationData->ComputeGraphs.Emplace(Pair.Key, FPCGComputeGraphs(std::move(Pair.Value)));
 			}
 		}
 	}
