[PCG] Fixing Copy attributes that was not working with 1 default value and 1 value set (Like in a Merge points scenario)

#rnx

diff --git a/Engine/Plugins/PCG/Source/PCG/Private/Helpers/PCGMetadataHelpers.cpp b/Engine/Plugins/PCG/Source/PCG/Private/Helpers/PCGMetadataHelpers.cpp
--- a/Engine/Plugins/PCG/Source/PCG/Private/Helpers/PCGMetadataHelpers.cpp
+++ b/Engine/Plugins/PCG/Source/PCG/Private/Helpers/PCGMetadataHelpers.cpp
@@ -164,22 +164,15 @@ namespace PCGMetadataHelpers
 
 				if (FPCGMetadataAttributeBase* NewAttr = TargetMetadata->CopyAttribute(SourceAttribute, LocalDestinationAttribute, /*bKeepParent=*/false, /*bCopyEntries=*/bSameOrigin, /*bCopyValues=*/true))
 				{
-					// To keep the previous behavior for single value attributes, we force the copied attribute to have its default value set to the first entry (only done if there is just a single value).
-					const bool bHasJustASingleEntryOrNoEntry = SourceAttribute->GetNumberOfEntriesWithParents() <= 1;
-					if (bHasJustASingleEntryOrNoEntry)
-					{
-						NewAttr->SetDefaultValueToFirstEntry();
-					}
-
 					// And finally, we create our source and target keys, to get all the metadata entry keys, to remap if we are not the same origin and that we have multiple entries.
-					if (!bSameOrigin && !bHasJustASingleEntryOrNoEntry)
+					if (!bSameOrigin)
 					{
 						TUniquePtr<const IPCGAttributeAccessorKeys> SourceKeys = PCGAttributeAccessorHelpers::CreateConstKeys(SourceData, InputSource);
 						TUniquePtr<IPCGAttributeAccessorKeys> TargetKeys = PCGAttributeAccessorHelpers::CreateKeys(TargetData, OutputTarget);
 
 						// They must exist, so we can check.
 						check(SourceKeys && TargetKeys);
-
+						
 						if (SourceKeys->GetNum() > 0 && TargetKeys->GetNum() > 0)
 						{
 							TArray<const PCGMetadataEntryKey*> AllSourceEntryKeysPtr;
@@ -194,14 +187,23 @@ namespace PCGMetadataHelpers
 								TArray<PCGMetadataEntryKey> AllSourceEntryKeys;
 								Algo::Transform(AllSourceEntryKeysPtr, AllSourceEntryKeys, [](const PCGMetadataEntryKey* KeyPtr) { return *KeyPtr; });
 								TArray<PCGMetadataValueKey> ValueKeys;
-								ValueKeys.Reserve(AllTargetEntryKeysPtr.Num());
+								TArray<PCGMetadataValueKey> FinalValueKeys;
+								ValueKeys.Reserve(AllSourceEntryKeys.Num());
 								SourceAttribute->GetValueKeys(TArrayView<const PCGMetadataEntryKey>(AllSourceEntryKeys), ValueKeys);
 
 								// Extends values keys to match target entry keys size. It will loop on value keys.
-								const int32 ValueKeysSize = ValueKeys.Num();
-								for (int32 i = ValueKeysSize; i < AllTargetEntryKeysPtr.Num(); ++i)
+								const int32 TargetEntryKeysSize = AllTargetEntryKeysPtr.Num();
+								if (ValueKeys.Num() >= TargetEntryKeysSize)
 								{
-									ValueKeys.Add(ValueKeys[i % ValueKeysSize]);
+									FinalValueKeys = std::move(ValueKeys);
+								}
+								else
+								{
+									FinalValueKeys.Reserve(TargetEntryKeysSize);
+									while (FinalValueKeys.Num() < TargetEntryKeysSize)
+									{
+										FinalValueKeys.Append(TArrayView<PCGMetadataValueKey>(ValueKeys.GetData(), FMath::Min(ValueKeys.Num(), TargetEntryKeysSize - FinalValueKeys.Num())));
+									}
 								}
 
 								// Make sure that the Target has some metadata entry
@@ -223,7 +225,13 @@ namespace PCGMetadataHelpers
 									TargetMetadata->AddEntriesInPlace(AllTargetEntryKeysPtrTemp);
 								}
 
-								NewAttr->SetValuesFromValueKeys(AllTargetEntryKeysPtr, ValueKeys, /*bResetValueOnDefaultValueKey=*/true);
+								NewAttr->SetValuesFromValueKeys(AllTargetEntryKeysPtr, FinalValueKeys, /*bResetValueOnDefaultValueKey=*/true);
+
+								// To keep the previous behavior for single value attributes, we force the copied attribute to have its default value set to the first entry (only done if there is just a single valid entry key).
+								if (AllSourceEntryKeys.Num() == 1 && AllSourceEntryKeys[0] != PCGInvalidEntryKey)
+								{
+									NewAttr->SetDefaultValueToFirstEntry();
+								}
 							}
 						}
 					}
diff --git a/Engine/Plugins/PCG/Source/PCG/Private/Tests/Elements/PCGCopyAttributeTest.cpp b/Engine/Plugins/PCG/Source/PCG/Private/Tests/Elements/PCGCopyAttributeTest.cpp
--- a/Engine/Plugins/PCG/Source/PCG/Private/Tests/Elements/PCGCopyAttributeTest.cpp
+++ b/Engine/Plugins/PCG/Source/PCG/Private/Tests/Elements/PCGCopyAttributeTest.cpp
@@ -342,7 +342,7 @@ bool FPCGCopyAttributeTests_Params_SingleValue::RunTest(const FString& Parameter
 
 	check(InputAttribute && OutputAttribute);
 
-	UTEST_EQUAL("Output Attribute default value is the same as first entry of Input Attribute", OutputAttribute->GetValueFromItemKey(PCGInvalidEntryKey), InputAttribute->GetValueFromItemKey(PCGMetadataEntryKey(0)));
+	UTEST_EQUAL("Output Attribute default value is the same as default of Input Attribute", OutputAttribute->GetValueFromItemKey(PCGInvalidEntryKey), InputAttribute->GetValueFromItemKey(PCGInvalidEntryKey));
 	UTEST_EQUAL("Output Attribute has the same value as Input Attribute for entry 0", OutputAttribute->GetValueFromItemKey(PCGMetadataEntryKey(0)), InputAttribute->GetValueFromItemKey(PCGMetadataEntryKey(0)));
 
 	return true;
