[PCG] Allow to set the IsPartitioned flag on BP classes.

#rnx

diff --git a/Engine/Plugins/PCG/Source/PCG/Private/PCGComponent.cpp b/Engine/Plugins/PCG/Source/PCG/Private/PCGComponent.cpp
--- a/Engine/Plugins/PCG/Source/PCG/Private/PCGComponent.cpp
+++ b/Engine/Plugins/PCG/Source/PCG/Private/PCGComponent.cpp
@@ -91,6 +91,13 @@ void UPCGComponent::SetIsPartitioned(bool bIsNowPartitioned)
 	{
 		return;
 	}
+	
+	// Allow BP class to set the flag. BP class won't have an owner. They won't get registered anyway.
+	if (!GetOwner())
+	{
+		bIsComponentPartitioned = bIsNowPartitioned;
+		return;
+	}
 
 	bool bDoActorMapping = bGenerated || PCGHelpers::IsRuntimeOrPIE();
 
