[PCG] Make sure the views in the PCG editor refresh if the currently inspected stack root component has done generating

#rnx

diff --git a/Engine/Plugins/PCG/Source/PCGEditor/Private/PCGEditor.cpp b/Engine/Plugins/PCG/Source/PCGEditor/Private/PCGEditor.cpp
--- a/Engine/Plugins/PCG/Source/PCGEditor/Private/PCGEditor.cpp
+++ b/Engine/Plugins/PCG/Source/PCGEditor/Private/PCGEditor.cpp
@@ -3131,7 +3131,12 @@ void FPCGEditor::OnComponentUnregistered(UPCGComponent* Component)
 
 void FPCGEditor::OnComponentGenerationDone(UPCGSubsystem* Subsystem, UPCGComponent* Component, EPCGGenerationStatus Status)
 {
-	if(Component && Component->GetGraph() != PCGGraphBeingEdited)
+	// We want to refresh if the component that is done generating has generated the current graph being edited,
+	// or if it is the root of the current stack being inspected (for subgraphs to also be refreshed).
+	// If we don't have a component, we refresh nonetheless.
+	const bool bShouldRefresh = !Component || StackBeingInspected.GetRootComponent() == Component || Component->GetGraph() == PCGGraphBeingEdited;
+
+	if (!bShouldRefresh)
 	{
 		return;
 	}
