[PCG] Fix serialized WeakObjPtrs in PCGComputeGraph.h

#rnx

diff --git a/Engine/Plugins/PCG/Source/PCG/Private/Compute/Elements/PCGComputeGraphElement.cpp b/Engine/Plugins/PCG/Source/PCG/Private/Compute/Elements/PCGComputeGraphElement.cpp
--- a/Engine/Plugins/PCG/Source/PCG/Private/Compute/Elements/PCGComputeGraphElement.cpp
+++ b/Engine/Plugins/PCG/Source/PCG/Private/Compute/Elements/PCGComputeGraphElement.cpp
@@ -139,7 +139,7 @@ bool FPCGComputeGraphElement::ExecuteInternal(FPCGContext* InContext) const
 		Context->DataBinding->Initialize(Context->ComputeGraph.Get(), Context->SourceComponent, Context->InputData, Context->ComputeGraph->GetAttributeLookupTable());
 
 		// Perform validation after input data is initialized.
-		for (TWeakObjectPtr<const UPCGNode> Node : Context->ComputeGraph->KernelToNode)
+		for (TSoftObjectPtr<const UPCGNode> Node : Context->ComputeGraph->KernelToNode)
 		{
 			const UPCGSettings* Settings = Node.Get() ? Node->GetSettings() : nullptr;
 
@@ -464,7 +464,7 @@ void FPCGComputeGraphElement::PostExecuteInternal(FPCGContext* InContext) const
 #if WITH_EDITOR
 	if (Context->bExecutionSuccess)
 	{
-		for (TWeakObjectPtr<const UPCGNode> NodePtr : Context->DataBinding->Graph->KernelToNode)
+		for (TSoftObjectPtr<const UPCGNode> NodePtr : Context->DataBinding->Graph->KernelToNode)
 		{
 			const UPCGNode* Node = NodePtr.Get();
 			UPCGComponent* Component = Context->SourceComponent.Get();
diff --git a/Engine/Plugins/PCG/Source/PCG/Private/Compute/PCGComputeGraph.h b/Engine/Plugins/PCG/Source/PCG/Private/Compute/PCGComputeGraph.h
--- a/Engine/Plugins/PCG/Source/PCG/Private/Compute/PCGComputeGraph.h
+++ b/Engine/Plugins/PCG/Source/PCG/Private/Compute/PCGComputeGraph.h
@@ -34,7 +34,7 @@ public:
 	
 	/** Set of input pins at the CPU -> GPU border. */
 	UPROPERTY()
-	TArray<TWeakObjectPtr<const UPCGPin>> PinsReceivingDataFromCPU;
+	TArray<TSoftObjectPtr<const UPCGPin>> PinsReceivingDataFromCPU;
 
 	/** Pin label aliases, used for selecting data items corresponding to an input pin from the input data collection. */
 	UPROPERTY()
@@ -46,7 +46,7 @@ public:
 
 	/** Node corresponding to each kernel, useful for compilation feedback. */
 	UPROPERTY()
-	TArray<TWeakObjectPtr<const UPCGNode>> KernelToNode;
+	TArray<TSoftObjectPtr<const UPCGNode>> KernelToNode;
 
 	/** List of settings for all nodes that spawn static meshes, so we can do required primitive & DI setup when compute graph element executes. */
 	UPROPERTY()
diff --git a/Engine/Plugins/PCG/Source/PCG/Private/Compute/PCGDataBinding.cpp b/Engine/Plugins/PCG/Source/PCG/Private/Compute/PCGDataBinding.cpp
--- a/Engine/Plugins/PCG/Source/PCG/Private/Compute/PCGDataBinding.cpp
+++ b/Engine/Plugins/PCG/Source/PCG/Private/Compute/PCGDataBinding.cpp
@@ -42,7 +42,7 @@ void UPCGDataBinding::InitializeInputData(const FPCGDataCollection& InComputeGra
 	DataForGPU.InputDataCollection = InComputeGraphElementInputData;
 
 	// Link each input pin to the data collection, so that data providers can find the data.
-	for (const TWeakObjectPtr<const UPCGPin>& InputPinPtr : Graph->PinsReceivingDataFromCPU)
+	for (const TSoftObjectPtr<const UPCGPin>& InputPinPtr : Graph->PinsReceivingDataFromCPU)
 	{
 		if (const UPCGPin* InputPin = InputPinPtr.Get())
 		{
@@ -220,7 +220,7 @@ void UPCGDataBinding::DebugLogDataDescriptions()
 	};
 
 	UE_LOG(LogPCG, Display, TEXT("\n### INPUT PIN DATA DESCRIPTIONS ###"));
-	for (const TWeakObjectPtr<const UPCGNode>& Node : Graph->KernelToNode)
+	for (const TSoftObjectPtr<const UPCGNode>& Node : Graph->KernelToNode)
 	{
 		const UPCGSettings* Settings = Node.IsValid() ? Node->GetSettings() : nullptr;
 		if (Settings && Settings->bDumpDataDescriptions)
