[PCG] Support None on Get Actor Property to extract the full struct.

* Add test for this usecase

#rnx

diff --git a/Engine/Plugins/PCG/Source/PCG/Private/Elements/PCGGetActorProperty.cpp b/Engine/Plugins/PCG/Source/PCG/Private/Elements/PCGGetActorProperty.cpp
--- a/Engine/Plugins/PCG/Source/PCG/Private/Elements/PCGGetActorProperty.cpp
+++ b/Engine/Plugins/PCG/Source/PCG/Private/Elements/PCGGetActorProperty.cpp
@@ -171,7 +171,7 @@ bool FPCGGetActorPropertyElement::ExecuteInternal(FPCGContext* InContext) const
 	check(Settings);
 
 	// Early out if arguments are not specified
-	if (Settings->PropertyName == NAME_None || (Settings->bSelectComponent && !Settings->ComponentClass))
+	if (Settings->bSelectComponent && !Settings->ComponentClass)
 	{
 		PCGE_LOG(Warning, GraphAndLog, LOCTEXT("ParametersMissing", "Some parameters are missing, aborting"));
 		return true;
diff --git a/Engine/Plugins/PCG/Source/PCG/Private/Tests/Elements/PCGGetActorPropertyTest.cpp b/Engine/Plugins/PCG/Source/PCG/Private/Tests/Elements/PCGGetActorPropertyTest.cpp
--- a/Engine/Plugins/PCG/Source/PCG/Private/Tests/Elements/PCGGetActorPropertyTest.cpp
+++ b/Engine/Plugins/PCG/Source/PCG/Private/Tests/Elements/PCGGetActorPropertyTest.cpp
@@ -13,6 +13,7 @@
 
 IMPLEMENT_CUSTOM_SIMPLE_AUTOMATION_TEST(FPCGPropertyToParamDataPropertyTypeTest, FPCGTestBaseClass, "Plugins.PCG.PropertyToParamData.PropertyType", PCGTestsCommon::TestFlags)
 IMPLEMENT_CUSTOM_SIMPLE_AUTOMATION_TEST(FPCGPropertyToParamDataActorFindTest, FPCGTestBaseClass, "Plugins.PCG.PropertyToParamData.ActorFind", PCGTestsCommon::TestFlags)
+IMPLEMENT_CUSTOM_SIMPLE_AUTOMATION_TEST(FPCGPropertyToParamDataFullExtractionTest, FPCGTestBaseClass, "Plugins.PCG.PropertyToParamData.FullExtraction", PCGTestsCommon::TestFlags)
 
 
 /**
@@ -314,10 +315,6 @@ bool FPCGPropertyToParamDataPropertyTypeTest::RunTest(const FString& Parameters)
 	AddExpectedError(TEXT("Fail to extract the property 'DummyMissingProperty' on actor"), EAutomationExpectedErrorFlags::Contains, 1);
 	bSuccess &= VerifyAttributeValueInvalid(this, TestData, TEXT("DummyMissingProperty"), 42, ExtraTestWhat);
 
-	// Missing property
-	AddExpectedError(TEXT("Some parameters are missing, abort."), EAutomationExpectedErrorFlags::Contains, 1);
-	bSuccess &= VerifyAttributeValueInvalid(this, TestData, NAME_None, 42, ExtraTestWhat);
-
 	ObjectValue->MarkAsGarbage();
 	SecondObjectValue->MarkAsGarbage();
 
@@ -445,4 +442,114 @@ bool FPCGPropertyToParamDataActorFindTest::RunTest(const FString& Parameters)
 	return bSuccess;
 }
 
+bool FPCGPropertyToParamDataFullExtractionTest::RunTest(const FString& Parameters)
+{
+	UPCGGetActorPropertySettings* Settings = NewObject<UPCGGetActorPropertySettings>();
+	Settings->ActorSelector.ActorSelection = EPCGActorSelection::ByClass;
+	Settings->ActorSelector.ActorSelectionClass = APCGUnitTestDummyActor::StaticClass();
+	Settings->ActorSelector.ActorFilter = EPCGActorFilter::Self;
+
+	static constexpr int32 Seed = 42;
+
+	PCGTestsCommon::FTestData TestData(Seed, Settings, APCGUnitTestDummyActor::StaticClass());
+
+	// Set all properties
+	const FName NameValue = TEXT("HelloWorld");
+	const FString StringValue = TEXT("HelloWorld");
+
+	const FVector VectorValue{ 1.0, 2.0, 3.0 };
+	const FVector SecondVectorValue{ 4.0, 5.0, 6.0 };
+	const FVector4 Vector4Value{ 1.0, 2.0, 3.0, 4.0 };
+	const FRotator RotatorValue{ 45.0, 45.0, 45.0 };
+	const FQuat QuatValue = RotatorValue.Quaternion();
+	const FTransform TransformValue{ QuatValue, VectorValue, VectorValue };
+
+	UPCGDummyGetPropertyTest* ObjectValue = NewObject<UPCGDummyGetPropertyTest>();
+	ObjectValue->SetFlags(RF_Transient);
+	ObjectValue->Int64Property = 42ll;
+	ObjectValue->DoubleProperty = 1.0;
+
+	UPCGDummyGetPropertyTest* SecondObjectValue = NewObject<UPCGDummyGetPropertyTest>();
+	SecondObjectValue->SetFlags(RF_Transient);
+	SecondObjectValue->Int64Property = 43ll;
+	SecondObjectValue->DoubleProperty = 2.0;
+
+	const FSoftObjectPath SoftObjectPathValue{ ObjectValue };
+	const FSoftClassPath SoftClassPathValue{ UPCGDummyGetPropertyTest::StaticClass() };
+
+	const FVector2D Vector2Value = { 1.0, 2.0 };
+	const FPCGTestMyColorStruct PCGColorValue{ 1.0, 1.0, 0.0, 1.0 };
+	const FPCGTestMyColorStruct SecondPCGColorValue{ 1.0, 0.0, 1.0, 1.0 };
+	const FColor ColorValue = FColor::White;
+	const FLinearColor LinearColorValue = FLinearColor::Blue;
+
+	APCGUnitTestDummyActor* Actor = Cast<APCGUnitTestDummyActor>(TestData.TestActor);
+	Actor->IntProperty = 42;
+	Actor->Int64Property = 42ll;
+	Actor->FloatProperty = 1.0f;
+	Actor->DoubleProperty = 1.0;
+	Actor->BoolProperty = true;
+	Actor->NameProperty = NameValue;
+	Actor->StringProperty = StringValue;
+	Actor->EnumProperty = EPCGUnitTestDummyEnum::Three;
+	Actor->VectorProperty = VectorValue;
+	Actor->Vector4Property = Vector4Value;
+	Actor->RotatorProperty = RotatorValue;
+	Actor->QuatProperty = QuatValue;
+	Actor->TransformProperty = TransformValue;
+	Actor->SoftObjectPathProperty = SoftObjectPathValue;
+	Actor->SoftClassPathProperty = SoftClassPathValue;
+	Actor->ClassProperty = UPCGDummyGetPropertyTest::StaticClass();
+	Actor->ObjectProperty = ObjectValue;
+	Actor->Vector2Property = Vector2Value;
+	Actor->ColorProperty = ColorValue;
+	Actor->LinearColorProperty = LinearColorValue;
+
+	const TArray<FName> ListOfAllExtractedValues =
+	{
+		GET_MEMBER_NAME_CHECKED(APCGUnitTestDummyActor, IntProperty),
+		GET_MEMBER_NAME_CHECKED(APCGUnitTestDummyActor, Int64Property),
+		GET_MEMBER_NAME_CHECKED(APCGUnitTestDummyActor, FloatProperty),
+		GET_MEMBER_NAME_CHECKED(APCGUnitTestDummyActor, DoubleProperty),
+		GET_MEMBER_NAME_CHECKED(APCGUnitTestDummyActor, BoolProperty),
+		GET_MEMBER_NAME_CHECKED(APCGUnitTestDummyActor, NameProperty),
+		GET_MEMBER_NAME_CHECKED(APCGUnitTestDummyActor, StringProperty),
+		GET_MEMBER_NAME_CHECKED(APCGUnitTestDummyActor, EnumProperty),
+		GET_MEMBER_NAME_CHECKED(APCGUnitTestDummyActor, VectorProperty),
+		GET_MEMBER_NAME_CHECKED(APCGUnitTestDummyActor, Vector4Property),
+		GET_MEMBER_NAME_CHECKED(APCGUnitTestDummyActor, RotatorProperty),
+		GET_MEMBER_NAME_CHECKED(APCGUnitTestDummyActor, QuatProperty),
+		GET_MEMBER_NAME_CHECKED(APCGUnitTestDummyActor, TransformProperty),
+		GET_MEMBER_NAME_CHECKED(APCGUnitTestDummyActor, SoftObjectPathProperty),
+		GET_MEMBER_NAME_CHECKED(APCGUnitTestDummyActor, SoftClassPathProperty),
+		GET_MEMBER_NAME_CHECKED(APCGUnitTestDummyActor, ClassProperty),
+		GET_MEMBER_NAME_CHECKED(APCGUnitTestDummyActor, ObjectProperty),
+		GET_MEMBER_NAME_CHECKED(APCGUnitTestDummyActor, Vector2Property),
+		GET_MEMBER_NAME_CHECKED(APCGUnitTestDummyActor, ColorProperty),
+		GET_MEMBER_NAME_CHECKED(APCGUnitTestDummyActor, LinearColorProperty),
+		// Arrays are not extracted, neither are deeper structs.
+	};
+
+	return VerifyAttributeValuesValid(this, TestData, NAME_None, ListOfAllExtractedValues, TEXT("FullExtraction"),
+		42,
+		42ll,
+		1.0,
+		1.0,
+		true,
+		NameValue,
+		StringValue,
+		(int64)EPCGUnitTestDummyEnum::Three,
+		VectorValue,
+		Vector4Value,
+		RotatorValue,
+		QuatValue,
+		TransformValue,
+		SoftObjectPathValue,
+		SoftClassPathValue,
+		FSoftClassPath(UPCGDummyGetPropertyTest::StaticClass()),
+		FSoftObjectPath(ObjectValue),
+		Vector2Value,
+		FVector4(ColorValue),
+		FVector4(LinearColorValue));
+}
 #endif // WITH_EDITOR
