[PCG] Fix up references to GetDataNumElements (should be GetNumElements)

#rnx

diff --git a/Engine/Plugins/PCG/Shaders/Private/Elements/PCGCopyPoints.usf b/Engine/Plugins/PCG/Shaders/Private/Elements/PCGCopyPoints.usf
--- a/Engine/Plugins/PCG/Shaders/Private/Elements/PCGCopyPoints.usf
+++ b/Engine/Plugins/PCG/Shaders/Private/Elements/PCGCopyPoints.usf
@@ -163,8 +163,8 @@ void Main(uint3 GroupId : SV_GroupID, uint GroupIndex : SV_GroupIndex)
 			SourceDataIndex = min(DataIndex, SourceNumData - 1);
 			TargetDataIndex = min(DataIndex, TargetNumData - 1);
 
-			const uint SourceElemCount = Source_GetDataNumElements(SourceDataIndex);
-			const uint TargetElemCount = Target_GetDataNumElements(TargetDataIndex);
+			const uint SourceElemCount = Source_GetNumElements(SourceDataIndex);
+			const uint TargetElemCount = Target_GetNumElements(TargetDataIndex);
 
 			const uint ElemCount = SourceElemCount * TargetElemCount;
 
diff --git a/Engine/Plugins/PCG/Shaders/Private/PCGDataCollectionDataInterface.ush b/Engine/Plugins/PCG/Shaders/Private/PCGDataCollectionDataInterface.ush
--- a/Engine/Plugins/PCG/Shaders/Private/PCGDataCollectionDataInterface.ush
+++ b/Engine/Plugins/PCG/Shaders/Private/PCGDataCollectionDataInterface.ush
@@ -64,7 +64,7 @@ bool GetThreadData_{DataInterfaceName}(uint InThreadIndex, out uint OutDataIndex
 	const uint NumData = GetNumData_{DataInterfaceName}();
 	for (uint DataIndex = 0; DataIndex < NumData; ++DataIndex)
 	{
-		const uint ElemCount = GetDataNumElements_{DataInterfaceName}(DataIndex);
+		const uint ElemCount = GetNumElements_{DataInterfaceName}(DataIndex);
 
 		if (ElementIndex < ElemCount)
 		{
@@ -88,7 +88,7 @@ uint GetNumElements_{DataInterfaceName}()
 	const uint NumData = GetNumData_{DataInterfaceName}();
 	for (uint DataIndex = 0; DataIndex < NumData; ++DataIndex)
 	{
-		NumElements += GetDataNumElements_{DataInterfaceName}(DataIndex);
+		NumElements += GetNumElements_{DataInterfaceName}(DataIndex);
 	}
 
 	return NumElements;
diff --git a/Engine/Plugins/PCG/Source/PCG/Private/Compute/DataInterfaces/PCGDataCollectionDataInterface.cpp b/Engine/Plugins/PCG/Source/PCG/Private/Compute/DataInterfaces/PCGDataCollectionDataInterface.cpp
--- a/Engine/Plugins/PCG/Source/PCG/Private/Compute/DataInterfaces/PCGDataCollectionDataInterface.cpp
+++ b/Engine/Plugins/PCG/Source/PCG/Private/Compute/DataInterfaces/PCGDataCollectionDataInterface.cpp
@@ -46,7 +46,7 @@ void UPCGDataCollectionDataInterface::GetSupportedInputs(TArray<FShaderFunctionD
 			.AddParam(EShaderFundamentalType::Uint); // InDataIndex
 
 		OutFunctions.AddDefaulted_GetRef()
-			.SetName(TEXT("GetDataNumElements"))
+			.SetName(TEXT("GetNumElements"))
 			.AddReturnType(EShaderFundamentalType::Uint)
 			.AddParam(EShaderFundamentalType::Uint); // InDataIndex
 
