[PCG] Fix Graph Determinism Test not setting a proper Generation task id on the PCG Component

#rnx

diff --git a/Engine/Plugins/PCG/Source/PCG/Private/Graph/PCGGraphExecutor.cpp b/Engine/Plugins/PCG/Source/PCG/Private/Graph/PCGGraphExecutor.cpp
--- a/Engine/Plugins/PCG/Source/PCG/Private/Graph/PCGGraphExecutor.cpp
+++ b/Engine/Plugins/PCG/Source/PCG/Private/Graph/PCGGraphExecutor.cpp
@@ -2705,6 +2705,7 @@ FPCGTaskId FPCGGraphExecutor::ScheduleDebugWithTaskCallback(UPCGComponent* InCom
 {
 	check(InComponent);
 	FPCGTaskId FinalTaskID = Schedule(InComponent, {});
+	InComponent->CurrentGenerationTask = FinalTaskID;
 
 	const bool bNonPartitionedComponent = !InComponent->IsLocalComponent() && !InComponent->IsPartitioned();
 	const uint32 GenerationGridSize = bNonPartitionedComponent ? PCGHiGenGrid::UninitializedGridSize() : InComponent->GetGenerationGridSize();
@@ -2862,7 +2863,7 @@ bool FPCGPreGraphElement::ExecuteInternal(FPCGContext* Context) const
 	Component->StartGenerationInProgress();
 #endif
 
-	check(Component->GetGenerationTaskId() != InvalidPCGTaskId);
+	ensureMsgf(Component->GetGenerationTaskId() != InvalidPCGTaskId, TEXT("Component was Scheduled for generation without having its CurrentGenerationTask assigned"));
 
 	{
 		PCG_EXECUTION_CACHE_VALIDATION_CREATE_SCOPE(Component);
diff --git a/Engine/Plugins/PCG/Source/PCG/Public/PCGComponent.h b/Engine/Plugins/PCG/Source/PCG/Public/PCGComponent.h
--- a/Engine/Plugins/PCG/Source/PCG/Public/PCGComponent.h
+++ b/Engine/Plugins/PCG/Source/PCG/Public/PCGComponent.h
@@ -91,6 +91,7 @@ class PCG_API UPCGComponent : public UActorComponent
 
 	friend class UPCGManagedActors;
 	friend class UPCGSubsystem;
+	friend class FPCGGraphExecutor;
 	friend class FPCGActorAndComponentMapping;
 	friend struct FPCGWorldPartitionBuilder;
 
