[PCG] Use component seed in GPU Static Mesh Spawner

#rnx

diff --git a/Engine/Plugins/PCG/Shaders/Private/Elements/PCGStaticMeshSpawner.usf b/Engine/Plugins/PCG/Shaders/Private/Elements/PCGStaticMeshSpawner.usf
--- a/Engine/Plugins/PCG/Shaders/Private/Elements/PCGStaticMeshSpawner.usf
+++ b/Engine/Plugins/PCG/Shaders/Private/Elements/PCGStaticMeshSpawner.usf
@@ -13,7 +13,7 @@
 #define MAX_NUM_ATTRS 128
 #define NUM_RESERVED_ATTRS 32
 
-void CopyAttributesToOuput(uint In_DataIndex, uint In_ElementIndex, uint Out_DataIndex, uint Out_ElementIndex);
+void CopyAttributesToOutput(uint In_DataIndex, uint In_ElementIndex, uint Out_DataIndex, uint Out_ElementIndex);
 
 [numthreads(64, 1, 1)]
 void Main(uint3 GroupId : SV_GroupID, uint GroupIndex : SV_GroupIndex)
@@ -55,6 +55,8 @@ void Main(uint3 GroupId : SV_GroupID, uint GroupIndex : SV_GroupIndex)
 		{
 			// Randomly select a primitive by inverting the cumulative distribution function.
 			uint Seed = In_GetSeed(In_DataIndex, In_ElementIndex);
+			Seed = ComputeSeed(Seed, GetSeed());
+
 			const float RandomDraw = FRand(Seed);
 
 			for (; PrimitiveIndex < uint(NumPrimitives) - 1; ++PrimitiveIndex)
@@ -127,7 +129,7 @@ void Main(uint3 GroupId : SV_GroupID, uint GroupIndex : SV_GroupIndex)
 	Out_SetDensity(Out_DataIndex, Out_ElementIndex, In_GetDensity(In_DataIndex, In_ElementIndex));
 	Out_SetSteepness(Out_DataIndex, Out_ElementIndex, In_GetSteepness(In_DataIndex, In_ElementIndex));
 
-	CopyAttributesToOuput(In_DataIndex, In_ElementIndex, Out_DataIndex, Out_ElementIndex);
+	CopyAttributesToOutput(In_DataIndex, In_ElementIndex, Out_DataIndex, Out_ElementIndex);
 
 	if (SMSpawner_GetSelectedMeshAttributeId() != (uint)-1)
 	{
@@ -137,45 +139,45 @@ void Main(uint3 GroupId : SV_GroupID, uint GroupIndex : SV_GroupIndex)
 }
 
 // TODO pass in attribute IDs to copy.
-void CopyAttributesToOuput(uint In_DataIndex, uint In_ElementIndex, uint Out_DataIndex, uint Out_ElementIndex)
+void CopyAttributesToOutput(uint In_DataIndex, uint In_ElementIndex, uint Out_DataIndex, uint Out_ElementIndex)
 {
-    const uint NumAttributes = Out_GetDataNumAttributesInternal(Out_DataIndex);
-    uint NumAttributesProcessed = 0;
-
-    for (int AttributeId = NUM_RESERVED_ATTRS; AttributeId < MAX_NUM_ATTRS; ++AttributeId)
-    {
-        const uint Out_Stride = Out_GetAttributeStrideInternal(Out_DataIndex, AttributeId);
-
-        if (Out_Stride == 0)
-        {
-            // No output attribute to write to.
-            continue;
-        }
-
-        const uint In_Stride = In_GetAttributeStrideInternal(In_DataIndex, AttributeId);
-        const uint Out_ElementAddress = Out_GetElementAddressInternal(Out_DataIndex, Out_ElementIndex, AttributeId);
-
-        if (In_Stride == Out_Stride)
-        {
-            const uint In_ElementAddress = In_GetElementAddressInternal(In_DataIndex, In_ElementIndex, AttributeId);
-
-            for (uint Offset = 0; Offset < Out_Stride; Offset += 4)
-            {
-                Out_StoreBufferInternal(Out_ElementAddress + Offset, In_LoadBufferInternal(In_ElementAddress + Offset));
-            }
-        }
-        else
-        {
-            for (uint Offset = 0; Offset < Out_Stride; Offset += 4)
-            {
+	const uint NumAttributes = Out_GetDataNumAttributesInternal(Out_DataIndex);
+	uint NumAttributesProcessed = 0;
+	
+	for (int AttributeId = NUM_RESERVED_ATTRS; AttributeId < MAX_NUM_ATTRS; ++AttributeId)
+	{
+		const uint Out_Stride = Out_GetAttributeStrideInternal(Out_DataIndex, AttributeId);
+		
+		if (Out_Stride == 0)
+		{
+			// No output attribute to write to.
+			continue;
+		}
+		
+		const uint In_Stride = In_GetAttributeStrideInternal(In_DataIndex, AttributeId);
+		const uint Out_ElementAddress = Out_GetElementAddressInternal(Out_DataIndex, Out_ElementIndex, AttributeId);
+		
+		if (In_Stride == Out_Stride)
+		{
+			const uint In_ElementAddress = In_GetElementAddressInternal(In_DataIndex, In_ElementIndex, AttributeId);
+			
+			for (uint Offset = 0; Offset < Out_Stride; Offset += 4)
+			{
+				Out_StoreBufferInternal(Out_ElementAddress + Offset, In_LoadBufferInternal(In_ElementAddress + Offset));
+			}
+		}
+		else
+		{
+			for (uint Offset = 0; Offset < Out_Stride; Offset += 4)
+			{
 				// Initialize output data (only header part of buffer will be initialized).
-                Out_StoreBufferInternal(Out_ElementAddress + Offset, 0u);
-            }
-        }
-
-        if (++NumAttributesProcessed >= NumAttributes)
-        {
-            break; // We can early-out when we've looked at all the possible attributes
-        }
-    }
+				Out_StoreBufferInternal(Out_ElementAddress + Offset, 0u);
+			}
+		}
+		
+		if (++NumAttributesProcessed >= NumAttributes)
+		{
+			break; // We can early-out when we've looked at all the possible attributes
+		}
+	}
 }
