[PCG] Rename PointCount to NumPoints in Custom HLSL node
* This requires manual fix-up of any references to the PointCount declaration in kernels

diff --git a/Engine/Plugins/PCG/Source/PCG/Private/Compute/Elements/PCGCustomHLSL.cpp b/Engine/Plugins/PCG/Source/PCG/Private/Compute/Elements/PCGCustomHLSL.cpp
--- a/Engine/Plugins/PCG/Source/PCG/Private/Compute/Elements/PCGCustomHLSL.cpp
+++ b/Engine/Plugins/PCG/Source/PCG/Private/Compute/Elements/PCGCustomHLSL.cpp
@@ -639,7 +639,7 @@ void UPCGCustomHLSLSettings::UpdateInputDeclarations()
 		if (KernelType == EPCGKernelType::PointGenerator)
 		{
 			InputDeclarations += TEXT("/*** INPUT CONSTANTS ***/\n\n");
-			InputDeclarations += FString::Format(TEXT("const uint PointCount = {0};\n\n"), { PointCount });
+			InputDeclarations += FString::Format(TEXT("const uint NumPoints = {0};\n\n"), { PointCount });
 		}
 
 		InputDeclarations += TEXT("/*** INPUT PER-THREAD CONSTANTS ***/\n\n");
@@ -1721,7 +1721,7 @@ FString UPCGCustomHLSLSettings::GetCookedKernelSource(const TMap<FName, FPCGKern
 	}
 	else if (KernelType == EPCGKernelType::PointGenerator)
 	{
-		KernelSpecificPreamble += FString::Format(TEXT("    const uint PointCount = {0};\n"), { PointCount });
+		KernelSpecificPreamble += FString::Format(TEXT("    const uint NumPoints = {0};\n"), { PointCount });
 
 		if (const UPCGPin* OutputPin = GetFirstPointOutputPin())
 		{
diff --git a/Engine/Plugins/PCG/Source/PCG/Public/Compute/Elements/PCGCustomHLSL.h b/Engine/Plugins/PCG/Source/PCG/Public/Compute/Elements/PCGCustomHLSL.h
--- a/Engine/Plugins/PCG/Source/PCG/Public/Compute/Elements/PCGCustomHLSL.h
+++ b/Engine/Plugins/PCG/Source/PCG/Public/Compute/Elements/PCGCustomHLSL.h
@@ -129,7 +129,7 @@ protected:
 	UPROPERTY(BlueprintReadWrite, EditAnywhere, Category = "Settings")
 	EPCGKernelType KernelType = EPCGKernelType::PointProcessor;
 
-	UPROPERTY(BlueprintReadWrite, EditAnywhere, Category = "Settings", meta = (EditCondition = "KernelType == EPCGKernelType::PointGenerator", EditConditionHides))
+	UPROPERTY(BlueprintReadWrite, EditAnywhere, DisplayName = "Number of Points", Category = "Settings", meta = (EditCondition = "KernelType == EPCGKernelType::PointGenerator", EditConditionHides))
 	int PointCount = 256;
 
 	UPROPERTY(BlueprintReadWrite, EditAnywhere, Category = "Settings|Thread Count", meta = (EditCondition = "KernelType == EPCGKernelType::Custom", EditConditionHides))
