[PCG] Fix up warnings on editor load related to missing images.

Remove styles that were ultimately unused.

#rnx

diff --git a/Engine/Plugins/PCG/Source/PCGEditor/Private/PCGEditorStyle.cpp b/Engine/Plugins/PCG/Source/PCGEditor/Private/PCGEditorStyle.cpp
--- a/Engine/Plugins/PCG/Source/PCGEditor/Private/PCGEditorStyle.cpp
+++ b/Engine/Plugins/PCG/Source/PCGEditor/Private/PCGEditorStyle.cpp
@@ -132,30 +132,8 @@ FPCGEditorStyle::FPCGEditorStyle() : FSlateStyleSet("PCGEditorStyle")
 
 	// Text editor styles
 	{
-		const FSlateFontInfo Consolas10 = FCoreStyle::GetDefaultFontStyle("Mono", 9);
-
-		const FTextBlockStyle NormalText = FTextBlockStyle()
-			.SetFont(Consolas10)
-			.SetColorAndOpacity(FLinearColor::White)
-			.SetShadowOffset(FVector2D::ZeroVector)
-			.SetShadowColorAndOpacity(FLinearColor::Black)
-			.SetSelectedBackgroundColor(FLinearColor::Blue)
-			.SetHighlightColor(FLinearColor::Yellow)
-			.SetHighlightShape(BOX_BRUSH("Images/TextBlockHighlightShape", FMargin(3.0f / 8.0f)));
-
-		Set("TextEditor.NormalText", NormalText);
-
 		Set("TextEditor.Border", new BOX_BRUSH("Images/TextEditorBorder", FMargin(4.0f / 16.0f), FLinearColor(0.02f, 0.02f, 0.02f, 1)));
 
-		const FEditableTextBoxStyle EditableTextBoxStyle = FEditableTextBoxStyle()
-			.SetTextStyle(NormalText)
-			.SetBackgroundImageNormal(FSlateNoResource())
-			.SetBackgroundImageHovered(FSlateNoResource())
-			.SetBackgroundImageFocused(FSlateNoResource())
-			.SetBackgroundImageReadOnly(FSlateNoResource());
-
-		Set("TextEditor.EditableTextBox", EditableTextBoxStyle);
-
 		FSearchBoxStyle SearchBoxStyle = FCoreStyle::Get().GetWidgetStyle<FSearchBoxStyle>("SearchBox");
 		SearchBoxStyle.SetLeftAlignGlassImageAndClearButton(true);
 		SearchBoxStyle.SetLeftAlignSearchResultButtons(false);
