[PCG] Fix failed compute graph execution when there are upstream reroutes.

We traverse the nodes/settings/edges to compute data descriptions and when encountering a CPU node we attempt to look up the data registered in the compute graph. Reroutes are culled before compute graph compilation we need to transparently jump over them when computing data descriptions.

This issue in theory applies to any node that is statically culled during compilation, however i think reroutes are the only problem case. Branch culling should be a non-issue by definition and higen works out because upstream grid size nodes are present. If we statically cull more things in the future, we might want to make a deeper fix - like have a full mapping of input pins->upstream output pins so we have the graph connectivity after culling, but this seems like a good low-impact step and might be sufficient.

#rnx

diff --git a/Engine/Plugins/PCG/Source/PCG/Private/Elements/PCGReroute.cpp b/Engine/Plugins/PCG/Source/PCG/Private/Elements/PCGReroute.cpp
--- a/Engine/Plugins/PCG/Source/PCG/Private/Elements/PCGReroute.cpp
+++ b/Engine/Plugins/PCG/Source/PCG/Private/Elements/PCGReroute.cpp
@@ -45,6 +45,21 @@ FPCGElementPtr UPCGRerouteSettings::CreateElement() const
 	return MakeShared<FPCGRerouteElement>();
 }
 
+FPCGDataCollectionDesc UPCGRerouteSettings::ComputeOutputPinDataDesc(const UPCGPin* OutputPin, const UPCGDataBinding* Binding) const
+{
+	// Reroutes always trivially forward data. Reroutes are culled and not visible to compute graph compilation and therefore we
+	// transparently forward the data here.
+	if (const UPCGNode* Node = Cast<UPCGNode>(GetOuter()))
+	{
+		if (const UPCGPin* Pin = Node->GetInputPin(PCGPinConstants::DefaultInputLabel))
+		{
+			return ComputeInputPinDataDesc(Pin, Binding);
+		}
+	}
+
+	return {};
+}
+
 TArray<FPCGPinProperties> UPCGNamedRerouteDeclarationSettings::OutputPinProperties() const
 {
 	TArray<FPCGPinProperties> PinProperties;
diff --git a/Engine/Plugins/PCG/Source/PCG/Private/PCGSettings.cpp b/Engine/Plugins/PCG/Source/PCG/Private/PCGSettings.cpp
--- a/Engine/Plugins/PCG/Source/PCG/Private/PCGSettings.cpp
+++ b/Engine/Plugins/PCG/Source/PCG/Private/PCGSettings.cpp
@@ -271,6 +271,9 @@ FPCGDataCollectionDesc UPCGSettings::ComputeOutputPinDataDesc(const FName& Outpu
 
 FPCGDataCollectionDesc UPCGSettings::ComputeOutputPinDataDesc(const UPCGPin* OutputPin, const UPCGDataBinding* Binding) const
 {
+	check(OutputPin);
+	check(Binding);
+
 	// This base class implementation will be called on upstream CPU nodes. Get the downstream pin alias in order to
 	// pick the relevant data items out of the compute graph element's input data collection.
 	if (const FName* FoundPinAlias = Binding->Graph->OutputCPUPinToInputGPUPinAlias.Find(OutputPin))
@@ -282,7 +285,11 @@ FPCGDataCollectionDesc UPCGSettings::ComputeOutputPinDataDesc(const UPCGPin* Out
 			Binding->GetStringTable());
 	}
 
-	ensure(false);
+	ensureMsgf(false, TEXT("Gathering data from CPU output pin '%s' on node '%s' failed. Pin was not present in OutputCPUPinToInputGPUPinAlias map (%d entries)."),
+		*OutputPin->Properties.Label.ToString(),
+		OutputPin->Node ? *OutputPin->Node->GetNodeTitle(EPCGNodeTitleType::ListView).ToString() : TEXT("MISSING"),
+		Binding->Graph->OutputCPUPinToInputGPUPinAlias.Num());
+
 	return {};
 }
 
diff --git a/Engine/Plugins/PCG/Source/PCG/Public/Elements/PCGReroute.h b/Engine/Plugins/PCG/Source/PCG/Public/Elements/PCGReroute.h
--- a/Engine/Plugins/PCG/Source/PCG/Public/Elements/PCGReroute.h
+++ b/Engine/Plugins/PCG/Source/PCG/Public/Elements/PCGReroute.h
@@ -40,6 +40,7 @@ protected:
 	virtual TArray<FPCGPinProperties> InputPinProperties() const override;
 	virtual TArray<FPCGPinProperties> OutputPinProperties() const override;
 	virtual FPCGElementPtr CreateElement() const override;
+	virtual FPCGDataCollectionDesc ComputeOutputPinDataDesc(const UPCGPin* OutputPin, const UPCGDataBinding* Binding) const override;
 	//~End UPCGSettings interface
 };
 
