Added spline, spline interior, volume and projection crc code.
Fixed an oversight in the point data crc computation, needed until we have a tightly packed data structure.
Added missing elements in the difference data crc.
Fixed missing shape copy in the volume data duplication code.
Added full data crc computation following subdivision nodes since they behave in a similar manner to samplers.
Small optimization to prevent component unregister/re-register when using or reusing it if it does not require a transform change.

#rnx

diff --git a/Engine/Plugins/PCG/Source/PCG/Private/Data/PCGCollisionShapeData.cpp b/Engine/Plugins/PCG/Source/PCG/Private/Data/PCGCollisionShapeData.cpp
--- a/Engine/Plugins/PCG/Source/PCG/Private/Data/PCGCollisionShapeData.cpp
+++ b/Engine/Plugins/PCG/Source/PCG/Private/Data/PCGCollisionShapeData.cpp
@@ -45,6 +45,8 @@ void UPCGCollisionShapeData::AddToCrc(FArchiveCrc32& Ar, bool bFullDataCrc) cons
 {
 	Super::AddToCrc(Ar, bFullDataCrc);
 
+	// Implementation note: no metadata at this point yet.
+
 	uint32 UniqueTypeID = StaticClass()->GetDefaultObject()->GetUniqueID();
 	Ar << UniqueTypeID;
 
diff --git a/Engine/Plugins/PCG/Source/PCG/Private/Data/PCGCollisionWrapperData.cpp b/Engine/Plugins/PCG/Source/PCG/Private/Data/PCGCollisionWrapperData.cpp
--- a/Engine/Plugins/PCG/Source/PCG/Private/Data/PCGCollisionWrapperData.cpp
+++ b/Engine/Plugins/PCG/Source/PCG/Private/Data/PCGCollisionWrapperData.cpp
@@ -259,10 +259,24 @@ const PhysicsInterfaceTypes::FInlineShapeArray& UPCGCollisionWrapperData::GetCac
 	}
 }
 
+FPCGCrc UPCGCollisionWrapperData::ComputeCrc(bool bFullDataCrc) const
+{
+	FArchiveCrc32 Ar;
+
+	AddToCrc(Ar, bFullDataCrc);
+	GetPointData()->AddToCrc(Ar, bFullDataCrc);
+
+	return FPCGCrc(Ar.GetCrc());
+}
+
 void UPCGCollisionWrapperData::AddToCrc(FArchiveCrc32& Ar, bool bFullDataCrc) const
 {
 	Super::AddToCrc(Ar, bFullDataCrc);
-	GetPointData()->AddToCrc(Ar, bFullDataCrc);
+
+	// Implementation note: no metadata at this point yet
+
+	uint32 UniqueTypeID = StaticClass()->GetDefaultObject()->GetUniqueID();
+	Ar << UniqueTypeID;
 
 	CollisionSelector.AddToCrc(Ar);
 
diff --git a/Engine/Plugins/PCG/Source/PCG/Private/Data/PCGDifferenceData.cpp b/Engine/Plugins/PCG/Source/PCG/Private/Data/PCGDifferenceData.cpp
--- a/Engine/Plugins/PCG/Source/PCG/Private/Data/PCGDifferenceData.cpp
+++ b/Engine/Plugins/PCG/Source/PCG/Private/Data/PCGDifferenceData.cpp
@@ -164,9 +164,14 @@ void UPCGDifferenceData::AddToCrc(FArchiveCrc32& Ar, bool bFullDataCrc) const
 {
 	Super::AddToCrc(Ar, bFullDataCrc);
 
+	// Implementation note: no metadata in composite data at this point.
+
 	uint32 UniqueTypeID = StaticClass()->GetDefaultObject()->GetUniqueID();
 	Ar << UniqueTypeID;
 
+	uint32 DiffMetadata = bDiffMetadata ? 1 : 0;
+	Ar << DiffMetadata;
+
 	uint32 DensityFunctionValue = static_cast<uint32>(DensityFunction);
 	Ar << DensityFunctionValue;
 }
diff --git a/Engine/Plugins/PCG/Source/PCG/Private/Data/PCGIntersectionData.cpp b/Engine/Plugins/PCG/Source/PCG/Private/Data/PCGIntersectionData.cpp
--- a/Engine/Plugins/PCG/Source/PCG/Private/Data/PCGIntersectionData.cpp
+++ b/Engine/Plugins/PCG/Source/PCG/Private/Data/PCGIntersectionData.cpp
@@ -101,6 +101,8 @@ void UPCGIntersectionData::AddToCrc(FArchiveCrc32& Ar, bool bFullDataCrc) const
 {
 	Super::AddToCrc(Ar, bFullDataCrc);
 
+	// Implementation note: no metadata in composite data at this point.
+
 	uint32 UniqueTypeID = StaticClass()->GetDefaultObject()->GetUniqueID();
 	Ar << UniqueTypeID;
 
diff --git a/Engine/Plugins/PCG/Source/PCG/Private/Data/PCGPointData.cpp b/Engine/Plugins/PCG/Source/PCG/Private/Data/PCGPointData.cpp
--- a/Engine/Plugins/PCG/Source/PCG/Private/Data/PCGPointData.cpp
+++ b/Engine/Plugins/PCG/Source/PCG/Private/Data/PCGPointData.cpp
@@ -298,52 +298,18 @@ void UPCGPointData::AddToCrc(FArchiveCrc32& Ar, bool bFullDataCrc) const
 	// Crc point data.
 	{
 		// Create copy so we can zero-out the metadata keys which are non-deterministic.
-		TArray<FPCGPoint> PointsCopy = Points;
-		for (FPCGPoint& Point : PointsCopy)
+		for (const FPCGPoint& OriginalPoint : Points)
 		{
+			FPCGPoint Point = OriginalPoint;
 			Point.MetadataEntry = 0;
-		}
-
-		Ar.Serialize(PointsCopy.GetData(), PointsCopy.Num() * PointsCopy.GetTypeSize());
+			Ar << Point;
+		} 
 	}
 
 	// Crc metadata.
 	if (const UPCGMetadata* PCGMetadata = ConstMetadata())
 	{
-		FPCGAttributeAccessorKeysPoints AccessorKeys(Points);
-
-		TArray<FName> AttributeNames;
-		{
-			TArray<EPCGMetadataTypes> AttributeTypes;
-			PCGMetadata->GetAttributes(AttributeNames, AttributeTypes);
-		}
-
-		// Attribute names might come in different orders for e.g. if edge order changes.
-		Algo::Sort(AttributeNames, [this](const FName& A, const FName& B) { return A.LexicalLess(B); });
-
-		for (FName AttributeName : AttributeNames)
-		{
-			Ar << AttributeName;
-
-			if (const FPCGMetadataAttributeBase* Attribute = PCGMetadata->GetConstAttribute(AttributeName))
-			{
-				for (const FPCGPoint& Point : Points)
-				{
-					auto Callback = [Attribute, PCGMetadata, &Ar, &Point](auto ValueWithType)
-					{
-						using AttributeType = decltype(ValueWithType);
-
-						if (const FPCGMetadataAttribute<AttributeType>* TypedAttribute = static_cast<const FPCGMetadataAttribute<AttributeType>*>(Attribute))
-						{
-							ValueWithType = TypedAttribute->GetValueFromItemKey(Point.MetadataEntry);
-							PCG::Private::Serialize(Ar, ValueWithType);
-						}
-					};
-
-					PCGMetadataAttribute::CallbackWithRightType(Attribute->GetTypeId(), Callback);
-				}
-			}
-		}
+		PCGMetadata->AddToCrc(Ar, bFullDataCrc);
 	}
 }
 
diff --git a/Engine/Plugins/PCG/Source/PCG/Private/Data/PCGProjectionData.cpp b/Engine/Plugins/PCG/Source/PCG/Private/Data/PCGProjectionData.cpp
--- a/Engine/Plugins/PCG/Source/PCG/Private/Data/PCGProjectionData.cpp
+++ b/Engine/Plugins/PCG/Source/PCG/Private/Data/PCGProjectionData.cpp
@@ -11,6 +11,8 @@
 #include "Helpers/PCGAsync.h"
 #include "Metadata/PCGMetadataAccessor.h"
 
+#include "Serialization/ArchiveCrc32.h"
+
 #include UE_INLINE_GENERATED_CPP_BY_NAME(PCGProjectionData)
 
 namespace PCGProjectionPrivate
@@ -68,12 +70,40 @@ void UPCGProjectionData::PostLoad()
 	ProjectionParams.ApplyDeprecation();
 }
 
+FPCGCrc UPCGProjectionData::ComputeCrc(bool bFullDataCrc) const
+{
+	FArchiveCrc32 Ar;
+
+	if (PropagateCrcThroughBooleanData())
+	{
+		AddToCrc(Ar, bFullDataCrc);
+
+		// Chain together CRCs of operands
+		check(Source && Target);
+		uint32 CrcSource = Source->GetOrComputeCrc(bFullDataCrc).GetValue();
+		uint32 CrcTarget = Target->GetOrComputeCrc(bFullDataCrc).GetValue();
+
+		Ar << CrcSource;
+		Ar << CrcTarget;
+	}
+	else
+	{
+		AddUIDToCrc(Ar);
+	}
+
+	return FPCGCrc(Ar.GetCrc());
+}
+
 void UPCGProjectionData::AddToCrc(FArchiveCrc32& Ar, bool bFullDataCrc) const
 {
 	Super::AddToCrc(Ar, bFullDataCrc);
 
-	// This data does not have a bespoke CRC implementation so just use a global unique data CRC.
-	AddUIDToCrc(Ar);
+	// Implementation note: no metadata in composite data at this point.
+
+	uint32 UniqueTypeID = StaticClass()->GetDefaultObject()->GetUniqueID();
+	Ar << UniqueTypeID;
+
+	Ar << ProjectionParams;
 }
 
 int UPCGProjectionData::GetDimension() const
diff --git a/Engine/Plugins/PCG/Source/PCG/Private/Data/PCGSplineData.cpp b/Engine/Plugins/PCG/Source/PCG/Private/Data/PCGSplineData.cpp
--- a/Engine/Plugins/PCG/Source/PCG/Private/Data/PCGSplineData.cpp
+++ b/Engine/Plugins/PCG/Source/PCG/Private/Data/PCGSplineData.cpp
@@ -11,6 +11,7 @@
 #include "Helpers/PCGHelpers.h"
 
 #include "Components/SplineComponent.h"
+#include "Serialization/ArchiveCrc32.h"
 
 #include UE_INLINE_GENERATED_CPP_BY_NAME(PCGSplineData)
 
@@ -80,8 +81,15 @@ void UPCGSplineData::AddToCrc(FArchiveCrc32& Ar, bool bFullDataCrc) const
 {
 	Super::AddToCrc(Ar, bFullDataCrc);
 
-	// This data does not have a bespoke CRC implementation so just use a global unique data CRC.
-	AddUIDToCrc(Ar);
+	if (Metadata)
+	{
+		Metadata->AddToCrc(Ar, bFullDataCrc);
+	}
+
+	uint32 UniqueTypeID = StaticClass()->GetDefaultObject()->GetUniqueID();
+	Ar << UniqueTypeID;
+
+	Ar << SplineStruct;
 }
 
 FTransform UPCGSplineData::GetTransform() const
diff --git a/Engine/Plugins/PCG/Source/PCG/Private/Data/PCGSplineInteriorSurfaceData.cpp b/Engine/Plugins/PCG/Source/PCG/Private/Data/PCGSplineInteriorSurfaceData.cpp
--- a/Engine/Plugins/PCG/Source/PCG/Private/Data/PCGSplineInteriorSurfaceData.cpp
+++ b/Engine/Plugins/PCG/Source/PCG/Private/Data/PCGSplineInteriorSurfaceData.cpp
@@ -7,6 +7,8 @@
 #include "Data/PCGSplineData.h"
 #include "Elements/PCGSplineSampler.h"
 
+#include "Serialization/ArchiveCrc32.h"
+
 #include UE_INLINE_GENERATED_CPP_BY_NAME(PCGSplineInteriorSurfaceData)
 
 void UPCGSplineInteriorSurfaceData::Initialize(FPCGContext* Context, const UPCGSplineData* InSplineData)
@@ -29,8 +31,12 @@ void UPCGSplineInteriorSurfaceData::AddToCrc(FArchiveCrc32& Ar, bool bFullDataCr
 {
 	Super::AddToCrc(Ar, bFullDataCrc);
 
-	// This data does not have a bespoke CRC implementation so just use a global unique data CRC.
-	AddUIDToCrc(Ar);
+	// Implementation note: no metadata in this data yet.
+
+	uint32 UniqueTypeID = StaticClass()->GetDefaultObject()->GetUniqueID();
+	Ar << UniqueTypeID;
+
+	Ar << SplineStruct;
 }
 
 bool UPCGSplineInteriorSurfaceData::SamplePoint(const FTransform& InTransform, const FBox& InBounds, FPCGPoint& OutPoint, UPCGMetadata* OutMetadata) const
diff --git a/Engine/Plugins/PCG/Source/PCG/Private/Data/PCGUnionData.cpp b/Engine/Plugins/PCG/Source/PCG/Private/Data/PCGUnionData.cpp
--- a/Engine/Plugins/PCG/Source/PCG/Private/Data/PCGUnionData.cpp
+++ b/Engine/Plugins/PCG/Source/PCG/Private/Data/PCGUnionData.cpp
@@ -118,6 +118,8 @@ void UPCGUnionData::AddToCrc(FArchiveCrc32& Ar, bool bFullDataCrc) const
 {
 	Super::AddToCrc(Ar, bFullDataCrc);
 
+	// Implementation note: no metadata in composite data at this point.
+
 	uint32 UniqueTypeID = StaticClass()->GetDefaultObject()->GetUniqueID();
 	Ar << UniqueTypeID;
 
diff --git a/Engine/Plugins/PCG/Source/PCG/Private/Data/PCGVolumeData.cpp b/Engine/Plugins/PCG/Source/PCG/Private/Data/PCGVolumeData.cpp
--- a/Engine/Plugins/PCG/Source/PCG/Private/Data/PCGVolumeData.cpp
+++ b/Engine/Plugins/PCG/Source/PCG/Private/Data/PCGVolumeData.cpp
@@ -12,6 +12,7 @@
 
 #include "Chaos/ChaosEngineInterface.h"
 #include "Physics/PhysicsInterfaceDeclares.h"
+#include "Serialization/ArchiveCrc32.h"
 
 #include UE_INLINE_GENERATED_CPP_BY_NAME(PCGVolumeData)
 
@@ -83,8 +84,40 @@ void UPCGVolumeData::AddToCrc(FArchiveCrc32& Ar, bool bFullDataCrc) const
 {
 	Super::AddToCrc(Ar, bFullDataCrc);
 
-	// This data does not have a bespoke CRC implementation so just use a global unique data CRC.
-	AddUIDToCrc(Ar);
+	if (VolumeBodyInstance)
+	{
+		// Implementation note: no metadata in this data at this point.
+
+		uint32 UniqueTypeID = StaticClass()->GetDefaultObject()->GetUniqueID();
+		Ar << UniqueTypeID;
+
+		// Weird implementation to fix ambiguous call between the FBodyInstance friend function (which isn't exposed) and the one in the FArchiveCrc32 code.
+		FBodyInstance::StaticStruct()->SerializeItem(Ar, const_cast<FBodyInstance*>(VolumeBodyInstance), nullptr);
+
+		// Implementation note: we will not consider the volume pointer in this instance
+		Ar.Serialize((void*)&VoxelSize, sizeof(FVector));
+
+		// TODO: move this to helper function
+		auto SerializeBounds = [&Ar](FBox* Box)
+		{
+			check(Box);
+			Ar << Box->IsValid;
+
+			if (Box->IsValid)
+			{
+				Ar << Box->Min;
+				Ar << Box->Max;
+			}
+		};
+
+		SerializeBounds(const_cast<FBox*>(&Bounds));
+		SerializeBounds(const_cast<FBox*>(&StrictBounds));
+	}
+	else
+	{
+		// This data does not have a bespoke CRC implementation so just use a global unique data CRC.
+		AddUIDToCrc(Ar);
+	}
 }
 
 FBox UPCGVolumeData::GetBounds() const
@@ -163,6 +196,13 @@ void UPCGVolumeData::CopyBaseVolumeData(UPCGVolumeData* NewVolumeData) const
 	NewVolumeData->Volume = Volume;
 	NewVolumeData->Bounds = Bounds;
 	NewVolumeData->StrictBounds = StrictBounds;
+
+	if (VolumeBodyInstance)
+	{
+		NewVolumeData->ReleaseInternalBodyInstance();
+		NewVolumeData->VolumeBodyInstance = new FBodyInstance();
+		NewVolumeData->VolumeBodyInstance->CopyBodyInstancePropertiesFrom(VolumeBodyInstance);
+	}
 }
 
 UPCGSpatialData* UPCGVolumeData::CopyInternal(FPCGContext* Context) const
diff --git a/Engine/Plugins/PCG/Source/PCG/Private/Metadata/PCGMetadata.cpp b/Engine/Plugins/PCG/Source/PCG/Private/Metadata/PCGMetadata.cpp
--- a/Engine/Plugins/PCG/Source/PCG/Private/Metadata/PCGMetadata.cpp
+++ b/Engine/Plugins/PCG/Source/PCG/Private/Metadata/PCGMetadata.cpp
@@ -8,10 +8,12 @@
 #include "Elements/Metadata/PCGMetadataElementCommon.h"
 #include "Helpers/PCGPropertyHelpers.h"
 #include "Metadata/PCGAttributePropertySelector.h"
+#include "Metadata/Accessors/PCGAttributeAccessorHelpers.h"
 
 #include "Algo/AnyOf.h"
 #include "Algo/Transform.h"
 #include "Async/ParallelFor.h"
+#include "Serialization/ArchiveCrc32.h"
 
 #include UE_INLINE_GENERATED_CPP_BY_NAME(PCGMetadata)
 
@@ -71,6 +73,64 @@ void UPCGMetadata::Serialize(FArchive& InArchive)
 	}
 }
 
+void UPCGMetadata::AddToCrc(FArchiveCrc32& Ar, bool bFullDataCrc) const
+{
+	const UPCGData* Data = Cast<UPCGData>(GetOuter());
+	check(Data);
+
+	// Sort attributes so we have a consistent processing path
+	TArray<FName> AttributeNames;
+
+	{
+		AttributeLock.ReadLock();
+		AttributeNames.Reserve(Attributes.Num());
+
+		for (const TPair<FName, FPCGMetadataAttributeBase*>& Attribute : Attributes)
+		{
+			AttributeNames.Add(Attribute.Key);
+		}
+		AttributeLock.ReadUnlock();
+
+		Algo::Sort(AttributeNames, [this](const FName& A, const FName& B) { return A.LexicalLess(B); });
+	}
+
+	// Add attributes to CRC
+	for (FName AttributeName : AttributeNames)
+	{
+		Ar << AttributeName;
+
+		FPCGAttributePropertyInputSelector InputSource;
+		InputSource.SetAttributeName(AttributeName);
+
+		TUniquePtr<const IPCGAttributeAccessor> InputAccessor = PCGAttributeAccessorHelpers::CreateConstAccessor(Data, InputSource);
+		TUniquePtr<const IPCGAttributeAccessorKeys> InputKeys = PCGAttributeAccessorHelpers::CreateConstKeys(Data, InputSource);
+
+		auto Callback = [&InputAccessor, &InputKeys, &Ar](auto&& Dummy)
+		{
+			using AttributeType = std::decay_t<decltype(Dummy)>;
+			TArray<AttributeType> Values;
+			if constexpr (std::is_trivially_copyable_v<AttributeType>)
+			{
+				Values.SetNumUninitialized(InputKeys->GetNum());
+			}
+			else
+			{
+				Values.SetNum(InputKeys->GetNum());
+			}
+			
+			InputAccessor->GetRange<AttributeType>(Values, 0, *InputKeys);
+
+			for (AttributeType& Value : Values)
+			{
+				// Add value to Crc
+				PCG::Private::Serialize(Ar, Value);
+			}
+		};
+
+		PCGMetadataAttribute::CallbackWithRightType(InputAccessor->GetUnderlyingType(), Callback);
+	}
+}
+
 void UPCGMetadata::BeginDestroy()
 {
 	FWriteScopeLock ScopeLock(AttributeLock);
diff --git a/Engine/Plugins/PCG/Source/PCG/Private/PCGData.cpp b/Engine/Plugins/PCG/Source/PCG/Private/PCGData.cpp
--- a/Engine/Plugins/PCG/Source/PCG/Private/PCGData.cpp
+++ b/Engine/Plugins/PCG/Source/PCG/Private/PCGData.cpp
@@ -20,7 +20,7 @@
 static TAutoConsoleVariable<bool> CVarCachePropagateCrcThroughBooleanData(
 	TEXT("pcg.Cache.PropagateCrcThroughBooleanData"),
 	false,
-	TEXT("Whether intersection, union, difference combine Crc values from operands. If false they fall back to using data UID."));
+	TEXT("Whether intersection, union, difference, projection combine Crc values from operands. If false they fall back to using data UID."));
 
 UPCGData::UPCGData(const FObjectInitializer& ObjectInitializer)
 	: Super(ObjectInitializer)
diff --git a/Engine/Plugins/PCG/Source/PCG/Private/PCGManagedResource.cpp b/Engine/Plugins/PCG/Source/PCG/Private/PCGManagedResource.cpp
--- a/Engine/Plugins/PCG/Source/PCG/Private/PCGManagedResource.cpp
+++ b/Engine/Plugins/PCG/Source/PCG/Private/PCGManagedResource.cpp
@@ -782,6 +782,8 @@ void UPCGManagedISMComponent::MarkAsUsed()
 
 	if (UInstancedStaticMeshComponent* ISMC = GetComponent())
 	{
+		const bool bHasPreviousRootLocation = bHasRootLocation;
+
 		// Keep track of the current root location so if we reuse this later we are able to update this appropriately
 		if (USceneComponent* RootComponent = ISMC->GetAttachmentRoot())
 		{
@@ -794,11 +796,14 @@ void UPCGManagedISMComponent::MarkAsUsed()
 			RootLocation = FVector::ZeroVector;
 		}
 
-		// Reset the rotation/scale to be identity otherwise if the root component transform has changed, the final transform will be wrong
-		// Since this is technically 'moving' the ISM, we need to unregister it before moving otherwise we could get a warning that we're moving a component with static mobility
-		ISMC->UnregisterComponent();
-		ISMC->SetWorldTransform(FTransform(FQuat::Identity, RootLocation, FVector::OneVector));
-		ISMC->RegisterComponent();
+		if (bHasPreviousRootLocation != bHasRootLocation || (ISMC->GetComponentLocation() - RootLocation).SquaredLength() > UE_DOUBLE_SMALL_NUMBER)
+		{
+			// Reset the rotation/scale to be identity otherwise if the root component transform has changed, the final transform will be wrong
+			// Since this is technically 'moving' the ISM, we need to unregister it before moving otherwise we could get a warning that we're moving a component with static mobility
+			ISMC->UnregisterComponent();
+			ISMC->SetWorldTransform(FTransform(FQuat::Identity, RootLocation, FVector::OneVector));
+			ISMC->RegisterComponent();
+		}
 	}
 }
 
@@ -819,10 +824,13 @@ void UPCGManagedISMComponent::MarkAsReused()
 			}
 		}
 
-		// Since this is technically 'moving' the ISM, we need to unregister it before moving otherwise we could get a warning that we're moving a component with static mobility
-		ISMC->UnregisterComponent();
-		ISMC->SetWorldTransform(FTransform(FQuat::Identity, TentativeRootLocation, FVector::OneVector));
-		ISMC->RegisterComponent();
+		if ((ISMC->GetComponentLocation() - TentativeRootLocation).SquaredLength() > UE_DOUBLE_SMALL_NUMBER)
+		{
+			// Since this is technically 'moving' the ISM, we need to unregister it before moving otherwise we could get a warning that we're moving a component with static mobility
+			ISMC->UnregisterComponent();
+			ISMC->SetWorldTransform(FTransform(FQuat::Identity, TentativeRootLocation, FVector::OneVector));
+			ISMC->RegisterComponent();
+		}
 	}
 }
 
diff --git a/Engine/Plugins/PCG/Source/PCG/Private/PCGParamData.cpp b/Engine/Plugins/PCG/Source/PCG/Private/PCGParamData.cpp
--- a/Engine/Plugins/PCG/Source/PCG/Private/PCGParamData.cpp
+++ b/Engine/Plugins/PCG/Source/PCG/Private/PCGParamData.cpp
@@ -21,48 +21,14 @@ UPCGParamData::UPCGParamData(const FObjectInitializer& ObjectInitializer)
 
 void UPCGParamData::AddToCrc(FArchiveCrc32& Ar, bool bFullDataCrc) const
 {
+	Super::AddToCrc(Ar, bFullDataCrc);
+
 	uint32 UniqueTypeID = StaticClass()->GetDefaultObject()->GetUniqueID();
 	Ar << UniqueTypeID;
 
-	if (!Metadata)
+	if (Metadata)
 	{
-		// Nothing to contribute
-		return;
-	}
-
-	// Get attribute names. Preserve order as attribute order matters.
-	TArray<FName> AttributeNames;
-	{
-		TArray<EPCGMetadataTypes> AttributeTypes;
-		Metadata->GetAttributes(AttributeNames, AttributeTypes);
-	}
-
-	// Add attributes to CRC
-	for (FName AttributeName : AttributeNames)
-	{
-		Ar << AttributeName;
-
-		FPCGAttributePropertyInputSelector InputSource;
-		InputSource.SetAttributeName(AttributeName);
-
-		TUniquePtr<const IPCGAttributeAccessor> InputAccessor = PCGAttributeAccessorHelpers::CreateConstAccessor(this, InputSource);
-		TUniquePtr<const IPCGAttributeAccessorKeys> InputKeys = PCGAttributeAccessorHelpers::CreateConstKeys(this, InputSource);
-
-		auto Callback = [&InputAccessor, &InputKeys, &Ar](auto && Dummy)
-		{
-			using AttributeType = std::decay_t<decltype(Dummy)>;
-			TArray<AttributeType> Values;
-			Values.SetNum(InputKeys->GetNum());
-			InputAccessor->GetRange<AttributeType>(Values, 0, *InputKeys);
-
-			for (AttributeType Value : Values)
-			{
-				// Add value to Crc
-				PCG::Private::Serialize(Ar, Value);
-			}
-		};
-
-		PCGMetadataAttribute::CallbackWithRightType(InputAccessor->GetUnderlyingType(), Callback);
+		Metadata->AddToCrc(Ar, bFullDataCrc);
 	}
 }
 
diff --git a/Engine/Plugins/PCG/Source/PCG/Public/Data/PCGCollisionWrapperData.h b/Engine/Plugins/PCG/Source/PCG/Public/Data/PCGCollisionWrapperData.h
--- a/Engine/Plugins/PCG/Source/PCG/Public/Data/PCGCollisionWrapperData.h
+++ b/Engine/Plugins/PCG/Source/PCG/Public/Data/PCGCollisionWrapperData.h
@@ -68,6 +68,7 @@ public:
 
 	// ~Begin UPCGData Interface
 	virtual EPCGDataType GetDataType() const override { return EPCGDataType::Primitive; }
+	virtual FPCGCrc ComputeCrc(bool bFullDataCrc) const override;
 	virtual void AddToCrc(FArchiveCrc32& Ar, bool bFullDataCrc) const override;
 	virtual void GetResourceSizeEx(FResourceSizeEx& CumulativeResourceSize) override;
 	// ~End UPCGData Interface
diff --git a/Engine/Plugins/PCG/Source/PCG/Public/Data/PCGProjectionData.h b/Engine/Plugins/PCG/Source/PCG/Public/Data/PCGProjectionData.h
--- a/Engine/Plugins/PCG/Source/PCG/Public/Data/PCGProjectionData.h
+++ b/Engine/Plugins/PCG/Source/PCG/Public/Data/PCGProjectionData.h
@@ -23,6 +23,7 @@ public:
 	const FPCGProjectionParams& GetProjectionParams() const { return ProjectionParams; }
 
 	// ~Begin UPCGData interface
+	virtual FPCGCrc ComputeCrc(bool bFullDataCrc) const override;
 	virtual void AddToCrc(FArchiveCrc32& Ar, bool bFullDataCrc) const override;
 	// ~End UPCGData interface
 
diff --git a/Engine/Plugins/PCG/Source/PCG/Public/Elements/Grammar/PCGSubdivisionBase.h b/Engine/Plugins/PCG/Source/PCG/Public/Elements/Grammar/PCGSubdivisionBase.h
--- a/Engine/Plugins/PCG/Source/PCG/Public/Elements/Grammar/PCGSubdivisionBase.h
+++ b/Engine/Plugins/PCG/Source/PCG/Public/Elements/Grammar/PCGSubdivisionBase.h
@@ -469,6 +469,10 @@ namespace PCGSubdivisionBase
 
 class FPCGSubdivisionBaseElement : public IPCGElement
 {
+public:
+	// Worth computing a full CRC in case we can halt change propagation/re-executions
+	virtual bool ShouldComputeFullOutputDataCrc(FPCGContext* Context) const override { return true; }
+
 protected:
 	using FModuleInfoMap = PCGSubdivisionBase::FModuleInfoMap;
 
diff --git a/Engine/Plugins/PCG/Source/PCG/Public/Metadata/PCGMetadata.h b/Engine/Plugins/PCG/Source/PCG/Public/Metadata/PCGMetadata.h
--- a/Engine/Plugins/PCG/Source/PCG/Public/Metadata/PCGMetadata.h
+++ b/Engine/Plugins/PCG/Source/PCG/Public/Metadata/PCGMetadata.h
@@ -348,6 +348,9 @@ public:
 		return InitializeAsCopyWithAttributeFilter(InMetadataToCopy, InFilteredAttributes, InFilterMode, !InOptionalEntriesToCopy.IsEmpty() ? &InOptionalEntriesToCopy : nullptr);
 	}
 
+	/** Computes Crc from all attributes & keys from outer's data. */
+	void AddToCrc(FArchiveCrc32& Ar, bool bFullDataCrc) const;
+
 protected:
 	FPCGMetadataAttributeBase* CopyAttribute(FName AttributeToCopy, FName NewAttributeName, bool bKeepParent, bool bCopyEntries, bool bCopyValues);
 
diff --git a/Engine/Plugins/PCGInterops/PCGGeometryScriptInterop/Source/PCGGeometryScriptInterop/Private/Data/PCGDynamicMeshData.cpp b/Engine/Plugins/PCGInterops/PCGGeometryScriptInterop/Source/PCGGeometryScriptInterop/Private/Data/PCGDynamicMeshData.cpp
--- a/Engine/Plugins/PCGInterops/PCGGeometryScriptInterop/Source/PCGGeometryScriptInterop/Private/Data/PCGDynamicMeshData.cpp
+++ b/Engine/Plugins/PCGInterops/PCGGeometryScriptInterop/Source/PCGGeometryScriptInterop/Private/Data/PCGDynamicMeshData.cpp
@@ -72,6 +72,11 @@ void UPCGDynamicMeshData::AddToCrc(FArchiveCrc32& Ar, bool bFullDataCrc) const
 
 	if (bFullDataCrc)
 	{
+		// Implementation note: metadata not supported at this point.
+
+		uint32 UniqueTypeID = StaticClass()->GetDefaultObject()->GetUniqueID();
+		Ar << UniqueTypeID;
+
 		Ar << DynamicMesh->GetMeshRef();
 	}
 	else
