[PCG] Cleared deprecated content from 5.2, 5.3

#rnx

diff --git a/Engine/Plugins/PCG/Source/PCG/Private/Helpers/PCGActorHelpers.cpp b/Engine/Plugins/PCG/Source/PCG/Private/Helpers/PCGActorHelpers.cpp
--- a/Engine/Plugins/PCG/Source/PCG/Private/Helpers/PCGActorHelpers.cpp
+++ b/Engine/Plugins/PCG/Source/PCG/Private/Helpers/PCGActorHelpers.cpp
@@ -432,16 +432,6 @@ bool UPCGActorHelpers::DeleteActors(UWorld* World, const TArray<TSoftObjectPtr<A
 	return true;
 }
 
-void UPCGActorHelpers::GetActorClassDefaultComponents(const TSubclassOf<AActor>& ActorClass, TArray<UActorComponent*>& OutComponents, const TSubclassOf<UActorComponent>& InComponentClass)
-{
-	OutComponents.Reset();
-	AActor::ForEachComponentOfActorClassDefault(ActorClass, InComponentClass, [&OutComponents](const UActorComponent* TemplateComponent)
-	{
-		OutComponents.Add(const_cast<UActorComponent*>(TemplateComponent));
-		return true;
-	});
-}
-
 void UPCGActorHelpers::ForEachActorInLevel(ULevel* Level, TSubclassOf<AActor> ActorClass, TFunctionRef<bool(AActor*)> Callback)
 {
 	if (!Level)
@@ -619,4 +609,3 @@ FIntVector UPCGActorHelpers::GetCellCoord(FVector InPosition, int InGridSize, bo
 		bUse2DGrid ? 0 : FMath::FloorToInt(Temp.Z)
 	);
 }
-
diff --git a/Engine/Plugins/PCG/Source/PCG/Private/Helpers/PCGSettingsHelpers.cpp b/Engine/Plugins/PCG/Source/PCG/Private/Helpers/PCGSettingsHelpers.cpp
--- a/Engine/Plugins/PCG/Source/PCG/Private/Helpers/PCGSettingsHelpers.cpp
+++ b/Engine/Plugins/PCG/Source/PCG/Private/Helpers/PCGSettingsHelpers.cpp
@@ -209,16 +209,6 @@ namespace PCGSettingsHelpers
 		}
 	}
 
-	PRAGMA_DISABLE_DEPRECATION_WARNINGS
-	int ComputeSeedWithOverride(const UPCGSettings* InSettings, const UPCGComponent* InComponent, UPCGParamData* InParams)
-	{
-		check(InSettings);
-
-		const int SettingsSeed = InParams ? PCGSettingsHelpers::GetValue(GET_MEMBER_NAME_CHECKED(UPCGSettings, Seed), InSettings->Seed, InParams) : InSettings->Seed;
-		return InComponent ? PCGHelpers::ComputeSeed(SettingsSeed, InComponent->Seed) : SettingsSeed;
-	}
-	PRAGMA_ENABLE_DEPRECATION_WARNINGS
-
 	void DeprecationBreakOutParamsToNewPin(UPCGNode* InOutNode, TArray<TObjectPtr<UPCGPin>>& InputPins, TArray<TObjectPtr<UPCGPin>>& OutputPins)
 	{
 		// Check basic conditions for which the code below should run.
diff --git a/Engine/Plugins/PCG/Source/PCG/Private/Metadata/Accessors/PCGAttributeAccessorHelpers.cpp b/Engine/Plugins/PCG/Source/PCG/Private/Metadata/Accessors/PCGAttributeAccessorHelpers.cpp
--- a/Engine/Plugins/PCG/Source/PCG/Private/Metadata/Accessors/PCGAttributeAccessorHelpers.cpp
+++ b/Engine/Plugins/PCG/Source/PCG/Private/Metadata/Accessors/PCGAttributeAccessorHelpers.cpp
@@ -542,21 +542,6 @@ TUniquePtr<IPCGAttributeAccessor> PCGAttributeAccessorHelpers::CreateExtraAccess
 	}
 }
 
-TUniquePtr<const IPCGAttributeAccessor> PCGAttributeAccessorHelpers::CreateConstAccessorForOverrideParam(const FPCGDataCollection& InInputData, const FPCGSettingsOverridableParam& InParam, FName* OutAttributeName)
-{
-	if (OutAttributeName)
-	{
-		AccessorParamResult Result{};
-		TUniquePtr<const IPCGAttributeAccessor> Accessor = CreateConstAccessorForOverrideParamWithResult(InInputData, InParam, &Result);
-		*OutAttributeName = Result.AttributeName;
-		return Accessor;
-	}
-	else
-	{
-		return CreateConstAccessorForOverrideParamWithResult(InInputData, InParam);
-	}
-}
-
 TUniquePtr<const IPCGAttributeAccessor> PCGAttributeAccessorHelpers::CreateConstAccessorForOverrideParamWithResult(const FPCGDataCollection& InInputData, const FPCGSettingsOverridableParam& InParam, AccessorParamResult* OutResult)
 {
 	bool bFromGlobalParamsPin = false;
diff --git a/Engine/Plugins/PCG/Source/PCG/Private/PCGData.cpp b/Engine/Plugins/PCG/Source/PCG/Private/PCGData.cpp
--- a/Engine/Plugins/PCG/Source/PCG/Private/PCGData.cpp
+++ b/Engine/Plugins/PCG/Source/PCG/Private/PCGData.cpp
@@ -317,19 +317,6 @@ UPCGParamData* FPCGDataCollection::GetParamsWithDeprecation(const UPCGNode* Node
 	return nullptr;
 }
 
-UPCGParamData* FPCGDataCollection::GetParams() const
-{
-	for (const FPCGTaggedData& TaggedDatum : TaggedData)
-	{
-		if (const UPCGParamData* Params = Cast<UPCGParamData>(TaggedDatum.Data))
-		{
-			return const_cast<UPCGParamData*>(Params); 
-		}
-	}
-
-	return nullptr;
-}
-
 UPCGParamData* FPCGDataCollection::GetFirstParamsOnParamsPin() const
 {
 	TArray<FPCGTaggedData> ParamsOnDefaultPin = GetParamsByPin(PCGPinConstants::DefaultParamsLabel);
diff --git a/Engine/Plugins/PCG/Source/PCG/Public/Data/PCGDifferenceData.h b/Engine/Plugins/PCG/Source/PCG/Public/Data/PCGDifferenceData.h
--- a/Engine/Plugins/PCG/Source/PCG/Public/Data/PCGDifferenceData.h
+++ b/Engine/Plugins/PCG/Source/PCG/Public/Data/PCGDifferenceData.h
@@ -116,7 +116,3 @@ private:
 	UPCGUnionData* RawPointerDifferencesUnion = nullptr;
 #endif
 };
-
-#if UE_ENABLE_INCLUDE_ORDER_DEPRECATED_IN_5_2
-#include "CoreMinimal.h"
-#endif
diff --git a/Engine/Plugins/PCG/Source/PCG/Public/Data/PCGLandscapeData.h b/Engine/Plugins/PCG/Source/PCG/Public/Data/PCGLandscapeData.h
--- a/Engine/Plugins/PCG/Source/PCG/Public/Data/PCGLandscapeData.h
+++ b/Engine/Plugins/PCG/Source/PCG/Public/Data/PCGLandscapeData.h
@@ -109,7 +109,3 @@ private:
 	TArray<ULandscapeInfo*> LandscapeInfos;
 	UPCGLandscapeCache* LandscapeCache = nullptr;
 };
-
-#if UE_ENABLE_INCLUDE_ORDER_DEPRECATED_IN_5_2
-#include "CoreMinimal.h"
-#endif
diff --git a/Engine/Plugins/PCG/Source/PCG/Public/Data/PCGLandscapeSplineData.h b/Engine/Plugins/PCG/Source/PCG/Public/Data/PCGLandscapeSplineData.h
--- a/Engine/Plugins/PCG/Source/PCG/Public/Data/PCGLandscapeSplineData.h
+++ b/Engine/Plugins/PCG/Source/PCG/Public/Data/PCGLandscapeSplineData.h
@@ -63,8 +63,3 @@ protected:
 	/** Reparameterization of the spline by distance. Useful to query the InputKey at arbitrary distance. */
 	FInterpCurveFloat ReparamTable;
 };
-
-#if UE_ENABLE_INCLUDE_ORDER_DEPRECATED_IN_5_2
-#include "CoreMinimal.h"
-#include "PCGProjectionData.h"
-#endif
diff --git a/Engine/Plugins/PCG/Source/PCG/Public/Data/PCGPolyLineData.h b/Engine/Plugins/PCG/Source/PCG/Public/Data/PCGPolyLineData.h
--- a/Engine/Plugins/PCG/Source/PCG/Public/Data/PCGPolyLineData.h
+++ b/Engine/Plugins/PCG/Source/PCG/Public/Data/PCGPolyLineData.h
@@ -75,8 +75,3 @@ public:
 	/** True if the line holds custom metadata. */
 	virtual bool HasCustomMetadata() const { return false; }
 };
-
-
-#if UE_ENABLE_INCLUDE_ORDER_DEPRECATED_IN_5_2
-#include "CoreMinimal.h"
-#endif
diff --git a/Engine/Plugins/PCG/Source/PCG/Public/Data/PCGProjectionData.h b/Engine/Plugins/PCG/Source/PCG/Public/Data/PCGProjectionData.h
--- a/Engine/Plugins/PCG/Source/PCG/Public/Data/PCGProjectionData.h
+++ b/Engine/Plugins/PCG/Source/PCG/Public/Data/PCGProjectionData.h
@@ -69,7 +69,3 @@ protected:
 	UPROPERTY(BlueprintReadwrite, VisibleAnywhere, Category = SpatialData)
 	FPCGProjectionParams ProjectionParams;
 };
-
-#if UE_ENABLE_INCLUDE_ORDER_DEPRECATED_IN_5_2
-#include "PCGPointData.h"
-#endif
diff --git a/Engine/Plugins/PCG/Source/PCG/Public/Data/PCGSpatialData.h b/Engine/Plugins/PCG/Source/PCG/Public/Data/PCGSpatialData.h
--- a/Engine/Plugins/PCG/Source/PCG/Public/Data/PCGSpatialData.h
+++ b/Engine/Plugins/PCG/Source/PCG/Public/Data/PCGSpatialData.h
@@ -229,8 +229,3 @@ private:
 
 	mutable FCriticalSection CacheLock;
 };
-
-#if UE_ENABLE_INCLUDE_ORDER_DEPRECATED_IN_5_2
-#include "CoreMinimal.h"
-#include "PCGContext.h"
-#endif
diff --git a/Engine/Plugins/PCG/Source/PCG/Public/Data/PCGSplineData.h b/Engine/Plugins/PCG/Source/PCG/Public/Data/PCGSplineData.h
--- a/Engine/Plugins/PCG/Source/PCG/Public/Data/PCGSplineData.h
+++ b/Engine/Plugins/PCG/Source/PCG/Public/Data/PCGSplineData.h
@@ -101,7 +101,3 @@ protected:
 	virtual UPCGSpatialData* CopyInternal(FPCGContext* Context) const override;
 	//~End UPCGSpatialData interface
 };
-
-#if UE_ENABLE_INCLUDE_ORDER_DEPRECATED_IN_5_2
-#include "CoreMinimal.h"
-#endif
diff --git a/Engine/Plugins/PCG/Source/PCG/Public/Data/PCGSurfaceData.h b/Engine/Plugins/PCG/Source/PCG/Public/Data/PCGSurfaceData.h
--- a/Engine/Plugins/PCG/Source/PCG/Public/Data/PCGSurfaceData.h
+++ b/Engine/Plugins/PCG/Source/PCG/Public/Data/PCGSurfaceData.h
@@ -33,7 +33,3 @@ protected:
 	UPROPERTY(BlueprintReadOnly, VisibleAnywhere, Category = SpatialData)
 	FBox LocalBounds = FBox(EForceInit::ForceInit);
 };
-
-#if UE_ENABLE_INCLUDE_ORDER_DEPRECATED_IN_5_2
-#include "CoreMinimal.h"
-#endif
diff --git a/Engine/Plugins/PCG/Source/PCG/Public/Data/PCGTextureData.h b/Engine/Plugins/PCG/Source/PCG/Public/Data/PCGTextureData.h
--- a/Engine/Plugins/PCG/Source/PCG/Public/Data/PCGTextureData.h
+++ b/Engine/Plugins/PCG/Source/PCG/Public/Data/PCGTextureData.h
@@ -214,7 +214,3 @@ public:
 	// Added to help deprecation in 5.5. To be removed when the deprecated Initialized function is removed.
 	TFunction<void()> PostInitializeCallback;
 };
-
-#if UE_ENABLE_INCLUDE_ORDER_DEPRECATED_IN_5_2
-#include "Engine/Texture2D.h"
-#endif
diff --git a/Engine/Plugins/PCG/Source/PCG/Public/Data/PCGUnionData.h b/Engine/Plugins/PCG/Source/PCG/Public/Data/PCGUnionData.h
--- a/Engine/Plugins/PCG/Source/PCG/Public/Data/PCGUnionData.h
+++ b/Engine/Plugins/PCG/Source/PCG/Public/Data/PCGUnionData.h
@@ -88,7 +88,3 @@ protected:
 private:
 	void CreateSequentialPointData(FPCGContext* Context, TArray<const UPCGSpatialData*>& DataRawPtr, TArray<const UPCGMetadata*>& InputMetadatas, UPCGPointData* PointData, UPCGMetadata* OutMetadata, bool bLeftToRight) const;
 };
-
-#if UE_ENABLE_INCLUDE_ORDER_DEPRECATED_IN_5_2
-#include "CoreMinimal.h"
-#endif
diff --git a/Engine/Plugins/PCG/Source/PCG/Public/Data/PCGVolumeData.h b/Engine/Plugins/PCG/Source/PCG/Public/Data/PCGVolumeData.h
--- a/Engine/Plugins/PCG/Source/PCG/Public/Data/PCGVolumeData.h
+++ b/Engine/Plugins/PCG/Source/PCG/Public/Data/PCGVolumeData.h
@@ -59,7 +59,3 @@ protected:
 	// Internal body instance to perform queries faster, used in static cases only
 	FBodyInstance* VolumeBodyInstance = nullptr;
 };
-
-#if UE_ENABLE_INCLUDE_ORDER_DEPRECATED_IN_5_2
-#include "CoreMinimal.h"
-#endif
diff --git a/Engine/Plugins/PCG/Source/PCG/Public/Elements/Metadata/PCGMetadataBreakTransform.h b/Engine/Plugins/PCG/Source/PCG/Public/Elements/Metadata/PCGMetadataBreakTransform.h
--- a/Engine/Plugins/PCG/Source/PCG/Public/Elements/Metadata/PCGMetadataBreakTransform.h
+++ b/Engine/Plugins/PCG/Source/PCG/Public/Elements/Metadata/PCGMetadataBreakTransform.h
@@ -54,7 +54,3 @@ class FPCGMetadataBreakTransformElement : public FPCGMetadataElementBase
 protected:
 	virtual bool DoOperation(PCGMetadataOps::FOperationData& OperationData) const override;
 };
-
-#if UE_ENABLE_INCLUDE_ORDER_DEPRECATED_IN_5_2
-#include "Elements/Metadata/PCGMetadataMakeTransform.h"
-#endif
diff --git a/Engine/Plugins/PCG/Source/PCG/Public/Elements/Metadata/PCGMetadataOpElementBase.h b/Engine/Plugins/PCG/Source/PCG/Public/Elements/Metadata/PCGMetadataOpElementBase.h
--- a/Engine/Plugins/PCG/Source/PCG/Public/Elements/Metadata/PCGMetadataOpElementBase.h
+++ b/Engine/Plugins/PCG/Source/PCG/Public/Elements/Metadata/PCGMetadataOpElementBase.h
@@ -509,8 +509,3 @@ inline bool FPCGMetadataElementBase::DoQuaternaryOp(PCGMetadataOps::FOperationDa
 {
 	return DoNAryOp<InType1, InType2, InType3, InType4>(InOperationData, ForwardAsTuple(std::forward<Callbacks>(InCallbacks)...));
 }
-
-#if UE_ENABLE_INCLUDE_ORDER_DEPRECATED_IN_5_2
-#include "Metadata/PCGMetadataAttribute.h"
-#include "Metadata/PCGMetadataAttributeTpl.h"
-#endif
diff --git a/Engine/Plugins/PCG/Source/PCG/Public/Elements/PCGAttributeNoise.h b/Engine/Plugins/PCG/Source/PCG/Public/Elements/PCGAttributeNoise.h
--- a/Engine/Plugins/PCG/Source/PCG/Public/Elements/PCGAttributeNoise.h
+++ b/Engine/Plugins/PCG/Source/PCG/Public/Elements/PCGAttributeNoise.h
@@ -129,8 +129,3 @@ protected:
 	virtual bool ExecuteInternal(FPCGContext* Context) const override;
 	virtual EPCGElementExecutionLoopMode ExecutionLoopMode(const UPCGSettings* Settings) const override { return EPCGElementExecutionLoopMode::SinglePrimaryPin; }
 };
-
-#if UE_ENABLE_INCLUDE_ORDER_DEPRECATED_IN_5_2
-#include "CoreMinimal.h"
-#include "PCGNode.h"
-#endif
diff --git a/Engine/Plugins/PCG/Source/PCG/Public/Elements/PCGCopyPoints.h b/Engine/Plugins/PCG/Source/PCG/Public/Elements/PCGCopyPoints.h
--- a/Engine/Plugins/PCG/Source/PCG/Public/Elements/PCGCopyPoints.h
+++ b/Engine/Plugins/PCG/Source/PCG/Public/Elements/PCGCopyPoints.h
@@ -99,7 +99,3 @@ class FPCGCopyPointsElement : public IPCGElement
 protected:
 	virtual bool ExecuteInternal(FPCGContext* Context) const override;
 };
-
-#if UE_ENABLE_INCLUDE_ORDER_DEPRECATED_IN_5_2
-#include "CoreMinimal.h"
-#endif
diff --git a/Engine/Plugins/PCG/Source/PCG/Public/Elements/PCGDataTableRowToParamData.h b/Engine/Plugins/PCG/Source/PCG/Public/Elements/PCGDataTableRowToParamData.h
--- a/Engine/Plugins/PCG/Source/PCG/Public/Elements/PCGDataTableRowToParamData.h
+++ b/Engine/Plugins/PCG/Source/PCG/Public/Elements/PCGDataTableRowToParamData.h
@@ -73,7 +73,3 @@ protected:
 	virtual bool PrepareDataInternal(FPCGContext* Context) const override;
 	virtual bool ExecuteInternal(FPCGContext* Context) const override;
 };
-
-#if UE_ENABLE_INCLUDE_ORDER_DEPRECATED_IN_5_2
-#include "CoreMinimal.h"
-#endif
diff --git a/Engine/Plugins/PCG/Source/PCG/Public/Elements/PCGDebugElement.h b/Engine/Plugins/PCG/Source/PCG/Public/Elements/PCGDebugElement.h
--- a/Engine/Plugins/PCG/Source/PCG/Public/Elements/PCGDebugElement.h
+++ b/Engine/Plugins/PCG/Source/PCG/Public/Elements/PCGDebugElement.h
@@ -43,7 +43,3 @@ public:
 protected:
 	virtual bool ExecuteInternal(FPCGContext* Context) const override;
 };
-
-#if UE_ENABLE_INCLUDE_ORDER_DEPRECATED_IN_5_2
-#include "CoreMinimal.h"
-#endif
diff --git a/Engine/Plugins/PCG/Source/PCG/Public/Elements/PCGDensityFilter.h b/Engine/Plugins/PCG/Source/PCG/Public/Elements/PCGDensityFilter.h
--- a/Engine/Plugins/PCG/Source/PCG/Public/Elements/PCGDensityFilter.h
+++ b/Engine/Plugins/PCG/Source/PCG/Public/Elements/PCGDensityFilter.h
@@ -48,8 +48,3 @@ protected:
 	virtual bool ExecuteInternal(FPCGContext* Context) const override;
 	virtual EPCGElementExecutionLoopMode ExecutionLoopMode(const UPCGSettings* Settings) const override { return EPCGElementExecutionLoopMode::SinglePrimaryPin; }
 };
-
-#if UE_ENABLE_INCLUDE_ORDER_DEPRECATED_IN_5_2
-#include "CoreMinimal.h"
-#include "PCGNode.h"
-#endif
diff --git a/Engine/Plugins/PCG/Source/PCG/Public/Elements/PCGDifferenceElement.h b/Engine/Plugins/PCG/Source/PCG/Public/Elements/PCGDifferenceElement.h
--- a/Engine/Plugins/PCG/Source/PCG/Public/Elements/PCGDifferenceElement.h
+++ b/Engine/Plugins/PCG/Source/PCG/Public/Elements/PCGDifferenceElement.h
@@ -62,7 +62,3 @@ protected:
 	virtual bool ExecuteInternal(FPCGContext* Context) const override;
 	virtual EPCGElementExecutionLoopMode ExecutionLoopMode(const UPCGSettings* Settings) const override { return EPCGElementExecutionLoopMode::SinglePrimaryPin; }
 };
-
-#if UE_ENABLE_INCLUDE_ORDER_DEPRECATED_IN_5_2
-#include "CoreMinimal.h"
-#endif
diff --git a/Engine/Plugins/PCG/Source/PCG/Public/Elements/PCGExecuteBlueprint.h b/Engine/Plugins/PCG/Source/PCG/Public/Elements/PCGExecuteBlueprint.h
--- a/Engine/Plugins/PCG/Source/PCG/Public/Elements/PCGExecuteBlueprint.h
+++ b/Engine/Plugins/PCG/Source/PCG/Public/Elements/PCGExecuteBlueprint.h
@@ -420,9 +420,3 @@ protected:
 	virtual void PostExecuteInternal(FPCGContext* Context) const override;
 	virtual FPCGContext* Initialize(const FPCGDataCollection& InputData, TWeakObjectPtr<UPCGComponent> SourceComponent, const UPCGNode* Node) override;	
 };
-
-#if UE_ENABLE_INCLUDE_ORDER_DEPRECATED_IN_5_2
-#include "Data/PCGPointData.h"
-#include "Math/RandomStream.h"
-#include "PCGPoint.h"
-#endif
diff --git a/Engine/Plugins/PCG/Source/PCG/Public/Elements/PCGGather.h b/Engine/Plugins/PCG/Source/PCG/Public/Elements/PCGGather.h
--- a/Engine/Plugins/PCG/Source/PCG/Public/Elements/PCGGather.h
+++ b/Engine/Plugins/PCG/Source/PCG/Public/Elements/PCGGather.h
@@ -42,8 +42,3 @@ namespace PCGGather
 	/** Gathers the input data into a single data collection and updates the tags */
 	FPCGDataCollection GatherDataForPin(const FPCGDataCollection& InputData, const FName InputLabel = PCGPinConstants::DefaultInputLabel, const FName OutputLabel = PCGPinConstants::DefaultOutputLabel);
 }
-
-
-#if UE_ENABLE_INCLUDE_ORDER_DEPRECATED_IN_5_2
-#include "CoreMinimal.h"
-#endif
diff --git a/Engine/Plugins/PCG/Source/PCG/Public/Elements/PCGInnerIntersectionElement.h b/Engine/Plugins/PCG/Source/PCG/Public/Elements/PCGInnerIntersectionElement.h
--- a/Engine/Plugins/PCG/Source/PCG/Public/Elements/PCGInnerIntersectionElement.h
+++ b/Engine/Plugins/PCG/Source/PCG/Public/Elements/PCGInnerIntersectionElement.h
@@ -46,7 +46,3 @@ class FPCGInnerIntersectionElement : public IPCGElement
 protected:
 	virtual bool ExecuteInternal(FPCGContext* Context) const override;
 };
-
-#if UE_ENABLE_INCLUDE_ORDER_DEPRECATED_IN_5_2
-#include "CoreMinimal.h"
-#endif
diff --git a/Engine/Plugins/PCG/Source/PCG/Public/Elements/PCGPointMatchAndSet.h b/Engine/Plugins/PCG/Source/PCG/Public/Elements/PCGPointMatchAndSet.h
--- a/Engine/Plugins/PCG/Source/PCG/Public/Elements/PCGPointMatchAndSet.h
+++ b/Engine/Plugins/PCG/Source/PCG/Public/Elements/PCGPointMatchAndSet.h
@@ -88,7 +88,3 @@ public:
 	virtual bool ExecuteInternal(FPCGContext* Context) const override;
 	virtual EPCGElementExecutionLoopMode ExecutionLoopMode(const UPCGSettings* Settings) const override { return EPCGElementExecutionLoopMode::SinglePrimaryPin; }
 };
-
-#if UE_ENABLE_INCLUDE_ORDER_DEPRECATED_IN_5_2
-#include "MatchAndSet/PCGMatchAndSetBase.h"
-#endif
diff --git a/Engine/Plugins/PCG/Source/PCG/Public/Elements/PCGProjectionElement.h b/Engine/Plugins/PCG/Source/PCG/Public/Elements/PCGProjectionElement.h
--- a/Engine/Plugins/PCG/Source/PCG/Public/Elements/PCGProjectionElement.h
+++ b/Engine/Plugins/PCG/Source/PCG/Public/Elements/PCGProjectionElement.h
@@ -60,8 +60,3 @@ class FPCGProjectionElement : public IPCGElement
 protected:
 	virtual bool ExecuteInternal(FPCGContext* Context) const override;
 };
-
-#if UE_ENABLE_INCLUDE_ORDER_DEPRECATED_IN_5_2
-#include "CoreMinimal.h"
-#include "Data/PCGProjectionData.h"
-#endif
diff --git a/Engine/Plugins/PCG/Source/PCG/Public/Elements/PCGProjectionParams.h b/Engine/Plugins/PCG/Source/PCG/Public/Elements/PCGProjectionParams.h
--- a/Engine/Plugins/PCG/Source/PCG/Public/Elements/PCGProjectionParams.h
+++ b/Engine/Plugins/PCG/Source/PCG/Public/Elements/PCGProjectionParams.h
@@ -91,7 +91,3 @@ struct PCG_API FPCGProjectionParams
 	bool bProjectColors_DEPRECATED = false;
 #endif // WITH_EDITORONLY_DATA
 };
-
-#if UE_ENABLE_INCLUDE_ORDER_DEPRECATED_IN_5_2
-#include "CoreMinimal.h"
-#endif
diff --git a/Engine/Plugins/PCG/Source/PCG/Public/Elements/PCGSanityCheckPointData.h b/Engine/Plugins/PCG/Source/PCG/Public/Elements/PCGSanityCheckPointData.h
--- a/Engine/Plugins/PCG/Source/PCG/Public/Elements/PCGSanityCheckPointData.h
+++ b/Engine/Plugins/PCG/Source/PCG/Public/Elements/PCGSanityCheckPointData.h
@@ -39,7 +39,3 @@ class FPCGSanityCheckPointDataElement : public IPCGElement
 protected:
 	virtual bool ExecuteInternal(FPCGContext* Context) const override;
 };
-
-#if UE_ENABLE_INCLUDE_ORDER_DEPRECATED_IN_5_2
-#include "CoreMinimal.h"
-#endif
diff --git a/Engine/Plugins/PCG/Source/PCG/Public/Elements/PCGSelectPoints.h b/Engine/Plugins/PCG/Source/PCG/Public/Elements/PCGSelectPoints.h
--- a/Engine/Plugins/PCG/Source/PCG/Public/Elements/PCGSelectPoints.h
+++ b/Engine/Plugins/PCG/Source/PCG/Public/Elements/PCGSelectPoints.h
@@ -43,8 +43,3 @@ protected:
 	virtual bool ExecuteInternal(FPCGContext* Context) const override;
 	virtual EPCGElementExecutionLoopMode ExecutionLoopMode(const UPCGSettings* Settings) const override { return EPCGElementExecutionLoopMode::SinglePrimaryPin; }
 };
-
-#if UE_ENABLE_INCLUDE_ORDER_DEPRECATED_IN_5_2
-#include "CoreMinimal.h"
-#include "PCGNode.h"
-#endif
diff --git a/Engine/Plugins/PCG/Source/PCG/Public/Elements/PCGSurfaceSampler.h b/Engine/Plugins/PCG/Source/PCG/Public/Elements/PCGSurfaceSampler.h
--- a/Engine/Plugins/PCG/Source/PCG/Public/Elements/PCGSurfaceSampler.h
+++ b/Engine/Plugins/PCG/Source/PCG/Public/Elements/PCGSurfaceSampler.h
@@ -165,7 +165,3 @@ protected:
 	virtual bool ExecuteInternal(FPCGContext* InContext) const override;
 	virtual EPCGElementExecutionLoopMode ExecutionLoopMode(const UPCGSettings* Settings) const override { return EPCGElementExecutionLoopMode::SinglePrimaryPin; }
 };
-
-#if UE_ENABLE_INCLUDE_ORDER_DEPRECATED_IN_5_2
-#include "CoreMinimal.h"
-#endif
diff --git a/Engine/Plugins/PCG/Source/PCG/Public/Elements/PCGTransformPoints.h b/Engine/Plugins/PCG/Source/PCG/Public/Elements/PCGTransformPoints.h
--- a/Engine/Plugins/PCG/Source/PCG/Public/Elements/PCGTransformPoints.h
+++ b/Engine/Plugins/PCG/Source/PCG/Public/Elements/PCGTransformPoints.h
@@ -78,8 +78,3 @@ protected:
 	virtual bool ExecuteInternal(FPCGContext* Context) const override;
 	virtual EPCGElementExecutionLoopMode ExecutionLoopMode(const UPCGSettings* Settings) const override { return EPCGElementExecutionLoopMode::SinglePrimaryPin; }
 };
-
-#if UE_ENABLE_INCLUDE_ORDER_DEPRECATED_IN_5_2
-#include "CoreMinimal.h"
-#include "Elements/PCGPointProcessingElementBase.h"
-#endif
diff --git a/Engine/Plugins/PCG/Source/PCG/Public/Elements/PCGUnionElement.h b/Engine/Plugins/PCG/Source/PCG/Public/Elements/PCGUnionElement.h
--- a/Engine/Plugins/PCG/Source/PCG/Public/Elements/PCGUnionElement.h
+++ b/Engine/Plugins/PCG/Source/PCG/Public/Elements/PCGUnionElement.h
@@ -40,7 +40,3 @@ class FPCGUnionElement : public IPCGElement
 protected:
 	virtual bool ExecuteInternal(FPCGContext* Context) const override;
 };
-
-#if UE_ENABLE_INCLUDE_ORDER_DEPRECATED_IN_5_2
-#include "CoreMinimal.h"
-#endif
diff --git a/Engine/Plugins/PCG/Source/PCG/Public/Helpers/PCGActorHelpers.h b/Engine/Plugins/PCG/Source/PCG/Public/Helpers/PCGActorHelpers.h
--- a/Engine/Plugins/PCG/Source/PCG/Public/Helpers/PCGActorHelpers.h
+++ b/Engine/Plugins/PCG/Source/PCG/Public/Helpers/PCGActorHelpers.h
@@ -189,12 +189,4 @@ PRAGMA_ENABLE_DEPRECATION_WARNINGS
 	 * Return the grid cell coordinates on the PCG partition grid given a position and the grid size.
 	 */
 	static FIntVector GetCellCoord(FVector InPosition, int InGridSize, bool bUse2DGrid);
-
-	UE_DEPRECATED(5.3, "Please use UE::BlueprintTools::GetActorClassDefaultComponents() instead")
-	static void GetActorClassDefaultComponents(const TSubclassOf<AActor>& ActorClass, TArray<UActorComponent*>& OutComponents, const TSubclassOf<UActorComponent>& InComponentClass = TSubclassOf<UActorComponent>());
 };
-
-#if UE_ENABLE_INCLUDE_ORDER_DEPRECATED_IN_5_2
-#include "CoreMinimal.h"
-#include "Engine/CollisionProfile.h"
-#endif
diff --git a/Engine/Plugins/PCG/Source/PCG/Public/Helpers/PCGAsync.h b/Engine/Plugins/PCG/Source/PCG/Public/Helpers/PCGAsync.h
--- a/Engine/Plugins/PCG/Source/PCG/Public/Helpers/PCGAsync.h
+++ b/Engine/Plugins/PCG/Source/PCG/Public/Helpers/PCGAsync.h
@@ -586,7 +586,3 @@ namespace FPCGAsync
 		return AsyncProcessingEx(AsyncState, NumIterations, Initialize, IterationInnerLoop, MoveData, Finished, bEnableTimeSlicing, ChunkSize, bAllowChunkSizeOverride);
 	}
 }
-
-#if UE_ENABLE_INCLUDE_ORDER_DEPRECATED_IN_5_2
-#include "CoreMinimal.h"
-#endif
diff --git a/Engine/Plugins/PCG/Source/PCG/Public/Helpers/PCGBlueprintHelpers.h b/Engine/Plugins/PCG/Source/PCG/Public/Helpers/PCGBlueprintHelpers.h
--- a/Engine/Plugins/PCG/Source/PCG/Public/Helpers/PCGBlueprintHelpers.h
+++ b/Engine/Plugins/PCG/Source/PCG/Public/Helpers/PCGBlueprintHelpers.h
@@ -98,10 +98,3 @@ public:
 	UFUNCTION(BlueprintCallable, Category="PCG|Data", meta = (ScriptMethod))
 	static UPCGData* DuplicateData(const UPCGData* InData, UPARAM(ref) FPCGContext& Context, bool bInitializeMetadata = true);
 };
-
-#if UE_ENABLE_INCLUDE_ORDER_DEPRECATED_IN_5_2
-#include "CoreMinimal.h"
-#include "Math/RandomStream.h"
-#include "PCGContext.h"
-#include "PCGPoint.h"
-#endif
diff --git a/Engine/Plugins/PCG/Source/PCG/Public/Helpers/PCGMetadataHelpers.h b/Engine/Plugins/PCG/Source/PCG/Public/Helpers/PCGMetadataHelpers.h
--- a/Engine/Plugins/PCG/Source/PCG/Public/Helpers/PCGMetadataHelpers.h
+++ b/Engine/Plugins/PCG/Source/PCG/Public/Helpers/PCGMetadataHelpers.h
@@ -59,7 +59,3 @@ namespace PCGMetadataHelpers
 	/** Copy all the attributes coming from Source Data into Target Data. */
 	PCG_API bool CopyAllAttributes(const UPCGData* SourceData, UPCGData* TargetData, FPCGContext* OptionalContext = nullptr);
 }
-
-#if UE_ENABLE_INCLUDE_ORDER_DEPRECATED_IN_5_2
-#include "CoreMinimal.h"
-#endif
diff --git a/Engine/Plugins/PCG/Source/PCG/Public/Helpers/PCGSettingsHelpers.h b/Engine/Plugins/PCG/Source/PCG/Public/Helpers/PCGSettingsHelpers.h
--- a/Engine/Plugins/PCG/Source/PCG/Public/Helpers/PCGSettingsHelpers.h
+++ b/Engine/Plugins/PCG/Source/PCG/Public/Helpers/PCGSettingsHelpers.h
@@ -21,137 +21,9 @@ struct FPCGDataCollection;
 
 namespace PCGSettingsHelpers
 {
-	/** Utility function to get the value of type T from a param data or a default value
-	* @param InName - Attribute to get from the param
-	* @param InValue - Default value to return if the param doesn't have the given attribute
-	* @param InParams - ParamData to get the value from.
-	* @param InKey - Metadata Entry Key to get the value from.
-	*/
-	template<typename T, typename TEnableIf<!TIsEnumClass<T>::Value>::Type* = nullptr>
-	UE_DEPRECATED(5.2, "GetValue is deprecated as overrides are now automatically applied on the settings. You should just query the settings value.")
-	T GetValue(const FName& InName, const T& InValue, const UPCGParamData* InParams, PCGMetadataEntryKey InKey)
-	{
-
-		if (InParams && InParams->Metadata)
-		{
-			const FPCGMetadataAttributeBase* MatchingAttribute = InParams->Metadata->GetConstAttribute(InName);
-
-			if (!MatchingAttribute)
-			{
-				return InValue;
-			}
-
-			auto GetTypedValue = [MatchingAttribute, &InValue, InKey](auto DummyValue) -> T
-			{
-				using AttributeType = decltype(DummyValue);
-
-				const FPCGMetadataAttribute<AttributeType>* ParamAttribute = static_cast<const FPCGMetadataAttribute<AttributeType>*>(MatchingAttribute);
-
-				if constexpr (std::is_same_v<T, AttributeType>)
-				{
-					return ParamAttribute->GetValueFromItemKey(InKey);
-				}
-				else if constexpr (std::is_constructible_v<T, AttributeType>)
-				{
-					UE_LOG(LogPCG, Verbose, TEXT("[GetAttributeValue] Matching attribute was found, but not the right type. Implicit conversion done (%d vs %d)"), MatchingAttribute->GetTypeId(), PCG::Private::MetadataTypes<T>::Id);
-					return T(ParamAttribute->GetValueFromItemKey(InKey));
-				}
-				else
-				{
-					UE_LOG(LogPCG, Error, TEXT("[GetAttributeValue] Matching attribute was found, but not the right type. %d vs %d"), MatchingAttribute->GetTypeId(), PCG::Private::MetadataTypes<T>::Id);
-					return InValue;
-				}
-			};
-
-			return PCGMetadataAttribute::CallbackWithRightType(MatchingAttribute->GetTypeId(), GetTypedValue);
-		}
-		else
-		{
-			return InValue;
-		}
-	}
-
-PRAGMA_DISABLE_DEPRECATION_WARNINGS
-
-	template<typename T, typename TEnableIf<!TIsEnumClass<T>::Value>::Type* = nullptr>
-	UE_DEPRECATED(5.2, "GetValue is deprecated as overrides are now automatically applied on the settings. You should just query the settings value.")
-	T GetValue(const FName& InName, const T& InValue, const UPCGParamData* InParams)
-	{
-		return GetValue(InName, InValue, InParams, 0);
-	}
-
-	template<typename T, typename TEnableIf<!TIsEnumClass<T>::Value>::Type* = nullptr>
-	UE_DEPRECATED(5.2, "GetValue is deprecated as overrides are now automatically applied on the settings. You should just query the settings value.")
-	T GetValue(const FName& InName, const T& InValue, const UPCGParamData* InParams, const FName& InParamName)
-	{
-		if (InParams && InParamName != NAME_None)
-		{
-			return GetValue(InName, InValue, InParams, InParams->FindMetadataKey(InParamName));
-		}
-		else
-		{
-			return InValue;
-		}
-	}
-
-	/** Specialized versions for enums */
-	template<typename T, typename TEnableIf<TIsEnumClass<T>::Value>::Type* = nullptr>
-	UE_DEPRECATED(5.2, "GetValue is deprecated as overrides are now automatically applied on the settings. You should just query the settings value.")
-	T GetValue(const FName& InName, const T& InValue, const UPCGParamData* InParams, PCGMetadataEntryKey InKey)
-	{
-		return static_cast<T>(GetValue(InName, static_cast<__underlying_type(T)>(InValue), InParams, InKey));
-	}
-
-	template<typename T, typename TEnableIf<TIsEnumClass<T>::Value>::Type* = nullptr>
-	UE_DEPRECATED(5.2, "GetValue is deprecated as overrides are now automatically applied on the settings. You should just query the settings value.")
-	T GetValue(const FName& InName, const T& InValue, const UPCGParamData* InParams)
-	{
-		return static_cast<T>(GetValue(InName, static_cast<__underlying_type(T)>(InValue), InParams));
-	}
-
-	template<typename T, typename TEnableIf<TIsEnumClass<T>::Value>::Type* = nullptr>
-	UE_DEPRECATED(5.2, "GetValue is deprecated as overrides are now automatically applied on the settings. You should just query the settings value.")
-	T GetValue(const FName& InName, const T& InValue, const UPCGParamData* InParams, const FName& InParamName)
-	{
-		return static_cast<T>(GetValue(InName, static_cast<__underlying_type(T)>(InValue), InParams, InParamName));
-	}
-
-	/** Specialized version for names */
-	template<>
-	UE_DEPRECATED(5.2, "GetValue is deprecated as overrides are now automatically applied on the settings. You should just query the settings value.")
-	FORCEINLINE FName GetValue(const FName& InName, const FName& InValue, const UPCGParamData* InParams, PCGMetadataEntryKey InKey)
-	{
-		return FName(GetValue(InName, InValue.ToString(), InParams, InKey));
-	}
-
-	template<>
-	UE_DEPRECATED(5.2, "GetValue is deprecated as overrides are now automatically applied on the settings. You should just query the settings value.")
-	FORCEINLINE FName GetValue(const FName& InName, const FName& InValue, const UPCGParamData* InParams)
-	{
-		return FName(GetValue(InName, InValue.ToString(), InParams));
-	}
-
-	template<>
-	UE_DEPRECATED(5.2, "GetValue is deprecated as overrides are now automatically applied on the settings. You should just query the settings value.")
-	FORCEINLINE FName GetValue(const FName& InName, const FName& InValue, const UPCGParamData* InParams, const FName& InParamName)
-	{
-		return FName(GetValue(InName, InValue.ToString(), InParams, InParamName));
-	}
-
 	/** Sets data from the params to a given property, matched on a name basis */
 	void SetValue(UPCGParamData* Params, UObject* Object, FProperty* Property);
 
-	UE_DEPRECATED(5.2, "ComputeSeedWithOverride is deprecated as overrides are now automatically applied on the settings. You should just query the settings value.")
-	int ComputeSeedWithOverride(const UPCGSettings* InSettings, const UPCGComponent* InComponent, UPCGParamData* InParams);
-
-	UE_DEPRECATED(5.2, "ComputeSeedWithOverride is deprecated as overrides are now automatically applied on the settings. You should just query the settings value.")
-	FORCEINLINE int ComputeSeedWithOverride(const UPCGSettings* InSettings, TWeakObjectPtr<UPCGComponent> InComponent, UPCGParamData* InParams)
-	{
-		return ComputeSeedWithOverride(InSettings, InComponent.Get(), InParams);
-	}
-
-PRAGMA_ENABLE_DEPRECATION_WARNINGS
-
 	/** Utility to call from before-node-update deprecation. A dedicated pin for params will be added when the pins are updated. Here we detect any params
 	*   connections to the In pin and disconnect them, and move the first params connection to a new params pin.
 	*/
diff --git a/Engine/Plugins/PCG/Source/PCG/Public/InstanceDataPackers/PCGInstanceDataPackerBase.h b/Engine/Plugins/PCG/Source/PCG/Public/InstanceDataPackers/PCGInstanceDataPackerBase.h
--- a/Engine/Plugins/PCG/Source/PCG/Public/InstanceDataPackers/PCGInstanceDataPackerBase.h
+++ b/Engine/Plugins/PCG/Source/PCG/Public/InstanceDataPackers/PCGInstanceDataPackerBase.h
@@ -53,9 +53,3 @@ public:
 	/** If OutNames is not null, returns a list of all attributes that will be packed. Returns true if this list can be statically determined (prior to execution). */
 	virtual bool GetAttributeNames(TArray<FName>* OutNames) { return false; }
 };
-
-#if UE_ENABLE_INCLUDE_ORDER_DEPRECATED_IN_5_2
-#include "CoreMinimal.h"
-#include "MeshSelectors/PCGMeshSelectorBase.h"
-#include "PCGElement.h"
-#endif
diff --git a/Engine/Plugins/PCG/Source/PCG/Public/MatchAndSet/PCGMatchAndSetBase.h b/Engine/Plugins/PCG/Source/PCG/Public/MatchAndSet/PCGMatchAndSetBase.h
--- a/Engine/Plugins/PCG/Source/PCG/Public/MatchAndSet/PCGMatchAndSetBase.h
+++ b/Engine/Plugins/PCG/Source/PCG/Public/MatchAndSet/PCGMatchAndSetBase.h
@@ -76,7 +76,3 @@ protected:
 	UPROPERTY()
 	EPCGMetadataTypesConstantStructStringMode StringMode_DEPRECATED;
 };
-
-#if UE_ENABLE_INCLUDE_ORDER_DEPRECATED_IN_5_2
-#include "CoreMinimal.h"
-#endif
diff --git a/Engine/Plugins/PCG/Source/PCG/Public/MatchAndSet/PCGMatchAndSetWeightedByCategory.h b/Engine/Plugins/PCG/Source/PCG/Public/MatchAndSet/PCGMatchAndSetWeightedByCategory.h
--- a/Engine/Plugins/PCG/Source/PCG/Public/MatchAndSet/PCGMatchAndSetWeightedByCategory.h
+++ b/Engine/Plugins/PCG/Source/PCG/Public/MatchAndSet/PCGMatchAndSetWeightedByCategory.h
@@ -83,7 +83,3 @@ public:
 	UPROPERTY(BlueprintReadWrite, EditAnywhere, Category = MatchAndSet)
 	bool bShouldMutateSeed = true;
 };
-
-#if UE_ENABLE_INCLUDE_ORDER_DEPRECATED_IN_5_2
-#include "MatchAndSet/PCGMatchAndSetWeighted.h"
-#endif
diff --git a/Engine/Plugins/PCG/Source/PCG/Public/MeshSelectors/PCGMeshSelectorBase.h b/Engine/Plugins/PCG/Source/PCG/Public/MeshSelectors/PCGMeshSelectorBase.h
--- a/Engine/Plugins/PCG/Source/PCG/Public/MeshSelectors/PCGMeshSelectorBase.h
+++ b/Engine/Plugins/PCG/Source/PCG/Public/MeshSelectors/PCGMeshSelectorBase.h
@@ -128,9 +128,3 @@ public:
 		TArray<FPCGMeshInstanceList>& OutMeshInstances,
 		UPCGPointData* OutPointData) const PURE_VIRTUAL(UPCGMeshSelectorBase::SelectInstances, return true;);
 };
-
-#if UE_ENABLE_INCLUDE_ORDER_DEPRECATED_IN_5_2
-#include "CoreMinimal.h"
-#include "Data/PCGPointData.h"
-#include "PCGElement.h"
-#endif
diff --git a/Engine/Plugins/PCG/Source/PCG/Public/Metadata/Accessors/IPCGAttributeAccessor.h b/Engine/Plugins/PCG/Source/PCG/Public/Metadata/Accessors/IPCGAttributeAccessor.h
--- a/Engine/Plugins/PCG/Source/PCG/Public/Metadata/Accessors/IPCGAttributeAccessor.h
+++ b/Engine/Plugins/PCG/Source/PCG/Public/Metadata/Accessors/IPCGAttributeAccessor.h
@@ -169,7 +169,3 @@ template <> bool PCG_API IPCGAttributeAccessor::GetRange<T>(TArrayView<T> OutVal
 template <> bool PCG_API IPCGAttributeAccessor::SetRange<T>(TArrayView<const T> InValues, int32 Index, IPCGAttributeAccessorKeys& Keys, EPCGAttributeAccessorFlags Flags);
 PCG_FOREACH_SUPPORTEDTYPES(IACCESSOR_DECL);
 #undef IACCESSOR_DECL
-
-#if UE_ENABLE_INCLUDE_ORDER_DEPRECATED_IN_5_2
-#include "Metadata/PCGMetadataCommon.h"
-#endif
diff --git a/Engine/Plugins/PCG/Source/PCG/Public/Metadata/Accessors/IPCGAttributeAccessorTpl.h b/Engine/Plugins/PCG/Source/PCG/Public/Metadata/Accessors/IPCGAttributeAccessorTpl.h
--- a/Engine/Plugins/PCG/Source/PCG/Public/Metadata/Accessors/IPCGAttributeAccessorTpl.h
+++ b/Engine/Plugins/PCG/Source/PCG/Public/Metadata/Accessors/IPCGAttributeAccessorTpl.h
@@ -146,7 +146,3 @@ private:
 
 	}
 };
-
-#if UE_ENABLE_INCLUDE_ORDER_DEPRECATED_IN_5_2
-#include "CoreMinimal.h"
-#endif
diff --git a/Engine/Plugins/PCG/Source/PCG/Public/Metadata/Accessors/PCGAttributeAccessorHelpers.h b/Engine/Plugins/PCG/Source/PCG/Public/Metadata/Accessors/PCGAttributeAccessorHelpers.h
--- a/Engine/Plugins/PCG/Source/PCG/Public/Metadata/Accessors/PCGAttributeAccessorHelpers.h
+++ b/Engine/Plugins/PCG/Source/PCG/Public/Metadata/Accessors/PCGAttributeAccessorHelpers.h
@@ -50,9 +50,6 @@ namespace PCGAttributeAccessorHelpers
 		bool bHasMultipleDataInAttributeSet = false;
 	};
 
-	UE_DEPRECATED(5.3, "Use the CreateConstAccessorForOverrideParamWithResult version")
-	PCG_API TUniquePtr<const IPCGAttributeAccessor> CreateConstAccessorForOverrideParam(const FPCGDataCollection& InInputData, const FPCGSettingsOverridableParam& InParam, FName* OutAttributeName = nullptr);
-
 	/**
 	* Create a const accessor depending on an overridable param
 	*/
diff --git a/Engine/Plugins/PCG/Source/PCG/Public/Metadata/Accessors/PCGAttributeAccessorKeys.h b/Engine/Plugins/PCG/Source/PCG/Public/Metadata/Accessors/PCGAttributeAccessorKeys.h
--- a/Engine/Plugins/PCG/Source/PCG/Public/Metadata/Accessors/PCGAttributeAccessorKeys.h
+++ b/Engine/Plugins/PCG/Source/PCG/Public/Metadata/Accessors/PCGAttributeAccessorKeys.h
@@ -416,7 +416,3 @@ inline bool IPCGAttributeAccessorKeys::GetKeys(int32 InStart, TArrayView<const O
 		return false;
 	}
 }
-
-#if UE_ENABLE_INCLUDE_ORDER_DEPRECATED_IN_5_2
-#include "PCGPoint.h"
-#endif
diff --git a/Engine/Plugins/PCG/Source/PCG/Public/Metadata/Accessors/PCGAttributeExtractor.h b/Engine/Plugins/PCG/Source/PCG/Public/Metadata/Accessors/PCGAttributeExtractor.h
--- a/Engine/Plugins/PCG/Source/PCG/Public/Metadata/Accessors/PCGAttributeExtractor.h
+++ b/Engine/Plugins/PCG/Source/PCG/Public/Metadata/Accessors/PCGAttributeExtractor.h
@@ -81,7 +81,3 @@ namespace PCGAttributeExtractor
 
 	TUniquePtr<IPCGAttributeAccessor> CreateTransformExtractor(TUniquePtr<IPCGAttributeAccessor> InAccessor, FName Name, bool& bOutSuccess);
 }
-
-#if UE_ENABLE_INCLUDE_ORDER_DEPRECATED_IN_5_2
-#include "UObject/UnrealType.h"
-#endif
diff --git a/Engine/Plugins/PCG/Source/PCG/Public/Metadata/PCGMetadata.h b/Engine/Plugins/PCG/Source/PCG/Public/Metadata/PCGMetadata.h
--- a/Engine/Plugins/PCG/Source/PCG/Public/Metadata/PCGMetadata.h
+++ b/Engine/Plugins/PCG/Source/PCG/Public/Metadata/PCGMetadata.h
@@ -473,8 +473,3 @@ const FPCGMetadataAttribute<T>* UPCGMetadata::GetConstTypedAttribute(FName Attri
 		? static_cast<const FPCGMetadataAttribute<T>*>(BaseAttribute)
 		: nullptr;
 }
-
-#if UE_ENABLE_INCLUDE_ORDER_DEPRECATED_IN_5_2
-#include "CoreMinimal.h"
-#include "PCGMetadataAccessor.h"
-#endif
diff --git a/Engine/Plugins/PCG/Source/PCG/Public/Metadata/PCGMetadataAttribute.h b/Engine/Plugins/PCG/Source/PCG/Public/Metadata/PCGMetadataAttribute.h
--- a/Engine/Plugins/PCG/Source/PCG/Public/Metadata/PCGMetadataAttribute.h
+++ b/Engine/Plugins/PCG/Source/PCG/Public/Metadata/PCGMetadataAttribute.h
@@ -111,8 +111,3 @@ public:
 	FName Name = NAME_None;
 	PCGMetadataAttributeKey AttributeId = -1;
 };
-
-#if UE_ENABLE_INCLUDE_ORDER_DEPRECATED_IN_5_2
-#include "CoreMinimal.h"
-#include "Misc/ScopeRWLock.h"
-#endif
diff --git a/Engine/Plugins/PCG/Source/PCG/Public/Metadata/PCGMetadataAttributeTpl.h b/Engine/Plugins/PCG/Source/PCG/Public/Metadata/PCGMetadataAttributeTpl.h
--- a/Engine/Plugins/PCG/Source/PCG/Public/Metadata/PCGMetadataAttributeTpl.h
+++ b/Engine/Plugins/PCG/Source/PCG/Public/Metadata/PCGMetadataAttributeTpl.h
@@ -1122,7 +1122,3 @@ FPCGMetadataAttributeBase* FPCGMetadataAttribute<T>::CopyToAnotherType(int16 Tar
 		}
 	});
 }
-
-#if UE_ENABLE_INCLUDE_ORDER_DEPRECATED_IN_5_2
-#include "CoreMinimal.h"
-#endif
diff --git a/Engine/Plugins/PCG/Source/PCG/Public/Metadata/PCGMetadataCommon.h b/Engine/Plugins/PCG/Source/PCG/Public/Metadata/PCGMetadataCommon.h
--- a/Engine/Plugins/PCG/Source/PCG/Public/Metadata/PCGMetadataCommon.h
+++ b/Engine/Plugins/PCG/Source/PCG/Public/Metadata/PCGMetadataCommon.h
@@ -42,7 +42,3 @@ enum class EPCGMetadataFilterMode : uint8
 	/** Only the listed attributes will be changed by the projection or added from the target data. */
 	IncludeAttributes,
 };
-
-#if UE_ENABLE_INCLUDE_ORDER_DEPRECATED_IN_5_2
-#include "CoreMinimal.h"
-#endif
diff --git a/Engine/Plugins/PCG/Source/PCG/Public/PCGComponent.h b/Engine/Plugins/PCG/Source/PCG/Public/PCGComponent.h
--- a/Engine/Plugins/PCG/Source/PCG/Public/PCGComponent.h
+++ b/Engine/Plugins/PCG/Source/PCG/Public/PCGComponent.h
@@ -715,8 +715,4 @@ protected:
 
 	UPROPERTY()
 	TObjectPtr<const UPCGComponent> SourceComponent;
-};
-
-#if UE_ENABLE_INCLUDE_ORDER_DEPRECATED_IN_5_2
-#include "CoreMinimal.h"
-#endif
+};
\ No newline at end of file
diff --git a/Engine/Plugins/PCG/Source/PCG/Public/PCGCustomVersion.h b/Engine/Plugins/PCG/Source/PCG/Public/PCGCustomVersion.h
--- a/Engine/Plugins/PCG/Source/PCG/Public/PCGCustomVersion.h
+++ b/Engine/Plugins/PCG/Source/PCG/Public/PCGCustomVersion.h
@@ -91,8 +91,3 @@ struct FPCGCustomVersion
 private:
 	FPCGCustomVersion() {}
 };
-
-#if UE_ENABLE_INCLUDE_ORDER_DEPRECATED_IN_5_2
-#include "CoreMinimal.h"
-#include "Misc/Guid.h"
-#endif
diff --git a/Engine/Plugins/PCG/Source/PCG/Public/PCGData.h b/Engine/Plugins/PCG/Source/PCG/Public/PCGData.h
--- a/Engine/Plugins/PCG/Source/PCG/Public/PCGData.h
+++ b/Engine/Plugins/PCG/Source/PCG/Public/PCGData.h
@@ -183,10 +183,6 @@ struct PCG_API FPCGDataCollection
 	template <typename PCGDataType>
 	TArray<FPCGTaggedData> GetTaggedTypedInputs(const FString& InTag) const;
 
-	UE_DEPRECATED(5.2, "GetParams is deprecated, please use GetParamsByPin or GetFirstParamsOnParamsPin.")
-	/** Returns the first params found in the collection */
-	UPCGParamData* GetParams() const;
-
 	// Only used as a temporary solution for old graph with nodes that didn't have params pins.
 	// Should NOT be used with new nodes.
 	UE_DEPRECATED(5.4, "Was not supposed to be used anyway, you should query the data per pin using GetParamsByPin")
@@ -312,8 +308,3 @@ public:
 protected:
 	static TArray<UPCGData*> GetInputsByPredicate(const FPCGDataCollection& InCollection, TArray<FPCGTaggedData>& OutTaggedData, TFunctionRef<bool(const FPCGTaggedData&)> InPredicate);
 };
-
-#if UE_ENABLE_INCLUDE_ORDER_DEPRECATED_IN_5_2
-#include "CoreMinimal.h"
-#include "PCGModule.h"
-#endif
diff --git a/Engine/Plugins/PCG/Source/PCG/Public/PCGDebug.h b/Engine/Plugins/PCG/Source/PCG/Public/PCGDebug.h
--- a/Engine/Plugins/PCG/Source/PCG/Public/PCGDebug.h
+++ b/Engine/Plugins/PCG/Source/PCG/Public/PCGDebug.h
@@ -44,7 +44,3 @@ public:
 	bool bDisplayProperties = true;
 #endif
 };
-
-#if UE_ENABLE_INCLUDE_ORDER_DEPRECATED_IN_5_2
-#include "CoreMinimal.h"
-#endif
diff --git a/Engine/Plugins/PCG/Source/PCG/Public/PCGEdge.h b/Engine/Plugins/PCG/Source/PCG/Public/PCGEdge.h
--- a/Engine/Plugins/PCG/Source/PCG/Public/PCGEdge.h
+++ b/Engine/Plugins/PCG/Source/PCG/Public/PCGEdge.h
@@ -43,7 +43,3 @@ public:
 	UPCGPin* GetOtherPin(const UPCGPin* Pin);
 	const UPCGPin* GetOtherPin(const UPCGPin* Pin) const;
 };
-
-#if UE_ENABLE_INCLUDE_ORDER_DEPRECATED_IN_5_2
-#include "CoreMinimal.h"
-#endif
diff --git a/Engine/Plugins/PCG/Source/PCG/Public/PCGElement.h b/Engine/Plugins/PCG/Source/PCG/Public/PCGElement.h
--- a/Engine/Plugins/PCG/Source/PCG/Public/PCGElement.h
+++ b/Engine/Plugins/PCG/Source/PCG/Public/PCGElement.h
@@ -198,10 +198,4 @@ public:
 };
 
 class UE_DEPRECATED(5.4, "This class has been deprecated. Please inherit from IPCGElement directly.") PCG_API FSimplePCGElement : public IPCGElement
-{};
-
-#if UE_ENABLE_INCLUDE_ORDER_DEPRECATED_IN_5_2
-#include "Misc/OutputDeviceRedirector.h"
-#include "PCGContext.h"
-#include "PCGData.h"
-#endif
+{};
\ No newline at end of file
diff --git a/Engine/Plugins/PCG/Source/PCG/Public/PCGGraph.h b/Engine/Plugins/PCG/Source/PCG/Public/PCGGraph.h
--- a/Engine/Plugins/PCG/Source/PCG/Public/PCGGraph.h
+++ b/Engine/Plugins/PCG/Source/PCG/Public/PCGGraph.h
@@ -598,7 +598,3 @@ UPCGNode* UPCGGraph::AddNodeOfType(T*& DefaultNodeSettings)
 	DefaultNodeSettings = Cast<T>(TempSettings);
 	return Node;
 }
-
-#if UE_ENABLE_INCLUDE_ORDER_DEPRECATED_IN_5_2
-#include "CoreMinimal.h"
-#endif
diff --git a/Engine/Plugins/PCG/Source/PCG/Public/PCGManagedResource.h b/Engine/Plugins/PCG/Source/PCG/Public/PCGManagedResource.h
--- a/Engine/Plugins/PCG/Source/PCG/Public/PCGManagedResource.h
+++ b/Engine/Plugins/PCG/Source/PCG/Public/PCGManagedResource.h
@@ -321,7 +321,3 @@ protected:
 	// Cached raw pointer to USplineMeshComponent
 	mutable USplineMeshComponent* CachedRawComponentPtr = nullptr;
 };
-
-#if UE_ENABLE_INCLUDE_ORDER_DEPRECATED_IN_5_2
-#include "CoreMinimal.h"
-#endif
diff --git a/Engine/Plugins/PCG/Source/PCG/Public/PCGModule.h b/Engine/Plugins/PCG/Source/PCG/Public/PCGModule.h
--- a/Engine/Plugins/PCG/Source/PCG/Public/PCGModule.h
+++ b/Engine/Plugins/PCG/Source/PCG/Public/PCGModule.h
@@ -59,9 +59,3 @@ private:
 	FPCGDataVisualizationRegistry PCGDataVisualizationRegistry;
 #endif
 };
-
-#if UE_ENABLE_INCLUDE_ORDER_DEPRECATED_IN_5_2
-#include "CoreMinimal.h"
-#include "Modules/ModuleManager.h"
-#include "Stats/Stats.h"
-#endif
diff --git a/Engine/Plugins/PCG/Source/PCG/Public/PCGNode.h b/Engine/Plugins/PCG/Source/PCG/Public/PCGNode.h
--- a/Engine/Plugins/PCG/Source/PCG/Public/PCGNode.h
+++ b/Engine/Plugins/PCG/Source/PCG/Public/PCGNode.h
@@ -244,10 +244,4 @@ protected:
 	// - Ability to be multithreaded (here or element)
 	// - Generates artifacts (here or element)
 	// - Priority
-};
-
-#if UE_ENABLE_INCLUDE_ORDER_DEPRECATED_IN_5_2
-#include "CoreMinimal.h"
-#include "PCGCommon.h"
-#include "PCGPin.h"
-#endif
+};
\ No newline at end of file
diff --git a/Engine/Plugins/PCG/Source/PCG/Public/PCGPin.h b/Engine/Plugins/PCG/Source/PCG/Public/PCGPin.h
--- a/Engine/Plugins/PCG/Source/PCG/Public/PCGPin.h
+++ b/Engine/Plugins/PCG/Source/PCG/Public/PCGPin.h
@@ -333,7 +333,3 @@ public:
 	UFUNCTION(BlueprintCallable, Category = "PCG|PinProperties", meta = (ScriptMethod))
 	static void SetNormalPin(UPARAM(ref) FPCGPinProperties& PinProperties) { PinProperties.SetNormalPin(); }
 };
-
-#if UE_ENABLE_INCLUDE_ORDER_DEPRECATED_IN_5_2
-#include "CoreMinimal.h"
-#endif
diff --git a/Engine/Plugins/PCG/Source/PCG/Public/PCGSubsystem.h b/Engine/Plugins/PCG/Source/PCG/Public/PCGSubsystem.h
--- a/Engine/Plugins/PCG/Source/PCG/Public/PCGSubsystem.h
+++ b/Engine/Plugins/PCG/Source/PCG/Public/PCGSubsystem.h
@@ -416,8 +416,3 @@ private:
 	}
 #endif
 };
-
-#if UE_ENABLE_INCLUDE_ORDER_DEPRECATED_IN_5_2
-#include "PCGComponent.h"
-#include "PCGVolume.h"
-#endif
diff --git a/Engine/Plugins/PCG/Source/PCG/Public/PCGVolume.h b/Engine/Plugins/PCG/Source/PCG/Public/PCGVolume.h
--- a/Engine/Plugins/PCG/Source/PCG/Public/PCGVolume.h
+++ b/Engine/Plugins/PCG/Source/PCG/Public/PCGVolume.h
@@ -26,7 +26,3 @@ public:
 	UPROPERTY(BlueprintReadWrite, EditAnywhere, Category = PCG)
 	TObjectPtr<UPCGComponent> PCGComponent;
 };
-
-#if UE_ENABLE_INCLUDE_ORDER_DEPRECATED_IN_5_2
-#include "CoreMinimal.h"
-#endif
diff --git a/Engine/Plugins/PCG/Source/PCG/Public/PCGWorldActor.h b/Engine/Plugins/PCG/Source/PCG/Public/PCGWorldActor.h
--- a/Engine/Plugins/PCG/Source/PCG/Public/PCGWorldActor.h
+++ b/Engine/Plugins/PCG/Source/PCG/Public/PCGWorldActor.h
@@ -80,8 +80,3 @@ private:
 	void OnPartitionGridSizeChanged();
 #endif
 };
-
-#if UE_ENABLE_INCLUDE_ORDER_DEPRECATED_IN_5_2
-#include "CoreMinimal.h"
-#include "Grid/PCGLandscapeCache.h"
-#endif
diff --git a/Engine/Plugins/PCG/Source/PCG/Public/Tests/Determinism/PCGDeterminismSettings.h b/Engine/Plugins/PCG/Source/PCG/Public/Tests/Determinism/PCGDeterminismSettings.h
--- a/Engine/Plugins/PCG/Source/PCG/Public/Tests/Determinism/PCGDeterminismSettings.h
+++ b/Engine/Plugins/PCG/Source/PCG/Public/Tests/Determinism/PCGDeterminismSettings.h
@@ -22,7 +22,3 @@ struct FPCGDeterminismSettings
 	UPROPERTY(BlueprintReadOnly, EditAnywhere, Category = Determinism, meta = (EditCondition = "bUseBlueprintDeterminismTest"))
 	TSubclassOf<UPCGDeterminismTestBlueprintBase> DeterminismTestBlueprint;
 };
-
-#if UE_ENABLE_INCLUDE_ORDER_DEPRECATED_IN_5_2
-#include "CoreMinimal.h"
-#endif
diff --git a/Engine/Plugins/PCG/Source/PCG/Public/Tests/Determinism/PCGDeterminismTestBlueprintBase.h b/Engine/Plugins/PCG/Source/PCG/Public/Tests/Determinism/PCGDeterminismTestBlueprintBase.h
--- a/Engine/Plugins/PCG/Source/PCG/Public/Tests/Determinism/PCGDeterminismTestBlueprintBase.h
+++ b/Engine/Plugins/PCG/Source/PCG/Public/Tests/Determinism/PCGDeterminismTestBlueprintBase.h
@@ -18,7 +18,3 @@ public:
 	UFUNCTION(BlueprintNativeEvent, BlueprintPure, CallInEditor, Category = Determinism)
 	void ExecuteTest(const UPCGNode* InPCGNode, UPARAM(ref)FDeterminismTestResult& InOutTestResult);
 };
-
-#if UE_ENABLE_INCLUDE_ORDER_DEPRECATED_IN_5_2
-#include "Tests/Determinism/PCGDeterminismTestsCommon.h"
-#endif
diff --git a/Engine/Plugins/PCG/Source/PCG/Public/Utils/PCGExtraCapture.h b/Engine/Plugins/PCG/Source/PCG/Public/Utils/PCGExtraCapture.h
--- a/Engine/Plugins/PCG/Source/PCG/Public/Utils/PCGExtraCapture.h
+++ b/Engine/Plugins/PCG/Source/PCG/Public/Utils/PCGExtraCapture.h
@@ -99,8 +99,3 @@ namespace PCGUtils
 	};
 #endif // WITH_EDITOR
 }
-
-#if UE_ENABLE_INCLUDE_ORDER_DEPRECATED_IN_5_2
-#include "PCGContext.h"
-#include "UObject/UnrealType.h"
-#endif
diff --git a/Engine/Plugins/PCG/Source/PCGEditor/Public/Details/PCGAttributePropertySelectorDetails.h b/Engine/Plugins/PCG/Source/PCGEditor/Public/Details/PCGAttributePropertySelectorDetails.h
--- a/Engine/Plugins/PCG/Source/PCGEditor/Public/Details/PCGAttributePropertySelectorDetails.h
+++ b/Engine/Plugins/PCG/Source/PCGEditor/Public/Details/PCGAttributePropertySelectorDetails.h
@@ -43,7 +43,3 @@ protected:
 
 	TSharedPtr<IPropertyHandle> PropertyHandle;
 };
-
-#if UE_ENABLE_INCLUDE_ORDER_DEPRECATED_IN_5_2
-#include "PCGPoint.h"
-#endif
diff --git a/Engine/Plugins/PCG/Source/PCGEditor/Public/Details/PCGBlueprintSettingsDetails.h b/Engine/Plugins/PCG/Source/PCGEditor/Public/Details/PCGBlueprintSettingsDetails.h
--- a/Engine/Plugins/PCG/Source/PCGEditor/Public/Details/PCGBlueprintSettingsDetails.h
+++ b/Engine/Plugins/PCG/Source/PCGEditor/Public/Details/PCGBlueprintSettingsDetails.h
@@ -31,10 +31,3 @@ private:
 private:
 	TWeakObjectPtr<UPCGBlueprintElement> SelectedObject;
 };
-
-#if UE_ENABLE_INCLUDE_ORDER_DEPRECATED_IN_5_2
-#include "Elements/PCGExecuteBlueprint.h"
-#include "IPropertyTypeCustomization.h"
-#include "Input/Reply.h"
-#include "PCGSettings.h"
-#endif
diff --git a/Engine/Plugins/PCG/Source/PCGEditor/Public/Details/PCGComponentDetails.h b/Engine/Plugins/PCG/Source/PCGEditor/Public/Details/PCGComponentDetails.h
--- a/Engine/Plugins/PCG/Source/PCGEditor/Public/Details/PCGComponentDetails.h
+++ b/Engine/Plugins/PCG/Source/PCGEditor/Public/Details/PCGComponentDetails.h
@@ -39,9 +39,3 @@ private:
 protected:
 	TArray<TWeakObjectPtr<UPCGComponent>> SelectedComponents;
 };
-
-#if UE_ENABLE_INCLUDE_ORDER_DEPRECATED_IN_5_2
-#include "Input/Reply.h"
-#include "Layout/Visibility.h"
-#include "UObject/WeakObjectPtr.h"
-#endif
diff --git a/Engine/Plugins/PCG/Source/PCGEditor/Public/Details/PCGGraphDetails.h b/Engine/Plugins/PCG/Source/PCGEditor/Public/Details/PCGGraphDetails.h
--- a/Engine/Plugins/PCG/Source/PCGEditor/Public/Details/PCGGraphDetails.h
+++ b/Engine/Plugins/PCG/Source/PCGEditor/Public/Details/PCGGraphDetails.h
@@ -17,8 +17,3 @@ public:
 private:
 	TArray<TWeakObjectPtr<UPCGGraph>> SelectedGraphs;
 };
-
-#if UE_ENABLE_INCLUDE_ORDER_DEPRECATED_IN_5_2
-#include "Input/Reply.h"
-#include "UObject/WeakObjectPtr.h"
-#endif
diff --git a/Engine/Plugins/PCG/Source/PCGEditor/Public/PCGEditorSettings.h b/Engine/Plugins/PCG/Source/PCGEditor/Public/PCGEditorSettings.h
--- a/Engine/Plugins/PCG/Source/PCGEditor/Public/PCGEditorSettings.h
+++ b/Engine/Plugins/PCG/Source/PCGEditor/Public/PCGEditorSettings.h
@@ -224,8 +224,3 @@ public:
 	UPROPERTY(EditAnywhere, config, Category = "Editor Performance", meta = (MultiLine = true, DisplayName = "Overrides default tooltip on the pause button", EditCondition = "bShowPauseButton", EditConditionHides))
 	FString OverridePausedButtonTooltip = FString();
 };
-
-#if UE_ENABLE_INCLUDE_ORDER_DEPRECATED_IN_5_2
-#include "CoreMinimal.h"
-#include "PCGSettings.h"
-#endif
