[PCG] Fix ISM Reuse when using Override Settings, the UID would change on each execution because Settings are duplicated and UID is transient (assigned in Constructor)

#rnx

diff --git a/Engine/Plugins/PCG/Source/PCG/Private/Elements/PCGSpawnActor.cpp b/Engine/Plugins/PCG/Source/PCG/Private/Elements/PCGSpawnActor.cpp
--- a/Engine/Plugins/PCG/Source/PCG/Private/Elements/PCGSpawnActor.cpp
+++ b/Engine/Plugins/PCG/Source/PCG/Private/Elements/PCGSpawnActor.cpp
@@ -709,7 +709,7 @@ void FPCGSpawnActorElement::CollapseIntoTargetActor(FPCGSubgraphContext* Context
 	{
 		const FPCGISMComponentBuilderParams& ISMCParams = ISMCBuilderTransforms.Key;
 
-		UPCGManagedISMComponent* MISMC = UPCGActorHelpers::GetOrCreateManagedISMC(TargetActor, Context->SourceComponent.Get(), Settings->UID, ISMCParams);
+		UPCGManagedISMComponent* MISMC = UPCGActorHelpers::GetOrCreateManagedISMC(TargetActor, Context->SourceComponent.Get(), Settings->GetStableUID(), ISMCParams);
 		if (!MISMC)
 		{
 			continue;
diff --git a/Engine/Plugins/PCG/Source/PCG/Private/Elements/PCGSpawnSplineMesh.cpp b/Engine/Plugins/PCG/Source/PCG/Private/Elements/PCGSpawnSplineMesh.cpp
--- a/Engine/Plugins/PCG/Source/PCG/Private/Elements/PCGSpawnSplineMesh.cpp
+++ b/Engine/Plugins/PCG/Source/PCG/Private/Elements/PCGSpawnSplineMesh.cpp
@@ -213,7 +213,7 @@ bool FPCGSpawnSplineMeshElement::ExecuteInternal(FPCGContext* Context) const
 			SMCBuilderParams.SplineMeshParams.StartScale = FVector2D(StartScale.Y, StartScale.Z);
 			SMCBuilderParams.SplineMeshParams.EndScale = FVector2D(EndScale.Y, EndScale.Z);
 
-			USplineMeshComponent* SplineMeshComponent = UPCGActorHelpers::GetOrCreateSplineMeshComponent(TargetActor, Context->SourceComponent.Get(), Settings->UID, SMCBuilderParams);
+			USplineMeshComponent* SplineMeshComponent = UPCGActorHelpers::GetOrCreateSplineMeshComponent(TargetActor, Context->SourceComponent.Get(), Settings->GetStableUID(), SMCBuilderParams);
 
 			// TODO: Write out the geometry to a dynamic mesh type.
 			//SplineMeshComponent->BodySetup->TriMeshGeometries
diff --git a/Engine/Plugins/PCG/Source/PCG/Private/Elements/PCGStaticMeshSpawner.cpp b/Engine/Plugins/PCG/Source/PCG/Private/Elements/PCGStaticMeshSpawner.cpp
--- a/Engine/Plugins/PCG/Source/PCG/Private/Elements/PCGStaticMeshSpawner.cpp
+++ b/Engine/Plugins/PCG/Source/PCG/Private/Elements/PCGStaticMeshSpawner.cpp
@@ -340,7 +340,7 @@ bool FPCGStaticMeshSpawnerElement::PrepareDataInternal(FPCGContext* InContext) c
 			{
 				if (UPCGManagedISMComponent* Resource = Cast<UPCGManagedISMComponent>(InResource))
 				{
-					if (Resource->GetSettingsUID() == Settings->UID && Resource->GetCrc().IsValid() && Resource->GetCrc() == Context->DependenciesCrc)
+					if (Resource->GetSettingsUID() == Settings->GetStableUID() && Resource->GetCrc().IsValid() && Resource->GetCrc() == Context->DependenciesCrc)
 					{
 						MISMCs.Add(Resource);
 					}
@@ -670,7 +670,7 @@ void FPCGStaticMeshSpawnerElement::SpawnStaticMeshInstances(FPCGStaticMeshSpawne
 
 	Params.bAllowDescriptorChanges = Settings->bAllowDescriptorChanges;
 
-	UPCGManagedISMComponent* MISMC = UPCGActorHelpers::GetOrCreateManagedISMC(TargetActor, Context->SourceComponent.Get(), Settings->UID, Params);
+	UPCGManagedISMComponent* MISMC = UPCGActorHelpers::GetOrCreateManagedISMC(TargetActor, Context->SourceComponent.Get(), Settings->GetStableUID(), Params);
 
 	check(MISMC);
 	MISMC->SetCrc(Context->DependenciesCrc);
diff --git a/Engine/Plugins/PCG/Source/PCG/Private/Helpers/PCGActorHelpers.cpp b/Engine/Plugins/PCG/Source/PCG/Private/Helpers/PCGActorHelpers.cpp
--- a/Engine/Plugins/PCG/Source/PCG/Private/Helpers/PCGActorHelpers.cpp
+++ b/Engine/Plugins/PCG/Source/PCG/Private/Helpers/PCGActorHelpers.cpp
@@ -82,15 +82,12 @@ UPCGManagedISMComponent* UPCGActorHelpers::GetOrCreateManagedISMC(AActor* InTarg
 
 	auto AddTagsToComponent = [InSourceComponent, &InParams](UInstancedStaticMeshComponent* ISMC)
 	{
-		if(ISMC)
-		{
-			ISMC->ComponentTags.AddUnique(PCGHelpers::DefaultPCGTag);
-			ISMC->ComponentTags.AddUnique(InSourceComponent->GetFName());
+		ISMC->ComponentTags.AddUnique(PCGHelpers::DefaultPCGTag);
+		ISMC->ComponentTags.AddUnique(InSourceComponent->GetFName());
 			
-			for (FName ComponentTag : InParams.Descriptor.ComponentTags)
-			{
-				ISMC->ComponentTags.AddUnique(ComponentTag);
-			}
+		for (FName ComponentTag : InParams.Descriptor.ComponentTags)
+		{
+			ISMC->ComponentTags.AddUnique(ComponentTag);
 		}
 	};
 
@@ -156,7 +153,13 @@ UPCGManagedISMComponent* UPCGActorHelpers::GetOrCreateManagedISMC(AActor* InTarg
 		{
 			TRACE_CPUPROFILER_EVENT_SCOPE(UPCGActorHelpers::GetOrCreateManagedISMC::MarkAsUsed);
 			MatchingResource->MarkAsUsed();
-			AddTagsToComponent(Cast<UInstancedStaticMeshComponent>(MatchingResource->GeneratedComponent.Get()));
+
+			UInstancedStaticMeshComponent* ISMC = Cast<UInstancedStaticMeshComponent>(MatchingResource->GeneratedComponent.Get());
+			if (ensure(ISMC))
+			{
+				ISMC->Modify(!InSourceComponent->IsInPreviewMode());
+				AddTagsToComponent(ISMC);
+			}
 
 			return MatchingResource;
 		}
@@ -268,6 +271,12 @@ UPCGManagedSplineMeshComponent* UPCGActorHelpers::GetOrCreateManagedSplineMeshCo
 			TRACE_CPUPROFILER_EVENT_SCOPE(UPCGActorHelpers::GetOrCreateManagedSplineMeshComponent::MarkAsUsed);
 			MatchingResource->MarkAsUsed();
 
+			USplineMeshComponent* SplineMeshComponent = MatchingResource->GetComponent();
+			if (ensure(SplineMeshComponent))
+			{
+				SplineMeshComponent->Modify(!InSourceComponent->IsInPreviewMode());
+			}
+
 			return MatchingResource;
 		}
 	}
diff --git a/Engine/Plugins/PCG/Source/PCG/Public/PCGSettings.h b/Engine/Plugins/PCG/Source/PCG/Public/PCGSettings.h
--- a/Engine/Plugins/PCG/Source/PCG/Public/PCGSettings.h
+++ b/Engine/Plugins/PCG/Source/PCG/Public/PCGSettings.h
@@ -408,6 +408,9 @@ public:
 	FText Description;
 #endif
 
+	// Returns Original UID when this is a duplicated settings so we can compare successive executions for reuse cases
+	uint64 GetStableUID() const { return OriginalSettings ? OriginalSettings->UID : UID; }
+
 	// Holds the original settings used to duplicate this object if it was overridden
 	const UPCGSettings* OriginalSettings = nullptr;
 
diff --git a/Engine/Plugins/PCG/Source/PCGEditor/Private/DataVisualizations/PCGSpatialDataVisualization.cpp b/Engine/Plugins/PCG/Source/PCGEditor/Private/DataVisualizations/PCGSpatialDataVisualization.cpp
--- a/Engine/Plugins/PCG/Source/PCGEditor/Private/DataVisualizations/PCGSpatialDataVisualization.cpp
+++ b/Engine/Plugins/PCG/Source/PCGEditor/Private/DataVisualizations/PCGSpatialDataVisualization.cpp
@@ -195,7 +195,7 @@ void IPCGSpatialDataVisualization::ExecuteDebugDisplay(FPCGContext* Context, con
 			continue;
 		}
 
-		UInstancedStaticMeshComponent* ISMC = UPCGActorHelpers::GetOrCreateISMC(TargetActor, Context->SourceComponent.Get(), SettingsInterface->GetSettings()->UID, Params[Direction]);
+		UInstancedStaticMeshComponent* ISMC = UPCGActorHelpers::GetOrCreateISMC(TargetActor, Context->SourceComponent.Get(), SettingsInterface->GetSettings()->GetStableUID(), Params[Direction]);
 		check(ISMC && ISMC->NumCustomDataFloats == NumCustomData);
 
 		ISMC->ComponentTags.AddUnique(PCGHelpers::DefaultPCGDebugTag);
