[PCG] Fix GPU nodes showing up as "cached" in profiling window
* Also fix disabled nodes showing up as "cached" in profiling window

#rnx

diff --git a/Engine/Plugins/PCG/Source/PCGEditor/Private/SPCGEditorGraphProfilingView.cpp b/Engine/Plugins/PCG/Source/PCGEditor/Private/SPCGEditorGraphProfilingView.cpp
--- a/Engine/Plugins/PCG/Source/PCGEditor/Private/SPCGEditorGraphProfilingView.cpp
+++ b/Engine/Plugins/PCG/Source/PCGEditor/Private/SPCGEditorGraphProfilingView.cpp
@@ -82,6 +82,20 @@ FText FPCGProfilingListViewItem::GetTextForColumn(FName ColumnId, bool bNoGroupi
 
 	if (ColumnId == PCGEditorGraphProfilingView::NAME_Node)
 	{
+		const UPCGSettings* PCGSettings = PCGNode ? PCGNode->GetSettings() : nullptr;
+
+		if (PCGSettings)
+		{
+			if (PCGSettings->ShouldExecuteOnGPU())
+			{
+				return FText::FromString(Name + TEXT(" (profiling GPU nodes not supported)"));
+			}
+			else if (!PCGSettings->bEnabled)
+			{
+				return FText::FromString(Name + TEXT(" (disabled)"));
+			}
+		}
+
 		if (bHasData)
 		{
 			return FText::FromString(Name);
