[PCG] Expose the User Parameters struct for modifying it and more get/set functions

* Can now get/set structs and objects
* Can now modify arrays and sets
* Needed to expose a new function for mutable sets
* Update tests

#rnx

diff --git a/Engine/Plugins/PCG/Source/PCG/Private/PCGGraph.cpp b/Engine/Plugins/PCG/Source/PCG/Private/PCGGraph.cpp
--- a/Engine/Plugins/PCG/Source/PCG/Private/PCGGraph.cpp
+++ b/Engine/Plugins/PCG/Source/PCG/Private/PCGGraph.cpp
@@ -133,6 +133,42 @@ EPropertyBagResult UPCGGraphInterface::SetGraphParameter(const FName PropertyNam
 	return Result;
 }
 
+bool UPCGGraphInterface::UpdateArrayGraphParameter(const FName PropertyName, TFunctionRef<bool(FPropertyBagArrayRef& PropertyBagArrayRef)> Callback)
+{
+	FInstancedPropertyBag* UserParameters = GetMutableUserParametersStruct();
+	check(UserParameters);
+
+	TValueOrError<FPropertyBagArrayRef, EPropertyBagResult> Result = UserParameters->GetMutableArrayRef(PropertyName);
+
+	if (!Result.HasError() && Result.HasValue() && Callback(Result.GetValue()))
+	{
+		OnGraphParametersChanged(EPCGGraphParameterEvent::ValueModifiedLocally, PropertyName);
+		return true;
+	}
+	else
+	{
+		return false;
+	}
+}
+
+bool UPCGGraphInterface::UpdateSetGraphParameter(const FName PropertyName, TFunctionRef<bool(FPropertyBagSetRef& PropertyBagSetRef)> Callback)
+{
+	FInstancedPropertyBag* UserParameters = GetMutableUserParametersStruct();
+	check(UserParameters);
+
+	TValueOrError<FPropertyBagSetRef, EPropertyBagResult> Result = UserParameters->GetMutableSetRef(PropertyName);
+
+	if (!Result.HasError() && Result.HasValue() && Callback(Result.GetValue()))
+	{
+		OnGraphParametersChanged(EPCGGraphParameterEvent::ValueModifiedLocally, PropertyName);
+		return true;
+	}
+	else
+	{
+		return false;
+	}
+}
+
 bool UPCGGraphInterface::IsInstance() const
 {
 	return this != GetGraph();
@@ -1643,6 +1679,13 @@ void UPCGGraph::OnGraphParametersChanged(EPCGGraphParameterEvent InChangeType, F
 #endif // WITH_EDITOR
 }
 
+void UPCGGraph::UpdateUserParametersStruct(TFunctionRef<void(FInstancedPropertyBag&)> Callback)
+{
+	Callback(UserParameters);
+	// Since anything could have changed, trigger a refresh like a post load (to compare what changed)
+	OnGraphParametersChanged(EPCGGraphParameterEvent::GraphPostLoad, NAME_None);
+}
+
 FInstancedPropertyBag* UPCGGraph::GetMutableUserParametersStruct()
 {
 	return &UserParameters;
diff --git a/Engine/Plugins/PCG/Source/PCG/Private/Tests/Elements/PCGGetActorPropertyTest.cpp b/Engine/Plugins/PCG/Source/PCG/Private/Tests/Elements/PCGGetActorPropertyTest.cpp
--- a/Engine/Plugins/PCG/Source/PCG/Private/Tests/Elements/PCGGetActorPropertyTest.cpp
+++ b/Engine/Plugins/PCG/Source/PCG/Private/Tests/Elements/PCGGetActorPropertyTest.cpp
@@ -247,6 +247,7 @@ bool FPCGPropertyToParamDataPropertyTypeTest::RunTest(const FString& Parameters)
 	Actor->ArrayOfVectorsProperty = { VectorValue, SecondVectorValue };
 	Actor->ArrayOfStructsProperty = { PCGColorValue, SecondPCGColorValue };
 	Actor->ArrayOfObjectsProperty = { ObjectValue, SecondObjectValue };
+	Actor->SetOfIntsProperty = {1, 2, 3};
 	Actor->DummyStruct.FloatProperty = 1.2f;
 	Actor->DummyStruct.IntArrayProperty = { 5, 6, 7 };
 	Actor->DummyStruct.Level2Struct.DoubleArrayProperty = { 0.1, 0.2, 0.3 };
@@ -296,6 +297,7 @@ bool FPCGPropertyToParamDataPropertyTypeTest::RunTest(const FString& Parameters)
 	bSuccess &= VerifyAttributeValuesValid(this, TestData, GET_MEMBER_NAME_CHECKED(APCGUnitTestDummyActor, ArrayOfIntsProperty), { GET_MEMBER_NAME_CHECKED(APCGUnitTestDummyActor, ArrayOfIntsProperty) }, ExtraTestWhat, 42, 43, 44);
 	bSuccess &= VerifyAttributeValuesValid(this, TestData, GET_MEMBER_NAME_CHECKED(APCGUnitTestDummyActor, ArrayOfVectorsProperty), { GET_MEMBER_NAME_CHECKED(APCGUnitTestDummyActor, ArrayOfVectorsProperty) }, ExtraTestWhat, VectorValue, SecondVectorValue);
 	bSuccess &= VerifyAttributeValuesValid(this, TestData, GET_MEMBER_NAME_CHECKED(APCGUnitTestDummyActor, ArrayOfObjectsProperty), { GET_MEMBER_NAME_CHECKED(APCGUnitTestDummyActor, ArrayOfObjectsProperty) }, ExtraTestWhat, FSoftObjectPath(ObjectValue), FSoftObjectPath(SecondObjectValue));
+	bSuccess &= VerifyAttributeValuesValid(this, TestData, GET_MEMBER_NAME_CHECKED(APCGUnitTestDummyActor, SetOfIntsProperty), { GET_MEMBER_NAME_CHECKED(APCGUnitTestDummyActor, SetOfIntsProperty) }, ExtraTestWhat, 1, 2, 3);
 
 	// Arrays of extracted properties
 	bSuccess &= VerifyAttributeValuesValid(this, TestData, GET_MEMBER_NAME_CHECKED(APCGUnitTestDummyActor, ArrayOfStructsProperty), ColorPropertyNames, ExtraTestWhat, 
diff --git a/Engine/Plugins/PCG/Source/PCG/Private/Tests/Elements/PCGGetActorPropertyTest.h b/Engine/Plugins/PCG/Source/PCG/Private/Tests/Elements/PCGGetActorPropertyTest.h
--- a/Engine/Plugins/PCG/Source/PCG/Private/Tests/Elements/PCGGetActorPropertyTest.h
+++ b/Engine/Plugins/PCG/Source/PCG/Private/Tests/Elements/PCGGetActorPropertyTest.h
@@ -152,6 +152,9 @@ public:
 
 	UPROPERTY(BlueprintReadOnly, Category = "")
 	FPCGDummyGetPropertyStruct DummyStruct;
+	
+	UPROPERTY(BlueprintReadOnly, Category = "")
+    TArray<int32> SetOfIntsProperty;
 };
 
 UCLASS(MinimalAPI, NotBlueprintable, NotPlaceable, NotBlueprintType, Transient, HideDropdown, meta = (Hidden))
diff --git a/Engine/Plugins/PCG/Source/PCG/Private/Tests/PCGGraphParametersTest.cpp b/Engine/Plugins/PCG/Source/PCG/Private/Tests/PCGGraphParametersTest.cpp
--- a/Engine/Plugins/PCG/Source/PCG/Private/Tests/PCGGraphParametersTest.cpp
+++ b/Engine/Plugins/PCG/Source/PCG/Private/Tests/PCGGraphParametersTest.cpp
@@ -436,29 +436,16 @@ bool FPCGGraphSetValueUserParametersPropagates::RunTest(const FString& Parameter
 	const FName MyPropertyName = TEXT("MyProperty");
 	const FName MyPropertyName2 = TEXT("MyProperty2");
 
-	// Creating a new property MyProperty, which is a double
-	PCGTests::EmulateModifyingUserParameters(Graph, [Graph, MyPropertyName]()
-	{
-		FInstancedPropertyBag* UserParameters = const_cast<FInstancedPropertyBag*>(Graph->GetUserParametersStruct());
-		UserParameters->AddProperty(MyPropertyName, EPropertyBagPropertyType::Double, nullptr);
-	});
-
-	// Creating another property MyProperty2, which is a double
-	PCGTests::EmulateModifyingUserParameters(Graph, [Graph, MyPropertyName2]()
-	{
-		FInstancedPropertyBag* UserParameters = const_cast<FInstancedPropertyBag*>(Graph->GetUserParametersStruct());
-		UserParameters->AddProperty(MyPropertyName2, EPropertyBagPropertyType::Int64, nullptr);
-	});
+	// Creating new properties MyProperty, which is a double and MyProperty2, which is a int64
+	Graph->UpdateUserParametersStruct([MyPropertyName, MyPropertyName2](FInstancedPropertyBag& UserParameters)
+		{
+			UserParameters.AddProperty(MyPropertyName, EPropertyBagPropertyType::Double, nullptr);
+			UserParameters.AddProperty(MyPropertyName2, EPropertyBagPropertyType::Int64, nullptr);
+		});
 
 	// Setting MyProperty to 3.0
 	const double NewValue = 3.0;
-	PCGTests::EmulateModifyingUserParametersValue(Graph, MyPropertyName, [Graph, MyPropertyName, NewValue]()
-	{
-		const FPropertyBagPropertyDesc* PropertyDesc = Graph->GetUserParametersStruct()->FindPropertyDescByName(MyPropertyName);
-		FInstancedPropertyBag* UserParameters = const_cast<FInstancedPropertyBag*>(Graph->GetUserParametersStruct());
-
-		PropertyDesc->CachedProperty->SetValue_InContainer(UserParameters->GetMutableValue().GetMemory(), &NewValue);
-	});
+	Graph->SetGraphParameter<float>(MyPropertyName, NewValue);
 
 	auto Verification = [this, MyPropertyName, NewValue](UPCGGraphInstance* GraphInstance, const TCHAR* Name) -> bool
 	{
@@ -467,13 +454,10 @@ bool FPCGGraphSetValueUserParametersPropagates::RunTest(const FString& Parameter
 		UTEST_TRUE(PCG_FMT("%s UserParameters property bag is valid", Name), GraphInstance->GetUserParametersStruct()->IsValid());
 		UTEST_EQUAL(PCG_FMT("%s UserParameters has 2 properties", Name), GraphInstance->GetUserParametersStruct()->GetNumPropertiesInBag(), 2);
 
-		const FPropertyBagPropertyDesc* PropertyDesc = GraphInstance->GetUserParametersStruct()->FindPropertyDescByName(MyPropertyName);
-		UTEST_NOT_NULL(PCG_FMT("%s: First property exists", Name), PropertyDesc);
-
-		double Value = 0.0;
-		PropertyDesc->CachedProperty->GetValue_InContainer(GraphInstance->GetUserParametersStruct()->GetValue().GetMemory(), &Value);
-
-		UTEST_EQUAL(PCG_FMT("%s: Property has the right value", Name), Value, NewValue);
+		TValueOrError<double, EPropertyBagResult> FirstPropertyValue = GraphInstance->GetGraphParameter<double>(MyPropertyName);
+		
+		UTEST_FALSE(PCG_FMT("%s: First property exists and is the expected type", Name), FirstPropertyValue.HasError());
+		UTEST_EQUAL(PCG_FMT("%s: Property has the right value", Name), FirstPropertyValue.GetValue(), NewValue);
 
 		return true;
 	};
diff --git a/Engine/Plugins/PCG/Source/PCG/Public/PCGGraph.h b/Engine/Plugins/PCG/Source/PCG/Public/PCGGraph.h
--- a/Engine/Plugins/PCG/Source/PCG/Public/PCGGraph.h
+++ b/Engine/Plugins/PCG/Source/PCG/Public/PCGGraph.h
@@ -104,8 +104,8 @@ public:
 
 	virtual const FInstancedPropertyBag* GetUserParametersStruct() const PURE_VIRTUAL(UPCGGraphInterface::GetUserParametersStruct, return nullptr;)
 
-	// Mutable version should not be used outside of testing, since there are callbacks fired when parameters changes.
-	// TODO: Make it safe to change parameters from the outside.
+	// Mutable version - should not be used outside of testing. Use UpdateUserParametersStruct on PCGGraph,
+	// or the different get/set/update functions in PCGGraphInterface instead to have proper callbacks.
 	FInstancedPropertyBag* GetMutableUserParametersStruct_Unsafe() const { return const_cast<FInstancedPropertyBag*>(GetUserParametersStruct()); }
 
 	bool IsInstance() const;
@@ -194,6 +194,22 @@ public:
 
 	EPropertyBagResult SetGraphParameter(const FName PropertyName, const uint64 Value, const UEnum* Enum);
 
+	/**
+	 * Allows to manipulate directly FPropertyBagArrayRef, while propagating changes to child instances. 
+	 * @param PropertyName Name of the property to access. Must be an Array.
+	 * @param Callback Callback to call with the FPropertyBagArrayRef. Returns a bool telling if there was a change.
+	 * @returns True if the change succeeded.
+	 */
+	bool UpdateArrayGraphParameter(const FName PropertyName, TFunctionRef<bool(FPropertyBagArrayRef& PropertyBagArrayRef)> Callback);
+
+	/**
+	 * Allows to manipulate directly FPropertyBagSetRef, while propagating changes to child instances. 
+	 * @param PropertyName Name of the property to access. Must be a Set.
+	 * @param Callback Callback to call with the FPropertyBagSetRef. Returns a bool telling if there was a change.
+	 * @returns True if the change succeeded.
+	 */
+	bool UpdateSetGraphParameter(const FName PropertyName, TFunctionRef<bool(FPropertyBagSetRef& PropertyBagSetRef)> Callback);
+
 	virtual void OnGraphParametersChanged(EPCGGraphParameterEvent InChangeType, FName InChangedPropertyName) PURE_VIRTUAL(UPCGGraphInterface::OnGraphParametersChanged, )
 
 protected:
@@ -470,6 +486,9 @@ public:
 
 	virtual void OnGraphParametersChanged(EPCGGraphParameterEvent InChangeType, FName InChangedPropertyName) override;
 
+	// Will call the callback function with a mutable property bag and will trigger the updates when it's done.
+	void UpdateUserParametersStruct(TFunctionRef<void(FInstancedPropertyBag&)> Callback);
+
 #if WITH_EDITOR
 private:
 	/** Sends a change notification. Demotes change if the compiled tasks are not significantly changed. */
diff --git a/Engine/Source/Runtime/CoreUObject/Private/StructUtils/PropertyBag.cpp b/Engine/Source/Runtime/CoreUObject/Private/StructUtils/PropertyBag.cpp
--- a/Engine/Source/Runtime/CoreUObject/Private/StructUtils/PropertyBag.cpp
+++ b/Engine/Source/Runtime/CoreUObject/Private/StructUtils/PropertyBag.cpp
@@ -2210,6 +2210,30 @@ TValueOrError<const FPropertyBagArrayRef, EPropertyBagResult> FInstancedProperty
 	return MakeValue(FPropertyBagArrayRef(*Desc, Address));
 }
 
+TValueOrError<FPropertyBagSetRef, EPropertyBagResult> FInstancedPropertyBag::GetMutableSetRef(const FName Name)
+{
+	const FPropertyBagPropertyDesc* Desc = FindPropertyDescByName(Name);
+	if (Desc == nullptr)
+	{
+		return MakeError(EPropertyBagResult::PropertyNotFound);
+	}
+	check(Desc->CachedProperty);
+
+	if (Desc->ContainerTypes.GetFirstContainerType() != EPropertyBagContainerType::Set)
+	{
+		return MakeError(EPropertyBagResult::TypeMismatch);
+	}
+
+	const void* Address = GetValueAddress(Desc);
+	if (Address == nullptr)
+	{
+		return MakeError(EPropertyBagResult::PropertyNotFound);
+	}
+	
+	return MakeValue(FPropertyBagSetRef(*Desc, Address));
+}
+
+
 TValueOrError<const FPropertyBagSetRef, EPropertyBagResult> FInstancedPropertyBag::GetSetRef(const FName Name) const
 {
 	const FPropertyBagPropertyDesc* Desc = FindPropertyDescByName(Name);
diff --git a/Engine/Source/Runtime/CoreUObject/Public/StructUtils/PropertyBag.h b/Engine/Source/Runtime/CoreUObject/Public/StructUtils/PropertyBag.h
--- a/Engine/Source/Runtime/CoreUObject/Public/StructUtils/PropertyBag.h
+++ b/Engine/Source/Runtime/CoreUObject/Public/StructUtils/PropertyBag.h
@@ -618,7 +618,7 @@ struct COREUOBJECT_API FInstancedPropertyBag
 	TValueOrError<FPropertyBagArrayRef, EPropertyBagResult> GetMutableArrayRef(const FName Name);
 
 	/**
-	 * Returns helper class to modify and access an array property.
+	 * Returns helper class to access an array property.
 	 * Note: The array reference is not valid after the layout of the referenced property bag has changed!
 	 * @returns helper class to modify and access arrays
 	*/
@@ -629,6 +629,13 @@ struct COREUOBJECT_API FInstancedPropertyBag
 	 * Note: The set reference is not valid after the layout of the referenced property bag has changed!
 	 * @returns helper class to modify and access sets
 	*/
+	TValueOrError<FPropertyBagSetRef, EPropertyBagResult> GetMutableSetRef(const FName Name);
+	
+	/**
+	 * Returns helper class to access a set property.
+	 * Note: The set reference is not valid after the layout of the referenced property bag has changed!
+	 * @returns helper class to modify and access sets
+	*/
 	TValueOrError<const FPropertyBagSetRef, EPropertyBagResult> GetSetRef(const FName Name) const;
 
 	bool Identical(const FInstancedPropertyBag* Other, uint32 PortFlags) const;
