[PCG] Fixed Attribute clobbering between Landscape and RayHit Query

* Streamlined attribute creation on World Query Data where applicable

#rnx

diff --git a/Engine/Plugins/PCG/Source/PCG/Private/Data/PCGWorldData.cpp b/Engine/Plugins/PCG/Source/PCG/Private/Data/PCGWorldData.cpp
--- a/Engine/Plugins/PCG/Source/PCG/Private/Data/PCGWorldData.cpp
+++ b/Engine/Plugins/PCG/Source/PCG/Private/Data/PCGWorldData.cpp
@@ -203,6 +203,7 @@ void UPCGWorldRayHitData::Initialize(UWorld* InWorld, const FTransform& InTransf
 	Transform = InTransform;
 	Bounds = InBounds;
 	LocalBounds = InLocalBounds;
+	PCGWorldQueryHelpers::CreateRayHitAttributes(QueryParams, Metadata);
 }
 
 void UPCGWorldRayHitData::AddToCrc(FArchiveCrc32& Ar, bool bFullDataCrc) const
@@ -237,8 +238,7 @@ bool UPCGWorldRayHitData::SamplePoint(const FTransform& InTransform, const FBox&
 		{
 			const FHitResult& Hit = HitResult.GetValue();
 			OutPoint = FPCGPoint(PCGWorldQueryHelpers::GetOrthonormalImpactTransform(Hit), 1.0f, UPCGBlueprintHelpers::ComputeSeedFromPosition(Hit.Location));
-			// TODO: Pre-create attributes within caller or find a better solution than per point
-			PCGWorldQueryHelpers::ApplyRayHitMetadata(Hit, QueryParams, QueryParams.RayDirection, OutPoint, OutMetadata, World, /*bShouldCreateAttributes=*/true);
+			PCGWorldQueryHelpers::ApplyRayHitMetadata(Hit, QueryParams, QueryParams.RayDirection, OutPoint, OutMetadata, World);
 			return true;
 		}
 	}
@@ -250,6 +250,7 @@ const UPCGPointData* UPCGWorldRayHitData::CreatePointData(FPCGContext* Context,
 {
 	UPCGPointData* Data = FPCGContext::NewObject_AnyThread<UPCGPointData>(Context);
 	Data->InitializeFromData(this);
+	PCGWorldQueryHelpers::CreateRayHitAttributes(QueryParams, Metadata);
 
 	FBox EffectiveBounds = Bounds;
 	if (InBounds.IsValid)
@@ -295,4 +296,4 @@ UPCGSpatialData* UPCGWorldRayHitData::CopyInternal(FPCGContext* Context) const
 	NewData->QueryParams.Initialize();
 
 	return NewData;
-}
+}
\ No newline at end of file
diff --git a/Engine/Plugins/PCG/Source/PCG/Private/Elements/PCGWorldQuery.cpp b/Engine/Plugins/PCG/Source/PCG/Private/Elements/PCGWorldQuery.cpp
--- a/Engine/Plugins/PCG/Source/PCG/Private/Elements/PCGWorldQuery.cpp
+++ b/Engine/Plugins/PCG/Source/PCG/Private/Elements/PCGWorldQuery.cpp
@@ -134,28 +134,10 @@ bool FPCGWorldRayHitQueryElement::ExecuteInternal(FPCGContext* Context) const
 	UWorld* World = Context->SourceComponent->GetWorld();
 
 	UPCGWorldRayHitData* Data = FPCGContext::NewObject_AnyThread<UPCGWorldRayHitData>(Context);
-	Data->Initialize(World, Transform, /*InBounds=*/FBox(EForceInit::ForceInit), LocalBounds);
 	Data->QueryParams = QueryParams;
 	Data->QueryParams.Initialize();
 	Data->OriginatingComponent = Context->SourceComponent;
-
-	if (Data->QueryParams.bGetReferenceToActorHit && Data->Metadata)
-	{
-		Data->Metadata->FindOrCreateAttribute(PCGPointDataConstants::ActorReferenceAttribute, FSoftObjectPath(), /*bAllowInterpolation=*/false, /*bOverrideParent=*/false);
-	}
-	else
-	{
-		Data->QueryParams.bGetReferenceToActorHit = false;
-	}
-
-	if (Data->QueryParams.bGetReferenceToPhysicalMaterial && Data->Metadata)
-	{
-		Data->Metadata->FindOrCreateAttribute(PCGWorldQueryConstants::PhysicalMaterialReferenceAttribute, FSoftObjectPath(), /*bAllowInterpolation=*/false, /*bOverrideParent=*/false);
-	}
-	else
-	{
-		Data->QueryParams.bGetReferenceToPhysicalMaterial = false;
-	}
+	Data->Initialize(World, Transform, /*InBounds=*/FBox(EForceInit::ForceInit), LocalBounds);
 
 	bool bHasLandscapeMetadata = false;
 	if (QueryParams.bApplyMetadataFromLandscape && Data->Metadata)
@@ -182,7 +164,7 @@ bool FPCGWorldRayHitQueryElement::ExecuteInternal(FPCGContext* Context) const
 					return OtherActorBounds.Intersect(Data->Bounds);
 				};
 			}
-	
+
 			TArray<AActor*> LandscapeActors = PCGActorSelector::FindActors(ActorSelector, Context->SourceComponent.Get(), BoundsCheck, SelfIgnoreCheck);
 			for (AActor* Landscape : LandscapeActors)
 			{
diff --git a/Engine/Plugins/PCG/Source/PCG/Private/Elements/PCGWorldRaycast.cpp b/Engine/Plugins/PCG/Source/PCG/Private/Elements/PCGWorldRaycast.cpp
--- a/Engine/Plugins/PCG/Source/PCG/Private/Elements/PCGWorldRaycast.cpp
+++ b/Engine/Plugins/PCG/Source/PCG/Private/Elements/PCGWorldRaycast.cpp
@@ -348,7 +348,7 @@ bool FPCGWorldRaycastElement::ExecuteInternal(FPCGContext* InContext) const
 				OutPoint.Density = 1.f;
 				UPCGBlueprintHelpers::SetSeedFromPosition(OutPoint);
 
-				bAttributeSuccess &= PCGWorldQueryHelpers::ApplyRayHitMetadata(Hit, WorldQueryParams, Direction, OutputPoints.Last(), OutMetadata, World, /*bShouldCreateAttributes=*/false);
+				bAttributeSuccess &= PCGWorldQueryHelpers::ApplyRayHitMetadata(Hit, WorldQueryParams, Direction, OutputPoints.Last(), OutMetadata, World);
 			}
 			else if (Settings->bKeepOriginalPointOnMiss)
 			{
@@ -375,4 +375,4 @@ bool FPCGWorldRaycastElement::ExecuteInternal(FPCGContext* InContext) const
 	});
 }
 
-#undef LOCTEXT_NAMESPACE
+#undef LOCTEXT_NAMESPACE
\ No newline at end of file
diff --git a/Engine/Plugins/PCG/Source/PCG/Private/Grid/PCGLandscapeCache.cpp b/Engine/Plugins/PCG/Source/PCG/Private/Grid/PCGLandscapeCache.cpp
--- a/Engine/Plugins/PCG/Source/PCG/Private/Grid/PCGLandscapeCache.cpp
+++ b/Engine/Plugins/PCG/Source/PCG/Private/Grid/PCGLandscapeCache.cpp
@@ -269,7 +269,7 @@ void FPCGLandscapeCacheEntry::GetInterpolatedPointMetadataInternal(const PCGLand
 
 	if (OutMetadata && !LayerData.IsEmpty())
 	{
-		OutPoint.MetadataEntry = OutMetadata->AddEntry();
+		OutMetadata->InitializeOnSet(OutPoint.MetadataEntry);
 
 		check(LayerData.Num() == LayerDataNames.Num());
 
@@ -1085,4 +1085,4 @@ void UPCGLandscapeCache::CacheLayerNames(ALandscapeProxy* InLandscape)
 	}
 }
 
-#endif // WITH_EDITOR
+#endif // WITH_EDITOR
\ No newline at end of file
diff --git a/Engine/Plugins/PCG/Source/PCG/Private/Helpers/PCGWorldQueryHelpers.cpp b/Engine/Plugins/PCG/Source/PCG/Private/Helpers/PCGWorldQueryHelpers.cpp
--- a/Engine/Plugins/PCG/Source/PCG/Private/Helpers/PCGWorldQueryHelpers.cpp
+++ b/Engine/Plugins/PCG/Source/PCG/Private/Helpers/PCGWorldQueryHelpers.cpp
@@ -172,22 +172,9 @@ namespace PCGWorldQueryHelpers
 
 	bool CreateRayHitAttributes(const FPCGWorldRaycastQueryParams& QueryParams, UPCGMetadata* OutMetadata)
 	{
-		auto CreateAttribute = [OutMetadata]<typename Type>(FName AttributeName, bool bShouldCreate, const Type& DefaultValue)
+		auto CreateAttribute = [OutMetadata]<typename Type>(FName AttributeName, bool bShouldCreate, const Type& DefaultValue) -> bool
 		{
-			if (!bShouldCreate)
-			{
-				return true;
-			}
-
-			if (!OutMetadata->HasAttribute(AttributeName))
-			{
-				if (OutMetadata->FindOrCreateAttribute<Type>(AttributeName, DefaultValue, /*bAllowsInterpolation=*/true, /*bOverrideParent=*/false, /*bOverwriteIfTypeMismatch=*/true))
-				{
-					return true;
-				}
-			}
-
-			return false;
+			return !bShouldCreate || OutMetadata->FindOrCreateAttribute<Type>(AttributeName, DefaultValue, /*bAllowsInterpolation=*/true, /*bOverrideParent=*/false, /*bOverwriteIfTypeMismatch=*/true);
 		};
 
 		bool bResult = true;
@@ -221,7 +208,7 @@ namespace PCGWorldQueryHelpers
 			return false;
 		}
 
-		if (FPCGMetadataAttribute<bool>* Attribute = OutMetadata->FindOrCreateAttribute<bool>(PCGWorldQueryConstants::ImpactAttribute, true, /*bAllowsInterpolation=*/true, /*bOverrideParent=*/false))
+		if (FPCGMetadataAttribute<bool>* Attribute = OutMetadata->FindOrCreateAttribute<bool>(PCGWorldQueryConstants::ImpactAttribute, true, /*bAllowsInterpolation=*/true, /*bOverrideParent=*/false, /*bOverwriteIfTypeMismatch=*/true))
 		{
 			OutMetadata->InitializeOnSet(OutPoint.MetadataEntry);
 			Attribute->SetValue(OutPoint.MetadataEntry, false);
@@ -238,19 +225,13 @@ namespace PCGWorldQueryHelpers
 		const FVector& RayDirection,
 		FPCGPoint& OutPoint,
 		UPCGMetadata* OutMetadata,
-		TWeakObjectPtr<UWorld> World,
-		bool bShouldCreateAttributes)
+		TWeakObjectPtr<UWorld> World)
 	{
 		if (!OutMetadata)
 		{
 			return false;
 		}
 
-		if (bShouldCreateAttributes)
-		{
-			CreateRayHitAttributes(QueryParams, OutMetadata);
-		}
-
 		auto ApplyAttribute = [&OutPoint, OutMetadata]<typename Type>(FName AttributeName, const Type& Value, bool bShouldApply = true)
 		{
 			if (!bShouldApply)
@@ -302,10 +283,8 @@ namespace PCGWorldQueryHelpers
 					const TArray<FName> Layers = LandscapeCache->GetLayerNames(Landscape);
 					for (const FName& Layer : Layers)
 					{
-						if (!OutMetadata->HasAttribute(Layer))
-						{
-							OutMetadata->CreateAttribute<float>(Layer, float(), /*bAllowInterpolation=*/true, /*bOverrideParent=*/true);
-						}
+						// Use Threadsafe version for on the fly attribute creation.
+						OutMetadata->FindOrCreateAttribute<float>(Layer, float());
 					}
 
 					LandscapeCache->SampleMetadataOnPoint(Landscape, OutPoint, OutMetadata);
diff --git a/Engine/Plugins/PCG/Source/PCG/Public/Helpers/PCGWorldQueryHelpers.h b/Engine/Plugins/PCG/Source/PCG/Public/Helpers/PCGWorldQueryHelpers.h
--- a/Engine/Plugins/PCG/Source/PCG/Public/Helpers/PCGWorldQueryHelpers.h
+++ b/Engine/Plugins/PCG/Source/PCG/Public/Helpers/PCGWorldQueryHelpers.h
@@ -75,6 +75,5 @@ namespace PCGWorldQueryHelpers
 		const FVector& RayDirection,
 		FPCGPoint& OutPoint,
 		UPCGMetadata* OutMetadata,
-		TWeakObjectPtr<UWorld> World,
-		bool bShouldCreateAttributes = true);
-}
+		TWeakObjectPtr<UWorld> World);
+}
\ No newline at end of file
diff --git a/Engine/Plugins/PCG/Source/PCG/Public/Metadata/PCGMetadata.h b/Engine/Plugins/PCG/Source/PCG/Public/Metadata/PCGMetadata.h
--- a/Engine/Plugins/PCG/Source/PCG/Public/Metadata/PCGMetadata.h
+++ b/Engine/Plugins/PCG/Source/PCG/Public/Metadata/PCGMetadata.h
@@ -197,6 +197,9 @@ public:
 	template <typename T>
 	FPCGMetadataAttribute<T>* GetMutableTypedAttribute(FName AttributeName);
 
+	template <typename T>
+	FPCGMetadataAttribute<T>* GetMutableTypedAttribute_Unsafe(FName AttributeName);
+
 	template <typename T>
 	const FPCGMetadataAttribute<T>* GetConstTypedAttribute(FName AttributeName) const;
 
@@ -399,7 +402,7 @@ FPCGMetadataAttribute<T>* UPCGMetadata::CreateAttribute(FName AttributeName, con
 
 	if (ParentAttribute && (ParentAttribute->GetTypeId() != PCG::Private::MetadataTypes<T>::Id))
 	{
-		// Can't parent if the types doesn't match
+		// Can't parent if the types do not match
 		ParentAttribute = nullptr;
 	}
 
@@ -435,24 +438,71 @@ FPCGMetadataAttribute<T>* UPCGMetadata::CreateAttribute(FName AttributeName, con
 	return NewAttribute;
 }
 
+template<typename T>
+FPCGMetadataAttribute<T>* UPCGMetadata::GetMutableTypedAttribute_Unsafe(FName AttributeName)
+{
+	FPCGMetadataAttribute<T>* Attribute = nullptr;
+
+	if (FPCGMetadataAttributeBase** FoundAttribute = Attributes.Find(AttributeName))
+	{
+		FPCGMetadataAttributeBase* BaseAttribute = *FoundAttribute;
+		if (BaseAttribute && BaseAttribute->GetTypeId() == PCG::Private::MetadataTypes<T>::Id)
+		{
+			Attribute = static_cast<FPCGMetadataAttribute<T>*>(BaseAttribute);
+		}
+	}
+
+	return Attribute;
+}
+
 template<typename T>
 FPCGMetadataAttribute<T>* UPCGMetadata::FindOrCreateAttribute(FName AttributeName, const T& DefaultValue, bool bAllowsInterpolation, bool bOverrideParent, bool bOverwriteIfTypeMismatch)
 {
-	FPCGMetadataAttribute<T>* Attribute = GetMutableTypedAttribute<T>(AttributeName);
+	{
+		FReadScopeLock ScopeLock(AttributeLock);
+		if (FPCGMetadataAttribute<T>* Attribute = GetMutableTypedAttribute_Unsafe<T>(AttributeName))
+		{
+			return Attribute;
+		}
+	}
 
-	// If Attribute is null, but we have an attribute with this name, we have a type mismatch.
-	// Will be overwrite if flag bOverwriteIfTypeMismatch is at true.
-	if (!Attribute && HasAttribute(AttributeName) && bOverwriteIfTypeMismatch)
+	FWriteScopeLock ScopeLock(AttributeLock);
+	if (FPCGMetadataAttribute<T>* Attribute = GetMutableTypedAttribute_Unsafe<T>(AttributeName))
 	{
-		DeleteAttribute(AttributeName);
+		return Attribute;
 	}
 
-	if (!Attribute)
+	// If an attribute with this name exists here, there is a type mismatch.
+	if (FPCGMetadataAttributeBase** FoundAttribute = Attributes.Find(AttributeName))
 	{
-		Attribute = CreateAttribute<T>(AttributeName, DefaultValue, bAllowsInterpolation, bOverrideParent);
+		if (bOverwriteIfTypeMismatch)
+		{
+			delete *FoundAttribute;
+			RemoveAttributeInternal(AttributeName);
+		}
+		else
+		{
+			return nullptr;
+		}
 	}
 
-	return Attribute;
+	// A new attribute will be created.
+	if (!FPCGMetadataAttributeBase::IsValidName(AttributeName))
+	{
+		UE_LOG(LogPCG, Error, TEXT("Attribute name '%s' is invalid"), *AttributeName.ToString());
+		return nullptr;
+	}
+
+	// Parent is const and therefore should be safe
+	const FPCGMetadataAttributeBase* ParentAttribute = (bOverrideParent && Parent) ? Parent->GetConstTypedAttribute<T>(AttributeName) : nullptr;
+	FPCGMetadataAttribute<T>* NewAttribute = new FPCGMetadataAttribute<T>(this, AttributeName, ParentAttribute, DefaultValue, bAllowsInterpolation);
+	NewAttribute->AttributeId = NextAttributeId++;
+	AddAttributeInternal(AttributeName, NewAttribute);
+
+	// Also when creating an attribute, notify the PCG Data owner that the latest attribute manipulated is this one.
+	SetLastCachedSelectorOnOwner(AttributeName);
+
+	return NewAttribute;
 }
 
 template <typename T>
@@ -462,7 +512,6 @@ FPCGMetadataAttribute<T>* UPCGMetadata::GetMutableTypedAttribute(FName Attribute
 	return (BaseAttribute && (BaseAttribute->GetTypeId() == PCG::Private::MetadataTypes<T>::Id))
 		? static_cast<FPCGMetadataAttribute<T>*>(BaseAttribute)
 		: nullptr;
-
 }
 
 template <typename T>
@@ -472,4 +521,4 @@ const FPCGMetadataAttribute<T>* UPCGMetadata::GetConstTypedAttribute(FName Attri
 	return (BaseAttribute && (BaseAttribute->GetTypeId() == PCG::Private::MetadataTypes<T>::Id))
 		? static_cast<const FPCGMetadataAttribute<T>*>(BaseAttribute)
 		: nullptr;
-}
+}
\ No newline at end of file
