[PCG] If GPU SM Spawners are present, always allow the component to generate.

The procedural instances are very transient and this is needed to make regenerating trivial upon editor restart.

#rnx

diff --git a/Engine/Plugins/PCG/Source/PCG/Private/Compute/Elements/PCGComputeGraphElement.cpp b/Engine/Plugins/PCG/Source/PCG/Private/Compute/Elements/PCGComputeGraphElement.cpp
--- a/Engine/Plugins/PCG/Source/PCG/Private/Compute/Elements/PCGComputeGraphElement.cpp
+++ b/Engine/Plugins/PCG/Source/PCG/Private/Compute/Elements/PCGComputeGraphElement.cpp
@@ -43,6 +43,8 @@ bool FPCGComputeGraphElement::ExecuteInternal(FPCGContext* InContext) const
 {
 	TRACE_CPUPROFILER_EVENT_SCOPE(FPCGComputeGraphElement::ExecuteInternal);
 	check(InContext);
+	check(InContext->SourceComponent.IsValid());
+
 	FPCGComputeGraphContext* Context = static_cast<FPCGComputeGraphContext*>(InContext);
 
 	if (!Context->ComputeGraph)
@@ -151,6 +153,11 @@ bool FPCGComputeGraphElement::ExecuteInternal(FPCGContext* InContext) const
 
 		const bool bAnyComponentsSetup = SetupProceduralISMComponents(Context, DataBinding);
 
+		if (bAnyComponentsSetup)
+		{
+			Context->SourceComponent->NotifyProceduralInstancesInUse();
+		}
+
 		{
 			TRACE_CPUPROFILER_EVENT_SCOPE(ComputeGraphInstance.CreateDataProviders);
 			Context->ComputeGraphInstance.CreateDataProviders(Context->ComputeGraph.Get(), 0, Context->DataBinding.Get());
diff --git a/Engine/Plugins/PCG/Source/PCG/Private/PCGComponent.cpp b/Engine/Plugins/PCG/Source/PCG/Private/PCGComponent.cpp
--- a/Engine/Plugins/PCG/Source/PCG/Private/PCGComponent.cpp
+++ b/Engine/Plugins/PCG/Source/PCG/Private/PCGComponent.cpp
@@ -323,6 +323,12 @@ bool UPCGComponent::ShouldGenerate(bool bForce, EPCGComponentGenerationTrigger R
 	}
 #endif
 
+	// Always generate if procedural ISMs are being used, because the instance data is not persistent, and is currently lost regularly when the GPU Scene is flushed.
+	if (bProceduralInstancesInUse)
+	{
+		return true;
+	}
+
 	// A request is invalid only if it was requested "GenerateOnLoad", but it is "GenerateOnDemand"
 	// Meaning that all "GenerateOnDemand" requests are always valid, and "GenerateOnLoad" request is only valid if we want a "GenerateOnLoad" trigger.
 	bool bValidRequest = !(RequestedGenerationTrigger == EPCGComponentGenerationTrigger::GenerateOnLoad && GenerationTrigger == EPCGComponentGenerationTrigger::GenerateOnDemand);
@@ -453,6 +459,9 @@ FPCGTaskId UPCGComponent::GenerateInternal(bool bForce, EPCGHiGenGrid Grid, EPCG
 
 	Modify(!IsInPreviewMode());
 
+	// Clear prior to generation.
+	bProceduralInstancesInUse = false;
+
 	CurrentGenerationTask = GetSubsystem()->ScheduleComponent(this, Grid, bForce, Dependencies);
 
 	if (CurrentGenerationTask != InvalidPCGTaskId)
diff --git a/Engine/Plugins/PCG/Source/PCG/Public/PCGComponent.h b/Engine/Plugins/PCG/Source/PCG/Public/PCGComponent.h
--- a/Engine/Plugins/PCG/Source/PCG/Public/PCGComponent.h
+++ b/Engine/Plugins/PCG/Source/PCG/Public/PCGComponent.h
@@ -236,6 +236,9 @@ public:
 	/** Compute the runtime cleanup radius for the given grid size. */
 	double GetCleanupRadiusFromGrid(EPCGHiGenGrid Grid) const;
 
+	/** Called during execution if one or more procedural ISM components are in use. */
+	void NotifyProceduralInstancesInUse() { bProceduralInstancesInUse = true; }
+
 	UPROPERTY(BlueprintReadWrite, EditAnywhere, Category = Settings, meta = (DisplayPriority = 600))
 	int Seed = 42;
 
@@ -636,6 +639,10 @@ private:
 	UPROPERTY(VisibleAnywhere, Transient, Category = Debug, meta = (EditCondition = false, EditConditionHides))
 	bool bIsComponentLocal = false;
 
+	/** Whether procedural ISM components were used/generated in the last execution. */
+	UPROPERTY()
+	bool bProceduralInstancesInUse = false;
+
 #if WITH_EDITOR
 	bool bWasGeneratedThisSession = false;
 	int32 InspectionCounter = 0;
