Fix little ifdef error in PCGMatchAndSetAttributes

diff --git a/Engine/Plugins/PCG/Source/PCG/Private/Elements/PCGMatchAndSetAttributes.cpp b/Engine/Plugins/PCG/Source/PCG/Private/Elements/PCGMatchAndSetAttributes.cpp
--- a/Engine/Plugins/PCG/Source/PCG/Private/Elements/PCGMatchAndSetAttributes.cpp
+++ b/Engine/Plugins/PCG/Source/PCG/Private/Elements/PCGMatchAndSetAttributes.cpp
@@ -275,7 +275,7 @@ public:
 			MatchingVK->Value.AddEntry(EntryKey, Weight);
 		}
 
-#ifdef WITH_EDITOR
+#if WITH_EDITOR
 		// Check for empty entries
 		const bool bHasEmptyEntries = Algo::AnyOf(PartitionData, [](const TPair<PCGMetadataValueKey, AttributeSetPartitionEntry>& PartitionEntry) { return PartitionEntry.Value.TotalWeight == 0; });
 		if (bHasEmptyEntries)
