[PCG] Fix custom kernel attributes not working in standalone builds.

diff --git a/Engine/Plugins/PCG/Source/PCG/Private/Compute/Elements/PCGCustomHLSL.cpp b/Engine/Plugins/PCG/Source/PCG/Private/Compute/Elements/PCGCustomHLSL.cpp
--- a/Engine/Plugins/PCG/Source/PCG/Private/Compute/Elements/PCGCustomHLSL.cpp
+++ b/Engine/Plugins/PCG/Source/PCG/Private/Compute/Elements/PCGCustomHLSL.cpp
@@ -71,7 +71,6 @@ UPCGCustomHLSLSettings::UPCGCustomHLSLSettings()
 	bExecuteOnGPU = true;
 }
 
-#if WITH_EDITOR
 void UPCGCustomHLSLSettings::PostLoad()
 {
 	Super::PostLoad();
@@ -79,7 +78,10 @@ void UPCGCustomHLSLSettings::PostLoad()
 	// Note: We update here so that Custom HLSL nodes will have the correct pin settings & declarations on load.
 	UpdatePinSettings();
 	UpdateAttributeKeys();
+
+#if WITH_EDITOR
 	UpdateDeclarations();
+#endif
 }
 
 void UPCGCustomHLSLSettings::PostInitProperties()
@@ -88,9 +90,11 @@ void UPCGCustomHLSLSettings::PostInitProperties()
 
 	// Note: We update here so that Custom HLSL nodes will have the correct pin settings & declarations on creation.
 	UpdatePinSettings();
+
+#if WITH_EDITOR
 	UpdateDeclarations();
-}
 #endif
+}
 
 TArray<FPCGPinProperties> UPCGCustomHLSLSettings::OutputPinProperties() const
 {
@@ -988,6 +992,7 @@ void UPCGCustomHLSLSettings::UpdateHelperDeclarations()
 
 	HelperDeclarations.TrimStartAndEndInline();
 }
+#endif // WITH_EDITOR
 
 void UPCGCustomHLSLSettings::UpdatePinSettings()
 {
@@ -1011,17 +1016,25 @@ void UPCGCustomHLSLSettings::UpdatePinSettings()
 			// Also disallow multi-data for Attribute Sets, since we require attributes to be uniform on a pin, thus having different
 			// attribute sets wouldn't work as expected.
 			Properties.bAllowMultipleData = false;
+
+#if WITH_EDITOR
 			Properties.bAllowEditMultipleData = false;
+#endif
 		}
 		else
 		{
+#if WITH_EDITOR
 			Properties.bAllowEditMultipleData = true;
+#endif
 		}
 
 		// TODO: We have work to do to allow dynamic merging of data. Also we will likely inject Gather
 		// nodes on the CPU side so that merging is handled CPU side where possible.
 		Properties.SetAllowMultipleConnections(false);
+
+#if WITH_EDITOR
 		Properties.bAllowEditMultipleConnections = false;
+#endif
 	}
 
 	// Setup output pins.
@@ -1035,6 +1048,7 @@ void UPCGCustomHLSLSettings::UpdatePinSettings()
 			Properties.AllowedTypes = EPCGDataType::Point;
 		}
 
+#if WITH_EDITOR
 		// Only allow editing the initialization mode if it's not driven by the kernel type.
 		const bool bInitModeDrivenByKernel = PinIndex == 0 && (KernelType == EPCGKernelType::PointProcessor || KernelType == EPCGKernelType::PointGenerator);
 		Properties.PropertiesGPU.bAllowEditInitMode = !bInitModeDrivenByKernel;
@@ -1043,19 +1057,25 @@ void UPCGCustomHLSLSettings::UpdatePinSettings()
 		// Output pins should always allow multiple connections.
 		// TODO this could be hoisted up somewhere in the future.
 		Properties.bAllowEditMultipleConnections = false;
+#endif
 
 		if (!!(Properties.AllowedTypes & EPCGDataType::Param))
 		{
 			Properties.bAllowMultipleData = false;
-			Properties.bAllowEditMultipleData = false;
-			Properties.PropertiesGPU.bAllowEditDataCount = false;
 			Properties.PropertiesGPU.DataCountMode = EPCGDataCountMode::Fixed;
 			Properties.PropertiesGPU.DataCount = 1;
+
+#if WITH_EDITOR
+			Properties.bAllowEditMultipleData = false;
+			Properties.PropertiesGPU.bAllowEditDataCount = false;
+#endif
 		}
 		else
 		{
+#if WITH_EDITOR
 			Properties.bAllowEditMultipleData = true;
 			Properties.PropertiesGPU.bAllowEditDataCount = true;
+#endif
 		}
 	}
 }
@@ -1136,6 +1156,7 @@ void UPCGCustomHLSLSettings::UpdateAttributeKeys()
 	}
 }
 
+#if WITH_EDITOR
 TArray<FName> UPCGCustomHLSLSettings::GetInputPinNames() const
 {
 	TArray<FName> PinNames;
@@ -1155,7 +1176,7 @@ TArray<FName> UPCGCustomHLSLSettings::GetInputPinNamesAndNone() const
 
 	return PinNames;
 }
-#endif
+#endif // WITH_EDITOR
 
 bool UPCGCustomHLSLSettings::IsKernelValid(FPCGContext* InContext, bool bQuiet) const
 {
diff --git a/Engine/Plugins/PCG/Source/PCG/Private/Metadata/Accessors/PCGAttributeAccessorKeys.cpp b/Engine/Plugins/PCG/Source/PCG/Private/Metadata/Accessors/PCGAttributeAccessorKeys.cpp
--- a/Engine/Plugins/PCG/Source/PCG/Private/Metadata/Accessors/PCGAttributeAccessorKeys.cpp
+++ b/Engine/Plugins/PCG/Source/PCG/Private/Metadata/Accessors/PCGAttributeAccessorKeys.cpp
@@ -151,7 +151,6 @@ FPCGAttributeAccessorKeysPointsSubset::FPCGAttributeAccessorKeysPointsSubset(con
 FPCGAttributeAccessorKeysPointsSubset::FPCGAttributeAccessorKeysPointsSubset(const TArrayView<const FPCGPoint>& InPoints, const TArrayView<const int32>& InPointIndices)
 	: IPCGAttributeAccessorKeys(/*bInReadOnly=*/ true)
 {
-	check(InPoints.Num() == InPointIndices.Num());
 	Points.Reserve(InPointIndices.Num());
 	Algo::Transform(InPointIndices, Points, [&InPoints](const int32 Index) -> FPCGPoint* { return const_cast<FPCGPoint*>(&InPoints[Index]); });
 }
diff --git a/Engine/Plugins/PCG/Source/PCG/Public/Compute/Elements/PCGCustomHLSL.h b/Engine/Plugins/PCG/Source/PCG/Public/Compute/Elements/PCGCustomHLSL.h
--- a/Engine/Plugins/PCG/Source/PCG/Public/Compute/Elements/PCGCustomHLSL.h
+++ b/Engine/Plugins/PCG/Source/PCG/Public/Compute/Elements/PCGCustomHLSL.h
@@ -41,10 +41,8 @@ public:
 	UPCGCustomHLSLSettings();
 
 	//~Begin UObject interface
-#if WITH_EDITOR
 	virtual void PostLoad() override;
 	virtual void PostInitProperties() override;
-#endif
 	//~End UObject interface
 
 	//~Begin UPCGSettings interface
@@ -68,7 +66,7 @@ public:
 
 	virtual bool IsKernelValid(FPCGContext* InContext = nullptr, bool bQuiet = true) const override;
 	virtual FString GetCookedKernelSource(const TMap<FName, FPCGKernelAttributeIDAndType>& GlobalAttributeLookupTable) const override;
-	virtual const TArray<FPCGKernelAttributeKey> GetKernelAttributeKeys() const { return KernelAttributeKeys; }
+	virtual const TArray<FPCGKernelAttributeKey> GetKernelAttributeKeys() const override { return KernelAttributeKeys; }
 	virtual int ComputeKernelThreadCount(const UPCGDataBinding* Binding) const override;
 	virtual FPCGDataCollectionDesc ComputeOutputPinDataDesc(const UPCGPin* OutputPin, const UPCGDataBinding* Binding) const override;
 
@@ -96,13 +94,17 @@ protected:
 	/** Gets the GPU pin properties for the output pin with the given label. */
 	const FPCGPinPropertiesGPU* GetOutputPinPropertiesGPU(const FName& InPinLabel) const;
 
+	/** Enforce required pin settings and set display toggles to drive UI. */
+	void UpdatePinSettings();
+
+	/** Update the kernel keys read/written/created by this node. */
+	void UpdateAttributeKeys();
+
 #if WITH_EDITOR
 	void UpdateDeclarations();
 	void UpdateInputDeclarations();
 	void UpdateOutputDeclarations();
 	void UpdateHelperDeclarations();
-	void UpdatePinSettings();
-	void UpdateAttributeKeys();
 
 	/** List of all non-advanced input pin names. */
 	UFUNCTION()
