Fixed issue where components spawned by PCG wouldn't retain their transform after loading

#rnx

diff --git a/Engine/Plugins/PCG/Source/PCG/Private/Data/PCGSplineStruct.cpp b/Engine/Plugins/PCG/Source/PCG/Private/Data/PCGSplineStruct.cpp
--- a/Engine/Plugins/PCG/Source/PCG/Private/Data/PCGSplineStruct.cpp
+++ b/Engine/Plugins/PCG/Source/PCG/Private/Data/PCGSplineStruct.cpp
@@ -107,7 +107,7 @@ void FPCGSplineStruct::ApplyTo(USplineComponent* InSplineComponent) const
 	check(InSplineComponent);
 
 	InSplineComponent->ClearSplinePoints(false);
-	InSplineComponent->SetComponentToWorld(Transform);
+	InSplineComponent->SetWorldTransform(Transform);
 	InSplineComponent->DefaultUpVector = DefaultUpVector;
 	InSplineComponent->ReparamStepsPerSegment = ReparamStepsPerSegment;
 
diff --git a/Engine/Plugins/PCG/Source/PCG/Private/Elements/PCGAddComponent.cpp b/Engine/Plugins/PCG/Source/PCG/Private/Elements/PCGAddComponent.cpp
--- a/Engine/Plugins/PCG/Source/PCG/Private/Elements/PCGAddComponent.cpp
+++ b/Engine/Plugins/PCG/Source/PCG/Private/Elements/PCGAddComponent.cpp
@@ -489,7 +489,7 @@ void FPCGAddComponentElement::AddComponents(FPCGAddComponentContext* Context,
 				if (USceneComponent* SceneComponent = Cast<USceneComponent>(Component))
 				{
 					SceneComponent->AttachToComponent(TargetActor->GetRootComponent(), FAttachmentTransformRules(EAttachmentRule::KeepRelative, EAttachmentRule::KeepWorld, EAttachmentRule::KeepWorld, false));
-					SceneComponent->SetComponentToWorld(Transform);
+					SceneComponent->SetWorldTransform(Transform);
 				}
 
 				Component->ComponentTags.Add(SourceComponent->GetFName());
