[PCG] Fix bug with named reroutes allowing naming conflicts.

#rnx

diff --git a/Engine/Plugins/PCG/Source/PCGEditor/Private/PCGEditor.cpp b/Engine/Plugins/PCG/Source/PCGEditor/Private/PCGEditor.cpp
--- a/Engine/Plugins/PCG/Source/PCGEditor/Private/PCGEditor.cpp
+++ b/Engine/Plugins/PCG/Source/PCGEditor/Private/PCGEditor.cpp
@@ -1641,11 +1641,13 @@ void FPCGEditor::OnConvertRerouteToNamedReroute()
 		TArray<UPCGEdge*> RerouteInputEdges = Reroute->GetInputPin(PCGPinConstants::DefaultInputLabel)->Edges;
 		ensure(RerouteInputEdges.Num() <= 1);
 
-		if(RerouteInputEdges.Num() >= 1)
+		if (RerouteInputEdges.Num() >= 1)
 		{
 			UPCGEdge* InputEdge = RerouteInputEdges[0];
 			PCGGraph->AddEdge(InputEdge->InputPin->Node, InputEdge->InputPin->Properties.Label, Declaration, PCGPinConstants::DefaultInputLabel);
-			CastChecked<UPCGEditorGraphNodeNamedRerouteDeclaration>(GetEditorNode(Declaration))->SetNodeName(InputEdge->InputPin->Node, InputEdge->InputPin->Properties.Label);
+			UPCGEditorGraphNodeNamedRerouteDeclaration* DeclarationNode = CastChecked<UPCGEditorGraphNodeNamedRerouteDeclaration>(GetEditorNode(Declaration));
+			const FString NodeName = DeclarationNode->GenerateNodeName(InputEdge->InputPin->Node, InputEdge->InputPin->Properties.Label);
+			DeclarationNode->OnRenameNode(NodeName);
 		}
 
 		// Add invisible edge between declaration and usage
@@ -3385,4 +3387,4 @@ FText FPCGEditor::GetAttributesTabLabel(int AttributesIndex)
 	}
 }
 
-#undef LOCTEXT_NAMESPACE
+#undef LOCTEXT_NAMESPACE
\ No newline at end of file
diff --git a/Engine/Plugins/PCG/Source/PCGEditor/Private/PCGEditorGraphNode.cpp b/Engine/Plugins/PCG/Source/PCGEditor/Private/PCGEditorGraphNode.cpp
--- a/Engine/Plugins/PCG/Source/PCGEditor/Private/PCGEditorGraphNode.cpp
+++ b/Engine/Plugins/PCG/Source/PCGEditor/Private/PCGEditorGraphNode.cpp
@@ -63,23 +63,36 @@ void UPCGEditorGraphNode::AllocateDefaultPins()
 }
 
 void UPCGEditorGraphNode::OnRenameNode(const FString& NewName)
+{
+	if (CanRenameNode(NewName))
+	{
+		CommitNodeRename(NewName);
+	}
+}
+
+bool UPCGEditorGraphNode::CanRenameNode(const FString& NewName) const
 {
 	if (!GetCanRenameNode())
 	{
-		return;
+		return false;
 	}
 
-	if(NewName.Len() >= NAME_SIZE)
+	if (NewName.Len() >= NAME_SIZE)
 	{
 		UE_LOG(LogPCGEditor, Error, TEXT("New name for PCG node is too long."));
-		return;
+		return false;
 	}
 
 	if (!PCGNode)
 	{
-		return;
+		return false;
 	}
 
+	return true;
+}
+
+void UPCGEditorGraphNode::CommitNodeRename(const FString& NewName)
+{
 	if (PCGNode->GetAuthoredTitleLine().ToString() != NewName)
 	{
 		Modify();
@@ -88,4 +101,4 @@ void UPCGEditorGraphNode::OnRenameNode(const FString& NewName)
 	}
 }
 
-#undef LOCTEXT_NAMESPACE
+#undef LOCTEXT_NAMESPACE
\ No newline at end of file
diff --git a/Engine/Plugins/PCG/Source/PCGEditor/Private/PCGEditorGraphNode.h b/Engine/Plugins/PCG/Source/PCGEditor/Private/PCGEditorGraphNode.h
--- a/Engine/Plugins/PCG/Source/PCGEditor/Private/PCGEditorGraphNode.h
+++ b/Engine/Plugins/PCG/Source/PCGEditor/Private/PCGEditorGraphNode.h
@@ -24,4 +24,7 @@ public:
 	virtual void AllocateDefaultPins() override;
 	virtual void OnRenameNode(const FString& NewName) override;
 	// ~End UEdGraphNode interface
-};
+
+	virtual bool CanRenameNode(const FString& NewName) const;
+	virtual void CommitNodeRename(const FString& NewName);
+};
\ No newline at end of file
diff --git a/Engine/Plugins/PCG/Source/PCGEditor/Private/PCGEditorGraphNodeReroute.cpp b/Engine/Plugins/PCG/Source/PCGEditor/Private/PCGEditorGraphNodeReroute.cpp
--- a/Engine/Plugins/PCG/Source/PCGEditor/Private/PCGEditorGraphNodeReroute.cpp
+++ b/Engine/Plugins/PCG/Source/PCGEditor/Private/PCGEditorGraphNodeReroute.cpp
@@ -173,41 +173,38 @@ void UPCGEditorGraphNodeNamedRerouteDeclaration::ReconstructNodeOnChange()
 	});
 }
 
-void UPCGEditorGraphNodeNamedRerouteDeclaration::SetNodeName(const UPCGNode* FromNode, FName FromPinName)
+FString UPCGEditorGraphNodeNamedRerouteDeclaration::GenerateNodeName(const UPCGNode* FromNode, FName FromPinName)
 {
-	FString NewName;
-
 	if (FromNode)
 	{
-		NewName = FromNode->GetNodeTitle(EPCGNodeTitleType::ListView).ToString() + " " + FromPinName.ToString();
+		return FromNode->GetNodeTitle(EPCGNodeTitleType::ListView).ToString() + " " + FromPinName.ToString();
 	}
 	else if (FromPinName != NAME_None)
 	{
-		NewName = FromPinName.ToString();
+		return FromPinName.ToString();
 	}
 	else
 	{
-		NewName = TEXT("Reroute");
+		return TEXT("Reroute");
 	}
-
-	bCanRenameNode = true;
-	OnRenameNode(GetCollisionFreeNodeName(NewName));
-	bCanRenameNode = false;
 }
 
 void UPCGEditorGraphNodeNamedRerouteDeclaration::OnRenameNode(const FString& NewName)
 {
-	Super::OnRenameNode(NewName);
-
-	// Propagate the name change to downstream usage nodes
-	ApplyToUsageNodes([&NewName](UPCGEditorGraphNodeNamedRerouteUsage* RerouteNode)
+	const FString ModifiedName = GetCollisionFreeNodeName(NewName);
+	if (CanRenameNode(ModifiedName))
 	{
-		RerouteNode->bCanRenameNode = true;
-		RerouteNode->OnRenameNode(NewName);
-		RerouteNode->bCanRenameNode = false;
-	});
+		CommitNodeRename(ModifiedName);
+		// Propagate the name change to downstream usage nodes
+		ApplyToUsageNodes([&ModifiedName](UPCGEditorGraphNodeNamedRerouteUsage* RerouteNode)
+		{
+			RerouteNode->bCanRenameNode = true;
+			RerouteNode->OnRenameNode(ModifiedName);
+			RerouteNode->bCanRenameNode = false;
+		});
 
-	ReconstructNodeOnChange();
+		ReconstructNodeOnChange();
+	}
 }
 
 void UPCGEditorGraphNodeNamedRerouteDeclaration::OnColorPicked(FLinearColor NewColor)
@@ -232,7 +229,7 @@ void UPCGEditorGraphNodeNamedRerouteDeclaration::FixNodeNameCollision()
 	const FString BaseName = GetNodeTitle(ENodeTitleType::ListView).ToString();
 
 	bCanRenameNode = true;
-	OnRenameNode(GetCollisionFreeNodeName(BaseName));
+	OnRenameNode(BaseName);
 	bCanRenameNode = false;
 }
 
diff --git a/Engine/Plugins/PCG/Source/PCGEditor/Private/PCGEditorGraphNodeReroute.h b/Engine/Plugins/PCG/Source/PCGEditor/Private/PCGEditorGraphNodeReroute.h
--- a/Engine/Plugins/PCG/Source/PCGEditor/Private/PCGEditorGraphNodeReroute.h
+++ b/Engine/Plugins/PCG/Source/PCGEditor/Private/PCGEditorGraphNodeReroute.h
@@ -60,7 +60,7 @@ public:
 	virtual void OnRenameNode(const FString& NewName) override;
 	virtual void PostPaste() override;
 
-	void SetNodeName(const UPCGNode* FromNode, FName FromPinName);
+	FString GenerateNodeName(const UPCGNode* FromNode, FName FromPinName);
 	void FixNodeNameCollision();
 	FString GetCollisionFreeNodeName(const FString& BaseName) const;
 
diff --git a/Engine/Plugins/PCG/Source/PCGEditor/Private/PCGEditorGraphSchemaActions.cpp b/Engine/Plugins/PCG/Source/PCGEditor/Private/PCGEditorGraphSchemaActions.cpp
--- a/Engine/Plugins/PCG/Source/PCGEditor/Private/PCGEditorGraphSchemaActions.cpp
+++ b/Engine/Plugins/PCG/Source/PCGEditor/Private/PCGEditorGraphSchemaActions.cpp
@@ -760,9 +760,10 @@ UEdGraphNode* FPCGEditorGraphSchemaAction_NewNamedRerouteDeclaration::PerformAct
 		NewNode->AutowireNewNode(FromPin);
 	}
 
-	NewNode->SetNodeName(FromNode, FromPin ? FromPin->PinName : NAME_None);
+	const FString NewName = NewNode->GenerateNodeName(FromNode, FromPin ? FromPin->PinName : NAME_None);
+	NewNode->OnRenameNode(NewName);
 
 	return NewNode;
 }
 
-#undef LOCTEXT_NAMESPACE
+#undef LOCTEXT_NAMESPACE
\ No newline at end of file
