[PCG]
- PCG preview actors now set the RF_NonPIEDuplicateTransient on transient preview actrors so they are properly handled in PIE.
- Todo: PCG preview components should set the AActor::bIgnoreInPIE on generated actors when generating preview actors.

diff --git a/Engine/Plugins/PCG/Source/PCG/Private/Elements/PCGCreateTargetActor.cpp b/Engine/Plugins/PCG/Source/PCG/Private/Elements/PCGCreateTargetActor.cpp
--- a/Engine/Plugins/PCG/Source/PCG/Private/Elements/PCGCreateTargetActor.cpp
+++ b/Engine/Plugins/PCG/Source/PCG/Private/Elements/PCGCreateTargetActor.cpp
@@ -302,7 +302,7 @@ bool FPCGCreateTargetActorElement::ExecuteInternal(FPCGContext* Context) const
 
 	if (PCGHelpers::IsRuntimeOrPIE() || (Context->SourceComponent.IsValid() && Context->SourceComponent->IsInPreviewMode()))
 	{
-		SpawnParams.ObjectFlags |= RF_Transient;
+		SpawnParams.ObjectFlags |= RF_Transient | RF_NonPIEDuplicateTransient;
 	}
 
 	FTransform Transform = TargetActor->GetTransform();
@@ -366,4 +366,4 @@ bool FPCGCreateTargetActorElement::ExecuteInternal(FPCGContext* Context) const
 	return true;
 }
 
-#undef LOCTEXT_NAMESPACE
\ No newline at end of file
+#undef LOCTEXT_NAMESPACE
diff --git a/Engine/Plugins/PCG/Source/PCG/Private/Elements/PCGSpawnActor.cpp b/Engine/Plugins/PCG/Source/PCG/Private/Elements/PCGSpawnActor.cpp
--- a/Engine/Plugins/PCG/Source/PCG/Private/Elements/PCGSpawnActor.cpp
+++ b/Engine/Plugins/PCG/Source/PCG/Private/Elements/PCGSpawnActor.cpp
@@ -806,7 +806,7 @@ void FPCGSpawnActorElement::SpawnActors(FPCGSubgraphContext* Context, AActor* Ta
 
 	if (PCGHelpers::IsRuntimeOrPIE() || (Context->SourceComponent.IsValid() && Context->SourceComponent->IsInPreviewMode()))
 	{
-		SpawnParams.ObjectFlags |= RF_Transient;
+		SpawnParams.ObjectFlags |= RF_Transient | RF_NonPIEDuplicateTransient;
 	}
 
 	UPCGActorHelpers::FSpawnDefaultActorParams SpawnDefaultActorParams(TargetActor->GetWorld(), InTemplateActorClass, FTransform::Identity, SpawnParams);
diff --git a/Engine/Plugins/PCG/Source/PCG/Private/Grid/PCGPartitionActor.cpp b/Engine/Plugins/PCG/Source/PCG/Private/Grid/PCGPartitionActor.cpp
--- a/Engine/Plugins/PCG/Source/PCG/Private/Grid/PCGPartitionActor.cpp
+++ b/Engine/Plugins/PCG/Source/PCG/Private/Grid/PCGPartitionActor.cpp
@@ -507,7 +507,7 @@ void APCGPartitionActor::AddGraphInstance(UPCGComponent* OriginalComponent)
 	Modify(!OriginalComponent->IsInPreviewMode());
 
 	// Create a new local component
-	LocalComponent = NewObject<UPCGComponent>(this, NAME_None, OriginalComponent->IsInPreviewMode() ? RF_Transient : RF_NoFlags);
+	LocalComponent = NewObject<UPCGComponent>(this, NAME_None, OriginalComponent->IsInPreviewMode() ? RF_Transient | RF_NonPIEDuplicateTransient : RF_NoFlags);
 	LocalComponent->MarkAsLocalComponent();
 	LocalComponent->SetGenerationGridSize(PCGGridSize);
 
@@ -750,4 +750,4 @@ FString APCGPartitionActor::GetPCGPartitionActorName(uint32 GridSize, const FInt
 FString APCGPartitionActor::GetPCGPartitionActorName(const FPCGGridDescriptor& GridDescriptor, const FIntVector& GridCoords)
 {
 	return GridDescriptor.GetPartitionActorName(GridCoords);
-}
\ No newline at end of file
+}
diff --git a/Engine/Plugins/PCG/Source/PCG/Private/Helpers/PCGActorHelpers.cpp b/Engine/Plugins/PCG/Source/PCG/Private/Helpers/PCGActorHelpers.cpp
--- a/Engine/Plugins/PCG/Source/PCG/Private/Helpers/PCGActorHelpers.cpp
+++ b/Engine/Plugins/PCG/Source/PCG/Private/Helpers/PCGActorHelpers.cpp
@@ -181,7 +181,7 @@ UPCGManagedISMComponent* UPCGActorHelpers::GetOrCreateManagedISMC(AActor* InTarg
 
 	ComponentName += StaticMesh->GetName();
 
-	const EObjectFlags ObjectFlags = (InSourceComponent->IsInPreviewMode() ? RF_Transient : RF_NoFlags);
+	const EObjectFlags ObjectFlags = (InSourceComponent->IsInPreviewMode() ? RF_Transient | RF_NonPIEDuplicateTransient : RF_NoFlags);
 	UInstancedStaticMeshComponent* ISMC = NewObject<UInstancedStaticMeshComponent>(InTargetActor, Descriptor.ComponentClass, MakeUniqueObjectName(InTargetActor, Descriptor.ComponentClass, FName(ComponentName)), ObjectFlags);
 	Descriptor.InitComponent(ISMC);
 	ISMC->SetNumCustomDataFloats(InParams.NumCustomDataFloats);
@@ -285,7 +285,7 @@ UPCGManagedSplineMeshComponent* UPCGActorHelpers::GetOrCreateManagedSplineMeshCo
 	InTargetActor->Modify(!InSourceComponent->IsInPreviewMode());
 
 	FString ComponentName = TEXT("PCGSplineMeshComponent_") + StaticMesh->GetName();
-	const EObjectFlags ObjectFlags = (InSourceComponent->IsInPreviewMode() ? RF_Transient : RF_NoFlags);
+	const EObjectFlags ObjectFlags = (InSourceComponent->IsInPreviewMode() ? RF_Transient | RF_NonPIEDuplicateTransient : RF_NoFlags);
 	USplineMeshComponent* SplineMeshComponent = NewObject<USplineMeshComponent>(InTargetActor, MakeUniqueObjectName(InTargetActor, USplineMeshComponent::StaticClass(), FName(ComponentName)), ObjectFlags);
 
 	// Init Component
@@ -522,7 +522,7 @@ AActor* UPCGActorHelpers::SpawnDefaultActor(const FSpawnDefaultActorParams& Para
 
 	if (PCGHelpers::IsRuntimeOrPIE())
 	{
-		SpawnParams.ObjectFlags |= RF_Transient;
+		SpawnParams.ObjectFlags |= RF_Transient | RF_NonPIEDuplicateTransient;
 	}
 
 #if WITH_EDITOR
