[PCG] Fix Graph Determinism Test bug causing engine to hang when running tests that required waiting on async or other framebound tests.

#rnx

diff --git a/Engine/Plugins/PCG/Source/PCG/Private/Compute/PCGDataForGPU.cpp b/Engine/Plugins/PCG/Source/PCG/Private/Compute/PCGDataForGPU.cpp
--- a/Engine/Plugins/PCG/Source/PCG/Private/Compute/PCGDataForGPU.cpp
+++ b/Engine/Plugins/PCG/Source/PCG/Private/Compute/PCGDataForGPU.cpp
@@ -810,12 +810,16 @@ void FPCGDataCollectionDesc::PackDataCollection(const FPCGDataCollection& InData
 		if (const UPCGPointData* PointData = Cast<UPCGPointData>(InputData[DataIndex].Data))
 		{
 			const TArray<FPCGPoint>& Points = PointData->GetPoints();
+			if (Points.IsEmpty())
+			{
+				continue;
+			}
+
 			const uint32 NumElements = Points.Num();
 
 			for (const FPCGKernelAttributeDesc& AttributeDesc : DataDesc.AttributeDescs)
 			{
 				const uint32 AttributeId = AttributeDesc.Index;
-				const uint32 AttributeStrideBytes = PCGDataForGPUHelpers::GetAttributeTypeStrideBytes(AttributeDesc.Type);
 
 				const FPCGMetadataAttributeBase* AttributeBase = (AttributeId >= NUM_RESERVED_ATTRS) ? Metadata->GetConstAttribute(AttributeDesc.Name) : nullptr;
 
@@ -1237,4 +1241,4 @@ uint32 FPCGDataCollectionDesc::ComputeDataElementCount(EPCGDataType InDataType)
 void FPCGDataCollectionDesc::Combine(const FPCGDataCollectionDesc& Other)
 {
 	DataDescs.Append(Other.DataDescs);
-}
+}
\ No newline at end of file
diff --git a/Engine/Plugins/PCG/Source/PCG/Private/Graph/PCGGraphExecutor.cpp b/Engine/Plugins/PCG/Source/PCG/Private/Graph/PCGGraphExecutor.cpp
--- a/Engine/Plugins/PCG/Source/PCG/Private/Graph/PCGGraphExecutor.cpp
+++ b/Engine/Plugins/PCG/Source/PCG/Private/Graph/PCGGraphExecutor.cpp
@@ -591,7 +591,7 @@ TSet<UPCGComponent*> FPCGGraphExecutor::Cancel(TFunctionRef<bool(TWeakObjectPtr<
 	return CancelledComponents;
 }
 
-bool FPCGGraphExecutor::IsGraphCurrentlyExecuting(UPCGGraph* InGraph)
+bool FPCGGraphExecutor::IsGraphCurrentlyExecuting(const UPCGGraph* InGraph)
 {
 	bool bAnyPresent = false;
 
@@ -3096,4 +3096,4 @@ FPCGElementPtr UPCGGridLinkageSettings::CreateElement() const
 	};
 
 	return MakeShared<PCGGraphExecutor::FPCGGridLinkageElement>(GridLinkageOperation, ContextAllocator, FromGrid, ToGrid, GenerationGrid, ResourceKey, UpstreamPin.Get());
-}
+}
\ No newline at end of file
diff --git a/Engine/Plugins/PCG/Source/PCG/Private/Graph/PCGGraphExecutor.h b/Engine/Plugins/PCG/Source/PCG/Private/Graph/PCGGraphExecutor.h
--- a/Engine/Plugins/PCG/Source/PCG/Private/Graph/PCGGraphExecutor.h
+++ b/Engine/Plugins/PCG/Source/PCG/Private/Graph/PCGGraphExecutor.h
@@ -90,7 +90,7 @@ public:
 	TArray<UPCGComponent*> CancelAll();
 
 	/** Returns true if any task is scheduled or executing for the given graph. */
-	bool IsGraphCurrentlyExecuting(UPCGGraph* InGraph);
+	bool IsGraphCurrentlyExecuting(const UPCGGraph* InGraph);
 
 	/** Returns true if any task is scheduled or executing for any graph */
 	bool IsAnyGraphCurrentlyExecuting() const;
diff --git a/Engine/Plugins/PCG/Source/PCG/Private/Tests/Determinism/PCGDeterminismTestsCommon.cpp b/Engine/Plugins/PCG/Source/PCG/Private/Tests/Determinism/PCGDeterminismTestsCommon.cpp
--- a/Engine/Plugins/PCG/Source/PCG/Private/Tests/Determinism/PCGDeterminismTestsCommon.cpp
+++ b/Engine/Plugins/PCG/Source/PCG/Private/Tests/Determinism/PCGDeterminismTestsCommon.cpp
@@ -25,10 +25,15 @@
 
 #define LOCTEXT_NAMESPACE "PCGDeterminism"
 
+static TAutoConsoleVariable<int32> CVarDeterminismGraphTestCycleLimit(
+	TEXT("pcg.DeterminismGraphTestCycleLimit"),
+	10,
+	TEXT("How many execution attempts before discarding a graph test."));
+
 static TAutoConsoleVariable<int32> CVarDeterminismPermutationLimit(
 	TEXT("pcg.DeterminismPermutationLimit"),
 	10000,
-	TEXT("A limit for the maximium amount of input permutations before autofailing a basic determinism test"));
+	TEXT("A limit for the maximum amount of input permutations before auto-failing a basic determinism test."));
 
 namespace PCGDeterminismTests
 {
@@ -82,14 +87,22 @@ namespace PCGDeterminismTests
 		OutResult.DataTypesTested |= InputPCGData->GetDataType();
 
 		// Clone the actor and component
-		const AActor* PCGActor = InPCGComponent->GetOwner();
+		AActor* PCGActor = InPCGComponent->GetOwner();
 		check(PCGActor);
 
-		AActor* PCGActorCopy = DuplicateObject<AActor>(PCGActor, GetTransientPackage());
-		PCGActorCopy->SetFlags(RF_Transient);
+		FObjectDuplicationParameters DuplicateParams(PCGActor, PCGActor->GetOuter());
+		DuplicateParams.ApplyFlags = RF_Transient;
+
+		AActor* PCGActorCopy = nullptr;
+		{
+			FGCScopeGuard Scope;
+			PCGActorCopy = Cast<AActor>(StaticDuplicateObjectEx(DuplicateParams));
+		}
+
+		PCGActorCopy->AddToRoot();
+
 		UPCGComponent* PCGComponentCopy = PCGActorCopy->FindComponentByClass<UPCGComponent>();
 		check(PCGComponentCopy);
-		PCGComponentCopy->SetFlags(RF_Transient);
 		PCGComponentCopy->SetIsPartitioned(false);
 		if (PCGActor->GetWorld())
 		{
@@ -98,7 +111,26 @@ namespace PCGDeterminismTests
 
 		FPCGGraphExecutor Executor;
 
-		auto ScheduleAndWaitForExecution = [&Executor, PCGComponentCopy](FPCGTaskId FinalTaskID)
+		ON_SCOPE_EXIT
+		{
+			if (Executor.IsGraphCurrentlyExecuting(InPCGGraph))
+			{
+				Executor.CancelAll();
+			}
+
+			// Clean up anything generated
+			PCGComponentCopy->Cleanup();
+			if (PCGComponentCopy->IsRegistered())
+			{
+				PCGComponentCopy->UnregisterComponent();
+			}
+
+			// Clean up the debug actor
+			PCGActorCopy->RemoveFromRoot();
+			PCGActorCopy->MarkAsGarbage();
+		};
+
+		auto ScheduleAndWaitForExecution = [&Executor, PCGComponentCopy](FPCGTaskId FinalTaskID) -> bool
 		{
 			// TODO: Consider randomizing/iterating through possible input orders
 			volatile bool bTasksComplete = false;
@@ -113,11 +145,25 @@ namespace PCGDeterminismTests
 				return true;
 			}, PCGComponentCopy, {FinalTaskID});
 
-			// Run first iteration
+			// TODO: Reshape the test framework to run a graph independently and use the callback events from the context.
+			/* Implementor's Note: Async loading or other frame bound activities can cause this to stall out, so the
+			 * cycle limit will stop early if needed.
+			 **/
+			int32 ExecutionCount = 0;
 			while (!bTasksComplete)
 			{
-				Executor.Execute();
+				if (ExecutionCount++ < CVarDeterminismGraphTestCycleLimit.GetValueOnAnyThread())
+				{
+					Executor.Execute();
+				}
+				else
+				{
+					Executor.CancelAll();
+					return false;
+				}
 			}
+
+			return true;
 		};
 
 		TMap<FPCGTaskId, TTuple<const UPCGNode*, const FPCGDataCollection>> IntermediateOutputArray;
@@ -131,18 +177,39 @@ namespace PCGDeterminismTests
 
 		// Schedule the graph execution
 		FPCGTaskId FinalTaskID = Executor.ScheduleDebugWithTaskCallback(PCGComponentCopy, RecordMappedOutput);
-		ScheduleAndWaitForExecution(FinalTaskID);
 
-		// Copy the first array over, so it can be used again
-		TMap<FPCGTaskId, TTuple<const UPCGNode*, const FPCGDataCollection>> FirstIntermediateOutputArray = IntermediateOutputArray;
+		TMap<FPCGTaskId, TTuple<const UPCGNode*, const FPCGDataCollection>> FirstIntermediateOutputArray;
+		bool bExecutionCompleted = ScheduleAndWaitForExecution(FinalTaskID);
+		if (bExecutionCompleted)
+		{
+			// Copy the first array over, so it can be used again
+			FirstIntermediateOutputArray = IntermediateOutputArray;
 
-		// Clean up generated content
-		IntermediateOutputArray.Empty();
-		PCGComponentCopy->Cleanup();
+			// Clean up generated content
+			IntermediateOutputArray.Empty();
+			PCGComponentCopy->Cleanup();
+
+			// Run again
+			FinalTaskID = Executor.ScheduleDebugWithTaskCallback(PCGComponentCopy, RecordMappedOutput);
+			bExecutionCompleted &= ScheduleAndWaitForExecution(FinalTaskID);
+		}
 
-		// Run again
-		FinalTaskID = Executor.ScheduleDebugWithTaskCallback(PCGComponentCopy, RecordMappedOutput);
-		ScheduleAndWaitForExecution(FinalTaskID);
+		// Early out if one of the executions was unable to complete
+		if (!bExecutionCompleted)
+		{
+			OutResult.bFlagRaised = true;
+			UpdateTestResults(Defaults::GraphResultName, OutResult, EDeterminismLevel::NoDeterminism);
+			OutResult.AdditionalDetails.Add(TEXT("Required async or timed out"));
+			return;
+		}
+
+		if (IntermediateOutputArray.IsEmpty() || FirstIntermediateOutputArray.IsEmpty())
+		{
+			OutResult.bFlagRaised = true;
+			UpdateTestResults(Defaults::GraphResultName, OutResult, EDeterminismLevel::NoDeterminism);
+			OutResult.AdditionalDetails.Add(TEXT("No output was generated"));
+			return;
+		}
 
 		// Early out for mismatched outputs
 		if (IntermediateOutputArray.Num() != FirstIntermediateOutputArray.Num())
@@ -163,7 +230,7 @@ namespace PCGDeterminismTests
 			const UPCGNode* NodePtr = NodeDataTuple.Get<0>();
 			const FString NodeNameString = NodePtr->GetName();
 
-			UE_LOG(LogPCG, Log, TEXT("[%s] Evaluating Node [%d]..."), *InPCGGraph->GetName(), *NodeNameString);
+			UE_LOG(LogPCG, Log, TEXT("[%s] Evaluating Node [%s]..."), *InPCGGraph->GetName(), *NodeNameString);
 
 			const FPCGDataCollection& FirstOutput = NodeDataTuple.Get<1>();
 			const FPCGDataCollection& SecondOutput = IntermediateOutputArray[TaskIdKey].Get<1>();
@@ -207,13 +274,6 @@ namespace PCGDeterminismTests
 			}
 		}
 
-		// Clean up anything generated
-		PCGComponentCopy->Cleanup();
-		if (PCGComponentCopy->IsRegistered())
-		{
-			PCGComponentCopy->UnregisterComponent();
-		}
-
 		// Finalize the results
 		UpdateTestResults(Defaults::GraphResultName, OutResult, HighestDeterminismLevel);
 	}
@@ -1588,4 +1648,4 @@ namespace PCGDeterminismTests
 	}
 }
 
-#undef LOCTEXT_NAMESPACE
+#undef LOCTEXT_NAMESPACE
\ No newline at end of file
