[PCG] Support Actor reuse in CreateTargetActor node
- PCGSettings now excludes UID (which changes between engine/editor runs) and uses its ObjectPath crc as a stable id when computing CacheCRC
- Add Stack Crc to CreateTargetActor Crc computation to give the Crc some context
- CVar: pcg.CreateTargetActor.AllowReuse (defaults to false, but will be turned on soon enough after some more testing)
- Fix successive component refresh calls preventing proper reuse because managed resources were getting released to early

#rnx

diff --git a/Engine/Plugins/PCG/Source/PCG/Private/Elements/PCGCreateTargetActor.cpp b/Engine/Plugins/PCG/Source/PCG/Private/Elements/PCGCreateTargetActor.cpp
--- a/Engine/Plugins/PCG/Source/PCG/Private/Elements/PCGCreateTargetActor.cpp
+++ b/Engine/Plugins/PCG/Source/PCG/Private/Elements/PCGCreateTargetActor.cpp
@@ -4,6 +4,7 @@
 
 #include "PCGComponent.h"
 #include "PCGContext.h"
+#include "PCGCrc.h"
 #include "PCGElement.h"
 #include "PCGManagedResource.h"
 #include "PCGModule.h"
@@ -15,7 +16,9 @@
 #include "Engine/Blueprint.h"
 #include "Engine/BlueprintGeneratedClass.h"
 #include "Engine/Engine.h"
+#include "Engine/Level.h"
 #include "GameFramework/Actor.h"
+#include "Serialization/ArchiveCrc32.h"
 #include "UObject/Package.h"
 
 #if WITH_EDITOR
@@ -41,6 +44,59 @@ namespace PCGCreateTargetActorConstants
 	const FText ActorPropertyOverridesTooltip = LOCTEXT("ActorOverrideToolTip", "Provide property overrides for the created target actor. The attribute name must match the InputSource name in the actor property override description.");
 }
 
+namespace PCGCreateTargetActor
+{
+	static TAutoConsoleVariable<bool> CVarCreateTargetActorAllowReuse(
+		TEXT("pcg.CreateTargetActor.AllowReuse"),
+		false,
+		TEXT("Controls whether PCG Create Target Actor node can reuse actors when re-executing (requires Create Target Actor node resave so they have a Stable Reuse Guid)"));
+
+	FPCGCrc GetAdditionalDependenciesCrc(const UPCGActorHelpers::FSpawnDefaultActorParams& InParams, AActor* TargetActor)
+	{
+		FArchiveCrc32 Ar;
+
+		Ar << TargetActor;
+
+		// Do not CRC Everything as some of those params come from the Settings which is already CRCed
+		UObject* Level = InParams.SpawnParams.OverrideLevel;
+		Ar << Level;
+
+		int32 ObjectFlags = InParams.SpawnParams.ObjectFlags;
+		Ar << ObjectFlags;
+				
+		// Include transform - round sufficiently to ensure stability
+		FVector TransformLocation = InParams.Transform.GetLocation();
+		FIntVector Location(FMath::RoundToInt(TransformLocation.X), FMath::RoundToInt(TransformLocation.Y), FMath::RoundToInt(TransformLocation.Z));
+		Ar << Location;
+
+		FRotator Rotator(InParams.Transform.Rotator().GetDenormalized());
+		const int32 MAX_DEGREES = 360;
+		FIntVector Rotation(FMath::RoundToInt(Rotator.Pitch) % MAX_DEGREES, FMath::RoundToInt(Rotator.Yaw) % MAX_DEGREES, FMath::RoundToInt(Rotator.Roll) % MAX_DEGREES);
+		Ar << Rotation;
+
+		FVector TransformScale = InParams.Transform.GetScale3D();
+		const float SCALE_FACTOR = 100;
+		FIntVector Scale(FMath::RoundToInt(TransformScale.X * SCALE_FACTOR), FMath::RoundToInt(TransformScale.Y * SCALE_FACTOR), FMath::RoundToInt(TransformScale.Z * SCALE_FACTOR));
+		Ar << Scale;
+				
+#if WITH_EDITOR
+		TArray<const UDataLayerInstance*> DataLayerInstances = InParams.DataLayerInstances;
+		DataLayerInstances.Sort([](const UDataLayerInstance& A, const UDataLayerInstance& B)
+		{
+			return A.GetDataLayerFullName() < B.GetDataLayerFullName();
+		});
+
+		for (const UDataLayerInstance* DataLayerInstance : DataLayerInstances)
+		{
+			UDataLayerInstance* NonConstDataLayerInstance = const_cast<UDataLayerInstance*>(DataLayerInstance);
+			Ar << NonConstDataLayerInstance;
+		}
+#endif
+		
+		return FPCGCrc(Ar.GetCrc());
+	}
+}
+
 UPCGCreateTargetActor::UPCGCreateTargetActor(const FObjectInitializer& ObjectInitializer)
 	: UPCGSettings(ObjectInitializer)
 {
@@ -317,39 +373,98 @@ bool FPCGCreateTargetActorElement::ExecuteInternal(FPCGContext* Context) const
 	SpawnDefaultActorParams.DataLayerInstances = TargetActor->GetDataLayerInstances();
 #endif
 
-	AActor* GeneratedActor = UPCGActorHelpers::SpawnDefaultActor(SpawnDefaultActorParams);
+	UPCGManagedActors* ReusableResource = nullptr;
+	FPCGCrc ResourceCrc;
 
-	if (!GeneratedActor)
+	if (PCGCreateTargetActor::CVarCreateTargetActorAllowReuse.GetValueOnAnyThread())
 	{
-		PCGE_LOG(Error, GraphAndLog, LOCTEXT("ActorSpawnFailed", "Failed to spawn actor"));
-		return true;
-	}
+		if (!Context->DependenciesCrc.IsValid())
+		{
+			// TODO: we should be able to use the InputData to compute Crcs but this Crc contains a non stable UID
+			// since CreateTargetActor settings are already overriden by the inputdata, it doesn't matter here and we can just ignore it
+			// but for future reference if we want to have better reusing and stable generation across, we need to fix this.
+			FPCGDataCollection EmptyCollection;
+			GetDependenciesCrc(EmptyCollection, Settings, Context->SourceComponent.Get(), Context->DependenciesCrc);
+		}
 
-	// Always attach if root actor is provided
-	PCGHelpers::AttachToParent(GeneratedActor, TargetActor, Settings->RootActor.Get() ? EPCGAttachOptions::Attached : Settings->AttachOptions, Context);
+		if (Context->DependenciesCrc.IsValid())
+		{
+			ResourceCrc = Context->DependenciesCrc;
 
-#if WITH_EDITOR
-	if (Settings->ActorLabel != FString())
+			const FPCGCrc StackCRC = Context->Stack->GetCrc();
+			ResourceCrc.Combine(StackCRC);
+
+			const FPCGCrc AdditionalCRC = PCGCreateTargetActor::GetAdditionalDependenciesCrc(SpawnDefaultActorParams, TargetActor);
+			ResourceCrc.Combine(AdditionalCRC);
+
+			Context->SourceComponent->ForEachManagedResource([&ReusableResource, ResourceCrc, &Context](UPCGManagedResource* InResource)
+			{
+				if (ReusableResource)
+				{
+					return;
+				}
+
+				if (UPCGManagedActors* Resource = Cast<UPCGManagedActors>(InResource))
+				{
+					if (Resource->GetCrc().IsValid() && Resource->GetCrc() == ResourceCrc && Resource->GeneratedActors.Num() == 1 && Resource->GeneratedActors.Array()[0].IsValid())
+					{
+						ReusableResource = Resource;
+					}
+				}
+			});
+		}
+	}
+		
+	AActor* GeneratedActor = nullptr;
+	if (ReusableResource)
 	{
-		GeneratedActor->SetActorLabel(Settings->ActorLabel);
+		ReusableResource->MarkAsReused();
+		GeneratedActor = ReusableResource->GeneratedActors.Array()[0].Get();
+		ensure(GeneratedActor->Tags.Contains(PCGHelpers::DefaultPCGActorTag));
 	}
-#endif
+	else
+	{
+		GeneratedActor = UPCGActorHelpers::SpawnDefaultActor(SpawnDefaultActorParams);
+		if (!GeneratedActor)
+		{
+			PCGE_LOG(Error, GraphAndLog, LOCTEXT("ActorSpawnFailed", "Failed to spawn actor"));
+			return true;
+		}
 
-	GeneratedActor->Tags.Add(PCGHelpers::DefaultPCGActorTag);
+		GeneratedActor->Tags.Add(PCGHelpers::DefaultPCGActorTag);
 
-	// Apply property overrides to the GeneratedActor
-	PCGObjectPropertyOverrideHelpers::ApplyOverridesFromParams(Settings->PropertyOverrideDescriptions, GeneratedActor, PCGCreateTargetActorConstants::ActorPropertyOverridesLabel, Context);
+#if WITH_EDITOR
+		if (!Settings->ActorLabel.IsEmpty())
+		{
+			GeneratedActor->SetActorLabel(Settings->ActorLabel);
+		}
+#endif
+
+		// Always attach if root actor is provided
+		PCGHelpers::AttachToParent(GeneratedActor, TargetActor, Settings->RootActor.Get() ? EPCGAttachOptions::Attached : Settings->AttachOptions, Context);
 
+		// Apply property overrides to the GeneratedActor
+		PCGObjectPropertyOverrideHelpers::ApplyOverridesFromParams(Settings->PropertyOverrideDescriptions, GeneratedActor, PCGCreateTargetActorConstants::ActorPropertyOverridesLabel, Context);
+	}
+	
 	for (UFunction* Function : PCGHelpers::FindUserFunctions(GeneratedActor->GetClass(), Settings->PostProcessFunctionNames, { UPCGFunctionPrototypes::GetPrototypeWithNoParams() }, Context))
 	{
 		GeneratedActor->ProcessEvent(Function, nullptr);
 	}
 
-	if (UPCGComponent* SourceComponent = Context->SourceComponent.Get())
+	// Create Resource if it isn't reused
+	if (!ReusableResource)
 	{
-		UPCGManagedActors* ManagedActors = NewObject<UPCGManagedActors>(SourceComponent);
-		ManagedActors->GeneratedActors.Add(GeneratedActor);
-		SourceComponent->AddToManagedResources(ManagedActors);
+		if (UPCGComponent* SourceComponent = Context->SourceComponent.Get())
+		{
+			UPCGManagedActors* ManagedActors = NewObject<UPCGManagedActors>(SourceComponent);
+			if (ResourceCrc.IsValid())
+			{
+				ManagedActors->SetCrc(ResourceCrc);
+			}
+			ManagedActors->GeneratedActors.Add(GeneratedActor);
+			SourceComponent->AddToManagedResources(ManagedActors);
+		}
 	}
 
 	// Create param data output with reference to actor
diff --git a/Engine/Plugins/PCG/Source/PCG/Private/Graph/PCGStackContext.cpp b/Engine/Plugins/PCG/Source/PCG/Private/Graph/PCGStackContext.cpp
--- a/Engine/Plugins/PCG/Source/PCG/Private/Graph/PCGStackContext.cpp
+++ b/Engine/Plugins/PCG/Source/PCG/Private/Graph/PCGStackContext.cpp
@@ -11,6 +11,7 @@
 #include "Algo/Find.h"
 #include "Containers/UnrealString.h"
 #include "Misc/StringBuilder.h"
+#include "Serialization/ArchiveCrc32.h"
 
 #include UE_INLINE_GENERATED_CPP_BY_NAME(PCGStackContext)
 
@@ -199,6 +200,27 @@ bool FPCGStack::operator==(const FPCGStack& Other) const
 	return true;
 }
 
+FPCGCrc FPCGStack::GetCrc() const
+{
+	FArchiveCrc32 Ar;
+		
+	for (const FPCGStackFrame& StackFrame : StackFrames)
+	{
+		if (!StackFrame.Object.IsNull())
+		{
+			TSoftObjectPtr<const UObject> SoftObjectPtr = StackFrame.Object;
+			Ar << SoftObjectPtr;
+		}
+		else
+		{
+			int32 LoopIndex = StackFrame.LoopIndex;
+			Ar << LoopIndex;
+		}
+	}
+
+	return FPCGCrc(Ar.GetCrc());
+}
+
 int32 FPCGStackContext::PushFrame(const UObject* InFrameObject)
 {
 	if (CurrentStackIndex == INDEX_NONE)
@@ -302,4 +324,4 @@ void FPCGStackContext::SetGraphExecutionTaskId(FPCGTaskId InGraphExecutionTaskId
 	{
 		Stack.GraphExecutionTaskId = GraphExecutionTaskId;
 	}
-}
\ No newline at end of file
+}
diff --git a/Engine/Plugins/PCG/Source/PCG/Private/PCGComponent.cpp b/Engine/Plugins/PCG/Source/PCG/Private/PCGComponent.cpp
--- a/Engine/Plugins/PCG/Source/PCG/Private/PCGComponent.cpp
+++ b/Engine/Plugins/PCG/Source/PCG/Private/PCGComponent.cpp
@@ -752,7 +752,7 @@ void UPCGComponent::PostCleanupGraph()
 	PCGComponent::BroadcastDynamicDelegate(OnPCGGraphCleanedExternal, this);
 }
 
-void UPCGComponent::OnProcessGraphAborted(bool bQuiet)
+void UPCGComponent::OnProcessGraphAborted(bool bQuiet, bool bCleanupUnusedResources)
 {
 	if (!bQuiet)
 	{
@@ -764,7 +764,10 @@ void UPCGComponent::OnProcessGraphAborted(bool bQuiet)
 	ResetIgnoredChangeOrigins(/*bLogIfAnyPresent=*/false);
 #endif
 
-	CleanupUnusedManagedResources();
+	if (bCleanupUnusedResources)
+	{
+		CleanupUnusedManagedResources();
+	}
 
 	CurrentGenerationTask = InvalidPCGTaskId;
 	CurrentCleanupTask = InvalidPCGTaskId; // this is needed to support cancellation
@@ -2625,14 +2628,16 @@ void UPCGComponent::Refresh(EPCGChangeType ChangeType, bool bCancelExistingRefre
 			CurrentRefreshTask = InvalidPCGTaskId;
 		}
 
+		const bool bScheduleRefresh = CurrentRefreshTask == InvalidPCGTaskId && CurrentCleanupTask == InvalidPCGTaskId;
+
 		if (bNeedToCancelCurrentTasks)
 		{
-			Subsystem->CancelGeneration(this);
+			Subsystem->CancelGeneration(this, /*bCleanupManagedResources=*/!bScheduleRefresh);
 		}
 
 		// Calling a new refresh here might not be sufficient; if the current component was generating but was not previously generated,
 		// then the bGenerated flag will be false, which will prevent a subsequent update here
-		if (CurrentRefreshTask == InvalidPCGTaskId && CurrentCleanupTask == InvalidPCGTaskId)
+		if (bScheduleRefresh)
 		{
 			CurrentRefreshTask = Subsystem->ScheduleRefresh(this, bGenerationWasInProgress);
 		}
diff --git a/Engine/Plugins/PCG/Source/PCG/Private/PCGContext.cpp b/Engine/Plugins/PCG/Source/PCG/Private/PCGContext.cpp
--- a/Engine/Plugins/PCG/Source/PCG/Private/PCGContext.cpp
+++ b/Engine/Plugins/PCG/Source/PCG/Private/PCGContext.cpp
@@ -401,6 +401,9 @@ void FPCGContext::OverrideSettings()
 			OverriddenParams.Add(&Param);
 		}
 	}
+
+	// Make sure CacheCrc is up to date
+	SettingsWithOverride->CacheCrc();
 }
 
 bool FPCGContext::IsValueOverriden(const FName PropertyName)
diff --git a/Engine/Plugins/PCG/Source/PCG/Private/PCGSettings.cpp b/Engine/Plugins/PCG/Source/PCG/Private/PCGSettings.cpp
--- a/Engine/Plugins/PCG/Source/PCG/Private/PCGSettings.cpp
+++ b/Engine/Plugins/PCG/Source/PCG/Private/PCGSettings.cpp
@@ -21,6 +21,7 @@
 #include "Metadata/Accessors/PCGAttributeAccessorHelpers.h"
 #include "Metadata/Accessors/IPCGAttributeAccessor.h"
 
+#include "Misc/Crc.h"
 #include "Serialization/ArchiveObjectCrc32.h"
 #include "UObject/ObjectSaveContext.h"
 #include "UObject/FortniteMainBranchObjectVersion.h"
@@ -54,18 +55,27 @@ namespace PCGSettings
 }
 
 /** Custom Crc computation that ignores properties that will not affect the computed result of a node. */
-class FPCGSettingsObjectCrc32 : public FArchiveObjectCrc32
+class FPCGSettingsObjectCrc : public FArchiveObjectCrc32
 {
 public:
+	static FPCGCrc PCGCrc(const UPCGSettings* PCGSettings)
+	{
+		// If Settings has an original, use it to get a unique Id
+		const UPCGSettings* OriginalSettings = PCGSettings->OriginalSettings ? PCGSettings->OriginalSettings : PCGSettings;
+				
+		// If Settings is not transient, use its Object Path as a unique Id (this allows persistent Crcs)
+		const bool bSkipUID = !OriginalSettings->GetPackage()->HasAnyPackageFlags(PKG_TransientFlags) && OriginalSettings->GetPackage() != GetTransientPackage();
+		const uint32 SettingsPathOrUIDCrc = bSkipUID ? FCrc::StrCrc32(*OriginalSettings->GetPathName()) : FCrc::TypeCrc32(OriginalSettings->GetStableUID());
+				
+		FPCGSettingsObjectCrc Ar;
+		uint32 ArCrc = Ar.Crc32(const_cast<UPCGSettings*>(PCGSettings), SettingsPathOrUIDCrc);
+		return FPCGCrc(ArCrc);
+	}
+
 #if WITH_EDITOR
 	virtual bool ShouldSkipProperty(const FProperty* InProperty) const override
 	{
-		// Currently we rely on the 'UID' property getting included in the Crc. An example of this are asset settings
-		// for which we avoid doing a full data CRC and instead rely on including hte UID. This property is transient
-		// and will only serialize if IsPersistent() is false (see tests in FProperty::ShouldSerializeValue()).
-		ensure(!IsPersistent());
-
-		// Omit CRC'ing data collections here as it is very slow. Rely instead on UID.
+		// Omit CRC'ing data collections here as it is very slow. Rely instead on UID/ObjectPath.
 		const FStructProperty* StructProperty = CastField<FStructProperty>(InProperty);
 		if (StructProperty && StructProperty->Struct && StructProperty->Struct->IsChildOf(FPCGDataCollection::StaticStruct()))
 		{
@@ -73,7 +83,8 @@ public:
 		}
 
 		const bool bSkip = InProperty && (
-			InProperty->GetFName() == GET_MEMBER_NAME_CHECKED(UPCGSettings, DebugSettings)
+			InProperty->GetFName() == GET_MEMBER_NAME_CHECKED(UPCGSettings, UID)
+			|| InProperty->GetFName() == GET_MEMBER_NAME_CHECKED(UPCGSettings, DebugSettings)
 			|| InProperty->GetFName() == GET_MEMBER_NAME_CHECKED(UPCGSettings, DeterminismSettings)
 			|| InProperty->GetFName() == GET_MEMBER_NAME_CHECKED(UPCGSettings, bDebug)
 			|| InProperty->GetFName() == GET_MEMBER_NAME_CHECKED(UPCGSettings, Category)
@@ -306,10 +317,7 @@ bool UPCGSettings::operator==(const UPCGSettings& Other) const
 	}
 	else
 	{
-		FPCGSettingsObjectCrc32 Ar;
-		uint32 ThisCrc = Ar.Crc32(const_cast<UPCGSettings*>(this));
-		uint32 OtherCrc = Ar.Crc32(const_cast<UPCGSettings*>(&Other));
-		return ThisCrc == OtherCrc;
+		return FPCGSettingsObjectCrc::PCGCrc(this) == FPCGSettingsObjectCrc::PCGCrc(&Other);
 	}
 }
 
@@ -1142,9 +1150,7 @@ void UPCGSettings::FixingOverridableParamPropertyClass(FPCGSettingsOverridablePa
 
 void UPCGSettings::CacheCrc()
 {
-	FPCGSettingsObjectCrc32 Ar;
-	const uint32 CrcValue = Ar.Crc32(const_cast<UPCGSettings*>(this));
-	CachedCrc = FPCGCrc(CrcValue);
+	CachedCrc = FPCGSettingsObjectCrc::PCGCrc(this);
 }
 
 TArray<FPCGPinProperties> UPCGSettings::DefaultPointInputPinProperties() const
diff --git a/Engine/Plugins/PCG/Source/PCG/Private/PCGSubsystem.cpp b/Engine/Plugins/PCG/Source/PCG/Private/PCGSubsystem.cpp
--- a/Engine/Plugins/PCG/Source/PCG/Private/PCGSubsystem.cpp
+++ b/Engine/Plugins/PCG/Source/PCG/Private/PCGSubsystem.cpp
@@ -843,6 +843,11 @@ FPCGTaskId UPCGSubsystem::ScheduleGenericWithContext(TFunction<bool(FPCGContext*
 }
 
 void UPCGSubsystem::CancelGeneration(UPCGComponent* Component)
+{
+	CancelGeneration(Component, /*bCleanupUnusedResources=*/true);
+}
+
+void UPCGSubsystem::CancelGeneration(UPCGComponent* Component, bool bCleanupUnusedResources)
 {
 	check(GraphExecutor && IsInGameThread());
 	if (!Component || !Component->IsGenerating())
@@ -852,11 +857,11 @@ void UPCGSubsystem::CancelGeneration(UPCGComponent* Component)
 
 	if (Component->IsPartitioned())
 	{
-		auto LocalCancel = [this](UPCGComponent* LocalComponent)
+		auto LocalCancel = [this, bCleanupUnusedResources](UPCGComponent* LocalComponent)
 		{
 			if (LocalComponent->IsGenerating())
 			{
-				CancelGeneration(LocalComponent);
+				CancelGeneration(LocalComponent, bCleanupUnusedResources);
 			}
 
 			return InvalidPCGTaskId;
@@ -870,7 +875,7 @@ void UPCGSubsystem::CancelGeneration(UPCGComponent* Component)
 	{
 		if (CancelledComponent)
 		{
-			CancelledComponent->OnProcessGraphAborted(/*bQuiet=*/true);
+			CancelledComponent->OnProcessGraphAborted(/*bQuiet=*/true, bCleanupUnusedResources);
 		}
 	}	
 }
diff --git a/Engine/Plugins/PCG/Source/PCG/Public/Graph/PCGStackContext.h b/Engine/Plugins/PCG/Source/PCG/Public/Graph/PCGStackContext.h
--- a/Engine/Plugins/PCG/Source/PCG/Public/Graph/PCGStackContext.h
+++ b/Engine/Plugins/PCG/Source/PCG/Public/Graph/PCGStackContext.h
@@ -2,6 +2,8 @@
 
 #pragma once
 
+#include "PCGCrc.h"
+
 #include "Utils/PCGExtraCapture.h"
 
 #include "Containers/Array.h"
@@ -145,6 +147,8 @@ public:
 		return Hash;
 	}
 
+	FPCGCrc GetCrc() const;
+
 private:
 	FPCGTaskId GraphExecutionTaskId = InvalidPCGTaskId;
 
diff --git a/Engine/Plugins/PCG/Source/PCG/Public/PCGComponent.h b/Engine/Plugins/PCG/Source/PCG/Public/PCGComponent.h
--- a/Engine/Plugins/PCG/Source/PCG/Public/PCGComponent.h
+++ b/Engine/Plugins/PCG/Source/PCG/Public/PCGComponent.h
@@ -534,7 +534,7 @@ private:
 	void PostProcessGraph(const FBox& InNewBounds, bool bInGenerated, FPCGContext* Context);
 	void CallPostGenerateFunctions(FPCGContext* Context) const;
 	void PostCleanupGraph();
-	void OnProcessGraphAborted(bool bQuiet = false);
+	void OnProcessGraphAborted(bool bQuiet = false, bool bCleanupUnusedResources = true);
 	void CleanupUnusedManagedResources();
 	bool MoveResourcesToNewActor(AActor* InNewActor, bool bCreateChild);
 
diff --git a/Engine/Plugins/PCG/Source/PCG/Public/PCGSettings.h b/Engine/Plugins/PCG/Source/PCG/Public/PCGSettings.h
--- a/Engine/Plugins/PCG/Source/PCG/Public/PCGSettings.h
+++ b/Engine/Plugins/PCG/Source/PCG/Public/PCGSettings.h
@@ -215,8 +215,9 @@ class PCG_API UPCGSettings : public UPCGSettingsInterface
 {
 	GENERATED_BODY()
 
-	friend class FPCGSettingsObjectCrc32;
+	friend class FPCGSettingsObjectCrc;
 	friend class UPCGSettingsInterface;
+	friend struct FPCGContext;
 
 public:
 	// ~Begin UPCGData interface
diff --git a/Engine/Plugins/PCG/Source/PCG/Public/PCGSubsystem.h b/Engine/Plugins/PCG/Source/PCG/Public/PCGSubsystem.h
--- a/Engine/Plugins/PCG/Source/PCG/Public/PCGSubsystem.h
+++ b/Engine/Plugins/PCG/Source/PCG/Public/PCGSubsystem.h
@@ -360,6 +360,8 @@ private:
 #endif // WITH_EDITOR
 	
 private:
+	void CancelGeneration(UPCGComponent* Component, bool bCleanupUnusedResources);
+
 	UPCGData* GetPCGData(FPCGTaskId InGraphExecutionTaskId);
 	UPCGData* GetInputPCGData(FPCGTaskId InGraphExecutionTaskId);
 	UPCGData* GetActorPCGData(FPCGTaskId InGraphExecutionTaskId);
