[PCG] Remove graph cache memory accounting ensure that is happening at very low rates.

The reason for the misaccounting is not clear given all access to MemoryRecords and TotalMemoryUsed are protected by a lock and interact in a simple way. In any case this memory total is a heuristic to stop memory growing unboundedly rather than a precise measure and clamping without ensuring should be fine.

#rnx

diff --git a/Engine/Plugins/PCG/Source/PCG/Private/Graph/PCGGraphCache.cpp b/Engine/Plugins/PCG/Source/PCG/Private/Graph/PCGGraphCache.cpp
--- a/Engine/Plugins/PCG/Source/PCG/Private/Graph/PCGGraphCache.cpp
+++ b/Engine/Plugins/PCG/Source/PCG/Private/Graph/PCGGraphCache.cpp
@@ -320,10 +320,15 @@ void FPCGGraphCache::RemoveFromMemoryTotal(const FPCGDataCollection& InCollectio
 					if (Record->InstanceCount == 0)
 					{
 						// Last instance removed, update accordingly
-						if (ensure(TotalMemoryUsed >= Record->MemoryPerInstance))
+						if (TotalMemoryUsed >= Record->MemoryPerInstance)
 						{
 							TotalMemoryUsed -= Record->MemoryPerInstance;
 						}
+						else
+						{
+							// Should not normally reach here but it seems to happen in rare cases. Clamp to 0.
+							TotalMemoryUsed = 0;
+						}
 
 						MemoryRecords.Remove(Data->UID);
 					}
