[PCG] Fix cooking determinism (introduced in 36654764), avoid serializing the redundant hash (it was a pointer hash, causing the non-deterministic behavior) and instead compute on load
- use the SoftObjectPath string as the hashed value instead of the pointer

#rnx

diff --git a/Engine/Plugins/PCG/Source/PCG/Private/Graph/PCGStackContext.cpp b/Engine/Plugins/PCG/Source/PCG/Private/Graph/PCGStackContext.cpp
--- a/Engine/Plugins/PCG/Source/PCG/Private/Graph/PCGStackContext.cpp
+++ b/Engine/Plugins/PCG/Source/PCG/Private/Graph/PCGStackContext.cpp
@@ -15,6 +15,26 @@
 
 #include UE_INLINE_GENERATED_CPP_BY_NAME(PCGStackContext)
 
+void FPCGStackFrame::PostSerialize(const FArchive& Ar)
+{
+	if (Ar.IsLoading())
+	{
+		ComputeHash();
+	}
+}
+
+void FPCGStackFrame::ComputeHash()
+{
+	if (Object.IsNull())
+	{
+		Hash = GetTypeHash(LoopIndex);
+	}
+	else
+	{
+		Hash = GetTypeHash(Object.ToString());
+	}
+}
+
 const FPCGStack* FPCGStackContext::GetStack(int32 InStackIndex) const
 {
 	if (ensure(Stacks.IsValidIndex(InStackIndex)))
diff --git a/Engine/Plugins/PCG/Source/PCG/Public/Graph/PCGStackContext.h b/Engine/Plugins/PCG/Source/PCG/Public/Graph/PCGStackContext.h
--- a/Engine/Plugins/PCG/Source/PCG/Public/Graph/PCGStackContext.h
+++ b/Engine/Plugins/PCG/Source/PCG/Public/Graph/PCGStackContext.h
@@ -48,16 +48,18 @@ struct PCG_API FPCGStackFrame
 	{
 		Object = InObject;
 		LoopIndex = INDEX_NONE;
-		Hash = PointerHash(Object.Get());
+		ComputeHash();
 	}
 
 	void SetLoopIndex(int32 InLoopIndex)
 	{
 		Object = nullptr;
 		LoopIndex = InLoopIndex;
-		Hash = GetTypeHash(LoopIndex);
+		ComputeHash();
 	}
 
+	void PostSerialize(const FArchive& Ar);
+
 	// Stores object this frame refers to. Use SetObject to change this and properly update the hash.
 	UPROPERTY()
 	TSoftObjectPtr<const UObject> Object;
@@ -67,10 +69,20 @@ struct PCG_API FPCGStackFrame
 	int32 LoopIndex = INDEX_NONE;
 
 private:
-	UPROPERTY()
+	void ComputeHash();
+
 	uint32 Hash = 0;
 };
 
+template<>
+struct TStructOpsTypeTraits<FPCGStackFrame> : public TStructOpsTypeTraitsBase2<FPCGStackFrame>
+{
+	enum
+	{
+		WithPostSerialize = true,
+	};
+};
+
 /** A call stack, represented as an array of stack frames. */
 USTRUCT()
 struct PCG_API FPCGStack
