Make FInstancedStaticMeshSceneProxyDesc use the non-ISMC specific instance data proxy to decouple it from the ISMC component.

This allows the Procedural ISM Component to use the FInstancedStaticMeshSceneProxyDesc and specify cull start/end distances, and also to create the non-Nanite scene proxy to enable non-Nanite meshes.

#rnx

diff --git a/Engine/Plugins/Experimental/GPULightmass/Source/GPULightmass/Private/Scene/InstancedStaticMesh.h b/Engine/Plugins/Experimental/GPULightmass/Source/GPULightmass/Private/Scene/InstancedStaticMesh.h
--- a/Engine/Plugins/Experimental/GPULightmass/Source/GPULightmass/Private/Scene/InstancedStaticMesh.h
+++ b/Engine/Plugins/Experimental/GPULightmass/Source/GPULightmass/Private/Scene/InstancedStaticMesh.h
@@ -18,7 +18,7 @@ public:
 
 	TArray<FIntPoint> LODPerInstanceLightmapSize;
 	int32 NumInstances = 0;
-	TSharedPtr<FISMCInstanceDataSceneProxy, ESPMode::ThreadSafe> InstanceDataSceneProxy;
+	TSharedPtr<FInstanceDataSceneProxy, ESPMode::ThreadSafe> InstanceDataSceneProxy;
 	const FInstanceSceneDataBuffers *InstanceSceneDataBuffers = nullptr;
 
 	TArray<FMeshBatch> GetMeshBatchesForGBufferRendering(int32 LODIndex, FTileVirtualCoordinates CoordsForCulling);
diff --git a/Engine/Plugins/Experimental/GPULightmass/Source/GPULightmass/Private/Scene/Scene.cpp b/Engine/Plugins/Experimental/GPULightmass/Source/GPULightmass/Private/Scene/Scene.cpp
--- a/Engine/Plugins/Experimental/GPULightmass/Source/GPULightmass/Private/Scene/Scene.cpp
+++ b/Engine/Plugins/Experimental/GPULightmass/Source/GPULightmass/Private/Scene/Scene.cpp
@@ -1053,7 +1053,7 @@ void FScene::AddGeometryInstanceFromComponent(UInstancedStaticMeshComponent* InC
 		{
 			TaskInfo->WaitForUpdateCompletion();
 		}
-		InstanceRenderState.InstanceSceneDataBuffers = &InstanceRenderState.InstanceDataSceneProxy->GetData();
+		InstanceRenderState.InstanceSceneDataBuffers = InstanceRenderState.InstanceDataSceneProxy->GeInstanceSceneDataBuffers();
 		InstanceRenderState.NumInstances = InstanceRenderState.InstanceSceneDataBuffers->GetNumInstances();
 
 		FInstanceGroupRenderStateRef InstanceRenderStateRef = RenderState.InstanceGroupRenderStates.Emplace(MoveTemp(InstanceRenderState));
diff --git a/Engine/Plugins/PCG/Source/PCG/Private/Components/PCGProceduralISMComponent.cpp b/Engine/Plugins/PCG/Source/PCG/Private/Components/PCGProceduralISMComponent.cpp
--- a/Engine/Plugins/PCG/Source/PCG/Private/Components/PCGProceduralISMComponent.cpp
+++ b/Engine/Plugins/PCG/Source/PCG/Private/Components/PCGProceduralISMComponent.cpp
@@ -9,8 +9,10 @@
 #include "MeshSelectors/PCGISMDescriptor.h"
 
 #include "InstanceDataSceneProxy.h"
+#include "InstancedStaticMeshSceneProxyDesc.h"
 #include "NaniteSceneProxy.h"
 #include "SceneInterface.h"
+#include "StaticMeshSceneProxyDesc.h"
 #include "Engine/StaticMesh.h"
 #include "Serialization/ArchiveCrc32.h"
 #include "VT/RuntimeVirtualTexture.h"
@@ -222,6 +224,17 @@ FPrimitiveSceneProxy* UPCGProceduralISMComponent::CreateStaticMeshSceneProxy(Nan
 {
 	LLM_SCOPE(ELLMTag::InstancedMesh);
 
+	if (!ensure(GetWorld()) || !ensure(GetWorld()->Scene))
+	{
+		return nullptr;
+	}
+	
+	if (!UseGPUScene(GetWorld()->Scene->GetShaderPlatform(), GetWorld()->Scene->GetFeatureLevel()))
+	{
+		UE_LOG(LogPCG, Warning, TEXT("PCGProceduralISMComponent depends on GPUScene functionality which is not available on this platform, component will not render."));
+		return nullptr;
+	}
+
 	if (CheckPSOPrecachingAndBoostPriority() && GetPSOPrecacheProxyCreationStrategy() == EPSOPrecacheProxyCreationStrategy::DelayUntilPSOPrecached)
 	{
 		UE_LOG(LogPCG, Verbose, TEXT("Skipping CreateSceneProxy for PCGProceduralISMComponent %s (PCGProceduralISMComponent PSOs are still compiling)"), *GetFullName());
@@ -229,32 +242,37 @@ FPrimitiveSceneProxy* UPCGProceduralISMComponent::CreateStaticMeshSceneProxy(Nan
 	}
 
 	FInstanceSceneDataBuffers InstanceSceneDataBuffers(/*InbInstanceDataIsGPUOnly=*/true);
+	{
+		FInstanceSceneDataBuffers::FAccessTag AccessTag(PointerHash(this));
+		FInstanceSceneDataBuffers::FWriteView ProxyData = InstanceSceneDataBuffers.BeginWriteAccess(AccessTag);
 
-	FInstanceSceneDataBuffers::FAccessTag AccessTag(PointerHash(this));
-	FInstanceSceneDataBuffers::FWriteView ProxyData = InstanceSceneDataBuffers.BeginWriteAccess(AccessTag);
-
-	InstanceSceneDataBuffers.SetPrimitiveLocalToWorld(GetRenderMatrix(), AccessTag);
+		InstanceSceneDataBuffers.SetPrimitiveLocalToWorld(GetRenderMatrix(), AccessTag);
 
-	ProxyData.NumInstancesGPUOnly = GetNumInstances();
-	ProxyData.NumCustomDataFloats = GetNumCustomDataFloats();
-	ProxyData.InstanceLocalBounds.SetNum(1);
-	ProxyData.InstanceLocalBounds[0] = GetStaticMesh()->GetBounds();
+		ProxyData.NumInstancesGPUOnly = GetNumInstances();
+		ProxyData.NumCustomDataFloats = GetNumCustomDataFloats();
+		ProxyData.InstanceLocalBounds.SetNum(1);
+		ProxyData.InstanceLocalBounds[0] = GetStaticMesh()->GetBounds();
 
-	ProxyData.Flags.bHasPerInstanceCustomData = ProxyData.NumCustomDataFloats > 0;
+		ProxyData.Flags.bHasPerInstanceCustomData = ProxyData.NumCustomDataFloats > 0;
 
-	InstanceSceneDataBuffers.EndWriteAccess(AccessTag);
-	InstanceSceneDataBuffers.ValidateData();
+		InstanceSceneDataBuffers.EndWriteAccess(AccessTag);
+		InstanceSceneDataBuffers.ValidateData();
+	}
 
-	auto InstanceDataSceneProxy = MakeShared<FInstanceDataSceneProxy, ESPMode::ThreadSafe>(MoveTemp(InstanceSceneDataBuffers));
+	FInstancedStaticMeshSceneProxyDesc Desc;
+	Desc.InitializeFromStaticMeshComponent(this);
+	Desc.InstanceDataSceneProxy = MakeShared<FInstanceDataSceneProxy, ESPMode::ThreadSafe>(MoveTemp(InstanceSceneDataBuffers));
+	Desc.InstanceStartCullDistance = InstanceStartCullDistance;
+	Desc.InstanceEndCullDistance = InstanceEndCullDistance;
+	Desc.bUseGpuLodSelection = true;
 
 	if (bCreateNanite)
 	{
-		return ::new Nanite::FSceneProxy(NaniteMaterials, this, InstanceDataSceneProxy);
+		return ::new Nanite::FSceneProxy(NaniteMaterials, Desc);
 	}
 	else
 	{
-		UE_LOG(LogPCG, Warning, TEXT("Only Nanite meshes are supported by PCGProceduralISMComponent currently, component will not render."));
-		return nullptr;
+		return ::new FInstancedStaticMeshSceneProxy(Desc, GetWorld()->GetFeatureLevel());
 	}
 }
 
diff --git a/Engine/Plugins/PCG/Source/PCG/Private/Elements/PCGStaticMeshSpawner.cpp b/Engine/Plugins/PCG/Source/PCG/Private/Elements/PCGStaticMeshSpawner.cpp
--- a/Engine/Plugins/PCG/Source/PCG/Private/Elements/PCGStaticMeshSpawner.cpp
+++ b/Engine/Plugins/PCG/Source/PCG/Private/Elements/PCGStaticMeshSpawner.cpp
@@ -66,19 +66,11 @@ bool UPCGStaticMeshSpawnerSettings::IsKernelValid(FPCGContext* InContext, bool b
 
 		for (const FPCGMeshSelectorWeightedEntry& Entry : SelectorWeighted->MeshEntries)
 		{
-			UStaticMesh* StaticMesh = Entry.Descriptor.StaticMesh.LoadSynchronous();
-
-			if (!StaticMesh)
+			if (Entry.Descriptor.StaticMesh.IsNull())
 			{
 				PCG_KERNEL_VALIDATION_ERR(InContext, this, bQuiet, LOCTEXT("UnassignedMesh", "Unassigned mesh."));
 				return false;
 			}
-
-			if (!StaticMesh->HasValidNaniteData())
-			{
-				PCG_KERNEL_VALIDATION_ERR(InContext, this, bQuiet, LOCTEXT("NonNaniteMesh", "Only Nanite meshes are currently supported by the GPU Static Mesh Spawner path."));
-				return false;
-			}
 		}
 	}
 	else if (UPCGMeshSelectorByAttribute* SelectorByAttribute = Cast<UPCGMeshSelectorByAttribute>(MeshSelectorParameters))
@@ -593,18 +585,7 @@ bool FPCGStaticMeshSpawnerElement::CanExecuteOnlyOnMainThread(FPCGContext* Conte
 	// PrepareData can call UPCGManagedComponent::MarkAsReused which registers the ISMC, which can go into Chaos code that asserts if not on main thread.
 	// TODO: We can likely re-enable multi-threading for PrepareData if we move the call to MarkAsReused to Execute. There should hopefully not be
 	// wider contention on resources resources are not shared across nodes and are also per-component.
-	bool bRequiresMainThread = !Context || Context->CurrentPhase == EPCGExecutionPhase::Execute || Context->CurrentPhase == EPCGExecutionPhase::PrepareData;
-
-	if (!bRequiresMainThread && Context->GetInputSettings<UPCGStaticMeshSpawnerSettings>())
-	{
-		// TODO: Unfortunate result of convoluted logic. If the spawner is set to GPU, but the validation failed (e.g. static mesh is not compatible),
-		// then the compute graph is not compiled and the CPU spawner task is left in the graph, and the spawner validation will be executed from the element
-		// so that graph warnings and errors can be logged. This validation loads the static mesh to check its settings and this must be done on the main
-		// thread. In the future we may want to invite nodes to log warnings/errors at graph schedule time, and always remove CPU tasks for GPU nodes.
-		bRequiresMainThread |= Context->GetInputSettings<UPCGStaticMeshSpawnerSettings>()->ShouldExecuteOnGPU();
-	}
-
-	return bRequiresMainThread;
+	return !Context || Context->CurrentPhase == EPCGExecutionPhase::Execute || Context->CurrentPhase == EPCGExecutionPhase::PrepareData;
 }
 
 void FPCGStaticMeshSpawnerElement::SpawnStaticMeshInstances(FPCGStaticMeshSpawnerContext* Context, const FPCGMeshInstanceList& InstanceList, AActor* TargetActor, const FPCGPackedCustomData* InPackedCustomData) const
diff --git a/Engine/Source/Runtime/Engine/Classes/Engine/InstancedStaticMesh.h b/Engine/Source/Runtime/Engine/Classes/Engine/InstancedStaticMesh.h
--- a/Engine/Source/Runtime/Engine/Classes/Engine/InstancedStaticMesh.h
+++ b/Engine/Source/Runtime/Engine/Classes/Engine/InstancedStaticMesh.h
@@ -566,7 +566,7 @@ private:
 	/** Stores a loose uniform buffer per LOD, used for static view relevance. */
 	TMap<uint32, FInstancedStaticMeshVFLooseUniformShaderParametersRef> LODLooseUniformBuffers;
 
-	TSharedPtr<FISMCInstanceDataSceneProxy, ESPMode::ThreadSafe> InstanceDataSceneProxy; 
+	TSharedPtr<FInstanceDataSceneProxy, ESPMode::ThreadSafe> InstanceDataSceneProxy; 
 };
 
 #if WITH_EDITOR
diff --git a/Engine/Source/Runtime/Engine/Private/InstancedStaticMesh.cpp b/Engine/Source/Runtime/Engine/Private/InstancedStaticMesh.cpp
--- a/Engine/Source/Runtime/Engine/Private/InstancedStaticMesh.cpp
+++ b/Engine/Source/Runtime/Engine/Private/InstancedStaticMesh.cpp
@@ -1296,12 +1296,12 @@ FInstancedStaticMeshVFLooseUniformShaderParametersRef FInstancedStaticMeshSceneP
 FInstancedStaticMeshSceneProxyDesc::FInstancedStaticMeshSceneProxyDesc(const UInstancedStaticMeshComponent* InComponent)
 	: FInstancedStaticMeshSceneProxyDesc()
 {
-	InitializeFrom(InComponent);
+	InitializeFromInstancedStaticMeshComponent(InComponent);
 }
 
-void FInstancedStaticMeshSceneProxyDesc::InitializeFrom(const UInstancedStaticMeshComponent* InComponent)
+void FInstancedStaticMeshSceneProxyDesc::InitializeFromInstancedStaticMeshComponent(const UInstancedStaticMeshComponent* InComponent)
 {
-	FStaticMeshSceneProxyDesc::InitializeFrom(InComponent);
+	InitializeFromStaticMeshComponent(InComponent);
 
 	InstanceDataSceneProxy = InComponent->GetInstanceDataSceneProxy();
 	check(InstanceDataSceneProxy)
diff --git a/Engine/Source/Runtime/Engine/Private/PrimitiveSceneProxy.cpp b/Engine/Source/Runtime/Engine/Private/PrimitiveSceneProxy.cpp
--- a/Engine/Source/Runtime/Engine/Private/PrimitiveSceneProxy.cpp
+++ b/Engine/Source/Runtime/Engine/Private/PrimitiveSceneProxy.cpp
@@ -261,12 +261,12 @@ static FRenderBounds PadLocalRenderBounds(const FRenderBounds& InBounds, const F
 FPrimitiveSceneProxyDesc::FPrimitiveSceneProxyDesc(const UPrimitiveComponent* InComponent)
 	: FPrimitiveSceneProxyDesc()
 {
-	InitializeFrom(InComponent);
+	InitializeFromPrimitiveComponent(InComponent);
 }
 
 FPrimitiveSceneProxy::FPrimitiveSceneProxy(FPrimitiveSceneProxy const&) = default;
 
-void FPrimitiveSceneProxyDesc::InitializeFrom(const UPrimitiveComponent* InComponent)
+void FPrimitiveSceneProxyDesc::InitializeFromPrimitiveComponent(const UPrimitiveComponent* InComponent)
 {
 	CastShadow = InComponent->CastShadow;
 	bReceivesDecals = InComponent->bReceivesDecals;
diff --git a/Engine/Source/Runtime/Engine/Private/SkinnedMeshSceneProxyDesc.cpp b/Engine/Source/Runtime/Engine/Private/SkinnedMeshSceneProxyDesc.cpp
--- a/Engine/Source/Runtime/Engine/Private/SkinnedMeshSceneProxyDesc.cpp
+++ b/Engine/Source/Runtime/Engine/Private/SkinnedMeshSceneProxyDesc.cpp
@@ -24,12 +24,12 @@ extern ENGINE_API int32 GEnableGPUSkinCache;
 
 FSkinnedMeshSceneProxyDesc::FSkinnedMeshSceneProxyDesc(const USkinnedMeshComponent* Component)
 { 
-	InitializeFrom(Component);
+	InitializeFromSkinnedMeshComponent(Component);
 }
 
-void FSkinnedMeshSceneProxyDesc::InitializeFrom(const USkinnedMeshComponent* InComponent)
+void FSkinnedMeshSceneProxyDesc::InitializeFromSkinnedMeshComponent(const USkinnedMeshComponent* InComponent)
 {
-	FPrimitiveSceneProxyDesc::InitializeFrom(InComponent);
+	InitializeFromPrimitiveComponent(InComponent);
 	bForceWireframe = InComponent->bForceWireframe;
 	bCanHighlightSelectedSections = InComponent->bCanHighlightSelectedSections;
 	bRenderStatic = InComponent->bRenderStatic;
diff --git a/Engine/Source/Runtime/Engine/Private/StaticMeshSceneProxy.cpp b/Engine/Source/Runtime/Engine/Private/StaticMeshSceneProxy.cpp
--- a/Engine/Source/Runtime/Engine/Private/StaticMeshSceneProxy.cpp
+++ b/Engine/Source/Runtime/Engine/Private/StaticMeshSceneProxy.cpp
@@ -2536,12 +2536,12 @@ FTextureResource* FStaticMeshSceneProxyDesc::GetMeshPaintTextureResource() const
 FStaticMeshSceneProxyDesc::FStaticMeshSceneProxyDesc(const UStaticMeshComponent* InComponent)
 	: FStaticMeshSceneProxyDesc()
 {	
-	InitializeFrom(InComponent);
+	InitializeFromStaticMeshComponent(InComponent);
 }
 
-void FStaticMeshSceneProxyDesc::InitializeFrom(const UStaticMeshComponent* InComponent)
+void FStaticMeshSceneProxyDesc::InitializeFromStaticMeshComponent(const UStaticMeshComponent* InComponent)
 {
-	FPrimitiveSceneProxyDesc::InitializeFrom(InComponent);	
+	InitializeFromPrimitiveComponent(InComponent);	
 
 	StaticMesh = InComponent->GetStaticMesh();
 	OverrideMaterials = const_cast<UStaticMeshComponent*>(InComponent)->OverrideMaterials;	
diff --git a/Engine/Source/Runtime/Engine/Public/InstanceDataSceneProxy.h b/Engine/Source/Runtime/Engine/Public/InstanceDataSceneProxy.h
--- a/Engine/Source/Runtime/Engine/Public/InstanceDataSceneProxy.h
+++ b/Engine/Source/Runtime/Engine/Public/InstanceDataSceneProxy.h
@@ -9,6 +9,7 @@
 #include "Rendering/RenderingSpatialHash.h"
 
 class FPrimitiveDrawInterface;
+class FStaticMeshInstanceBuffer;
 
 class FInstanceIdIndexMap
 {
@@ -440,6 +441,10 @@ public:
 
 	ENGINE_API virtual void DebugDrawInstanceChanges(FPrimitiveDrawInterface* DebugPDI, ESceneDepthPriorityGroup SceneDepthPriorityGroup) {};
 
+	/**
+	 * Workaround to support the ISMC non GPU scene rendering path.
+	 */
+	virtual ENGINE_API FStaticMeshInstanceBuffer* GetLegacyInstanceBuffer() { return nullptr; }
 
 protected:
 	ENGINE_API void IncStatCounters();
diff --git a/Engine/Source/Runtime/Engine/Public/InstancedStaticMesh/ISMInstanceDataSceneProxy.h b/Engine/Source/Runtime/Engine/Public/InstancedStaticMesh/ISMInstanceDataSceneProxy.h
--- a/Engine/Source/Runtime/Engine/Public/InstancedStaticMesh/ISMInstanceDataSceneProxy.h
+++ b/Engine/Source/Runtime/Engine/Public/InstancedStaticMesh/ISMInstanceDataSceneProxy.h
@@ -42,11 +42,6 @@ public:
 
 	inline const FInstanceSceneDataBuffers& GetData() const { return InstanceSceneDataBuffers; }
 
-	/**
-	 * Get the legacy instance data (which is available on legacy platforms).
-	 */
-	virtual ENGINE_API FStaticMeshInstanceBuffer* GetLegacyInstanceBuffer() { return nullptr; }
-
 	//void BuildLegacyData();
 
 	/**
diff --git a/Engine/Source/Runtime/Engine/Public/InstancedStaticMeshSceneProxyDesc.h b/Engine/Source/Runtime/Engine/Public/InstancedStaticMeshSceneProxyDesc.h
--- a/Engine/Source/Runtime/Engine/Public/InstancedStaticMeshSceneProxyDesc.h
+++ b/Engine/Source/Runtime/Engine/Public/InstancedStaticMeshSceneProxyDesc.h
@@ -11,9 +11,13 @@ struct FInstancedStaticMeshSceneProxyDesc : public FStaticMeshSceneProxyDesc
 {		
 	FInstancedStaticMeshSceneProxyDesc() = default;
 	ENGINE_API FInstancedStaticMeshSceneProxyDesc(const UInstancedStaticMeshComponent*);
-	ENGINE_API void InitializeFrom(const UInstancedStaticMeshComponent*);
 
-	TSharedPtr<FISMCInstanceDataSceneProxy, ESPMode::ThreadSafe> InstanceDataSceneProxy;
+	ENGINE_API void InitializeFromInstancedStaticMeshComponent(const UInstancedStaticMeshComponent*);
+
+	UE_DEPRECATED(5.5, "Use InitializeFromInstancedStaticMeshComponent instead.")
+	ENGINE_API void InitializeFrom(const UInstancedStaticMeshComponent* InComponent) { InitializeFromInstancedStaticMeshComponent(InComponent); }
+
+	TSharedPtr<FInstanceDataSceneProxy, ESPMode::ThreadSafe> InstanceDataSceneProxy;
 #if WITH_EDITOR
 	bool bHasSelectedInstances = false;
 #endif
@@ -24,4 +28,4 @@ struct FInstancedStaticMeshSceneProxyDesc : public FStaticMeshSceneProxyDesc
 
 	bool bUseGpuLodSelection = false;
 
-};
\ No newline at end of file
+};
diff --git a/Engine/Source/Runtime/Engine/Public/PrimitiveSceneProxyDesc.h b/Engine/Source/Runtime/Engine/Public/PrimitiveSceneProxyDesc.h
--- a/Engine/Source/Runtime/Engine/Public/PrimitiveSceneProxyDesc.h
+++ b/Engine/Source/Runtime/Engine/Public/PrimitiveSceneProxyDesc.h
@@ -77,8 +77,11 @@ struct FPrimitiveSceneProxyDesc
 
 	ENGINE_API FPrimitiveSceneProxyDesc(const UPrimitiveComponent*);
 
-	void InitializeFrom(const UPrimitiveComponent*);
+	void InitializeFromPrimitiveComponent(const UPrimitiveComponent*);
 	
+	UE_DEPRECATED(5.5, "Use InitializeFromPrimitiveComponent instead.")
+	void InitializeFrom(const UPrimitiveComponent* InComponent) { InitializeFromPrimitiveComponent(InComponent); }
+
 	virtual ~FPrimitiveSceneProxyDesc() = default;
 
 	uint32 CastShadow : 1;
diff --git a/Engine/Source/Runtime/Engine/Public/SkinnedMeshSceneProxyDesc.h b/Engine/Source/Runtime/Engine/Public/SkinnedMeshSceneProxyDesc.h
--- a/Engine/Source/Runtime/Engine/Public/SkinnedMeshSceneProxyDesc.h
+++ b/Engine/Source/Runtime/Engine/Public/SkinnedMeshSceneProxyDesc.h
@@ -24,7 +24,7 @@ struct FSkinnedMeshSceneProxyDesc : public FPrimitiveSceneProxyDesc
 
 	FSkinnedMeshSceneProxyDesc() = default;
 	ENGINE_API FSkinnedMeshSceneProxyDesc(const USkinnedMeshComponent* Component);
-	ENGINE_API void InitializeFrom(const USkinnedMeshComponent*);
+	ENGINE_API void InitializeFromSkinnedMeshComponent(const USkinnedMeshComponent*);
 	ENGINE_API USkinnedAsset* GetSkinnedAsset() const;
 	ENGINE_API UPhysicsAsset* GetPhysicsAsset() const;
 	ENGINE_API bool ShouldDrawDebugSkeleton() const;
@@ -53,6 +53,8 @@ struct FSkinnedMeshSceneProxyDesc : public FPrimitiveSceneProxyDesc
 
 	ENGINE_API UMeshDeformerInstance* GetMeshDeformerInstanceForLOD(int32 LODIndex) const;
 
+	UE_DEPRECATED(5.5, "Use InitializeFromSkinnedMeshComponent instead.")
+	ENGINE_API void InitializeFrom(const USkinnedMeshComponent* InComponent) { InitializeFromSkinnedMeshComponent(InComponent); }
 
 	TOptional<FLinearColor> DebugDrawColor;
 	FMaterialRelevance MaterialRelevance{};
diff --git a/Engine/Source/Runtime/Engine/Public/StaticMeshSceneProxyDesc.h b/Engine/Source/Runtime/Engine/Public/StaticMeshSceneProxyDesc.h
--- a/Engine/Source/Runtime/Engine/Public/StaticMeshSceneProxyDesc.h
+++ b/Engine/Source/Runtime/Engine/Public/StaticMeshSceneProxyDesc.h
@@ -44,7 +44,11 @@ struct FStaticMeshSceneProxyDesc : public FPrimitiveSceneProxyDesc
 	}
 
 	ENGINE_API FStaticMeshSceneProxyDesc(const UStaticMeshComponent*);
-	ENGINE_API void InitializeFrom(const UStaticMeshComponent*);
+
+	ENGINE_API void InitializeFromStaticMeshComponent(const UStaticMeshComponent*);
+
+	UE_DEPRECATED(5.5, "Use InitializeFromStaticMeshComponent instead.")
+	ENGINE_API void InitializeFrom(const UStaticMeshComponent* InComponent) { InitializeFromStaticMeshComponent(InComponent); }
 
 	UStaticMesh* StaticMesh = nullptr;
 	TArrayView<TObjectPtr<UMaterialInterface>>	OverrideMaterials;
