[PCG] Fix stack replacement in Debug tree view when BP is re-instantiated, breaking attributes list view in subgraphs.

* Made a helper function for stack root replacement instead of duplicating the code.

#rnx

diff --git a/Engine/Plugins/PCG/Source/PCG/Private/Graph/PCGStackContext.cpp b/Engine/Plugins/PCG/Source/PCG/Private/Graph/PCGStackContext.cpp
--- a/Engine/Plugins/PCG/Source/PCG/Private/Graph/PCGStackContext.cpp
+++ b/Engine/Plugins/PCG/Source/PCG/Private/Graph/PCGStackContext.cpp
@@ -201,6 +201,36 @@ bool FPCGStack::HasObject(const UObject* InObject) const
 	});
 }
 
+bool FPCGStack::ReplaceRoot(const TMap<UObject*, UObject*>& ReplacementMap)
+{
+	if (!StackFrames.IsEmpty())
+	{
+		UObject* NewStackRoot = ReplacementMap.FindRef(StackFrames[0].Object.Get());
+
+		// If the stack frame was marked as garbage, NewStackRoot will be nullptr, but we still match against the object ptr hash.
+		if (!NewStackRoot)
+		{
+			for (const TPair<UObject*, UObject*>& Pair : ReplacementMap)
+			{
+				// TODO: If we ever have a replacement where the name changed, this won't work.
+				// We might need to have a serialization pointer (soft) and an execution pointer (weak).
+				if (Pair.Key && Pair.Value && StackFrames[0].Object == Pair.Value)
+				{
+					NewStackRoot = Pair.Value;
+				}
+			}
+		}
+
+		if (NewStackRoot)
+		{
+			StackFrames[0].SetObject(NewStackRoot);
+			return true;
+		}
+	}
+
+	return false;
+}
+
 bool FPCGStack::operator==(const FPCGStack& Other) const
 {
 	// Stacks are the same if all stack frames are the same
diff --git a/Engine/Plugins/PCG/Source/PCG/Public/Graph/PCGStackContext.h b/Engine/Plugins/PCG/Source/PCG/Public/Graph/PCGStackContext.h
--- a/Engine/Plugins/PCG/Source/PCG/Public/Graph/PCGStackContext.h
+++ b/Engine/Plugins/PCG/Source/PCG/Public/Graph/PCGStackContext.h
@@ -144,6 +144,9 @@ public:
 	/** Stack has a frame corresponding to the given object. */
 	bool HasObject(const UObject* InObject) const;
 
+	/** Utility function to replace the root object, in case of re-instantiation. Returns true if a replacement happened. */
+	bool ReplaceRoot(const TMap<UObject*, UObject*>& ReplacementMap);
+
 	bool operator==(const FPCGStack& Other) const;
 	bool operator!=(const FPCGStack& Other) const { return !(*this == Other); }
 
diff --git a/Engine/Plugins/PCG/Source/PCGEditor/Private/PCGEditor.cpp b/Engine/Plugins/PCG/Source/PCGEditor/Private/PCGEditor.cpp
--- a/Engine/Plugins/PCG/Source/PCGEditor/Private/PCGEditor.cpp
+++ b/Engine/Plugins/PCG/Source/PCGEditor/Private/PCGEditor.cpp
@@ -252,31 +252,7 @@ void FPCGEditor::OnObjectsReplaced(const TMap<UObject*, UObject*>& ReplacementMa
 		LastValidPCGComponentBeingInspected = Cast<UPCGComponent>(NewLastValidPCGComponentBeingInspected);
 	}
 
-	TArray<FPCGStackFrame>& StackFrames = StackBeingInspected.GetStackFramesMutable();
-	if (!StackFrames.IsEmpty())
-	{
-		UObject* NewStackRoot = ReplacementMap.FindRef(StackFrames[0].Object.Get());
-
-		// If the stack frame was marked as garbage, NewStackRoot will be nullptr, but we still match against the object path.
-		if (!NewStackRoot)
-		{
-			for(const TPair<UObject*, UObject*>& Pair : ReplacementMap)
-			{
-				if (Pair.Key && Pair.Value)
-				{
-					if (Pair.Key->GetPathName() == StackFrames[0].Object.ToString())
-					{
-						NewStackRoot = Pair.Value;
-					}
-				}
-			}
-		}
-
-		if (NewStackRoot)
-		{
-			StackFrames[0].SetObject(NewStackRoot);
-		}
-	}
+	StackBeingInspected.ReplaceRoot(ReplacementMap);
 
 	// Propagate object replacement to the debug object tree view too
 	if (DebugObjectTreeWidget)
diff --git a/Engine/Plugins/PCG/Source/PCGEditor/Private/SPCGEditorGraphDebugObjectTree.cpp b/Engine/Plugins/PCG/Source/PCGEditor/Private/SPCGEditorGraphDebugObjectTree.cpp
--- a/Engine/Plugins/PCG/Source/PCGEditor/Private/SPCGEditorGraphDebugObjectTree.cpp
+++ b/Engine/Plugins/PCG/Source/PCGEditor/Private/SPCGEditorGraphDebugObjectTree.cpp
@@ -138,31 +138,7 @@ void FPCGEditorGraphDebugObjectItem::OnObjectsReplaced(const TMap<UObject*, UObj
 {
 	if (FPCGStack* Stack = GetMutablePCGStack())
 	{
-		TArray<FPCGStackFrame>& StackFrames = Stack->GetStackFramesMutable();
-		if (!StackFrames.IsEmpty())
-		{
-			UObject* NewStackRoot = ReplacementMap.FindRef(StackFrames[0].Object.Get());
-
-			// If the stack frame was marked as garbage, NewStackRoot will be nullptr, but we still match against the object path.
-			if (!NewStackRoot)
-			{
-				for(const TPair<UObject*, UObject*>& Pair : ReplacementMap)
-				{
-					if (Pair.Key && Pair.Value)
-					{
-						if (Pair.Key->GetPathName() == StackFrames[0].Object.ToString())
-						{
-							NewStackRoot = Pair.Value;
-						}
-					}
-				}
-			}
-
-			if (NewStackRoot)
-			{
-				StackFrames[0].SetObject(NewStackRoot);
-			}
-		}
+		Stack->ReplaceRoot(ReplacementMap);
 	}
 }
 
@@ -435,12 +411,25 @@ void SPCGEditorGraphDebugObjectTree::OnObjectsReplaced(const TMap<UObject*, UObj
 		}
 	}
 
-	// While we'll update the selected original component here (might make stack matching more efficient),
-	// we will not update the previously selected stack, as the selection will be redone anyway.
+	// We need to replace the selected component, selected stacks and expanded stacks as they are used for selection matching.
 	if (UObject* NewComponent = ReplacementMap.FindRef(SelectedOriginalComponent.GetEvenIfUnreachable()))
 	{
 		SelectedOriginalComponent = Cast<UPCGComponent>(NewComponent);
 	}
+
+	SelectedStack.ReplaceRoot(ReplacementMap);
+
+	// For the expanded stacks, since it is a set, if we ever modify the root, we should probably re-hash it.
+	// We can't modify it in place so we'll just rebuild it.
+	TSet<FPCGStack> UpdatedExpandedStacks;
+	UpdatedExpandedStacks.Reserve(ExpandedStacks.Num());
+	for (FPCGStack& Stack : ExpandedStacks)
+	{
+		Stack.ReplaceRoot(ReplacementMap);
+		UpdatedExpandedStacks.Add(std::move(Stack));
+	}
+	
+	ExpandedStacks = std::move(UpdatedExpandedStacks);
 }
 
 void SPCGEditorGraphDebugObjectTree::SetNodeBeingInspected(const UPCGNode* InPCGNode)
@@ -877,8 +866,9 @@ void SPCGEditorGraphDebugObjectTree::SortTreeItems(bool bIsAscending, bool bIsRe
 
 void SPCGEditorGraphDebugObjectTree::RestoreTreeState()
 {
-	// Try to restore user item expansion.
-	TSet<FPCGStack> ExpandedStacksBefore = ExpandedStacks;
+	// Try to restore user item expansion. We can't modify ExpandedStack, so we'll rebuild it.
+	TSet<FPCGStack> ExpandedStacksBefore = std::move(ExpandedStacks);
+	ExpandedStacks.Empty(ExpandedStacksBefore.Num());
 	for (const FPCGStack& ExpandedStack : ExpandedStacksBefore)
 	{
 		for (FPCGEditorGraphDebugObjectItemPtr& Item : AllGraphItems)
