[PCG] Fix RuntimeGen scheduling policy being overwritten on PostLoad

#rnx

diff --git a/Engine/Plugins/PCG/Source/PCG/Private/PCGComponent.cpp b/Engine/Plugins/PCG/Source/PCG/Private/PCGComponent.cpp
--- a/Engine/Plugins/PCG/Source/PCG/Private/PCGComponent.cpp
+++ b/Engine/Plugins/PCG/Source/PCG/Private/PCGComponent.cpp
@@ -1768,7 +1768,7 @@ void UPCGComponent::PostLoad()
 	}
 #endif
 
-	if (!SchedulingPolicy || IsManagedByRuntimeGenSystem())
+	if (!IsValid(SchedulingPolicy))
 	{
 		RefreshSchedulingPolicy();
 	}
