[PCG] Fix bTreatEditorViewportAsGenSource causing constant refresh of RuntimeGen when floating windows are focused

#rnx

diff --git a/Engine/Plugins/PCG/Source/PCG/Private/RuntimeGen/PCGRuntimeGenScheduler.cpp b/Engine/Plugins/PCG/Source/PCG/Private/RuntimeGen/PCGRuntimeGenScheduler.cpp
--- a/Engine/Plugins/PCG/Source/PCG/Private/RuntimeGen/PCGRuntimeGenScheduler.cpp
+++ b/Engine/Plugins/PCG/Source/PCG/Private/RuntimeGen/PCGRuntimeGenScheduler.cpp
@@ -190,13 +190,10 @@ bool FPCGRuntimeGenScheduler::ShouldTick()
 	// same behaviour as other things).
 	if (!World->IsGameWorld())
 	{
-		bool bAnyVisible = false;
-		for (FEditorViewportClient* EditorViewportClient : GEditor->GetAllViewportClients())
-		{
-			bAnyVisible |= EditorViewportClient->IsVisible();
-		}
+		FViewport* Viewport = GEditor ? GEditor->GetActiveViewport() : nullptr;
+		FEditorViewportClient* ViewportClient = Viewport ? static_cast<FEditorViewportClient*>(Viewport->GetClient()) : nullptr;
 
-		if (!bAnyVisible)
+		if (!ViewportClient || !ViewportClient->IsVisible())
 		{
 			return false;
 		}
