[PCG] Fix ensure when force regenerating multiple PCG Components with GPU nodes

#rnx

diff --git a/Engine/Plugins/PCG/Source/PCGEditor/Private/Details/PCGComponentDetails.cpp b/Engine/Plugins/PCG/Source/PCGEditor/Private/Details/PCGComponentDetails.cpp
--- a/Engine/Plugins/PCG/Source/PCGEditor/Private/Details/PCGComponentDetails.cpp
+++ b/Engine/Plugins/PCG/Source/PCGEditor/Private/Details/PCGComponentDetails.cpp
@@ -256,19 +256,26 @@ EVisibility FPCGComponentDetails::RefreshButtonVisible() const
 
 FReply FPCGComponentDetails::OnGenerateClicked()
 {
-	for (TWeakObjectPtr<UPCGComponent>& Component : SelectedComponents)
+	const FModifierKeysState ModifierKeys = FSlateApplication::Get().GetModifierKeys();
+	const bool bIsControlDown = ModifierKeys.IsControlDown();
+
+	if (bIsControlDown)
 	{
-		if (Component.IsValid())
+		for (TWeakObjectPtr<UPCGComponent>& Component : SelectedComponents)
 		{
-			bool bForce = false;
-			FModifierKeysState ModifierKeys = FSlateApplication::Get().GetModifierKeys();
-			if (ModifierKeys.IsControlDown())
+			if (Component.IsValid())
 			{
 				Component->GetSubsystem()->FlushCache();
-				bForce = true;
+				break;
 			}
+		}
+	}
 
-			Component.Get()->Generate(bForce);
+	for (TWeakObjectPtr<UPCGComponent>& Component : SelectedComponents)
+	{
+		if (Component.IsValid())
+		{
+			Component.Get()->Generate(/*bForce=*/bIsControlDown);
 		}
 	}
 
