[PCG] Fix stale raw Context ptr
- Introduce an optional FPCGContextHandle, lambda captures can rely on this to determine if Context is still valid

#rnx

diff --git a/Engine/Plugins/PCG/Source/PCG/Private/Compute/Elements/PCGComputeGraphElement.cpp b/Engine/Plugins/PCG/Source/PCG/Private/Compute/Elements/PCGComputeGraphElement.cpp
--- a/Engine/Plugins/PCG/Source/PCG/Private/Compute/Elements/PCGComputeGraphElement.cpp
+++ b/Engine/Plugins/PCG/Source/PCG/Private/Compute/Elements/PCGComputeGraphElement.cpp
@@ -106,13 +106,19 @@ bool FPCGComputeGraphElement::ExecuteInternal(FPCGContext* InContext) const
 		return true;
 	}
 
-	auto SleepUntilNextFrame = [Context]()
+	auto SleepUntilNextFrame = [](FPCGComputeGraphContext* InContext)
 	{
-		// TODO unsafe access to raw pointer, need cancellation lambda
-		Context->bIsPaused = true;
-		Context->SourceComponent->GetSubsystem()->RegisterBeginTickAction([Context]()
+		check(InContext);
+		InContext->bIsPaused = true;
+		InContext->SourceComponent->GetSubsystem()->RegisterBeginTickAction([ContextHandle = InContext->GetOrCreateHandle()]()
 		{
-			Context->bIsPaused = false;
+			if (TSharedPtr<FPCGContextHandle> SharedHandle = ContextHandle.Pin())
+			{
+				if (FPCGContext* ContextPtr = SharedHandle->GetContext())
+				{
+					ContextPtr->bIsPaused = false;
+				}
+			}
 		});
 	};
 
@@ -120,7 +126,7 @@ bool FPCGComputeGraphElement::ExecuteInternal(FPCGContext* InContext) const
 	if (Context->ComputeGraph->HasKernelResourcesPendingShaderCompilation())
 	{
 		UE_LOG(LogPCG, Log, TEXT("Deferring until next frame as the kernel has pending shader compilations."));
-		SleepUntilNextFrame();
+		SleepUntilNextFrame(Context);
 		return false;
 	}
 
@@ -128,7 +134,7 @@ bool FPCGComputeGraphElement::ExecuteInternal(FPCGContext* InContext) const
 	if (Context->bGraphEnqueued)
 	{
 		// Likely we need a frame to pass in order to make progress.
-		SleepUntilNextFrame();
+		SleepUntilNextFrame(Context);
 		return true;
 	}
 
@@ -193,7 +199,7 @@ bool FPCGComputeGraphElement::ExecuteInternal(FPCGContext* InContext) const
 		if (bAnyComponentsSetup)
 		{
 			// Delay to give time for proxy to settle and for scene update to allocate space in the GPU scene.
-			SleepUntilNextFrame();
+			SleepUntilNextFrame(Context);
 			return false;
 		}
 	}
@@ -207,7 +213,7 @@ bool FPCGComputeGraphElement::ExecuteInternal(FPCGContext* InContext) const
 
 		Context->ComputeGraph->UpdateResources();
 
-		SleepUntilNextFrame();
+		SleepUntilNextFrame(Context);
 		return false;
 	}
 
diff --git a/Engine/Plugins/PCG/Source/PCG/Private/Elements/Landscape/PCGWaitLandscapeReady.cpp b/Engine/Plugins/PCG/Source/PCG/Private/Elements/Landscape/PCGWaitLandscapeReady.cpp
--- a/Engine/Plugins/PCG/Source/PCG/Private/Elements/Landscape/PCGWaitLandscapeReady.cpp
+++ b/Engine/Plugins/PCG/Source/PCG/Private/Elements/Landscape/PCGWaitLandscapeReady.cpp
@@ -102,15 +102,21 @@ bool FPCGWaitLandscapeReadyElement::ExecuteInternal(FPCGContext* InContext) cons
 		Context->bIsPaused = true;
 
 		Subsystem->ScheduleGeneric(
-			[Context]() // Normal execution: wake up the current task
+			[ContextHandle = Context->GetOrCreateHandle()]() // Normal execution: wake up the current task
 			{
-				Context->bIsPaused = false;
+				if (FPCGWaitLandscapeReadyElementContext* ContextPtr = FPCGContext::GetContextFromHandle<FPCGWaitLandscapeReadyElementContext>(ContextHandle))
+				{
+					ContextPtr->bIsPaused = false;
+				}
 				return true;
 			},
-			[Context]() // On abort: wakeup and cancel
+			[ContextHandle = Context->GetOrCreateHandle()]() // On abort: wakeup and cancel
 			{
-				Context->bIsPaused = false;
-				Context->OutputData.bCancelExecution = true;
+				if (FPCGWaitLandscapeReadyElementContext* ContextPtr = FPCGContext::GetContextFromHandle<FPCGWaitLandscapeReadyElementContext>(ContextHandle))
+				{
+					ContextPtr->bIsPaused = false;
+					ContextPtr->OutputData.bCancelExecution = true;
+				}
 				return true;
 			},
 			Context->SourceComponent.Get(),
diff --git a/Engine/Plugins/PCG/Source/PCG/Private/Elements/PCGDataFromActor.cpp b/Engine/Plugins/PCG/Source/PCG/Private/Elements/PCGDataFromActor.cpp
--- a/Engine/Plugins/PCG/Source/PCG/Private/Elements/PCGDataFromActor.cpp
+++ b/Engine/Plugins/PCG/Source/PCG/Private/Elements/PCGDataFromActor.cpp
@@ -492,16 +492,22 @@ bool FPCGDataFromActorElement::ExecuteInternal(FPCGContext* InContext) const
 					Context->bIsPaused = true;
 
 					Subsystem->ScheduleGeneric(
-						[Context]() // Normal execution: Wake up the current task
+						[ContextHandle = Context->GetOrCreateHandle()]() // Normal execution: Wake up the current task
 						{
-							Context->bIsPaused = false; 
+							if (FPCGDataFromActorContext* ContextPtr = FPCGContext::GetContextFromHandle<FPCGDataFromActorContext>(ContextHandle))
+							{
+								ContextPtr->bIsPaused = false;
+							}
 							return true;
 						}, 
-						[Context]() // On Abort: Wake up on abort, clear all results and mark as cancelled
+						[ContextHandle = Context->GetOrCreateHandle()]() // On Abort: Wake up on abort, clear all results and mark as cancelled
 						{
-							Context->bIsPaused = false; 
-							Context->FoundActors.Reset();
-							Context->OutputData.bCancelExecution = true;
+							if (FPCGDataFromActorContext* ContextPtr = FPCGContext::GetContextFromHandle<FPCGDataFromActorContext>(ContextHandle))
+							{
+								ContextPtr->bIsPaused = false;
+								ContextPtr->FoundActors.Reset();
+								ContextPtr->OutputData.bCancelExecution = true;
+							}
 							return true;
 						},
 						Context->SourceComponent.Get(),
diff --git a/Engine/Plugins/PCG/Source/PCG/Private/Elements/PCGLoopElement.cpp b/Engine/Plugins/PCG/Source/PCG/Private/Elements/PCGLoopElement.cpp
--- a/Engine/Plugins/PCG/Source/PCG/Private/Elements/PCGLoopElement.cpp
+++ b/Engine/Plugins/PCG/Source/PCG/Private/Elements/PCGLoopElement.cpp
@@ -312,16 +312,22 @@ bool FPCGLoopElement::ExecuteInternal(FPCGContext* InContext) const
 			Context->bIsPaused = true;
 
 			Subsystem->ScheduleGeneric(
-				[Context]() // Normal execution: Wake up the current task
+				[ContextHandle = Context->GetOrCreateHandle()]() // Normal execution: Wake up the current task
 				{
-					Context->bIsPaused = false;
+					if (FPCGSubgraphContext* ContextPtr = FPCGContext::GetContextFromHandle<FPCGSubgraphContext>(ContextHandle))
+					{
+						ContextPtr->bIsPaused = false;
+					}
 					return true;
 				},
-				[Context]() // On abort: wakeup and cancel, forget subgraphs
+				[ContextHandle = Context->GetOrCreateHandle()]() // On abort: wakeup and cancel, forget subgraphs
 				{
-					Context->bIsPaused = false;
-					Context->SubgraphTaskIds.Reset();
-					Context->OutputData.bCancelExecution = true;
+					if (FPCGSubgraphContext* ContextPtr = FPCGContext::GetContextFromHandle<FPCGSubgraphContext>(ContextHandle))
+					{
+						ContextPtr->bIsPaused = false;
+						ContextPtr->SubgraphTaskIds.Reset();
+						ContextPtr->OutputData.bCancelExecution = true;
+					}
 					return true;
 				},
 				Context->SourceComponent.Get(),
diff --git a/Engine/Plugins/PCG/Source/PCG/Private/Elements/PCGSpawnActor.cpp b/Engine/Plugins/PCG/Source/PCG/Private/Elements/PCGSpawnActor.cpp
--- a/Engine/Plugins/PCG/Source/PCG/Private/Elements/PCGSpawnActor.cpp
+++ b/Engine/Plugins/PCG/Source/PCG/Private/Elements/PCGSpawnActor.cpp
@@ -638,15 +638,21 @@ bool FPCGSpawnActorElement::SpawnAndPrepareSubgraphs(FPCGSubgraphContext* Contex
 		Context->bIsPaused = true;
 
 		Subsystem->ScheduleGeneric(
-			[Context]() // Normal execution: Wake up the current task
+			[ContextHandle = Context->GetOrCreateHandle()]() // Normal execution: Wake up the current task
 			{
-				Context->bIsPaused = false;
+				if (FPCGSubgraphContext* ContextPtr = FPCGContext::GetContextFromHandle<FPCGSubgraphContext>(ContextHandle))
+				{
+					ContextPtr->bIsPaused = false;
+				}
 				return true;
 			}, 
-			[Context]() // On Abort: Wake up & cancel
+			[ContextHandle = Context->GetOrCreateHandle()]() // On Abort: Wake up & cancel
 			{
-				Context->bIsPaused = false;
-				Context->OutputData.bCancelExecution = true;
+				if (FPCGSubgraphContext* ContextPtr = FPCGContext::GetContextFromHandle<FPCGSubgraphContext>(ContextHandle))
+				{
+					ContextPtr->bIsPaused = false;
+					ContextPtr->OutputData.bCancelExecution = true;
+				}
 			},
 			Context->SourceComponent.Get(), 
 			Context->SubgraphTaskIds);
diff --git a/Engine/Plugins/PCG/Source/PCG/Private/Elements/PCGTextureSampler.cpp b/Engine/Plugins/PCG/Source/PCG/Private/Elements/PCGTextureSampler.cpp
--- a/Engine/Plugins/PCG/Source/PCG/Private/Elements/PCGTextureSampler.cpp
+++ b/Engine/Plugins/PCG/Source/PCG/Private/Elements/PCGTextureSampler.cpp
@@ -245,9 +245,12 @@ bool FPCGTextureSamplerElement::ExecuteInternal(FPCGContext* InContext) const
 		{
 			// Initialization not complete. Could be waiting on async texture processing or for GPU readback. Sleep until next frame.
 			Context->bIsPaused = true;
-			Context->SourceComponent->GetSubsystem()->RegisterBeginTickAction([Context]()
+			Context->SourceComponent->GetSubsystem()->RegisterBeginTickAction([ContextHandle = Context->GetOrCreateHandle()]()
 			{
-				Context->bIsPaused = false;
+				if (FPCGTextureSamplerContext* ContextPtr = FPCGContext::GetContextFromHandle<FPCGTextureSamplerContext>(ContextHandle))
+				{
+					ContextPtr->bIsPaused = false;
+				}
 			});
 
 			return false;
diff --git a/Engine/Plugins/PCG/Source/PCG/Private/PCGSubgraph.cpp b/Engine/Plugins/PCG/Source/PCG/Private/PCGSubgraph.cpp
--- a/Engine/Plugins/PCG/Source/PCG/Private/PCGSubgraph.cpp
+++ b/Engine/Plugins/PCG/Source/PCG/Private/PCGSubgraph.cpp
@@ -728,16 +728,22 @@ bool FPCGSubgraphElement::ExecuteInternal(FPCGContext* InContext) const
 					
 					// add a trivial task after the output task that wakes up this task
 					Subsystem->ScheduleGeneric(
-						[Context]() // Normal execution: Wake up the current task
+						[ContextHandle = Context->GetOrCreateHandle()]() // Normal execution: Wake up the current task
 						{
-							Context->bIsPaused = false;
+							if (FPCGSubgraphContext* ContextPtr = FPCGContext::GetContextFromHandle<FPCGSubgraphContext>(ContextHandle))
+							{
+								ContextPtr->bIsPaused = false;
+							}					
 							return true;
 						}, 
-						[Context]() // On Abort: Wake up and cancel the execution
+						[ContextHandle = Context->GetOrCreateHandle()]() // On Abort: Wake up and cancel the execution
 						{
-							Context->bIsPaused = false;
-							Context->OutputData.bCancelExecution = true;
-							Context->SubgraphTaskIds.Reset();
+							if (FPCGSubgraphContext* ContextPtr = FPCGContext::GetContextFromHandle<FPCGSubgraphContext>(ContextHandle))
+							{
+								ContextPtr->bIsPaused = false;
+								ContextPtr->OutputData.bCancelExecution = true;
+								ContextPtr->SubgraphTaskIds.Reset();
+							}
 							return true;
 						},
 						Context->SourceComponent.Get(), 
diff --git a/Engine/Plugins/PCG/Source/PCG/Public/PCGContext.h b/Engine/Plugins/PCG/Source/PCG/Public/PCGContext.h
--- a/Engine/Plugins/PCG/Source/PCG/Public/PCGContext.h
+++ b/Engine/Plugins/PCG/Source/PCG/Public/PCGContext.h
@@ -45,6 +45,21 @@ enum class EPCGExecutionPhase : uint8
 		Done
 };
 
+struct FPCGContext;
+
+struct PCG_API FPCGContextHandle : public TSharedFromThis<FPCGContextHandle>
+{
+public:
+	FPCGContextHandle(FPCGContext* InContext)
+		: Context(InContext)
+	{
+	}
+
+	FPCGContext* GetContext() { return Context; }
+private:
+	FPCGContext* Context = nullptr;
+};
+
 USTRUCT(BlueprintType)
 struct PCG_API FPCGContext
 {
@@ -145,6 +160,27 @@ struct PCG_API FPCGContext
 		return PCGContextHelpers::GetInputSettings<SettingsType>(Node, InputData);
 	}
 
+	TWeakPtr<FPCGContextHandle> GetOrCreateHandle()
+	{
+		if (!Handle)
+		{
+			Handle = MakeShared<FPCGContextHandle>(this);
+		}
+
+		return Handle.ToWeakPtr();
+	}
+
+	template <typename ContextType>
+	static ContextType* GetContextFromHandle(TWeakPtr<FPCGContextHandle> WeakHandle)
+	{
+		if (TSharedPtr<FPCGContextHandle> SharedHandle = WeakHandle.Pin())
+		{
+			return (ContextType*)SharedHandle->GetContext();
+		}
+
+		return nullptr;
+	}
+
 protected:
 	virtual UObject* GetExternalContainerForOverridableParam(const FPCGSettingsOverridableParam& InParam) { return nullptr; }
 	virtual void* GetUnsafeExternalContainerForOverridableParam(const FPCGSettingsOverridableParam& InParam) { return nullptr; }
@@ -183,9 +219,12 @@ private:
 	// so that they can be considered as existing on the main thread (and get properly GCed)
 	TSet<TObjectPtr<UObject>> AsyncObjects;
 
+	// Lazy initialized shared handle pointer that can be used in lambda captures to test if Context is still valid before accessing it
+	TSharedPtr<FPCGContextHandle> Handle;
+
 	friend struct FPCGGraphActiveTask;
 #if WITH_EDITOR
 	friend class PCGUtils::FExtraCapture;
 	PCGUtils::FCallTime Timer;
 #endif
-};
\ No newline at end of file
+};
