[PCG] Fix grammar attribute override not exposed.

* Also removed the duplicated GrammarAttribute, it doesn't require any deprecation, it will work as-is.

#rnx

diff --git a/Engine/Plugins/PCG/Source/PCG/Public/Elements/Grammar/PCGSubdivisionBase.h b/Engine/Plugins/PCG/Source/PCG/Public/Elements/Grammar/PCGSubdivisionBase.h
--- a/Engine/Plugins/PCG/Source/PCG/Public/Elements/Grammar/PCGSubdivisionBase.h
+++ b/Engine/Plugins/PCG/Source/PCG/Public/Elements/Grammar/PCGSubdivisionBase.h
@@ -102,13 +102,9 @@ public:
 	FPCGSubdivisionModuleAttributeNames ModulesInfoAttributeNames;
 
 	/** An encoded string that represents how to apply a set of rules to a series of defined modules. */
-	UPROPERTY(BlueprintReadWrite, EditAnywhere, Category = Settings, meta = (ShowOnlyInnerProperties))
+	UPROPERTY(BlueprintReadWrite, EditAnywhere, Category = Settings, meta = (ShowOnlyInnerProperties, PCG_Overridable))
 	FPCGGrammarSelection GrammarSelection;
 
-	/** Attribute to be taken from the input spline containing the grammar to use for the subdivision. */
-	UPROPERTY(BlueprintReadWrite, EditAnywhere, Category = Settings, meta = (EditCondition = "bGrammarAsAttribute", EditConditionHides, DisplayAfter = bGrammarAsAttribute, PCG_Overridable))
-	FPCGAttributePropertyInputSelector GrammarAttribute;
-
 	/** Controls whether we'll use an attribute to drive random seeding for stochastic processes in the subdivision. */
 	UPROPERTY(BlueprintReadWrite, EditAnywhere, Category = Settings, meta = (PCG_Overridable, InlineEditConditionToggle))
 	bool bUseSeedAttribute = false;
diff --git a/Engine/Plugins/PCG/Source/PCG/Public/Grammar/PCGGrammar.h b/Engine/Plugins/PCG/Source/PCG/Public/Grammar/PCGGrammar.h
--- a/Engine/Plugins/PCG/Source/PCG/Public/Grammar/PCGGrammar.h
+++ b/Engine/Plugins/PCG/Source/PCG/Public/Grammar/PCGGrammar.h
@@ -57,7 +57,7 @@ struct FPCGGrammarSelection
 	* - Modules can be repeated a fixed number of times, by adding a number after it (like [A,B]3 will produce ABABAB)
 	* - Modules can be marked repeated an indefinite number of times, with '*'. (like [A,B]* will produce ABABABAB... while it fits the allowed size).
 	*/
-	UPROPERTY(BlueprintReadWrite, EditAnywhere, Category = Settings)
+	UPROPERTY(BlueprintReadWrite, EditAnywhere, Category = Settings, meta = (PCG_Overridable))
 	bool bGrammarAsAttribute = false;
 
 	/** An encoded string that represents how to apply a set of rules to a series of defined modules. */
