[PCG] GPU - fix failed static mesh spawning in standalone builds where proxies were not ready yet.

#rnx

diff --git a/Engine/Plugins/PCG/Source/PCG/Private/Compute/Elements/PCGComputeGraphElement.cpp b/Engine/Plugins/PCG/Source/PCG/Private/Compute/Elements/PCGComputeGraphElement.cpp
--- a/Engine/Plugins/PCG/Source/PCG/Private/Compute/Elements/PCGComputeGraphElement.cpp
+++ b/Engine/Plugins/PCG/Source/PCG/Private/Compute/Elements/PCGComputeGraphElement.cpp
@@ -206,6 +206,37 @@ bool FPCGComputeGraphElement::ExecuteInternal(FPCGContext* InContext) const
 
 	check(Context->DataBinding && InContext->SourceComponent.Get());
 
+	if (!Context->bPrimitiveProxiesValidated)
+	{
+		bool bAllProxiesCreated = true;
+
+		for (const TPair<TObjectPtr<const UPCGSettings>, FPCGSpawnerPrimitives>& SpawnerToPrimitives : Context->DataBinding->MeshSpawnersToPrimitives)
+		{
+			for (const TObjectPtr<UPrimitiveComponent>& Primitive : SpawnerToPrimitives.Value.Primitives)
+			{
+				const UPCGProceduralISMComponent* PISMC = Cast<UPCGProceduralISMComponent>(Primitive);
+				if (PISMC && PISMC->GetNumInstances() > 0 && !PISMC->GetSceneProxy())
+				{
+					bAllProxiesCreated = false;
+					break;
+				}
+			}
+
+			if (!bAllProxiesCreated)
+			{
+				break;
+			}
+		}
+
+		if (!bAllProxiesCreated)
+		{
+			SleepUntilNextFrame(Context);
+			return false;
+		}
+
+		Context->bPrimitiveProxiesValidated = true;
+	}
+
 	// 1. Prepare render resources. In editor, this will trigger shader compile if not compiled already.
 	if (!Context->ComputeGraph->GetRenderProxy())
 	{
diff --git a/Engine/Plugins/PCG/Source/PCG/Private/Compute/Elements/PCGComputeGraphElement.h b/Engine/Plugins/PCG/Source/PCG/Private/Compute/Elements/PCGComputeGraphElement.h
--- a/Engine/Plugins/PCG/Source/PCG/Private/Compute/Elements/PCGComputeGraphElement.h
+++ b/Engine/Plugins/PCG/Source/PCG/Private/Compute/Elements/PCGComputeGraphElement.h
@@ -40,6 +40,9 @@ public:
 
 	/** Graph executed successfully. */
 	bool bExecutionSuccess = false;
+
+	/** Whether all required primitive proxies have been created. */
+	bool bPrimitiveProxiesValidated = false;
 };
 
 /** Executes a CF graph. Created by the compiler when collapsing GPU nodes rather than by a settings/node. */
