[PCG] Apply some sanity limits to GPU buffer sizes and instance counts, abort generation with error log if sizes dangerously large.

#rnx

diff --git a/Engine/Plugins/PCG/Source/PCG/Private/Compute/DataInterfaces/PCGDataCollectionDataInterface.cpp b/Engine/Plugins/PCG/Source/PCG/Private/Compute/DataInterfaces/PCGDataCollectionDataInterface.cpp
--- a/Engine/Plugins/PCG/Source/PCG/Private/Compute/DataInterfaces/PCGDataCollectionDataInterface.cpp
+++ b/Engine/Plugins/PCG/Source/PCG/Private/Compute/DataInterfaces/PCGDataCollectionDataInterface.cpp
@@ -413,10 +413,17 @@ UComputeDataProvider* UPCGDataCollectionDataInterface::CreateDataProvider(TObjec
 	TRACE_CPUPROFILER_EVENT_SCOPE(UPCGDataCollectionDataInterface::CreateDataProvider);
 	UPCGDataBinding* Binding = CastChecked<UPCGDataBinding>(InBinding);
 
+	FPCGDataCollectionDesc PinDataDesc = ProducerSettings->ComputeOutputPinDataDesc(OutputPinLabel, Binding);
+
+	if (PCGComputeHelpers::IsBufferSizeTooLarge(PinDataDesc.ComputePackedSizeBytes()))
+	{
+		return nullptr;
+	}
+
 	UPCGDataCollectionDataProvider* Provider = NewObject<UPCGDataCollectionDataProvider>();
 	Provider->Binding = Binding;
 	Provider->ProducerSettings = ProducerSettings;
-	Provider->PinDesc = ProducerSettings->ComputeOutputPinDataDesc(OutputPinLabel, Binding);
+	Provider->PinDesc = MoveTemp(PinDataDesc);
 
 	Provider->ReadbackMode = bRequiresReadback ? EPCGReadbackMode::GraphOutput : EPCGReadbackMode::None;
 
diff --git a/Engine/Plugins/PCG/Source/PCG/Private/Compute/DataInterfaces/PCGDataCollectionUploadDataInterface.cpp b/Engine/Plugins/PCG/Source/PCG/Private/Compute/DataInterfaces/PCGDataCollectionUploadDataInterface.cpp
--- a/Engine/Plugins/PCG/Source/PCG/Private/Compute/DataInterfaces/PCGDataCollectionUploadDataInterface.cpp
+++ b/Engine/Plugins/PCG/Source/PCG/Private/Compute/DataInterfaces/PCGDataCollectionUploadDataInterface.cpp
@@ -6,6 +6,7 @@
 #include "PCGModule.h"
 #include "PCGNode.h"
 #include "PCGSettings.h"
+#include "Compute/PCGComputeCommon.h"
 #include "Compute/PCGComputeGraph.h"
 #include "Compute/PCGDataBinding.h"
 
@@ -43,6 +44,11 @@ FComputeDataProviderRenderProxy* UPCGDataProviderDataCollectionUpload::GetRender
 	check(!DownstreamInputPinLabels.IsEmpty());
 	PinDesc.PackDataCollection(Binding->DataForGPU.InputDataCollection, DownstreamInputPinLabels[0], Binding->GetStringTable(), PackedDataCollection);
 
+	if (PCGComputeHelpers::IsBufferSizeTooLarge(PackedDataCollection.Num() * PackedDataCollection.GetTypeSize()))
+	{
+		return nullptr;
+	}
+
 	return new FPCGDataProviderDataCollectionUploadProxy(PinDesc, MoveTemp(PackedDataCollection), DownstreamInputPinLabels);
 }
 
diff --git a/Engine/Plugins/PCG/Source/PCG/Private/Compute/DataInterfaces/PCGInstanceDataInterface.cpp b/Engine/Plugins/PCG/Source/PCG/Private/Compute/DataInterfaces/PCGInstanceDataInterface.cpp
--- a/Engine/Plugins/PCG/Source/PCG/Private/Compute/DataInterfaces/PCGInstanceDataInterface.cpp
+++ b/Engine/Plugins/PCG/Source/PCG/Private/Compute/DataInterfaces/PCGInstanceDataInterface.cpp
@@ -113,18 +113,16 @@ UComputeDataProvider* UPCGInstanceDataInterface::CreateDataProvider(TObjectPtr<U
 {
 	TRACE_CPUPROFILER_EVENT_SCOPE(UPCGInstanceDataInterface::CreateDataProvider);
 	UPCGDataBinding* Binding = CastChecked<UPCGDataBinding>(InBinding);
-
-	TObjectPtr<UPCGInstanceDataProvider> DataProvider = NewObject<UPCGInstanceDataProvider>();
-
 	check(ProducerSettings);
-	const FPCGDataCollectionDesc InputDataDesc = ProducerSettings->ComputeInputPinDataDesc(InputPinProvidingData, Binding);
-	DataProvider->NumInstancesAllPrimitives = InputDataDesc.ComputeDataElementCount(EPCGDataType::Point);
 
-	FPCGSpawnerPrimitives* FoundPrimitives = Binding->MeshSpawnersToPrimitives.Find(ProducerSettings);
+	TObjectPtr<UPCGInstanceDataProvider> DataProvider = NewObject<UPCGInstanceDataProvider>();
 
-	// If there were 0 input points for this execution, we will not have created any primitives, so check for null.
-	if (FoundPrimitives)
+	// If there were 0 input points (or too many input points) for this execution, we will not have created any primitives, so check for null.
+	if (FPCGSpawnerPrimitives* FoundPrimitives = Binding->MeshSpawnersToPrimitives.Find(ProducerSettings))
 	{
+		const FPCGDataCollectionDesc InputDataDesc = ProducerSettings->ComputeInputPinDataDesc(InputPinProvidingData, Binding);
+		DataProvider->NumInstancesAllPrimitives = InputDataDesc.ComputeDataElementCount(EPCGDataType::Point);
+
 		if (!ensure(FoundPrimitives->Primitives.Num() <= PCGComputeConstants::MAX_PRIMITIVE_COMPONENTS_PER_SPAWNER))
 		{
 			// Last resort - should be clamped earlier during setup and not come to this.
@@ -140,10 +138,6 @@ UComputeDataProvider* UPCGInstanceDataInterface::CreateDataProvider(TObjectPtr<U
 			DataProvider->NumInstancesAllPrimitives *= FoundPrimitives->Primitives.Num();
 		}
 	}
-	else
-	{
-		ensure(DataProvider->NumInstancesAllPrimitives == 0);
-	}
 
 	return DataProvider;
 }
diff --git a/Engine/Plugins/PCG/Source/PCG/Private/Compute/Elements/PCGComputeGraphElement.cpp b/Engine/Plugins/PCG/Source/PCG/Private/Compute/Elements/PCGComputeGraphElement.cpp
--- a/Engine/Plugins/PCG/Source/PCG/Private/Compute/Elements/PCGComputeGraphElement.cpp
+++ b/Engine/Plugins/PCG/Source/PCG/Private/Compute/Elements/PCGComputeGraphElement.cpp
@@ -283,7 +283,7 @@ bool FPCGComputeGraphElement::ExecuteInternal(FPCGContext* InContext) const
 			FSimpleDelegate());
 	}
 
-	if (ensure(Context->bGraphEnqueued))
+	if (Context->bGraphEnqueued)
 	{
 		FReadScopeLock Lock(Context->ProvidersRunningAsyncOperationsLock);
 
@@ -459,6 +459,21 @@ bool FPCGComputeGraphElement::SetupProceduralISMComponents(FPCGContext* InContex
 			}
 		}
 
+		// Validate instance count is not too large.
+		{
+			uint64 TotalInstanceCount = 0;
+			for (const FPCGProceduralISMComponentDescriptor& Desc : ComponentsToCreate)
+			{
+				TotalInstanceCount += Desc.NumInstances;
+			}
+
+			if (TotalInstanceCount >= MAX_INSTANCE_ID)
+			{
+				UE_LOG(LogPCG, Error, TEXT("Tried to spawn too many instances (%lu), procedural ISM component creation skipped and instances will not be rendered."), TotalInstanceCount);
+				return false;
+			}
+		}
+
 		FPCGSpawnerPrimitives& Primitives = InBinding->MeshSpawnersToPrimitives.FindOrAdd(SpawnerSettings);
 		Primitives.NumCustomFloats = CustomFloatCount;
 		Primitives.AttributeIdOffsetStrides = MoveTemp(AttributeIdOffsetStrides);
diff --git a/Engine/Plugins/PCG/Source/PCG/Private/Compute/PCGComputeCommon.cpp b/Engine/Plugins/PCG/Source/PCG/Private/Compute/PCGComputeCommon.cpp
--- a/Engine/Plugins/PCG/Source/PCG/Private/Compute/PCGComputeCommon.cpp
+++ b/Engine/Plugins/PCG/Source/PCG/Private/Compute/PCGComputeCommon.cpp
@@ -9,8 +9,16 @@
 #include "Compute/PCGPinPropertiesGPU.h"
 #include "Data/PCGPointData.h"
 
+#include "DynamicRHI.h"
+#include "RHIStats.h"
+
 namespace PCGComputeHelpers
 {
+	static TAutoConsoleVariable<float> CVarMaxGPUBufferSizeProportion(
+		TEXT("pcg.GraphExecution.GPU.MaxBufferSize"),
+		0.5f,
+		TEXT("Maximum GPU buffer size as proportion of total available graphics memory."));
+
 	int GetElementCount(const UPCGData* InData)
 	{
 		if (const UPCGPointData* PointData = Cast<UPCGPointData>(InData))
@@ -288,4 +296,24 @@ namespace PCGComputeHelpers
 			checkNoEntry();
 		}
 	}
+
+	bool IsBufferSizeTooLarge(uint64 InBufferSizeBytes, bool bInLogError)
+	{
+		FTextureMemoryStats TextureMemStats;
+		RHIGetTextureMemoryStats(TextureMemStats);
+
+		// If buffer size exceeds a proportion of total graphics memory, then it is deemed too large. Using this as a heuristic as there
+		// is no RHI API to obtain available graphics memory outside of D3D12.
+		const uint64 BudgetBytes = static_cast<uint64>(TextureMemStats.TotalGraphicsMemory * CVarMaxGPUBufferSizeProportion.GetValueOnAnyThread());
+		const bool bBufferTooLarge = TextureMemStats.TotalGraphicsMemory > 0 && InBufferSizeBytes > BudgetBytes;
+
+		if (bBufferTooLarge && bInLogError)
+		{
+			UE_LOG(LogPCG, Error, TEXT("Attempted to allocate a GPU buffer of size %llu bytes which is larger than the safety threshold (%llu bytes). Compute graph execution aborted."),
+				InBufferSizeBytes,
+				BudgetBytes);
+		}
+
+		return bBufferTooLarge;
+	}
 }
diff --git a/Engine/Plugins/PCG/Source/PCG/Private/Compute/PCGComputeCommon.h b/Engine/Plugins/PCG/Source/PCG/Private/Compute/PCGComputeCommon.h
--- a/Engine/Plugins/PCG/Source/PCG/Private/Compute/PCGComputeCommon.h
+++ b/Engine/Plugins/PCG/Source/PCG/Private/Compute/PCGComputeCommon.h
@@ -84,4 +84,7 @@ namespace PCGComputeHelpers
 
 	/** Computes a pin data description based on the given GPU pin properties. */
 	void ComputeOutputPinDataDesc(const FPCGPinPropertiesGPU& PinProperties, const UPCGSettings* Settings, const UPCGDataBinding* Binding, FPCGDataCollectionDesc& OutPinDesc);
+
+	/** Returns true if the given buffer size is dangerously large. Optionally emits error log. */
+	bool IsBufferSizeTooLarge(uint64 InBufferSizeBytes, bool bInLogError = true);
 }
diff --git a/Engine/Plugins/PCG/Source/PCG/Private/Compute/PCGDataForGPU.cpp b/Engine/Plugins/PCG/Source/PCG/Private/Compute/PCGDataForGPU.cpp
--- a/Engine/Plugins/PCG/Source/PCG/Private/Compute/PCGDataForGPU.cpp
+++ b/Engine/Plugins/PCG/Source/PCG/Private/Compute/PCGDataForGPU.cpp
@@ -496,15 +496,15 @@ FPCGDataDesc::FPCGDataDesc(const UPCGData* InData, const TMap<FName, FPCGKernelA
 	InitializeAttributeDescs(InData, InGlobalAttributeLookupTable, InStringTable);
 }
 
-uint32 FPCGDataDesc::ComputePackedSize() const
+uint64 FPCGDataDesc::ComputePackedSize() const
 {
 	check(PCGComputeHelpers::IsTypeAllowedInDataCollection(Type));
 
-	uint32 DataSizeBytes = DATA_HEADER_SIZE_BYTES;
+	uint64 DataSizeBytes = DATA_HEADER_SIZE_BYTES;
 
 	for (const FPCGKernelAttributeDesc& AttributeDesc : AttributeDescs)
 	{
-		DataSizeBytes += PCGDataForGPUHelpers::GetAttributeTypeStrideBytes(AttributeDesc.Type) * ElementCount;
+		DataSizeBytes += static_cast<uint64>(PCGDataForGPUHelpers::GetAttributeTypeStrideBytes(AttributeDesc.Type)) * static_cast<uint64>(ElementCount);
 	}
 
 	return DataSizeBytes;
@@ -678,11 +678,11 @@ uint32 FPCGDataCollectionDesc::ComputePackedHeaderSizeBytes() const
 	return PCGComputeConstants::DATA_COLLECTION_HEADER_SIZE_BYTES + PCGComputeConstants::DATA_HEADER_SIZE_BYTES * DataDescs.Num();
 }
 
-uint32 FPCGDataCollectionDesc::ComputePackedSizeBytes() const
+uint64 FPCGDataCollectionDesc::ComputePackedSizeBytes() const
 {
 	TRACE_CPUPROFILER_EVENT_SCOPE(FPCGDataCollectionDesc::ComputePackedSize);
 
-	uint32 TotalCollectionSizeBytes = ComputePackedHeaderSizeBytes();
+	uint64 TotalCollectionSizeBytes = ComputePackedHeaderSizeBytes();
 
 	for (const FPCGDataDesc& DataDesc : DataDescs)
 	{
diff --git a/Engine/Plugins/PCG/Source/PCG/Public/Compute/PCGDataForGPU.h b/Engine/Plugins/PCG/Source/PCG/Public/Compute/PCGDataForGPU.h
--- a/Engine/Plugins/PCG/Source/PCG/Public/Compute/PCGDataForGPU.h
+++ b/Engine/Plugins/PCG/Source/PCG/Public/Compute/PCGDataForGPU.h
@@ -112,7 +112,7 @@ struct FPCGDataDesc
 	FPCGDataDesc(EPCGDataType InType, int32 InElementCount);
 	FPCGDataDesc(const UPCGData* InData, const TMap<FName, FPCGKernelAttributeIDAndType>& InGlobalAttributeLookupTable, const TArray<FString>& InStringTable);
 
-	uint32 ComputePackedSize() const;
+	uint64 ComputePackedSize() const;
 
 	EPCGDataType Type = EPCGDataType::Point;
 	TArray<FPCGKernelAttributeDesc> AttributeDescs;
@@ -134,7 +134,7 @@ struct FPCGDataCollectionDesc
 	uint32 ComputePackedHeaderSizeBytes() const;
 
 	/** Computes the size (in bytes) of the data collection after packing. */
-	uint32 ComputePackedSizeBytes() const;
+	uint64 ComputePackedSizeBytes() const;
 
 	void WriteHeader(TArray<uint32>& OutPackedDataCollectionHeader) const;
 
