[PCG] Make sure the Static Mesh Spawner is not spawing grass ISMCs as it is not meant to be used outside of landscape grass.

* Also hide UGrassInstancedStaticMeshComponent class so it's not selectable in the ISM descriptor.
* Was asserting

#rnx

diff --git a/Engine/Plugins/PCG/Source/PCG/Private/Elements/PCGStaticMeshSpawner.cpp b/Engine/Plugins/PCG/Source/PCG/Private/Elements/PCGStaticMeshSpawner.cpp
--- a/Engine/Plugins/PCG/Source/PCG/Private/Elements/PCGStaticMeshSpawner.cpp
+++ b/Engine/Plugins/PCG/Source/PCG/Private/Elements/PCGStaticMeshSpawner.cpp
@@ -19,6 +19,7 @@
 #include "MeshSelectors/PCGMeshSelectorByAttribute.h"
 #include "MeshSelectors/PCGMeshSelectorWeighted.h"
 
+#include "GrassInstancedStaticMeshComponent.h"
 #include "Components/InstancedStaticMeshComponent.h"
 #include "Engine/StaticMesh.h"
 #include "Materials/MaterialInterface.h"
@@ -597,6 +598,12 @@ void FPCGStaticMeshSpawnerElement::SpawnStaticMeshInstances(FPCGStaticMeshSpawne
 	{
 		return;
 	}
+	
+	if (InstanceList.Descriptor.ComponentClass && InstanceList.Descriptor.ComponentClass->IsChildOf<UGrassInstancedStaticMeshComponent>())
+	{
+		PCGLog::LogErrorOnGraph(LOCTEXT("NoSupportForGrassComponent", "Grass Instanced Static Mesh Component are not meant to be used outside of the landscape grass system, and are not supported by PCG."), Context);
+		return;
+	}
 
 	// Will be synchronously loaded if not loaded. But by default it should already have been loaded asynchronously in PrepareData, so this is free.
 	UStaticMesh* LoadedMesh = InstanceList.Descriptor.StaticMesh.LoadSynchronous();
diff --git a/Engine/Source/Runtime/Foliage/Public/GrassInstancedStaticMeshComponent.h b/Engine/Source/Runtime/Foliage/Public/GrassInstancedStaticMeshComponent.h
--- a/Engine/Source/Runtime/Foliage/Public/GrassInstancedStaticMeshComponent.h
+++ b/Engine/Source/Runtime/Foliage/Public/GrassInstancedStaticMeshComponent.h
@@ -6,7 +6,7 @@
 #include "Components/HierarchicalInstancedStaticMeshComponent.h"
 #include "GrassInstancedStaticMeshComponent.generated.h"
 
-UCLASS(ClassGroup = Foliage, MinimalAPI)
+UCLASS(ClassGroup = Foliage, MinimalAPI, Hidden)
 class UGrassInstancedStaticMeshComponent : public UHierarchicalInstancedStaticMeshComponent
 {
 	GENERATED_UCLASS_BODY()
