[PCG] Added Multiline to the grammar text editing

#rnx

diff --git a/Engine/Plugins/PCG/Source/PCG/Private/Grammar/PCGGrammarParser.cpp b/Engine/Plugins/PCG/Source/PCG/Private/Grammar/PCGGrammarParser.cpp
--- a/Engine/Plugins/PCG/Source/PCG/Private/Grammar/PCGGrammarParser.cpp
+++ b/Engine/Plugins/PCG/Source/PCG/Private/Grammar/PCGGrammarParser.cpp
@@ -35,6 +35,8 @@ namespace PCGGrammar
 		// TODO: parse tokens in a more efficient way; this relies on token separation from what should essentially be token delimiters
 		// This would also allow better information to the end user when the grammar is invalid (as we wouldn't need to modify the string).
 		FString Grammar = InGrammar;
+		Grammar.ReplaceInline(TEXT("\r"), TEXT(""));
+		Grammar.ReplaceInline(TEXT("\n"), TEXT(""));
 		Grammar.ReplaceInline(Delimiter, TEXT(" , "));
 		Grammar.ReplaceInline(WeightDelimiter, TEXT(" : "));
 		Grammar.ReplaceInline(SequenceStart, TEXT(" [ "));
@@ -201,4 +203,4 @@ namespace PCGGrammar
 	}
 }
 
-#undef LOCTEXT_NAMESPACE
+#undef LOCTEXT_NAMESPACE
\ No newline at end of file
diff --git a/Engine/Plugins/PCG/Source/PCG/Public/Elements/Grammar/PCGPrintGrammar.h b/Engine/Plugins/PCG/Source/PCG/Public/Elements/Grammar/PCGPrintGrammar.h
--- a/Engine/Plugins/PCG/Source/PCG/Public/Elements/Grammar/PCGPrintGrammar.h
+++ b/Engine/Plugins/PCG/Source/PCG/Public/Elements/Grammar/PCGPrintGrammar.h
@@ -31,7 +31,7 @@ protected:
 
 public:
 	/** The grammar to interpret. */
-	UPROPERTY(EditAnywhere, BlueprintReadWrite, Category = Settings, meta = (PCG_Overridable))
+	UPROPERTY(EditAnywhere, BlueprintReadWrite, Category = Settings, meta = (MultiLine, PCG_Overridable))
 	FString Grammar;
 };
 
@@ -39,4 +39,4 @@ class FPCGPrintGrammarElement : public IPCGElement
 {
 protected:
 	virtual bool ExecuteInternal(FPCGContext* Context) const override;
-};
+};
\ No newline at end of file
diff --git a/Engine/Plugins/PCG/Source/PCG/Public/Elements/Grammar/PCGSelectGrammar.h b/Engine/Plugins/PCG/Source/PCG/Public/Elements/Grammar/PCGSelectGrammar.h
--- a/Engine/Plugins/PCG/Source/PCG/Public/Elements/Grammar/PCGSelectGrammar.h
+++ b/Engine/Plugins/PCG/Source/PCG/Public/Elements/Grammar/PCGSelectGrammar.h
@@ -49,7 +49,7 @@ struct FPCGSelectGrammarCriterion
 	double SecondValue = 0.0;
 
 	/** The selected grammar result if the condition evaluates to true. */
-	UPROPERTY(BlueprintReadWrite, EditAnywhere, Category = Settings)
+	UPROPERTY(BlueprintReadWrite, EditAnywhere, Category = Settings, meta = (MultiLine))
 	FString Grammar;
 };
 
@@ -146,4 +146,4 @@ class FPCGSelectGrammarElement : public FPCGSubdivisionBaseElement
 {
 protected:
 	virtual bool ExecuteInternal(FPCGContext* InContext) const override;
-};
+};
\ No newline at end of file
diff --git a/Engine/Plugins/PCG/Source/PCG/Public/Grammar/PCGGrammar.h b/Engine/Plugins/PCG/Source/PCG/Public/Grammar/PCGGrammar.h
--- a/Engine/Plugins/PCG/Source/PCG/Public/Grammar/PCGGrammar.h
+++ b/Engine/Plugins/PCG/Source/PCG/Public/Grammar/PCGGrammar.h
@@ -61,10 +61,10 @@ struct FPCGGrammarSelection
 	bool bGrammarAsAttribute = false;
 
 	/** An encoded string that represents how to apply a set of rules to a series of defined modules. */
-	UPROPERTY(BlueprintReadWrite, EditAnywhere, Category = Settings, meta = (EditCondition = "!bGrammarAsAttribute", EditConditionHides, PCG_Overridable))
+	UPROPERTY(BlueprintReadWrite, EditAnywhere, Category = Settings, meta = (MultiLine, EditCondition = "!bGrammarAsAttribute", EditConditionHides, PCG_Overridable))
 	FString GrammarString;
 
 	/** Attribute to be taken from the input spline containing the grammar to use. */
 	UPROPERTY(BlueprintReadWrite, EditAnywhere, Category = Settings, meta = (EditCondition = "bGrammarAsAttribute", EditConditionHides, PCG_Overridable))
 	FPCGAttributePropertyInputSelector GrammarAttribute;
-};
+};
\ No newline at end of file
