[PCG] Tagged released load as preview managed actor resources with bIgnoreInPIE = true to avoid them being present in PIE (replaced by non pie duplicate transient preview actors) (following CL 36811270)

#rnx

diff --git a/Engine/Plugins/PCG/Source/PCG/Private/PCGManagedResource.cpp b/Engine/Plugins/PCG/Source/PCG/Private/PCGManagedResource.cpp
--- a/Engine/Plugins/PCG/Source/PCG/Private/PCGManagedResource.cpp
+++ b/Engine/Plugins/PCG/Source/PCG/Private/PCGManagedResource.cpp
@@ -109,6 +109,7 @@ bool UPCGManagedActors::Release(bool bHardRelease, TSet<TSoftObjectPtr<AActor>>&
 				Actor->SetIsTemporarilyHiddenInEditor(true);
 				Actor->SetHidden(true);
 				Actor->SetActorEnableCollision(false);
+				Actor->bIgnoreInPIE = true;
 			}
 		}
 	}
