[PCG] Fix crash when GPU static mesh spawner references a deleted mesh

#rnx

diff --git a/Engine/Plugins/PCG/Source/PCG/Private/Compute/Elements/PCGComputeGraphElement.cpp b/Engine/Plugins/PCG/Source/PCG/Private/Compute/Elements/PCGComputeGraphElement.cpp
--- a/Engine/Plugins/PCG/Source/PCG/Private/Compute/Elements/PCGComputeGraphElement.cpp
+++ b/Engine/Plugins/PCG/Source/PCG/Private/Compute/Elements/PCGComputeGraphElement.cpp
@@ -442,20 +442,25 @@ bool FPCGComputeGraphElement::SetupProceduralISMComponents(FPCGContext* InContex
 
 			for (const FPCGMeshSelectorWeightedEntry& Entry : SelectorWeighted->MeshEntries)
 			{
-				UStaticMesh* StaticMesh = Entry.Descriptor.StaticMesh.LoadSynchronous();
-
-				const float Weight = float(Entry.Weight) / TotalWeight;
-				CumulativeWeight += Weight;
-				PrimitiveSelectionCDF.Add(CumulativeWeight);
-				PrimitiveStringKeys.Add(InBinding->GetStringTable().IndexOfByKey(Entry.Descriptor.StaticMesh.ToString()));
-
-				FPCGProceduralISMComponentDescriptor Descriptor;
-				Descriptor = Entry.Descriptor;
-				Descriptor.NumInstances = FMath::CeilToInt(InputPointCount * Weight);
-				Descriptor.LocalBounds = LocalBounds;
-				Descriptor.NumCustomFloats = CustomFloatCount;
-				Descriptor.StaticMesh = StaticMesh;
-				ComponentsToCreate.Add(MoveTemp(Descriptor));
+				if (UStaticMesh* StaticMesh = Entry.Descriptor.StaticMesh.LoadSynchronous())
+				{
+					const float Weight = float(Entry.Weight) / TotalWeight;
+					CumulativeWeight += Weight;
+					PrimitiveSelectionCDF.Add(CumulativeWeight);
+					PrimitiveStringKeys.Add(InBinding->GetStringTable().IndexOfByKey(Entry.Descriptor.StaticMesh.ToString()));
+
+					FPCGProceduralISMComponentDescriptor Descriptor;
+					Descriptor = Entry.Descriptor;
+					Descriptor.NumInstances = FMath::CeilToInt(InputPointCount * Weight);
+					Descriptor.LocalBounds = LocalBounds;
+					Descriptor.NumCustomFloats = CustomFloatCount;
+					Descriptor.StaticMesh = StaticMesh;
+					ComponentsToCreate.Add(MoveTemp(Descriptor));
+				}
+				else
+				{
+					UE_LOG(LogPCG, Error, TEXT("Could not load static mesh from path '%s'."), *Entry.Descriptor.StaticMesh.ToString())
+				}
 			}
 		}
 
