[PCG] Make sure we always cook tasks with grid size == unitialized, so that we support components that have partitioned=false.

#rnx

diff --git a/Engine/Plugins/PCG/Source/PCG/Private/PCGGraph.cpp b/Engine/Plugins/PCG/Source/PCG/Private/PCGGraph.cpp
--- a/Engine/Plugins/PCG/Source/PCG/Private/PCGGraph.cpp
+++ b/Engine/Plugins/PCG/Source/PCG/Private/PCGGraph.cpp
@@ -538,7 +538,6 @@ void UPCGGraph::PreSave(FObjectPreSaveContext ObjectSaveContext)
 			for (uint32 GridSize : GridSizes)
 			{
 				FPCGStackContext StackContext;
-
 				GraphCompiler.GetCompiledTasks(this, GridSize, StackContext);
 			}
 
@@ -548,7 +547,8 @@ void UPCGGraph::PreSave(FObjectPreSaveContext ObjectSaveContext)
 				GraphCompiler.GetCompiledTasks(this, PCGHiGenGrid::UnboundedGridSize(), StackContext);
 			}
 		}
-		else
+
+		// Always cook unitialized grid tasks which are used if component is not partitioned.
 		{
 			FPCGStackContext StackContext;
 			GraphCompiler.GetCompiledTasks(this, PCGHiGenGrid::UninitializedGridSize(), StackContext);
