[PCG] Level to Asset : Make sure to load all Level Instance actors so they can get exported

diff --git a/Engine/Plugins/PCG/Source/PCGEditor/Private/PCGLevelToAsset.cpp b/Engine/Plugins/PCG/Source/PCGEditor/Private/PCGLevelToAsset.cpp
--- a/Engine/Plugins/PCG/Source/PCGEditor/Private/PCGLevelToAsset.cpp
+++ b/Engine/Plugins/PCG/Source/PCGEditor/Private/PCGLevelToAsset.cpp
@@ -17,9 +17,11 @@
 #include "IContentBrowserSingleton.h"
 #include "AssetRegistry/AssetRegistryModule.h"
 #include "Components/InstancedStaticMeshComponent.h"
+#include "Engine/LevelStreaming.h"
 #include "UObject/UObjectGlobals.h"
 #include "UObject/Package.h"
 #include "WorldPartition/WorldPartition.h"
+#include "WorldPartition/WorldPartitionHelpers.h"
 
 void UPCGLevelToAsset::CreateOrUpdatePCGAssets(const TArray<FAssetData>& WorldAssets, const FPCGAssetExporterParameters& InParameters, TSubclassOf<UPCGLevelToAsset> ExporterSubclass)
 {
@@ -211,12 +213,51 @@ bool UPCGLevelToAsset::BP_ExportWorld_Implementation(UWorld* World, const FStrin
 
 
 	// Make sure all actors are loaded
-	TArray<FWorldPartitionReference> ActorReferences;
-	if (UWorldPartition* WorldPartition = World->GetWorldPartition())
+	TMap<UWorldPartition*, TArray<FWorldPartitionReference>> ActorReferencesPerWorldPartition;
+
+	auto LoadAllActors = [&ActorReferencesPerWorldPartition](UWorldPartition* InWorldPartition) -> bool
 	{
-		WorldPartition->LoadAllActors(ActorReferences);
-	}
+		if (InWorldPartition)
+		{
+			TArray<FWorldPartitionReference> ActorReferences;
+			InWorldPartition->LoadAllActors(ActorReferences);
+			ActorReferencesPerWorldPartition.Emplace(InWorldPartition, MoveTemp(ActorReferences));
+
+			return true;
+		}
+
+		return false;
+	};
+
+	// Load all actors
+	bool bProcessedNewActors = LoadAllActors(World->GetWorldPartition());
+		
+	// Load all level instances and their actors recursively
+	TSet<ULevelStreaming*> ProcessedLevelStreamings;
+	
+	while(bProcessedNewActors)
+	{
+		bProcessedNewActors = false;
 
+		// Make sure to load all Level Instances
+		World->BlockTillLevelStreamingCompleted();
+
+		// For each Streaming Level, make sure to load its actors if it is a World Partition
+		for (ULevelStreaming* LevelStreaming : World->GetStreamingLevels())
+		{
+			bool bAlreadySet = false;
+			if (!ProcessedLevelStreamings.Contains(LevelStreaming))
+			{
+				ProcessedLevelStreamings.Add(LevelStreaming);
+				
+				if (ULevel* LoadedLevel = LevelStreaming->GetLoadedLevel())
+				{
+					bProcessedNewActors |= LoadAllActors(FWorldPartitionHelpers::GetWorldPartition(LoadedLevel));
+				}
+			}
+		}
+	}
+	
 	// Build actor-index map
 	TMap<AActor*, int> ActorIndexMap;
 	int LastActorIndex = 1; // Since the root is the "first" point we'll have, we'll have the map start from 1.
