[PCG] Fix GPU support for multi-data unpacking

#rnx

diff --git a/Engine/Plugins/PCG/Source/PCG/Private/Compute/PCGDataForGPU.cpp b/Engine/Plugins/PCG/Source/PCG/Private/Compute/PCGDataForGPU.cpp
--- a/Engine/Plugins/PCG/Source/PCG/Private/Compute/PCGDataForGPU.cpp
+++ b/Engine/Plugins/PCG/Source/PCG/Private/Compute/PCGDataForGPU.cpp
@@ -967,14 +967,18 @@ void FPCGDataCollectionDesc::PackDataCollection(const FPCGDataCollection& InData
 
 EPCGUnpackDataCollectionResult FPCGDataCollectionDesc::UnpackDataCollection(const TArray<uint8>& InPackedData, FName InPin, const TArray<FString>& InStringTable, FPCGDataCollection& OutDataCollection) const
 {
+	if (InPackedData.IsEmpty())
+	{
+		ensureMsgf(false, TEXT("Tried to unpack a GPU data collection, but the readback buffer was empty."));
+		return EPCGUnpackDataCollectionResult::NoData;
+	}
+
 	const void* PackedData = InPackedData.GetData();
 	const float* DataAsFloat = static_cast<const float*>(PackedData);
 	const uint32* DataAsUint = static_cast<const uint32*>(PackedData);
 	const int32* DataAsInt = static_cast<const int32*>(PackedData);
 
-	uint32 ReadAddress = 0;
-
-	const uint32 PackedExecutionFlagAndNumData = DataAsUint[ReadAddress++];
+	const uint32 PackedExecutionFlagAndNumData = DataAsUint[0];
 
 	// Most significant bit of NumData is reserved to flag whether or not the kernel executed.
 	ensureMsgf(PackedExecutionFlagAndNumData & PCGComputeConstants::KernelExecutedFlag, TEXT("Tried to unpack a GPU data collection, but the compute shader did not execute."));
@@ -991,9 +995,11 @@ EPCGUnpackDataCollectionResult FPCGDataCollectionDesc::UnpackDataCollection(cons
 	{
 		TRACE_CPUPROFILER_EVENT_SCOPE(FPCGDataCollectionDesc::UnpackDataItem);
 
-		const uint32 TypeId = DataAsUint[ReadAddress++];
-		const uint32 NumAttributes = DataAsUint[ReadAddress++];
-		const uint32 NumElements = DataAsUint[ReadAddress++];
+		const uint32 DataHeaderAddress = (DATA_COLLECTION_HEADER_SIZE_BYTES + DATA_HEADER_SIZE_BYTES * DataIndex) / sizeof(uint32);
+
+		const uint32 TypeId =        DataAsUint[DataHeaderAddress + 0];
+		const uint32 NumAttributes = DataAsUint[DataHeaderAddress + 1];
+		const uint32 NumElements =   DataAsUint[DataHeaderAddress + 2];
 
 		const TArray<FPCGKernelAttributeDesc>& AttributeDescs = DataDescs[DataIndex].AttributeDescs;
 		check(NumAttributes == AttributeDescs.Num());
diff --git a/Engine/Plugins/PCG/Source/PCG/Public/Compute/PCGDataForGPU.h b/Engine/Plugins/PCG/Source/PCG/Public/Compute/PCGDataForGPU.h
--- a/Engine/Plugins/PCG/Source/PCG/Public/Compute/PCGDataForGPU.h
+++ b/Engine/Plugins/PCG/Source/PCG/Public/Compute/PCGDataForGPU.h
@@ -16,7 +16,8 @@ enum class EPCGMetadataTypes : uint8;
 enum class EPCGUnpackDataCollectionResult
 {
 	Success,
-	DataMismatch
+	DataMismatch,
+	NoData
 };
 
 UENUM()
