[PCG] Fix upstream Graph parameters in recursive/dynamic subgraphs

* Dynamic and recursive subgraphs nodes were not hooked to pre task, making multiple depth call loosing the UserParameterData needed for the upstream functionality to work
* Needed to make sure we don't hook to multiple pre-tasks, so added a bool in the graph task.

#rnx

diff --git a/Engine/Plugins/PCG/Source/PCG/Private/Graph/PCGGraphCompiler.cpp b/Engine/Plugins/PCG/Source/PCG/Private/Graph/PCGGraphCompiler.cpp
--- a/Engine/Plugins/PCG/Source/PCG/Private/Graph/PCGGraphCompiler.cpp
+++ b/Engine/Plugins/PCG/Source/PCG/Private/Graph/PCGGraphCompiler.cpp
@@ -279,6 +279,7 @@ TArray<FPCGGraphTask> FPCGGraphCompiler::CompileGraph(UPCGGraph* InGraph, FPCGTa
 				if (bRequiresDataFromPreTask || Subtask.Inputs.IsEmpty())
 				{
 					Subtask.Inputs.Emplace(PreId, /*InUpstreamPin=*/FPCGGraphTaskInput::NoPin, /*InDownstreamPin=*/FPCGGraphTaskInput::NoPin, bRequiresDataFromPreTask);
+					Subtask.bWasHookedToPreTask = true;
 				}
 			}
 
@@ -1352,9 +1353,14 @@ void FPCGGraphCompiler::CompileTopGraph(UPCGGraph* InGraph, uint32 GenerationGri
 	for (int TaskIndex = 0; TaskIndex < TaskNum; ++TaskIndex)
 	{
 		FPCGGraphTask& Task = CompiledTasks[TaskIndex];
-		if (Task.Inputs.IsEmpty())
+		const UPCGSettings* Settings = Task.Node ? Task.Node->GetSettings() : nullptr;
+		const bool bRequiresDataFromPreTask = Settings && Settings->RequiresDataFromPreTask();
+
+		// Tasks could have already been hooked to the pre task of the subgraph, so don't re-hook it if that was the case.
+		if (!Task.bWasHookedToPreTask && (Task.Inputs.IsEmpty() || bRequiresDataFromPreTask))
 		{
-			Task.Inputs.Emplace(PreExecuteTaskId);
+			Task.Inputs.Emplace(PreExecuteTaskId, /*InUpstreamPin=*/FPCGGraphTaskInput::NoPin, /*InDownstreamPin=*/FPCGGraphTaskInput::NoPin, bRequiresDataFromPreTask);
+			Task.bWasHookedToPreTask = true;
 		}
 
 		Task.CompiledTaskId = Task.NodeId;
diff --git a/Engine/Plugins/PCG/Source/PCG/Private/Graph/PCGGraphTask.h b/Engine/Plugins/PCG/Source/PCG/Private/Graph/PCGGraphTask.h
--- a/Engine/Plugins/PCG/Source/PCG/Private/Graph/PCGGraphTask.h
+++ b/Engine/Plugins/PCG/Source/PCG/Private/Graph/PCGGraphTask.h
@@ -147,6 +147,9 @@ struct FPCGGraphTask
 
 	// Whether PrepareForExecute as been called on this task
 	bool bHasDonePrepareForExecute = false;
+	
+	// Keep track if this task input has already been hooked to the pre task, like to a pre-task of a subgraph, to avoid hooking it twice, like to the graph pre-task.
+	bool bWasHookedToPreTask = false;
 
 #if WITH_EDITOR
 	// Can be true when we want to have debug display on a task but have taken the results from the cache
