[PCG] Fix accessor creation that did illegal stuff

* We were setting the last selector on const data...
* Was breaking [REDACTED_HANDLE]

#rnx

diff --git a/Engine/Plugins/PCG/Source/PCG/Private/Metadata/Accessors/PCGAttributeAccessorHelpers.cpp b/Engine/Plugins/PCG/Source/PCG/Private/Metadata/Accessors/PCGAttributeAccessorHelpers.cpp
--- a/Engine/Plugins/PCG/Source/PCG/Private/Metadata/Accessors/PCGAttributeAccessorHelpers.cpp
+++ b/Engine/Plugins/PCG/Source/PCG/Private/Metadata/Accessors/PCGAttributeAccessorHelpers.cpp
@@ -21,33 +21,29 @@
 
 namespace PCGAttributeAccessorHelpers
 {
-	/** Retrieves metadata object and the corresponding attribute if any for the provided data. */
-	void ExtractMetadataAttribute(UPCGData* InData, FName Name, UPCGMetadata*& OutMetadata, FPCGMetadataAttributeBase*& OutAttribute)
+	TUniquePtr<IPCGAttributeAccessor> CreateAttributeAccessorImpl(FPCGMetadataAttributeBase* Attribute, UPCGMetadata* Metadata)
 	{
-		OutMetadata = InData ? InData->MutableMetadata() : nullptr;
-		OutAttribute = OutMetadata ? OutMetadata->GetMutableAttribute(Name) : nullptr;
+		auto CreateTypedAccessor = [Attribute, Metadata]<typename T>(T) -> TUniquePtr<IPCGAttributeAccessor>
+		{
+			return MakeUnique<FPCGAttributeAccessor<T>>(static_cast<FPCGMetadataAttribute<T>*>(Attribute), Metadata);
+		};
+
+		if (Attribute && Metadata)
+		{
+			return PCGMetadataAttribute::CallbackWithRightType(Attribute->GetTypeId(), CreateTypedAccessor);
+		}
+		else
+		{
+			return {};
+		}
 	}
 
-	// This template creates a const or non-const accessor on an attribute depending on the provided templated type. 
-	template <
-		typename AccessorType,
-		typename = typename std::enable_if_t<
-		std::is_same_v<IPCGAttributeAccessor, std::remove_const_t<AccessorType>>
-		>
-	>
-	TUniquePtr<AccessorType> CreateAttributeAccessorImpl(FPCGMetadataAttributeBase* Attribute, UPCGMetadata* Metadata, bool bForceReadOnly = false)
+	// Return a non-const version but it will be read only anyway. Done like that to factorize some non-const/const version of accessor creation.
+	TUniquePtr<IPCGAttributeAccessor> CreateConstAttributeAccessorImpl(const FPCGMetadataAttributeBase* Attribute, const UPCGMetadata* Metadata)
 	{
-		auto CreateTypedAccessor = [Attribute, Metadata, bForceReadOnly](auto Dummy) -> TUniquePtr<IPCGAttributeAccessor>
+		auto CreateTypedAccessor = [Attribute, Metadata]<typename T>(T) -> TUniquePtr<IPCGAttributeAccessor>
 		{
-			using AttributeType = decltype(Dummy);
-			if constexpr (std::is_const_v<AccessorType>)
-			{
-				return MakeUnique<FPCGAttributeAccessor<AttributeType>>(static_cast<const FPCGMetadataAttribute<AttributeType>*>(Attribute), Metadata, bForceReadOnly);
-			}
-			else
-			{
-				return MakeUnique<FPCGAttributeAccessor<AttributeType>>(static_cast<FPCGMetadataAttribute<AttributeType>*>(Attribute), Metadata, bForceReadOnly);
-			}
+			return MakeUnique<FPCGAttributeAccessor<T>>(static_cast<const FPCGMetadataAttribute<T>*>(Attribute), Metadata, /*bForceReadOnly=*/true);
 		};
 
 		if (Attribute && Metadata)
@@ -56,7 +52,7 @@ namespace PCGAttributeAccessorHelpers
 		}
 		else
 		{
-			return TUniquePtr<AccessorType>();
+			return {};
 		}
 	}
 
@@ -124,14 +120,22 @@ namespace PCGAttributeAccessorHelpers
 
 		if (InSelector.GetSelection() == EPCGAttributePropertySelection::Attribute)
 		{
-			UPCGMetadata* Metadata = nullptr;
-			FPCGMetadataAttributeBase* Attribute = nullptr;
-
-			// It is OK to const_cast here, since we will create a const accessor if the input is const.
-			ExtractMetadataAttribute(const_cast<UPCGData*>(InData), Name, Metadata, Attribute);
-
-			// To simplify the code here & below we'll get a non-const accessor but force readonly if it should be
-			Accessor = CreateAttributeAccessorImpl<std::remove_const_t<AccessorType>>(Attribute, Metadata, /*bForceReadOnly=*/std::is_const_v<AccessorType>);
+			// We can't use const_cast here to only have a single path, as GetMutableAttribute will set the last selector
+			// and we don't want that on const data.
+			if constexpr (std::is_const_v<AccessorType>)
+			{
+				const UPCGMetadata* ConstMetadata = InData ? InData->ConstMetadata() : nullptr;
+				const FPCGMetadataAttributeBase* ConstAttribute = ConstMetadata ? ConstMetadata->GetConstAttribute(Name) : nullptr;
+				
+				Accessor = CreateConstAttributeAccessorImpl(ConstAttribute, ConstMetadata);
+			}
+			else
+			{
+				UPCGMetadata* MutableMetadata = InData ? InData->MutableMetadata() : nullptr;
+				FPCGMetadataAttributeBase* MutableAttribute = MutableMetadata ? MutableMetadata->GetMutableAttribute(Name) : nullptr;
+				
+				Accessor = CreateAttributeAccessorImpl(MutableAttribute, MutableMetadata);
+			}
 		}
 
 		if (!Accessor.IsValid())
@@ -619,7 +623,7 @@ TUniquePtr<const IPCGAttributeAccessor> PCGAttributeAccessorHelpers::CreateConst
 
 TUniquePtr<const IPCGAttributeAccessor> PCGAttributeAccessorHelpers::CreateConstAccessor(const FPCGMetadataAttributeBase* InAttribute, const UPCGMetadata* InMetadata, bool bQuiet)
 {
-	return CreateAttributeAccessorImpl<const IPCGAttributeAccessor>(const_cast<FPCGMetadataAttributeBase*>(InAttribute), const_cast<UPCGMetadata*>(InMetadata), bQuiet);
+	return CreateConstAttributeAccessorImpl(const_cast<FPCGMetadataAttributeBase*>(InAttribute), const_cast<UPCGMetadata*>(InMetadata));
 }
 
 TUniquePtr<IPCGAttributeAccessor> PCGAttributeAccessorHelpers::CreateAccessor(UPCGData* InData, const FPCGAttributePropertySelector& InSelector, bool bQuiet)
@@ -629,7 +633,7 @@ TUniquePtr<IPCGAttributeAccessor> PCGAttributeAccessorHelpers::CreateAccessor(UP
 
 TUniquePtr<IPCGAttributeAccessor> PCGAttributeAccessorHelpers::CreateAccessor(FPCGMetadataAttributeBase* InAttribute, UPCGMetadata* InMetadata, bool bQuiet)
 {
-	return CreateAttributeAccessorImpl<IPCGAttributeAccessor>(InAttribute, InMetadata, bQuiet);
+	return CreateAttributeAccessorImpl(InAttribute, InMetadata);
 }
 
 TUniquePtr<IPCGAttributeAccessor> PCGAttributeAccessorHelpers::CreateAccessorWithAttributeCreation(UPCGData* InData, const FPCGAttributePropertySelector& InSelector, const IPCGAttributeAccessor* InMatchingAccessor, EPCGAttributeAccessorFlags InTypeMatching, bool bQuiet)
