[PCG] Address GPU Scene 'flushes' that lose instance data in editor.

Catches known flush cases:

* After level load in editor
* After material used by the PISMC instances is modified

Temporary workaround until we get to planned work to make procedural instance data explicitly persistent.

#rnx

diff --git a/Engine/Plugins/PCG/Source/PCG/Private/Components/PCGProceduralISMComponent.cpp b/Engine/Plugins/PCG/Source/PCG/Private/Components/PCGProceduralISMComponent.cpp
--- a/Engine/Plugins/PCG/Source/PCG/Private/Components/PCGProceduralISMComponent.cpp
+++ b/Engine/Plugins/PCG/Source/PCG/Private/Components/PCGProceduralISMComponent.cpp
@@ -5,6 +5,7 @@
 #include "PCGComponent.h"
 #include "PCGModule.h"
 #include "PCGManagedResource.h"
+#include "PCGSubsystem.h"
 #include "Helpers/PCGHelpers.h"
 #include "MeshSelectors/PCGISMDescriptor.h"
 
@@ -220,6 +221,23 @@ UPCGProceduralISMComponent::UPCGProceduralISMComponent(const FObjectInitializer&
 #endif
 }
 
+#if WITH_EDITOR
+void UPCGProceduralISMComponent::OnRegister()
+{
+	Super::OnRegister();
+
+	USceneComponent::MarkRenderStateDirtyEvent.RemoveAll(this);
+	USceneComponent::MarkRenderStateDirtyEvent.AddUObject(this, &UPCGProceduralISMComponent::OnRenderStateDirty);
+}
+
+void UPCGProceduralISMComponent::OnUnregister()
+{
+	USceneComponent::MarkRenderStateDirtyEvent.RemoveAll(this);
+
+	Super::OnUnregister();
+}
+#endif
+
 FPrimitiveSceneProxy* UPCGProceduralISMComponent::CreateStaticMeshSceneProxy(Nanite::FMaterialAudit& NaniteMaterials, bool bCreateNanite)
 {
 	LLM_SCOPE(ELLMTag::InstancedMesh);
@@ -380,6 +398,77 @@ void UPCGProceduralISMComponent::ValidateComponentSetup()
 	}
 }
 
+#if WITH_EDITOR
+void UPCGProceduralISMComponent::OnRenderStateDirty(UActorComponent& InComponent)
+{
+	// Currently, there is no explicit persistence of instance data in the GPU scene. When this component is dirtied, the instance data is cleared.
+	// TODO: This function is a stop gap that requests a refresh of the PCG Component managing this component, and should be removed later.
+
+	if (&InComponent != this)
+	{
+		return;
+	}
+
+	UPCGSubsystem* Subsystem = GetOwner() ? UPCGSubsystem::GetInstance(GetOwner()->GetWorld()) : nullptr;
+	if (!Subsystem)
+	{
+		return;
+	}
+
+	// Helper that returns true if the given PCG component is managing this PISMC.
+	auto PCGComponentManagesThisPISMC = [this](UPCGComponent* InComponent) -> bool
+	{
+		bool bManagesThis = false;
+
+		if (InComponent->bGenerated && InComponent->AreProceduralInstancesInUse())
+		{
+			InComponent->ForEachManagedResource([this, InComponent, &bManagesThis](UPCGManagedResource* InResource)
+			{
+				if (!bManagesThis)
+				{
+					const UPCGManagedProceduralISMComponent* PISMC = Cast<UPCGManagedProceduralISMComponent>(InResource);
+					if (PISMC && PISMC->GetComponent() == this)
+					{
+						bManagesThis = true;
+					}
+				}
+			});
+		}
+
+		return bManagesThis;
+	};
+
+	Subsystem->RefreshAllComponentsFiltered(
+		[this, Subsystem, &PCGComponentManagesThisPISMC](UPCGComponent* InComponent) -> bool
+		{
+			// If the original component manages this PISMC, request a refresh of it and we're done. If it has local components they will also be refreshed.
+			if (PCGComponentManagesThisPISMC(InComponent))
+			{
+				return true;
+			}
+
+			// A local component of the original component might manage this, so check those.
+			bool bManagesThis = false;
+
+			if (InComponent->bIsComponentPartitioned)
+			{
+				Subsystem->ForAllRegisteredLocalComponents(
+					InComponent,
+					[this, &bManagesThis, &PCGComponentManagesThisPISMC](UPCGComponent* InComponent)
+					{
+						if (!bManagesThis)
+						{
+							bManagesThis = PCGComponentManagesThisPISMC(InComponent);
+						}
+					});
+			}
+
+			return bManagesThis;
+		},
+		EPCGChangeType::Structural);
+}
+#endif // WITH_EDITOR
+
 void UPCGProceduralISMComponent::SetNumInstances(int32 InNumInstances)
 {
 	if (NumInstances != InNumInstances)
diff --git a/Engine/Plugins/PCG/Source/PCG/Private/Components/PCGProceduralISMComponent.h b/Engine/Plugins/PCG/Source/PCG/Private/Components/PCGProceduralISMComponent.h
--- a/Engine/Plugins/PCG/Source/PCG/Private/Components/PCGProceduralISMComponent.h
+++ b/Engine/Plugins/PCG/Source/PCG/Private/Components/PCGProceduralISMComponent.h
@@ -28,6 +28,13 @@ class UPCGProceduralISMComponent : public UStaticMeshComponent
 	GENERATED_UCLASS_BODY()
 
 public:
+#if WITH_EDITOR
+	//~Begin UActorComponent Interface
+	virtual void OnRegister() override;
+	virtual void OnUnregister() override;
+	//~End UActorComponent Interface
+#endif
+
 	UFUNCTION(BlueprintCallable, Category = "Components|PCGProceduralISMC")
 	int32 GetNumInstances() const { return NumInstances; }
 
@@ -100,6 +107,12 @@ protected:
 	/** Validate component setup, disable features that are not supported (typically those that require instance transforms on CPU). */
 	virtual void ValidateComponentSetup();
 
+#if WITH_EDITOR
+private:
+	/** Executed when the render state of any actor component is dirtied. */
+	void OnRenderStateDirty(UActorComponent& InComponent);
+#endif
+
 protected:
 	UPROPERTY(EditAnywhere, Category = Instances)
 	int32 NumInstances = 0;
diff --git a/Engine/Plugins/PCG/Source/PCG/Private/PCGSubsystem.cpp b/Engine/Plugins/PCG/Source/PCG/Private/PCGSubsystem.cpp
--- a/Engine/Plugins/PCG/Source/PCG/Private/PCGSubsystem.cpp
+++ b/Engine/Plugins/PCG/Source/PCG/Private/PCGSubsystem.cpp
@@ -932,9 +932,25 @@ void UPCGSubsystem::RefreshAllRuntimeGenComponents(EPCGChangeType ChangeType)
 {
 	for (UPCGComponent* Component : GetAllRegisteredComponents())
 	{
-		RefreshRuntimeGenComponent(Component, ChangeType);
+		if (Component->GenerationTrigger == EPCGComponentGenerationTrigger::GenerateAtRuntime)
+		{
+			RefreshRuntimeGenComponent(Component, ChangeType);
+		}
+	}
+}
+
+#if WITH_EDITOR
+void UPCGSubsystem::RefreshAllComponentsFiltered(const TFunction<bool(UPCGComponent*)>& ComponentFilter, EPCGChangeType ChangeType)
+{
+	for (UPCGComponent* Component : GetAllRegisteredComponents())
+	{
+		if (ComponentFilter(Component))
+		{
+			Component->Refresh(ChangeType);
+		}
 	}
 }
+#endif
 
 bool UPCGSubsystem::IsGraphCurrentlyExecuting(UPCGGraph* Graph)
 {
diff --git a/Engine/Plugins/PCG/Source/PCG/Public/PCGComponent.h b/Engine/Plugins/PCG/Source/PCG/Public/PCGComponent.h
--- a/Engine/Plugins/PCG/Source/PCG/Public/PCGComponent.h
+++ b/Engine/Plugins/PCG/Source/PCG/Public/PCGComponent.h
@@ -239,6 +239,9 @@ public:
 	/** Called during execution if one or more procedural ISM components are in use. */
 	void NotifyProceduralInstancesInUse() { bProceduralInstancesInUse = true; }
 
+	/** Whether this component created one or more procedural ISM components when last generated. */
+	bool AreProceduralInstancesInUse() const { return bProceduralInstancesInUse; }
+
 	UPROPERTY(BlueprintReadWrite, EditAnywhere, Category = Settings, meta = (DisplayPriority = 600))
 	int Seed = 42;
 
diff --git a/Engine/Plugins/PCG/Source/PCG/Public/PCGSubsystem.h b/Engine/Plugins/PCG/Source/PCG/Public/PCGSubsystem.h
--- a/Engine/Plugins/PCG/Source/PCG/Public/PCGSubsystem.h
+++ b/Engine/Plugins/PCG/Source/PCG/Public/PCGSubsystem.h
@@ -184,6 +184,11 @@ public:
 	/** Asks the runtime generation scheduler to refresh all GenerateAtRuntime components. ChangeType should be 'GenerationGrid' to perform a full cleanup of PAs and local components. */
 	void RefreshAllRuntimeGenComponents(EPCGChangeType ChangeType = EPCGChangeType::None);
 
+#if WITH_EDITOR
+	/** Refresh all components selected by the filter (runtime generated or otherwise). */
+	void RefreshAllComponentsFiltered(const TFunction<bool(UPCGComponent*)>& ComponentFilter, EPCGChangeType ChangeType = EPCGChangeType::None);
+#endif
+
 	FPCGRuntimeGenScheduler* GetRuntimeGenScheduler() const { return RuntimeGenScheduler; }
 
 	/** Returns true if there are any tasks for this graph currently scheduled or executing. */
