[PCG] Add missing FGCScopeGuard around a SoftObjectPath resolve (see OverrideSettings)

#rnx

diff --git a/Engine/Plugins/PCG/Source/PCG/Private/PCGContext.cpp b/Engine/Plugins/PCG/Source/PCG/Private/PCGContext.cpp
--- a/Engine/Plugins/PCG/Source/PCG/Private/PCGContext.cpp
+++ b/Engine/Plugins/PCG/Source/PCG/Private/PCGContext.cpp
@@ -166,11 +166,15 @@ void FPCGContext::InitializeSettings(bool bSkipPostLoad)
 						}
 
 						FSoftObjectPath ObjectPath;
-						if (PCGContextHelpers::GetOverrideParamValue(*AttributeAccessor, ObjectPath) && !ObjectPath.ResolveObject())
+						if (PCGContextHelpers::GetOverrideParamValue(*AttributeAccessor, ObjectPath))
 						{
-							// We have an override value that is not loaded, and we are not on the main thread. We need to schedule the task on the main thread.
-							bOverrideSettingsOnMainThread = true;
-							break;
+							FGCScopeGuard GCScope;
+							if (!ObjectPath.ResolveObject())
+							{
+								// We have an override value that is not loaded, and we are not on the main thread. We need to schedule the task on the main thread.
+								bOverrideSettingsOnMainThread = true;
+								break;
+							}
 						}
 					}
 				}
