[PCG] Properly record transactions when modifying Custom HLSL source/functions from shader source editor
* Now calls Modify() on the graph, which fixes issues where source changes are not marked to be saved and are thus are lost
* Now allows undo/redo of source changes. Changes are applied anytime the text is committed or the text box loses focus.

#rnx

diff --git a/Engine/Plugins/PCG/Source/PCG/Private/Compute/Elements/PCGCustomHLSL.cpp b/Engine/Plugins/PCG/Source/PCG/Private/Compute/Elements/PCGCustomHLSL.cpp
--- a/Engine/Plugins/PCG/Source/PCG/Private/Compute/Elements/PCGCustomHLSL.cpp
+++ b/Engine/Plugins/PCG/Source/PCG/Private/Compute/Elements/PCGCustomHLSL.cpp
@@ -17,6 +17,10 @@
 
 #include "Internationalization/Regex.h"
 
+#if WITH_EDITOR
+#include "ScopedTransaction.h"
+#endif
+
 #include UE_INLINE_GENERATED_CPP_BY_NAME(PCGCustomHLSL)
 
 #define LOCTEXT_NAMESPACE "PCGCustomHLSLElement"
@@ -594,27 +598,45 @@ FString UPCGCustomHLSLSettings::GetShaderText() const
 
 void UPCGCustomHLSLSettings::SetShaderFunctionsText(const FString& NewFunctionsText)
 {
-	ShaderFunctions = NewFunctionsText;
+	const FName PropertyName = GET_MEMBER_NAME_CHECKED(UPCGCustomHLSLSettings, ShaderFunctions);
+	FProperty* Property = FindFProperty<FProperty>(StaticClass(), PropertyName);
+	FPropertyChangedEvent PropertyChangedEvent(Property, EPropertyChangeType::ValueSet);
 
-	UpdateAttributeKeys();
+	{
+		FScopedTransaction Transaction(LOCTEXT("OnSetShaderFunctionsText", "Set Shader Functions Text"));
+	
+		PreEditChange(Property);
+		Modify();
+		ShaderFunctions = NewFunctionsText;
+		PostEditChangeProperty(PropertyChangedEvent);
+	}
+	
+
+	OnSettingsChangedDelegate.Broadcast(this, GetChangeTypeForProperty(PropertyName));
 }
 
 void UPCGCustomHLSLSettings::SetShaderText(const FString& NewText)
 {
-	ShaderSource = NewText;
+	const FName PropertyName = GET_MEMBER_NAME_CHECKED(UPCGCustomHLSLSettings, ShaderSource);
+	FProperty* Property = FindFProperty<FProperty>(StaticClass(), PropertyName);
+	FPropertyChangedEvent PropertyChangedEvent(Property, EPropertyChangeType::ValueSet);
+	
+	{
+		FScopedTransaction Transaction(LOCTEXT("OnSetShaderSourceText", "Set Shader Source Text"));
 
-	UpdateAttributeKeys();
+		PreEditChange(Property);
+		Modify();
+		ShaderSource = NewText;
+		PostEditChangeProperty(PropertyChangedEvent);
+	}
+
+	OnSettingsChangedDelegate.Broadcast(this, GetChangeTypeForProperty(PropertyName));
 }
 
 bool UPCGCustomHLSLSettings::IsShaderTextReadOnly() const
 {
 	return false;
 }
-
-void UPCGCustomHLSLSettings::ApplySourceChanges()
-{
-	OnSettingsChangedDelegate.Broadcast(this, GetChangeTypeForProperty(GET_MEMBER_NAME_CHECKED(UPCGCustomHLSLSettings, ShaderSource)));
-}
 #endif //WITH_EDITOR
 
 const FPCGPinPropertiesGPU* UPCGCustomHLSLSettings::GetOutputPinPropertiesGPU(const FName& InPinLabel) const
diff --git a/Engine/Plugins/PCG/Source/PCG/Public/Compute/Elements/PCGCustomHLSL.h b/Engine/Plugins/PCG/Source/PCG/Public/Compute/Elements/PCGCustomHLSL.h
--- a/Engine/Plugins/PCG/Source/PCG/Public/Compute/Elements/PCGCustomHLSL.h
+++ b/Engine/Plugins/PCG/Source/PCG/Public/Compute/Elements/PCGCustomHLSL.h
@@ -87,7 +87,6 @@ protected:
 	void SetShaderFunctionsText(const FString& NewFunctionsText) override;
 	void SetShaderText(const FString& NewText) override;
 	bool IsShaderTextReadOnly() const override;
-	void ApplySourceChanges() override;
 	//~End IPCGNodeSourceTextProvider interface
 #endif
 
diff --git a/Engine/Plugins/PCG/Source/PCG/Public/Compute/IPCGNodeSourceTextProvider.h b/Engine/Plugins/PCG/Source/PCG/Public/Compute/IPCGNodeSourceTextProvider.h
--- a/Engine/Plugins/PCG/Source/PCG/Public/Compute/IPCGNodeSourceTextProvider.h
+++ b/Engine/Plugins/PCG/Source/PCG/Public/Compute/IPCGNodeSourceTextProvider.h
@@ -37,8 +37,5 @@ public:
 
 	/** Return true is the shader text can be edited */
 	virtual bool IsShaderTextReadOnly() const = 0;
-
-	/** Apply source changes (compile and execute). */
-	virtual void ApplySourceChanges() = 0;
 #endif
 };
diff --git a/Engine/Plugins/PCG/Source/PCGEditor/Private/Widgets/SPCGEditorNodeSource.cpp b/Engine/Plugins/PCG/Source/PCGEditor/Private/Widgets/SPCGEditorNodeSource.cpp
--- a/Engine/Plugins/PCG/Source/PCGEditor/Private/Widgets/SPCGEditorNodeSource.cpp
+++ b/Engine/Plugins/PCG/Source/PCGEditor/Private/Widgets/SPCGEditorNodeSource.cpp
@@ -2,6 +2,7 @@
 
 #include "SPCGEditorNodeSource.h"
 
+#include "PCGSettings.h"
 #include "Compute/IPCGNodeSourceTextProvider.h"
 
 #include "PCGEditor.h"
@@ -61,7 +62,8 @@ void SPCGEditorNodeSource::Construct(const FArguments& InArgs, TSharedPtr<FPCGEd
 				.IsReadOnly(this, &SPCGEditorNodeSource::IsShaderTextReadOnly)
 				.Marshaller(SyntaxHighlighterDeclarations)
 				.OnTextChanged(this, &SPCGEditorNodeSource::OnShaderFunctionsTextChanged)
-				.OnTextChangesApplied(this, &SPCGEditorNodeSource::OnShaderTextChangesApplied)
+				.OnTextCommitted(this, &SPCGEditorNodeSource::OnShaderFunctionsTextCommitted)
+				.OnTextChangesApplied(this, &SPCGEditorNodeSource::OnShaderFunctionsTextChangesApplied)
 			]
 		]
 		+SSplitter::Slot()
@@ -80,8 +82,9 @@ void SPCGEditorNodeSource::Construct(const FArguments& InArgs, TSharedPtr<FPCGEd
 				.Text(this, &SPCGEditorNodeSource::GetShaderTextAsText)
 				.IsReadOnly(this, &SPCGEditorNodeSource::IsShaderTextReadOnly)
 				.Marshaller(SyntaxHighlighterShaderText)
-				.OnTextChanged(this, &SPCGEditorNodeSource::OnShaderTextChanged)
-				.OnTextChangesApplied(this, &SPCGEditorNodeSource::OnShaderTextChangesApplied)
+				.OnTextChanged(this, &SPCGEditorNodeSource::OnShaderSourceTextChanged)
+				.OnTextCommitted(this, &SPCGEditorNodeSource::OnShaderSourceTextCommitted)
+				.OnTextChangesApplied(this, &SPCGEditorNodeSource::OnShaderSourceTextChangesApplied)
 			]
 		]
 	];
@@ -91,6 +94,9 @@ void SPCGEditorNodeSource::SetTextProviderObject(UObject* InProviderObject)
 {
 	ShaderTextProviderObject = InProviderObject;
 
+	ShaderFunctionsText = GetShaderFunctionsAsText();
+	ShaderSourceText = GetShaderTextAsText();
+
 	Refresh();
 }
 
@@ -169,30 +175,67 @@ bool SPCGEditorNodeSource::IsShaderTextReadOnly() const
 	return false;
 }
 
-void SPCGEditorNodeSource::OnShaderFunctionsTextChanged(const FText& InText) const
+void SPCGEditorNodeSource::OnShaderFunctionsTextChanged(const FText& InText)
 {
-	if (GetTextProviderInterface())
-	{
-		GetTextProviderInterface()->SetShaderFunctionsText(InText.ToString());
-	}
+	// Always clear compiler messages when text is edited so that markup doesn't lurk in random places.
+	SyntaxHighlighterShaderText->ClearCompilerMessages();
+
+	ShaderFunctionsText = InText;
 }
 
-void SPCGEditorNodeSource::OnShaderTextChanged(const FText& InText) const
+void SPCGEditorNodeSource::OnShaderSourceTextChanged(const FText& InText)
 {
 	// Always clear compiler messages when text is edited so that markup doesn't lurk in random places.
 	SyntaxHighlighterShaderText->ClearCompilerMessages();
 
-	if (GetTextProviderInterface())
+	ShaderSourceText = InText;
+}
+
+void SPCGEditorNodeSource::OnShaderFunctionsTextCommitted(const FText& InText, ETextCommit::Type InCommitInfo)
+{
+	ShaderFunctionsText = InText;
+	SetShaderFunctionsText();
+}
+
+void SPCGEditorNodeSource::OnShaderSourceTextCommitted(const FText& InText, ETextCommit::Type InCommitInfo)
+{
+	ShaderSourceText = InText;
+	SetShaderSourceText();
+}
+
+void SPCGEditorNodeSource::OnShaderFunctionsTextChangesApplied() const
+{
+	SetShaderFunctionsText();
+}
+
+void SPCGEditorNodeSource::OnShaderSourceTextChangesApplied() const
+{
+	SetShaderSourceText();
+}
+
+void SPCGEditorNodeSource::SetShaderFunctionsText() const
+{
+	if (IPCGNodeSourceTextProvider* Provider = GetTextProviderInterface())
 	{
-		GetTextProviderInterface()->SetShaderText(InText.ToString());
+		const FString ShaderFunctionsString = ShaderFunctionsText.ToString();
+
+		if (ShaderFunctionsString != Provider->GetShaderFunctionsText())
+		{
+			Provider->SetShaderFunctionsText(ShaderFunctionsString);
+		}
 	}
 }
 
-void SPCGEditorNodeSource::OnShaderTextChangesApplied() const
+void SPCGEditorNodeSource::SetShaderSourceText() const
 {
-	if (GetTextProviderInterface())
+	if (IPCGNodeSourceTextProvider* Provider = GetTextProviderInterface())
 	{
-		GetTextProviderInterface()->ApplySourceChanges();
+		const FString ShaderSourceString = ShaderSourceText.ToString();
+
+		if (ShaderSourceString != Provider->GetShaderText())
+		{
+			Provider->SetShaderText(ShaderSourceString);
+		}
 	}
 }
 
diff --git a/Engine/Plugins/PCG/Source/PCGEditor/Private/Widgets/SPCGEditorNodeSource.h b/Engine/Plugins/PCG/Source/PCGEditor/Private/Widgets/SPCGEditorNodeSource.h
--- a/Engine/Plugins/PCG/Source/PCGEditor/Private/Widgets/SPCGEditorNodeSource.h
+++ b/Engine/Plugins/PCG/Source/PCGEditor/Private/Widgets/SPCGEditorNodeSource.h
@@ -44,9 +44,17 @@ private:
 
 	bool IsShaderTextReadOnly() const;
 
-	void OnShaderFunctionsTextChanged(const FText& InText) const;
-	void OnShaderTextChanged(const FText& InText) const;
-	void OnShaderTextChangesApplied() const;
+	void OnShaderFunctionsTextChanged(const FText& InText);
+	void OnShaderSourceTextChanged(const FText& InText);
+
+	void OnShaderFunctionsTextCommitted(const FText& InText, ETextCommit::Type InCommitInfo);
+	void OnShaderSourceTextCommitted(const FText& InText, ETextCommit::Type InCommitInfo);
+
+	void OnShaderFunctionsTextChangesApplied() const;
+	void OnShaderSourceTextChangesApplied() const;
+
+	void SetShaderFunctionsText() const;
+	void SetShaderSourceText() const;
 
 	TSharedRef<SWidget> ConstructNonExpandableHeaderWidget(const SExpandableArea::FArguments& InArgs) const;
 
@@ -61,4 +69,7 @@ private:
 	TSharedPtr<SPCGNodeSourceTextBox> DeclarationsTextBox;
 	TSharedPtr<SPCGNodeSourceTextBox> ShaderFunctionsTextBox;
 	TSharedPtr<SPCGNodeSourceTextBox> ShaderTextTextBox;
+
+	FText ShaderFunctionsText;
+	FText ShaderSourceText;
 };
diff --git a/Engine/Plugins/PCG/Source/PCGEditor/Private/Widgets/SPCGNodeSourceTextBox.cpp b/Engine/Plugins/PCG/Source/PCGEditor/Private/Widgets/SPCGNodeSourceTextBox.cpp
--- a/Engine/Plugins/PCG/Source/PCGEditor/Private/Widgets/SPCGNodeSourceTextBox.cpp
+++ b/Engine/Plugins/PCG/Source/PCGEditor/Private/Widgets/SPCGNodeSourceTextBox.cpp
@@ -107,6 +107,7 @@ void SPCGNodeSourceTextBox::Construct(const FArguments& InArgs)
 			.TextStyle(&TextStyle)
 			.Text(InArgs._Text)
 			.OnTextChanged(InArgs._OnTextChanged)
+			.OnTextCommitted(InArgs._OnTextCommitted)
 			.OnKeyCharHandler(this, &SPCGNodeSourceTextBox::OnTextKeyChar)
 			.OnKeyDownHandler(this, &SPCGNodeSourceTextBox::OnTextKeyDown)
 			// By default, the Tab key gets routed to "next widget". We want to disable that behaviour.
diff --git a/Engine/Plugins/PCG/Source/PCGEditor/Private/Widgets/SPCGNodeSourceTextBox.h b/Engine/Plugins/PCG/Source/PCGEditor/Private/Widgets/SPCGNodeSourceTextBox.h
--- a/Engine/Plugins/PCG/Source/PCGEditor/Private/Widgets/SPCGNodeSourceTextBox.h
+++ b/Engine/Plugins/PCG/Source/PCGEditor/Private/Widgets/SPCGNodeSourceTextBox.h
@@ -56,6 +56,9 @@ public:
 		/** Called whenever the text is changed programmatically or interactively by the user */
 		SLATE_EVENT(FOnTextChanged, OnTextChanged)
 
+		/** Called whenever the text is committed. This happens when the text box loses focus. */
+		SLATE_EVENT(FOnTextCommitted, OnTextCommitted)
+
 		/** Called when user applies the changes */
 		SLATE_EVENT(FOnTextChangesApplied, OnTextChangesApplied)
 	SLATE_END_ARGS()
