[PCG] Remove unsafe uses of CastChecked from GPU data interfaces and PCGTextureData

#rnx

diff --git a/Engine/Plugins/PCG/Source/PCG/Private/Compute/DataInterfaces/PCGLandscapeDataInterface.cpp b/Engine/Plugins/PCG/Source/PCG/Private/Compute/DataInterfaces/PCGLandscapeDataInterface.cpp
--- a/Engine/Plugins/PCG/Source/PCG/Private/Compute/DataInterfaces/PCGLandscapeDataInterface.cpp
+++ b/Engine/Plugins/PCG/Source/PCG/Private/Compute/DataInterfaces/PCGLandscapeDataInterface.cpp
@@ -292,19 +292,21 @@ UComputeDataProvider* UPCGLandscapeDataInterface::CreateDataProvider(TObjectPtr<
 	if (!TaggedDatas.IsEmpty())
 	{
 		ensure(TaggedDatas.Num() == 1); // There should only be one landscape data
+		ensure(TaggedDatas[0].Data->IsA<UPCGLandscapeData>());
 
-		const UPCGLandscapeData* LandscapeData = CastChecked<UPCGLandscapeData>(TaggedDatas[0].Data);
-
-		for (TSoftObjectPtr<ALandscapeProxy> LandscapeProxyPtr : LandscapeData->Landscapes)
+		if (const UPCGLandscapeData* LandscapeData = Cast<UPCGLandscapeData>(TaggedDatas[0].Data))
 		{
-			if (ALandscapeProxy* LandscapeProxy = LandscapeProxyPtr.Get())
+			for (TSoftObjectPtr<ALandscapeProxy> LandscapeProxyPtr : LandscapeData->Landscapes)
 			{
-				if (ALandscape* Landscape = LandscapeProxy->GetLandscapeActor())
+				if (ALandscapeProxy* LandscapeProxy = LandscapeProxyPtr.Get())
 				{
-					if (TestLandscape(Landscape))
+					if (ALandscape* Landscape = LandscapeProxy->GetLandscapeActor())
 					{
-						DataProvider->Initialize(Landscape, ComponentBounds);
-						break;
+						if (TestLandscape(Landscape))
+						{
+							DataProvider->Initialize(Landscape, ComponentBounds);
+							break;
+						}
 					}
 				}
 			}
diff --git a/Engine/Plugins/PCG/Source/PCG/Private/Compute/DataInterfaces/PCGTextureDataInterface.cpp b/Engine/Plugins/PCG/Source/PCG/Private/Compute/DataInterfaces/PCGTextureDataInterface.cpp
--- a/Engine/Plugins/PCG/Source/PCG/Private/Compute/DataInterfaces/PCGTextureDataInterface.cpp
+++ b/Engine/Plugins/PCG/Source/PCG/Private/Compute/DataInterfaces/PCGTextureDataInterface.cpp
@@ -83,20 +83,23 @@ UComputeDataProvider* UPCGTextureDataInterface::CreateDataProvider(TObjectPtr<UO
 	if (!TaggedDatas.IsEmpty())
 	{
 		ensure(TaggedDatas.Num() == 1); // There should only be one texture data
+		ensure(TaggedDatas[0].Data->IsA<UPCGTextureData>());
 
-		const UPCGTextureData* TextureData = CastChecked<UPCGTextureData>(TaggedDatas[0].Data);
-		const UTexture* Texture = TextureData->Texture.IsValid() ? TextureData->Texture.Get() : nullptr;
-
-		// TODO: Support texture2d array, texture3d
-		// TODO: Bindless resources to handle multiple textures
-
-		if (Texture)
+		if (const UPCGTextureData* TextureData = Cast<UPCGTextureData>(TaggedDatas[0].Data))
 		{
-			const FTextureResource* TextureResource = Texture->GetResource();
+			const UTexture* Texture = TextureData->Texture.IsValid() ? TextureData->Texture.Get() : nullptr;
 
-			if (TextureResource && TextureResource->TextureRHI)
+			// TODO: Support texture2d array, texture3d
+			// TODO: Bindless resources to handle multiple textures
+
+			if (Texture)
 			{
-				DataProvider->Texture = TextureResource->TextureRHI;
+				const FTextureResource* TextureResource = Texture->GetResource();
+
+				if (TextureResource && TextureResource->TextureRHI)
+				{
+					DataProvider->Texture = TextureResource->TextureRHI;
+				}
 			}
 		}
 	}
diff --git a/Engine/Plugins/PCG/Source/PCG/Private/Data/PCGTextureData.cpp b/Engine/Plugins/PCG/Source/PCG/Private/Data/PCGTextureData.cpp
--- a/Engine/Plugins/PCG/Source/PCG/Private/Data/PCGTextureData.cpp
+++ b/Engine/Plugins/PCG/Source/PCG/Private/Data/PCGTextureData.cpp
@@ -432,8 +432,15 @@ void UPCGTextureData::InitializeInternal(UTexture* InTexture, uint32 InTextureIn
 		UTexture2D* Texture2D = Cast<UTexture2D>(Texture);
 		if (!Texture2D)
 		{
-			UTexture2DArray* Texture2DArray = CastChecked<UTexture2DArray>(Texture);
-			Texture2D = Texture2DArray->SourceTextures.IsValidIndex(TextureIndex) ? Texture2DArray->SourceTextures[TextureIndex] : nullptr;
+			if (UTexture2DArray* Texture2DArray = Cast<UTexture2DArray>(Texture))
+			{
+				Texture2D = Texture2DArray->SourceTextures.IsValidIndex(TextureIndex) ? Texture2DArray->SourceTextures[TextureIndex] : nullptr;
+			}
+			else
+			{
+				SetInitCompleted(true);
+				return;
+			}
 		}
 
 		TOptional<bool> CanGPUTextureBeCPUAccessed = PCGTextureSamplingHelpers::CanGPUTextureBeCPUAccessed(Texture2D);
