[PCG] Fix a couple of fail cases for GPU compute graph compilation.

* For some graph setups the compiler would fail to assemble the full graph and emit a graph that has missing data interfaces. This was observed occasionally with the Copy Points node which appeared to have undefined behaviour on some runs.
* Fix case where a CPU node is embedded in an 'island' of GPU nodes which was not handled properly. Generally, adjacent connected compute graphs were not working without this fix.

#rnx

diff --git a/Engine/Plugins/PCG/Source/PCG/Private/Compute/DataInterfaces/PCGDataCollectionDataInterface.cpp b/Engine/Plugins/PCG/Source/PCG/Private/Compute/DataInterfaces/PCGDataCollectionDataInterface.cpp
--- a/Engine/Plugins/PCG/Source/PCG/Private/Compute/DataInterfaces/PCGDataCollectionDataInterface.cpp
+++ b/Engine/Plugins/PCG/Source/PCG/Private/Compute/DataInterfaces/PCGDataCollectionDataInterface.cpp
@@ -411,9 +411,11 @@ void UPCGDataCollectionDataInterface::GetHLSL(FString& OutHLSL, FString const& I
 UComputeDataProvider* UPCGDataCollectionDataInterface::CreateDataProvider(TObjectPtr<UObject> InBinding, uint64 InInputMask, uint64 InOutputMask) const
 {
 	TRACE_CPUPROFILER_EVENT_SCOPE(UPCGDataCollectionDataInterface::CreateDataProvider);
+	check(ProducerSettings);
 	UPCGDataBinding* Binding = CastChecked<UPCGDataBinding>(InBinding);
 
-	FPCGDataCollectionDesc PinDataDesc = ProducerSettings->ComputeOutputPinDataDesc(OutputPinLabel, Binding);
+	FPCGDataCollectionDesc PinDataDesc;
+	ensure(ProducerSettings->ComputeOutputPinDataDesc(OutputPinLabel, Binding, PinDataDesc));
 
 	if (PCGComputeHelpers::IsBufferSizeTooLarge(PinDataDesc.ComputePackedSizeBytes()))
 	{
@@ -464,6 +466,17 @@ FComputeDataProviderRenderProxy* UPCGDataCollectionDataProvider::GetRenderProxy(
 	TRACE_CPUPROFILER_EVENT_SCOPE(UPCGDataCollectionDataProvider::GetRenderProxy);
 	FPCGDataCollectionDataProviderProxy* Proxy = new FPCGDataCollectionDataProviderProxy(PinDesc, ReadbackMode);
 
+#if WITH_EDITOR
+	if (const UPCGNode* Node = Cast<UPCGNode>(ProducerSettings ? ProducerSettings->GetOuter() : nullptr))
+	{
+		Proxy->SetBufferName(FString::Format(TEXT("PCG_DC_{0}_{1}"), { Node->GetNodeTitle(EPCGNodeTitleType::ListView).ToString(), OutputPinLabel.ToString()}));
+	}
+	else
+	{
+		Proxy->SetBufferName(TEXT("PCG_DC"));
+	}
+#endif
+
 	if (ReadbackMode != EPCGReadbackMode::None)
 	{
 		TWeakObjectPtr<UPCGDataCollectionDataProvider> ThisWeakPtr(this);
@@ -624,8 +637,15 @@ void FPCGDataCollectionDataProviderProxy::AllocateResources(FRDGBuilder& GraphBu
 	TRACE_CPUPROFILER_EVENT_SCOPE(FPCGDataCollectionDataProviderProxy::AllocateResources);
 	check(SizeBytes > 0);
 
+	const TCHAR* DebugName =
+#if WITH_EDITOR
+		*BufferName;
+#else
+		TEXT("PCGDataCollectionBuffer");
+#endif
+
 	FRDGBufferDesc Desc = FRDGBufferDesc::CreateStructuredDesc(sizeof(uint32), SizeBytes >> 2);
-	Buffer = GraphBuilder.CreateBuffer(Desc, TEXT("PCGDataCollectionBuffer"));
+	Buffer = GraphBuilder.CreateBuffer(Desc, DebugName);
 	BufferUAV = GraphBuilder.CreateUAV(Buffer);
 
 	// Initialize with an empty data collection. The kernel may not run, for example if indirect dispatch args end up being 0. Ensure
@@ -649,4 +669,11 @@ void FPCGDataCollectionDataProviderProxy::GetReadbackData(TArray<FReadbackData>&
 	}
 }
 
+#if WITH_EDITOR
+void FPCGDataCollectionDataProviderProxy::SetBufferName(const FString& InBufferName)
+{
+	BufferName = InBufferName;
+}
+#endif
+
 #undef LOCTEXT_NAMESPACE
diff --git a/Engine/Plugins/PCG/Source/PCG/Private/Compute/DataInterfaces/PCGDataCollectionDataInterface.h b/Engine/Plugins/PCG/Source/PCG/Private/Compute/DataInterfaces/PCGDataCollectionDataInterface.h
--- a/Engine/Plugins/PCG/Source/PCG/Private/Compute/DataInterfaces/PCGDataCollectionDataInterface.h
+++ b/Engine/Plugins/PCG/Source/PCG/Private/Compute/DataInterfaces/PCGDataCollectionDataInterface.h
@@ -118,6 +118,10 @@ public:
 	void GetReadbackData(TArray<FReadbackData>& OutReadbackData) const override;
 	//~ End FComputeDataProviderRenderProxy Interface
 
+#if WITH_EDITOR
+	void SetBufferName(const FString& InBufferName);
+#endif
+
 public:
 	FReadbackCallback AsyncReadbackCallback_RenderThread;
 
@@ -132,4 +136,8 @@ protected:
 
 	FRDGBufferRef Buffer = nullptr;
 	FRDGBufferUAVRef BufferUAV = nullptr;
+
+#if WITH_EDITOR
+	FString BufferName;
+#endif
 };
diff --git a/Engine/Plugins/PCG/Source/PCG/Private/Compute/DataInterfaces/PCGDataCollectionUploadDataInterface.cpp b/Engine/Plugins/PCG/Source/PCG/Private/Compute/DataInterfaces/PCGDataCollectionUploadDataInterface.cpp
--- a/Engine/Plugins/PCG/Source/PCG/Private/Compute/DataInterfaces/PCGDataCollectionUploadDataInterface.cpp
+++ b/Engine/Plugins/PCG/Source/PCG/Private/Compute/DataInterfaces/PCGDataCollectionUploadDataInterface.cpp
@@ -49,7 +49,20 @@ FComputeDataProviderRenderProxy* UPCGDataProviderDataCollectionUpload::GetRender
 		return nullptr;
 	}
 
-	return new FPCGDataProviderDataCollectionUploadProxy(PinDesc, MoveTemp(PackedDataCollection), DownstreamInputPinLabels);
+	FPCGDataProviderDataCollectionUploadProxy* Proxy = new FPCGDataProviderDataCollectionUploadProxy(PinDesc, MoveTemp(PackedDataCollection), DownstreamInputPinLabels);
+
+#if WITH_EDITOR
+	if (const UPCGNode* Node = Cast<UPCGNode>(ProducerSettings ? ProducerSettings->GetOuter() : nullptr))
+	{
+		Proxy->SetBufferName(FString::Format(TEXT("PCG_UDC_{0}_{1}"), { Node->GetNodeTitle(EPCGNodeTitleType::ListView).ToString(), OutputPinLabel.ToString() }));
+	}
+	else
+	{
+		Proxy->SetBufferName(TEXT("PCG_UDC"));
+	}
+#endif
+
+	return Proxy;
 }
 
 FPCGDataProviderDataCollectionUploadProxy::FPCGDataProviderDataCollectionUploadProxy(const FPCGDataCollectionDesc& InPinDesc, TArray<uint32>&& InPackedDataCollection, const TArray<FName>& InDownstreamInputPinLabels)
@@ -63,8 +76,15 @@ void FPCGDataProviderDataCollectionUploadProxy::AllocateResources(FRDGBuilder& G
 {
 	TRACE_CPUPROFILER_EVENT_SCOPE(FPCGDataProviderDataCollectionUploadProxy::AllocateResources);
 
+	const TCHAR* DebugName =
+#if WITH_EDITOR
+		*BufferName;
+#else
+		TEXT("PCGDataCollectionUploadBuffer");
+#endif
+
 	const FRDGBufferDesc Desc = FRDGBufferDesc::CreateStructuredDesc(sizeof(uint32), PackedDataCollection.Num());
-	Buffer = GraphBuilder.CreateBuffer(Desc, TEXT("PCGDataCollectionUploadBuffer"));
+	Buffer = GraphBuilder.CreateBuffer(Desc, DebugName);
 	BufferUAV = GraphBuilder.CreateUAV(Buffer);
 
 	GraphBuilder.QueueBufferUpload(Buffer, PackedDataCollection.GetData(), PackedDataCollection.Num() * PackedDataCollection.GetTypeSize(), ERDGInitialDataFlags::None);
diff --git a/Engine/Plugins/PCG/Source/PCG/Private/Compute/Elements/PCGCustomHLSL.cpp b/Engine/Plugins/PCG/Source/PCG/Private/Compute/Elements/PCGCustomHLSL.cpp
--- a/Engine/Plugins/PCG/Source/PCG/Private/Compute/Elements/PCGCustomHLSL.cpp
+++ b/Engine/Plugins/PCG/Source/PCG/Private/Compute/Elements/PCGCustomHLSL.cpp
@@ -224,8 +224,11 @@ int UPCGCustomHLSLSettings::GetProcessingElemCountForInputPin(const UPCGPin* Inp
 			const UPCGSettings* UpstreamSettings = UpstreamOutputPin->Node ? UpstreamOutputPin->Node->GetSettings() : nullptr;
 			check(UpstreamSettings);
 
-			const FPCGDataCollectionDesc PinDesc = UpstreamSettings->ComputeOutputPinDataDesc(UpstreamOutputPin, Binding);
-			ProcessingElemCount += PinDesc.ComputeDataElementCount(InputPin->Properties.AllowedTypes);
+			FPCGDataCollectionDesc PinDesc;
+			if (ensure(UpstreamSettings->ComputeOutputPinDataDesc(UpstreamOutputPin, Binding, PinDesc)))
+			{
+				ProcessingElemCount += PinDesc.ComputeDataElementCount(InputPin->Properties.AllowedTypes);
+			}
 		}
 	}
 
@@ -255,8 +258,11 @@ int UPCGCustomHLSLSettings::ComputeKernelThreadCount(const UPCGDataBinding* Bind
 		{
 			if (const UPCGPin* OutputPin = GetFirstOutputPin())
 			{
-				const FPCGDataCollectionDesc Desc = ComputeOutputPinDataDesc(OutputPin, Binding);
-				ThreadCount = Desc.ComputeDataElementCount(EPCGDataType::Any);
+				FPCGDataCollectionDesc Desc;
+				if (ensure(ComputeOutputPinDataDesc(OutputPin, Binding, Desc)))
+				{
+					ThreadCount = Desc.ComputeDataElementCount(EPCGDataType::Any);
+				}
 			}
 		}
 		else if (DispatchThreadCount == EPCGDispatchThreadCount::FromProductOfInputPins)
@@ -287,15 +293,15 @@ int UPCGCustomHLSLSettings::ComputeKernelThreadCount(const UPCGDataBinding* Bind
 	return ThreadCount;
 }
 
-FPCGDataCollectionDesc UPCGCustomHLSLSettings::ComputeOutputPinDataDesc(const UPCGPin* OutputPin, const UPCGDataBinding* Binding) const
+bool UPCGCustomHLSLSettings::ComputeOutputPinDataDesc(const UPCGPin* OutputPin, const UPCGDataBinding* Binding, FPCGDataCollectionDesc& OutDesc) const
 {
 	check(OutputPin);
 	check(Binding);
 
-	// If node will not execute on GPU then its data comes straight from CPU.
-	if (!ShouldExecuteOnGPU() || !bEnabled)
+	// First inspect data coming from CPU. Will return description if this node is a CPU node or in a different compute graph.
+	if (Binding->ComputeCPUOutputPinDataDesc(OutputPin, OutDesc))
 	{
-		return Super::ComputeOutputPinDataDesc(OutputPin, Binding);
+		return true;
 	}
 
 	FPCGDataCollectionDesc PinDesc;
@@ -449,7 +455,8 @@ FPCGDataCollectionDesc UPCGCustomHLSLSettings::ComputeOutputPinDataDesc(const UP
 		}
 	}
 
-	return PinDesc;
+	OutDesc = PinDesc;
+	return true;
 }
 
 #if WITH_EDITOR
@@ -1493,7 +1500,12 @@ bool UPCGCustomHLSLSettings::AreKernelAttributesValid(FPCGContext* InContext, FT
 		for (const UPCGPin* OutputPin : OutPins)
 		{
 			check(OutputPin);
-			OutputPinDescs.Add(OutputPin->Properties.Label, ComputeOutputPinDataDesc(OutputPin, DataBinding));
+
+			FPCGDataCollectionDesc PinDesc;
+			if (ensure(ComputeOutputPinDataDesc(OutputPin, DataBinding, PinDesc)))
+			{
+				OutputPinDescs.Add(OutputPin->Properties.Label, PinDesc);
+			}
 		}
 
 		const FString Source = ShaderFunctions + ShaderSource;
diff --git a/Engine/Plugins/PCG/Source/PCG/Private/Compute/PCGDataBinding.cpp b/Engine/Plugins/PCG/Source/PCG/Private/Compute/PCGDataBinding.cpp
--- a/Engine/Plugins/PCG/Source/PCG/Private/Compute/PCGDataBinding.cpp
+++ b/Engine/Plugins/PCG/Source/PCG/Private/Compute/PCGDataBinding.cpp
@@ -35,6 +35,29 @@ void UPCGDataBinding::Initialize(
 #endif
 }
 
+bool UPCGDataBinding::ComputeCPUOutputPinDataDesc(const UPCGPin* OutputPin, FPCGDataCollectionDesc& OutDesc) const
+{
+	check(OutputPin);
+	check(Graph);
+
+	// Find the unique alias for the virtual pin created during compilation to receive the data from the upstream OutputPin.
+	// OutputPin can either be on a CPU node or on a GPU node in a different compute graph.
+	const FName* FoundPinAlias = Graph->OutputCPUPinToInputGPUPinAlias.Find(OutputPin);
+	if (!FoundPinAlias)
+	{
+		return false;
+	}
+
+	// Create description from compute graph element input data.
+	OutDesc = FPCGDataCollectionDesc::BuildFromInputDataCollectionAndInputPinLabel(
+		DataForGPU.InputDataCollection,
+		*FoundPinAlias,
+		GetAttributeLookupTable(),
+		GetStringTable());
+
+	return true;
+}
+
 void UPCGDataBinding::InitializeInputData(const FPCGDataCollection& InComputeGraphElementInputData)
 {
 	TRACE_CPUPROFILER_EVENT_SCOPE(UPCGDataBinding::InitializeInputData);
@@ -236,9 +259,16 @@ void UPCGDataBinding::DebugLogDataDescriptions()
 
 			for (const UPCGPin* Pin : Node->GetOutputPins())
 			{
-				const FPCGDataCollectionDesc DataDesc = Settings->ComputeOutputPinDataDesc(Pin, this);
-				UE_LOG(LogPCG, Display, TEXT("\tOutput Pin: %s (%d data)"), *Pin->Properties.Label.ToString(), DataDesc.DataDescs.Num());
-				LogDataDescription(DataDesc);
+				FPCGDataCollectionDesc DataDesc;
+				if (ensure(Settings->ComputeOutputPinDataDesc(Pin, this, DataDesc)))
+				{
+					UE_LOG(LogPCG, Display, TEXT("\tOutput Pin: %s (%d data)"), *Pin->Properties.Label.ToString(), DataDesc.DataDescs.Num());
+					LogDataDescription(DataDesc);
+				}
+				else
+				{
+					UE_LOG(LogPCG, Display, TEXT("\tOutput Pin: %s MISSING!"), *Pin->Properties.Label.ToString());
+				}
 			}
 		}
 	}
diff --git a/Engine/Plugins/PCG/Source/PCG/Private/Compute/PCGDataBinding.h b/Engine/Plugins/PCG/Source/PCG/Private/Compute/PCGDataBinding.h
--- a/Engine/Plugins/PCG/Source/PCG/Private/Compute/PCGDataBinding.h
+++ b/Engine/Plugins/PCG/Source/PCG/Private/Compute/PCGDataBinding.h
@@ -51,6 +51,9 @@ public:
 	/** Get the set of all unique strings entering the compute graph. */
 	const TArray<FString>& GetStringTable() const { return StringTable; }
 
+	/** If the output pin is a connected upstream pin of the compute graph (on a CPU node or a GPU node from a different compute graph), returns a description of this data. */
+	bool ComputeCPUOutputPinDataDesc(const UPCGPin* OutputPin, FPCGDataCollectionDesc& OutDesc) const;
+
 private:
 	void InitializeInputData(const FPCGDataCollection& InComputeGraphElementInputData);
 
diff --git a/Engine/Plugins/PCG/Source/PCG/Private/Elements/PCGCopyPoints.cpp b/Engine/Plugins/PCG/Source/PCG/Private/Elements/PCGCopyPoints.cpp
--- a/Engine/Plugins/PCG/Source/PCG/Private/Elements/PCGCopyPoints.cpp
+++ b/Engine/Plugins/PCG/Source/PCG/Private/Elements/PCGCopyPoints.cpp
@@ -5,6 +5,7 @@
 #include "PCGContext.h"
 #include "Compute/DataInterfaces/PCGCopyPointsDataInterface.h"
 #include "Compute/PCGComputeCommon.h"
+#include "Compute/PCGDataBinding.h"
 #include "Data/PCGPointData.h"
 #include "Data/PCGSpatialData.h"
 #include "Helpers/PCGAsync.h"
@@ -32,15 +33,15 @@ FString UPCGCopyPointsSettings::GetCookedKernelSource(const TMap<FName, FPCGKern
 	return TemplateFile;
 }
 
-FPCGDataCollectionDesc UPCGCopyPointsSettings::ComputeOutputPinDataDesc(const UPCGPin* OutputPin, const UPCGDataBinding* Binding) const
+bool UPCGCopyPointsSettings::ComputeOutputPinDataDesc(const UPCGPin* OutputPin, const UPCGDataBinding* Binding, FPCGDataCollectionDesc& OutDesc) const
 {
 	check(OutputPin);
 	check(Binding);
 
-	// If node will not execute on GPU then its data comes straight from CPU.
-	if (!ShouldExecuteOnGPU() || !bEnabled)
+	// First inspect data coming from CPU. Will return description if this node is a CPU node or in a different compute graph.
+	if (Binding->ComputeCPUOutputPinDataDesc(OutputPin, OutDesc))
 	{
-		return Super::ComputeOutputPinDataDesc(OutputPin, Binding);
+		return true;
 	}
 
 	const UPCGNode* Node = CastChecked<UPCGNode>(GetOuter());
@@ -100,7 +101,8 @@ FPCGDataCollectionDesc UPCGCopyPointsSettings::ComputeOutputPinDataDesc(const UP
 		}
 	}
 
-	return PinDesc;
+	OutDesc = PinDesc;
+	return true;
 }
 
 int UPCGCopyPointsSettings::ComputeKernelThreadCount(const UPCGDataBinding* Binding) const
@@ -108,12 +110,15 @@ int UPCGCopyPointsSettings::ComputeKernelThreadCount(const UPCGDataBinding* Bind
 	const UPCGNode* Node = CastChecked<UPCGNode>(GetOuter());
 	const UPCGPin* OutPin = Node->GetOutputPin(PCGPinConstants::DefaultOutputLabel);
 
-	const FPCGDataCollectionDesc OutputPinDesc = ComputeOutputPinDataDesc(OutPin, Binding);
 	int ThreadCount = 0;
 
-	for (const FPCGDataDesc& DataDesc : OutputPinDesc.DataDescs)
+	FPCGDataCollectionDesc OutputPinDesc;
+	if (ensure(ComputeOutputPinDataDesc(OutPin, Binding, OutputPinDesc)))
 	{
-		ThreadCount += DataDesc.ElementCount;
+		for (const FPCGDataDesc& DataDesc : OutputPinDesc.DataDescs)
+		{
+			ThreadCount += DataDesc.ElementCount;
+		}
 	}
 
 	return ThreadCount;
diff --git a/Engine/Plugins/PCG/Source/PCG/Private/Elements/PCGReroute.cpp b/Engine/Plugins/PCG/Source/PCG/Private/Elements/PCGReroute.cpp
--- a/Engine/Plugins/PCG/Source/PCG/Private/Elements/PCGReroute.cpp
+++ b/Engine/Plugins/PCG/Source/PCG/Private/Elements/PCGReroute.cpp
@@ -45,7 +45,7 @@ FPCGElementPtr UPCGRerouteSettings::CreateElement() const
 	return MakeShared<FPCGRerouteElement>();
 }
 
-FPCGDataCollectionDesc UPCGRerouteSettings::ComputeOutputPinDataDesc(const UPCGPin* OutputPin, const UPCGDataBinding* Binding) const
+bool UPCGRerouteSettings::ComputeOutputPinDataDesc(const UPCGPin* OutputPin, const UPCGDataBinding* Binding, FPCGDataCollectionDesc& OutDesc) const
 {
 	// Reroutes always trivially forward data. Reroutes are culled and not visible to compute graph compilation and therefore we
 	// transparently forward the data here.
@@ -53,11 +53,12 @@ FPCGDataCollectionDesc UPCGRerouteSettings::ComputeOutputPinDataDesc(const UPCGP
 	{
 		if (const UPCGPin* Pin = Node->GetInputPin(PCGPinConstants::DefaultInputLabel))
 		{
-			return ComputeInputPinDataDesc(Pin, Binding);
+			OutDesc = ComputeInputPinDataDesc(Pin, Binding);
+			return true;
 		}
 	}
 
-	return {};
+	return false;
 }
 
 TArray<FPCGPinProperties> UPCGNamedRerouteDeclarationSettings::OutputPinProperties() const
diff --git a/Engine/Plugins/PCG/Source/PCG/Private/Elements/PCGStaticMeshSpawner.cpp b/Engine/Plugins/PCG/Source/PCG/Private/Elements/PCGStaticMeshSpawner.cpp
--- a/Engine/Plugins/PCG/Source/PCG/Private/Elements/PCGStaticMeshSpawner.cpp
+++ b/Engine/Plugins/PCG/Source/PCG/Private/Elements/PCGStaticMeshSpawner.cpp
@@ -107,15 +107,15 @@ FString UPCGStaticMeshSpawnerSettings::GetCookedKernelSource(const TMap<FName, F
 	return TemplateFile;
 }
 
-FPCGDataCollectionDesc UPCGStaticMeshSpawnerSettings::ComputeOutputPinDataDesc(const UPCGPin* OutputPin, const UPCGDataBinding* Binding) const
+bool UPCGStaticMeshSpawnerSettings::ComputeOutputPinDataDesc(const UPCGPin* OutputPin, const UPCGDataBinding* Binding, FPCGDataCollectionDesc& OutDesc) const
 {
 	check(OutputPin);
 	check(Binding);
 
-	// If node will not execute on GPU then its data comes straight from CPU.
-	if (!ShouldExecuteOnGPU() || !bEnabled)
+	// First inspect data coming from CPU. Will return description if this node is a CPU node or in a different compute graph.
+	if (Binding->ComputeCPUOutputPinDataDesc(OutputPin, OutDesc))
 	{
-		return Super::ComputeOutputPinDataDesc(OutputPin, Binding);
+		return true;
 	}
 
 	const FPCGDataCollectionDesc InputPinDesc = ComputeInputPinDataDesc(PCGPinConstants::DefaultInputLabel, Binding);
@@ -176,7 +176,8 @@ FPCGDataCollectionDesc UPCGStaticMeshSpawnerSettings::ComputeOutputPinDataDesc(c
 		}
 	}
 
-	return OutputPinDesc;
+	OutDesc = OutputPinDesc;
+	return true;
 }
 
 const TArray<FPCGKernelAttributeKey> UPCGStaticMeshSpawnerSettings::GetKernelAttributeKeys() const
diff --git a/Engine/Plugins/PCG/Source/PCG/Private/Graph/PCGGraphCompiler.cpp b/Engine/Plugins/PCG/Source/PCG/Private/Graph/PCGGraphCompiler.cpp
--- a/Engine/Plugins/PCG/Source/PCG/Private/Graph/PCGGraphCompiler.cpp
+++ b/Engine/Plugins/PCG/Source/PCG/Private/Graph/PCGGraphCompiler.cpp
@@ -974,6 +974,74 @@ void FPCGGraphCompiler::AddReferencedObjects(FReferenceCollector& Collector)
 	}
 }
 
+bool FPCGGraphCompiler::VisitTasksInExecutionOrder(const TArray<FPCGGraphTask>& InTasks, const TMap<FPCGTaskId, TArray<FPCGTaskId>>& InTaskToTaskSuccessors, const TFunction<bool(FPCGTaskId)>& InVisitor)
+{
+	// Populate initial sets of tasks that are ready to consume vs ones that currently blocked.
+	TArray<FPCGTaskId> ReadyTaskIds;
+	TSet<FPCGTaskId> RemainingTaskIds;
+	ReadyTaskIds.Reserve(InTasks.Num());
+	RemainingTaskIds.Reserve(InTasks.Num());
+
+	// Sort tasks into those that don't have inputs (ready tasks) vs those that cannot execute yet (remaining).
+	for (FPCGTaskId TaskId = 0; TaskId < InTasks.Num(); ++TaskId)
+	{
+		if (InTasks[TaskId].Inputs.IsEmpty())
+		{
+			ReadyTaskIds.Add(TaskId);
+		}
+		else
+		{
+			RemainingTaskIds.Add(TaskId);
+		}
+	}
+
+	// Loop until we have consumed all tasks.
+	while (!ReadyTaskIds.IsEmpty() || !RemainingTaskIds.IsEmpty())
+	{
+		FPCGTaskId ReadyTaskId = ReadyTaskIds.Pop();
+
+		if (!InVisitor(ReadyTaskId))
+		{
+			return false;
+		}
+
+		// Queue up any successors that are ready to go.
+		const TArray<FPCGTaskId>* Successors = InTaskToTaskSuccessors.Find(ReadyTaskId);
+		if (!Successors)
+		{
+			continue;
+		}
+
+		for (FPCGTaskId Successor : *Successors)
+		{
+			const bool bSuccessorQueued = ReadyTaskIds.Contains(Successor);
+
+			// All successors should either be already queued, or waiting to be queued.
+			check(bSuccessorQueued || RemainingTaskIds.Contains(Successor));
+
+			if (!bSuccessorQueued)
+			{
+				bool bSuccessorReady = true;
+				for (const FPCGGraphTaskInput& Input : InTasks[Successor].Inputs)
+				{
+					if (ReadyTaskIds.Contains(Input.TaskId) || RemainingTaskIds.Contains(Input.TaskId))
+					{
+						bSuccessorReady = false;
+						break;
+					}
+				}
+
+				if (bSuccessorReady)
+				{
+					ReadyTaskIds.Add(Successor);
+					RemainingTaskIds.Remove(Successor);
+				}
+			}
+		}
+	}
+
+	return true;
+}
 
 void FPCGGraphCompiler::PostCullStackCleanup(TArray<FPCGGraphTask>& InCompiledTasks, FPCGStackContext& InOutStackContext)
 {
diff --git a/Engine/Plugins/PCG/Source/PCG/Private/Graph/PCGGraphCompiler.h b/Engine/Plugins/PCG/Source/PCG/Private/Graph/PCGGraphCompiler.h
--- a/Engine/Plugins/PCG/Source/PCG/Private/Graph/PCGGraphCompiler.h
+++ b/Engine/Plugins/PCG/Source/PCG/Private/Graph/PCGGraphCompiler.h
@@ -98,10 +98,14 @@ public:
 
 	void AddReferencedObjects(FReferenceCollector& Collector);
 
+	/** Executes lambda over provided tasks in execution order (task is not visited until all upstream dependencies are visited). Aborts if visitor lambda returns false.
+	* Returns true if all tasks visited. */
+	static bool VisitTasksInExecutionOrder(const TArray<FPCGGraphTask>& InTasks, const TMap<FPCGTaskId, TArray<FPCGTaskId>>& InTaskToTaskSuccessors, const TFunction<bool(FPCGTaskId)>& InVisitor);
+
 private:
 	TArray<FPCGGraphTask> CompileGraph(UPCGGraph* InGraph, FPCGTaskId& NextId, FPCGStackContext& InOutStackContext);
 
-	/** Compiles the top graph and applies culling optimizations if a non-unitialized grid size is provided. */
+	/** Compiles the top graph and applies culling optimizations if a non-uninitialized grid size is provided. */
 	void CompileTopGraph(UPCGGraph* InGraph, uint32 GenerationGridSize);
 
 	/** Propagates grid sizes through a graph's compiled tasks. */
diff --git a/Engine/Plugins/PCG/Source/PCG/Private/Graph/PCGGraphCompilerGPU.cpp b/Engine/Plugins/PCG/Source/PCG/Private/Graph/PCGGraphCompilerGPU.cpp
--- a/Engine/Plugins/PCG/Source/PCG/Private/Graph/PCGGraphCompilerGPU.cpp
+++ b/Engine/Plugins/PCG/Source/PCG/Private/Graph/PCGGraphCompilerGPU.cpp
@@ -5,6 +5,8 @@
 #include "PCGGraph.h"
 #include "PCGModule.h"
 #include "PCGPin.h"
+#include "Compute/PCGComputeCommon.h"
+#include "Compute/PCGComputeKernelSource.h"
 #include "Compute/DataInterfaces/PCGCustomKernelDataInterface.h"
 #include "Compute/DataInterfaces/PCGDataCollectionDataInterface.h"
 #include "Compute/DataInterfaces/PCGDataCollectionUploadDataInterface.h"
@@ -12,8 +14,6 @@
 #include "Compute/DataInterfaces/PCGLandscapeDataInterface.h"
 #include "Compute/DataInterfaces/PCGTextureDataInterface.h"
 #include "Compute/Elements/PCGComputeGraphElement.h"
-#include "Compute/PCGComputeCommon.h"
-#include "Compute/PCGComputeKernelSource.h"
 #include "Elements/PCGStaticMeshSpawner.h"
 #include "Graph/PCGGraphCompiler.h"
 #include "Graph/PCGGraphExecutor.h"
@@ -95,167 +95,132 @@ void FPCGGraphCompilerGPU::CollectGPUNodeSubsets(
 	const TSet<FPCGTaskId>& InGPUCompatibleTaskIds,
 	TArray<TSet<FPCGTaskId>>& OutNodeSubsetsToConvertToCFGraph)
 {
+	TRACE_CPUPROFILER_EVENT_SCOPE(FPCGGraphCompilerGPU::CollectGPUNodeSubsets);
+
 	// Identifies connected sets of GPU nodes, giving each a non-zero ID value.
 	TArray<uint32> ConnectedGPUNodeIslandIDs;
 	LabelConnectedGPUNodeIslands(InCompiledTasks, InGPUCompatibleTaskIds, InTaskSuccessors, ConnectedGPUNodeIslandIDs);
 
-	// Populate initial sets of tasks that are ready to consume vs ones that currently blocked.
-	TSet<FPCGTaskId> ReadyTaskIds;
-	TSet<FPCGTaskId> RemainingTaskIds;
-	ReadyTaskIds.Reserve(InCompiledTasks.Num());
-	RemainingTaskIds.Reserve(InCompiledTasks.Num());
-
-	for (FPCGTaskId TaskId = 0; TaskId < InCompiledTasks.Num(); ++TaskId)
-	{
-		if (InCompiledTasks[TaskId].Inputs.IsEmpty())
-		{
-			ReadyTaskIds.Add(TaskId);
-		}
-		else
-		{
-			RemainingTaskIds.Add(TaskId);
-		}
-	}
-
-	// Queue all successors of InTaskId that are ready to go (all upstream input tasks have been processed).
-	auto QueueSuccessors = [&InTaskSuccessors, &ReadyTaskIds, &RemainingTaskIds, &InCompiledTasks](int InTaskId)
-	{
-		bool bQueuedTask = false;
-
-		// Queue up any successors that are ready to go.
-		if (const TArray<FPCGTaskId>* Successors = InTaskSuccessors.Find(InTaskId))
+	// For every CPU node that has one or more downstream GPU node connected, check for a GPU -> CPU -> GPU pattern where data flows
+	// from a GPU node island to the CPU and then back to the same island. For such cases, we traverse the entire tree of GPU nodes that
+	// are in the island and downstream of the CPU node and bump their island ID - splitting the island into a portion that is independent
+	// of the CPU node and a portion that is dependent on it, so that we can read the data back to CPU, execute the CPU portion, then
+	// re-upload to GPU.
+	
+	// Any new island IDs will be created from the task count which will be larger than any island IDs presently set.
+	uint32 NextIslandId = InCompiledTasks.Num();
+
+	// Cache dependencies. Since our islands are only ever split (rather than replaced or removed), the dependency on each island ID is invariant.
+	TMap<TPair<FPCGTaskId, /*IslandId*/uint32>, /*bIsDependent*/bool> CPUNodeIsDependentOnIslandCached;
+
+	// Visit tasks in execution order so that splits happen as upstream to minimize island splits.
+	FPCGGraphCompiler::VisitTasksInExecutionOrder(
+		InCompiledTasks,
+		InTaskSuccessors,
+		[&InCompiledTasks, &InTaskSuccessors, &InGPUCompatibleTaskIds, &ConnectedGPUNodeIslandIDs, &NextIslandId, &CPUNodeIsDependentOnIslandCached](FPCGTaskId InTaskId) -> bool
 		{
-			for (FPCGTaskId Successor : *Successors)
+			if (InGPUCompatibleTaskIds.Contains(InTaskId))
 			{
-				const bool bSuccessorQueued = ReadyTaskIds.Contains(Successor);
-
-				// All successors should either be already queued, or waiting to be queued.
-				check(bSuccessorQueued || RemainingTaskIds.Contains(Successor));
-
-				if (!bSuccessorQueued)
-				{
-					bool bSuccessorReady = true;
-					for (const FPCGGraphTaskInput& Input : InCompiledTasks[Successor].Inputs)
-					{
-						if (ReadyTaskIds.Contains(Input.TaskId) || RemainingTaskIds.Contains(Input.TaskId))
-						{
-							bSuccessorReady = false;
-						}
-					}
-
-					if (bSuccessorReady)
-					{
-						ReadyTaskIds.Add(Successor);
-						RemainingTaskIds.Remove(Successor);
-						bQueuedTask = true;
-					}
-				}
+				// Skip GPU nodes.
+				return true;
 			}
-		}
-
-		return bQueuedTask;
-	};
 
-	// Local to loops below but pulled out for performance.
-	TArray<FPCGTaskId> FoundReadyTaskIds;
-	TSet<FPCGTaskId> GPUSubsetTaskIds;
-	FoundReadyTaskIds.Reserve(InCompiledTasks.Num());
-	GPUSubsetTaskIds.Reserve(InCompiledTasks.Num());
-
-	// Build subsets of nodes that are GPU compatible and can be dispatched together.
-	while (!ReadyTaskIds.IsEmpty() || !RemainingTaskIds.IsEmpty())
-	{
-		// Consume as many CPU nodes as we can.
-		bool bQueuedTasks = true;
-		while (bQueuedTasks)
-		{
-			FoundReadyTaskIds.Reset();
-
-			for (FPCGTaskId ReadyTaskId : ReadyTaskIds)
+			if (const TArray<FPCGTaskId>* Successors = InTaskSuccessors.Find(InTaskId))
 			{
-				const bool bIsCPUNode = ConnectedGPUNodeIslandIDs[ReadyTaskId] == 0;
-				if (bIsCPUNode)
+				for (FPCGTaskId SuccessorTaskId : *Successors)
 				{
-					FoundReadyTaskIds.Add(ReadyTaskId);
-				}
-			}
+					if (!InGPUCompatibleTaskIds.Contains(SuccessorTaskId))
+					{
+						continue;
+					}
 
-			bQueuedTasks = false;
+					const uint32 SuccessorIsland = ConnectedGPUNodeIslandIDs[SuccessorTaskId];
 
-			for (FPCGTaskId ReadyCPUTaskId : FoundReadyTaskIds)
-			{
-				ReadyTaskIds.Remove(ReadyCPUTaskId);
-				bQueuedTasks |= QueueSuccessors(ReadyCPUTaskId);
-			}
-		}
+					// Recursively check entire node tree upstream of this CPU node to see if it can be fed by any node in the GPU island.
+					auto CPUNodeIsDependentOnIsland = [&InCompiledTasks, &InTaskSuccessors, &InGPUCompatibleTaskIds, &ConnectedGPUNodeIslandIDs, &CPUNodeIsDependentOnIslandCached](FPCGTaskId InTaskId, uint32 InIslandID, auto&& RecursiveCall) -> bool
+					{
+						if (bool* CachedValue = CPUNodeIsDependentOnIslandCached.Find({ InTaskId, InIslandID }))
+						{
+							return *CachedValue;
+						}
 
-		GPUSubsetTaskIds.Reset();
+						// Is this task is part of the specified island.
+						bool bIsDependent = ConnectedGPUNodeIslandIDs[InTaskId] == InIslandID;
 
-		int StackIndex = INDEX_NONE;
-		uint32 IslandID = INDEX_NONE;
+						if (!bIsDependent)
+						{
+							// Check upstream tasks recursively.
+							for (const FPCGGraphTaskInput& Input : InCompiledTasks[InTaskId].Inputs)
+							{
+								if (RecursiveCall(Input.TaskId, InIslandID, RecursiveCall))
+								{
+									bIsDependent = true;
+									break;
+								}
+							}
+						}
 
-		// Now the opposite - consume as many GPU nodes as we can and accumulate them into a set that will be compiled into a compute graph.
-		bQueuedTasks = !ReadyTaskIds.IsEmpty();
-		while (bQueuedTasks)
-		{
-			FoundReadyTaskIds.Reset();
+						CPUNodeIsDependentOnIslandCached.Add({ InTaskId, InIslandID }, bIsDependent);
 
-			for (FPCGTaskId ReadyTaskId : ReadyTaskIds)
-			{
-				const uint32 TaskIslandID = ConnectedGPUNodeIslandIDs[ReadyTaskId];
-				if (TaskIslandID == 0)
-				{
-					// Non-gpu task - skip
-					continue;
-				}
+						return bIsDependent;
+					};
 
-				const bool bIslandMatches = (IslandID == INDEX_NONE) || (IslandID == TaskIslandID);
+					if (CPUNodeIsDependentOnIsland(InTaskId, SuccessorIsland, CPUNodeIsDependentOnIsland))
+					{
+						// Propagate a new island ID to all downstream GPU tasks within the island.
 
-				// For now don't mix tasks from different execution stacks (in and out of subgraphs for instance) into one compute graph.
-				const bool bStackMatches = (StackIndex == INDEX_NONE) || (InCompiledTasks[ReadyTaskId].StackIndex == StackIndex);
+						auto PropagateIslandIDDownstream = [&InCompiledTasks, &InTaskSuccessors, &InGPUCompatibleTaskIds, &ConnectedGPUNodeIslandIDs](FPCGTaskId InTaskId, uint32 InOldIslandID, uint32 InNewIslandID, auto&& RecursiveCall) -> void
+						{
+							ConnectedGPUNodeIslandIDs[InTaskId] = InNewIslandID;
 
-				if (bIslandMatches && bStackMatches)
-				{
-					IslandID = TaskIslandID;
-					StackIndex = InCompiledTasks[ReadyTaskId].StackIndex;
+							if (const TArray<FPCGTaskId>* Successors = InTaskSuccessors.Find(InTaskId))
+							{
+								for (FPCGTaskId Successor : *Successors)
+								{
+									if (ConnectedGPUNodeIslandIDs[Successor] == InOldIslandID && InGPUCompatibleTaskIds.Contains(Successor))
+									{
+										RecursiveCall(Successor, InOldIslandID, InNewIslandID, RecursiveCall);
+									}
+								}
+							}
+						};
 
-					FoundReadyTaskIds.Add(ReadyTaskId);
+						PropagateIslandIDDownstream(SuccessorTaskId, SuccessorIsland, NextIslandId++, PropagateIslandIDDownstream);
+					}
 				}
 			}
 
-			bQueuedTasks = false;
-
-			for (FPCGTaskId ReadyGPUTaskId : FoundReadyTaskIds)
-			{
-				GPUSubsetTaskIds.Add(ReadyGPUTaskId);
-				ReadyTaskIds.Remove(ReadyGPUTaskId);
-				bQueuedTasks |= QueueSuccessors(ReadyGPUTaskId);
-			}
-		}
+			return true;
+		});
 
-		if (!GPUSubsetTaskIds.IsEmpty())
+	// Island IDs now correctly identify subsets of nodes that will be assembled into compute graphs for GPU execution.
+	for (FPCGTaskId TaskId = 0; TaskId < InCompiledTasks.Num(); ++TaskId)
+	{
+		if (ConnectedGPUNodeIslandIDs[TaskId] != 0)
 		{
-			bool bAllNodesValid = true;
-			for (FPCGTaskId& TaskId : GPUSubsetTaskIds)
+			TSet<FPCGTaskId> GPUNodeSubset;
+
+			const uint32 IslandId = ConnectedGPUNodeIslandIDs[TaskId];
+
+			for (FPCGTaskId OtherTaskId = TaskId; OtherTaskId < InCompiledTasks.Num(); ++OtherTaskId)
 			{
-				const UPCGSettings* Settings = InCompiledTasks[TaskId].Node ? InCompiledTasks[TaskId].Node->GetSettings() : nullptr;
-				if (Settings && !Settings->IsKernelValid())
+				if (ConnectedGPUNodeIslandIDs[OtherTaskId] == IslandId)
 				{
-					bAllNodesValid = false;
-					break;
+					GPUNodeSubset.Add(OtherTaskId);
+
+					ConnectedGPUNodeIslandIDs[OtherTaskId] = 0;
 				}
 			}
 
-			if (bAllNodesValid)
-			{
-				OutNodeSubsetsToConvertToCFGraph.Add(MoveTemp(GPUSubsetTaskIds));
-			}
+			OutNodeSubsetsToConvertToCFGraph.Add(MoveTemp(GPUNodeSubset));
 		}
 	}
 }
 
 void FPCGGraphCompilerGPU::CreateGatherTasksAtGPUInputs(UPCGGraph* InGraph, const TSet<FPCGTaskId>& InGPUCompatibleTaskIds, TArray<FPCGGraphTask>& InOutCompiledTasks)
 {
+	TRACE_CPUPROFILER_EVENT_SCOPE(FPCGGraphCompilerGPU::CreateGatherTasksAtGPUInputs);
+
 	using FOriginalInputPinKey = TPair<FPCGTaskId /* Original GPU task */, FName /* Input pin label */>;
 
 	// These are local to loop below but hoisted here for efficiency.
@@ -358,16 +323,14 @@ void FPCGGraphCompilerGPU::CreateGatherTasksAtGPUInputs(UPCGGraph* InGraph, cons
 	}
 }
 
-void FPCGGraphCompilerGPU::WireGPUGraphNode(
-	FPCGTaskId InGPUGraphTaskId,
+void FPCGGraphCompilerGPU::SetupVirtualPins(
 	const TSet<FPCGTaskId>& InCollapsedTasks,
-	const TSet<FPCGTaskId>& InGPUCompatibleTaskIds,
-	TArray<FPCGGraphTask>& InOutCompiledTasks,
+	const TArray<FPCGGraphTask>& InCompiledTasks,
 	const FTaskToSuccessors& InTaskSuccessors,
 	FOriginalToVirtualPin& OutOriginalToVirtualPin,
 	TMap<TSoftObjectPtr<const UPCGPin>, FName>& OutOutputCPUPinToVirtualPin)
 {
-	FPCGGraphTask& GPUGraphTask = InOutCompiledTasks[InGPUGraphTaskId];
+	TRACE_CPUPROFILER_EVENT_SCOPE(FPCGGraphCompilerGPU::SetupVirtualPins);
 
 	// Used to construct unique input/output labels, ultimately consumed in graph executor in BuildTaskInput and PostExecute for input/output respectively.
 	int InputCount = 0;
@@ -376,28 +339,25 @@ void FPCGGraphCompilerGPU::WireGPUGraphNode(
 	// Add all compute graph task inputs and outputs.
 	for (FPCGTaskId GPUTaskId : InCollapsedTasks)
 	{
-		// First find CPU to GPU edges and wire in the GPU graph node inputs.
-		for (const FPCGGraphTaskInput& Input : InOutCompiledTasks[GPUTaskId].Inputs)
+		// First input edges to the compute graph.
+		for (const FPCGGraphTaskInput& Input : InCompiledTasks[GPUTaskId].Inputs)
 		{
-			if (InGPUCompatibleTaskIds.Contains(Input.TaskId))
+			if (InCollapsedTasks.Contains(Input.TaskId))
 			{
 				continue;
 			}
 
-			FPCGGraphTaskInput& AddedInput = GPUGraphTask.Inputs.Add_GetRef(Input);
-
 			// TODO is pinless fine with skipping?
-			if (AddedInput.DownstreamPin.IsSet())
+			if (Input.DownstreamPin.IsSet())
 			{
-				const FName VirtualLabel = *FString::Format(TEXT("{0}-VirtualIn{1}"), { AddedInput.DownstreamPin->Label.ToString(), InputCount });
+				const FName VirtualLabel = *FString::Format(TEXT("{0}-VirtualIn{1}"), { Input.DownstreamPin->Label.ToString(), InputCount });
 				const bool bIsInputPin = true;
-				OutOriginalToVirtualPin.Add({ GPUTaskId, AddedInput.DownstreamPin->Label, bIsInputPin }, VirtualLabel);
-				AddedInput.DownstreamPin->Label = VirtualLabel;
+				OutOriginalToVirtualPin.Add({ GPUTaskId, Input.DownstreamPin->Label, bIsInputPin }, VirtualLabel);
 
 				++InputCount;
 
-				const FPCGGraphTask& UpstreamTask = InOutCompiledTasks[Input.TaskId];
-				const UPCGPin* OutputPin = UpstreamTask.Node ? UpstreamTask.Node->GetOutputPin(AddedInput.UpstreamPin->Label) : nullptr;
+				const FPCGGraphTask& UpstreamTask = InCompiledTasks[Input.TaskId];
+				const UPCGPin* OutputPin = UpstreamTask.Node ? UpstreamTask.Node->GetOutputPin(Input.UpstreamPin->Label) : nullptr;
 
 				// Grid Linkages don't have a node associated, so ask the GridLinkageElement for the output pin instead.
 				if (!OutputPin)
@@ -428,10 +388,86 @@ void FPCGGraphCompilerGPU::WireGPUGraphNode(
 			continue;
 		}
 
+		// Next consider output edges of the compute graph.
+		for (FPCGTaskId Successor : InTaskSuccessors[GPUTaskId])
+		{
+			if (InCollapsedTasks.Contains(Successor))
+			{
+				continue;
+			}
+
+			// Rewire inputs of this downstream CPU node to the outputs of the compute graph task.
+			const FPCGGraphTask& DownstreamCPUNode = InCompiledTasks[Successor];
+
+			for (int SuccessorInputIndex = 0; SuccessorInputIndex < DownstreamCPUNode.Inputs.Num(); ++SuccessorInputIndex)
+			{
+				// Skip irrelevant edges.
+				if (DownstreamCPUNode.Inputs[SuccessorInputIndex].TaskId != GPUTaskId)
+				{
+					continue;
+				}
+
+				if (DownstreamCPUNode.Inputs[SuccessorInputIndex].UpstreamPin.IsSet())
+				{
+					const FNodePin PinKey = { GPUTaskId, DownstreamCPUNode.Inputs[SuccessorInputIndex].UpstreamPin->Label, /*Pin is input*/false };
+					if (!OutOriginalToVirtualPin.Contains(PinKey))
+					{
+						const FName VirtualLabel = *FString::Format(TEXT("{0}-VirtualOut{1}"), { DownstreamCPUNode.Inputs[SuccessorInputIndex].UpstreamPin->Label.ToString(), OutputCount });
+						OutOriginalToVirtualPin.Add(PinKey, VirtualLabel);
+
+						++OutputCount;
+					}
+				}
+			}
+		}
+	}
+}
+
+void FPCGGraphCompilerGPU::WireGPUGraphNode(
+	FPCGTaskId InGPUGraphTaskId,
+	const TSet<FPCGTaskId>& InCollapsedTasks,
+	TArray<FPCGGraphTask>& InOutCompiledTasks,
+	const FTaskToSuccessors& InTaskSuccessors,
+	const FOriginalToVirtualPin& InOriginalToVirtualPin,
+	const TMap<TSoftObjectPtr<const UPCGPin>, FName>& InOutputCPUPinToVirtualPin)
+{
+	TRACE_CPUPROFILER_EVENT_SCOPE(FPCGGraphCompilerGPU::WireGPUGraphNode);
+
+	FPCGGraphTask& GPUGraphTask = InOutCompiledTasks[InGPUGraphTaskId];
+
+	// Add all compute graph task inputs and outputs.
+	for (FPCGTaskId GPUTaskId : InCollapsedTasks)
+	{
+		// First find CPU to GPU edges and wire in the GPU graph node inputs.
+		for (const FPCGGraphTaskInput& Input : InOutCompiledTasks[GPUTaskId].Inputs)
+		{
+			if (InCollapsedTasks.Contains(Input.TaskId))
+			{
+				continue;
+			}
+
+			FPCGGraphTaskInput& AddedInput = GPUGraphTask.Inputs.Add_GetRef(Input);
+
+			// TODO is pinless fine with skipping?
+			if (AddedInput.DownstreamPin.IsSet())
+			{
+				const FName* FoundVirtualLabel = InOriginalToVirtualPin.Find({ GPUTaskId, AddedInput.DownstreamPin->Label, /*bIsInputPin=*/true });
+				if (ensure(FoundVirtualLabel))
+				{
+					AddedInput.DownstreamPin->Label = *FoundVirtualLabel;
+				}
+			}
+		}
+
+		if (!InTaskSuccessors.Contains(GPUTaskId))
+		{
+			continue;
+		}
+
 		// Next consider GPU to CPU edges to wire in the GPU graph node outputs.
 		for (FPCGTaskId Successor : InTaskSuccessors[GPUTaskId])
 		{
-			if (InGPUCompatibleTaskIds.Contains(Successor))
+			if (InCollapsedTasks.Contains(Successor))
 			{
 				continue;
 			}
@@ -459,20 +495,12 @@ void FPCGGraphCompilerGPU::WireGPUGraphNode(
 				if (DownstreamCPUNode.Inputs[SuccessorInputIndex].UpstreamPin.IsSet())
 				{
 					const FNodePin PinKey = { GPUTaskId, InputCopy.UpstreamPin->Label, /*Pin is input*/false };
-					if (const FName* FoundVirtualPinLabel = OutOriginalToVirtualPin.Find(PinKey))
+					const FName* FoundVirtualPinLabel = InOriginalToVirtualPin.Find(PinKey);
+					if (ensure(FoundVirtualPinLabel))
 					{
 						// Wire to the existing virtual output pin.
 						InputCopy.UpstreamPin->Label = *FoundVirtualPinLabel;
 					}
-					else
-					{
-						const FName VirtualLabel = *FString::Format(TEXT("{0}-VirtualOut{1}"), { InputCopy.UpstreamPin->Label.ToString(), OutputCount });
-						OutOriginalToVirtualPin.Add(PinKey, VirtualLabel);
-
-						InputCopy.UpstreamPin->Label = VirtualLabel;
-
-						++OutputCount;
-					}
 				}
 
 				DownstreamCPUNode.Inputs.Add(MoveTemp(InputCopy));
@@ -487,11 +515,14 @@ void FPCGGraphCompilerGPU::BuildGPUGraphTask(
 	uint32 InGridSize,
 	FPCGTaskId InGPUGraphTaskId,
 	const TSet<FPCGTaskId>& InCollapsedTasks,
+	const TSet<FPCGTaskId>& InAllGPUCompatibleTasks,
 	const FTaskToSuccessors& InTaskSuccessors,
 	TArray<FPCGGraphTask>& InOutCompiledTasks,
 	const FOriginalToVirtualPin& InOriginalToVirtualPin,
 	const TMap<TSoftObjectPtr<const UPCGPin>, FName>& InOutputCPUPinToVirtualPin)
 {
+	TRACE_CPUPROFILER_EVENT_SCOPE(FPCGGraphCompilerGPU::BuildGPUGraphTask);
+
 	check(InGraph);
 	const FName GraphName = MakeUniqueObjectName(InGraph, UPCGComputeGraph::StaticClass(), InGraph->GetFName());
 
@@ -507,6 +538,8 @@ void FPCGGraphCompilerGPU::BuildGPUGraphTask(
 	// For CPU->GPU edges, an upload data interface is created. For GPU->CPU edges, a readback data interface is created.
 	auto CreateDataInterface = [&InCollapsedTasks, &InOutCompiledTasks, ComputeGraph](FPCGTaskId InTaskId, bool bRequiresReadback, const FPCGPinProperties& InOutputPinProperties) -> UPCGComputeDataInterface*
 	{
+		TRACE_CPUPROFILER_EVENT_SCOPE(CreateDataInterface);
+
 		const bool bUpstreamIsGPUTask = InCollapsedTasks.Contains(InTaskId);
 
 		EPCGDataType PinType = InOutputPinProperties.AllowedTypes;
@@ -646,11 +679,11 @@ void FPCGGraphCompilerGPU::BuildGPUGraphTask(
 			}
 		}
 
-		// Create any DIs for upstream CPU nodes.
+		// Create any DIs for upstream nodes outside of this compute graph.
 		for (const FPCGGraphTaskInput& Input : InOutCompiledTasks[TaskId].Inputs)
 		{
-			// Only deal with upstream CPU tasks.
-			if (InCollapsedTasks.Contains(Input.TaskId))
+			// Never wire directly to the GPU-compatible tasks, these are collapsed into compute graphs (and will be culled in a next step).
+			if (InAllGPUCompatibleTasks.Contains(Input.TaskId))
 			{
 				continue;
 			}
@@ -744,6 +777,15 @@ void FPCGGraphCompilerGPU::BuildGPUGraphTask(
 
 		for (const FPCGGraphTaskInput& Input : InOutCompiledTasks[TaskId].Inputs)
 		{
+			// Currently the new compute graph tasks are wired into the graph in parallel to each GPU node task. The GPU node tasks
+			// will be culled at the end, leaving only the compute graphs. Only create DIs for tasks within this compute graph (in InCollapsedTasks),
+			// or for tasks that will not be culled (not in InAllGPUCompatibleTasks).
+			const bool bValidInput = InCollapsedTasks.Contains(Input.TaskId) || !InAllGPUCompatibleTasks.Contains(Input.TaskId);
+			if (!bValidInput)
+			{
+				continue;
+			}
+
 			if (!Input.UpstreamPin.IsSet())
 			{
 				// Execution-only dependencies not supported currently. Unclear if this should ever be supported for GPU graphs.
@@ -1101,6 +1143,8 @@ void FPCGGraphCompilerGPU::BuildGPUGraphTask(
 
 	if (FApp::CanEverRender())
 	{
+		TRACE_CPUPROFILER_EVENT_SCOPE(UpdateResources);
+
 		// Compile shader resources and create render proxies.
 		ComputeGraph->UpdateResources();
 	}
@@ -1152,6 +1196,27 @@ void FPCGGraphCompilerGPU::CreateGPUNodes(FPCGGraphCompiler& InOutCompiler, UPCG
 	TArray<TSet<FPCGTaskId>> NodeSubsetsToConvertToCFGraph;
 	CollectGPUNodeSubsets(InOutCompiledTasks, TaskSuccessors, GPUCompatibleTaskIds, NodeSubsetsToConvertToCFGraph);
 
+	// Mapping from task ID & pin label to a virtual pin label. Compute graphs are executed within a generated element,
+	// and the input and output pins of this element must have unique virtual pin labels so that we can parse the data that
+	// PCG provides through the input data collection correctly, and route the output data to the downstream pins correctly.
+	TArray<FOriginalToVirtualPin> OriginalToVirtualPin;
+
+	TArray<TMap<TSoftObjectPtr<const UPCGPin>, FName>> OutputCPUPinToVirtualPin;
+
+	OriginalToVirtualPin.SetNum(NodeSubsetsToConvertToCFGraph.Num());
+	OutputCPUPinToVirtualPin.SetNum(NodeSubsetsToConvertToCFGraph.Num());
+
+	// Setup mappings from existing pins to compute graph element virtual pins as a prestep before wiring in the compute graph tasks.
+	for (int32 ComputeGraphIndex = 0; ComputeGraphIndex < NodeSubsetsToConvertToCFGraph.Num(); ++ComputeGraphIndex)
+	{
+		SetupVirtualPins(
+			NodeSubsetsToConvertToCFGraph[ComputeGraphIndex],
+			InOutCompiledTasks,
+			TaskSuccessors,
+			OriginalToVirtualPin[ComputeGraphIndex],
+			OutputCPUPinToVirtualPin[ComputeGraphIndex]);
+	}
+
 	// Do actual collapsing now, one subset at a time. Each collapse will do all fixup of task ids? That will invalidate
 	// ids in NodeSubsetsToConvertToCFGraph, so may need remap table. But can ignore this for now.
 	for (int32 ComputeGraphIndex = 0; ComputeGraphIndex < NodeSubsetsToConvertToCFGraph.Num(); ++ComputeGraphIndex)
@@ -1177,22 +1242,14 @@ void FPCGGraphCompilerGPU::CreateGPUNodes(FPCGGraphCompiler& InOutCompiler, UPCG
 			break;
 		}
 
-		// Mapping from task ID & pin label to a virtual pin label. Compute graphs are executed within a generated element,
-		// and the input and output pins of this element must have unique virtual pin labels so that we can parse the data that
-		// PCG provides through the input data collection correctly, and route the output data to the downstream pins correctly.
-		FOriginalToVirtualPin OriginalToVirtualPin;
-
-		TMap<TSoftObjectPtr<const UPCGPin>, FName> OutputCPUPinToVirtualPin;
-
 		// Wire in the compute graph task, side by side with the individual GPU tasks, which will be culled below.
 		WireGPUGraphNode(
 			ComputeGraphTaskId,
 			NodeSubsetToConvertToCFGraph,
-			GPUCompatibleTaskIds,
 			InOutCompiledTasks,
 			TaskSuccessors,
-			OriginalToVirtualPin,
-			OutputCPUPinToVirtualPin);
+			OriginalToVirtualPin[ComputeGraphIndex],
+			OutputCPUPinToVirtualPin[ComputeGraphIndex]);
 
 		// Generate a compute graph from all of the individual GPU tasks.
 		BuildGPUGraphTask(
@@ -1201,10 +1258,11 @@ void FPCGGraphCompilerGPU::CreateGPUNodes(FPCGGraphCompiler& InOutCompiler, UPCG
 			InGridSize,
 			ComputeGraphTaskId,
 			NodeSubsetToConvertToCFGraph,
+			GPUCompatibleTaskIds,
 			TaskSuccessors,
 			InOutCompiledTasks,
-			OriginalToVirtualPin,
-			OutputCPUPinToVirtualPin);
+			OriginalToVirtualPin[ComputeGraphIndex],
+			OutputCPUPinToVirtualPin[ComputeGraphIndex]);
 	}
 
 	// Now cull all the GPU compatible nodes. The compute graph task are already wired in so we're fine to just delete.
diff --git a/Engine/Plugins/PCG/Source/PCG/Private/Graph/PCGGraphCompilerGPU.h b/Engine/Plugins/PCG/Source/PCG/Private/Graph/PCGGraphCompilerGPU.h
--- a/Engine/Plugins/PCG/Source/PCG/Private/Graph/PCGGraphCompilerGPU.h
+++ b/Engine/Plugins/PCG/Source/PCG/Private/Graph/PCGGraphCompilerGPU.h
@@ -37,16 +37,23 @@ public:
 	* using unique virtual input pin labels. */
 	static void CreateGatherTasksAtGPUInputs(UPCGGraph* InGraph, const TSet<FPCGTaskId>& InGPUCompatibleTaskIds, TArray<FPCGGraphTask>& InOutCompiledTasks);
 	
+	/** The compute graph element does not have a node or pins, so create maps of unique "virtual pins" that will allow us to marshal
+	 * data properly from an input data collection at execution time. */
+	static void SetupVirtualPins(
+		const TSet<FPCGTaskId>& InCollapsedTasks,
+		const TArray<FPCGGraphTask>& InCompiledTasks,
+		const TMap<FPCGTaskId, TArray<FPCGTaskId>>& InTaskSuccessors,
+		FOriginalToVirtualPin& OutOriginalToVirtualPin,
+		TMap<TSoftObjectPtr<const UPCGPin>, FName>& OutOutputCPUPinToVirtualPin);
+
 	/** Wires in a compute graph element alongside each set of GPU compatible nodes. The tasks for each node will be culled later. */
 	static void WireGPUGraphNode(
 		FPCGTaskId InGPUGraphTaskId,
 		const TSet<FPCGTaskId>& InCollapsedTasks,
-		const TSet<FPCGTaskId>& InGPUCompatibleTaskIds,
 		TArray<FPCGGraphTask>& InOutCompiledTasks,
-		const TMap<FPCGTaskId,
-		TArray<FPCGTaskId>>&InTaskSuccessors,
-		FOriginalToVirtualPin& OutOriginalToVirtualPin,
-		TMap<TSoftObjectPtr<const UPCGPin>, FName>& OutOutputCPUPinToVirtualPin);
+		const TMap<FPCGTaskId, TArray<FPCGTaskId>>& InTaskSuccessors,
+		const FOriginalToVirtualPin& InOriginalToVirtualPin,
+		const TMap<TSoftObjectPtr<const UPCGPin>, FName>& InOutputCPUPinToVirtualPin);
 	
 	/** Compiles a compute graph. */
 	static void BuildGPUGraphTask(
@@ -55,6 +62,7 @@ public:
 		uint32 InGridSize,
 		FPCGTaskId InGPUGraphTaskId,
 		const TSet<FPCGTaskId>& InCollapsedTasks,
+		const TSet<FPCGTaskId>& InAllGPUCompatibleTasks,
 		const FTaskToSuccessors& InTaskSuccessors,
 		TArray<FPCGGraphTask>& InOutCompiledTasks,
 		const FOriginalToVirtualPin& InOriginalToVirtualPin,
diff --git a/Engine/Plugins/PCG/Source/PCG/Private/PCGSettings.cpp b/Engine/Plugins/PCG/Source/PCG/Private/PCGSettings.cpp
--- a/Engine/Plugins/PCG/Source/PCG/Private/PCGSettings.cpp
+++ b/Engine/Plugins/PCG/Source/PCG/Private/PCGSettings.cpp
@@ -268,40 +268,38 @@ FPCGDataCollectionDesc UPCGSettings::ComputeInputPinDataDesc(const UPCGPin* Inpu
 		check(UpstreamSettings);
 
 		// Add data from connected upstream output pin.
-		PinDesc.Combine(UpstreamSettings->ComputeOutputPinDataDesc(UpstreamOutputPin, Binding));
+		FPCGDataCollectionDesc EdgeDesc;
+		if (ensure(UpstreamSettings->ComputeOutputPinDataDesc(UpstreamOutputPin, Binding, EdgeDesc)))
+		{
+			PinDesc.Combine(EdgeDesc);
+		}
 	}
 
 	return PinDesc;
 }
 
-FPCGDataCollectionDesc UPCGSettings::ComputeOutputPinDataDesc(const FName& OutputPinLabel, const UPCGDataBinding* Binding) const
+bool UPCGSettings::ComputeOutputPinDataDesc(const FName& OutputPinLabel, const UPCGDataBinding* InBinding, FPCGDataCollectionDesc& OutDesc) const
 {
 	const UPCGNode* Node = CastChecked<UPCGNode>(GetOuter());
-	return ComputeOutputPinDataDesc(Node->GetOutputPin(OutputPinLabel), Binding);
+	return ComputeOutputPinDataDesc(Node->GetOutputPin(OutputPinLabel), InBinding, OutDesc);
 }
 
-FPCGDataCollectionDesc UPCGSettings::ComputeOutputPinDataDesc(const UPCGPin* OutputPin, const UPCGDataBinding* Binding) const
+bool UPCGSettings::ComputeOutputPinDataDesc(const UPCGPin* OutputPin, const UPCGDataBinding* InBinding, FPCGDataCollectionDesc& OutDesc) const
 {
 	check(OutputPin);
-	check(Binding);
+	check(InBinding);
 
-	// This base class implementation will be called on upstream CPU nodes. Get the downstream pin alias in order to
-	// pick the relevant data items out of the compute graph element's input data collection.
-	if (const FName* FoundPinAlias = Binding->Graph->OutputCPUPinToInputGPUPinAlias.Find(OutputPin))
+	const bool bSuccess = InBinding->ComputeCPUOutputPinDataDesc(OutputPin, OutDesc);
+
+	if (!bSuccess)
 	{
-		return FPCGDataCollectionDesc::BuildFromInputDataCollectionAndInputPinLabel(
-			Binding->DataForGPU.InputDataCollection,
-			*FoundPinAlias,
-			Binding->GetAttributeLookupTable(),
-			Binding->GetStringTable());
+		ensureMsgf(false, TEXT("Gathering data from CPU output pin '%s' on node '%s' failed. Pin was not present in OutputCPUPinToInputGPUPinAlias map (%d entries)."),
+			*OutputPin->Properties.Label.ToString(),
+			OutputPin->Node ? *OutputPin->Node->GetNodeTitle(EPCGNodeTitleType::ListView).ToString() : TEXT("MISSING"),
+			InBinding->Graph->OutputCPUPinToInputGPUPinAlias.Num());
 	}
 
-	ensureMsgf(false, TEXT("Gathering data from CPU output pin '%s' on node '%s' failed. Pin was not present in OutputCPUPinToInputGPUPinAlias map (%d entries)."),
-		*OutputPin->Properties.Label.ToString(),
-		OutputPin->Node ? *OutputPin->Node->GetNodeTitle(EPCGNodeTitleType::ListView).ToString() : TEXT("MISSING"),
-		Binding->Graph->OutputCPUPinToInputGPUPinAlias.Num());
-
-	return {};
+	return bSuccess;
 }
 
 uint32 UPCGSettings::GetTypeNameHash() const
diff --git a/Engine/Plugins/PCG/Source/PCG/Public/Compute/Elements/PCGCustomHLSL.h b/Engine/Plugins/PCG/Source/PCG/Public/Compute/Elements/PCGCustomHLSL.h
--- a/Engine/Plugins/PCG/Source/PCG/Public/Compute/Elements/PCGCustomHLSL.h
+++ b/Engine/Plugins/PCG/Source/PCG/Public/Compute/Elements/PCGCustomHLSL.h
@@ -68,7 +68,7 @@ public:
 	virtual FString GetCookedKernelSource(const TMap<FName, FPCGKernelAttributeIDAndType>& GlobalAttributeLookupTable) const override;
 	virtual const TArray<FPCGKernelAttributeKey> GetKernelAttributeKeys() const override { return KernelAttributeKeys; }
 	virtual int ComputeKernelThreadCount(const UPCGDataBinding* Binding) const override;
-	virtual FPCGDataCollectionDesc ComputeOutputPinDataDesc(const UPCGPin* OutputPin, const UPCGDataBinding* Binding) const override;
+	virtual bool ComputeOutputPinDataDesc(const UPCGPin* OutputPin, const UPCGDataBinding* Binding, FPCGDataCollectionDesc& OutDesc) const override;
 
 protected:
 #if WITH_EDITOR
diff --git a/Engine/Plugins/PCG/Source/PCG/Public/Elements/PCGCopyPoints.h b/Engine/Plugins/PCG/Source/PCG/Public/Elements/PCGCopyPoints.h
--- a/Engine/Plugins/PCG/Source/PCG/Public/Elements/PCGCopyPoints.h
+++ b/Engine/Plugins/PCG/Source/PCG/Public/Elements/PCGCopyPoints.h
@@ -55,7 +55,7 @@ public:
 
 	virtual FString GetCookedKernelSource(const TMap<FName, FPCGKernelAttributeIDAndType>& GlobalAttributeLookupTable) const override;
 	virtual int ComputeKernelThreadCount(const UPCGDataBinding* Binding) const override;
-	virtual FPCGDataCollectionDesc ComputeOutputPinDataDesc(const UPCGPin* OutputPin, const UPCGDataBinding* Binding) const override;
+	virtual bool ComputeOutputPinDataDesc(const UPCGPin* OutputPin, const UPCGDataBinding* Binding, FPCGDataCollectionDesc& OutDesc) const override;
 	virtual void CreateAdditionalInputDataInterfaces(TArray<TObjectPtr<UComputeDataInterface>>& OutDataInterfaces) const override;
 
 protected:
diff --git a/Engine/Plugins/PCG/Source/PCG/Public/Elements/PCGReroute.h b/Engine/Plugins/PCG/Source/PCG/Public/Elements/PCGReroute.h
--- a/Engine/Plugins/PCG/Source/PCG/Public/Elements/PCGReroute.h
+++ b/Engine/Plugins/PCG/Source/PCG/Public/Elements/PCGReroute.h
@@ -40,7 +40,7 @@ protected:
 	virtual TArray<FPCGPinProperties> InputPinProperties() const override;
 	virtual TArray<FPCGPinProperties> OutputPinProperties() const override;
 	virtual FPCGElementPtr CreateElement() const override;
-	virtual FPCGDataCollectionDesc ComputeOutputPinDataDesc(const UPCGPin* OutputPin, const UPCGDataBinding* Binding) const override;
+	virtual bool ComputeOutputPinDataDesc(const UPCGPin* OutputPin, const UPCGDataBinding* Binding, FPCGDataCollectionDesc& OutDesc) const override;
 	//~End UPCGSettings interface
 };
 
diff --git a/Engine/Plugins/PCG/Source/PCG/Public/Elements/PCGStaticMeshSpawner.h b/Engine/Plugins/PCG/Source/PCG/Public/Elements/PCGStaticMeshSpawner.h
--- a/Engine/Plugins/PCG/Source/PCG/Public/Elements/PCGStaticMeshSpawner.h
+++ b/Engine/Plugins/PCG/Source/PCG/Public/Elements/PCGStaticMeshSpawner.h
@@ -34,7 +34,7 @@ public:
 	virtual const TArray<FPCGKernelAttributeKey> GetKernelAttributeKeys() const override;
 	virtual void AddStaticCreatedStrings(TArray<FString>& InOutStringTable) const override;
 	virtual int ComputeKernelThreadCount(const UPCGDataBinding* Binding) const override;
-	virtual FPCGDataCollectionDesc ComputeOutputPinDataDesc(const UPCGPin* OutputPin, const UPCGDataBinding* Binding) const override;
+	virtual bool ComputeOutputPinDataDesc(const UPCGPin* OutputPin, const UPCGDataBinding* Binding, FPCGDataCollectionDesc& OutDesc) const override;
 	virtual void CreateAdditionalInputDataInterfaces(TArray<TObjectPtr<UComputeDataInterface>>& OutDataInterfaces) const override;
 	virtual void CreateAdditionalOutputDataInterfaces(TArray<TObjectPtr<UComputeDataInterface>>& OutDataInterfaces) const override;
 #if WITH_EDITOR
diff --git a/Engine/Plugins/PCG/Source/PCG/Public/PCGSettings.h b/Engine/Plugins/PCG/Source/PCG/Public/PCGSettings.h
--- a/Engine/Plugins/PCG/Source/PCG/Public/PCGSettings.h
+++ b/Engine/Plugins/PCG/Source/PCG/Public/PCGSettings.h
@@ -549,8 +549,8 @@ public:
 	FPCGDataCollectionDesc ComputeInputPinDataDesc(const UPCGPin* InputPin, const UPCGDataBinding* Binding) const;
 
 	/** Compute a description of data that will be output from OutputPinLabel/OutputPin. */
-	FPCGDataCollectionDesc ComputeOutputPinDataDesc(const FName& OutputPinLabel, const UPCGDataBinding* Binding) const;
-	virtual FPCGDataCollectionDesc ComputeOutputPinDataDesc(const UPCGPin* OutputPin, const UPCGDataBinding* Binding) const;
+	bool ComputeOutputPinDataDesc(const FName& OutputPinLabel, const UPCGDataBinding* InBinding, FPCGDataCollectionDesc& OutDesc) const;
+	virtual bool ComputeOutputPinDataDesc(const UPCGPin* OutputPin, const UPCGDataBinding* InBinding, FPCGDataCollectionDesc& OutDesc) const;
 
 	/** Create additional data interfaces to marshal any data required by this settings. */
 	virtual void CreateAdditionalInputDataInterfaces(TArray<TObjectPtr<UComputeDataInterface>>& OutDataInterfaces) const {}
