[PCG] Make GPU compilation code thread safe so that it works with Graph Executor v2 & multi-threaded execution.

#rnx

diff --git a/Engine/Plugins/PCG/Source/PCG/Private/Elements/PCGCopyPoints.cpp b/Engine/Plugins/PCG/Source/PCG/Private/Elements/PCGCopyPoints.cpp
--- a/Engine/Plugins/PCG/Source/PCG/Private/Elements/PCGCopyPoints.cpp
+++ b/Engine/Plugins/PCG/Source/PCG/Private/Elements/PCGCopyPoints.cpp
@@ -8,6 +8,7 @@
 #include "Compute/PCGDataBinding.h"
 #include "Data/PCGPointData.h"
 #include "Data/PCGSpatialData.h"
+#include "Graph/PCGGPUGraphCompilationContext.h"
 #include "Helpers/PCGAsync.h"
 #include "Helpers/PCGHelpers.h"
 
@@ -124,14 +125,16 @@ int UPCGCopyPointsSettings::ComputeKernelThreadCount(const UPCGDataBinding* Bind
 	return ThreadCount;
 }
 
-void UPCGCopyPointsSettings::CreateAdditionalInputDataInterfaces(TArray<TObjectPtr<UComputeDataInterface>>& OutDataInterfaces) const
+#if WITH_EDITOR
+void UPCGCopyPointsSettings::CreateAdditionalInputDataInterfaces(FPCGGPUCompilationContext& InOutContext, UObject* InObjectOuter, TArray<TObjectPtr<UComputeDataInterface>>& OutDataInterfaces) const
 {
-	Super::CreateAdditionalInputDataInterfaces(OutDataInterfaces);
+	Super::CreateAdditionalInputDataInterfaces(InOutContext, InObjectOuter, OutDataInterfaces);
 
-	TObjectPtr<UPCGCopyPointsDataInterface> NodeDI = NewObject<UPCGCopyPointsDataInterface>();
+	TObjectPtr<UPCGCopyPointsDataInterface> NodeDI = InOutContext.NewObject_AnyThread<UPCGCopyPointsDataInterface>(InObjectOuter);
 	NodeDI->Settings = this;
 	OutDataInterfaces.Add(NodeDI);
 }
+#endif
 
 TArray<FPCGPinProperties> UPCGCopyPointsSettings::InputPinProperties() const
 {
diff --git a/Engine/Plugins/PCG/Source/PCG/Private/Elements/PCGStaticMeshSpawner.cpp b/Engine/Plugins/PCG/Source/PCG/Private/Elements/PCGStaticMeshSpawner.cpp
--- a/Engine/Plugins/PCG/Source/PCG/Private/Elements/PCGStaticMeshSpawner.cpp
+++ b/Engine/Plugins/PCG/Source/PCG/Private/Elements/PCGStaticMeshSpawner.cpp
@@ -11,6 +11,7 @@
 #include "Compute/DataInterfaces/Elements/PCGStaticMeshSpawnerDataInterface.h"
 #include "Data/PCGPointData.h"
 #include "Data/PCGSpatialData.h"
+#include "Graph/PCGGPUGraphCompilationContext.h"
 #include "Elements/PCGStaticMeshSpawnerContext.h"
 #include "Helpers/PCGActorHelpers.h"
 #include "Helpers/PCGHelpers.h"
@@ -228,26 +229,26 @@ int UPCGStaticMeshSpawnerSettings::ComputeKernelThreadCount(const UPCGDataBindin
 	return ComputeInputPinDataDesc(Pin, Binding).ComputeDataElementCount(EPCGDataType::Point);
 }
 
-void UPCGStaticMeshSpawnerSettings::CreateAdditionalInputDataInterfaces(TArray<TObjectPtr<UComputeDataInterface>>& OutDataInterfaces) const
+#if WITH_EDITOR
+void UPCGStaticMeshSpawnerSettings::CreateAdditionalInputDataInterfaces(FPCGGPUCompilationContext& InOutContext, UObject* InObjectOuter, TArray<TObjectPtr<UComputeDataInterface>>& OutDataInterfaces) const
 {
-	Super::CreateAdditionalInputDataInterfaces(OutDataInterfaces);
+	Super::CreateAdditionalInputDataInterfaces(InOutContext, InObjectOuter, OutDataInterfaces);
 
-	TObjectPtr<UPCGStaticMeshSpawnerDataInterface> NodeDI = NewObject<UPCGStaticMeshSpawnerDataInterface>();
+	TObjectPtr<UPCGStaticMeshSpawnerDataInterface> NodeDI = InOutContext.NewObject_AnyThread<UPCGStaticMeshSpawnerDataInterface>(InObjectOuter);
 	NodeDI->Settings = this;
 	OutDataInterfaces.Add(NodeDI);
 }
 
-void UPCGStaticMeshSpawnerSettings::CreateAdditionalOutputDataInterfaces(TArray<TObjectPtr<UComputeDataInterface>>& OutDataInterfaces) const
+void UPCGStaticMeshSpawnerSettings::CreateAdditionalOutputDataInterfaces(FPCGGPUCompilationContext& InOutContext, UObject* InObjectOuter, TArray<TObjectPtr<UComputeDataInterface>>& OutDataInterfaces) const
 {
-	Super::CreateAdditionalOutputDataInterfaces(OutDataInterfaces);
+	Super::CreateAdditionalOutputDataInterfaces(InOutContext, InObjectOuter, OutDataInterfaces);
 
-	UPCGInstanceDataInterface* InstanceDI = NewObject<UPCGInstanceDataInterface>();
+	UPCGInstanceDataInterface* InstanceDI = InOutContext.NewObject_AnyThread<UPCGInstanceDataInterface>(InObjectOuter);
 	InstanceDI->SetProducerSettings(this);
 	InstanceDI->InputPinProvidingData = PCGPinConstants::DefaultInputLabel;
 	OutDataInterfaces.Add(InstanceDI);
 }
 
-#if WITH_EDITOR
 FText UPCGStaticMeshSpawnerSettings::GetDefaultNodeTitle() const
 {
 	return LOCTEXT("NodeTitle", "Static Mesh Spawner");
diff --git a/Engine/Plugins/PCG/Source/PCG/Private/Graph/PCGGPUGraphCompilationContext.h b/Engine/Plugins/PCG/Source/PCG/Private/Graph/PCGGPUGraphCompilationContext.h
new file mode 100644
--- /dev/null
+++ b/Engine/Plugins/PCG/Source/PCG/Private/Graph/PCGGPUGraphCompilationContext.h
@@ -0,0 +1,73 @@
+// Copyright Epic Games, Inc. All Rights Reserved.
+
+#pragma once
+
+#if WITH_EDITOR
+
+#include "PCGCommon.h"
+
+class FPCGGraphCompiler;
+
+/** Current context/state used during compilation of GPU graphs. */
+struct FPCGGPUCompilationContext
+{
+	FPCGGPUCompilationContext(FPCGGraphCompiler& InGraphCompiler)
+		: GraphCompiler(InGraphCompiler)
+	{
+	}
+
+	~FPCGGPUCompilationContext()
+	{
+		// Remove Async Flags from objects created during compilation so that they can get tracked by the GarbageCollector properly.
+		for (TObjectPtr<UObject>& AsyncObject : AsyncObjects)
+		{
+			if (ensure(AsyncObject->HasAnyInternalFlags(EInternalObjectFlags::Async)))
+			{
+				AsyncObject->ClearInternalFlags(EInternalObjectFlags::Async);
+				ForEachObjectWithOuter(
+					AsyncObject,
+					[&AsyncObjects = AsyncObjects](UObject* SubObject)
+					{
+						if (AsyncObjects.Contains(SubObject))
+						{
+							return;
+						}
+
+						SubObject->ClearInternalFlags(EInternalObjectFlags::Async);
+					},
+					true);
+			}
+		}
+	}
+
+	/** Thread safe NewObject. If object created off game thread, object added to AsyncObjects. */
+	template<class T, typename... Args>
+	T* NewObject_AnyThread(Args&&... InArgs)
+	{
+		if (!IsInGameThread())
+		{
+			T* Object = nullptr;
+			{
+				FGCScopeGuard Scope;
+				Object = ::NewObject<T>(std::forward<Args>(InArgs)...);
+			}
+
+			check(Object);
+			AsyncObjects.Add(Object);
+			return Object;
+		}
+
+		return ::NewObject<T>(std::forward<Args>(InArgs)...);
+	}
+
+	FPCGGraphCompiler& GetGraphCompiler() { return GraphCompiler; }
+
+private:
+	// List of objects created during compilation, we need to track them so we can remove their Async flags when compilation is done
+	// so that they can be considered as existing on the main thread (and get properly GCed).
+	TSet<TObjectPtr<UObject>> AsyncObjects;
+
+	FPCGGraphCompiler& GraphCompiler;
+};
+
+#endif // WITH_EDITOR
diff --git a/Engine/Plugins/PCG/Source/PCG/Private/Graph/PCGGraphCompilerGPU.cpp b/Engine/Plugins/PCG/Source/PCG/Private/Graph/PCGGraphCompilerGPU.cpp
--- a/Engine/Plugins/PCG/Source/PCG/Private/Graph/PCGGraphCompilerGPU.cpp
+++ b/Engine/Plugins/PCG/Source/PCG/Private/Graph/PCGGraphCompilerGPU.cpp
@@ -15,6 +15,7 @@
 #include "Compute/DataInterfaces/PCGTextureDataInterface.h"
 #include "Compute/Elements/PCGComputeGraphElement.h"
 #include "Elements/PCGStaticMeshSpawner.h"
+#include "Graph/PCGGPUGraphCompilationContext.h"
 #include "Graph/PCGGraphCompiler.h"
 #include "Graph/PCGGraphExecutor.h"
 
@@ -510,7 +511,7 @@ void FPCGGraphCompilerGPU::WireGPUGraphNode(
 }
 
 void FPCGGraphCompilerGPU::BuildGPUGraphTask(
-	FPCGGraphCompiler& InOutCompiler,
+	FPCGGPUCompilationContext& InOutContext,
 	UPCGGraph* InGraph,
 	uint32 InGridSize,
 	FPCGTaskId InGPUGraphTaskId,
@@ -526,8 +527,9 @@ void FPCGGraphCompilerGPU::BuildGPUGraphTask(
 	check(InGraph);
 	const FName GraphName = MakeUniqueObjectName(InGraph, UPCGComputeGraph::StaticClass(), InGraph->GetFName());
 
-	UObject* ComputeGraphOuter = InOutCompiler.IsCooking() ? InGraph : Cast<UObject>(GetTransientPackage());
-	UPCGComputeGraph* ComputeGraph = NewObject<UPCGComputeGraph>(ComputeGraphOuter, GraphName);
+	UObject* ComputeGraphOuter = InOutContext.GetGraphCompiler().IsCooking() ? InGraph : Cast<UObject>(GetTransientPackage());
+
+	UPCGComputeGraph* ComputeGraph = InOutContext.NewObject_AnyThread<UPCGComputeGraph>(ComputeGraphOuter, GraphName);
 	ComputeGraph->OutputCPUPinToInputGPUPinAlias = InOutputCPUPinToVirtualPin;
 
 	// Not incredibly useful for us - DG adds GetComponentSource()->GetComponentClass() object which allows it to bind at execution time by class.
@@ -536,7 +538,7 @@ void FPCGGraphCompilerGPU::BuildGPUGraphTask(
 
 	// Create data interfaces which allow kernels to read or write data. Each data interface is associated with a node output pin.
 	// For CPU->GPU edges, an upload data interface is created. For GPU->CPU edges, a readback data interface is created.
-	auto CreateDataInterface = [&InCollapsedTasks, &InOutCompiledTasks, ComputeGraph](FPCGTaskId InTaskId, bool bRequiresReadback, const FPCGPinProperties& InOutputPinProperties) -> UPCGComputeDataInterface*
+	auto CreateDataInterface = [&InOutContext, &InCollapsedTasks, &InOutCompiledTasks, ComputeGraph](FPCGTaskId InTaskId, bool bRequiresReadback, const FPCGPinProperties& InOutputPinProperties) -> UPCGComputeDataInterface*
 	{
 		TRACE_CPUPROFILER_EVENT_SCOPE(CreateDataInterface);
 
@@ -564,12 +566,12 @@ void FPCGGraphCompilerGPU::BuildGPUGraphTask(
 			if (bUpstreamIsGPUTask)
 			{
 				// Provides data for GPU -> GPU and GPU -> CPU edges.
-				DataInterfacePCGData = NewObject<UPCGDataCollectionDataInterface>(ComputeGraph);
+				DataInterfacePCGData = InOutContext.NewObject_AnyThread<UPCGDataCollectionDataInterface>(ComputeGraph);
 			}
 			else
 			{
 				// Provides data for CPU -> GPU edge.
-				DataInterfacePCGData = NewObject<UPCGDataCollectionUploadDataInterface>(ComputeGraph);
+				DataInterfacePCGData = InOutContext.NewObject_AnyThread<UPCGDataCollectionUploadDataInterface>(ComputeGraph);
 			}
 
 			check(DataInterfacePCGData);
@@ -581,12 +583,12 @@ void FPCGGraphCompilerGPU::BuildGPUGraphTask(
 		}
 		case EPCGDataType::Texture:
 		{
-			DataInterface = NewObject<UPCGTextureDataInterface>(ComputeGraph);
+			DataInterface = InOutContext.NewObject_AnyThread<UPCGTextureDataInterface>(ComputeGraph);
 			break;
 		}
 		case EPCGDataType::Landscape:
 		{
-			DataInterface = NewObject<UPCGLandscapeDataInterface>(ComputeGraph);
+			DataInterface = InOutContext.NewObject_AnyThread<UPCGLandscapeDataInterface>(ComputeGraph);
 			break;
 		}
 		default:
@@ -871,7 +873,7 @@ void FPCGGraphCompilerGPU::BuildGPUGraphTask(
 			}
 		}
 
-		UPCGCustomKernelDataInterface* KernelDI = NewObject<UPCGCustomKernelDataInterface>(ComputeGraph);
+		UPCGCustomKernelDataInterface* KernelDI = InOutContext.NewObject_AnyThread<UPCGCustomKernelDataInterface>(ComputeGraph);
 		KernelDI->Settings = Settings;
 
 		const int32 KernelDIIndex = ComputeGraph->DataInterfaces.Num();
@@ -879,16 +881,13 @@ void FPCGGraphCompilerGPU::BuildGPUGraphTask(
 		InputDataInterfaceIndexAndPin.Emplace(KernelDIIndex, NAME_None);
 
 		TArray<TObjectPtr<UComputeDataInterface>> AdditionalInputDIs, AdditionalOutputDIs;
-		Settings->CreateAdditionalInputDataInterfaces(AdditionalInputDIs);
-		Settings->CreateAdditionalOutputDataInterfaces(AdditionalOutputDIs);
+		Settings->CreateAdditionalInputDataInterfaces(InOutContext, ComputeGraph, AdditionalInputDIs);
+		Settings->CreateAdditionalOutputDataInterfaces(InOutContext, ComputeGraph, AdditionalOutputDIs);
 
 		auto AddAdditionalDataInterface = [ComputeGraph](TObjectPtr<UComputeDataInterface> DataInterface) -> int32
 		{
 			if (ensure(DataInterface))
 			{
-				// Re-outer to the ComputeGraph
-				DataInterface->Rename(/*Name=*/nullptr, ComputeGraph);
-
 				const int32 DataInterfaceIndex = ComputeGraph->DataInterfaces.Num();
 				ComputeGraph->DataInterfaces.Add(DataInterface);
 
@@ -921,7 +920,7 @@ void FPCGGraphCompilerGPU::BuildGPUGraphTask(
 		// TODO once we support cooking for different platforms/configs, don't create the interface if logging is not present.
 		if (Settings->bPrintShaderDebugValues)
 		{
-			UPCGDebugDataInterface* DebugDI = NewObject<UPCGDebugDataInterface>(ComputeGraph);
+			UPCGDebugDataInterface* DebugDI = InOutContext.NewObject_AnyThread<UPCGDebugDataInterface>(ComputeGraph);
 			DebugDI->SetDebugBufferSize(Settings->DebugBufferSize);
 
 			const int32 DebugDIIndex = ComputeGraph->DataInterfaces.Num();
@@ -951,7 +950,7 @@ void FPCGGraphCompilerGPU::BuildGPUGraphTask(
 
 		FKernelWithDataBindings KernelWithBindings;
 
-		KernelWithBindings.Kernel = NewObject<UComputeKernel>(ComputeGraph);
+		KernelWithBindings.Kernel = InOutContext.NewObject_AnyThread<UComputeKernel>(ComputeGraph);
 		const int KernelIndex = ComputeGraph->KernelInvocations.Num();
 		ComputeGraph->KernelInvocations.Add(KernelWithBindings.Kernel);
 		ComputeGraph->KernelToNode.Add(Node);
@@ -1000,7 +999,7 @@ void FPCGGraphCompilerGPU::BuildGPUGraphTask(
 		}
 
 		{
-			UPCGComputeKernelSource* KernelSource = NewObject<UPCGComputeKernelSource>(KernelWithBindings.Kernel); // is outer to kernel fine?
+			UPCGComputeKernelSource* KernelSource = InOutContext.NewObject_AnyThread<UPCGComputeKernelSource>(KernelWithBindings.Kernel); // is outer to kernel fine?
 			KernelWithBindings.Kernel->KernelSource = KernelSource;
 			KernelSource->EntryPoint = Settings->GetKernelEntryPoint();
 			KernelSource->GroupSize = Settings->GetThreadGroupSize();
@@ -1124,12 +1123,12 @@ void FPCGGraphCompilerGPU::BuildGPUGraphTask(
 		}
 	}
 
-	TMap<uint32, TArray<TObjectPtr<UPCGComputeGraph>>>& GridToComputeGraphs = InOutCompiler.GetCache().TopGraphToComputeGraphMap.FindOrAdd(InGraph);
+	TMap<uint32, TArray<TObjectPtr<UPCGComputeGraph>>>& GridToComputeGraphs = InOutContext.GetGraphCompiler().GetCache().TopGraphToComputeGraphMap.FindOrAdd(InGraph);
 	TArray<TObjectPtr<UPCGComputeGraph>>& ComputeGraphs = GridToComputeGraphs.FindOrAdd(InGridSize);
 	const uint32 ComputeGraphIndex = ComputeGraphs.Num();
 	ComputeGraphs.Add(ComputeGraph);
 
-	if (InOutCompiler.IsCooking())
+	if (InOutContext.GetGraphCompiler().IsCooking())
 	{
 		TObjectPtr<UPCGComputeGraphSettings> Settings = NewObject<UPCGComputeGraphSettings>(InGraph);
 		Settings->ComputeGraphIndex = ComputeGraphIndex;
@@ -1154,6 +1153,8 @@ void FPCGGraphCompilerGPU::CreateGPUNodes(FPCGGraphCompiler& InOutCompiler, UPCG
 {
 	TRACE_CPUPROFILER_EVENT_SCOPE(FPCGGraphCompilerGPU::CreateGPUNodes);
 
+	FPCGGPUCompilationContext Context(InOutCompiler);
+
 	// Clear out any previous compute graphs for this grid level before adding new ones.
 	if (ensure(InGraph))
 	{
@@ -1253,7 +1254,7 @@ void FPCGGraphCompilerGPU::CreateGPUNodes(FPCGGraphCompiler& InOutCompiler, UPCG
 
 		// Generate a compute graph from all of the individual GPU tasks.
 		BuildGPUGraphTask(
-			InOutCompiler,
+			Context,
 			InGraph,
 			InGridSize,
 			ComputeGraphTaskId,
diff --git a/Engine/Plugins/PCG/Source/PCG/Private/Graph/PCGGraphCompilerGPU.h b/Engine/Plugins/PCG/Source/PCG/Private/Graph/PCGGraphCompilerGPU.h
--- a/Engine/Plugins/PCG/Source/PCG/Private/Graph/PCGGraphCompilerGPU.h
+++ b/Engine/Plugins/PCG/Source/PCG/Private/Graph/PCGGraphCompilerGPU.h
@@ -8,11 +8,16 @@ class FPCGGraphCompiler;
 class UPCGGraph;
 class UPCGPin;
 struct FPCGGraphTask;
+struct FPCGGPUCompilationContext;
 
 #if WITH_EDITOR
 class FPCGGraphCompilerGPU
 {
 public:
+	/** Finds connected subgraphs of GPU - enabled nodes that can be dispatched together and replaces each one with a compute graph. */
+	static void CreateGPUNodes(FPCGGraphCompiler& InOutCompiler, UPCGGraph* InGraph, uint32 InGridSize, TArray<FPCGGraphTask>& InOutCompiledTasks);
+
+private:
 	/** Used to track new unique virtual pins created on generated compute graph elements. */
 	using FNodePin = TTuple<FPCGTaskId, /*Pin label*/FName, /*Pin is input*/bool>;
 	using FOriginalToVirtualPin = TMap<FNodePin, /*Virtual pin label*/FName>;
@@ -57,7 +62,7 @@ public:
 	
 	/** Compiles a compute graph. */
 	static void BuildGPUGraphTask(
-		FPCGGraphCompiler& InOutCompiler,
+		FPCGGPUCompilationContext& InOutContext,
 		UPCGGraph* InGraph,
 		uint32 InGridSize,
 		FPCGTaskId InGPUGraphTaskId,
@@ -68,7 +73,5 @@ public:
 		const FOriginalToVirtualPin& InOriginalToVirtualPin,
 		const TMap<TSoftObjectPtr<const UPCGPin>, FName>& InOutputCPUPinToVirtualPin);
 	
-	/** Finds connected subgraphs of GPU - enabled nodes that can be dispatched together and replaces each one with a compute graph. */
-	static void CreateGPUNodes(FPCGGraphCompiler& InOutCompiler, UPCGGraph* InGraph, uint32 InGridSize, TArray<FPCGGraphTask>& InOutCompiledTasks);
 };
 #endif // WITH_EDITOR
diff --git a/Engine/Plugins/PCG/Source/PCG/Public/Elements/PCGCopyPoints.h b/Engine/Plugins/PCG/Source/PCG/Public/Elements/PCGCopyPoints.h
--- a/Engine/Plugins/PCG/Source/PCG/Public/Elements/PCGCopyPoints.h
+++ b/Engine/Plugins/PCG/Source/PCG/Public/Elements/PCGCopyPoints.h
@@ -51,12 +51,12 @@ public:
 	virtual FText GetNodeTooltipText() const override;
 	virtual EPCGSettingsType GetType() const override { return EPCGSettingsType::Sampler; }
 	virtual bool DisplayExecuteOnGPUSetting() const override { return true; }
+	virtual void CreateAdditionalInputDataInterfaces(FPCGGPUCompilationContext& InOutContext, UObject* InObjectOuter, TArray<TObjectPtr<UComputeDataInterface>>& OutDataInterfaces) const override;
 #endif
 
 	virtual FString GetCookedKernelSource(const TMap<FName, FPCGKernelAttributeIDAndType>& GlobalAttributeLookupTable) const override;
 	virtual int ComputeKernelThreadCount(const UPCGDataBinding* Binding) const override;
 	virtual bool ComputeOutputPinDataDesc(const UPCGPin* OutputPin, const UPCGDataBinding* Binding, FPCGDataCollectionDesc& OutDesc) const override;
-	virtual void CreateAdditionalInputDataInterfaces(TArray<TObjectPtr<UComputeDataInterface>>& OutDataInterfaces) const override;
 
 protected:
 	virtual TArray<FPCGPinProperties> InputPinProperties() const override;
diff --git a/Engine/Plugins/PCG/Source/PCG/Public/Elements/PCGStaticMeshSpawner.h b/Engine/Plugins/PCG/Source/PCG/Public/Elements/PCGStaticMeshSpawner.h
--- a/Engine/Plugins/PCG/Source/PCG/Public/Elements/PCGStaticMeshSpawner.h
+++ b/Engine/Plugins/PCG/Source/PCG/Public/Elements/PCGStaticMeshSpawner.h
@@ -35,9 +35,9 @@ public:
 	virtual void AddStaticCreatedStrings(TArray<FString>& InOutStringTable) const override;
 	virtual int ComputeKernelThreadCount(const UPCGDataBinding* Binding) const override;
 	virtual bool ComputeOutputPinDataDesc(const UPCGPin* OutputPin, const UPCGDataBinding* Binding, FPCGDataCollectionDesc& OutDesc) const override;
-	virtual void CreateAdditionalInputDataInterfaces(TArray<TObjectPtr<UComputeDataInterface>>& OutDataInterfaces) const override;
-	virtual void CreateAdditionalOutputDataInterfaces(TArray<TObjectPtr<UComputeDataInterface>>& OutDataInterfaces) const override;
 #if WITH_EDITOR
+	virtual void CreateAdditionalInputDataInterfaces(FPCGGPUCompilationContext& InOutContext, UObject* InObjectOuter, TArray<TObjectPtr<UComputeDataInterface>>& OutDataInterfaces) const override;
+	virtual void CreateAdditionalOutputDataInterfaces(FPCGGPUCompilationContext& InOutContext, UObject* InObjectOuter, TArray<TObjectPtr<UComputeDataInterface>>& OutDataInterfaces) const override;
 	virtual bool DisplayExecuteOnGPUSetting() const override { return true; }
 	virtual FName GetDefaultNodeName() const override { return FName(TEXT("StaticMeshSpawner")); }
 	virtual FText GetDefaultNodeTitle() const override;
diff --git a/Engine/Plugins/PCG/Source/PCG/Public/PCGSettings.h b/Engine/Plugins/PCG/Source/PCG/Public/PCGSettings.h
--- a/Engine/Plugins/PCG/Source/PCG/Public/PCGSettings.h
+++ b/Engine/Plugins/PCG/Source/PCG/Public/PCGSettings.h
@@ -16,6 +16,7 @@
 
 class UPCGComponent;
 class UPCGPin;
+struct FPCGGPUCompilationContext;
 struct FPCGPinProperties;
 struct FPropertyChangedEvent;
 
@@ -552,11 +553,11 @@ public:
 	bool ComputeOutputPinDataDesc(const FName& OutputPinLabel, const UPCGDataBinding* InBinding, FPCGDataCollectionDesc& OutDesc) const;
 	virtual bool ComputeOutputPinDataDesc(const UPCGPin* OutputPin, const UPCGDataBinding* InBinding, FPCGDataCollectionDesc& OutDesc) const;
 
+#if WITH_EDITOR
 	/** Create additional data interfaces to marshal any data required by this settings. */
-	virtual void CreateAdditionalInputDataInterfaces(TArray<TObjectPtr<UComputeDataInterface>>& OutDataInterfaces) const {}
-	virtual void CreateAdditionalOutputDataInterfaces(TArray<TObjectPtr<UComputeDataInterface>>& OutDataInterfaces) const {}
+	virtual void CreateAdditionalInputDataInterfaces(FPCGGPUCompilationContext& InOutContext, UObject* InObjectOuter, TArray<TObjectPtr<UComputeDataInterface>>& OutDataInterfaces) const {}
+	virtual void CreateAdditionalOutputDataInterfaces(FPCGGPUCompilationContext& InOutContext, UObject* InObjectOuter, TArray<TObjectPtr<UComputeDataInterface>>& OutDataInterfaces) const {}
 
-#if WITH_EDITOR
 	/** Whether to display GPU execution option in node settings UI. */
 	UFUNCTION()
 	virtual bool DisplayExecuteOnGPUSetting() const { return false; }
