Replaced CastChecked usage to prevent crashes when disabling upstream nodes

#rnx

diff --git a/Engine/Plugins/PCG/Source/PCG/Private/Elements/PCGAttractElement.cpp b/Engine/Plugins/PCG/Source/PCG/Private/Elements/PCGAttractElement.cpp
--- a/Engine/Plugins/PCG/Source/PCG/Private/Elements/PCGAttractElement.cpp
+++ b/Engine/Plugins/PCG/Source/PCG/Private/Elements/PCGAttractElement.cpp
@@ -324,14 +324,19 @@ bool FPCGAttractElement::PrepareDataInternal(FPCGContext* InContext) const
 
 	Context->InitializePerIterationStates(SourceInputs.Num(), [&SourceInputs, &TargetInputs, Settings, Context](IterStateType& OutState, const ExecStateType&, const uint32 IterationIndex)
 	{
-		const UPCGPointData* SourcePointData = CastChecked<UPCGPointData>(SourceInputs[IterationIndex].Data);
+		const UPCGPointData* SourcePointData = Cast<UPCGPointData>(SourceInputs[IterationIndex].Data);
 
 		if (!SourcePointData || SourcePointData->IsEmpty())
 		{
 			return EPCGTimeSliceInitResult::NoOperation;
 		}
 
-		const UPCGPointData* TargetPointData = CastChecked<UPCGPointData>(TargetInputs[IterationIndex % TargetInputs.Num()].Data);
+		const UPCGPointData* TargetPointData = Cast<UPCGPointData>(TargetInputs[IterationIndex % TargetInputs.Num()].Data);
+
+		if (!TargetPointData)
+		{
+			return EPCGTimeSliceInitResult::NoOperation;
+		}
 
 		UPCGPointData* OutPointData = FPCGContext::NewObject_AnyThread<UPCGPointData>(Context);
 		OutPointData->InitializeFromData(SourcePointData);
diff --git a/Engine/Plugins/PCG/Source/PCG/Private/Elements/PCGClusterElement.cpp b/Engine/Plugins/PCG/Source/PCG/Private/Elements/PCGClusterElement.cpp
--- a/Engine/Plugins/PCG/Source/PCG/Private/Elements/PCGClusterElement.cpp
+++ b/Engine/Plugins/PCG/Source/PCG/Private/Elements/PCGClusterElement.cpp
@@ -483,7 +483,7 @@ bool FPCGClusterElement::PrepareDataInternal(FPCGContext* InContext) const
 	{
 		TRACE_CPUPROFILER_EVENT_SCOPE(FPCGClusterElement::PrepareData::InitializeIteration);
 
-		const UPCGPointData* InputPointData = CastChecked<UPCGPointData>(PointInputs[IterationIndex].Data);
+		const UPCGPointData* InputPointData = Cast<UPCGPointData>(PointInputs[IterationIndex].Data);
 		if (!InputPointData || InputPointData->IsEmpty())
 		{
 			return EPCGTimeSliceInitResult::NoOperation;
diff --git a/Engine/Plugins/PCG/Source/PCG/Private/Elements/PCGCollapsePoints.cpp b/Engine/Plugins/PCG/Source/PCG/Private/Elements/PCGCollapsePoints.cpp
--- a/Engine/Plugins/PCG/Source/PCG/Private/Elements/PCGCollapsePoints.cpp
+++ b/Engine/Plugins/PCG/Source/PCG/Private/Elements/PCGCollapsePoints.cpp
@@ -231,7 +231,7 @@ bool FPCGCollapsePointsElement::PrepareDataInternal(FPCGContext* InContext) cons
 
 	Context->InitializePerIterationStates(Inputs.Num(), [&Inputs, Settings, Context](IterStateType& OutState, const ExecStateType&, const uint32 IterationIndex)
 	{
-		const UPCGPointData* PointData = CastChecked<UPCGPointData>(Inputs[IterationIndex].Data);
+		const UPCGPointData* PointData = Cast<UPCGPointData>(Inputs[IterationIndex].Data);
 
 		if (!PointData || PointData->IsEmpty())
 		{
diff --git a/Engine/Plugins/PCG/Source/PCG/Private/Elements/PCGGetActorProperty.cpp b/Engine/Plugins/PCG/Source/PCG/Private/Elements/PCGGetActorProperty.cpp
--- a/Engine/Plugins/PCG/Source/PCG/Private/Elements/PCGGetActorProperty.cpp
+++ b/Engine/Plugins/PCG/Source/PCG/Private/Elements/PCGGetActorProperty.cpp
@@ -135,6 +135,18 @@ FPCGElementPtr UPCGGetActorPropertySettings::CreateElement() const
 	return MakeShared<FPCGGetActorPropertyElement>();
 }
 
+bool FPCGGetActorPropertyElement::IsCacheable(const UPCGSettings* InSettings) const
+{
+	if (const UPCGGetActorPropertySettings* Settings = Cast<UPCGGetActorPropertySettings>(InSettings))
+	{
+		return Settings->bAlwaysRequeryActors;
+	}
+	else
+	{
+		return false;
+	}
+}
+
 bool FPCGGetActorPropertyElement::PrepareDataInternal(FPCGContext* Context) const
 {
 	TRACE_CPUPROFILER_EVENT_SCOPE(FPCGGetActorPropertyElement::PrepareData);
diff --git a/Engine/Plugins/PCG/Source/PCG/Private/Graph/PCGGraphExecutor.cpp b/Engine/Plugins/PCG/Source/PCG/Private/Graph/PCGGraphExecutor.cpp
--- a/Engine/Plugins/PCG/Source/PCG/Private/Graph/PCGGraphExecutor.cpp
+++ b/Engine/Plugins/PCG/Source/PCG/Private/Graph/PCGGraphExecutor.cpp
@@ -2192,24 +2192,26 @@ void FPCGGraphExecutor::CombineParams(FPCGGraphTask& Task)
 
 		for (const FPCGTaggedData& TaggedDatum : AllParamsData)
 		{
-			const UPCGParamData* ParamData = CastChecked<UPCGParamData>(TaggedDatum.Data);
-			if (!CombinedParamData)
+			if (const UPCGParamData* ParamData = Cast<UPCGParamData>(TaggedDatum.Data))
 			{
-				// We always need a Context to allow capture of Async objects if this runs outside of the game thread
-				FPCGContext TempContext;
-				TempContext.AsyncState.bIsRunningOnMainThread = IsInGameThread();
+				if (!CombinedParamData)
+				{
+					// We always need a Context to allow capture of Async objects if this runs outside of the game thread
+					FPCGContext TempContext;
+					TempContext.AsyncState.bIsRunningOnMainThread = IsInGameThread();
 				
-				CombinedParamData = ParamData->DuplicateData(&TempContext);
+					CombinedParamData = ParamData->DuplicateData(&TempContext);
 
-				if (!TempContext.AsyncState.bIsRunningOnMainThread)
+					if (!TempContext.AsyncState.bIsRunningOnMainThread)
+					{
+						Task.CombineParamsAsyncObjects = std::move(TempContext.AsyncObjects);
+					}
+				}
+				else
 				{
-					Task.CombineParamsAsyncObjects = std::move(TempContext.AsyncObjects);
+					bSuccess &= PCGMetadataHelpers::CopyAllAttributes(ParamData, CombinedParamData, nullptr);
 				}
 			}
-			else
-			{
-				bSuccess &= PCGMetadataHelpers::CopyAllAttributes(ParamData, CombinedParamData, nullptr);
-			}
 		}
 
 		if (!bSuccess)
@@ -2217,8 +2219,8 @@ void FPCGGraphExecutor::CombineParams(FPCGGraphTask& Task)
 			return;
 		}
 
-		const int32 NewNumberOfInputs = Task.TaskInput.TaggedData.Num() - AllParamsData.Num() + 1;
-		check(NewNumberOfInputs >= 1);
+		const int32 NewNumberOfInputs = Task.TaskInput.TaggedData.Num() - AllParamsData.Num() + (CombinedParamData ? 1 : 0);
+		check(NewNumberOfInputs >= 0);
 
 		TArray<FPCGTaggedData> TempTaggedData{};
 		TempTaggedData.Reserve(NewNumberOfInputs);
@@ -2231,10 +2233,13 @@ void FPCGGraphExecutor::CombineParams(FPCGGraphTask& Task)
 		}
 
 		// Add to the root set since we created a new object, that needs to be kept alive for the duration of the task.
-		FPCGTaggedData CombineParams{};
-		CombineParams.Data = CombinedParamData;
-		CombineParams.Pin = PCGPinConstants::DefaultParamsLabel;
-		TempTaggedData.Add(CombineParams);
+		if (CombinedParamData)
+		{
+			FPCGTaggedData CombineParams{};
+			CombineParams.Data = CombinedParamData;
+			CombineParams.Pin = PCGPinConstants::DefaultParamsLabel;
+			TempTaggedData.Add(CombineParams);
+		}
 
 		Task.TaskInput.TaggedData = std::move(TempTaggedData);
 	}
diff --git a/Engine/Plugins/PCG/Source/PCG/Private/Metadata/Accessors/PCGAttributeAccessorHelpers.cpp b/Engine/Plugins/PCG/Source/PCG/Private/Metadata/Accessors/PCGAttributeAccessorHelpers.cpp
--- a/Engine/Plugins/PCG/Source/PCG/Private/Metadata/Accessors/PCGAttributeAccessorHelpers.cpp
+++ b/Engine/Plugins/PCG/Source/PCG/Private/Metadata/Accessors/PCGAttributeAccessorHelpers.cpp
@@ -562,7 +562,7 @@ TUniquePtr<const IPCGAttributeAccessor> PCGAttributeAccessorHelpers::CreateConst
 		OutResult->bHasMultipleAttributeSetsOnOverridePin = InputParamData.Num() > 1;
 	}
 
-	const UPCGParamData* ParamData = !InputParamData.IsEmpty() ? CastChecked<UPCGParamData>(InputParamData[0].Data) : nullptr;
+	const UPCGParamData* ParamData = !InputParamData.IsEmpty() ? Cast<UPCGParamData>(InputParamData[0].Data) : nullptr;
 
 	if (OutResult && ParamData && !bFromGlobalParamsPin)
 	{
diff --git a/Engine/Plugins/PCG/Source/PCG/Public/Elements/PCGGetActorProperty.h b/Engine/Plugins/PCG/Source/PCG/Public/Elements/PCGGetActorProperty.h
--- a/Engine/Plugins/PCG/Source/PCG/Public/Elements/PCGGetActorProperty.h
+++ b/Engine/Plugins/PCG/Source/PCG/Public/Elements/PCGGetActorProperty.h
@@ -113,7 +113,7 @@ class FPCGGetActorPropertyElement : public IPCGElementWithCustomContext<FPCGLoad
 {
 public:
 	virtual bool CanExecuteOnlyOnMainThread(FPCGContext* Context) const override { return true; }
-	virtual bool IsCacheable(const UPCGSettings* InSettings) const override { return !CastChecked<UPCGGetActorPropertySettings>(InSettings)->bAlwaysRequeryActors; }
+	virtual bool IsCacheable(const UPCGSettings* InSettings) const override;
 protected:
 	virtual bool PrepareDataInternal(FPCGContext* Context) const override;
 	virtual bool ExecuteInternal(FPCGContext* Context) const override;
