[PCG] Add robustness to OOM security measures in PCG Samplers to help mitigate continued crashes

EDIT: Does not solve the ISM issue from UE-221385, which ought to be handled on the ISM side.

diff --git a/Engine/Plugins/PCG/Source/PCG/Private/Elements/PCGSurfaceSampler.cpp b/Engine/Plugins/PCG/Source/PCG/Private/Elements/PCGSurfaceSampler.cpp
--- a/Engine/Plugins/PCG/Source/PCG/Private/Elements/PCGSurfaceSampler.cpp
+++ b/Engine/Plugins/PCG/Source/PCG/Private/Elements/PCGSurfaceSampler.cpp
@@ -131,7 +131,8 @@ namespace PCGSurfaceSampler
 			return false;
 		}
 
-		if (PCGFeatureSwitches::CVarCheckSamplerMemory.GetValueOnAnyThread() && (PCGFeatureSwitches::CVarSamplerMemoryThreshold.GetValueOnAnyThread() * FPlatformMemory::GetStats().AvailablePhysical) < sizeof(FPCGPoint) * FMath::Min(TargetPointCount, CellCount))
+		if (PCGFeatureSwitches::CVarCheckSamplerMemory.GetValueOnAnyThread()
+			&& PCGFeatureSwitches::Helpers::GetAvailableMemoryForSamplers() < (sizeof(FPCGPoint) * FMath::Min(TargetPointCount, CellCount)))
 		{
 			PCGLog::LogErrorOnGraph(FText::Format((LOCTEXT("TooManyPoints", "Skipped - tried to generate too many points ({0}).\nAdjust 'pcg.SamplerMemoryThreshold' if needed.")), FMath::Min(TargetPointCount, CellCount)), Context);
 			return false;
diff --git a/Engine/Plugins/PCG/Source/PCG/Private/Elements/PCGVolumeSampler.cpp b/Engine/Plugins/PCG/Source/PCG/Private/Elements/PCGVolumeSampler.cpp
--- a/Engine/Plugins/PCG/Source/PCG/Private/Elements/PCGVolumeSampler.cpp
+++ b/Engine/Plugins/PCG/Source/PCG/Private/Elements/PCGVolumeSampler.cpp
@@ -81,8 +81,8 @@ namespace PCGVolumeSampler
 				NumIterationsXY64 < MAX_int32 &&
 				NumIterations64 > 0 &&
 				NumIterations64 < MAX_int32 &&
-				(!PCGFeatureSwitches::CVarCheckSamplerMemory.GetValueOnAnyThread() ||
-					(PCGFeatureSwitches::CVarSamplerMemoryThreshold.GetValueOnAnyThread() * FPlatformMemory::GetStats().AvailablePhysical) >= sizeof(FPCGPoint) * NumIterations64))
+				(!PCGFeatureSwitches::CVarCheckSamplerMemory.GetValueOnAnyThread()
+					|| PCGFeatureSwitches::Helpers::GetAvailableMemoryForSamplers() >= sizeof(FPCGPoint) * NumIterations64))
 			{
 				NumIterations = static_cast<int32>(NumIterations64);
 			}
diff --git a/Engine/Plugins/PCG/Source/PCG/Private/PCGCommon.cpp b/Engine/Plugins/PCG/Source/PCG/Private/PCGCommon.cpp
--- a/Engine/Plugins/PCG/Source/PCG/Private/PCGCommon.cpp
+++ b/Engine/Plugins/PCG/Source/PCG/Private/PCGCommon.cpp
@@ -24,6 +24,17 @@ namespace PCGFeatureSwitches
 			}
 		})
 	};
+
+	namespace Helpers
+	{
+		uint64 GetAvailableMemoryForSamplers()
+		{
+			const FPlatformMemoryStats MemoryStats = FPlatformMemory::GetStats();
+			// Also uses AvailableVirtual because the system might have plenty of physical memory but still be limited by virtual memory available in some cases.
+			// (i.e. per-process quota, paging file size lower than actual memory available, etc.).
+			return CVarSamplerMemoryThreshold.GetValueOnAnyThread() * FMath::Min(MemoryStats.AvailablePhysical, MemoryStats.AvailableVirtual);
+		}
+	}
 }
 
 namespace PCGSystemSwitches
diff --git a/Engine/Plugins/PCG/Source/PCG/Public/PCGCommon.h b/Engine/Plugins/PCG/Source/PCG/Public/PCGCommon.h
--- a/Engine/Plugins/PCG/Source/PCG/Public/PCGCommon.h
+++ b/Engine/Plugins/PCG/Source/PCG/Public/PCGCommon.h
@@ -188,6 +188,12 @@ namespace PCGFeatureSwitches
 {
 	extern PCG_API TAutoConsoleVariable<bool> CVarCheckSamplerMemory;
 	extern PCG_API TAutoConsoleVariable<float> CVarSamplerMemoryThreshold;
+
+	namespace Helpers
+	{
+		/** Checks the cvar for allowed physical and virtual memory ratio to be used with samplers. */
+		PCG_API uint64 GetAvailableMemoryForSamplers();
+	}
 }
 
 namespace PCGSystemSwitches
@@ -444,4 +450,4 @@ enum class EPCGGenerationStatus : uint8
 #define PCG_EXECUTION_CACHE_VALIDATION_CREATE_ORIGINAL_SCOPE(PCGComponent)
 #define PCG_EXECUTION_CACHE_VALIDATION_CHECK(PCGComponent)
 
-#endif // PCG_EXECUTION_CACHE_VALIDATION_ENABLED
\ No newline at end of file
+#endif // PCG_EXECUTION_CACHE_VALIDATION_ENABLED
