[PCG] GPU - fix compute graphs not working in dynamic subgraphs

#rnx

diff --git a/Engine/Plugins/PCG/Source/PCG/Private/Compute/Elements/PCGComputeGraphElement.cpp b/Engine/Plugins/PCG/Source/PCG/Private/Compute/Elements/PCGComputeGraphElement.cpp
--- a/Engine/Plugins/PCG/Source/PCG/Private/Compute/Elements/PCGComputeGraphElement.cpp
+++ b/Engine/Plugins/PCG/Source/PCG/Private/Compute/Elements/PCGComputeGraphElement.cpp
@@ -50,25 +50,30 @@ bool FPCGComputeGraphElement::ExecuteInternal(FPCGContext* InContext) const
 	if (!Context->ComputeGraph)
 	{
 		const UPCGComponent* Component = InContext->SourceComponent.Get();
-
-		if (const UPCGGraph* PCGGraph = Component ? Component->GetGraph() : nullptr)
+		const UPCGGraph* TopGraph = Component ? Component->GetGraph() : nullptr;
+		if (!ensure(TopGraph))
 		{
-			uint32 GenerationGridSize = PCGHiGenGrid::UninitializedGridSize();
+			return true;
+		}
 
-			if (PCGGraph->IsHierarchicalGenerationEnabled())
-			{
-				if (Component->IsLocalComponent() || Component->IsPartitioned())
-				{
-					GenerationGridSize = Component->GetGenerationGridSize();
-				}
-			}
+		uint32 GenerationGridSize = PCGHiGenGrid::UninitializedGridSize();
 
-			if (UPCGSubsystem* Subsystem = Context->SourceComponent->GetSubsystem())
+		// Higen is always disabled within dynamic subgraphs - will retrieve Uninitialized tasks (which are always cooked).
+		const UPCGGraph* DynamicSubgraph = InContext->Stack ? InContext->Stack->GetNearestDynamicSubgraphForCurrentFrame() : nullptr;
+			
+		if (TopGraph->IsHierarchicalGenerationEnabled() && !DynamicSubgraph)
+		{
+			if (Component->IsLocalComponent() || Component->IsPartitioned())
 			{
-				Context->ComputeGraph.Reset(Subsystem->GetComputeGraph(PCGGraph, GenerationGridSize, ComputeGraphIndex));
+				GenerationGridSize = Component->GetGenerationGridSize();
 			}
 		}
 
+		if (UPCGSubsystem* Subsystem = Context->SourceComponent->GetSubsystem())
+		{
+			Context->ComputeGraph.Reset(Subsystem->GetComputeGraph(DynamicSubgraph ? DynamicSubgraph : TopGraph, GenerationGridSize, ComputeGraphIndex));
+		}
+
 		if (!ensure(Context->ComputeGraph))
 		{
 			return true;
diff --git a/Engine/Plugins/PCG/Source/PCG/Private/Graph/PCGStackContext.cpp b/Engine/Plugins/PCG/Source/PCG/Private/Graph/PCGStackContext.cpp
--- a/Engine/Plugins/PCG/Source/PCG/Private/Graph/PCGStackContext.cpp
+++ b/Engine/Plugins/PCG/Source/PCG/Private/Graph/PCGStackContext.cpp
@@ -188,6 +188,21 @@ const UPCGGraph* FPCGStack::GetGraphForCurrentFrame() const
 	return nullptr;
 }
 
+const UPCGGraph* FPCGStack::GetNearestDynamicSubgraphForCurrentFrame() const
+{
+	// Dynamic subgraphs and looped subgraphs always have 3 stacks: SubgraphNode/LoopIndex/Subgraph
+	// Look for a loop index frame and then return the corresponding graph.
+	for (int StackIndex = GetStackFrames().Num() - 1; StackIndex >= 2; --StackIndex)
+	{
+		if (StackFrames[StackIndex - 1].IsLoopIndexFrame())
+		{
+			return Cast<UPCGGraph>(StackFrames[StackIndex].Object.Get());
+		}
+	}
+
+	return nullptr;
+}
+
 const UPCGNode* FPCGStack::GetCurrentFrameNode() const
 {
 	return StackFrames.IsEmpty() ? nullptr : Cast<const UPCGNode>(StackFrames.Last().Object.Get());
diff --git a/Engine/Plugins/PCG/Source/PCG/Public/Graph/PCGStackContext.h b/Engine/Plugins/PCG/Source/PCG/Public/Graph/PCGStackContext.h
--- a/Engine/Plugins/PCG/Source/PCG/Public/Graph/PCGStackContext.h
+++ b/Engine/Plugins/PCG/Source/PCG/Public/Graph/PCGStackContext.h
@@ -44,6 +44,9 @@ struct PCG_API FPCGStackFrame
 	// A valid frame should either point to an object or have a loop index >= 0.
 	bool IsValid() const { return LoopIndex != INDEX_NONE || Object.IsValid(); }
 
+	/** Loop index frames are injected for any dynamic subgraph invocation and have Object set explicitly null. */
+	bool IsLoopIndexFrame() const { return Object.IsNull(); }
+
 	void SetObject(const UObject* InObject)
 	{
 		Object = InObject;
@@ -138,6 +141,9 @@ public:
 	/** Gets the graph from the graph frame closest to the top of the stack (most recent), or null if no such graph present. */
 	const UPCGGraph* GetGraphForCurrentFrame() const;
 
+	/** Walks up the stack to find nearest dynamic subgraph frame and if found returns the graph. */
+	const UPCGGraph* GetNearestDynamicSubgraphForCurrentFrame() const;
+
 	/** If current frame (top of stack) corresponds to a node returns that node, otherwise returns null. */
 	const UPCGNode* GetCurrentFrameNode() const;
 
