[PCG] Hide GPU related things in private or internal folder where possible.

#rnx

diff --git a/Engine/Plugins/PCG/Source/PCG/Internal/Compute/IPCGNodeSourceTextProvider.h b/Engine/Plugins/PCG/Source/PCG/Internal/Compute/IPCGNodeSourceTextProvider.h
new file mode 100644
--- /dev/null
+++ b/Engine/Plugins/PCG/Source/PCG/Internal/Compute/IPCGNodeSourceTextProvider.h
@@ -0,0 +1,42 @@
+// Copyright Epic Games, Inc. All Rights Reserved.
+
+#pragma once
+
+#include "UObject/Interface.h"
+#include "IPCGNodeSourceTextProvider.generated.h"
+
+UINTERFACE()
+class PCG_API UPCGNodeSourceTextProvider : public UInterface
+{
+	GENERATED_BODY()
+};
+
+/**
+* [EXPERIMENTAL] Interface for node source text editing widgets.
+* Note: This struct is subject to change without deprecation.
+*/
+class PCG_API IPCGNodeSourceTextProvider
+{
+	GENERATED_BODY()
+
+#if WITH_EDITOR	
+public:
+	/** Get text for declarations pane. */
+	virtual FString GetDeclarationsText() const = 0;
+
+	/** Get text for shader functions pane. */
+	virtual FString GetShaderFunctionsText() const = 0;
+
+	/** Get shader text for editor pane. */
+	virtual FString GetShaderText() const = 0;
+
+	/** Set shader functions text after edit modifications. */
+	virtual void SetShaderFunctionsText(const FString& NewFunctionsText) = 0;
+
+	/** Set shader text after edit modifications. */
+	virtual void SetShaderText(const FString& NewText) = 0;
+
+	/** Return true is the shader text can be edited */
+	virtual bool IsShaderTextReadOnly() const = 0;
+#endif
+};
diff --git a/Engine/Plugins/PCG/Source/PCG/Internal/Compute/PCGCompilerDiagnostic.h b/Engine/Plugins/PCG/Source/PCG/Internal/Compute/PCGCompilerDiagnostic.h
new file mode 100644
--- /dev/null
+++ b/Engine/Plugins/PCG/Source/PCG/Internal/Compute/PCGCompilerDiagnostic.h
@@ -0,0 +1,49 @@
+// Copyright Epic Games, Inc. All Rights Reserved.
+
+#pragma once
+
+#include "CoreMinimal.h"
+
+/** [EXPERIMENTAL] Compiler diagnostic severity.
+* Note: This struct is subject to change without deprecation.
+*/
+enum class EPCGDiagnosticLevel : uint8
+{
+	None,
+	Info,
+	Warning,
+	Error
+};
+
+/** [EXPERIMENTAL] Diagnostic message emitted by shader compiler.
+* Note: This struct is subject to change without deprecation.
+*/
+struct FPCGCompilerDiagnostic
+{
+	FPCGCompilerDiagnostic() = default;
+	FPCGCompilerDiagnostic(const FPCGCompilerDiagnostic&) = default;
+	FPCGCompilerDiagnostic& operator=(const FPCGCompilerDiagnostic&) = default;
+
+	// The severity of the issue.
+	EPCGDiagnosticLevel Level = EPCGDiagnosticLevel::None;
+
+	// The actual diagnostic message.
+	FText Message;
+
+	// Line location in source.
+	int32 Line = INDEX_NONE;
+
+	// Starting column (inclusive).
+	int32 ColumnStart = INDEX_NONE;
+
+	// Ending column (inclusive).
+	int32 ColumnEnd = INDEX_NONE;
+};
+
+/** [EXPERIMENTAL] Convenience struct to wrap up array of diagnostics.
+ * Note: This struct is subject to change without deprecation.
+ */
+struct FPCGCompilerDiagnostics
+{
+	TArray<FPCGCompilerDiagnostic> Diagnostics;
+};
diff --git a/Engine/Plugins/PCG/Source/PCG/Private/Compute/DataInterfaces/PCGInstanceDataInterface.cpp b/Engine/Plugins/PCG/Source/PCG/Private/Compute/DataInterfaces/PCGInstanceDataInterface.cpp
--- a/Engine/Plugins/PCG/Source/PCG/Private/Compute/DataInterfaces/PCGInstanceDataInterface.cpp
+++ b/Engine/Plugins/PCG/Source/PCG/Private/Compute/DataInterfaces/PCGInstanceDataInterface.cpp
@@ -120,7 +120,7 @@ UComputeDataProvider* UPCGInstanceDataInterface::CreateDataProvider(TObjectPtr<U
 	// If there were 0 input points (or too many input points) for this execution, we will not have created any primitives, so check for null.
 	if (FPCGSpawnerPrimitives* FoundPrimitives = Binding->MeshSpawnersToPrimitives.Find(ProducerSettings))
 	{
-		const FPCGDataCollectionDesc InputDataDesc = ProducerSettings->ComputeInputPinDataDesc(InputPinProvidingData, Binding);
+		const FPCGDataCollectionDesc InputDataDesc = PCGDataForGPUHelpers::ComputeInputPinDataDesc(ProducerSettings, InputPinProvidingData, Binding);
 		DataProvider->NumInstancesAllPrimitives = InputDataDesc.ComputeDataElementCount(EPCGDataType::Point);
 
 		if (!ensure(FoundPrimitives->Primitives.Num() <= PCGComputeConstants::MAX_PRIMITIVE_COMPONENTS_PER_SPAWNER))
diff --git a/Engine/Plugins/PCG/Source/PCG/Private/Compute/Elements/PCGComputeGraphElement.cpp b/Engine/Plugins/PCG/Source/PCG/Private/Compute/Elements/PCGComputeGraphElement.cpp
--- a/Engine/Plugins/PCG/Source/PCG/Private/Compute/Elements/PCGComputeGraphElement.cpp
+++ b/Engine/Plugins/PCG/Source/PCG/Private/Compute/Elements/PCGComputeGraphElement.cpp
@@ -337,7 +337,7 @@ bool FPCGComputeGraphElement::SetupProceduralISMComponents(FPCGContext* InContex
 			continue;
 		}
 
-		const FPCGDataCollectionDesc InputDataDesc = SpawnerSettings->ComputeInputPinDataDesc(PCGPinConstants::DefaultInputLabel, InBinding);
+		const FPCGDataCollectionDesc InputDataDesc = PCGDataForGPUHelpers::ComputeInputPinDataDesc(SpawnerSettings, PCGPinConstants::DefaultInputLabel, InBinding);
 
 		const uint32 InputPointCount = InputDataDesc.ComputeDataElementCount(EPCGDataType::Point);
 		if (InputPointCount == 0)
diff --git a/Engine/Plugins/PCG/Source/PCG/Private/Compute/Elements/PCGCustomHLSL.cpp b/Engine/Plugins/PCG/Source/PCG/Private/Compute/Elements/PCGCustomHLSL.cpp
--- a/Engine/Plugins/PCG/Source/PCG/Private/Compute/Elements/PCGCustomHLSL.cpp
+++ b/Engine/Plugins/PCG/Source/PCG/Private/Compute/Elements/PCGCustomHLSL.cpp
@@ -315,7 +315,7 @@ bool UPCGCustomHLSLSettings::ComputeOutputPinDataDesc(const UPCGPin* OutputPin,
 		// First output pin passes through first input pin.
 		if (const UPCGPin* PointProcessingInputPin = GetPointProcessingInputPin())
 		{
-			PinDesc = ComputeInputPinDataDesc(PointProcessingInputPin, Binding);
+			PinDesc = PCGDataForGPUHelpers::ComputeInputPinDataDesc(PointProcessingInputPin, Binding);
 		}
 	}
 	else if (OutputPin == FirstOutputPin && KernelType == EPCGKernelType::PointGenerator)
@@ -375,7 +375,7 @@ bool UPCGCustomHLSLSettings::ComputeOutputPinDataDesc(const UPCGPin* OutputPin,
 		// Collect descriptions of input data items that have string key attributes.
 		for (const UPCGPin* InputPin : Node->GetInputPins())
 		{
-			const FPCGDataCollectionDesc InputPinDesc = ComputeInputPinDataDesc(InputPin, Binding);
+			const FPCGDataCollectionDesc InputPinDesc = PCGDataForGPUHelpers::ComputeInputPinDataDesc(InputPin, Binding);
 
 			bool bFoundStringKeyAttribute = false;
 
@@ -1494,7 +1494,7 @@ bool UPCGCustomHLSLSettings::AreKernelAttributesValid(FPCGContext* InContext, FT
 		for (const UPCGPin* InputPin : InPins)
 		{
 			check(InputPin);
-			InputPinDescs.Add(InputPin->Properties.Label, ComputeInputPinDataDesc(InputPin, DataBinding));
+			InputPinDescs.Add(InputPin->Properties.Label, PCGDataForGPUHelpers::ComputeInputPinDataDesc(InputPin, DataBinding));
 		}
 
 		for (const UPCGPin* OutputPin : OutPins)
@@ -1636,7 +1636,7 @@ bool UPCGCustomHLSLSettings::AreKernelAttributesValid(FPCGContext* InContext, FT
 
 FString UPCGCustomHLSLSettings::GetCookedKernelSource(const TMap<FName, FPCGKernelAttributeIDAndType>& GlobalAttributeLookupTable) const
 {
-	const FIntVector GroupSize = GetThreadGroupSize();
+	const FIntVector GroupSize(PCGComputeConstants::THREAD_GROUP_SIZE, 1, 1);
 
 	// FIXME: Create source range mappings so that we can go from error location to our source.
 	FString Source = ShaderSource;
@@ -1677,7 +1677,7 @@ FString UPCGCustomHLSLSettings::GetCookedKernelSource(const TMap<FName, FPCGKern
 
 	const FString KernelFunc = FString::Printf(
 		TEXT("[numthreads(%d, %d, %d)]\nvoid %s(uint3 GroupId : SV_GroupID, uint GroupIndex : SV_GroupIndex)"),
-		GroupSize.X, GroupSize.Y, GroupSize.Z, *GetKernelEntryPoint());
+		GroupSize.X, GroupSize.Y, GroupSize.Z, TEXT("Main"));
 
 	const FString UnWrappedDispatchThreadId = FString::Printf(
 		TEXT("GetUnWrappedDispatchThreadId(GroupId, GroupIndex, %d)"),
diff --git a/Engine/Plugins/PCG/Source/PCG/Private/Compute/Elements/PCGCustomHLSL.h b/Engine/Plugins/PCG/Source/PCG/Private/Compute/Elements/PCGCustomHLSL.h
new file mode 100644
--- /dev/null
+++ b/Engine/Plugins/PCG/Source/PCG/Private/Compute/Elements/PCGCustomHLSL.h
@@ -0,0 +1,193 @@
+// Copyright Epic Games, Inc. All Rights Reserved.
+
+#pragma once
+
+#include "PCGSettings.h"
+#include "Compute/PCGPinPropertiesGPU.h"
+
+#include "Compute/IPCGNodeSourceTextProvider.h"
+
+#include "PCGCustomHLSL.generated.h"
+
+class UPCGPin;
+
+/** Type of kernel allows us to make decisions about execution automatically, streamlining authoring. */
+UENUM()
+enum class EPCGKernelType : uint8
+{
+	PointProcessor UMETA(Tooltip = "Kernel executes on each point in first input pin."),
+	PointGenerator UMETA(Tooltip = "Kernel executes for fixed number of points, configurable on node."),
+	Custom UMETA(Tooltip = "Execution thread counts and output buffer sizes configurable on node. All data read/write indices must be manually bounds checked."),
+};
+
+/** Total number of threads that will be dispatched for this kernel. */
+UENUM()
+enum class EPCGDispatchThreadCount : uint8
+{
+	FromFirstOutputPin UMETA(Tooltip = "One thread per pin data element."),
+	Fixed UMETA(DisplayName = "Fixed Thread Count"),
+	FromProductOfInputPins UMETA(Tooltip = "Dispatches a thread per element in the product of one or more pins. So if there are 4 data elements in pin A and 6 data elements in pin B, 24 threads will be dispatched."),
+};
+
+/** Produces a HLSL compute shader which will be executed on the GPU. */
+UCLASS(MinimalAPI, BlueprintType, ClassGroup = (Procedural))
+class UPCGCustomHLSLSettings
+	: public UPCGSettings
+	, public IPCGNodeSourceTextProvider
+{
+	GENERATED_BODY()
+
+public:
+	UPCGCustomHLSLSettings();
+
+	//~Begin UObject interface
+	virtual void PostLoad() override;
+	virtual void PostInitProperties() override;
+	//~End UObject interface
+
+	//~Begin UPCGSettings interface
+	virtual TArray<FPCGPinProperties> InputPinProperties() const override { return InputPins; }
+	virtual TArray<FPCGPinProperties> OutputPinProperties() const override;
+	virtual bool HasOverridableParams() const override { return false; }
+	virtual bool IsInputPinRequiredByExecution(const UPCGPin* InPin) const override { return true; }
+	virtual bool UseSeed() const override { return true; }
+#if WITH_EDITOR
+	virtual FName GetDefaultNodeName() const override { return FName(TEXT("CustomHLSL")); }
+	virtual FText GetDefaultNodeTitle() const override { return NSLOCTEXT("PCGCustomHLSLElement", "NodeTitle", "Custom HLSL"); }
+	virtual FText GetNodeTooltipText() const override { return NSLOCTEXT("PCGCustomHLSLElement", "NodeTooltip", "[EXPERIMENTAL] Produces a HLSL compute shader which will be executed on the GPU."); }
+	virtual EPCGSettingsType GetType() const override { return EPCGSettingsType::GPU; }
+
+	virtual TArray<FPCGPreConfiguredSettingsInfo> GetPreconfiguredInfo() const override;
+	virtual bool OnlyExposePreconfiguredSettings() const override { return true; }
+#endif
+
+	virtual FString GetAdditionalTitleInformation() const override;
+	virtual void ApplyPreconfiguredSettings(const FPCGPreConfiguredSettingsInfo& PreconfigureInfo) override;
+
+	virtual bool IsKernelValid(FPCGContext* InContext = nullptr, bool bQuiet = true) const override;
+	virtual FString GetCookedKernelSource(const TMap<FName, FPCGKernelAttributeIDAndType>& GlobalAttributeLookupTable) const override;
+	virtual void GetKernelAttributeKeys(TArray<FPCGKernelAttributeKey>& OutKeys) const override { OutKeys.Append(KernelAttributeKeys); }
+	virtual int ComputeKernelThreadCount(const UPCGDataBinding* Binding) const override;
+	virtual bool ComputeOutputPinDataDesc(const UPCGPin* OutputPin, const UPCGDataBinding* Binding, FPCGDataCollectionDesc& OutDesc) const override;
+
+protected:
+#if WITH_EDITOR
+	virtual void PreEditChange(FProperty* PropertyAboutToChange) override;
+	virtual void PostEditChangeProperty(FPropertyChangedEvent& PropertyChangedEvent) override;
+	virtual EPCGChangeType GetChangeTypeForProperty(const FName& InPropertyName) const override;
+#endif
+	virtual FPCGElementPtr CreateElement() const override;
+	//~End UPCGSettings interface
+
+#if WITH_EDITOR
+	//~Begin IPCGNodeSourceTextProvider interface
+	FString GetShaderText() const override;
+	FString GetDeclarationsText() const override;
+	FString GetShaderFunctionsText() const override;
+	void SetShaderFunctionsText(const FString& NewFunctionsText) override;
+	void SetShaderText(const FString& NewText) override;
+	bool IsShaderTextReadOnly() const override;
+	//~End IPCGNodeSourceTextProvider interface
+#endif
+
+	/** Gets the GPU pin properties for the output pin with the given label. */
+	const FPCGPinPropertiesGPU* GetOutputPinPropertiesGPU(const FName& InPinLabel) const;
+
+	/** Enforce required pin settings and set display toggles to drive UI. */
+	void UpdatePinSettings();
+
+	/** Update the kernel keys read/written/created by this node. */
+	void UpdateAttributeKeys();
+
+#if WITH_EDITOR
+	void UpdateDeclarations();
+	void UpdateInputDeclarations();
+	void UpdateOutputDeclarations();
+	void UpdateHelperDeclarations();
+
+	/** List of all non-advanced input pin names. */
+	UFUNCTION()
+	TArray<FName> GetInputPinNames() const;
+
+	/** List of all non-advanced input pin names, prepended with 'Name_NONE'. */
+	UFUNCTION()
+	TArray<FName> GetInputPinNamesAndNone() const;
+#endif
+
+	const UPCGPin* GetInputPin(FName Label) const;
+	const UPCGPin* GetOutputPin(FName Label) const;
+	const UPCGPin* GetFirstInputPin() const;
+	const UPCGPin* GetPointProcessingInputPin() const;
+	const UPCGPin* GetFirstOutputPin() const;
+	const UPCGPin* GetFirstPointOutputPin() const;
+	int GetProcessingElemCountForInputPin(const UPCGPin* InputPin, const UPCGDataBinding* Binding) const;
+	bool AreKernelAttributesValid(FPCGContext* InContext, FText* OutErrorText) const;
+
+	/** Will the ThreadCountMultiplier value be applied when calculating the dispatch thread count. */
+	bool IsThreadCountMultiplierInUse() const { return KernelType == EPCGKernelType::Custom && DispatchThreadCount != EPCGDispatchThreadCount::Fixed; }
+
+protected:
+	UPROPERTY(BlueprintReadWrite, EditAnywhere, Category = "Settings")
+	EPCGKernelType KernelType = EPCGKernelType::PointProcessor;
+
+	UPROPERTY(BlueprintReadWrite, EditAnywhere, DisplayName = "Number of Points", Category = "Settings", meta = (EditCondition = "KernelType == EPCGKernelType::PointGenerator", EditConditionHides))
+	int PointCount = 256;
+
+	UPROPERTY(BlueprintReadWrite, EditAnywhere, Category = "Settings|Thread Count", meta = (EditCondition = "KernelType == EPCGKernelType::Custom", EditConditionHides))
+	EPCGDispatchThreadCount DispatchThreadCount = EPCGDispatchThreadCount::FromFirstOutputPin;
+
+	UPROPERTY(BlueprintReadWrite, EditAnywhere, Category = "Settings|Thread Count", meta = (EditCondition = "KernelType == EPCGKernelType::Custom && DispatchThreadCount != EPCGDispatchThreadCount::Fixed", EditConditionHides))
+	int ThreadCountMultiplier = 1;
+
+	UPROPERTY(BlueprintReadWrite, EditAnywhere, Category = "Settings|Thread Count", meta = (EditCondition = "KernelType == EPCGKernelType::Custom && DispatchThreadCount == EPCGDispatchThreadCount::Fixed", EditConditionHides))
+	int FixedThreadCount = 1;
+
+	UPROPERTY(BlueprintReadWrite, EditAnywhere, DisplayName = "Input Pins", Category = "Settings|Thread Count", meta = (EditCondition = "KernelType == EPCGKernelType::Custom && DispatchThreadCount == EPCGDispatchThreadCount::FromProductOfInputPins", EditConditionHides, GetOptions = "GetInputPinNames"))
+	TArray<FName> ThreadCountInputPinLabels;
+
+public:
+	UPROPERTY(BlueprintReadWrite, EditAnywhere, Category = "Settings")
+	TArray<FPCGPinProperties> InputPins = Super::DefaultPointInputPinProperties();
+
+	UPROPERTY(BlueprintReadWrite, EditAnywhere, Category = "Settings")
+	TArray<FPCGPinPropertiesGPU> OutputPins = { FPCGPinPropertiesGPU(PCGPinConstants::DefaultOutputLabel, EPCGDataType::Point) };
+
+#if WITH_EDITOR
+	/** Holds input pin labels from PreEditChange, used in PostEditPropertyChange to update any references in output pin setup. */
+	TArray<FName> InputPinLabelsPreEditChange;
+#endif
+
+protected:
+	/** Optional functions that can be called from the source. Intended to be edited using the Node Source Editor window. */
+	UPROPERTY(BlueprintReadWrite, EditAnywhere, Category = "Settings", AdvancedDisplay, meta = (MultiLine = true))
+	FString ShaderFunctions = "/** CUSTOM SHADER FUNCTIONS **/\n";
+
+	/** Shader code that forms the body of the kernel. Intended to be edited using the Node Source Editor window. */
+	UPROPERTY(BlueprintReadWrite, EditAnywhere, Category = "Settings", AdvancedDisplay, meta = (MultiLine = true))
+	FString ShaderSource;
+
+	/** Inputs data accessors that can be used from the shader code. Intended to be viewed using the Node Source Editor window. */
+	UPROPERTY(Transient, VisibleAnywhere, Category = "Settings", AdvancedDisplay, meta = (MultiLine = true))
+	FString InputDeclarations;
+
+	/** Output data accessors that can be used from the shader code. Intended to be viewed using the Node Source Editor window. */
+	UPROPERTY(Transient, VisibleAnywhere, Category = "Settings", AdvancedDisplay, meta = (MultiLine = true))
+	FString OutputDeclarations;
+
+	/** Helper data and functions that can be used from the shader code. Intended to be viewed using the Node Source Editor window. */
+	UPROPERTY(Transient, VisibleAnywhere, Category = "Settings", AdvancedDisplay, meta = (MultiLine = true))
+	FString HelperDeclarations;
+
+	/** Attributes statically detected as being read, written, or created by this node. */
+	TArray<FPCGKernelAttributeKey> KernelAttributeKeys;
+
+	/** Maps pins to their attribute keys and whether or not they were created on the GPU. */
+	TMap<FName, TArray<TTuple<FPCGKernelAttributeKey, bool /*bCreatedOnGPU*/>>> PinToAttributeKeys;
+};
+
+class FPCGCustomHLSLElement : public IPCGElement
+{
+protected:
+	// This will only be called if the custom HLSL node is not set up correctly (valid nodes are replaced with a compute graph element).
+	virtual bool ExecuteInternal(FPCGContext* Context) const override;
+};
diff --git a/Engine/Plugins/PCG/Source/PCG/Private/Compute/PCGComputeCommon.cpp b/Engine/Plugins/PCG/Source/PCG/Private/Compute/PCGComputeCommon.cpp
--- a/Engine/Plugins/PCG/Source/PCG/Private/Compute/PCGComputeCommon.cpp
+++ b/Engine/Plugins/PCG/Source/PCG/Private/Compute/PCGComputeCommon.cpp
@@ -108,7 +108,7 @@ namespace PCGComputeHelpers
 				if (const UPCGPin* InitPin = Node->GetInputPin(PinToInitFrom))
 				{
 					InitPins.Emplace(PinToInitFrom);
-					InputDescs.Emplace(Settings->ComputeInputPinDataDesc(InitPin, Binding));
+					InputDescs.Emplace(PCGDataForGPUHelpers::ComputeInputPinDataDesc(InitPin, Binding));
 				}
 			}
 
diff --git a/Engine/Plugins/PCG/Source/PCG/Private/Compute/PCGComputeCommon.h b/Engine/Plugins/PCG/Source/PCG/Private/Compute/PCGComputeCommon.h
--- a/Engine/Plugins/PCG/Source/PCG/Private/Compute/PCGComputeCommon.h
+++ b/Engine/Plugins/PCG/Source/PCG/Private/Compute/PCGComputeCommon.h
@@ -58,6 +58,8 @@ namespace PCGComputeConstants
 	constexpr uint32 KernelExecutedFlag = 1 << 31;
 
 	constexpr uint32 MAX_PRIMITIVE_COMPONENTS_PER_SPAWNER = 64;
+
+	constexpr uint32 THREAD_GROUP_SIZE = 64;
 }
 
 namespace PCGComputeHelpers
diff --git a/Engine/Plugins/PCG/Source/PCG/Private/Compute/PCGDataBinding.cpp b/Engine/Plugins/PCG/Source/PCG/Private/Compute/PCGDataBinding.cpp
--- a/Engine/Plugins/PCG/Source/PCG/Private/Compute/PCGDataBinding.cpp
+++ b/Engine/Plugins/PCG/Source/PCG/Private/Compute/PCGDataBinding.cpp
@@ -252,7 +252,7 @@ void UPCGDataBinding::DebugLogDataDescriptions()
 
 			for (const UPCGPin* Pin : Node->GetInputPins())
 			{
-				const FPCGDataCollectionDesc DataDesc = Settings->ComputeInputPinDataDesc(Pin, this);
+				const FPCGDataCollectionDesc DataDesc = PCGDataForGPUHelpers::ComputeInputPinDataDesc(Pin, this);
 				UE_LOG(LogPCG, Display, TEXT("\tInput Pin: %s (%d data)"), *Pin->Properties.Label.ToString(), DataDesc.DataDescs.Num());
 				LogDataDescription(DataDesc);
 			}
diff --git a/Engine/Plugins/PCG/Source/PCG/Private/Compute/PCGDataForGPU.cpp b/Engine/Plugins/PCG/Source/PCG/Private/Compute/PCGDataForGPU.cpp
--- a/Engine/Plugins/PCG/Source/PCG/Private/Compute/PCGDataForGPU.cpp
+++ b/Engine/Plugins/PCG/Source/PCG/Private/Compute/PCGDataForGPU.cpp
@@ -3,8 +3,11 @@
 #include "Compute/PCGDataForGPU.h"
 
 #include "PCGData.h"
+#include "PCGEdge.h"
+#include "PCGNode.h"
 #include "PCGParamData.h"
 #include "PCGPoint.h"
+#include "PCGSettings.h"
 #include "Compute/PCGComputeCommon.h"
 #include "Compute/PCGComputeGraph.h"
 #include "Compute/PCGDataBinding.h"
@@ -461,6 +464,54 @@ namespace PCGDataForGPUHelpers
 
 		OutCustomFloatCount = OffsetFloats;
 	}
+
+	FPCGDataCollectionDesc ComputeInputPinDataDesc(const UPCGSettings* Settings, const FName& InputPinLabel, const UPCGDataBinding* Binding)
+	{
+		check(Settings);
+
+		const UPCGNode* Node = Cast<UPCGNode>(Settings->GetOuter());
+		const UPCGPin* InputPin = Node ? Node->GetInputPin(InputPinLabel) : nullptr;
+
+		if (ensure(InputPin))
+		{
+			return ComputeInputPinDataDesc(InputPin, Binding);
+		}
+		else
+		{
+			return {};
+		}
+	}
+
+	FPCGDataCollectionDesc ComputeInputPinDataDesc(const UPCGPin* InputPin, const UPCGDataBinding* Binding)
+	{
+		check(InputPin);
+		check(Binding);
+
+		FPCGDataCollectionDesc PinDesc;
+
+		// Grab data from all incident edges.
+		for (const UPCGEdge* Edge : InputPin->Edges)
+		{
+			// InputPin is upstream output pin.
+			const UPCGPin* UpstreamOutputPin = Edge->InputPin;
+			if (!UpstreamOutputPin)
+			{
+				continue;
+			}
+
+			const UPCGSettings* UpstreamSettings = UpstreamOutputPin->Node ? UpstreamOutputPin->Node->GetSettings() : nullptr;
+			check(UpstreamSettings);
+
+			// Add data from connected upstream output pin.
+			FPCGDataCollectionDesc EdgeDesc;
+			if (ensure(UpstreamSettings->ComputeOutputPinDataDesc(UpstreamOutputPin, Binding, EdgeDesc)))
+			{
+				PinDesc.Combine(EdgeDesc);
+			}
+		}
+
+		return PinDesc;
+	}
 }
 
 bool FPCGKernelAttributeKey::operator==(const FPCGKernelAttributeKey& Other) const
diff --git a/Engine/Plugins/PCG/Source/PCG/Private/Compute/PCGDataForGPU.h b/Engine/Plugins/PCG/Source/PCG/Private/Compute/PCGDataForGPU.h
new file mode 100644
--- /dev/null
+++ b/Engine/Plugins/PCG/Source/PCG/Private/Compute/PCGDataForGPU.h
@@ -0,0 +1,190 @@
+// Copyright Epic Games, Inc. All Rights Reserved.
+
+#pragma once
+
+#include "PCGCommon.h"
+#include "PCGData.h"
+
+#include "PCGDataForGPU.generated.h"
+
+class UPCGDataBinding;
+class UPCGMetadata;
+class UPCGPin;
+class UPCGSettings;
+enum class EPCGMetadataTypes : uint8;
+
+enum class EPCGUnpackDataCollectionResult
+{
+	Success,
+	DataMismatch,
+	NoData
+};
+
+UENUM()
+enum class EPCGKernelAttributeType : uint8
+{
+	None = 0,
+	Bool,
+	Int,
+	Float,
+	Float2,
+	Float3,
+	Float4,
+	Rotator,
+	Quat,
+	Transform,
+	StringKey,
+
+	Invalid = std::numeric_limits<uint8>::max()
+};
+
+USTRUCT()
+struct FPCGKernelAttributeKey
+{
+	GENERATED_BODY()
+
+	FPCGKernelAttributeKey() = default;
+
+	explicit FPCGKernelAttributeKey(FName InName, EPCGKernelAttributeType InType)
+		: Name(InName)
+		, Type(InType)
+	{}
+
+	UPROPERTY(EditAnywhere, Category = "Settings")
+	FName Name = NAME_None;
+
+	UPROPERTY(EditAnywhere, Category = "Settings")
+	EPCGKernelAttributeType Type = EPCGKernelAttributeType::Float;
+
+	bool operator==(const FPCGKernelAttributeKey& Other) const;
+	friend uint32 GetTypeHash(const FPCGKernelAttributeKey& In);
+};
+
+USTRUCT()
+struct FPCGKernelAttributeIDAndType
+{
+	GENERATED_BODY()
+
+	FPCGKernelAttributeIDAndType() = default;
+
+	explicit FPCGKernelAttributeIDAndType(int32 InId, EPCGKernelAttributeType InType)
+		: Id(InId)
+		, Type(InType)
+	{}
+
+	UPROPERTY()
+	int32 Id = INDEX_NONE;
+
+	UPROPERTY()
+	EPCGKernelAttributeType Type = EPCGKernelAttributeType::Invalid;
+};
+
+struct FPCGKernelAttributeDesc
+{
+	explicit FPCGKernelAttributeDesc(int32 InIndex, EPCGKernelAttributeType InType, FName InName)
+		: Index(InIndex)
+		, Type(InType)
+		, Name(InName)
+	{
+	}
+
+	explicit FPCGKernelAttributeDesc(int32 InIndex, EPCGKernelAttributeType InType, FName InName, TArray<int32>&& InUniqueStringKeys)
+		: Index(InIndex)
+		, Type(InType)
+		, Name(InName)
+		, UniqueStringKeys(MoveTemp(InUniqueStringKeys))
+	{
+	}
+
+	// Attribute Id
+	int32 Index = INDEX_NONE;
+
+	EPCGKernelAttributeType Type = EPCGKernelAttributeType::Float;
+	FName Name = NAME_None;
+
+	/* All possible string keys arriving on this attribute (string keys are indices into the string table in the data binding). */
+	TArray<int32> UniqueStringKeys;
+
+	bool operator==(const FPCGKernelAttributeDesc& Other) const;
+};
+
+struct FPCGDataDesc
+{
+	FPCGDataDesc(EPCGDataType InType, int32 InElementCount);
+	FPCGDataDesc(const UPCGData* InData, const TMap<FName, FPCGKernelAttributeIDAndType>& InGlobalAttributeLookupTable, const TArray<FString>& InStringTable);
+
+	uint64 ComputePackedSize() const;
+
+	EPCGDataType Type = EPCGDataType::Point;
+	TArray<FPCGKernelAttributeDesc> AttributeDescs;
+	int32 ElementCount = 0;
+
+private:
+	void InitializeAttributeDescs(const UPCGData* InData, const TMap<FName, FPCGKernelAttributeIDAndType>& InGlobalAttributeLookupTable, const TArray<FString>& InStringTable);
+};
+
+struct FPCGDataCollectionDesc
+{
+	static FPCGDataCollectionDesc BuildFromInputDataCollectionAndInputPinLabel(
+		const FPCGDataCollection& InDataCollection,
+		FName InputPinLabel,
+		const TMap<FName, FPCGKernelAttributeIDAndType>& InAttributeLookupTable,
+		const TArray<FString>& InStringTable);
+
+	/** Computes the size (in bytes) of the header portion of the packed data collection buffer. */
+	uint32 ComputePackedHeaderSizeBytes() const;
+
+	/** Computes the size (in bytes) of the data collection after packing. */
+	uint64 ComputePackedSizeBytes() const;
+
+	void WriteHeader(TArray<uint32>& OutPackedDataCollectionHeader) const;
+
+	/** Pack a data collection into the GPU data format. DataDescs defines which attributes are packed. */
+	void PackDataCollection(const FPCGDataCollection& InDataCollection, FName InPin, const TArray<FString>& InStringTable, TArray<uint32>& OutPackedDataCollection) const;
+
+	/** Unpack a buffer of 8-bit uints to a data collection. */
+	EPCGUnpackDataCollectionResult UnpackDataCollection(const TArray<uint8>& InPackedData, FName InPin, const TArray<FString>& InStringTable, FPCGDataCollection& OutDataCollection) const;
+
+	/** Compute total number of processing elements of the given type. */
+	uint32 ComputeDataElementCount(EPCGDataType InDataType) const;
+
+	/** Aggregate another data description. */
+	void Combine(const FPCGDataCollectionDesc& Other);
+
+	TArray<FPCGDataDesc> DataDescs;
+};
+
+USTRUCT()
+struct FPCGDataForGPU
+{
+	GENERATED_BODY()
+
+	// All the node input pins that have edges that cross from CPU to GPU.
+	UPROPERTY()
+	TSet<TSoftObjectPtr<const UPCGPin>> InputPins;
+
+	UPROPERTY()
+	TMap<TSoftObjectPtr<const UPCGPin>, FName> InputPinLabelAliases;
+
+	// Since the compute graph is collapsed to a single element, all data crossing from CPU to GPU is in a single collection.
+	UPROPERTY()
+	FPCGDataCollection InputDataCollection;
+};
+
+namespace PCGDataForGPUHelpers
+{
+	/** Returns GPU type that will be used to represent the given metadata type. */
+	EPCGKernelAttributeType GetAttributeTypeFromMetadataType(EPCGMetadataTypes MetadataType);
+
+	/** Compute how attributes can be packed to custom floats. */
+	void ComputeCustomFloatPacking(
+		TArray<FName>& InAttributeNames,
+		const UPCGDataBinding* InBinding,
+		const FPCGDataCollectionDesc& InDataCollectionDescription,
+		uint32& OutCustomFloatCount,
+		TArray<FUint32Vector4>& OutAttributeIdOffsetStrides);
+
+	/** Compute a description of all data arriving on InputPin. */
+	FPCGDataCollectionDesc ComputeInputPinDataDesc(const UPCGSettings* Settings, const FName& InputPinLabel, const UPCGDataBinding* Binding);
+	FPCGDataCollectionDesc ComputeInputPinDataDesc(const UPCGPin* InputPin, const UPCGDataBinding* Binding);
+}
diff --git a/Engine/Plugins/PCG/Source/PCG/Private/Compute/PCGPinPropertiesGPU.h b/Engine/Plugins/PCG/Source/PCG/Private/Compute/PCGPinPropertiesGPU.h
new file mode 100644
--- /dev/null
+++ b/Engine/Plugins/PCG/Source/PCG/Private/Compute/PCGPinPropertiesGPU.h
@@ -0,0 +1,152 @@
+// Copyright Epic Games, Inc. All Rights Reserved.
+
+#pragma once
+
+#include "PCGPin.h"
+#include "Compute/PCGDataForGPU.h"
+
+#include "PCGPinPropertiesGPU.generated.h"
+
+/** Method for initializing an output pin. */
+UENUM()
+enum class EPCGPinInitMode : uint8
+{
+	FromInputPins UMETA(Tooltip = "Initialize from one or more input pins."),
+	Custom UMETA(Tooltip = "Initialize from fixed values."),
+};
+
+/** Method for computing data count. */
+UENUM()
+enum class EPCGDataCountMode : uint8
+{
+	FromInputData,
+	Fixed,
+};
+
+/** Method for computing element count. */
+UENUM()
+enum class EPCGElementCountMode : uint8
+{
+	FromInputData,
+	Fixed,
+};
+
+/** Method for combining two or more data counts. */
+UENUM()
+enum class EPCGDataMultiplicity : uint8
+{
+	Pairwise UMETA(Tooltip = "A data item will be produced for each pair/tuple/etc of input data items across the input pins. Requires all pins to have the same number of data."),
+	CartesianProduct UMETA(Tooltip = "If there are two input pins with N and M data items respectively, the output will have NxM data items."),
+};
+
+/** Method for combining two or more element counts. */
+UENUM()
+enum class EPCGElementMultiplicity : uint8
+{
+	Product UMETA(Tooltip = "Element count will be the product of the elements in each pair/tuple/etc of input data."),
+	Sum UMETA(Tooltip = "Element count will be the sum of the elements in each pair/tuple/etc of input data."),
+};
+
+/** Method for inheriting attribute data from input pins. */
+UENUM()
+enum class EPCGAttributeInheritanceMode : uint8
+{
+	None UMETA(Tooltip = "No attributes inherited from initialization pins."),
+	CopyAttributeSetup UMETA(Tooltip = "Take attribute names and types from initialization pins. Pins declared first will take priority during conflicts. Does not copy values."),
+	// TODO: Mode to automatically copy values. Requires code-gen to loop on batches of elements for pins based on the total number of elements and total number of threads.
+	//CopyAttributeSetupAndValues UMETA(Tooltip = "Take attribute names, types, and initial values from initialization pins. Pins declared first will take priority during conflicts."),
+};
+
+/** Helper struct to nest GPU pin properties inside a UI category. */
+USTRUCT(BlueprintType)
+struct FPCGPinPropertiesGPUStruct
+{
+	GENERATED_BODY()
+
+	/** How the output data for this pin will be initialized. */
+	UPROPERTY(BlueprintReadWrite, EditAnywhere, Category = Settings, meta = (EditCondition = "bAllowEditInitMode", EditConditionHides, HideEditConditionToggle))
+	EPCGPinInitMode InitializationMode = EPCGPinInitMode::FromInputPins;
+
+	/** Input pins to initialize this pin's data from. */
+	UPROPERTY(BlueprintReadWrite, EditAnywhere, Category = Settings, meta = (GetOptions = "GetInputPinNames", EditCondition = "bAllowEditInitMode && InitializationMode == EPCGPinInitMode::FromInputPins", EditConditionHides))
+	TArray<FName> PinsToInititalizeFrom;
+
+	/** How the number of data is determined. */
+	UPROPERTY(BlueprintReadWrite, EditAnywhere, Category = Settings, meta = (EditCondition = "bAllowEditInitMode && bAllowEditDataCount && InitializationMode == EPCGPinInitMode::FromInputPins", EditConditionHides))
+	EPCGDataCountMode DataCountMode = EPCGDataCountMode::FromInputData;
+
+	/** How to combine data counts. */
+	UPROPERTY(BlueprintReadWrite, EditAnywhere, Category = Settings, meta = (EditCondition = "bAllowEditInitMode && bAllowEditDataCount && InitializationMode == EPCGPinInitMode::FromInputPins && DataCountMode == EPCGDataCountMode::FromInputData && bMultipleInitPins", EditConditionHides))
+	EPCGDataMultiplicity DataMultiplicity = EPCGDataMultiplicity::Pairwise;
+
+	/** Number of data to create. */
+	UPROPERTY(BlueprintReadWrite, EditAnywhere, Category = Settings, meta = (EditCondition = "bAllowEditInitMode && (InitializationMode == EPCGPinInitMode::Custom || DataCountMode == EPCGDataCountMode::Fixed)", EditConditionHides, ClampMin = 1))
+	int DataCount = 1;
+
+	/** How the number of elements is determined. */
+	UPROPERTY(BlueprintReadWrite, EditAnywhere, Category = Settings, meta = (EditCondition = "bAllowEditInitMode && InitializationMode == EPCGPinInitMode::FromInputPins", EditConditionHides))
+	EPCGElementCountMode ElementCountMode = EPCGElementCountMode::FromInputData;
+
+	/** How to combine element counts. */
+	UPROPERTY(BlueprintReadWrite, EditAnywhere, Category = Settings, meta = (EditCondition = "bAllowEditInitMode && InitializationMode == EPCGPinInitMode::FromInputPins && ElementCountMode == EPCGElementCountMode::FromInputData && bMultipleInitPins", EditConditionHides))
+	EPCGElementMultiplicity ElementMultiplicity = EPCGElementMultiplicity::Product;
+
+	/** Fixed number of elements to create in each output data. */
+	UPROPERTY(BlueprintReadWrite, EditAnywhere, Category = Settings, meta = (EditCondition = "bAllowEditInitMode && (InitializationMode == EPCGPinInitMode::Custom || ElementCountMode == EPCGElementCountMode::Fixed)", EditConditionHides, ClampMin = 1))
+	int ElementCount = 1;
+
+	/** How to inherit attribute names, types, and values. */
+	UPROPERTY(BlueprintReadWrite, EditAnywhere, Category = Settings, meta = (EditCondition = "bAllowEditInitMode && InitializationMode == EPCGPinInitMode::FromInputPins", EditConditionHides))
+	EPCGAttributeInheritanceMode AttributeInheritanceMode = EPCGAttributeInheritanceMode::CopyAttributeSetup;
+
+	/** Add entries to create new attributes on data emitted by this pin. */
+	UPROPERTY(EditAnywhere, DisplayName = "Attributes to Create", Category = Settings)
+	TArray<FPCGKernelAttributeKey> CreatedKernelAttributeKeys;
+
+#if WITH_EDITORONLY_DATA
+	UPROPERTY(Transient)
+	bool bAllowEditInitMode = false;
+
+	/** Some types force single data, such as Attribute Sets. We should not allow data count to be edited for such types. */
+	UPROPERTY(Transient)
+	bool bAllowEditDataCount = false;
+
+	UPROPERTY(Transient)
+	bool bMultipleInitPins = false;
+#endif // WITH_EDITORONLY_DATA
+};
+
+/** An extension of the pin properties that adds hints for GPU specific properties, such as buffer size and data layout. */
+USTRUCT(BlueprintType)
+struct PCG_API FPCGPinPropertiesGPU : public FPCGPinProperties
+{
+	GENERATED_BODY()
+
+public:
+	FPCGPinPropertiesGPU() = default;
+
+	explicit FPCGPinPropertiesGPU(const FName& InLabel, EPCGDataType InAllowedTypes)
+		: FPCGPinProperties(InLabel, InAllowedTypes)
+	{}
+
+	explicit FPCGPinPropertiesGPU(const FPCGPinProperties& InPinProps)
+		: FPCGPinProperties(InPinProps.Label, InPinProps.AllowedTypes)
+	{}
+
+#if WITH_EDITOR
+	bool CanEditChange(const FEditPropertyChain& PropertyChain) const;
+#endif
+
+public:
+	UPROPERTY(BlueprintReadWrite, EditAnywhere, Category = Settings, meta = (DisplayName = "GPU Properties", DisplayAfter = "Tooltip"))
+	FPCGPinPropertiesGPUStruct PropertiesGPU;
+};
+
+template<>
+struct TStructOpsTypeTraits<FPCGPinPropertiesGPU> : public TStructOpsTypeTraitsBase2<FPCGPinPropertiesGPU>
+{
+	enum
+	{
+		WithCanEditChange = true,
+	};
+};
diff --git a/Engine/Plugins/PCG/Source/PCG/Private/Elements/PCGCopyPoints.cpp b/Engine/Plugins/PCG/Source/PCG/Private/Elements/PCGCopyPoints.cpp
--- a/Engine/Plugins/PCG/Source/PCG/Private/Elements/PCGCopyPoints.cpp
+++ b/Engine/Plugins/PCG/Source/PCG/Private/Elements/PCGCopyPoints.cpp
@@ -53,8 +53,8 @@ bool UPCGCopyPointsSettings::ComputeOutputPinDataDesc(const UPCGPin* OutputPin,
 
 	if (ensure(SourcePin && TargetPin))
 	{
-		const FPCGDataCollectionDesc SourcePinDesc = ComputeInputPinDataDesc(SourcePin, Binding);
-		const FPCGDataCollectionDesc TargetPinDesc = ComputeInputPinDataDesc(TargetPin, Binding);
+		const FPCGDataCollectionDesc SourcePinDesc = PCGDataForGPUHelpers::ComputeInputPinDataDesc(SourcePin, Binding);
+		const FPCGDataCollectionDesc TargetPinDesc = PCGDataForGPUHelpers::ComputeInputPinDataDesc(TargetPin, Binding);
 
 		const int32 NumSources = SourcePinDesc.DataDescs.Num();
 		const int32 NumTargets = TargetPinDesc.DataDescs.Num();
diff --git a/Engine/Plugins/PCG/Source/PCG/Private/Elements/PCGReroute.cpp b/Engine/Plugins/PCG/Source/PCG/Private/Elements/PCGReroute.cpp
--- a/Engine/Plugins/PCG/Source/PCG/Private/Elements/PCGReroute.cpp
+++ b/Engine/Plugins/PCG/Source/PCG/Private/Elements/PCGReroute.cpp
@@ -7,6 +7,7 @@
 #include "PCGModule.h"
 #include "PCGNode.h"
 #include "PCGPin.h"
+#include "Compute/PCGDataForGPU.h"
 
 UPCGRerouteSettings::UPCGRerouteSettings()
 {
@@ -53,7 +54,7 @@ bool UPCGRerouteSettings::ComputeOutputPinDataDesc(const UPCGPin* OutputPin, con
 	{
 		if (const UPCGPin* Pin = Node->GetInputPin(PCGPinConstants::DefaultInputLabel))
 		{
-			OutDesc = ComputeInputPinDataDesc(Pin, Binding);
+			OutDesc = PCGDataForGPUHelpers::ComputeInputPinDataDesc(Pin, Binding);
 			return true;
 		}
 	}
diff --git a/Engine/Plugins/PCG/Source/PCG/Private/Elements/PCGStaticMeshSpawner.cpp b/Engine/Plugins/PCG/Source/PCG/Private/Elements/PCGStaticMeshSpawner.cpp
--- a/Engine/Plugins/PCG/Source/PCG/Private/Elements/PCGStaticMeshSpawner.cpp
+++ b/Engine/Plugins/PCG/Source/PCG/Private/Elements/PCGStaticMeshSpawner.cpp
@@ -11,8 +11,8 @@
 #include "Compute/DataInterfaces/Elements/PCGStaticMeshSpawnerDataInterface.h"
 #include "Data/PCGPointData.h"
 #include "Data/PCGSpatialData.h"
-#include "Graph/PCGGPUGraphCompilationContext.h"
 #include "Elements/PCGStaticMeshSpawnerContext.h"
+#include "Graph/PCGGPUGraphCompilationContext.h"
 #include "Helpers/PCGActorHelpers.h"
 #include "Helpers/PCGHelpers.h"
 #include "InstanceDataPackers/PCGInstanceDataPackerBase.h"
@@ -119,7 +119,7 @@ bool UPCGStaticMeshSpawnerSettings::ComputeOutputPinDataDesc(const UPCGPin* Outp
 		return true;
 	}
 
-	const FPCGDataCollectionDesc InputPinDesc = ComputeInputPinDataDesc(PCGPinConstants::DefaultInputLabel, Binding);
+	const FPCGDataCollectionDesc InputPinDesc = PCGDataForGPUHelpers::ComputeInputPinDataDesc(this, PCGPinConstants::DefaultInputLabel, Binding);
 	FPCGDataCollectionDesc OutputPinDesc = InputPinDesc;
 
 	if (ensure(OutputPin->Properties.Label == PCGPinConstants::DefaultOutputLabel))
@@ -181,10 +181,8 @@ bool UPCGStaticMeshSpawnerSettings::ComputeOutputPinDataDesc(const UPCGPin* Outp
 	return true;
 }
 
-const TArray<FPCGKernelAttributeKey> UPCGStaticMeshSpawnerSettings::GetKernelAttributeKeys() const
+void UPCGStaticMeshSpawnerSettings::GetKernelAttributeKeys(TArray<FPCGKernelAttributeKey>& OutKeys) const
 {
-	TArray<FPCGKernelAttributeKey> AttributeKeys;
-
 	if (InstanceDataPackerParameters)
 	{
 		TArray<FName> AttributeNames;
@@ -193,7 +191,7 @@ const TArray<FPCGKernelAttributeKey> UPCGStaticMeshSpawnerSettings::GetKernelAtt
 			for (const FName& AttributeName : AttributeNames)
 			{
 				// We don't know the type statically before execution, leave unset.
-				AttributeKeys.AddUnique(FPCGKernelAttributeKey(AttributeName, EPCGKernelAttributeType::None));
+				OutKeys.AddUnique(FPCGKernelAttributeKey(AttributeName, EPCGKernelAttributeType::None));
 			}
 		}
 	}
@@ -201,13 +199,11 @@ const TArray<FPCGKernelAttributeKey> UPCGStaticMeshSpawnerSettings::GetKernelAtt
 	if (const UPCGMeshSelectorByAttribute* Selector = Cast<UPCGMeshSelectorByAttribute>(MeshSelectorParameters))
 	{
 		// Add an attribute key for the given attribute so we register it as being read.
-		AttributeKeys.AddUnique(FPCGKernelAttributeKey(Selector->AttributeName, EPCGKernelAttributeType::None));
+		OutKeys.AddUnique(FPCGKernelAttributeKey(Selector->AttributeName, EPCGKernelAttributeType::None));
 	}
 
 	// Write out attribute.
-	AttributeKeys.AddUnique(FPCGKernelAttributeKey(OutAttributeName, EPCGKernelAttributeType::StringKey));
-
-	return AttributeKeys;
+	OutKeys.AddUnique(FPCGKernelAttributeKey(OutAttributeName, EPCGKernelAttributeType::StringKey));
 }
 
 void UPCGStaticMeshSpawnerSettings::AddStaticCreatedStrings(TArray<FString>& InOutStringTable) const
@@ -226,7 +222,7 @@ int UPCGStaticMeshSpawnerSettings::ComputeKernelThreadCount(const UPCGDataBindin
 	const UPCGNode* Node = CastChecked<UPCGNode>(GetOuter());
 	const UPCGPin* Pin = Node->GetInputPin(PCGPinConstants::DefaultInputLabel);
 
-	return ComputeInputPinDataDesc(Pin, Binding).ComputeDataElementCount(EPCGDataType::Point);
+	return PCGDataForGPUHelpers::ComputeInputPinDataDesc(Pin, Binding).ComputeDataElementCount(EPCGDataType::Point);
 }
 
 #if WITH_EDITOR
diff --git a/Engine/Plugins/PCG/Source/PCG/Private/Graph/PCGGraphCompilerGPU.cpp b/Engine/Plugins/PCG/Source/PCG/Private/Graph/PCGGraphCompilerGPU.cpp
--- a/Engine/Plugins/PCG/Source/PCG/Private/Graph/PCGGraphCompilerGPU.cpp
+++ b/Engine/Plugins/PCG/Source/PCG/Private/Graph/PCGGraphCompilerGPU.cpp
@@ -649,7 +649,10 @@ void FPCGGraphCompilerGPU::BuildGPUGraphTask(
 					OutputPinDataInterfaces.Add({ TaskId, OutputPinProperties.Label }, OutputDI);
 
 					// TODO: This seems like a once-per-settings process rather than repeating for each output pin?
-					for (const FPCGKernelAttributeKey& Key : Settings->GetKernelAttributeKeys())
+					TArray<FPCGKernelAttributeKey> AttributeKeys;
+					Settings->GetKernelAttributeKeys(AttributeKeys);
+
+					for (const FPCGKernelAttributeKey& Key : AttributeKeys)
 					{
 						FPCGKernelAttributeIDAndType* FoundEntry = ComputeGraph->GlobalAttributeLookupTable.Find(Key.Name);
 						if (FoundEntry)
@@ -1001,8 +1004,11 @@ void FPCGGraphCompilerGPU::BuildGPUGraphTask(
 		{
 			UPCGComputeKernelSource* KernelSource = InOutContext.NewObject_AnyThread<UPCGComputeKernelSource>(KernelWithBindings.Kernel); // is outer to kernel fine?
 			KernelWithBindings.Kernel->KernelSource = KernelSource;
-			KernelSource->EntryPoint = Settings->GetKernelEntryPoint();
-			KernelSource->GroupSize = Settings->GetThreadGroupSize();
+
+			// These could be exposed through PCGSettings API later when the need arises (and/or when GPU feature matures).
+			KernelSource->EntryPoint = TEXT("Main");
+			KernelSource->GroupSize = FIntVector(PCGComputeConstants::THREAD_GROUP_SIZE, 1, 1);
+
 			KernelSource->SetSource(Settings->GetCookedKernelSource(ComputeGraph->GlobalAttributeLookupTable));
 
 			if (Settings->bDumpCookedHLSL)
diff --git a/Engine/Plugins/PCG/Source/PCG/Private/PCGSettings.cpp b/Engine/Plugins/PCG/Source/PCG/Private/PCGSettings.cpp
--- a/Engine/Plugins/PCG/Source/PCG/Private/PCGSettings.cpp
+++ b/Engine/Plugins/PCG/Source/PCG/Private/PCGSettings.cpp
@@ -240,44 +240,6 @@ bool UPCGSettings::IsKernelValid(FPCGContext* InContext, bool bQuiet) const
 	return true;
 }
 
-FPCGDataCollectionDesc UPCGSettings::ComputeInputPinDataDesc(const FName& InputPinLabel, const UPCGDataBinding* Binding) const
-{
-	const UPCGNode* Node = Cast<UPCGNode>(GetOuter());
-	const UPCGPin* InputPin = Node ? Node->GetInputPin(InputPinLabel) : nullptr;
-	check(InputPin);
-	return ComputeInputPinDataDesc(InputPin, Binding);
-}
-
-FPCGDataCollectionDesc UPCGSettings::ComputeInputPinDataDesc(const UPCGPin* InputPin, const UPCGDataBinding* Binding) const
-{
-	check(InputPin && Binding);
-
-	FPCGDataCollectionDesc PinDesc;
-
-	// Grab data from all incident edges.
-	for (const UPCGEdge* Edge : InputPin->Edges)
-	{
-		// InputPin is upstream output pin.
-		const UPCGPin* UpstreamOutputPin = Edge->InputPin;
-		if (!UpstreamOutputPin)
-		{
-			continue;
-		}
-
-		const UPCGSettings* UpstreamSettings = UpstreamOutputPin->Node ? UpstreamOutputPin->Node->GetSettings() : nullptr;
-		check(UpstreamSettings);
-
-		// Add data from connected upstream output pin.
-		FPCGDataCollectionDesc EdgeDesc;
-		if (ensure(UpstreamSettings->ComputeOutputPinDataDesc(UpstreamOutputPin, Binding, EdgeDesc)))
-		{
-			PinDesc.Combine(EdgeDesc);
-		}
-	}
-
-	return PinDesc;
-}
-
 bool UPCGSettings::ComputeOutputPinDataDesc(const FName& OutputPinLabel, const UPCGDataBinding* InBinding, FPCGDataCollectionDesc& OutDesc) const
 {
 	const UPCGNode* Node = CastChecked<UPCGNode>(GetOuter());
diff --git a/Engine/Plugins/PCG/Source/PCG/Public/Compute/Elements/PCGCustomHLSL.h b/Engine/Plugins/PCG/Source/PCG/Public/Compute/Elements/PCGCustomHLSL.h
deleted file mode 100644
diff --git a/Engine/Plugins/PCG/Source/PCG/Public/Compute/IPCGNodeSourceTextProvider.h b/Engine/Plugins/PCG/Source/PCG/Public/Compute/IPCGNodeSourceTextProvider.h
deleted file mode 100644
diff --git a/Engine/Plugins/PCG/Source/PCG/Public/Compute/PCGCompilerDiagnostic.h b/Engine/Plugins/PCG/Source/PCG/Public/Compute/PCGCompilerDiagnostic.h
deleted file mode 100644
diff --git a/Engine/Plugins/PCG/Source/PCG/Public/Compute/PCGDataForGPU.h b/Engine/Plugins/PCG/Source/PCG/Public/Compute/PCGDataForGPU.h
deleted file mode 100644
diff --git a/Engine/Plugins/PCG/Source/PCG/Public/Compute/PCGPinPropertiesGPU.h b/Engine/Plugins/PCG/Source/PCG/Public/Compute/PCGPinPropertiesGPU.h
deleted file mode 100644
diff --git a/Engine/Plugins/PCG/Source/PCG/Public/Elements/PCGStaticMeshSpawner.h b/Engine/Plugins/PCG/Source/PCG/Public/Elements/PCGStaticMeshSpawner.h
--- a/Engine/Plugins/PCG/Source/PCG/Public/Elements/PCGStaticMeshSpawner.h
+++ b/Engine/Plugins/PCG/Source/PCG/Public/Elements/PCGStaticMeshSpawner.h
@@ -31,7 +31,7 @@ public:
 	// ~Begin UPCGSettings interface
 	virtual bool IsKernelValid(FPCGContext* InContext = nullptr, bool bQuiet = true) const;
 	virtual FString GetCookedKernelSource(const TMap<FName, FPCGKernelAttributeIDAndType>& GlobalAttributeLookupTable) const override;
-	virtual const TArray<FPCGKernelAttributeKey> GetKernelAttributeKeys() const override;
+	virtual void GetKernelAttributeKeys(TArray<FPCGKernelAttributeKey>& OutKeys) const override;
 	virtual void AddStaticCreatedStrings(TArray<FString>& InOutStringTable) const override;
 	virtual int ComputeKernelThreadCount(const UPCGDataBinding* Binding) const override;
 	virtual bool ComputeOutputPinDataDesc(const UPCGPin* OutputPin, const UPCGDataBinding* Binding, FPCGDataCollectionDesc& OutDesc) const override;
diff --git a/Engine/Plugins/PCG/Source/PCG/Public/PCGSettings.h b/Engine/Plugins/PCG/Source/PCG/Public/PCGSettings.h
--- a/Engine/Plugins/PCG/Source/PCG/Public/PCGSettings.h
+++ b/Engine/Plugins/PCG/Source/PCG/Public/PCGSettings.h
@@ -6,7 +6,6 @@
 #include "PCGDebug.h"
 #include "PCGElement.h"
 #include "PCGPin.h"
-#include "Compute/PCGDataForGPU.h"
 #include "Elements/PCGActorSelector.h"
 #include "Tests/Determinism/PCGDeterminismSettings.h"
 
@@ -16,16 +15,18 @@
 
 class UPCGComponent;
 class UPCGPin;
-struct FPCGGPUCompilationContext;
-struct FPCGPinProperties;
-struct FPropertyChangedEvent;
-
 class UPCGComputeGraph;
 class UPCGDataBinding;
 class UPCGGraph;
 class UComputeDataInterface;
 class UPCGNode;
 class UPCGSettings;
+struct FPCGDataCollectionDesc;
+struct FPCGGPUCompilationContext;
+struct FPCGKernelAttributeIDAndType;
+struct FPCGKernelAttributeKey;
+struct FPCGPinProperties;
+struct FPropertyChangedEvent;
 
 using FPCGSettingsAndCulling = TPair<TSoftObjectPtr<const UPCGSettings>, bool>;
 using FPCGSelectionKeyToSettingsMap = TMap<FPCGSelectionKey, TArray<FPCGSettingsAndCulling>>;
@@ -528,60 +529,85 @@ public:
 
 	// GPU section
 public:
-	/** Whether this node should be executed on the GPU. */
-	bool ShouldExecuteOnGPU() const { return bExecuteOnGPU; }
 
+	/** [EXPERIMENTAL] Whether this node should be executed on the GPU.
+	* Note: This API function is subject to change without deprecation.
+	*/
+	virtual bool ShouldExecuteOnGPU() const { return bExecuteOnGPU; }
+
+	/** [EXPERIMENTAL] Performs validation and returns true if this node is suitable for deployment to the GPU.
+	* Note: This API function is subject to change without deprecation.
+	*/
 	virtual bool IsKernelValid(FPCGContext* InContext = nullptr, bool bQuiet = true) const;
+
+	/** [EXPERIMENTAL] Produces the node specific portion of kernel shader source text, including the main entry point.
+	* Note: This API function is subject to change without deprecation.
+	*/
 	virtual FString GetCookedKernelSource(const TMap<FName, FPCGKernelAttributeIDAndType>& GlobalAttributeLookupTable) const { return TEXT(""); }
-	virtual FString GetKernelEntryPoint() const { return TEXT("Main"); }
-	virtual FIntVector GetThreadGroupSize() const { return FIntVector(64, 1, 1); }
 
-	/** Get a list of the attributes read or written by this node. */
-	virtual const TArray<FPCGKernelAttributeKey> GetKernelAttributeKeys() const { return {}; }
+	/** [EXPERIMENTAL] Get a list of the attributes read or written by this node.
+	* Note: This API function is subject to change without deprecation.
+	*/
+	virtual void GetKernelAttributeKeys(TArray<FPCGKernelAttributeKey>& OutKeys) const {}
 
-	/** Add any strings emitted by this node that are known statically at compile time. */
+	/** [EXPERIMENTAL] Add any strings emitted by this node that are known statically at compile time.
+	* Note: This API function is subject to change without deprecation.
+	*/
 	virtual void AddStaticCreatedStrings(TArray<FString>& InOutStringTable) const {};
 
-	/** Compute how many threads should be dispatched to execute this node on the GPU. */
+	/** [EXPERIMENTAL] Compute how many threads should be dispatched to execute this node on the GPU.
+	* Note: This API function is subject to change without deprecation.
+	*/
 	virtual int ComputeKernelThreadCount(const UPCGDataBinding* Binding) const { return 0; };
 
-	/** Compute a description of all data arriving on InputPin. */
-	FPCGDataCollectionDesc ComputeInputPinDataDesc(const FName& InputPinLabel, const UPCGDataBinding* Binding) const;
-	FPCGDataCollectionDesc ComputeInputPinDataDesc(const UPCGPin* InputPin, const UPCGDataBinding* Binding) const;
-
-	/** Compute a description of data that will be output from OutputPinLabel/OutputPin. */
+	/** [EXPERIMENTAL] Compute a description of data that will be output from OutputPinLabel/OutputPin.
+	* Note: This API function is subject to change without deprecation.
+	*/
 	bool ComputeOutputPinDataDesc(const FName& OutputPinLabel, const UPCGDataBinding* InBinding, FPCGDataCollectionDesc& OutDesc) const;
+
+	/** [EXPERIMENTAL] Compute a description of data that will be output from OutputPinLabel/OutputPin.
+	* Note: This API function is subject to change without deprecation.
+	*/
 	virtual bool ComputeOutputPinDataDesc(const UPCGPin* OutputPin, const UPCGDataBinding* InBinding, FPCGDataCollectionDesc& OutDesc) const;
 
 #if WITH_EDITOR
-	/** Create additional data interfaces to marshal any data required by this settings. */
+	/** [EXPERIMENTAL] Create additional input data interfaces to marshal any required input data.
+	* Note: This API function is subject to change without deprecation.
+	*/
 	virtual void CreateAdditionalInputDataInterfaces(FPCGGPUCompilationContext& InOutContext, UObject* InObjectOuter, TArray<TObjectPtr<UComputeDataInterface>>& OutDataInterfaces) const {}
+
+	/** [EXPERIMENTAL] Create additional output data interfaces to marshal any required output data.
+	* Note: This API function is subject to change without deprecation.
+	*/
 	virtual void CreateAdditionalOutputDataInterfaces(FPCGGPUCompilationContext& InOutContext, UObject* InObjectOuter, TArray<TObjectPtr<UComputeDataInterface>>& OutDataInterfaces) const {}
+	virtual void CreateAdditionalOutputDataInterfaces(TArray<TObjectPtr<UComputeDataInterface>>& OutDataInterfaces) const {}
 
-	/** Whether to display GPU execution option in node settings UI. */
+	/** [EXPERIMENTAL] Whether to display GPU execution option in node settings UI. 
+	* Note: This API function is subject to change without deprecation.
+	*/
 	UFUNCTION()
 	virtual bool DisplayExecuteOnGPUSetting() const { return false; }
 #endif
 
-public:
-	/** Whether this node should be executed on the GPU. */
-	UPROPERTY(BlueprintReadWrite, EditAnywhere, Category = "GPU", meta = (EditCondition = "DisplayExecuteOnGPUSetting()", EditConditionHides, HideEditConditionToggle))
+protected:
+	/** [EXPERIMENTAL] Note: This variable is subject to change without deprecation. */
+	UPROPERTY(EditAnywhere, Category = "GPU", meta = (Tooltip = "Whether this node should be executed on the GPU.", EditCondition = "DisplayExecuteOnGPUSetting()", EditConditionHides, HideEditConditionToggle))
 	bool bExecuteOnGPU = false;
 
-	/** Dump the cooked HLSL into the log after it is generated. */
-	UPROPERTY(EditAnywhere, Category = "GPU", AdvancedDisplay, meta = (EditCondition = "bExecuteOnGPU", EditConditionHides))
+	/** [EXPERIMENTAL] Note: This variable is subject to change without deprecation. */
+	UPROPERTY(EditAnywhere, Category = "GPU", AdvancedDisplay, meta = (Tooltip = "Dump the cooked HLSL into the log after it is generated.", EditCondition = "bExecuteOnGPU", EditConditionHides))
 	bool bDumpCookedHLSL = false;
 
-	/** Dump the data descriptions of input/output pins to the log. */
-	UPROPERTY(EditAnywhere, Category = "GPU", AdvancedDisplay, meta = (EditCondition = "bExecuteOnGPU", EditConditionHides))
+	/** [EXPERIMENTAL] Note: This variable is subject to change without deprecation. */
+	UPROPERTY(EditAnywhere, Category = "GPU", AdvancedDisplay, meta = (Tooltip = "Dump the data descriptions of input/output pins to the log.", EditCondition = "bExecuteOnGPU", EditConditionHides))
 	bool bDumpDataDescriptions = false;
 
-	/** Enable use of 'WriteDebugValue(uint Index, float Value)' function in your kernel. Allows you to write float values to a buffer for logging on the CPU. */
-	UPROPERTY(EditAnywhere, Category = "GPU", AdvancedDisplay, meta = (EditCondition = "bExecuteOnGPU", EditConditionHides))
+	/** [EXPERIMENTAL] Note: This variable is subject to change without deprecation. */
+	UPROPERTY(EditAnywhere, Category = "GPU", AdvancedDisplay, meta = (Tooltip = "Enable use of 'WriteDebugValue(uint Index, float Value)' function in your kernel. Allows you to write float values to a buffer for logging on the CPU.", EditCondition = "bExecuteOnGPU", EditConditionHides))
 	bool bPrintShaderDebugValues = false;
 
-	/** Size (in number of floats) of the shader debug print buffer. */
-	UPROPERTY(EditAnywhere, Category = "GPU", AdvancedDisplay, meta = (EditCondition="bExecuteOnGPU && bPrintShaderDebugValues", EditConditionHides))
+	/** [EXPERIMENTAL] Note: This variable is subject to change without deprecation. */
+	UPROPERTY(EditAnywhere, Category = "GPU", AdvancedDisplay, meta = (Tooltip = "Size (in number of floats) of the shader debug print buffer.", EditCondition="bExecuteOnGPU && bPrintShaderDebugValues", EditConditionHides))
 	int DebugBufferSize = 16;
 
 private:
@@ -590,6 +616,9 @@ private:
 
 	/** The cached Crc for these settings. */
 	FPCGCrc CachedCrc;
+
+	friend class FPCGGraphCompilerGPU;
+	friend class UPCGDataBinding;
 };
 
 UCLASS(BlueprintType, ClassGroup = (Procedural))
diff --git a/Engine/Plugins/PCG/Source/PCGEditor/Private/Details/PCGCustomHLSLSettingsDetails.cpp b/Engine/Plugins/PCG/Source/PCGEditor/Private/Details/PCGCustomHLSLSettingsDetails.cpp
--- a/Engine/Plugins/PCG/Source/PCGEditor/Private/Details/PCGCustomHLSLSettingsDetails.cpp
+++ b/Engine/Plugins/PCG/Source/PCGEditor/Private/Details/PCGCustomHLSLSettingsDetails.cpp
@@ -2,8 +2,6 @@
 
 #include "Details/PCGCustomHLSLSettingsDetails.h"
 
-#include "Compute/Elements/PCGCustomHLSL.h"
-
 #include "PCGEditor.h"
 #include "PCGEditorGraph.h"
 
diff --git a/Engine/Plugins/PCG/Source/PCGEditor/Private/PCGEditorModule.cpp b/Engine/Plugins/PCG/Source/PCGEditor/Private/PCGEditorModule.cpp
--- a/Engine/Plugins/PCG/Source/PCGEditor/Private/PCGEditorModule.cpp
+++ b/Engine/Plugins/PCG/Source/PCGEditor/Private/PCGEditorModule.cpp
@@ -8,7 +8,6 @@
 #include "PCGParamData.h"
 #include "PCGSubsystem.h"
 #include "PCGVolumeFactory.h"
-#include "Compute/Elements/PCGCustomHLSL.h"
 #include "Data/PCGSpatialData.h"
 #include "Data/PCGSplineData.h"
 #include "Grid/PCGPartitionActor.h"
@@ -234,7 +233,7 @@ void FPCGEditorModule::RegisterDetailsCustomizations()
 	PropertyEditor.RegisterCustomClassLayout("PCGGraphInstance", FOnGetDetailCustomizationInstance::CreateStatic(&FPCGGraphInstanceDetails::MakeInstance));
 	PropertyEditor.RegisterCustomClassLayout("PCGVolume", FOnGetDetailCustomizationInstance::CreateStatic(&FPCGVolumeDetails::MakeInstance));
 	PropertyEditor.RegisterCustomClassLayout("PCGUserParameterGetSettings", FOnGetDetailCustomizationInstance::CreateStatic(&FPCGEditableUserParameterDetails::MakeInstance));
-	PropertyEditor.RegisterCustomClassLayout(UPCGCustomHLSLSettings::StaticClass()->GetFName(), FOnGetDetailCustomizationInstance::CreateStatic(&FPCGCustomHLSLSettingsDetails::MakeInstance));
+	PropertyEditor.RegisterCustomClassLayout("PCGCustomHLSLSettings", FOnGetDetailCustomizationInstance::CreateStatic(&FPCGCustomHLSLSettingsDetails::MakeInstance));
 
 	PropertyEditor.RegisterCustomPropertyTypeLayout("PCGAttributePropertySelector", FOnGetPropertyTypeCustomizationInstance::CreateStatic(&FPCGAttributePropertySelectorDetails::MakeInstance));
 	PropertyEditor.RegisterCustomPropertyTypeLayout("PCGAttributePropertyInputSelector", FOnGetPropertyTypeCustomizationInstance::CreateStatic(&FPCGAttributePropertySelectorDetails::MakeInstance));
