[PCG] Fixing few issues with Write To Niagara Data Channel

* Visible to GPU was always false, making it not working to send to GPU emitters. Now all visibility options are exposed.
* Duplicate the node was losing attribute mapping. It was because of the Import/Export on Attribute Selector not working with more complex export string (like a map).

#rnx

diff --git a/Engine/Plugins/Experimental/PCGInterops/PCGNiagaraInterop/Source/PCGNiagaraInterop/Private/Elements/PCGWriteToNiagaraDataChannel.cpp b/Engine/Plugins/Experimental/PCGInterops/PCGNiagaraInterop/Source/PCGNiagaraInterop/Private/Elements/PCGWriteToNiagaraDataChannel.cpp
--- a/Engine/Plugins/Experimental/PCGInterops/PCGNiagaraInterop/Source/PCGNiagaraInterop/Private/Elements/PCGWriteToNiagaraDataChannel.cpp
+++ b/Engine/Plugins/Experimental/PCGInterops/PCGNiagaraInterop/Source/PCGNiagaraInterop/Private/Elements/PCGWriteToNiagaraDataChannel.cpp
@@ -174,7 +174,7 @@ bool FPCGWriteToNiagaraDataChannelElement::ExecuteInternal(FPCGContext* InContex
 
 		const FString DebugSource = FString::Format(TEXT("PCGWriteToNiagaraChannel - {0} - Input {1}"), { InContext->Node ? InContext->Node->GetName() : TEXT("Unknown node"), FString::FromInt(InputNum - 1) });
 
-		UNiagaraDataChannelWriter* NiagaraWriter = UNiagaraDataChannelLibrary::CreateDataChannelWriter(World, DataChannel, std::move(SearchParameters), Count, /*bVisibleToGame=*/ true, /*bVisibleToCPU=*/ true, /*bVisibleToGPU=*/ false, DebugSource);
+		UNiagaraDataChannelWriter* NiagaraWriter = UNiagaraDataChannelLibrary::CreateDataChannelWriter(World, DataChannel, std::move(SearchParameters), Count, Settings->bVisibleToGame, Settings->bVisibleToCPU, Settings->bVisibleToGPU, DebugSource);
 
 		if (!NiagaraWriter)
 		{
diff --git a/Engine/Plugins/Experimental/PCGInterops/PCGNiagaraInterop/Source/PCGNiagaraInterop/Public/Elements/PCGWriteToNiagaraDataChannel.h b/Engine/Plugins/Experimental/PCGInterops/PCGNiagaraInterop/Source/PCGNiagaraInterop/Public/Elements/PCGWriteToNiagaraDataChannel.h
--- a/Engine/Plugins/Experimental/PCGInterops/PCGNiagaraInterop/Source/PCGNiagaraInterop/Public/Elements/PCGWriteToNiagaraDataChannel.h
+++ b/Engine/Plugins/Experimental/PCGInterops/PCGNiagaraInterop/Source/PCGNiagaraInterop/Public/Elements/PCGWriteToNiagaraDataChannel.h
@@ -40,6 +40,18 @@ public:
 	UPROPERTY(BlueprintReadWrite, EditAnywhere, Category = "Settings")
 	TMap<FName, FPCGAttributePropertyInputSelector> NiagaraVariablesPCGAttributeMapping;
 
+	/** Data written to this data channel is visible to Blueprint and C++ logic reading from it */
+	UPROPERTY(BlueprintReadWrite, EditAnywhere, Category = "Settings|Visibility", meta = (PCG_Overridable))
+	bool bVisibleToGame = true;
+
+	/** Data written to this data channel is visible to Niagara CPU emitters */
+	UPROPERTY(BlueprintReadWrite, EditAnywhere, Category = "Settings|Visibility", meta = (PCG_Overridable))
+	bool bVisibleToCPU = true;
+
+	/** Data written to this data channel is visible to Niagara GPU emitters */
+	UPROPERTY(BlueprintReadWrite, EditAnywhere, Category = "Settings|Visibility", meta = (PCG_Overridable))
+	bool bVisibleToGPU = false;
+	
 	UPROPERTY(BlueprintReadWrite, EditAnywhere, Category = "Settings|Debug")
 	bool bSynchronousLoad = false;
 };
diff --git a/Engine/Plugins/PCG/Source/PCG/Private/Metadata/PCGAttributePropertySelector.cpp b/Engine/Plugins/PCG/Source/PCG/Private/Metadata/PCGAttributePropertySelector.cpp
--- a/Engine/Plugins/PCG/Source/PCG/Private/Metadata/PCGAttributePropertySelector.cpp
+++ b/Engine/Plugins/PCG/Source/PCG/Private/Metadata/PCGAttributePropertySelector.cpp
@@ -29,7 +29,7 @@ bool FPCGAttributePropertySelector::ExportTextItem(FString& ValueStr, FPCGAttrib
 	StringBuilder.Append(ToString());
 	StringBuilder.Append(PCGAttributePropertySelectorConstants::ExportTextRightSentinel);
 
-	ValueStr = StringBuilder.ToString();
+	ValueStr += StringBuilder.ToString();
 	return true;
 }
 
@@ -40,16 +40,23 @@ bool FPCGAttributePropertySelector::ImportTextItem(const TCHAR*& Buffer, int32 P
 	using PCGAttributePropertySelectorConstants::ExportTextLeftSentinel;
 	using PCGAttributePropertySelectorConstants::ExportTextRightSentinel;
 
-	if (!BufferView.StartsWith(ExportTextLeftSentinel) || !BufferView.EndsWith(ExportTextRightSentinel))
+	// Look for the first occurence of the left and right sentinel
+	int32 Start = BufferView.Find(ExportTextLeftSentinel);
+	const int32 End = BufferView.Find(ExportTextRightSentinel);
+
+	if (Start == INDEX_NONE || End == INDEX_NONE)
 	{
 		// Didn't find our sentinels, abort
 		return false;
 	}
 
-	const int32 Start = ExportTextLeftSentinel.Len();
-	const int32 End = BufferView.Len() - ExportTextRightSentinel.Len();
+	// Offset our start accounting the size of the left sentinel
+	Start += ExportTextLeftSentinel.Len();
 
 	Update(FString(BufferView.SubStr(Start, End - Start)));
+
+	// Offset buffer to the end of the right sentinel.
+	Buffer += (End + ExportTextRightSentinel.Len());
 	return true;
 }
 
