Fixed component transform issue that was still causing an issue with splines.

#rnx

diff --git a/Engine/Plugins/PCG/Source/PCG/Private/Elements/PCGSpawnSpline.cpp b/Engine/Plugins/PCG/Source/PCG/Private/Elements/PCGSpawnSpline.cpp
--- a/Engine/Plugins/PCG/Source/PCG/Private/Elements/PCGSpawnSpline.cpp
+++ b/Engine/Plugins/PCG/Source/PCG/Private/Elements/PCGSpawnSpline.cpp
@@ -109,12 +109,12 @@ bool FPCGSpawnSplineElement::ExecuteInternal(FPCGContext* InContext) const
 		// Also propagate the tags on the tagged data to the component.
 		Algo::Transform(InputData.Tags, SplineComponent->ComponentTags, [](const FString& Tag) { return FName(Tag); });
 
-		InSplineData->ApplyTo(SplineComponent);
-
 		SplineComponent->RegisterComponent();
 		TargetActor->AddInstanceComponent(SplineComponent);
 		// Spline Data apply to Spline Component is done in world coordinate, so keep world location.
-		SplineComponent->AttachToComponent(TargetActor->GetRootComponent(), FAttachmentTransformRules(EAttachmentRule::KeepWorld, EAttachmentRule::KeepWorld, EAttachmentRule::KeepWorld, false));
+		SplineComponent->AttachToComponent(TargetActor->GetRootComponent(), FAttachmentTransformRules(EAttachmentRule::KeepRelative, EAttachmentRule::KeepWorld, EAttachmentRule::KeepWorld, false));
+
+		InSplineData->ApplyTo(SplineComponent);
 
 		UPCGManagedComponent* ManagedComponent = NewObject<UPCGManagedComponent>(InContext->SourceComponent.Get());
 		ManagedComponent->GeneratedComponent = SplineComponent;
diff --git a/Engine/Plugins/PCG/Source/PCG/Private/Helpers/PCGActorHelpers.cpp b/Engine/Plugins/PCG/Source/PCG/Private/Helpers/PCGActorHelpers.cpp
--- a/Engine/Plugins/PCG/Source/PCG/Private/Helpers/PCGActorHelpers.cpp
+++ b/Engine/Plugins/PCG/Source/PCG/Private/Helpers/PCGActorHelpers.cpp
@@ -308,6 +308,9 @@ UPCGManagedSplineMeshComponent* UPCGActorHelpers::GetOrCreateManagedSplineMeshCo
 	SplineMeshComponent->RegisterComponent();
 	InTargetActor->AddInstanceComponent(SplineMeshComponent);
 
+	// Implementation note: since the data passed to the params here is in world space,
+	// We need the transform on the spline mesh component to be the identity - in world space, because unlike the ISMs
+	// where we can set the instances and specify the data is in world space, we can't do that here.
 	SplineMeshComponent->AttachToComponent(InTargetActor->GetRootComponent(), FAttachmentTransformRules(EAttachmentRule::KeepWorld, EAttachmentRule::KeepWorld, EAttachmentRule::KeepWorld, false));
 	SplineMeshComponent->ComponentTags.Add(InSourceComponent->GetFName());
 	SplineMeshComponent->ComponentTags.Add(PCGHelpers::DefaultPCGTag);
