[PCG] Emit a warning when trying to use GPU settings in proxy nodes

#rnx

diff --git a/Engine/Plugins/PCG/Source/PCG/Private/Elements/PCGIndirectionElement.cpp b/Engine/Plugins/PCG/Source/PCG/Private/Elements/PCGIndirectionElement.cpp
--- a/Engine/Plugins/PCG/Source/PCG/Private/Elements/PCGIndirectionElement.cpp
+++ b/Engine/Plugins/PCG/Source/PCG/Private/Elements/PCGIndirectionElement.cpp
@@ -303,6 +303,13 @@ bool FPCGIndirectionElement::ExecuteInternal(FPCGContext* InContext) const
 		return true;
 	}
 
+	if (Context->InnerSettings && Context->InnerSettings->ShouldExecuteOnGPU())
+	{
+		PCGE_LOG(Error, GraphAndLog, LOCTEXT("UnsupportedGPUProxyNode", "GPU nodes do not currently support execution via proxy."));
+		Context->OutputData = Context->InputData;
+		return true;
+	}
+
 	// TODO: use caching when possible
 	// TODO: see what we can do for inspection data
 	// TODO: support pausing in inner element, might require some upstream changes in the graph executor
