[PCG] Fix unbounded grid size not generating from BeginPlay()
* Also fixes partitioned BP components not properly generating in PIE and thread-safety issues during compute graph creation

#rnx

diff --git a/Engine/Plugins/PCG/Source/PCG/Private/Graph/PCGGPUGraphCompilationContext.h b/Engine/Plugins/PCG/Source/PCG/Private/Graph/PCGGPUGraphCompilationContext.h
--- a/Engine/Plugins/PCG/Source/PCG/Private/Graph/PCGGPUGraphCompilationContext.h
+++ b/Engine/Plugins/PCG/Source/PCG/Private/Graph/PCGGPUGraphCompilationContext.h
@@ -7,11 +7,12 @@
 #include "PCGCommon.h"
 
 class FPCGGraphCompiler;
+class UPCGComputeGraph;
 
 /** Current context/state used during compilation of GPU graphs. */
 struct FPCGGPUCompilationContext
 {
-	FPCGGPUCompilationContext(FPCGGraphCompiler& InGraphCompiler)
+	explicit FPCGGPUCompilationContext(FPCGGraphCompiler& InGraphCompiler)
 		: GraphCompiler(InGraphCompiler)
 	{
 	}
@@ -61,6 +62,10 @@ struct FPCGGPUCompilationContext
 	}
 
 	FPCGGraphCompiler& GetGraphCompiler() { return GraphCompiler; }
+	TArray<TObjectPtr<UPCGComputeGraph>>& GetCompiledComputeGraphs() { return CompiledComputeGraphs; }
+
+	/** Returns the index of the newly added compute graph. */
+	int32 AddCompiledComputeGraph(TObjectPtr<UPCGComputeGraph> NewComputeGraph) { return CompiledComputeGraphs.Add(NewComputeGraph); }
 
 private:
 	// List of objects created during compilation, we need to track them so we can remove their Async flags when compilation is done
@@ -68,6 +73,7 @@ private:
 	TSet<TObjectPtr<UObject>> AsyncObjects;
 
 	FPCGGraphCompiler& GraphCompiler;
+	TArray<TObjectPtr<UPCGComputeGraph>> CompiledComputeGraphs;
 };
 
 #endif // WITH_EDITOR
diff --git a/Engine/Plugins/PCG/Source/PCG/Private/Graph/PCGGraphCompiler.cpp b/Engine/Plugins/PCG/Source/PCG/Private/Graph/PCGGraphCompiler.cpp
--- a/Engine/Plugins/PCG/Source/PCG/Private/Graph/PCGGraphCompiler.cpp
+++ b/Engine/Plugins/PCG/Source/PCG/Private/Graph/PCGGraphCompiler.cpp
@@ -78,13 +78,17 @@ namespace PCGGraphCompiler
 UPCGComputeGraph* FPCGGraphCompilerCache::GetCompiledComputeGraph(const UPCGGraph* InGraph, uint32 GridSize, uint32 ComputeGraphIndex)
 {
 #if WITH_EDITOR
-	if (TMap<uint32, TArray<TObjectPtr<UPCGComputeGraph>>>* GridSizeToComputeGraphs = TopGraphToComputeGraphMap.Find(InGraph))
 	{
-		if (TArray<TObjectPtr<UPCGComputeGraph>>* ComputeGraphs = GridSizeToComputeGraphs->Find(GridSize))
+		FReadScopeLock Lock(GraphToTaskMapLock);
+
+		if (TMap<uint32, TArray<TObjectPtr<UPCGComputeGraph>>>* GridSizeToComputeGraphs = TopGraphToComputeGraphMap.Find(InGraph))
 		{
-			if (ComputeGraphs->IsValidIndex(ComputeGraphIndex))
+			if (TArray<TObjectPtr<UPCGComputeGraph>>* ComputeGraphs = GridSizeToComputeGraphs->Find(GridSize))
 			{
-				return (*ComputeGraphs)[ComputeGraphIndex];
+				if (ComputeGraphs->IsValidIndex(ComputeGraphIndex))
+				{
+					return (*ComputeGraphs)[ComputeGraphIndex];
+				}
 			}
 		}
 	}
@@ -960,7 +964,7 @@ void FPCGGraphCompiler::AddReferencedObjects(FReferenceCollector& Collector)
 {
 	TRACE_CPUPROFILER_EVENT_SCOPE(FCPGGraphCompiler::AddReferencedObjects);
 
-	TMap<UPCGGraph*, TMap<uint32, TArray<TObjectPtr<UPCGComputeGraph>>>> TopGraphToComputeGraphMap;
+	FReadScopeLock Lock(Cache.GraphToTaskMapLock);
 
 	for (const TPair<UPCGGraph*, TMap<uint32, TArray<TObjectPtr<UPCGComputeGraph>>>>& GraphToGridSizes : Cache.TopGraphToComputeGraphMap)
 	{
diff --git a/Engine/Plugins/PCG/Source/PCG/Private/Graph/PCGGraphCompilerGPU.cpp b/Engine/Plugins/PCG/Source/PCG/Private/Graph/PCGGraphCompilerGPU.cpp
--- a/Engine/Plugins/PCG/Source/PCG/Private/Graph/PCGGraphCompilerGPU.cpp
+++ b/Engine/Plugins/PCG/Source/PCG/Private/Graph/PCGGraphCompilerGPU.cpp
@@ -514,6 +514,7 @@ void FPCGGraphCompilerGPU::BuildGPUGraphTask(
 	FPCGGPUCompilationContext& InOutContext,
 	UPCGGraph* InGraph,
 	uint32 InGridSize,
+	uint32 InComputeGraphIndex,
 	FPCGTaskId InGPUGraphTaskId,
 	const TSet<FPCGTaskId>& InCollapsedTasks,
 	const TSet<FPCGTaskId>& InAllGPUCompatibleTasks,
@@ -525,11 +526,68 @@ void FPCGGraphCompilerGPU::BuildGPUGraphTask(
 	TRACE_CPUPROFILER_EVENT_SCOPE(FPCGGraphCompilerGPU::BuildGPUGraphTask);
 
 	check(InGraph);
-	const FName GraphName = MakeUniqueObjectName(InGraph, UPCGComputeGraph::StaticClass(), InGraph->GetFName());
 
-	UObject* ComputeGraphOuter = InOutContext.GetGraphCompiler().IsCooking() ? InGraph : Cast<UObject>(GetTransientPackage());
+	FPCGGraphCompiler& GraphCompiler = InOutContext.GetGraphCompiler();
+	UPCGComputeGraph* ComputeGraph = GraphCompiler.GetComputeGraph(InGraph, InGridSize, InComputeGraphIndex);
+
+	// If the graph does not exist already, create it.
+	if (!ComputeGraph)
+	{
+		// TODO: There is a race condition here where two threads can try to find this graph at the same time and both fail, so both create the graph.
+		// This is fine, as they'll both create the same graph and place it at the same compute graph index. However it would be ideal perf-wise to
+		// avoid this, as creating the compute graph can be expensive.
+		ComputeGraph = CompileComputeGraph(
+			InOutContext,
+			InGraph,
+			InCollapsedTasks,
+			InAllGPUCompatibleTasks,
+			InTaskSuccessors,
+			InOutCompiledTasks,
+			InOriginalToVirtualPin,
+			InOutputCPUPinToVirtualPin);
+	}
+
+	const int32 ComputeGraphIndex = InOutContext.AddCompiledComputeGraph(ComputeGraph);
+	ensure(ComputeGraphIndex == InComputeGraphIndex);
+
+	if (GraphCompiler.IsCooking())
+	{
+		TObjectPtr<UPCGComputeGraphSettings> Settings = NewObject<UPCGComputeGraphSettings>(InGraph);
+		Settings->ComputeGraphIndex = ComputeGraphIndex;
+		InOutCompiledTasks[InGPUGraphTaskId].ElementSource = EPCGElementSource::FromCookedSettings;
+		InOutCompiledTasks[InGPUGraphTaskId].CookedSettings = Settings;
+	}
+	else
+	{
+		InOutCompiledTasks[InGPUGraphTaskId].Element = MakeShared<FPCGComputeGraphElement>(ComputeGraphIndex);
+	}
+
+	if (FApp::CanEverRender() && IsInGameThread() && ensure(ComputeGraph))
+	{
+		TRACE_CPUPROFILER_EVENT_SCOPE(UpdateResources);
 
+		// Compile shader resources and create render proxies.
+		ComputeGraph->UpdateResources();
+	}
+}
+
+UPCGComputeGraph* FPCGGraphCompilerGPU::CompileComputeGraph(
+	FPCGGPUCompilationContext& InOutContext,
+	UPCGGraph* InGraph,
+	const TSet<FPCGTaskId>& InCollapsedTasks,
+	const TSet<FPCGTaskId>& InAllGPUCompatibleTasks,
+	const FTaskToSuccessors& InTaskSuccessors,
+	TArray<FPCGGraphTask>& InOutCompiledTasks,
+	const FOriginalToVirtualPin& InOriginalToVirtualPin,
+	const TMap<TSoftObjectPtr<const UPCGPin>, FName>& InOutputCPUPinToVirtualPin)
+{
+	TRACE_CPUPROFILER_EVENT_SCOPE(FPCGGraphCompilerGPU::CompileComputeGraph);
+
+	// Create a new compute graph.
+	UObject* ComputeGraphOuter = InOutContext.GetGraphCompiler().IsCooking() ? InGraph : Cast<UObject>(GetTransientPackage());
+	const FName GraphName = MakeUniqueObjectName(InGraph, UPCGComputeGraph::StaticClass(), InGraph->GetFName());
 	UPCGComputeGraph* ComputeGraph = InOutContext.NewObject_AnyThread<UPCGComputeGraph>(ComputeGraphOuter, GraphName);
+
 	ComputeGraph->OutputCPUPinToInputGPUPinAlias = InOutputCPUPinToVirtualPin;
 
 	// Not incredibly useful for us - DG adds GetComponentSource()->GetComponentClass() object which allows it to bind at execution time by class.
@@ -1129,45 +1187,21 @@ void FPCGGraphCompilerGPU::BuildGPUGraphTask(
 		}
 	}
 
-	TMap<uint32, TArray<TObjectPtr<UPCGComputeGraph>>>& GridToComputeGraphs = InOutContext.GetGraphCompiler().GetCache().TopGraphToComputeGraphMap.FindOrAdd(InGraph);
-	TArray<TObjectPtr<UPCGComputeGraph>>& ComputeGraphs = GridToComputeGraphs.FindOrAdd(InGridSize);
-	const uint32 ComputeGraphIndex = ComputeGraphs.Num();
-	ComputeGraphs.Add(ComputeGraph);
-
-	if (InOutContext.GetGraphCompiler().IsCooking())
-	{
-		TObjectPtr<UPCGComputeGraphSettings> Settings = NewObject<UPCGComputeGraphSettings>(InGraph);
-		Settings->ComputeGraphIndex = ComputeGraphIndex;
-		InOutCompiledTasks[InGPUGraphTaskId].ElementSource = EPCGElementSource::FromCookedSettings;
-		InOutCompiledTasks[InGPUGraphTaskId].CookedSettings = Settings;
-	}
-	else
-	{
-		InOutCompiledTasks[InGPUGraphTaskId].Element = MakeShared<FPCGComputeGraphElement>(ComputeGraphIndex);
-	}
-
-	if (FApp::CanEverRender())
-	{
-		TRACE_CPUPROFILER_EVENT_SCOPE(UpdateResources);
-
-		// Compile shader resources and create render proxies.
-		ComputeGraph->UpdateResources();
-	}
+	return ComputeGraph;
 }
 
 void FPCGGraphCompilerGPU::CreateGPUNodes(FPCGGraphCompiler& InOutCompiler, UPCGGraph* InGraph, uint32 InGridSize, TArray<FPCGGraphTask>& InOutCompiledTasks)
 {
 	TRACE_CPUPROFILER_EVENT_SCOPE(FPCGGraphCompilerGPU::CreateGPUNodes);
 
-	FPCGGPUCompilationContext Context(InOutCompiler);
-
-	// Clear out any previous compute graphs for this grid level before adding new ones.
-	if (ensure(InGraph))
+	if (!InGraph)
 	{
-		TMap<uint32, TArray<TObjectPtr<UPCGComputeGraph>>>& ComputeGraphs = InOutCompiler.GetCache().TopGraphToComputeGraphMap.FindOrAdd(InGraph);
-		ComputeGraphs.FindOrAdd(InGridSize).Reset();
+		ensure(false);
+		return;
 	}
 
+	FPCGGPUCompilationContext Context(InOutCompiler);
+
 	TSet<FPCGTaskId> GPUCompatibleTaskIds;
 	GPUCompatibleTaskIds.Reserve(InOutCompiledTasks.Num());
 	for (FPCGTaskId TaskId = 0; TaskId < InOutCompiledTasks.Num(); ++TaskId)
@@ -1203,6 +1237,8 @@ void FPCGGraphCompilerGPU::CreateGPUNodes(FPCGGraphCompiler& InOutCompiler, UPCG
 	TArray<TSet<FPCGTaskId>> NodeSubsetsToConvertToCFGraph;
 	CollectGPUNodeSubsets(InOutCompiledTasks, TaskSuccessors, GPUCompatibleTaskIds, NodeSubsetsToConvertToCFGraph);
 
+	const uint32 NumComputeGraphs = NodeSubsetsToConvertToCFGraph.Num();
+
 	// Mapping from task ID & pin label to a virtual pin label. Compute graphs are executed within a generated element,
 	// and the input and output pins of this element must have unique virtual pin labels so that we can parse the data that
 	// PCG provides through the input data collection correctly, and route the output data to the downstream pins correctly.
@@ -1210,11 +1246,11 @@ void FPCGGraphCompilerGPU::CreateGPUNodes(FPCGGraphCompiler& InOutCompiler, UPCG
 
 	TArray<TMap<TSoftObjectPtr<const UPCGPin>, FName>> OutputCPUPinToVirtualPin;
 
-	OriginalToVirtualPin.SetNum(NodeSubsetsToConvertToCFGraph.Num());
-	OutputCPUPinToVirtualPin.SetNum(NodeSubsetsToConvertToCFGraph.Num());
+	OriginalToVirtualPin.SetNum(NumComputeGraphs);
+	OutputCPUPinToVirtualPin.SetNum(NumComputeGraphs);
 
 	// Setup mappings from existing pins to compute graph element virtual pins as a prestep before wiring in the compute graph tasks.
-	for (int32 ComputeGraphIndex = 0; ComputeGraphIndex < NodeSubsetsToConvertToCFGraph.Num(); ++ComputeGraphIndex)
+	for (uint32 ComputeGraphIndex = 0; ComputeGraphIndex < NumComputeGraphs; ++ComputeGraphIndex)
 	{
 		SetupVirtualPins(
 			NodeSubsetsToConvertToCFGraph[ComputeGraphIndex],
@@ -1226,7 +1262,7 @@ void FPCGGraphCompilerGPU::CreateGPUNodes(FPCGGraphCompiler& InOutCompiler, UPCG
 
 	// Do actual collapsing now, one subset at a time. Each collapse will do all fixup of task ids? That will invalidate
 	// ids in NodeSubsetsToConvertToCFGraph, so may need remap table. But can ignore this for now.
-	for (int32 ComputeGraphIndex = 0; ComputeGraphIndex < NodeSubsetsToConvertToCFGraph.Num(); ++ComputeGraphIndex)
+	for (uint32 ComputeGraphIndex = 0; ComputeGraphIndex < NumComputeGraphs; ++ComputeGraphIndex)
 	{
 		TSet<FPCGTaskId>& NodeSubsetToConvertToCFGraph = NodeSubsetsToConvertToCFGraph[ComputeGraphIndex];
 
@@ -1263,6 +1299,7 @@ void FPCGGraphCompilerGPU::CreateGPUNodes(FPCGGraphCompiler& InOutCompiler, UPCG
 			Context,
 			InGraph,
 			InGridSize,
+			ComputeGraphIndex,
 			ComputeGraphTaskId,
 			NodeSubsetToConvertToCFGraph,
 			GPUCompatibleTaskIds,
@@ -1272,6 +1309,17 @@ void FPCGGraphCompilerGPU::CreateGPUNodes(FPCGGraphCompiler& InOutCompiler, UPCG
 			OutputCPUPinToVirtualPin[ComputeGraphIndex]);
 	}
 
+	{
+		FWriteScopeLock Lock(InOutCompiler.GetCache().GraphToTaskMapLock);
+		TMap<uint32, TArray<TObjectPtr<UPCGComputeGraph>>>& GridSizeToComputeGraphs = InOutCompiler.GetCache().TopGraphToComputeGraphMap.FindOrAdd(InGraph);
+		TArray<TObjectPtr<UPCGComputeGraph>>& ComputeGraphs = GridSizeToComputeGraphs.FindOrAdd(InGridSize);
+
+		// Replace any existing compute graphs with the newly compiled ones. It's okay if multiple threads do this, because
+		// compute graph index order should be deterministic, so different threads will produce the same results.
+		ComputeGraphs = MoveTemp(Context.GetCompiledComputeGraphs());
+		ensure(ComputeGraphs.Num() == NumComputeGraphs);
+	}
+
 	// Now cull all the GPU compatible nodes. The compute graph task are already wired in so we're fine to just delete.
 	FPCGGraphCompiler::CullTasks(InOutCompiledTasks, /*bAddPassthroughWires=*/false, [&NodeSubsetsToConvertToCFGraph](const FPCGGraphTask& InTask)
 	{
diff --git a/Engine/Plugins/PCG/Source/PCG/Private/Graph/PCGGraphCompilerGPU.h b/Engine/Plugins/PCG/Source/PCG/Private/Graph/PCGGraphCompilerGPU.h
--- a/Engine/Plugins/PCG/Source/PCG/Private/Graph/PCGGraphCompilerGPU.h
+++ b/Engine/Plugins/PCG/Source/PCG/Private/Graph/PCGGraphCompilerGPU.h
@@ -5,6 +5,7 @@
 #include "PCGCommon.h"
 
 class FPCGGraphCompiler;
+class UPCGComputeGraph;
 class UPCGGraph;
 class UPCGPin;
 struct FPCGGraphTask;
@@ -65,6 +66,7 @@ private:
 		FPCGGPUCompilationContext& InOutContext,
 		UPCGGraph* InGraph,
 		uint32 InGridSize,
+		uint32 InComputeGraphIndex,
 		FPCGTaskId InGPUGraphTaskId,
 		const TSet<FPCGTaskId>& InCollapsedTasks,
 		const TSet<FPCGTaskId>& InAllGPUCompatibleTasks,
@@ -72,6 +74,15 @@ private:
 		TArray<FPCGGraphTask>& InOutCompiledTasks,
 		const FOriginalToVirtualPin& InOriginalToVirtualPin,
 		const TMap<TSoftObjectPtr<const UPCGPin>, FName>& InOutputCPUPinToVirtualPin);
-	
+
+	static UPCGComputeGraph* CompileComputeGraph(
+		FPCGGPUCompilationContext& InOutContext,
+		UPCGGraph* InGraph,
+		const TSet<FPCGTaskId>& InCollapsedTasks,
+		const TSet<FPCGTaskId>& InAllGPUCompatibleTasks,
+		const FTaskToSuccessors& InTaskSuccessors,
+		TArray<FPCGGraphTask>& InOutCompiledTasks,
+		const FOriginalToVirtualPin& InOriginalToVirtualPin,
+		const TMap<TSoftObjectPtr<const UPCGPin>, FName>& InOutputCPUPinToVirtualPin);
 };
 #endif // WITH_EDITOR
diff --git a/Engine/Plugins/PCG/Source/PCG/Private/PCGComponent.cpp b/Engine/Plugins/PCG/Source/PCG/Private/PCGComponent.cpp
--- a/Engine/Plugins/PCG/Source/PCG/Private/PCGComponent.cpp
+++ b/Engine/Plugins/PCG/Source/PCG/Private/PCGComponent.cpp
@@ -1535,24 +1535,10 @@ void UPCGComponent::BeginPlay()
 		GetSubsystem()->RegisterOrUpdatePCGComponent(this);
 	}
 
-	if(bActivated && !bGenerated && GenerationTrigger == EPCGComponentGenerationTrigger::GenerateOnLoad)
+	if (bActivated && !bGenerated && GenerationTrigger == EPCGComponentGenerationTrigger::GenerateOnLoad)
 	{
-		if (IsPartitioned())
-		{
-			// If we are partitioned, the responsibility of the generation is to the partition actors.
-			// but we still need to know that we are currently generated (even if the state is held by the partition actors)
-			// TODO: Will be cleaner when we have dynamic association.
-			const FBox NewBounds = GetGridBounds();
-			if (NewBounds.IsValid)
-			{
-				PostProcessGraph(NewBounds, true, nullptr);
-			}
-		}
-		else
-		{
-			GenerateInternal(/*bForce=*/false, EPCGHiGenGrid::Uninitialized, EPCGComponentGenerationTrigger::GenerateOnLoad, {});
-			bRuntimeGenerated = true;
-		}
+		GenerateInternal(/*bForce=*/false, EPCGHiGenGrid::Uninitialized, EPCGComponentGenerationTrigger::GenerateOnLoad, {});
+		bRuntimeGenerated = true;
 	}
 }
 
@@ -2745,7 +2731,7 @@ void UPCGComponent::OnRefresh(bool bForceRefresh)
 	// Only redo the mapping if we are generated
 	UPCGSubsystem* Subsystem = GetSubsystem();
 	const bool bWasGenerated = bGenerated;
-	const bool bWasGeneratedOrGenerating = bWasGenerated || bForceRefresh;
+	const bool bWasGeneratedOrGenerating = bWasGenerated || bForceRefresh || IsGenerating();
 
 	// If we are partitioned but we have resources, we need to force a cleanup
 	if (IsPartitioned() && !GeneratedResources.IsEmpty())
diff --git a/Engine/Plugins/PCG/Source/PCG/Private/PCGGraph.cpp b/Engine/Plugins/PCG/Source/PCG/Private/PCGGraph.cpp
--- a/Engine/Plugins/PCG/Source/PCG/Private/PCGGraph.cpp
+++ b/Engine/Plugins/PCG/Source/PCG/Private/PCGGraph.cpp
@@ -587,7 +587,9 @@ void UPCGGraph::PreSave(FObjectPreSaveContext ObjectSaveContext)
 			}
 		}
 
-		if (ensure(CompiledComputeGraphs))
+		// Note: We don't have an ensure on the CompiledComputeGraphs like the other compiled data since graphs that do not
+		// produce compute graphs will never create an entry in this mapping.
+		if (CompiledComputeGraphs)
 		{
 			CookedCompilationData->ComputeGraphs.Reserve(CompiledComputeGraphs->Num());
 
