[PCG] Rename Attribute was not setting the last selector correctly

#rnx

diff --git a/Engine/Plugins/PCG/Source/PCG/Private/Metadata/PCGMetadata.cpp b/Engine/Plugins/PCG/Source/PCG/Private/Metadata/PCGMetadata.cpp
--- a/Engine/Plugins/PCG/Source/PCG/Private/Metadata/PCGMetadata.cpp
+++ b/Engine/Plugins/PCG/Source/PCG/Private/Metadata/PCGMetadata.cpp
@@ -896,6 +896,9 @@ bool UPCGMetadata::RenameAttribute(FName AttributeToRename, FName NewAttributeNa
 		RemoveAttributeInternal(AttributeToRename);
 		Attribute->Name = NewAttributeName;
 		AddAttributeInternal(NewAttributeName, Attribute);
+
+		// Also when renaming an attribute, notify the PCG Data owner that the latest attribute manipulated is this one.
+		SetLastCachedSelectorOnOwner(NewAttributeName);
 		
 		bRenamed = true;
 	}
