[PCG] Change the behavior of the match and set attribute in the slicers

* If the attribute already exists in the output, override it (instead of ignoring it)
* Also update tooltip for properties to display the right name.

#rnx

diff --git a/Engine/Plugins/PCG/Source/PCG/Private/Elements/Grammar/PCGSubdivisionBase.cpp b/Engine/Plugins/PCG/Source/PCG/Private/Elements/Grammar/PCGSubdivisionBase.cpp
--- a/Engine/Plugins/PCG/Source/PCG/Private/Elements/Grammar/PCGSubdivisionBase.cpp
+++ b/Engine/Plugins/PCG/Source/PCG/Private/Elements/Grammar/PCGSubdivisionBase.cpp
@@ -394,10 +394,10 @@ bool FPCGSubdivisionBaseElement::MatchAndSetAttributes(const TArray<FPCGTaggedDa
 					continue;
 				}
 
-				// Skip any attribute that already exists in the output data
+				// If the attribute already exists, override it
 				if (OutputMetadata->HasAttribute(AttributeName))
 				{
-					continue;
+					OutputMetadata->DeleteAttribute(AttributeName);
 				}
 
 				InAttributes.Add(InputMetadata->GetConstAttribute(AttributeName));
diff --git a/Engine/Plugins/PCG/Source/PCG/Public/Elements/Grammar/PCGSubdivisionBase.h b/Engine/Plugins/PCG/Source/PCG/Public/Elements/Grammar/PCGSubdivisionBase.h
--- a/Engine/Plugins/PCG/Source/PCG/Public/Elements/Grammar/PCGSubdivisionBase.h
+++ b/Engine/Plugins/PCG/Source/PCG/Public/Elements/Grammar/PCGSubdivisionBase.h
@@ -124,21 +124,21 @@ public:
 	UPROPERTY(BlueprintReadWrite, EditAnywhere, Category = "Settings|Extra Output Attributes", meta = (PCG_Overridable, InlineEditConditionToggle))
 	bool bOutputSizeAttribute = true;
 
-	/** Name of the Size output attribute name, ignored if match and set from module info is true. */
+	/** Name of the Size output attribute name, ignored if Forward Attributes From Modules Info is true. */
 	UPROPERTY(BlueprintReadWrite, EditAnywhere, Category = "Settings|Extra Output Attributes", meta = (PCG_Overridable, EditCondition = "bOutputSizeAttribute"))
 	FName SizeAttributeName = PCGSubdivisionBase::Constants::SizeAttributeName;
 
 	UPROPERTY(BlueprintReadWrite, EditAnywhere, Category = "Settings|Extra Output Attributes", meta = (PCG_Overridable, InlineEditConditionToggle))
 	bool bOutputScalableAttribute = true;
 
-	/** Name of the Scalable output attribute name, ignored if match and set from module info is true. */
+	/** Name of the Scalable output attribute name, ignored if Forward Attributes From Modules Info is true. */
 	UPROPERTY(BlueprintReadWrite, EditAnywhere, Category = "Settings|Extra Output Attributes", meta = (PCG_Overridable, EditCondition = "bOutputScalableAttribute"))
 	FName ScalableAttributeName = PCGSubdivisionBase::Constants::ScalableAttributeName;
 
 	UPROPERTY(BlueprintReadWrite, EditAnywhere, Category = "Settings|Extra Output Attributes", meta = (PCG_Overridable, InlineEditConditionToggle))
 	bool bOutputDebugColorAttribute = false;
 
-	/** Name of the Debug Color output attribute name, ignored if match and set from module info is true. */
+	/** Name of the Debug Color output attribute name, ignored if Forward Attributes From Modules Info is true. */
 	UPROPERTY(BlueprintReadWrite, EditAnywhere, Category = "Settings|Extra Output Attributes", meta = (PCG_Overridable, EditCondition = "bOutputDebugColorAttribute"))
 	FName DebugColorAttributeName = PCGSubdivisionBase::Constants::DebugColorAttributeName;
 
