[PCG] Fix the issue where only the first component gets extracted in Get Actor Property

* Add an option to process all components + add component reference in order to filter them afterwards.

#rnx

diff --git a/Engine/Plugins/PCG/Source/PCG/Private/Elements/PCGGetActorProperty.cpp b/Engine/Plugins/PCG/Source/PCG/Private/Elements/PCGGetActorProperty.cpp
--- a/Engine/Plugins/PCG/Source/PCG/Private/Elements/PCGGetActorProperty.cpp
+++ b/Engine/Plugins/PCG/Source/PCG/Private/Elements/PCGGetActorProperty.cpp
@@ -12,10 +12,19 @@
 #include "Helpers/PCGPropertyHelpers.h"
 #include "Helpers/PCGHelpers.h"
 #include "Metadata/PCGAttributePropertySelector.h"
-#include "Metadata/Accessors/PCGAttributeAccessorHelpers.h"
 
 #define LOCTEXT_NAMESPACE "PCGPropertyToParamDataElement"
 
+UPCGGetActorPropertySettings::UPCGGetActorPropertySettings(const FObjectInitializer& ObjectInitializer)
+{
+	// To not break previous behavior, we keep this false by default, but will force it to true for new nodes as it makes more sense to process them all,
+	// as we have no guarantee on the order components are found.
+	if (PCGHelpers::IsNewObjectAndNotDefault(this))
+	{
+		bProcessAllComponents = true;
+	}
+}
+
 #if WITH_EDITOR
 void UPCGGetActorPropertySettings::GetStaticTrackedKeys(FPCGSelectionKeyToSettingsMap& OutKeysToSettings, TArray<TObjectPtr<const UPCGGraph>>& OutVisitedGraphs) const
 {
@@ -255,69 +264,92 @@ bool FPCGGetActorPropertyElement::ExecuteInternal(FPCGContext* InContext) const
 		}
 
 		// From there, we either check the actor, or the component attached to it.
-		UObject* ObjectToInspect = FoundActor;
+		TArray<UObject*> ObjectsToInspect;
 		if (Settings->bSelectComponent)
 		{
-			ObjectToInspect = FoundActor->GetComponentByClass(Settings->ComponentClass);
-			if (!ObjectToInspect)
+			TInlineComponentArray<UActorComponent*> ActorComponents;
+			FoundActor->GetComponents(Settings->ComponentClass, ActorComponents);
+			if (ActorComponents.IsEmpty())
 			{
 				PCGE_LOG(Error, GraphAndLog, FText::Format(LOCTEXT("ComponentDoesNotExist", "Component class '{0}' does not exist in the found actor {1}"), FText::FromString(Settings->ComponentClass->GetName()), FText::FromString(FoundActor->GetName())));
 				return true;
 			}
-		}
-
-		PCGPropertyHelpers::FExtractorParameters Parameters(ObjectToInspect, ObjectToInspect->GetClass(), Settings->PropertyName.ToString(), Settings->OutputAttributeName, Settings->bForceObjectAndStructExtraction, /*bPropertyNeedsToBeVisible=*/true);
 
-		// Don't care for object traversed in non-editor build, since it is only useful for tracking.
-		TSet<FSoftObjectPath>* ObjectTraversedPtr = nullptr;
-#if WITH_EDITOR
-		TSet<FSoftObjectPath> ObjectTraversed;
-		ObjectTraversedPtr = &ObjectTraversed;
-#endif // WITH_EDITOR
-
-		if (UPCGParamData* ParamData = PCGPropertyHelpers::ExtractPropertyAsAttributeSet(Parameters, Context, ObjectTraversedPtr))
-		{
-			TArray<FPCGTaggedData>& Outputs = Context->OutputData.TaggedData;
-			FPCGTaggedData& Output = Outputs.Emplace_GetRef();
-			Output.Data = ParamData;
-			
-			for (FName ActorTag : FoundActor->Tags)
+			if (Settings->bProcessAllComponents)
 			{
-				Output.Tags.Add(ActorTag.ToString());
+				ObjectsToInspect.Append(ActorComponents);
 			}
-
-			if (Settings->bOutputActorReference)
+			else
 			{
-				check(ParamData->MutableMetadata());
-				ParamData->MutableMetadata()->FindOrCreateAttribute(PCGPointDataConstants::ActorReferenceAttribute, FSoftObjectPath(FoundActor), /*bAllowsInterpolation=*/false, /*bOverrideParent=*/false, /*bOverwriteIfTypeMismatch=*/true);
+				ObjectsToInspect.Add(ActorComponents[0]);
 			}
 		}
 		else
 		{
-			if(Parameters.PropertySelectors.IsEmpty() || Parameters.PropertySelectors[0].GetName() == NAME_None)
+			ObjectsToInspect.Add(FoundActor);
+		}
+
+		for (UObject* ObjectToInspect : ObjectsToInspect)
+		{
+			PCGPropertyHelpers::FExtractorParameters Parameters(ObjectToInspect, ObjectToInspect->GetClass(), Settings->PropertyName.ToString(), Settings->OutputAttributeName, Settings->bForceObjectAndStructExtraction, /*bPropertyNeedsToBeVisible=*/true);
+
+			// Don't care for object traversed in non-editor build, since it is only useful for tracking.
+			TSet<FSoftObjectPath>* ObjectTraversedPtr = nullptr;
+#if WITH_EDITOR
+			TSet<FSoftObjectPath> ObjectTraversed;
+			ObjectTraversedPtr = &ObjectTraversed;
+#endif // WITH_EDITOR
+
+			if (UPCGParamData* ParamData = PCGPropertyHelpers::ExtractPropertyAsAttributeSet(Parameters, Context, ObjectTraversedPtr))
 			{
-				PCGE_LOG(Error, GraphAndLog, FText::Format(LOCTEXT("FailedToExtractActor", "Fail to extract actor {0}."), FText::FromString(FoundActor->GetName())));
+				TArray<FPCGTaggedData>& Outputs = Context->OutputData.TaggedData;
+				FPCGTaggedData& Output = Outputs.Emplace_GetRef();
+				Output.Data = ParamData;
+			
+				for (FName ActorTag : FoundActor->Tags)
+				{
+					Output.Tags.Add(ActorTag.ToString());
+				}
+
+				if (Settings->bOutputActorReference)
+				{
+					check(ParamData->MutableMetadata());
+					ParamData->MutableMetadata()->FindOrCreateAttribute(PCGPointDataConstants::ActorReferenceAttribute, FSoftObjectPath(FoundActor), /*bAllowsInterpolation=*/false, /*bOverrideParent=*/false, /*bOverwriteIfTypeMismatch=*/true);
+				}
+				
+				if (Settings->bSelectComponent && Settings->bOutputComponentReference)
+				{
+					check(ParamData->MutableMetadata());
+					ParamData->MutableMetadata()->FindOrCreateAttribute(TEXT("ComponentReference"), FSoftObjectPath(ObjectToInspect), /*bAllowsInterpolation=*/false, /*bOverrideParent=*/false, /*bOverwriteIfTypeMismatch=*/true);
+				}
 			}
 			else
 			{
-				PCGE_LOG(Error, GraphAndLog, FText::Format(LOCTEXT("FailedToExtract", "Fail to extract the property '{0}' on actor {1}."), Parameters.PropertySelectors[0].GetDisplayText(), FText::FromString(FoundActor->GetName())));
+				if(Parameters.PropertySelectors.IsEmpty() || Parameters.PropertySelectors[0].GetName() == NAME_None)
+				{
+					PCGE_LOG(Error, GraphAndLog, FText::Format(LOCTEXT("FailedToExtractActor", "Fail to extract actor {0}."), FText::FromString(FoundActor->GetName())));
+				}
+				else
+				{
+					PCGE_LOG(Error, GraphAndLog, FText::Format(LOCTEXT("FailedToExtract", "Fail to extract the property '{0}' on actor {1}."), Parameters.PropertySelectors[0].GetDisplayText(), FText::FromString(FoundActor->GetName())));
+				}
 			}
-		}
 
-		// Register dynamic tracking
+			// Register dynamic tracking
 #if WITH_EDITOR
-		if (!ObjectTraversed.IsEmpty())
-		{
-			FPCGDynamicTrackingHelper DynamicTracking;
-			DynamicTracking.EnableAndInitialize(Context, ObjectTraversed.Num());
-			for (const FSoftObjectPath& Path : ObjectTraversed)
+			if (!ObjectTraversed.IsEmpty())
 			{
-				DynamicTracking.AddToTracking(FPCGSelectionKey::CreateFromPath(Path), /*bCulled=*/false);
+				FPCGDynamicTrackingHelper DynamicTracking;
+				DynamicTracking.EnableAndInitialize(Context, ObjectTraversed.Num());
+				for (const FSoftObjectPath& Path : ObjectTraversed)
+				{
+					DynamicTracking.AddToTracking(FPCGSelectionKey::CreateFromPath(Path), /*bCulled=*/false);
+				}
+
+				DynamicTracking.Finalize(Context);
 			}
-
-			DynamicTracking.Finalize(Context);
-		}
 #endif // WITH_EDITOR
+		}
 	}
 
 	return true;
diff --git a/Engine/Plugins/PCG/Source/PCG/Public/Elements/PCGGetActorProperty.h b/Engine/Plugins/PCG/Source/PCG/Public/Elements/PCGGetActorProperty.h
--- a/Engine/Plugins/PCG/Source/PCG/Public/Elements/PCGGetActorProperty.h
+++ b/Engine/Plugins/PCG/Source/PCG/Public/Elements/PCGGetActorProperty.h
@@ -20,6 +20,8 @@ class UPCGGetActorPropertySettings : public UPCGSettings
 	GENERATED_BODY()
 
 public:
+	UPCGGetActorPropertySettings(const FObjectInitializer& ObjectInitializer);
+	
 	//~Begin UObject interface
 	virtual void PostLoad() override;
 #if WITH_EDITOR
@@ -54,13 +56,21 @@ public:
 	FPCGActorSelectorSettings ActorSelector;
 
 	/** Allow to look for an actor component instead of an actor. It will need to be attached to the found actor. */
-	UPROPERTY(BlueprintReadWrite, EditAnywhere, Category = Settings)
+	UPROPERTY(BlueprintReadWrite, EditAnywhere, Category = Settings, meta = (PCG_Overridable))
 	bool bSelectComponent = false;
 
 	/** If we are looking for an actor component, the class can be specified here. */
-	UPROPERTY(BlueprintReadWrite, EditAnywhere, Category = Settings, meta = (EditCondition = "bSelectComponent", EditConditionHides))
+	UPROPERTY(BlueprintReadWrite, EditAnywhere, Category = Settings, meta = (PCG_Overridable, EditCondition = "bSelectComponent", EditConditionHides))
 	TSubclassOf<UActorComponent> ComponentClass;
 
+	/** Process all Actor components. If not set, only the first component found will be processed. */
+	UPROPERTY(BlueprintReadWrite, EditAnywhere, Category = Settings, meta = (PCG_Overridable, EditCondition = "bSelectComponent", EditConditionHides))
+	bool bProcessAllComponents = false;
+	
+	/** Controls whether a component reference attribute will be added to the result */
+	UPROPERTY(BlueprintReadWrite, EditAnywhere, Category = Settings, meta = (PCG_Overridable, EditCondition = "bSelectComponent", EditConditionHides))
+	bool bOutputComponentReference = false;
+
 	/** Property name to extract. Can only extract properties that are compatible with metadata types. If None, extract the actor/component directly. Can be a comma-separated list, assuming they have the same cardinality. */
 	UPROPERTY(BlueprintReadWrite, EditAnywhere, Category = Settings, meta = (PCG_Overridable))
 	FName PropertyName = NAME_None;
