[PCG] GPU SM Spawner - fix seed computation in 'weighted' selection mode so that result matches CPU.

#rnx

diff --git a/Engine/Plugins/PCG/Shaders/Private/Elements/PCGStaticMeshSpawner.usf b/Engine/Plugins/PCG/Shaders/Private/Elements/PCGStaticMeshSpawner.usf
--- a/Engine/Plugins/PCG/Shaders/Private/Elements/PCGStaticMeshSpawner.usf
+++ b/Engine/Plugins/PCG/Shaders/Private/Elements/PCGStaticMeshSpawner.usf
@@ -54,8 +54,13 @@ void Main(uint3 GroupId : SV_GroupID, uint GroupIndex : SV_GroupIndex)
 		else
 		{
 			// Randomly select a primitive by inverting the cumulative distribution function.
-			uint Seed = In_GetSeed(In_DataIndex, In_ElementIndex);
-			Seed = ComputeSeed(Seed, GetSeed());
+			
+			// Exactly matches UPCGBlueprintHelpers::GetRandomStreamFromPoint - three seeds passed in to consecutive arguments of ComputeSeed
+			// (argument count & order matters).
+			const uint SeedPoint = In_GetSeed(In_DataIndex, In_ElementIndex);
+			const uint SeedSettings = GetSettingsSeed();
+			const uint SeedComponent = GetComponentSeed();
+			const uint Seed = ComputeSeed(SeedPoint, SeedSettings, SeedComponent);
 
 			const float RandomDraw = FRand(Seed);
 
diff --git a/Engine/Plugins/PCG/Source/PCG/Private/Compute/DataInterfaces/PCGCustomKernelDataInterface.cpp b/Engine/Plugins/PCG/Source/PCG/Private/Compute/DataInterfaces/PCGCustomKernelDataInterface.cpp
--- a/Engine/Plugins/PCG/Source/PCG/Private/Compute/DataInterfaces/PCGCustomKernelDataInterface.cpp
+++ b/Engine/Plugins/PCG/Source/PCG/Private/Compute/DataInterfaces/PCGCustomKernelDataInterface.cpp
@@ -25,6 +25,14 @@ void UPCGCustomKernelDataInterface::GetSupportedInputs(TArray<FShaderFunctionDef
 		.SetName(TEXT("GetSeed"))
 		.AddReturnType(FShaderValueType::Get(EShaderFundamentalType::Uint));
 
+	OutFunctions.AddDefaulted_GetRef()
+		.SetName(TEXT("GetSettingsSeed"))
+		.AddReturnType(FShaderValueType::Get(EShaderFundamentalType::Uint));
+
+	OutFunctions.AddDefaulted_GetRef()
+		.SetName(TEXT("GetComponentSeed"))
+		.AddReturnType(FShaderValueType::Get(EShaderFundamentalType::Uint));
+
 	// A convenient way to serve component bounds to all kernels. Could be pulled out into a PCG context DI in the future.
 	OutFunctions.AddDefaulted_GetRef()
 		.SetName(TEXT("GetComponentBoundsMin"))
@@ -38,6 +46,8 @@ void UPCGCustomKernelDataInterface::GetSupportedInputs(TArray<FShaderFunctionDef
 BEGIN_SHADER_PARAMETER_STRUCT(FPCGCustomKernelDataInterfaceParameters, )
 	SHADER_PARAMETER(FIntVector3, NumThreads)
 	SHADER_PARAMETER(uint32, Seed)
+	SHADER_PARAMETER(uint32, SeedSettings)
+	SHADER_PARAMETER(uint32, SeedComponent)
 	SHADER_PARAMETER(FVector3f, ComponentBoundsMin)
 	SHADER_PARAMETER(FVector3f, ComponentBoundsMax)
 END_SHADER_PARAMETER_STRUCT()
@@ -57,11 +67,15 @@ void UPCGCustomKernelDataInterface::GetHLSL(FString& OutHLSL, FString const& InD
 	OutHLSL += FString::Format(TEXT(
 		"int3 {DataInterfaceName}_NumThreads;\n"
 		"uint {DataInterfaceName}_Seed;\n"
+		"uint {DataInterfaceName}_SeedSettings;\n"
+		"uint {DataInterfaceName}_SeedComponent;\n"
 		"float3 {DataInterfaceName}_ComponentBoundsMin;\n"
 		"float3 {DataInterfaceName}_ComponentBoundsMax;\n"
 		"\n"
 		"int3 GetNumThreads_{DataInterfaceName}()\n{\n\treturn {DataInterfaceName}_NumThreads;\n}\n\n"
 		"uint GetSeed_{DataInterfaceName}()\n{\n\treturn {DataInterfaceName}_Seed;\n}\n\n"
+		"uint GetSettingsSeed_{DataInterfaceName}()\n{\n\treturn {DataInterfaceName}_SeedSettings;\n}\n\n"
+		"uint GetComponentSeed_{DataInterfaceName}()\n{\n\treturn {DataInterfaceName}_SeedComponent;\n}\n\n"
 		"float3 GetComponentBoundsMin_{DataInterfaceName}()\n{\n\treturn {DataInterfaceName}_ComponentBoundsMin;\n}\n\n"
 		"float3 GetComponentBoundsMax_{DataInterfaceName}()\n{\n\treturn {DataInterfaceName}_ComponentBoundsMax;\n}\n\n"),
 		TemplateArgs);
@@ -77,7 +91,11 @@ UComputeDataProvider* UPCGCustomKernelDataInterface::CreateDataProvider(TObjectP
 
 	UPCGCustomComputeKernelDataProvider* Provider = NewObject<UPCGCustomComputeKernelDataProvider>();
 	Provider->ThreadCount = Settings->ComputeKernelThreadCount(Binding);
+	
 	Provider->Seed = static_cast<uint32>(Settings->GetSeed(Binding->SourceComponent.Get()));
+	Provider->SeedSettings = static_cast<uint32>(Settings->Seed);
+	Provider->SeedComponent = static_cast<uint32>(Binding->SourceComponent->Seed);
+
 	Provider->SourceComponentBounds = Binding->SourceComponent.Get()->GetGridBounds();
 
 	return Provider;
@@ -85,7 +103,7 @@ UComputeDataProvider* UPCGCustomKernelDataInterface::CreateDataProvider(TObjectP
 
 FComputeDataProviderRenderProxy* UPCGCustomComputeKernelDataProvider::GetRenderProxy()
 {
-	return new FPCGCustomComputeKernelDataProviderProxy(ThreadCount, Seed, SourceComponentBounds);
+	return new FPCGCustomComputeKernelDataProviderProxy(ThreadCount, Seed, SeedSettings, SeedComponent, SourceComponentBounds);
 }
 
 bool FPCGCustomComputeKernelDataProviderProxy::IsValid(FValidationData const& InValidationData) const
@@ -121,6 +139,8 @@ void FPCGCustomComputeKernelDataProviderProxy::GatherDispatchData(FDispatchData
 
 		// Seed for the node
 		Parameters.Seed = Seed;
+		Parameters.SeedSettings = SeedSettings;
+		Parameters.SeedComponent = SeedComponent;
 
 		// Set component bounds
 		Parameters.ComponentBoundsMin = (FVector3f)SourceComponentBounds.Min;
diff --git a/Engine/Plugins/PCG/Source/PCG/Private/Compute/DataInterfaces/PCGCustomKernelDataInterface.h b/Engine/Plugins/PCG/Source/PCG/Private/Compute/DataInterfaces/PCGCustomKernelDataInterface.h
--- a/Engine/Plugins/PCG/Source/PCG/Private/Compute/DataInterfaces/PCGCustomKernelDataInterface.h
+++ b/Engine/Plugins/PCG/Source/PCG/Private/Compute/DataInterfaces/PCGCustomKernelDataInterface.h
@@ -50,6 +50,8 @@ public:
 	int32 ThreadCount = -1;
 
 	uint32 Seed = 42;
+	uint32 SeedSettings = 42;
+	uint32 SeedComponent = 42;
 
 	FBox SourceComponentBounds;
 };
@@ -57,9 +59,11 @@ public:
 class FPCGCustomComputeKernelDataProviderProxy : public FComputeDataProviderRenderProxy
 {
 public:
-	FPCGCustomComputeKernelDataProviderProxy(int32 InThreadCount, int32 InSeed, const FBox& InSourceComponentBounds)
+	FPCGCustomComputeKernelDataProviderProxy(int32 InThreadCount, uint32 InSeed, uint32 InSeedSettings, uint32 InSeedComponent, const FBox& InSourceComponentBounds)
 		: ThreadCount(InThreadCount)
 		, Seed(InSeed)
+		, SeedSettings(InSeedSettings)
+		, SeedComponent(InSeedComponent)
 		, SourceComponentBounds(InSourceComponentBounds)
 	{}
 
@@ -73,6 +77,10 @@ protected:
 	using FParameters = FPCGCustomKernelDataInterfaceParameters;
 
 	int32 ThreadCount;
+
 	uint32 Seed = 42;
+	uint32 SeedSettings = 42;
+	uint32 SeedComponent = 42;
+
 	FBox SourceComponentBounds;
 };
