[PCG] GPU Copy Points - place files in correct folder and remove unsafe raw pointers and accesses to CPU data from the render proxy.

#rnx

diff --git a/Engine/Plugins/PCG/Source/PCG/Private/Compute/DataInterfaces/Elements/PCGCopyPointsDataInterface.cpp b/Engine/Plugins/PCG/Source/PCG/Private/Compute/DataInterfaces/Elements/PCGCopyPointsDataInterface.cpp
new file mode 100644
--- /dev/null
+++ b/Engine/Plugins/PCG/Source/PCG/Private/Compute/DataInterfaces/Elements/PCGCopyPointsDataInterface.cpp
@@ -0,0 +1,124 @@
+// Copyright Epic Games, Inc. All Rights Reserved.
+
+#include "PCGCopyPointsDataInterface.h"
+
+#include "Elements/PCGCopyPoints.h"
+
+#include "GlobalRenderResources.h"
+#include "RHIResources.h"
+#include "ShaderCompilerCore.h"
+#include "ShaderCore.h"
+#include "ShaderParameterMetadataBuilder.h"
+#include "ComputeFramework/ShaderParamTypeDefinition.h"
+
+#include UE_INLINE_GENERATED_CPP_BY_NAME(PCGCopyPointsDataInterface)
+
+void UPCGCopyPointsDataInterface::GetSupportedInputs(TArray<FShaderFunctionDefinition>& OutFunctions) const
+{
+	OutFunctions.AddDefaulted_GetRef()
+		.SetName(TEXT("GetRotationInheritance"))
+		.AddReturnType(FShaderValueType::Get(EShaderFundamentalType::Uint));
+
+	OutFunctions.AddDefaulted_GetRef()
+		.SetName(TEXT("GetScaleInheritance"))
+		.AddReturnType(FShaderValueType::Get(EShaderFundamentalType::Uint));
+
+	OutFunctions.AddDefaulted_GetRef()
+		.SetName(TEXT("GetColorInheritance"))
+		.AddReturnType(FShaderValueType::Get(EShaderFundamentalType::Uint));
+
+	OutFunctions.AddDefaulted_GetRef()
+		.SetName(TEXT("GetSeedInheritance"))
+		.AddReturnType(FShaderValueType::Get(EShaderFundamentalType::Uint));
+
+	OutFunctions.AddDefaulted_GetRef()
+		.SetName(TEXT("GetAttributeInheritance"))
+		.AddReturnType(FShaderValueType::Get(EShaderFundamentalType::Uint));
+
+	OutFunctions.AddDefaulted_GetRef()
+		.SetName(TEXT("GetCopyEachSourceOnEveryTarget"))
+		.AddReturnType(FShaderValueType::Get(EShaderFundamentalType::Bool));
+}
+
+BEGIN_SHADER_PARAMETER_STRUCT(FPCGCopyPointsDataInterfaceParameters,)
+	SHADER_PARAMETER(uint32, RotationInheritance)
+	SHADER_PARAMETER(uint32, ScaleInheritance)
+	SHADER_PARAMETER(uint32, ColorInheritance)
+	SHADER_PARAMETER(uint32, SeedInheritance)
+	SHADER_PARAMETER(uint32, AttributeInheritance)
+	SHADER_PARAMETER(uint32, bCopyEachSourceOnEveryTarget)
+END_SHADER_PARAMETER_STRUCT()
+
+void UPCGCopyPointsDataInterface::GetShaderParameters(TCHAR const* UID, FShaderParametersMetadataBuilder& InOutBuilder, FShaderParametersMetadataAllocations& InOutAllocations) const
+{
+	InOutBuilder.AddNestedStruct<FPCGCopyPointsDataInterfaceParameters>(UID);
+}
+
+void UPCGCopyPointsDataInterface::GetHLSL(FString& OutHLSL, FString const& InDataInterfaceName) const
+{
+	TMap<FString, FStringFormatArg> TemplateArgs =
+	{
+		{ TEXT("DataInterfaceName"), InDataInterfaceName },
+	};
+
+	OutHLSL += FString::Format(TEXT(
+		"uint {DataInterfaceName}_RotationInheritance;\n"
+		"uint {DataInterfaceName}_ScaleInheritance;\n"
+		"uint {DataInterfaceName}_ColorInheritance;\n"
+		"uint {DataInterfaceName}_SeedInheritance;\n"
+		"uint {DataInterfaceName}_AttributeInheritance;\n"
+		"uint {DataInterfaceName}_bCopyEachSourceOnEveryTarget;\n"
+		"\n"
+		"uint GetRotationInheritance_{DataInterfaceName}() { return {DataInterfaceName}_RotationInheritance; }\n"
+		"uint GetScaleInheritance_{DataInterfaceName}() { return {DataInterfaceName}_ScaleInheritance; }\n"
+		"uint GetColorInheritance_{DataInterfaceName}() { return {DataInterfaceName}_ColorInheritance; }\n"
+		"uint GetSeedInheritance_{DataInterfaceName}() { return {DataInterfaceName}_SeedInheritance; }\n"
+		"uint GetAttributeInheritance_{DataInterfaceName}() { return {DataInterfaceName}_AttributeInheritance; }\n"
+		"uint GetCopyEachSourceOnEveryTarget_{DataInterfaceName}() { return {DataInterfaceName}_bCopyEachSourceOnEveryTarget; }\n"
+		), TemplateArgs);
+}
+
+UComputeDataProvider* UPCGCopyPointsDataInterface::CreateDataProvider(TObjectPtr<UObject> InBinding, uint64 InInputMask, uint64 InOutputMask) const
+{
+	TObjectPtr<UPCGCopyPointsDataProvider> DataProvider = NewObject<UPCGCopyPointsDataProvider>();
+	DataProvider->Settings = Cast<UPCGCopyPointsSettings>(Settings);
+
+	return DataProvider;
+}
+
+FComputeDataProviderRenderProxy* UPCGCopyPointsDataProvider::GetRenderProxy()
+{
+	return new FPCGCopyPointsDataProviderProxy(
+		static_cast<uint32>(Settings->RotationInheritance),
+		static_cast<uint32>(Settings->ScaleInheritance),
+		static_cast<uint32>(Settings->ColorInheritance),
+		static_cast<uint32>(Settings->SeedInheritance),
+		static_cast<uint32>(Settings->AttributeInheritance),
+		static_cast<uint32>(Settings->bCopyEachSourceOnEveryTarget)
+	);
+}
+
+bool FPCGCopyPointsDataProviderProxy::IsValid(FValidationData const& InValidationData) const
+{
+	if (InValidationData.ParameterStructSize != sizeof(FParameters))
+	{
+		return false;
+	}
+
+	return true;
+}
+
+void FPCGCopyPointsDataProviderProxy::GatherDispatchData(FDispatchData const& InDispatchData)
+{
+	const TStridedView<FParameters> ParameterArray = MakeStridedParameterView<FParameters>(InDispatchData);
+	for (int32 InvocationIndex = 0; InvocationIndex < ParameterArray.Num(); ++InvocationIndex)
+	{
+		FParameters& Parameters = ParameterArray[InvocationIndex];
+		Parameters.RotationInheritance = RotationInheritance;
+		Parameters.ScaleInheritance = ScaleInheritance;
+		Parameters.ColorInheritance = ColorInheritance;
+		Parameters.SeedInheritance = SeedInheritance;
+		Parameters.AttributeInheritance = AttributeInheritance;
+		Parameters.bCopyEachSourceOnEveryTarget = bCopyEachSourceOnEveryTarget;
+	}
+}
diff --git a/Engine/Plugins/PCG/Source/PCG/Private/Compute/DataInterfaces/Elements/PCGCopyPointsDataInterface.h b/Engine/Plugins/PCG/Source/PCG/Private/Compute/DataInterfaces/Elements/PCGCopyPointsDataInterface.h
new file mode 100644
--- /dev/null
+++ b/Engine/Plugins/PCG/Source/PCG/Private/Compute/DataInterfaces/Elements/PCGCopyPointsDataInterface.h
@@ -0,0 +1,73 @@
+// Copyright Epic Games, Inc. All Rights Reserved.
+
+#pragma once
+
+#include "Compute/DataInterfaces/PCGComputeDataInterface.h"
+#include "ComputeFramework/ComputeDataProvider.h"
+
+#include "PCGCopyPointsDataInterface.generated.h"
+
+class FPCGCopyPointsDataInterfaceParameters;
+class UPCGCopyPointsSettings;
+class UPCGSettings;
+
+/** Data Interface to marshal Copy Points settings to the GPU. */
+UCLASS(ClassGroup = (Procedural))
+class UPCGCopyPointsDataInterface : public UPCGComputeDataInterface
+{
+	GENERATED_BODY()
+
+public:
+	//~ Begin UComputeDataInterface Interface
+	TCHAR const* GetClassName() const override { return TEXT("PCGCopyPoints"); }
+	void GetSupportedInputs(TArray<FShaderFunctionDefinition>& OutFunctions) const override;
+	void GetShaderParameters(TCHAR const* UID, FShaderParametersMetadataBuilder& InOutBuilder, FShaderParametersMetadataAllocations& InOutAllocations) const override;
+	void GetHLSL(FString& OutHLSL, FString const& InDataInterfaceName) const override;
+	UComputeDataProvider* CreateDataProvider(TObjectPtr<UObject> InBinding, uint64 InInputMask, uint64 InOutputMask) const override;
+	//~ End UComputeDataInterface Interface
+
+	UPROPERTY()
+	TObjectPtr<const UPCGSettings> Settings = nullptr;
+};
+
+UCLASS()
+class UPCGCopyPointsDataProvider : public UComputeDataProvider
+{
+	GENERATED_BODY()
+
+public:
+	//~ Begin UComputeDataProvider Interface
+	FComputeDataProviderRenderProxy* GetRenderProxy() override;
+	//~ End UComputeDataProvider Interface
+
+	UPROPERTY()
+	TObjectPtr<const UPCGCopyPointsSettings> Settings;
+};
+
+class FPCGCopyPointsDataProviderProxy : public FComputeDataProviderRenderProxy
+{
+public:
+	FPCGCopyPointsDataProviderProxy(uint32 InRotationInheritance, uint32 InScaleInheritance, uint32 InColorInheritance, uint32 InSeedInheritance, uint32 InAttributeInheritance, uint32 bInCopyEachSourceOnEveryTarget)
+		: RotationInheritance(InRotationInheritance)
+		, ScaleInheritance(InScaleInheritance)
+		, ColorInheritance(InColorInheritance)
+		, SeedInheritance(InSeedInheritance)
+		, AttributeInheritance(InAttributeInheritance)
+		, bCopyEachSourceOnEveryTarget(bInCopyEachSourceOnEveryTarget)
+	{}
+
+	//~ Begin FComputeDataProviderRenderProxy Interface
+	bool IsValid(FValidationData const& InValidationData) const override;
+	void GatherDispatchData(FDispatchData const& InDispatchData) override;
+	//~ End FComputeDataProviderRenderProxy Interface
+
+protected:
+	using FParameters = FPCGCopyPointsDataInterfaceParameters;
+
+	uint32 RotationInheritance = 0;
+	uint32 ScaleInheritance = 0;
+	uint32 ColorInheritance = 0;
+	uint32 SeedInheritance = 0;
+	uint32 AttributeInheritance = 0;
+	uint32 bCopyEachSourceOnEveryTarget = 0;
+};
diff --git a/Engine/Plugins/PCG/Source/PCG/Private/Compute/DataInterfaces/PCGCopyPointsDataInterface.cpp b/Engine/Plugins/PCG/Source/PCG/Private/Compute/DataInterfaces/PCGCopyPointsDataInterface.cpp
deleted file mode 100644
diff --git a/Engine/Plugins/PCG/Source/PCG/Private/Compute/DataInterfaces/PCGCopyPointsDataInterface.h b/Engine/Plugins/PCG/Source/PCG/Private/Compute/DataInterfaces/PCGCopyPointsDataInterface.h
deleted file mode 100644
diff --git a/Engine/Plugins/PCG/Source/PCG/Private/Elements/PCGCopyPoints.cpp b/Engine/Plugins/PCG/Source/PCG/Private/Elements/PCGCopyPoints.cpp
--- a/Engine/Plugins/PCG/Source/PCG/Private/Elements/PCGCopyPoints.cpp
+++ b/Engine/Plugins/PCG/Source/PCG/Private/Elements/PCGCopyPoints.cpp
@@ -3,9 +3,9 @@
 #include "Elements/PCGCopyPoints.h"
 
 #include "PCGContext.h"
-#include "Compute/DataInterfaces/PCGCopyPointsDataInterface.h"
 #include "Compute/PCGComputeCommon.h"
 #include "Compute/PCGDataBinding.h"
+#include "Compute/DataInterfaces/Elements/PCGCopyPointsDataInterface.h"
 #include "Data/PCGPointData.h"
 #include "Data/PCGSpatialData.h"
 #include "Graph/PCGGPUGraphCompilationContext.h"
