[PCG] Fix crash when refreshing Node Source Editor
* Refresh must take place on main thread, but can be requested by worker threads after compute graph compilation
* Move PCGGraphExecutor::ExecuteOnGameThread() to PCGHelpers::ExecuteOnGameThread()

#rnx

diff --git a/Engine/Plugins/PCG/Source/PCG/Private/Graph/PCGGraphExecutor.cpp b/Engine/Plugins/PCG/Source/PCG/Private/Graph/PCGGraphExecutor.cpp
--- a/Engine/Plugins/PCG/Source/PCG/Private/Graph/PCGGraphExecutor.cpp
+++ b/Engine/Plugins/PCG/Source/PCG/Private/Graph/PCGGraphExecutor.cpp
@@ -968,7 +968,7 @@ void FPCGGraphExecutor::PostTaskExecute(TSharedPtr<FPCGGraphActiveTask> ActiveTa
 			CollectGCReferenceTasks.Add(ActiveTaskPtr);
 		}
 
-		PCGGraphExecutor::ExecuteOnGameThread(UE_SOURCE_LOCATION, [WeakHandler = GameThreadHandler->AsWeak(), ActiveTaskPtr]()
+		PCGHelpers::ExecuteOnGameThread(UE_SOURCE_LOCATION, [WeakHandler = GameThreadHandler->AsWeak(), ActiveTaskPtr]()
 		{
 			if (TSharedPtr<FPCGGraphExecutor::FGameThreadHandler> Handler = WeakHandler.Pin())
 			{
diff --git a/Engine/Plugins/PCG/Source/PCG/Private/Graph/PCGGraphExecutor.h b/Engine/Plugins/PCG/Source/PCG/Private/Graph/PCGGraphExecutor.h
--- a/Engine/Plugins/PCG/Source/PCG/Private/Graph/PCGGraphExecutor.h
+++ b/Engine/Plugins/PCG/Source/PCG/Private/Graph/PCGGraphExecutor.h
@@ -12,7 +12,6 @@
 #include "Graph/PCGGraphTask.h"
 #include "Graph/PCGStackContext.h"
 
-#include "Containers/Ticker.h"
 #include "Misc/SpinLock.h"
 #include "Tasks/Task.h"
 #include "Templates/UniquePtr.h"
@@ -40,19 +39,6 @@ namespace PCGGraphExecutor
 	extern PCG_API TAutoConsoleVariable<float> CVarEditorTimePerFrame;
 #endif
 
-	template<typename FunctorType>
-	void ExecuteOnGameThread(const TCHAR* DebugName, FunctorType&& Functor)
-	{
-		if (IsInGameThread())
-		{
-			Functor();
-		}
-		else
-		{
-			::ExecuteOnGameThread(DebugName, std::forward<FunctorType>(Functor));
-		}
-	}
-
 	void ClearAsyncFlags(TSet<TObjectPtr<UObject>>& AsyncObjects);
 }
 
diff --git a/Engine/Plugins/PCG/Source/PCG/Private/Graph/PCGGraphTask.cpp b/Engine/Plugins/PCG/Source/PCG/Private/Graph/PCGGraphTask.cpp
--- a/Engine/Plugins/PCG/Source/PCG/Private/Graph/PCGGraphTask.cpp
+++ b/Engine/Plugins/PCG/Source/PCG/Private/Graph/PCGGraphTask.cpp
@@ -6,6 +6,7 @@
 #include "PCGSubsystem.h"
 #include "Compute/Elements/PCGComputeGraphElement.h"
 #include "Graph/PCGGraphExecutor.h"
+#include "Helpers/PCGHelpers.h"
 
 #include UE_INLINE_GENERATED_CPP_BY_NAME(PCGGraphTask)
 
@@ -112,7 +113,7 @@ FPCGGraphActiveTask::~FPCGGraphActiveTask()
 {
 	if (Context)
 	{
-		PCGGraphExecutor::ExecuteOnGameThread(UE_SOURCE_LOCATION, [ContextPtr = Context.Release()]()
+		PCGHelpers::ExecuteOnGameThread(UE_SOURCE_LOCATION, [ContextPtr = Context.Release()]()
 		{
 			PCGGraphExecutor::ClearAsyncFlags(ContextPtr->AsyncObjects);
 			delete ContextPtr;
diff --git a/Engine/Plugins/PCG/Source/PCG/Public/Helpers/PCGHelpers.h b/Engine/Plugins/PCG/Source/PCG/Public/Helpers/PCGHelpers.h
--- a/Engine/Plugins/PCG/Source/PCG/Public/Helpers/PCGHelpers.h
+++ b/Engine/Plugins/PCG/Source/PCG/Public/Helpers/PCGHelpers.h
@@ -5,6 +5,7 @@
 #include "PCGCommon.h"
 #include "PCGPoint.h"
 
+#include "Containers/Ticker.h"
 #include "Kismet/BlueprintFunctionLibrary.h"
 #include "Math/Box.h"
 
@@ -149,6 +150,19 @@ namespace PCGHelpers
 
 		Array = std::move(TempArray);
 	}
+
+	template<typename FunctorType>
+	void ExecuteOnGameThread(const TCHAR* DebugName, FunctorType&& Functor)
+	{
+		if (IsInGameThread())
+		{
+			Functor();
+		}
+		else
+		{
+			::ExecuteOnGameThread(DebugName, std::forward<FunctorType>(Functor));
+		}
+	}
 }
 
 /** Holds function prototypes used to match against actor function signatures. */
diff --git a/Engine/Plugins/PCG/Source/PCGEditor/Private/Widgets/SPCGEditorNodeSource.cpp b/Engine/Plugins/PCG/Source/PCGEditor/Private/Widgets/SPCGEditorNodeSource.cpp
--- a/Engine/Plugins/PCG/Source/PCGEditor/Private/Widgets/SPCGEditorNodeSource.cpp
+++ b/Engine/Plugins/PCG/Source/PCGEditor/Private/Widgets/SPCGEditorNodeSource.cpp
@@ -4,6 +4,7 @@
 
 #include "PCGSettings.h"
 #include "Compute/IPCGNodeSourceTextProvider.h"
+#include "Helpers/PCGHelpers.h"
 
 #include "PCGEditor.h"
 #include "PCGHLSLSyntaxHighlighter.h"
@@ -242,7 +243,17 @@ void SPCGEditorNodeSource::SetShaderSourceText() const
 void SPCGEditorNodeSource::OnDiagnosticsUpdated(const FPCGCompilerDiagnostics& InDiagnostics) const
 {
 	SyntaxHighlighterShaderText->SetCompilerMessages(InDiagnostics);
-	ShaderTextTextBox->Refresh();
+
+	PCGHelpers::ExecuteOnGameThread(UE_SOURCE_LOCATION, [ShaderTextBoxPtr = TWeakPtr<SPCGNodeSourceTextBox>(ShaderTextTextBox)]()
+	{
+		if (ShaderTextBoxPtr.IsValid())
+		{
+			if (TSharedPtr<SPCGNodeSourceTextBox> ShaderTextBox = ShaderTextBoxPtr.Pin())
+			{
+				ShaderTextBox->Refresh();
+			}
+		}
+	});
 }
 
 TSharedRef<SWidget> SPCGEditorNodeSource::ConstructNonExpandableHeaderWidget(const SExpandableArea::FArguments& InArgs) const
diff --git a/Engine/Plugins/PCG/Source/PCGEditor/Private/Widgets/SPCGNodeSourceTextBox.cpp b/Engine/Plugins/PCG/Source/PCGEditor/Private/Widgets/SPCGNodeSourceTextBox.cpp
--- a/Engine/Plugins/PCG/Source/PCGEditor/Private/Widgets/SPCGNodeSourceTextBox.cpp
+++ b/Engine/Plugins/PCG/Source/PCGEditor/Private/Widgets/SPCGNodeSourceTextBox.cpp
@@ -251,6 +251,7 @@ void SPCGNodeSourceTextBox::OnTriggerSearch()
 
 void SPCGNodeSourceTextBox::Refresh() const
 {
+	check(IsInGameThread());
 	Text->Refresh();
 }
 
