[PCG] Fix GPU readback emitting all data on all pins

#rnx

diff --git a/Engine/Plugins/PCG/Source/PCG/Private/Graph/PCGGraphCompilerGPU.cpp b/Engine/Plugins/PCG/Source/PCG/Private/Graph/PCGGraphCompilerGPU.cpp
--- a/Engine/Plugins/PCG/Source/PCG/Private/Graph/PCGGraphCompilerGPU.cpp
+++ b/Engine/Plugins/PCG/Source/PCG/Private/Graph/PCGGraphCompilerGPU.cpp
@@ -1226,7 +1226,9 @@ void FPCGGraphCompilerGPU::CreateGPUNodes(FPCGGraphCompiler& InOutCompiler, UPCG
 	{
 		for (int InputIndex = 0; InputIndex < InOutCompiledTasks[TaskId].Inputs.Num(); ++InputIndex)
 		{
-			TaskSuccessors.FindOrAdd(InOutCompiledTasks[TaskId].Inputs[InputIndex].TaskId).Add(TaskId);
+			// Only add unique task successors to avoid storing the data from each output pin
+			// multiple times when multiple pins connect to the same downstream node.
+			TaskSuccessors.FindOrAdd(InOutCompiledTasks[TaskId].Inputs[InputIndex].TaskId).AddUnique(TaskId);
 		}
 	}
 
