Made inspection non transactional to prevent issues on cleanup

#rnx

diff --git a/Engine/Plugins/PCG/Source/PCG/Public/PCGComponent.h b/Engine/Plugins/PCG/Source/PCG/Public/PCGComponent.h
--- a/Engine/Plugins/PCG/Source/PCG/Public/PCGComponent.h
+++ b/Engine/Plugins/PCG/Source/PCG/Public/PCGComponent.h
@@ -627,7 +627,7 @@ private:
 	FPCGDataCollection GeneratedGraphOutput;
 
 	/** If any graph edges cross execution grid sizes, data on the edge is stored / retrieved from this map. */
-	UPROPERTY(Transient, VisibleAnywhere, Category = Debug)
+	UPROPERTY(Transient, NonTransactional, VisibleAnywhere, Category = Debug)
 	TMap<FString, FPCGDataCollection> PerPinGeneratedOutput;
 
 	mutable FRWLock PerPinGeneratedOutputLock;
@@ -665,7 +665,7 @@ private:
 	// Serialization will be done in the Serialize function
 	FPCGSelectionKeyToSettingsMap DynamicallyTrackedKeysToSettings;
 
-	UPROPERTY(Transient)
+	UPROPERTY(Transient, NonTransactional)
 	TMap<FPCGStack, FPCGDataCollection> InspectionCache;
 #endif
 
