Added GPU widget to easily identify nodes that are running on GPU.

#rnx

diff --git a/Engine/Plugins/PCG/Source/PCGEditor/Private/PCGEditorStyle.cpp b/Engine/Plugins/PCG/Source/PCGEditor/Private/PCGEditorStyle.cpp
--- a/Engine/Plugins/PCG/Source/PCGEditor/Private/PCGEditorStyle.cpp
+++ b/Engine/Plugins/PCG/Source/PCGEditor/Private/PCGEditorStyle.cpp
@@ -90,6 +90,14 @@ FPCGEditorStyle::FPCGEditorStyle() : FSlateStyleSet("PCGEditorStyle")
 		FLinearColor::Black,
 		PCGEditorStyleConstants::Node_Overlay_GridSizeLabel_BorderStroke));
 
+	// Styles related to additional tokens on nodes (GPU, ...)
+	FTextBlockStyle AdditionalOverlayWidgetTextStyle = FTextBlockStyle(FAppStyle::Get().GetWidgetStyle<FTextBlockStyle>("NormalText"))
+		.SetFont(FCoreStyle::GetDefaultFontStyle("Bold", 8))
+		.SetColorAndOpacity(FLinearColor(218.0f / 255.0f, 218.0f / 255.0f, 218.0f / 255.0f))
+		.SetShadowOffset(FVector2D::ZeroVector)
+		.SetShadowColorAndOpacity(FLinearColor(0.f, 0.f, 0.f, 0.7f));
+	Set("PCG.Node.AdditionalOverlayWidgetText", AdditionalOverlayWidgetTextStyle);
+
 	// Plugin icon/editor/component icons
 	Set("ClassIcon.PCGComponent", new IMAGE_BRUSH_SVG("Icons/PCG_16", Icon16x16));
 	Set("ClassThumbnail.PCGComponent", new IMAGE_BRUSH_SVG("Icons/PCG_64", Icon64x64));
diff --git a/Engine/Plugins/PCG/Source/PCGEditor/Private/SPCGEditorGraphNode.cpp b/Engine/Plugins/PCG/Source/PCGEditor/Private/SPCGEditorGraphNode.cpp
--- a/Engine/Plugins/PCG/Source/PCGEditor/Private/SPCGEditorGraphNode.cpp
+++ b/Engine/Plugins/PCG/Source/PCGEditor/Private/SPCGEditorGraphNode.cpp
@@ -627,7 +627,6 @@ TArray<FOverlayWidgetInfo> SPCGEditorGraphNode::GetOverlayWidgets(bool bSelected
 		const EPCGHiGenGrid InspectedGrid = PCGEditorGraphNode->GetInspectedGenerationGrid();
 		UPCGNode* PCGNode = PCGEditorGraphNode->GetPCGNode();
 
-		//const bool bHigenEnabled = PCGEditorGraphNode->GetPCGNode() && PCGEditorGraphNode->getpcn
 		const bool bInspectingHigen = InspectedGrid != EPCGHiGenGrid::Uninitialized;
 		if (bInspectingHigen && PCGEditorGraphNode->IsNodeEnabled())
 		{
@@ -690,6 +689,46 @@ TArray<FOverlayWidgetInfo> SPCGEditorGraphNode::GetOverlayWidgets(bool bSelected
 
 			OverlayWidgets.Add(GridSizeLabelInfo);
 		}
+
+		if (PCGEditorGraphNode->GetPCGNode() && PCGEditorGraphNode->GetPCGNode()->GetSettings() && PCGEditorGraphNode->GetPCGNode()->GetSettings()->ShouldExecuteOnGPU())
+		{
+			const FLinearColor LightGray(0.7f, 0.7f, 0.7f);
+			const float BorderRadius = 7.0f;
+			const float BorderStroke = 1.5f;
+
+			FText GPUText = FText::FromString(TEXT("GPU"));
+			FLinearColor Tint = LightGray;
+			FLinearColor TextColor = LightGray;
+			FLinearColor BackgroundColor = FColor::White;
+			const FSlateBrush* BorderBrush = new FSlateRoundedBoxBrush(
+				FLinearColor::Black,
+				BorderRadius,
+				LightGray,
+				BorderStroke);
+
+			TSharedPtr<SWidget> GPUUsageLabel =
+				SNew(SHorizontalBox)
+				.Visibility(EVisibility::Visible)
+				+SHorizontalBox::Slot()
+				[
+					SNew(SBorder)
+					.BorderImage(BorderBrush)
+					.Padding(FMargin(6, 3))
+					.ColorAndOpacity(Tint)
+					[
+						SNew(STextBlock)
+						.TextStyle(FPCGEditorStyle::Get(), "PCG.Node.AdditionalOverlayWidgetText")
+						.Text(GPUText)
+						.Justification(ETextJustify::Center)
+						.ColorAndOpacity(TextColor)
+					]
+				];
+
+			FOverlayWidgetInfo GPUUsageLabelInfo(GPUUsageLabel);
+			GPUUsageLabelInfo.OverlayOffset = FVector2D(GetDesiredSize().X - 34.0f, GetDesiredSize().Y + 5.0f);
+
+			OverlayWidgets.Add(GPUUsageLabelInfo);
+		}
 	}
 
 	return OverlayWidgets;
