[PCG] Fix uninitialized values for attributes when using point generator, and expose InitializePoint function for users.

* Fix some default values
* Make sure the string table always has empty string in first slot so that 0-initializing string keys is fine
* Small SM Spawner shader code fix introduced by recent mis-shelve

#rnx

diff --git a/Engine/Plugins/PCG/Shaders/Private/Elements/PCGStaticMeshSpawner.usf b/Engine/Plugins/PCG/Shaders/Private/Elements/PCGStaticMeshSpawner.usf
--- a/Engine/Plugins/PCG/Shaders/Private/Elements/PCGStaticMeshSpawner.usf
+++ b/Engine/Plugins/PCG/Shaders/Private/Elements/PCGStaticMeshSpawner.usf
@@ -60,7 +60,7 @@ void Main(uint3 GroupId : SV_GroupID, uint GroupIndex : SV_GroupIndex)
 			const uint SeedPoint = In_GetSeed(In_DataIndex, In_ElementIndex);
 			const uint SeedSettings = GetSettingsSeed();
 			const uint SeedComponent = GetComponentSeed();
-			const uint Seed = ComputeSeed(SeedPoint, SeedSettings, SeedComponent);
+			uint Seed = ComputeSeed(SeedPoint, SeedSettings, SeedComponent);
 
 			const float RandomDraw = FRand(Seed);
 
diff --git a/Engine/Plugins/PCG/Shaders/Private/PCGDataCollectionDataInterface.ush b/Engine/Plugins/PCG/Shaders/Private/PCGDataCollectionDataInterface.ush
--- a/Engine/Plugins/PCG/Shaders/Private/PCGDataCollectionDataInterface.ush
+++ b/Engine/Plugins/PCG/Shaders/Private/PCGDataCollectionDataInterface.ush
@@ -39,6 +39,10 @@ RWByteAddressBuffer {DataInterfaceName}_DataCollectionBuffer;
 #define ATTRIBUTE_HEADER_SIZE_BYTES 8
 #define KERNEL_EXECUTED_FLAG (1u << 31)
 
+// TODO: should be code-gen'd from PCGComputeCommon.h
+#define MAX_NUM_ATTRS 128
+#define NUM_RESERVED_ATTRS 32
+
 // #################### HEADER READERS ##########################
 
 uint GetNumData_{DataInterfaceName}()
@@ -421,7 +425,7 @@ void RemovePoint_{DataInterfaceName}(uint InDataIndex, uint InElementIndex)
 // Initialize a single point with default values.
 void InitializePoint_{DataInterfaceName}(uint InDataIndex, uint InElementIndex)
 {
-	SetPosition_{DataInterfaceName}(InDataIndex, InElementIndex, 1.0f);
+	SetPosition_{DataInterfaceName}(InDataIndex, InElementIndex, 0.0f);
 	SetRotation_{DataInterfaceName}(InDataIndex, InElementIndex, float4(0.0f, 0.0f, 0.0f, 1.0f));
 	SetScale_{DataInterfaceName}(InDataIndex, InElementIndex, 1.0f);
 	SetBoundsMin_{DataInterfaceName}(InDataIndex, InElementIndex, -50.0f);
@@ -429,7 +433,32 @@ void InitializePoint_{DataInterfaceName}(uint InDataIndex, uint InElementIndex)
 	SetColor_{DataInterfaceName}(InDataIndex, InElementIndex, 1.0f);
 	SetDensity_{DataInterfaceName}(InDataIndex, InElementIndex, 1.0f);
 	SetSeed_{DataInterfaceName}(InDataIndex, InElementIndex, 42);
-	SetSteepness_{DataInterfaceName}(InDataIndex, InElementIndex, 0.5f);
+	SetSteepness_{DataInterfaceName}(InDataIndex, InElementIndex, 1.0f);
 	
-	// TODO - pass in relevant attribute IDs and init them
+	int NumAttributesRemaining = (int)GetDataNumAttributesInternal_{DataInterfaceName}(InDataIndex);
+
+	for (int AttributeId = NUM_RESERVED_ATTRS; AttributeId < MAX_NUM_ATTRS; ++AttributeId)
+	{
+		const uint Stride = GetAttributeStrideInternal_{DataInterfaceName}(InDataIndex, AttributeId);
+
+		if (Stride == 0)
+		{
+			// No output attribute to write to.
+			continue;
+		}
+
+		const uint ElementAddress = GetElementAddressInternal_{DataInterfaceName}(InDataIndex, InElementIndex, AttributeId);
+		
+		for (uint Offset = 0; Offset < Stride; Offset += 4)
+		{
+			// TODO: In future could upload actual default values for attributes rather than 0-initializing.
+			StoreBufferInternal_{DataInterfaceName}(ElementAddress + Offset, 0u);
+		}
+
+		if (--NumAttributesRemaining <= 0)
+		{
+			// Early-out when we've looked at all the possible attributes
+			break;
+		}
+	}
 }
diff --git a/Engine/Plugins/PCG/Source/PCG/Private/Compute/Elements/PCGCustomHLSL.cpp b/Engine/Plugins/PCG/Source/PCG/Private/Compute/Elements/PCGCustomHLSL.cpp
--- a/Engine/Plugins/PCG/Source/PCG/Private/Compute/Elements/PCGCustomHLSL.cpp
+++ b/Engine/Plugins/PCG/Source/PCG/Private/Compute/Elements/PCGCustomHLSL.cpp
@@ -922,7 +922,8 @@ void UPCGCustomHLSLSettings::UpdateOutputDeclarations()
 			"void {0}_SetSeed(uint DataIndex, uint ElementIndex, int Seed);\n"
 			"void {0}_SetSteepness(uint DataIndex, uint ElementIndex, float Steepness);\n"
 			"void {0}_SetPointTransform(uint DataIndex, uint ElementIndex, float4x4 Transform);\n"
-			"bool {0}_RemovePoint(uint DataIndex, uint ElementIndex);\n"),
+			"bool {0}_RemovePoint(uint DataIndex, uint ElementIndex);\n"
+			"void {0}_InitializePoint(uint DataIndex, uint ElementIndex);\n"),
 			{ bMultiPin ? PCGHLSLElement::PinDeclTemplateStr : PointDataPins[0] });
 
 		OutputDeclarations += TEXT("\n");
diff --git a/Engine/Plugins/PCG/Source/PCG/Private/Compute/PCGDataBinding.cpp b/Engine/Plugins/PCG/Source/PCG/Private/Compute/PCGDataBinding.cpp
--- a/Engine/Plugins/PCG/Source/PCG/Private/Compute/PCGDataBinding.cpp
+++ b/Engine/Plugins/PCG/Source/PCG/Private/Compute/PCGDataBinding.cpp
@@ -83,8 +83,11 @@ void UPCGDataBinding::BuildStringTable()
 	TRACE_CPUPROFILER_EVENT_SCOPE(UPCGDataBinding::BuildStringTable);
 	check(Graph);
 
-	// Start with any strings known statically at compilation time.
-	StringTable = Graph->GetStringTable();
+	// String table always contains empty string in index 0 (and string key attributes are 0-initialized).
+	StringTable = { FString() };
+
+	// Add any strings known statically at compilation time.
+	StringTable.Append(Graph->GetStringTable());
 
 	for (const FPCGTaggedData& Data : DataForGPU.InputDataCollection.TaggedData)
 	{
diff --git a/Engine/Plugins/PCG/Source/PCG/Private/Graph/PCGGraphCompilerGPU.cpp b/Engine/Plugins/PCG/Source/PCG/Private/Graph/PCGGraphCompilerGPU.cpp
--- a/Engine/Plugins/PCG/Source/PCG/Private/Graph/PCGGraphCompilerGPU.cpp
+++ b/Engine/Plugins/PCG/Source/PCG/Private/Graph/PCGGraphCompilerGPU.cpp
@@ -1187,6 +1187,10 @@ UPCGComputeGraph* FPCGGraphCompilerGPU::CompileComputeGraph(
 		}
 	}
 
+	// Remove empty strings (at execution time empty string is placed in table slot 0).
+	// NOTE: This can scramble order but order is not important at this stage (key values created at execution time in data binding).
+	ComputeGraph->StringTable.RemoveAllSwap([](const FString& InString) { return InString.IsEmpty(); });
+
 	return ComputeGraph;
 }
 
