[PCG] GetActorProperty: Fix condition inversion from 36979836

#rnx

diff --git a/Engine/Plugins/PCG/Source/PCG/Private/Elements/PCGGetActorProperty.cpp b/Engine/Plugins/PCG/Source/PCG/Private/Elements/PCGGetActorProperty.cpp
--- a/Engine/Plugins/PCG/Source/PCG/Private/Elements/PCGGetActorProperty.cpp
+++ b/Engine/Plugins/PCG/Source/PCG/Private/Elements/PCGGetActorProperty.cpp
@@ -148,7 +148,7 @@ bool FPCGGetActorPropertyElement::IsCacheable(const UPCGSettings* InSettings) co
 {
 	if (const UPCGGetActorPropertySettings* Settings = Cast<UPCGGetActorPropertySettings>(InSettings))
 	{
-		return Settings->bAlwaysRequeryActors;
+		return !Settings->bAlwaysRequeryActors;
 	}
 	else
 	{
