[PCG] Allow persistent ISM/Spline component reuse
- Replace ISM/Spline SettingsUID with SettingsCrc which are stable and persisted

diff --git a/Engine/Plugins/PCG/Source/PCG/Private/Elements/PCGSpawnActor.cpp b/Engine/Plugins/PCG/Source/PCG/Private/Elements/PCGSpawnActor.cpp
--- a/Engine/Plugins/PCG/Source/PCG/Private/Elements/PCGSpawnActor.cpp
+++ b/Engine/Plugins/PCG/Source/PCG/Private/Elements/PCGSpawnActor.cpp
@@ -678,13 +678,17 @@ void FPCGSpawnActorElement::CollapseIntoTargetActor(FPCGSubgraphContext* Context
 
 	const UPCGSpawnActorSettings* Settings = Context->GetInputSettings<UPCGSpawnActorSettings>();
 	check(Settings);
+	
+	FPCGCrc SettingsCrc = Settings->GetSettingsCrc();
+	ensure(SettingsCrc.IsValid());
 
 	TMap<FPCGISMComponentBuilderParams, TArray<FTransform>> MeshDescriptorTransforms;
 
-	AActor::ForEachComponentOfActorClassDefault<UStaticMeshComponent>(TemplateActorClass, [&MeshDescriptorTransforms](const UStaticMeshComponent* StaticMeshComponent)
+	AActor::ForEachComponentOfActorClassDefault<UStaticMeshComponent>(TemplateActorClass, [&MeshDescriptorTransforms, SettingsCrc](const UStaticMeshComponent* StaticMeshComponent)
 	{
 		FPCGISMComponentBuilderParams Params;
 		Params.Descriptor.InitFrom(StaticMeshComponent);
+		Params.SettingsCrc = SettingsCrc;
 		// TODO: No custom data float support?
 
 		TArray<FTransform>& Transforms = MeshDescriptorTransforms.FindOrAdd(Params);
@@ -715,7 +719,7 @@ void FPCGSpawnActorElement::CollapseIntoTargetActor(FPCGSubgraphContext* Context
 	{
 		const FPCGISMComponentBuilderParams& ISMCParams = ISMCBuilderTransforms.Key;
 
-		UPCGManagedISMComponent* MISMC = UPCGActorHelpers::GetOrCreateManagedISMC(TargetActor, Context->SourceComponent.Get(), Settings->GetStableUID(), ISMCParams);
+		UPCGManagedISMComponent* MISMC = UPCGActorHelpers::GetOrCreateManagedISMC(TargetActor, Context->SourceComponent.Get(), ISMCParams);
 		if (!MISMC)
 		{
 			continue;
diff --git a/Engine/Plugins/PCG/Source/PCG/Private/Elements/PCGSpawnSplineMesh.cpp b/Engine/Plugins/PCG/Source/PCG/Private/Elements/PCGSpawnSplineMesh.cpp
--- a/Engine/Plugins/PCG/Source/PCG/Private/Elements/PCGSpawnSplineMesh.cpp
+++ b/Engine/Plugins/PCG/Source/PCG/Private/Elements/PCGSpawnSplineMesh.cpp
@@ -212,8 +212,10 @@ bool FPCGSpawnSplineMeshElement::ExecuteInternal(FPCGContext* Context) const
 
 			SMCBuilderParams.SplineMeshParams.StartScale = FVector2D(StartScale.Y, StartScale.Z);
 			SMCBuilderParams.SplineMeshParams.EndScale = FVector2D(EndScale.Y, EndScale.Z);
+			SMCBuilderParams.SettingsCrc = Settings->GetSettingsCrc();
+			ensure(SMCBuilderParams.SettingsCrc.IsValid());
 
-			USplineMeshComponent* SplineMeshComponent = UPCGActorHelpers::GetOrCreateSplineMeshComponent(TargetActor, Context->SourceComponent.Get(), Settings->GetStableUID(), SMCBuilderParams);
+			USplineMeshComponent* SplineMeshComponent = UPCGActorHelpers::GetOrCreateSplineMeshComponent(TargetActor, Context->SourceComponent.Get(), SMCBuilderParams);
 
 			// TODO: Write out the geometry to a dynamic mesh type.
 			//SplineMeshComponent->BodySetup->TriMeshGeometries
diff --git a/Engine/Plugins/PCG/Source/PCG/Private/Elements/PCGStaticMeshSpawner.cpp b/Engine/Plugins/PCG/Source/PCG/Private/Elements/PCGStaticMeshSpawner.cpp
--- a/Engine/Plugins/PCG/Source/PCG/Private/Elements/PCGStaticMeshSpawner.cpp
+++ b/Engine/Plugins/PCG/Source/PCG/Private/Elements/PCGStaticMeshSpawner.cpp
@@ -331,7 +331,7 @@ bool FPCGStaticMeshSpawnerElement::PrepareDataInternal(FPCGContext* InContext) c
 			{
 				if (UPCGManagedISMComponent* Resource = Cast<UPCGManagedISMComponent>(InResource))
 				{
-					if (Resource->GetSettingsUID() == Settings->GetStableUID() && Resource->GetCrc().IsValid() && Resource->GetCrc() == Context->DependenciesCrc)
+					if (Resource->GetCrc().IsValid() && Resource->GetCrc() == Context->DependenciesCrc)
 					{
 						MISMCs.Add(Resource);
 					}
@@ -654,9 +654,12 @@ void FPCGStaticMeshSpawnerElement::SpawnStaticMeshInstances(FPCGStaticMeshSpawne
 	const UPCGStaticMeshSpawnerSettings* Settings = Context->GetInputSettings<UPCGStaticMeshSpawnerSettings>();
 	check(Settings);
 
+	Params.SettingsCrc = Settings->GetSettingsCrc();
+	ensure(Params.SettingsCrc.IsValid());
+
 	Params.bAllowDescriptorChanges = Settings->bAllowDescriptorChanges;
 
-	UPCGManagedISMComponent* MISMC = UPCGActorHelpers::GetOrCreateManagedISMC(TargetActor, Context->SourceComponent.Get(), Settings->GetStableUID(), Params);
+	UPCGManagedISMComponent* MISMC = UPCGActorHelpers::GetOrCreateManagedISMC(TargetActor, Context->SourceComponent.Get(), Params);
 
 	check(MISMC);
 	MISMC->SetCrc(Context->DependenciesCrc);
diff --git a/Engine/Plugins/PCG/Source/PCG/Private/Helpers/PCGActorHelpers.cpp b/Engine/Plugins/PCG/Source/PCG/Private/Helpers/PCGActorHelpers.cpp
--- a/Engine/Plugins/PCG/Source/PCG/Private/Helpers/PCGActorHelpers.cpp
+++ b/Engine/Plugins/PCG/Source/PCG/Private/Helpers/PCGActorHelpers.cpp
@@ -52,9 +52,21 @@ UPCGManagedISMComponent* UPCGActorHelpers::GetOrCreateManagedISMC(AActor* InTarg
 }
 PRAGMA_ENABLE_DEPRECATION_WARNINGS
 
+// deprecated
 UInstancedStaticMeshComponent* UPCGActorHelpers::GetOrCreateISMC(AActor* InTargetActor, UPCGComponent* InSourceComponent, uint64 SettingsUID, const FPCGISMComponentBuilderParams& InParams)
 {
-	UPCGManagedISMComponent* MISMC = GetOrCreateManagedISMC(InTargetActor, InSourceComponent, SettingsUID, InParams);
+	return GetOrCreateISMC(InTargetActor, InSourceComponent, InParams);
+}
+
+// deprecated
+UPCGManagedISMComponent* UPCGActorHelpers::GetOrCreateManagedISMC(AActor* InTargetActor, UPCGComponent* InSourceComponent, uint64 SettingsUID, const FPCGISMComponentBuilderParams& InParams)
+{
+	return GetOrCreateManagedISMC(InTargetActor, InSourceComponent, InParams);
+}
+
+UInstancedStaticMeshComponent* UPCGActorHelpers::GetOrCreateISMC(AActor* InTargetActor, UPCGComponent* InSourceComponent, const FPCGISMComponentBuilderParams& InParams)
+{
+	UPCGManagedISMComponent* MISMC = GetOrCreateManagedISMC(InTargetActor, InSourceComponent, InParams);
 	if (MISMC)
 	{
 		return MISMC->GetComponent();
@@ -65,7 +77,7 @@ UInstancedStaticMeshComponent* UPCGActorHelpers::GetOrCreateISMC(AActor* InTarge
 	}
 }
 
-UPCGManagedISMComponent* UPCGActorHelpers::GetOrCreateManagedISMC(AActor* InTargetActor, UPCGComponent* InSourceComponent, uint64 SettingsUID, const FPCGISMComponentBuilderParams& InParams)
+UPCGManagedISMComponent* UPCGActorHelpers::GetOrCreateManagedISMC(AActor* InTargetActor, UPCGComponent* InSourceComponent, const FPCGISMComponentBuilderParams& InParams)
 {
 	check(InTargetActor && InSourceComponent);
 
@@ -120,33 +132,36 @@ UPCGManagedISMComponent* UPCGActorHelpers::GetOrCreateManagedISMC(AActor* InTarg
 	{
 		TRACE_CPUPROFILER_EVENT_SCOPE(UPCGActorHelpers::GetOrCreateManagedISMC::FindMatchingMISMC);
 		UPCGManagedISMComponent* MatchingResource = nullptr;
-		InSourceComponent->ForEachManagedResource([&MatchingResource, &InParams, &InTargetActor, &Descriptor, SettingsUID](UPCGManagedResource* InResource)
+		if(InParams.SettingsCrc.IsValid())
 		{
-			// Early out if already found a match
-			if (MatchingResource)
+			InSourceComponent->ForEachManagedResource([&MatchingResource, &InParams, &InTargetActor, &Descriptor](UPCGManagedResource* InResource)
 			{
-				return;
-			}
-
-			if (UPCGManagedISMComponent* Resource = Cast<UPCGManagedISMComponent>(InResource))
-			{
-				if (Resource->GetSettingsUID() != SettingsUID || !Resource->CanBeUsed())
+				// Early out if already found a match
+				if (MatchingResource)
 				{
 					return;
 				}
 
-				if (UInstancedStaticMeshComponent* ISMC = Resource->GetComponent())
+				if (UPCGManagedISMComponent* Resource = Cast<UPCGManagedISMComponent>(InResource))
 				{
-					if (IsValid(ISMC) &&
-						ISMC->GetOwner() == InTargetActor &&
-						ISMC->NumCustomDataFloats == InParams.NumCustomDataFloats &&
-						Resource->GetDescriptor() == Descriptor)
+					if (!Resource->GetSettingsCrc().IsValid() || Resource->GetSettingsCrc() != InParams.SettingsCrc || !Resource->CanBeUsed())
+					{
+						return;
+					}
+
+					if (UInstancedStaticMeshComponent* ISMC = Resource->GetComponent())
 					{
-						MatchingResource = Resource;
+						if (IsValid(ISMC) &&
+							ISMC->GetOwner() == InTargetActor &&
+							ISMC->NumCustomDataFloats == InParams.NumCustomDataFloats &&
+							Resource->GetDescriptor() == Descriptor)
+						{
+							MatchingResource = Resource;
+						}
 					}
 				}
-			}
-		});
+			});
+		}
 
 		if (MatchingResource)
 		{
@@ -201,15 +216,28 @@ UPCGManagedISMComponent* UPCGActorHelpers::GetOrCreateManagedISMC(AActor* InTarg
 		Resource->SetRootLocation(InTargetActor->GetRootComponent()->GetComponentLocation());
 	}
 	
-	Resource->SetSettingsUID(SettingsUID);
+	Resource->SetSettingsCrc(InParams.SettingsCrc);
 	InSourceComponent->AddToManagedResources(Resource);
 
 	return Resource;
 }
 
+// deprecated
 USplineMeshComponent* UPCGActorHelpers::GetOrCreateSplineMeshComponent(AActor* InTargetActor, UPCGComponent* InSourceComponent, uint64 SettingsUID, const FPCGSplineMeshComponentBuilderParameters& InParams)
 {
-	UPCGManagedSplineMeshComponent* ManagedComponent = GetOrCreateManagedSplineMeshComponent(InTargetActor, InSourceComponent, SettingsUID, InParams);
+	return GetOrCreateSplineMeshComponent(InTargetActor, InSourceComponent, InParams);
+}
+
+
+// deprecated
+UPCGManagedSplineMeshComponent* UPCGActorHelpers::GetOrCreateManagedSplineMeshComponent(AActor* InTargetActor, UPCGComponent* InSourceComponent, uint64 SettingsUID, const FPCGSplineMeshComponentBuilderParameters& InParams)
+{
+	return GetOrCreateManagedSplineMeshComponent(InTargetActor, InSourceComponent, InParams);
+}
+
+USplineMeshComponent* UPCGActorHelpers::GetOrCreateSplineMeshComponent(AActor* InTargetActor, UPCGComponent* InSourceComponent, const FPCGSplineMeshComponentBuilderParameters& InParams)
+{
+	UPCGManagedSplineMeshComponent* ManagedComponent = GetOrCreateManagedSplineMeshComponent(InTargetActor, InSourceComponent, InParams);
 	if (ManagedComponent)
 	{
 		return ManagedComponent->GetComponent();
@@ -220,7 +248,7 @@ USplineMeshComponent* UPCGActorHelpers::GetOrCreateSplineMeshComponent(AActor* I
 	}
 }
 
-UPCGManagedSplineMeshComponent* UPCGActorHelpers::GetOrCreateManagedSplineMeshComponent(AActor* InTargetActor, UPCGComponent* InSourceComponent, uint64 SettingsUID, const FPCGSplineMeshComponentBuilderParameters& InParams)
+UPCGManagedSplineMeshComponent* UPCGActorHelpers::GetOrCreateManagedSplineMeshComponent(AActor * InTargetActor, UPCGComponent * InSourceComponent, const FPCGSplineMeshComponentBuilderParameters & InParams)
 {
 	check(InTargetActor && InSourceComponent);
 
@@ -237,33 +265,36 @@ UPCGManagedSplineMeshComponent* UPCGActorHelpers::GetOrCreateManagedSplineMeshCo
 		TRACE_CPUPROFILER_EVENT_SCOPE(UPCGActorHelpers::GetOrCreateManagedSplineMeshComponent::FindMatchingManagedSplineMeshComponent);
 
 		UPCGManagedSplineMeshComponent* MatchingResource = nullptr;
-		InSourceComponent->ForEachManagedResource([&MatchingResource, &InParams, &InTargetActor, SettingsUID](UPCGManagedResource* InResource)
+		if(InParams.SettingsCrc.IsValid())
 		{
-			// Early out if already found a match
-			if (MatchingResource)
+			InSourceComponent->ForEachManagedResource([&MatchingResource, &InParams, &InTargetActor](UPCGManagedResource* InResource)
 			{
-				return;
-			}
-
-			if (UPCGManagedSplineMeshComponent* Resource = Cast<UPCGManagedSplineMeshComponent>(InResource))
-			{
-				if (Resource->GetSettingsUID() != SettingsUID || !Resource->CanBeUsed())
+				// Early out if already found a match
+				if (MatchingResource)
 				{
 					return;
 				}
 
-				if (USplineMeshComponent* SplineMeshComponent = Resource->GetComponent())
+				if (UPCGManagedSplineMeshComponent* Resource = Cast<UPCGManagedSplineMeshComponent>(InResource))
 				{
-					if (IsValid(SplineMeshComponent)
-						&& SplineMeshComponent->GetOwner() == InTargetActor
-						&& Resource->GetDescriptor() == InParams.Descriptor
-						&& Resource->GetSplineMeshParams() == InParams.SplineMeshParams)
+					if (!Resource->GetSettingsCrc().IsValid() || Resource->GetSettingsCrc() != InParams.SettingsCrc || !Resource->CanBeUsed())
+					{
+						return;
+					}
+
+					if (USplineMeshComponent* SplineMeshComponent = Resource->GetComponent())
 					{
-						MatchingResource = Resource;
+						if (IsValid(SplineMeshComponent)
+							&& SplineMeshComponent->GetOwner() == InTargetActor
+							&& Resource->GetDescriptor() == InParams.Descriptor
+							&& Resource->GetSplineMeshParams() == InParams.SplineMeshParams)
+						{
+							MatchingResource = Resource;
+						}
 					}
 				}
-			}
-		});
+			});
+		}
 
 		if (MatchingResource)
 		{
@@ -320,7 +351,7 @@ UPCGManagedSplineMeshComponent* UPCGActorHelpers::GetOrCreateManagedSplineMeshCo
 	Resource->SetComponent(SplineMeshComponent);
 	Resource->SetDescriptor(InParams.Descriptor);
 	Resource->SetSplineMeshParams(InParams.SplineMeshParams);
-	Resource->SetSettingsUID(SettingsUID);
+	Resource->SetSettingsCrc(InParams.SettingsCrc);
 	InSourceComponent->AddToManagedResources(Resource);
 
 	return Resource;
diff --git a/Engine/Plugins/PCG/Source/PCG/Private/PCGElement.cpp b/Engine/Plugins/PCG/Source/PCG/Private/PCGElement.cpp
--- a/Engine/Plugins/PCG/Source/PCG/Private/PCGElement.cpp
+++ b/Engine/Plugins/PCG/Source/PCG/Private/PCGElement.cpp
@@ -870,7 +870,7 @@ void IPCGElement::GetDependenciesCrc(const FPCGDataCollection& InInput, const UP
 
 	if (InSettings)
 	{
-		const FPCGCrc& SettingsCrc = InSettings->GetCachedCrc();
+		FPCGCrc SettingsCrc = InSettings->GetSettingsCrc();
 		if (ensure(SettingsCrc.IsValid()))
 		{
 			Crc.Combine(SettingsCrc);
diff --git a/Engine/Plugins/PCG/Source/PCG/Private/PCGSettings.cpp b/Engine/Plugins/PCG/Source/PCG/Private/PCGSettings.cpp
--- a/Engine/Plugins/PCG/Source/PCG/Private/PCGSettings.cpp
+++ b/Engine/Plugins/PCG/Source/PCG/Private/PCGSettings.cpp
@@ -1108,9 +1108,32 @@ void UPCGSettings::FixingOverridableParamPropertyClass(FPCGSettingsOverridablePa
 	}
 }
 
+// deprecated
+const FPCGCrc& UPCGSettings::GetCachedCrc() const
+{
+	static FPCGCrc InvalidCrc;
+	return InvalidCrc;
+}
+
+FPCGCrc UPCGSettings::GetSettingsCrc() const
+{
+#if WITH_EDITOR
+	if (!CachedCrc.IsValid())
+	{
+		const_cast<UPCGSettings*>(this)->CacheCrc();
+	}
+	return CachedCrc;
+#else
+	// In non-editor builds we don't track changes to the UPCGSettings object so we always compute it
+	return FPCGSettingsObjectCrc::PCGCrc(this);
+#endif
+}
+
 void UPCGSettings::CacheCrc()
 {
+#if WITH_EDITOR
 	CachedCrc = FPCGSettingsObjectCrc::PCGCrc(this);
+#endif
 }
 
 TArray<FPCGPinProperties> UPCGSettings::DefaultPointInputPinProperties() const
diff --git a/Engine/Plugins/PCG/Source/PCG/Public/Helpers/PCGActorHelpers.h b/Engine/Plugins/PCG/Source/PCG/Public/Helpers/PCGActorHelpers.h
--- a/Engine/Plugins/PCG/Source/PCG/Public/Helpers/PCGActorHelpers.h
+++ b/Engine/Plugins/PCG/Source/PCG/Public/Helpers/PCGActorHelpers.h
@@ -2,6 +2,7 @@
 
 #pragma once
 
+#include "PCGCrc.h"
 #include "Elements/PCGSplineMeshParams.h"
 #include "MeshSelectors/PCGISMDescriptor.h"
 
@@ -47,6 +48,7 @@ struct FPCGISMComponentBuilderParams
 
 	FPCGSoftISMComponentDescriptor Descriptor;
 	int32 NumCustomDataFloats = 0;
+	FPCGCrc SettingsCrc;
 	bool bAllowDescriptorChanges = true;
 
 PRAGMA_DISABLE_DEPRECATION_WARNINGS
@@ -68,6 +70,7 @@ struct FPCGSplineMeshComponentBuilderParameters
 {
 	FSplineMeshComponentDescriptor Descriptor;
 	FPCGSplineMeshParams SplineMeshParams;
+	FPCGCrc SettingsCrc;
 
 	friend inline uint32 GetTypeHash(const FPCGSplineMeshComponentBuilderParameters& Key)
 	{
@@ -88,10 +91,20 @@ PRAGMA_DISABLE_DEPRECATION_WARNINGS
 	static UPCGManagedISMComponent* GetOrCreateManagedISMC(AActor* InTargetActor, UPCGComponent* SourceComponent, uint64 SettingsUID, const FPCGISMCBuilderParameters& Params);
 PRAGMA_ENABLE_DEPRECATION_WARNINGS
 
+	UE_DEPRECATED(5.6, "Use signature with no SettingsUID. Make sure to fill Params.SettingsCRC to enable component reuse.")
 	static UInstancedStaticMeshComponent* GetOrCreateISMC(AActor* InTargetActor, UPCGComponent* SourceComponent, uint64 SettingsUID, const FPCGISMComponentBuilderParams& Params);
+	UE_DEPRECATED(5.6, "Use signature with no SettingsUID. Make sure to fill Params.SettingsCRC to enable component reuse.")	
 	static UPCGManagedISMComponent* GetOrCreateManagedISMC(AActor* InTargetActor, UPCGComponent* SourceComponent, uint64 SettingsUID, const FPCGISMComponentBuilderParams& Params);
+	UE_DEPRECATED(5.6, "Use signature with no SettingsUID. Make sure to fill Params.SettingsCRC to enable component reuse.")
 	static USplineMeshComponent* GetOrCreateSplineMeshComponent(AActor* InTargetActor, UPCGComponent* SourceComponent, uint64 SettingsUID, const FPCGSplineMeshComponentBuilderParameters& Params);
+	UE_DEPRECATED(5.6, "Use signature with no SettingsUID. Make sure to fill Params.SettingsCRC to enable component reuse.")
 	static UPCGManagedSplineMeshComponent* GetOrCreateManagedSplineMeshComponent(AActor* InTargetActor, UPCGComponent* SourceComponent, uint64 SettingsUID, const FPCGSplineMeshComponentBuilderParameters& Params);
+	
+	static UInstancedStaticMeshComponent* GetOrCreateISMC(AActor* InTargetActor, UPCGComponent* SourceComponent, const FPCGISMComponentBuilderParams& Params);
+	static UPCGManagedISMComponent* GetOrCreateManagedISMC(AActor* InTargetActor, UPCGComponent* SourceComponent, const FPCGISMComponentBuilderParams& Params);
+	static USplineMeshComponent* GetOrCreateSplineMeshComponent(AActor* InTargetActor, UPCGComponent* SourceComponent, const FPCGSplineMeshComponentBuilderParameters& Params);
+	static UPCGManagedSplineMeshComponent* GetOrCreateManagedSplineMeshComponent(AActor* InTargetActor, UPCGComponent* SourceComponent, const FPCGSplineMeshComponentBuilderParameters& Params);
+
 	static bool DeleteActors(UWorld* World, const TArray<TSoftObjectPtr<AActor>>& ActorsToDelete);
 
 	template <typename T, typename = typename std::enable_if_t<std::is_base_of_v<AActor, T>>>
diff --git a/Engine/Plugins/PCG/Source/PCG/Public/PCGManagedResource.h b/Engine/Plugins/PCG/Source/PCG/Public/PCGManagedResource.h
--- a/Engine/Plugins/PCG/Source/PCG/Public/PCGManagedResource.h
+++ b/Engine/Plugins/PCG/Source/PCG/Public/PCGManagedResource.h
@@ -261,8 +261,14 @@ public:
 
 	void SetRootLocation(const FVector& InRootLocation);
 
-	uint64 GetSettingsUID() const { return SettingsUID; }
-	void SetSettingsUID(uint64 InSettingsUID) { SettingsUID = InSettingsUID; }
+	FPCGCrc GetSettingsCrc() const { return SettingsCrc; }
+	void SetSettingsCrc(const FPCGCrc& InSettingsCrc) { SettingsCrc = InSettingsCrc; }
+
+	UE_DEPRECATED(5.6, "Use GetSettingsCrc instead")
+	uint64 GetSettingsUID() const { return -1; }
+	
+	UE_DEPRECATED(5.6, "Use SetSettingsCrc instead")
+	void SetSettingsUID(uint64 InSettingsUID) { }
 
 protected:
 	UPROPERTY()
@@ -277,8 +283,8 @@ protected:
 	UPROPERTY()
 	FVector RootLocation = FVector::ZeroVector;
 
-	UPROPERTY(Transient)
-	uint64 SettingsUID = -1; // purposefully a value that will never happen in data
+	UPROPERTY()
+	FPCGCrc SettingsCrc;
 
 	// Cached raw pointer to ISM component
 	mutable UInstancedStaticMeshComponent* CachedRawComponentPtr = nullptr;
@@ -305,8 +311,14 @@ public:
 	void SetSplineMeshParams(const FPCGSplineMeshParams& InSplineMeshParams) { SplineMeshParams = InSplineMeshParams; }
 	const FPCGSplineMeshParams& GetSplineMeshParams() const { return SplineMeshParams; }
 
-	uint64 GetSettingsUID() const { return SettingsUID; }
-	void SetSettingsUID(uint64 InSettingsUID) { SettingsUID = InSettingsUID; }
+	FPCGCrc GetSettingsCrc() const { return SettingsCrc; }
+	void SetSettingsCrc(const FPCGCrc& InSettingsCrc) { SettingsCrc = InSettingsCrc; }
+
+	UE_DEPRECATED(5.6, "Use GetSettingsCrc instead")
+	uint64 GetSettingsUID() const { return -1; }
+
+	UE_DEPRECATED(5.6, "Use SetSettingsCrc instead")
+	void SetSettingsUID(uint64 InSettingsUID) { }
 
 protected:
 	UPROPERTY()
@@ -315,8 +327,8 @@ protected:
 	UPROPERTY()
 	FPCGSplineMeshParams SplineMeshParams;
 
-	UPROPERTY(Transient)
-	uint64 SettingsUID = -1; // purposefully a value that will never happen in data
+	UPROPERTY()
+	FPCGCrc SettingsCrc;
 
 	// Cached raw pointer to USplineMeshComponent
 	mutable USplineMeshComponent* CachedRawComponentPtr = nullptr;
diff --git a/Engine/Plugins/PCG/Source/PCG/Public/PCGSettings.h b/Engine/Plugins/PCG/Source/PCG/Public/PCGSettings.h
--- a/Engine/Plugins/PCG/Source/PCG/Public/PCGSettings.h
+++ b/Engine/Plugins/PCG/Source/PCG/Public/PCGSettings.h
@@ -279,8 +279,10 @@ public:
 	// Get the seed, combined with optional PCGComponent seed
 	int GetSeed(const UPCGComponent* InSourceComponent = nullptr) const;
 
-	/** Get the Crc value precomputed for these settings. */
-	const FPCGCrc& GetCachedCrc() const { return CachedCrc; }
+	UE_DEPRECATED(5.6, "Use GetSettingsCrc instead")
+	const FPCGCrc& GetCachedCrc() const;
+
+	FPCGCrc GetSettingsCrc() const;
 
 #if WITH_EDITOR
 	/** Puts node title on node body, reducing overall node size */
@@ -614,8 +616,10 @@ private:
 	/** Calculate Crc for these settings and save it. */
 	void CacheCrc();
 
+#if WITH_EDITORONLY_DATA
 	/** The cached Crc for these settings. */
 	FPCGCrc CachedCrc;
+#endif
 
 	friend class FPCGGraphCompilerGPU;
 	friend class UPCGDataBinding;
diff --git a/Engine/Plugins/PCG/Source/PCGEditor/Private/DataVisualizations/PCGSpatialDataVisualization.cpp b/Engine/Plugins/PCG/Source/PCGEditor/Private/DataVisualizations/PCGSpatialDataVisualization.cpp
--- a/Engine/Plugins/PCG/Source/PCGEditor/Private/DataVisualizations/PCGSpatialDataVisualization.cpp
+++ b/Engine/Plugins/PCG/Source/PCGEditor/Private/DataVisualizations/PCGSpatialDataVisualization.cpp
@@ -156,6 +156,7 @@ void IPCGSpatialDataVisualization::ExecuteDebugDisplay(FPCGContext* Context, con
 	}
 
 	FPCGISMComponentBuilderParams Params[2];
+	Params[0].SettingsCrc = SettingsInterface->GetSettings()->GetSettingsCrc();
 	Params[0].NumCustomDataFloats = NumCustomData;
 	Params[0].Descriptor.StaticMesh = Mesh;
 	Params[0].Descriptor.OverrideMaterials = Materials;
@@ -195,7 +196,7 @@ void IPCGSpatialDataVisualization::ExecuteDebugDisplay(FPCGContext* Context, con
 			continue;
 		}
 
-		UInstancedStaticMeshComponent* ISMC = UPCGActorHelpers::GetOrCreateISMC(TargetActor, Context->SourceComponent.Get(), SettingsInterface->GetSettings()->GetStableUID(), Params[Direction]);
+		UInstancedStaticMeshComponent* ISMC = UPCGActorHelpers::GetOrCreateISMC(TargetActor, Context->SourceComponent.Get(), Params[Direction]);
 		check(ISMC && ISMC->NumCustomDataFloats == NumCustomData);
 
 		ISMC->ComponentTags.AddUnique(PCGHelpers::DefaultPCGDebugTag);
