[PCG] Enable CreateTargetActor reuse by default

diff --git a/Engine/Plugins/PCG/Source/PCG/Private/Elements/PCGCreateTargetActor.cpp b/Engine/Plugins/PCG/Source/PCG/Private/Elements/PCGCreateTargetActor.cpp
--- a/Engine/Plugins/PCG/Source/PCG/Private/Elements/PCGCreateTargetActor.cpp
+++ b/Engine/Plugins/PCG/Source/PCG/Private/Elements/PCGCreateTargetActor.cpp
@@ -48,7 +48,7 @@ namespace PCGCreateTargetActor
 {
 	static TAutoConsoleVariable<bool> CVarCreateTargetActorAllowReuse(
 		TEXT("pcg.CreateTargetActor.AllowReuse"),
-		false,
+		true,
 		TEXT("Controls whether PCG Create Target Actor node can reuse actors when re-executing (requires Create Target Actor node resave so they have a Stable Reuse Guid)"));
 
 	FPCGCrc GetAdditionalDependenciesCrc(const UPCGActorHelpers::FSpawnDefaultActorParams& InParams, AActor* TargetActor)
