[PCG] GPU SM Spawner - fix SM bounds not being written to points.

#rnx

diff --git a/Engine/Plugins/PCG/Shaders/Private/Elements/PCGStaticMeshSpawner.usf b/Engine/Plugins/PCG/Shaders/Private/Elements/PCGStaticMeshSpawner.usf
--- a/Engine/Plugins/PCG/Shaders/Private/Elements/PCGStaticMeshSpawner.usf
+++ b/Engine/Plugins/PCG/Shaders/Private/Elements/PCGStaticMeshSpawner.usf
@@ -123,17 +123,26 @@ void Main(uint3 GroupId : SV_GroupID, uint GroupIndex : SV_GroupIndex)
 	}
 
 	// Write output point data (pass through).
-	// TODO - support writing mesh bounds.
 	Out_SetPosition(Out_DataIndex, Out_ElementIndex, In_GetPosition(In_DataIndex, In_ElementIndex));
 	Out_SetRotation(Out_DataIndex, Out_ElementIndex, In_GetRotation(In_DataIndex, In_ElementIndex));
 	Out_SetScale(Out_DataIndex, Out_ElementIndex, In_GetScale(In_DataIndex, In_ElementIndex));
-	Out_SetBoundsMin(Out_DataIndex, Out_ElementIndex, In_GetBoundsMin(In_DataIndex, In_ElementIndex));
-	Out_SetBoundsMax(Out_DataIndex, Out_ElementIndex, In_GetBoundsMax(In_DataIndex, In_ElementIndex));
 	Out_SetColor(Out_DataIndex, Out_ElementIndex, In_GetColor(In_DataIndex, In_ElementIndex));
 	Out_SetSeed(Out_DataIndex, Out_ElementIndex, In_GetSeed(In_DataIndex, In_ElementIndex));
 	Out_SetDensity(Out_DataIndex, Out_ElementIndex, In_GetDensity(In_DataIndex, In_ElementIndex));
 	Out_SetSteepness(Out_DataIndex, Out_ElementIndex, In_GetSteepness(In_DataIndex, In_ElementIndex));
 
+	// Either apply SM bounds to points or pass bounds through.
+	if (SMSpawner_ShouldApplyBounds())
+	{
+		Out_SetBoundsMin(Out_DataIndex, Out_ElementIndex, SMSpawner_GetPrimitiveMeshBoundsMin(PrimitiveIndex));
+		Out_SetBoundsMax(Out_DataIndex, Out_ElementIndex, SMSpawner_GetPrimitiveMeshBoundsMax(PrimitiveIndex));
+	}
+	else
+	{
+		Out_SetBoundsMin(Out_DataIndex, Out_ElementIndex, In_GetBoundsMin(In_DataIndex, In_ElementIndex));
+		Out_SetBoundsMax(Out_DataIndex, Out_ElementIndex, In_GetBoundsMax(In_DataIndex, In_ElementIndex));
+	}
+
 	CopyAttributesToOutput(In_DataIndex, In_ElementIndex, Out_DataIndex, Out_ElementIndex);
 
 	if (SMSpawner_GetSelectedMeshAttributeId() != (uint)-1)
diff --git a/Engine/Plugins/PCG/Source/PCG/Private/Compute/DataInterfaces/Elements/PCGStaticMeshSpawnerDataInterface.cpp b/Engine/Plugins/PCG/Source/PCG/Private/Compute/DataInterfaces/Elements/PCGStaticMeshSpawnerDataInterface.cpp
--- a/Engine/Plugins/PCG/Source/PCG/Private/Compute/DataInterfaces/Elements/PCGStaticMeshSpawnerDataInterface.cpp
+++ b/Engine/Plugins/PCG/Source/PCG/Private/Compute/DataInterfaces/Elements/PCGStaticMeshSpawnerDataInterface.cpp
@@ -27,6 +27,20 @@ void UPCGStaticMeshSpawnerDataInterface::GetSupportedInputs(TArray<FShaderFuncti
 		.SetName(TEXT("SMSpawner_GetNumPrimitives"))
 		.AddReturnType(EShaderFundamentalType::Uint); // Num primitives
 
+	OutFunctions.AddDefaulted_GetRef()
+		.SetName(TEXT("SMSpawner_ShouldApplyBounds"))
+		.AddReturnType(EShaderFundamentalType::Bool);
+
+	OutFunctions.AddDefaulted_GetRef()
+		.SetName(TEXT("SMSpawner_GetPrimitiveMeshBoundsMin"))
+		.AddReturnType(EShaderFundamentalType::Float, 3) // Local bounds min
+		.AddParam(EShaderFundamentalType::Uint); // InPrimitiveIndex
+
+	OutFunctions.AddDefaulted_GetRef()
+		.SetName(TEXT("SMSpawner_GetPrimitiveMeshBoundsMax"))
+		.AddReturnType(EShaderFundamentalType::Float, 3) // Local bounds max
+		.AddParam(EShaderFundamentalType::Uint); // InPrimitiveIndex
+
 	OutFunctions.AddDefaulted_GetRef()
 		.SetName(TEXT("SMSpawner_GetAttributeIdOffsetStride"))
 		.AddReturnType(EShaderFundamentalType::Uint, 4)
@@ -55,11 +69,14 @@ void UPCGStaticMeshSpawnerDataInterface::GetSupportedInputs(TArray<FShaderFuncti
 BEGIN_SHADER_PARAMETER_STRUCT(FPCGStaticMeshSpawnerDataInterfaceParameters,)
 	SHADER_PARAMETER_ARRAY(FUintVector4, AttributeIdOffsetStrides, [UPCGStaticMeshSpawnerDataInterface::MAX_ATTRIBUTES])
 	SHADER_PARAMETER_ARRAY(FIntVector4, PrimitiveStringKeys, [PCGComputeConstants::MAX_PRIMITIVE_COMPONENTS_PER_SPAWNER])
+	SHADER_PARAMETER_ARRAY(FVector4f, PrimitiveMeshBoundsMin, [PCGComputeConstants::MAX_PRIMITIVE_COMPONENTS_PER_SPAWNER])
+	SHADER_PARAMETER_ARRAY(FVector4f, PrimitiveMeshBoundsMax, [PCGComputeConstants::MAX_PRIMITIVE_COMPONENTS_PER_SPAWNER])
 	SHADER_PARAMETER_SCALAR_ARRAY(float, SelectionCDF, [PCGComputeConstants::MAX_PRIMITIVE_COMPONENTS_PER_SPAWNER])
 	SHADER_PARAMETER(uint32, NumAttributes)
 	SHADER_PARAMETER(uint32, NumPrimitives)
 	SHADER_PARAMETER(int32, SelectorAttributeId)
 	SHADER_PARAMETER(int32, SelectedMeshAttributeId)
+	SHADER_PARAMETER(uint32, ApplyBounds)
 END_SHADER_PARAMETER_STRUCT()
 
 void UPCGStaticMeshSpawnerDataInterface::GetShaderParameters(TCHAR const* UID, FShaderParametersMetadataBuilder& InOutBuilder, FShaderParametersMetadataAllocations& InOutAllocations) const
@@ -78,17 +95,50 @@ void UPCGStaticMeshSpawnerDataInterface::GetHLSL(FString& OutHLSL, FString const
 
 	OutHLSL += FString::Format(TEXT(
 		"int {DataInterfaceName}_SelectorAttributeId;\n"
-		"int SMSpawner_GetSelectorAttributeId_{DataInterfaceName}() { return {DataInterfaceName}_SelectorAttributeId; }\n"
-		"\n"
 		"uint {DataInterfaceName}_NumAttributes;\n"
-		"uint4 {DataInterfaceName}_AttributeIdOffsetStrides[{MaxAttributes}];\n"
-		"uint SMSpawner_GetNumAttributes_{DataInterfaceName}() { return {DataInterfaceName}_NumAttributes; }\n"
-		"uint4 SMSpawner_GetAttributeIdOffsetStride_{DataInterfaceName}(uint InAttributeIndex) { return {DataInterfaceName}_AttributeIdOffsetStrides[InAttributeIndex]; }\n"
-		"\n"
 		"uint {DataInterfaceName}_NumPrimitives;\n"
+		"uint {DataInterfaceName}_ApplyBounds;"
+		"int {DataInterfaceName}_SelectedMeshAttributeId;\n"
+		"uint4 {DataInterfaceName}_AttributeIdOffsetStrides[{MaxAttributes}];\n"
+		"float4 {DataInterfaceName}_PrimitiveMeshBoundsMin[{MaxPrimitives}];"
+		"float4 {DataInterfaceName}_PrimitiveMeshBoundsMax[{MaxPrimitives}];"
 		"int4 {DataInterfaceName}_PrimitiveStringKeys[{MaxPrimitives}];"
 		"DECLARE_SCALAR_ARRAY(float, {DataInterfaceName}_SelectionCDF, {MaxPrimitives});\n"
-		"uint SMSpawner_GetNumPrimitives_{DataInterfaceName}() { return {DataInterfaceName}_NumPrimitives; }\n"
+		"\n"
+		"int SMSpawner_GetSelectorAttributeId_{DataInterfaceName}()\n"
+		"{\n"
+		"	return {DataInterfaceName}_SelectorAttributeId;\n"
+		"}\n"
+		"\n"
+		"uint SMSpawner_GetNumAttributes_{DataInterfaceName}()\n"
+		"{\n"
+		"	return {DataInterfaceName}_NumAttributes;\n"
+		"}\n"
+		"\n"
+		"uint4 SMSpawner_GetAttributeIdOffsetStride_{DataInterfaceName}(uint InAttributeIndex)\n"
+		"{\n"
+		"	return {DataInterfaceName}_AttributeIdOffsetStrides[InAttributeIndex];\n"
+		"}\n"
+		"\n"
+		"uint SMSpawner_GetNumPrimitives_{DataInterfaceName}()\n"
+		"{\n"
+		"	return {DataInterfaceName}_NumPrimitives;\n"
+		"}\n"
+		"\n"
+		"bool SMSpawner_ShouldApplyBounds_{DataInterfaceName}()\n"
+		"{\n"
+		"	return {DataInterfaceName}_ApplyBounds > 0;\n"
+		"}\n"
+		"\n"
+		"float3 SMSpawner_GetPrimitiveMeshBoundsMin_{DataInterfaceName}(uint InPrimitiveIndex)\n"
+		"{\n"
+		"	return {DataInterfaceName}_PrimitiveMeshBoundsMin[InPrimitiveIndex].xyz;\n"
+		"}\n"
+		"\n"
+		"float3 SMSpawner_GetPrimitiveMeshBoundsMax_{DataInterfaceName}(uint InPrimitiveIndex)\n"
+		"{\n"
+		"	return {DataInterfaceName}_PrimitiveMeshBoundsMax[InPrimitiveIndex].xyz;\n"
+		"}\n"
 		"\n"
 		"int SMSpawner_GetPrimitiveStringKey_{DataInterfaceName}(uint InPrimitiveIndex)\n"
 		"{\n"
@@ -108,11 +158,16 @@ void UPCGStaticMeshSpawnerDataInterface::GetHLSL(FString& OutHLSL, FString const
 		"	return (uint)-1;\n"
 		"}\n"
 		"\n"
-		"float SMSpawner_GetPrimitiveSelectionCDF_{DataInterfaceName}(uint InPrimitiveIndex) { return GET_SCALAR_ARRAY_ELEMENT({DataInterfaceName}_SelectionCDF, InPrimitiveIndex); }\n"
+		"float SMSpawner_GetPrimitiveSelectionCDF_{DataInterfaceName}(uint InPrimitiveIndex)\n"
+		"{\n"
+		"	return GET_SCALAR_ARRAY_ELEMENT({DataInterfaceName}_SelectionCDF, InPrimitiveIndex);\n"
+		"}\n"
 		"\n"
-		"int {DataInterfaceName}_SelectedMeshAttributeId;\n"
-		"int SMSpawner_GetSelectedMeshAttributeId_{DataInterfaceName}() { return {DataInterfaceName}_SelectedMeshAttributeId; }\n"
-		), TemplateArgs);
+		"int SMSpawner_GetSelectedMeshAttributeId_{DataInterfaceName}()\n"
+		"{\n"
+		"	return {DataInterfaceName}_SelectedMeshAttributeId;\n"
+		"}\n"
+	), TemplateArgs);
 }
 
 UComputeDataProvider* UPCGStaticMeshSpawnerDataInterface::CreateDataProvider(TObjectPtr<UObject> InBinding, uint64 InInputMask, uint64 InOutputMask) const
@@ -130,6 +185,7 @@ UComputeDataProvider* UPCGStaticMeshSpawnerDataInterface::CreateDataProvider(TOb
 		DataProvider->SelectionCDF = Primitives->SelectionCDF;
 		DataProvider->SelectorAttributeId = Primitives->SelectorAttributeId;
 		DataProvider->PrimitiveStringKeys = Primitives->PrimitiveStringKeys;
+		DataProvider->PrimitiveMeshBounds = Primitives->PrimitiveMeshBounds;
 
 		const FPCGKernelAttributeIDAndType* FoundAttribute = Binding->GetAttributeLookupTable().Find(DataProvider->Settings->OutAttributeName);
 		DataProvider->SelectedMeshAttributeId = ensure(FoundAttribute) ? FoundAttribute->Id : -1;
@@ -140,7 +196,7 @@ UComputeDataProvider* UPCGStaticMeshSpawnerDataInterface::CreateDataProvider(TOb
 
 FComputeDataProviderRenderProxy* UPCGStaticMeshSpawnerDataProvider::GetRenderProxy()
 {
-	return new FPCGStaticMeshSpawnerDataProviderProxy(AttributeIdOffsetStrides, SelectorAttributeId, PrimitiveStringKeys, SelectionCDF, SelectedMeshAttributeId);
+	return new FPCGStaticMeshSpawnerDataProviderProxy(AttributeIdOffsetStrides, SelectorAttributeId, PrimitiveStringKeys, SelectionCDF, SelectedMeshAttributeId, PrimitiveMeshBounds);
 }
 
 bool FPCGStaticMeshSpawnerDataProviderProxy::IsValid(FValidationData const& InValidationData) const
@@ -179,5 +235,13 @@ void FPCGStaticMeshSpawnerDataProviderProxy::GatherDispatchData(FDispatchData co
 		{
 			GET_SCALAR_ARRAY_ELEMENT(Parameters.SelectionCDF, Index) = SelectionCDF[Index];
 		}
+
+		Parameters.ApplyBounds = PrimitiveMeshBounds.IsEmpty() ? 0 : 1;
+
+		for (int32 Index = 0; Index < PrimitiveMeshBounds.Num(); ++Index)
+		{
+			Parameters.PrimitiveMeshBoundsMin[Index] = FVector4f(PrimitiveMeshBounds[Index].Min.X, PrimitiveMeshBounds[Index].Min.Y, PrimitiveMeshBounds[Index].Min.Z, /*Unused*/0.0f);
+			Parameters.PrimitiveMeshBoundsMax[Index] = FVector4f(PrimitiveMeshBounds[Index].Max.X, PrimitiveMeshBounds[Index].Max.Y, PrimitiveMeshBounds[Index].Max.Z, /*Unused*/0.0f);
+		}
 	}
 }
diff --git a/Engine/Plugins/PCG/Source/PCG/Private/Compute/DataInterfaces/Elements/PCGStaticMeshSpawnerDataInterface.h b/Engine/Plugins/PCG/Source/PCG/Private/Compute/DataInterfaces/Elements/PCGStaticMeshSpawnerDataInterface.h
--- a/Engine/Plugins/PCG/Source/PCG/Private/Compute/DataInterfaces/Elements/PCGStaticMeshSpawnerDataInterface.h
+++ b/Engine/Plugins/PCG/Source/PCG/Private/Compute/DataInterfaces/Elements/PCGStaticMeshSpawnerDataInterface.h
@@ -51,6 +51,9 @@ public:
 	UPROPERTY()
 	TArray<int32> PrimitiveStringKeys;
 
+	UPROPERTY()
+	TArray<FBox> PrimitiveMeshBounds;
+
 	UPROPERTY()
 	TArray<float> SelectionCDF;
 
@@ -67,11 +70,18 @@ public:
 class FPCGStaticMeshSpawnerDataProviderProxy : public FComputeDataProviderRenderProxy
 {
 public:
-	FPCGStaticMeshSpawnerDataProviderProxy(TArray<FUintVector4>& InAttributeIdOffsetStrides, int32 InSelectorAttributeId, TArray<int32>& InPrimitiveStringKeys, TArray<float> InSelectionCDF, int32 InSelectedMeshAttributeId)
+	FPCGStaticMeshSpawnerDataProviderProxy(
+		const TArray<FUintVector4>& InAttributeIdOffsetStrides,
+		int32 InSelectorAttributeId,
+		const TArray<int32>& InPrimitiveStringKeys,
+		TArray<float> InSelectionCDF,
+		int32 InSelectedMeshAttributeId,
+		const TArray<FBox>& InPrimitiveMeshBounds)
 		: AttributeIdOffsetStrides(InAttributeIdOffsetStrides)
 		, SelectionCDF(InSelectionCDF)
 		, SelectorAttributeId(InSelectorAttributeId)
 		, PrimitiveStringKeys(InPrimitiveStringKeys)
+		, PrimitiveMeshBounds(InPrimitiveMeshBounds)
 		, SelectedMeshAttributeId(InSelectedMeshAttributeId)
 	{}
 
@@ -89,5 +99,7 @@ protected:
 	int32 SelectorAttributeId = INDEX_NONE;
 	TArray<int32> PrimitiveStringKeys;
 
+	TArray<FBox> PrimitiveMeshBounds;
+
 	int32 SelectedMeshAttributeId = INDEX_NONE;
 };
diff --git a/Engine/Plugins/PCG/Source/PCG/Private/Compute/Elements/PCGComputeGraphElement.cpp b/Engine/Plugins/PCG/Source/PCG/Private/Compute/Elements/PCGComputeGraphElement.cpp
--- a/Engine/Plugins/PCG/Source/PCG/Private/Compute/Elements/PCGComputeGraphElement.cpp
+++ b/Engine/Plugins/PCG/Source/PCG/Private/Compute/Elements/PCGComputeGraphElement.cpp
@@ -361,6 +361,7 @@ bool FPCGComputeGraphElement::SetupProceduralISMComponents(FPCGContext* InContex
 		TArray<float> PrimitiveSelectionCDF;
 		int32 SelectorAttributeId = -1;
 		TArray<int32> PrimitiveStringKeys;
+		TArray<FBox> PrimitiveMeshBounds;
 		TArray<FPCGProceduralISMComponentDescriptor> ComponentsToCreate;
 
 		if (const UPCGMeshSelectorByAttribute* SelectorByAttribute = Cast<UPCGMeshSelectorByAttribute>(SpawnerSettings->MeshSelectorParameters))
@@ -418,6 +419,11 @@ bool FPCGComputeGraphElement::SetupProceduralISMComponents(FPCGContext* InContex
 
 				PrimitiveStringKeys.Emplace(StringKey);
 				ComponentsToCreate.Add(MoveTemp(Descriptor));
+
+				if (SpawnerSettings->bApplyMeshBoundsToPoints)
+				{
+					PrimitiveMeshBounds.Add(Descriptor.StaticMesh->GetBoundingBox());
+				}
 			}
 
 			PrimitiveSelectionCDF.SetNumZeroed(ComponentsToCreate.Num());
@@ -461,6 +467,11 @@ bool FPCGComputeGraphElement::SetupProceduralISMComponents(FPCGContext* InContex
 					Descriptor.NumCustomFloats = CustomFloatCount;
 					Descriptor.StaticMesh = StaticMesh;
 					ComponentsToCreate.Add(MoveTemp(Descriptor));
+
+					if (SpawnerSettings->bApplyMeshBoundsToPoints)
+					{
+						PrimitiveMeshBounds.Add(Descriptor.StaticMesh->GetBoundingBox());
+					}
 				}
 				else
 				{
@@ -490,6 +501,7 @@ bool FPCGComputeGraphElement::SetupProceduralISMComponents(FPCGContext* InContex
 		Primitives.SelectorAttributeId = SelectorAttributeId;
 		Primitives.SelectionCDF = MoveTemp(PrimitiveSelectionCDF);
 		Primitives.PrimitiveStringKeys = MoveTemp(PrimitiveStringKeys);
+		Primitives.PrimitiveMeshBounds = MoveTemp(PrimitiveMeshBounds);
 
 		for (const FPCGProceduralISMComponentDescriptor& Desc : ComponentsToCreate)
 		{
diff --git a/Engine/Plugins/PCG/Source/PCG/Private/Compute/PCGDataBinding.h b/Engine/Plugins/PCG/Source/PCG/Private/Compute/PCGDataBinding.h
--- a/Engine/Plugins/PCG/Source/PCG/Private/Compute/PCGDataBinding.h
+++ b/Engine/Plugins/PCG/Source/PCG/Private/Compute/PCGDataBinding.h
@@ -18,6 +18,9 @@ struct FPCGSpawnerPrimitives
 	UPROPERTY()
 	TArray<TObjectPtr<UPrimitiveComponent>> Primitives;
 
+	UPROPERTY()
+	TArray<FBox> PrimitiveMeshBounds;
+
 	/** Cumulative distribution function values (one per primitive) to enable choosing a primitive based on a random draw value. */
 	UPROPERTY()
 	TArray<float> SelectionCDF;
