[PCG] Fix a crash caused by trying to debug too many points

Note: When hitting the instance count limit for ISMs, we were crashing by adding more custom data than can fit. This change aims to limit an ISM instance count provided by PCG when debugging more than the instance limit number of points. (Currently 16M)

diff --git a/Engine/Plugins/PCG/Source/PCGEditor/Private/DataVisualizations/PCGSpatialDataVisualization.cpp b/Engine/Plugins/PCG/Source/PCGEditor/Private/DataVisualizations/PCGSpatialDataVisualization.cpp
--- a/Engine/Plugins/PCG/Source/PCGEditor/Private/DataVisualizations/PCGSpatialDataVisualization.cpp
+++ b/Engine/Plugins/PCG/Source/PCGEditor/Private/DataVisualizations/PCGSpatialDataVisualization.cpp
@@ -132,10 +132,20 @@ void IPCGSpatialDataVisualization::ExecuteDebugDisplay(FPCGContext* Context, con
 	const bool bIsRelative = DebugSettings.ScaleMethod == EPCGDebugVisScaleMethod::Relative;
 	const bool bScaleWithExtents = DebugSettings.ScaleMethod == EPCGDebugVisScaleMethod::Extents;
 	const FVector MeshExtents = Mesh->GetBoundingBox().GetExtent();
-	const FVector MeshCenter = Mesh->GetBoundingBox().GetCenter();
 
-	for (const FPCGPoint& Point : Points)
+	/* Note: A re-used ISMC may have any number of pre-existing instances, so this won't prevent going over the max. But,
+	 * the renderer is robust to over-instancing attempts and will not crash. However, MAX_INSTANCE_ID still serves as a
+	 * good, scalable heuristic for a max limit.
+	 */
+	const int32 NumDesiredInstances = FMath::Min(Points.Num(), static_cast<int32>(MAX_INSTANCE_ID));
+	if (NumDesiredInstances != Points.Num())
 	{
+		PCGE_LOG_C(Error, GraphAndLog, Context, FText::Format(LOCTEXT( "DebugPointsOverLimit", "Debug point display ({0}) surpassed the max instance limit ({1}) and will be clamped."), FText::AsNumber(Points.Num()), FText::AsNumber(MAX_INSTANCE_ID)));
+	}
+
+	for (int i = 0; i < NumDesiredInstances; ++i)
+	{
+		const FPCGPoint& Point = Points[i];
 		TArray<FTransform>& Instances = ((bIsAbsolute || Point.Transform.GetDeterminant() >= 0) ? ForwardInstances : ReverseInstances);
 		FTransform& InstanceTransform = Instances.Add_GetRef(Point.Transform);
 		if (bIsRelative)
@@ -187,6 +197,8 @@ void IPCGSpatialDataVisualization::ExecuteDebugDisplay(FPCGContext* Context, con
 	Params[1] = Params[0];
 	Params[1].Descriptor.bReverseCulling = true;
 
+	// Since the instance count is global, track the current instances applied and previously belong to the ISMCs.
+	int32 NumCurrentInstances = 0;
 	for (int32 Direction = 0; Direction < 2; ++Direction)
 	{
 		TArray<FTransform>& Instances = (Direction == 0 ? ForwardInstances : ReverseInstances);
@@ -201,6 +213,21 @@ void IPCGSpatialDataVisualization::ExecuteDebugDisplay(FPCGContext* Context, con
 
 		ISMC->ComponentTags.AddUnique(PCGHelpers::DefaultPCGDebugTag);
 		const int32 PreExistingInstanceCount = ISMC->GetInstanceCount();
+		NumCurrentInstances += PreExistingInstanceCount;
+
+		// The renderer is robust to going over the instance count, so it's okay not to account for other scene instances here.
+		if (NumCurrentInstances + Instances.Num() > MAX_INSTANCE_ID)
+		{
+			// Drop instances to stay at the max.
+			// Account for less than 0 if, for example, the forward was over the limit and the reverse had less than the PreExisting.
+			Instances.SetNum(FMath::Max(0, static_cast<int32>(MAX_INSTANCE_ID) - NumCurrentInstances));
+			if (Instances.IsEmpty())
+			{
+				continue;
+			}
+		}
+
+		NumCurrentInstances += Instances.Num();
 		ISMC->AddInstances(Instances, /*bShouldReturnIndices=*/false, /*bWorldSpace=*/true);
 
 		// Scan all points looking for points that match current direction and add their custom data.
