[PCG] Make sure Data visualization ISMs are transient

#rnx

diff --git a/Engine/Plugins/PCG/Source/PCG/Private/Helpers/PCGActorHelpers.cpp b/Engine/Plugins/PCG/Source/PCG/Private/Helpers/PCGActorHelpers.cpp
--- a/Engine/Plugins/PCG/Source/PCG/Private/Helpers/PCGActorHelpers.cpp
+++ b/Engine/Plugins/PCG/Source/PCG/Private/Helpers/PCGActorHelpers.cpp
@@ -154,7 +154,8 @@ UPCGManagedISMComponent* UPCGActorHelpers::GetOrCreateManagedISMC(AActor* InTarg
 						if (IsValid(ISMC) &&
 							ISMC->GetOwner() == InTargetActor &&
 							ISMC->NumCustomDataFloats == InParams.NumCustomDataFloats &&
-							Resource->GetDescriptor() == Descriptor)
+							Resource->GetDescriptor() == Descriptor &&
+							ISMC->HasAnyFlags(RF_Transient) == InParams.bTransient)
 						{
 							MatchingResource = Resource;
 						}
@@ -195,7 +196,14 @@ UPCGManagedISMComponent* UPCGActorHelpers::GetOrCreateManagedISMC(AActor* InTarg
 
 	ComponentName += StaticMesh->GetName();
 
-	const EObjectFlags ObjectFlags = (InSourceComponent->IsInPreviewMode() ? RF_Transient | RF_NonPIEDuplicateTransient : RF_NoFlags);
+	EObjectFlags ObjectFlags = (InSourceComponent->IsInPreviewMode() ? RF_Transient | RF_NonPIEDuplicateTransient : RF_NoFlags);
+
+	// Used for debug data visualization
+	if (InParams.bTransient)
+	{
+		ObjectFlags |= RF_Transient;
+	}
+
 	UInstancedStaticMeshComponent* ISMC = NewObject<UInstancedStaticMeshComponent>(InTargetActor, Descriptor.ComponentClass, MakeUniqueObjectName(InTargetActor, Descriptor.ComponentClass, FName(ComponentName)), ObjectFlags);
 	Descriptor.InitComponent(ISMC);
 	ISMC->SetNumCustomDataFloats(InParams.NumCustomDataFloats);
@@ -208,7 +216,7 @@ UPCGManagedISMComponent* UPCGActorHelpers::GetOrCreateManagedISMC(AActor* InTarg
 	AddTagsToComponent(ISMC);
 
 	// Create managed resource on source component
-	UPCGManagedISMComponent* Resource = NewObject<UPCGManagedISMComponent>(InSourceComponent);
+	UPCGManagedISMComponent* Resource = NewObject<UPCGManagedISMComponent>(InSourceComponent, NAME_None, InParams.bTransient ? RF_Transient : RF_NoFlags);
 	Resource->SetComponent(ISMC);
 	Resource->SetDescriptor(Descriptor);
 	if (InTargetActor->GetRootComponent())
diff --git a/Engine/Plugins/PCG/Source/PCG/Private/PCGComponent.cpp b/Engine/Plugins/PCG/Source/PCG/Private/PCGComponent.cpp
--- a/Engine/Plugins/PCG/Source/PCG/Private/PCGComponent.cpp
+++ b/Engine/Plugins/PCG/Source/PCG/Private/PCGComponent.cpp
@@ -1668,6 +1668,9 @@ void UPCGComponent::PostLoad()
 	Super::PostLoad();
 
 #if WITH_EDITOR
+	// Debug resources are transient and will generate null entries on save, clean them up here
+	GeneratedResources.Remove(nullptr);
+
 	// Force dirty to be false on load. We should never refresh on load.
 	bDirtyGenerated = false;
 
diff --git a/Engine/Plugins/PCG/Source/PCG/Public/Helpers/PCGActorHelpers.h b/Engine/Plugins/PCG/Source/PCG/Public/Helpers/PCGActorHelpers.h
--- a/Engine/Plugins/PCG/Source/PCG/Public/Helpers/PCGActorHelpers.h
+++ b/Engine/Plugins/PCG/Source/PCG/Public/Helpers/PCGActorHelpers.h
@@ -50,6 +50,7 @@ struct FPCGISMComponentBuilderParams
 	int32 NumCustomDataFloats = 0;
 	FPCGCrc SettingsCrc;
 	bool bAllowDescriptorChanges = true;
+	bool bTransient = false;
 
 PRAGMA_DISABLE_DEPRECATION_WARNINGS
 	explicit FPCGISMComponentBuilderParams(const FPCGISMCBuilderParameters& Params)
diff --git a/Engine/Plugins/PCG/Source/PCGEditor/Private/DataVisualizations/PCGSpatialDataVisualization.cpp b/Engine/Plugins/PCG/Source/PCGEditor/Private/DataVisualizations/PCGSpatialDataVisualization.cpp
--- a/Engine/Plugins/PCG/Source/PCGEditor/Private/DataVisualizations/PCGSpatialDataVisualization.cpp
+++ b/Engine/Plugins/PCG/Source/PCGEditor/Private/DataVisualizations/PCGSpatialDataVisualization.cpp
@@ -167,6 +167,7 @@ void IPCGSpatialDataVisualization::ExecuteDebugDisplay(FPCGContext* Context, con
 
 	FPCGISMComponentBuilderParams Params[2];
 	Params[0].SettingsCrc = SettingsInterface->GetSettings()->GetSettingsCrc();
+	Params[0].bTransient = true;
 	Params[0].NumCustomDataFloats = NumCustomData;
 	Params[0].Descriptor.StaticMesh = Mesh;
 	Params[0].Descriptor.OverrideMaterials = Materials;
