Fixed some UX issues that impair understandability of graphs and hidden internal settings related to graph compilation.

#rnx

diff --git a/Engine/Plugins/PCG/Source/PCG/Private/Compute/Elements/PCGComputeGraphElement.cpp b/Engine/Plugins/PCG/Source/PCG/Private/Compute/Elements/PCGComputeGraphElement.cpp
--- a/Engine/Plugins/PCG/Source/PCG/Private/Compute/Elements/PCGComputeGraphElement.cpp
+++ b/Engine/Plugins/PCG/Source/PCG/Private/Compute/Elements/PCGComputeGraphElement.cpp
@@ -646,6 +646,13 @@ void FPCGComputeGraphElement::LogCompilationMessages(FPCGComputeGraphContext* In
 }
 #endif
 
+UPCGComputeGraphSettings::UPCGComputeGraphSettings()
+{
+#if WITH_EDITOR
+	bExposeToLibrary = false;
+#endif
+}
+
 FPCGElementPtr UPCGComputeGraphSettings::CreateElement() const
 {
 	return MakeShared<FPCGComputeGraphElement>(ComputeGraphIndex);
diff --git a/Engine/Plugins/PCG/Source/PCG/Private/Compute/Elements/PCGComputeGraphElement.h b/Engine/Plugins/PCG/Source/PCG/Private/Compute/Elements/PCGComputeGraphElement.h
--- a/Engine/Plugins/PCG/Source/PCG/Private/Compute/Elements/PCGComputeGraphElement.h
+++ b/Engine/Plugins/PCG/Source/PCG/Private/Compute/Elements/PCGComputeGraphElement.h
@@ -93,6 +93,9 @@ class UPCGComputeGraphSettings : public UPCGSettings
 {
 	GENERATED_BODY()
 
+public:
+	UPCGComputeGraphSettings();
+
 protected:
 	virtual FPCGElementPtr CreateElement() const override;
 
diff --git a/Engine/Plugins/PCG/Source/PCG/Private/Elements/ControlFlow/PCGQualitySelect.cpp b/Engine/Plugins/PCG/Source/PCG/Private/Elements/ControlFlow/PCGQualitySelect.cpp
--- a/Engine/Plugins/PCG/Source/PCG/Private/Elements/ControlFlow/PCGQualitySelect.cpp
+++ b/Engine/Plugins/PCG/Source/PCG/Private/Elements/ControlFlow/PCGQualitySelect.cpp
@@ -33,7 +33,13 @@ EPCGDataType UPCGQualitySelectSettings::GetCurrentPinTypes(const UPCGPin* InPin)
 		return Super::GetCurrentPinTypes(InPin);
 	}
 
-	const EPCGDataType InputType = GetTypeUnionOfIncidentEdges(PCGQualityHelpers::GetQualityPinLabel());
+	// In the case of this node, really we want the union of all inputs on all pins.
+	EPCGDataType InputType = EPCGDataType::None;
+	for (const FPCGPinProperties& PinProperties : InputPinProperties())
+	{
+		InputType |= GetTypeUnionOfIncidentEdges(PinProperties.Label);
+	}
+
 	return (InputType != EPCGDataType::None) ? InputType : EPCGDataType::Any;
 }
 
diff --git a/Engine/Plugins/PCG/Source/PCG/Private/Elements/PCGFilterElementsByIndex.cpp b/Engine/Plugins/PCG/Source/PCG/Private/Elements/PCGFilterElementsByIndex.cpp
--- a/Engine/Plugins/PCG/Source/PCG/Private/Elements/PCGFilterElementsByIndex.cpp
+++ b/Engine/Plugins/PCG/Source/PCG/Private/Elements/PCGFilterElementsByIndex.cpp
@@ -101,9 +101,16 @@ FString UPCGFilterElementsByIndexSettings::GetAdditionalTitleInformation() const
 	}
 
 	FString NodeName = TEXT("Indices: ");
-	NodeName += SelectedIndices.Len() <= PCGFilterElementsByIndex::Constants::IndexExpressionTruncation
+	if (SelectedIndices == FString(TEXT(":")))
+	{
+		NodeName += TEXT("All");
+	}
+	else
+	{
+		NodeName += SelectedIndices.Len() <= PCGFilterElementsByIndex::Constants::IndexExpressionTruncation
 					? SelectedIndices
 					: SelectedIndices.Left(PCGFilterElementsByIndex::Constants::IndexExpressionTruncation - 3) + TEXT("...");
+	}
 
 	return NodeName;
 }
diff --git a/Engine/Plugins/PCG/Source/PCG/Private/Graph/PCGGraphExecutor.cpp b/Engine/Plugins/PCG/Source/PCG/Private/Graph/PCGGraphExecutor.cpp
--- a/Engine/Plugins/PCG/Source/PCG/Private/Graph/PCGGraphExecutor.cpp
+++ b/Engine/Plugins/PCG/Source/PCG/Private/Graph/PCGGraphExecutor.cpp
@@ -3082,6 +3082,13 @@ namespace PCGGraphExecutor
 	}
 }
 
+UPCGGridLinkageSettings::UPCGGridLinkageSettings()
+{
+#if WITH_EDITOR
+	bExposeToLibrary = false;
+#endif // WITH_EDITOR
+}
+
 FPCGElementPtr UPCGGridLinkageSettings::CreateElement() const
 {
 	auto GridLinkageOperation = [this](FPCGContext* InContext)
diff --git a/Engine/Plugins/PCG/Source/PCG/Private/Graph/PCGGraphExecutor.h b/Engine/Plugins/PCG/Source/PCG/Private/Graph/PCGGraphExecutor.h
--- a/Engine/Plugins/PCG/Source/PCG/Private/Graph/PCGGraphExecutor.h
+++ b/Engine/Plugins/PCG/Source/PCG/Private/Graph/PCGGraphExecutor.h
@@ -489,6 +489,9 @@ class UPCGGridLinkageSettings : public UPCGSettings
 {
 	GENERATED_BODY()
 
+public:
+	UPCGGridLinkageSettings();
+
 protected:
 	virtual FPCGElementPtr CreateElement() const override;
 
diff --git a/Engine/Plugins/PCG/Source/PCG/Public/Elements/PCGAddComponent.h b/Engine/Plugins/PCG/Source/PCG/Public/Elements/PCGAddComponent.h
--- a/Engine/Plugins/PCG/Source/PCG/Public/Elements/PCGAddComponent.h
+++ b/Engine/Plugins/PCG/Source/PCG/Public/Elements/PCGAddComponent.h
@@ -49,7 +49,6 @@ public:
 	virtual FText GetNodeTooltipText() const override { return NSLOCTEXT("PCGAddComponentElement", "NodeTooltip", "Adds component(s) to specified target actor(s)."); }
 	virtual EPCGSettingsType GetType() const override { return EPCGSettingsType::Generic; }
 #endif
-	virtual bool HasDynamicPins() const override { return true; }
 
 protected:
 	virtual TArray<FPCGPinProperties> InputPinProperties() const override;
diff --git a/Engine/Plugins/PCG/Source/PCG/Public/Elements/PCGFilterElementsByIndex.h b/Engine/Plugins/PCG/Source/PCG/Public/Elements/PCGFilterElementsByIndex.h
--- a/Engine/Plugins/PCG/Source/PCG/Public/Elements/PCGFilterElementsByIndex.h
+++ b/Engine/Plugins/PCG/Source/PCG/Public/Elements/PCGFilterElementsByIndex.h
@@ -21,6 +21,7 @@ public:
 	virtual EPCGSettingsType GetType() const override { return EPCGSettingsType::Filter; }
 #endif // WITH_EDITOR
 	virtual FString GetAdditionalTitleInformation() const override;
+	virtual bool HasDynamicPins() const override { return true; }
 
 protected:
 	virtual TArray<FPCGPinProperties> InputPinProperties() const override;
diff --git a/Engine/Plugins/PCG/Source/PCG/Public/Elements/PCGPathfindingElement.h b/Engine/Plugins/PCG/Source/PCG/Public/Elements/PCGPathfindingElement.h
--- a/Engine/Plugins/PCG/Source/PCG/Public/Elements/PCGPathfindingElement.h
+++ b/Engine/Plugins/PCG/Source/PCG/Public/Elements/PCGPathfindingElement.h
@@ -81,7 +81,7 @@ public:
 	FVector Goal = FVector::ZeroVector;
 
 	/** The heuristic estimates a faster path to speed up processing. A higher than 1 heuristic weight can be faster, but it may cease being the optimal path. A weight of 0 is essentially flood fill. */
-	UPROPERTY(EditAnywhere, BlueprintReadWrite, Category = Settings, meta = (ClampMin = "0.0", EditCondition = "Algorithm == EPCGPathfindingAlgorithm::AStar", PCG_Overridable))
+	UPROPERTY(EditAnywhere, BlueprintReadWrite, Category = Settings, meta = (ClampMin = "0.0", PCG_Overridable))
 	double HeuristicWeight = 1.0;
 
 	/** Controls whether the cost function will use a given attribute as a scalar wrt to the distance. */
