[PCG] Fixed a crash when running a graph determinism test and then enabling disabled GPU nodes.

The PCG Component was not being cleaned up properly.

#rnx

diff --git a/Engine/Plugins/PCG/Source/PCG/Private/Tests/Determinism/PCGDeterminismTestsCommon.cpp b/Engine/Plugins/PCG/Source/PCG/Private/Tests/Determinism/PCGDeterminismTestsCommon.cpp
--- a/Engine/Plugins/PCG/Source/PCG/Private/Tests/Determinism/PCGDeterminismTestsCommon.cpp
+++ b/Engine/Plugins/PCG/Source/PCG/Private/Tests/Determinism/PCGDeterminismTestsCommon.cpp
@@ -27,7 +27,7 @@
 
 static TAutoConsoleVariable<int32> CVarDeterminismGraphTestCycleLimit(
 	TEXT("pcg.DeterminismGraphTestCycleLimit"),
-	10,
+	10000,
 	TEXT("How many execution attempts before discarding a graph test."));
 
 static TAutoConsoleVariable<int32> CVarDeterminismPermutationLimit(
@@ -120,12 +120,10 @@ namespace PCGDeterminismTests
 
 			// Clean up anything generated
 			PCGComponentCopy->Cleanup();
-			if (PCGComponentCopy->IsRegistered())
-			{
-				PCGComponentCopy->UnregisterComponent();
-			}
 
 			// Clean up the debug actor
+			PCGActorCopy->UnregisterAllComponents();
+			PCGActorCopy->MarkComponentsAsGarbage();
 			PCGActorCopy->RemoveFromRoot();
 			PCGActorCopy->MarkAsGarbage();
 		};
@@ -1648,4 +1646,4 @@ namespace PCGDeterminismTests
 	}
 }
 
-#undef LOCTEXT_NAMESPACE
\ No newline at end of file
+#undef LOCTEXT_NAMESPACE
