[PCG] Fix and improve on the node title for Spline Direction.

#rnx

diff --git a/Engine/Plugins/PCG/Source/PCG/Private/Elements/PCGReverseSpline.cpp b/Engine/Plugins/PCG/Source/PCG/Private/Elements/PCGReverseSpline.cpp
deleted file mode 100644
diff --git a/Engine/Plugins/PCG/Source/PCG/Private/Elements/PCGSplineDirection.cpp b/Engine/Plugins/PCG/Source/PCG/Private/Elements/PCGSplineDirection.cpp
new file mode 100644
--- /dev/null
+++ b/Engine/Plugins/PCG/Source/PCG/Private/Elements/PCGSplineDirection.cpp
@@ -0,0 +1,172 @@
+// Copyright Epic Games, Inc. All Rights Reserved.
+
+#include "Elements/PCGSplineDirection.h"
+
+#include "PCGContext.h"
+#include "Data/PCGSplineData.h"
+
+#define LOCTEXT_NAMESPACE "PCGSplineDirectionElement"
+
+namespace PCGSplineDirection
+{
+	bool IsClockwiseXY(const UPCGSplineData* InputSplineData)
+	{
+		check(InputSplineData);
+
+		double CumulativeAngle = 0.0;
+		int NumPoints = InputSplineData->SplineStruct.SplineCurves.Position.Points.Num();
+		if (!InputSplineData->IsClosed())
+		{
+			NumPoints--;
+		}
+
+		for (int32 Index = 0; Index < NumPoints; ++Index)
+		{
+			FVector ArriveTangent{}, LeaveTangent{};
+
+			InputSplineData->GetTangentsAtSegmentStart(Index, ArriveTangent, LeaveTangent);
+			const double CrossProduct = (ArriveTangent.X * LeaveTangent.Y - ArriveTangent.Y * LeaveTangent.X);
+			const double AbsAngle = FMath::Acos(ArriveTangent.CosineAngle2D(LeaveTangent));
+			const double Angle = FMath::Sign(CrossProduct) * AbsAngle;
+
+			CumulativeAngle += Angle;
+		}
+
+		return CumulativeAngle <= 0;
+	}
+
+	UPCGSplineData* Reverse(const UPCGSplineData* InputSplineData, FPCGContext* Context)
+	{
+		check(InputSplineData);
+
+		const FInterpCurveVector& ControlPointsPosition = InputSplineData->SplineStruct.SplineCurves.Position;
+		const FInterpCurveQuat& ControlPointsRotation = InputSplineData->SplineStruct.SplineCurves.Rotation;
+		const FInterpCurveVector& ControlPointsScale = InputSplineData->SplineStruct.SplineCurves.Scale;
+
+		TArray<FSplinePoint> NewControlPoints;
+		NewControlPoints.Reserve(ControlPointsPosition.Points.Num());
+
+		for (int i = ControlPointsPosition.Points.Num() - 1; i >= 0; --i)
+		{
+			/* Implementation Note: Segment interpolation is determined by the interpolation mode of the preceding
+			 * control point. When inverting order of control points, we can decay the interpolation mode by setting all
+			 * modes to Custom Tangent to use the pre-calculated tangents as-is, so long as they were actually
+			 * calculated. This is a slightly destructive process and some information will be lost. Also, its worth
+			 * noting that since each spline segment is calculated from [0..1) there is a slight inconsistency when
+			 * evaluated in reverse order, as effectively the reverse is [1..0) per segment.
+			 */
+			NewControlPoints.Emplace(static_cast<float>(NewControlPoints.Num()),
+				ControlPointsPosition.Points[i].OutVal,
+				-ControlPointsPosition.Points[i].LeaveTangent, // Tangents are inverted and swapped
+				-ControlPointsPosition.Points[i].ArriveTangent,
+				ControlPointsRotation.Points[i].OutVal.Rotator(),
+				ControlPointsScale.Points[i].OutVal,
+				ESplinePointType::CurveCustomTangent);
+		}
+
+		UPCGSplineData* NewSplineData = FPCGContext::NewObject_AnyThread<UPCGSplineData>(Context);
+		NewSplineData->InitializeFromData(InputSplineData);
+		NewSplineData->Initialize(NewControlPoints, InputSplineData->IsClosed(), InputSplineData->GetTransform());
+
+		return NewSplineData;
+	}
+}
+
+#if WITH_EDITOR
+FName UPCGReverseSplineSettings::GetDefaultNodeName() const
+{
+	return FName(TEXT("SplineDirection"));
+}
+
+FText UPCGReverseSplineSettings::GetDefaultNodeTitle() const
+{
+	return LOCTEXT("NodeTitle", "Spline Direction");
+}
+
+EPCGChangeType UPCGReverseSplineSettings::GetChangeTypeForProperty(const FName& InPropertyName) const
+{
+	EPCGChangeType ChangeType = Super::GetChangeTypeForProperty(InPropertyName);
+	if (InPropertyName == GET_MEMBER_NAME_CHECKED(UPCGReverseSplineSettings, Operation))
+	{
+		ChangeType |= EPCGChangeType::Cosmetic;
+	}
+
+	return ChangeType;
+}
+#endif // WITH_EDITOR
+
+FPCGElementPtr UPCGReverseSplineSettings::CreateElement() const
+{
+	return MakeShared<FPCGSplineDirectionElement>();
+}
+
+FString UPCGReverseSplineSettings::GetAdditionalTitleInformation() const
+{
+	if (const UEnum* EnumPtr = StaticEnum<EPCGReverseSplineOperation>())
+	{
+		return EnumPtr->GetDisplayNameTextByValue(static_cast<int64>(Operation)).ToString();
+	}
+
+	return {};
+}
+
+TArray<FPCGPinProperties> UPCGReverseSplineSettings::InputPinProperties() const
+{
+	TArray<FPCGPinProperties> Properties;
+	Properties.Emplace_GetRef(PCGPinConstants::DefaultInputLabel, EPCGDataType::Spline).SetRequiredPin();
+	return Properties;
+}
+
+TArray<FPCGPinProperties> UPCGReverseSplineSettings::OutputPinProperties() const
+{
+	TArray<FPCGPinProperties> Properties;
+	Properties.Emplace(PCGPinConstants::DefaultOutputLabel, EPCGDataType::Spline);
+	return Properties;
+}
+
+bool FPCGSplineDirectionElement::ExecuteInternal(FPCGContext* InContext) const
+{
+	TRACE_CPUPROFILER_EVENT_SCOPE(FPCGSplineDirectionElement::Execute);
+
+	check(InContext);
+
+	const UPCGReverseSplineSettings* Settings = InContext->GetInputSettings<UPCGReverseSplineSettings>();
+	check(Settings);
+
+	for (const FPCGTaggedData& InputData : InContext->InputData.GetInputsByPin(PCGPinConstants::DefaultInputLabel))
+	{
+		FPCGTaggedData& Output = InContext->OutputData.TaggedData.Emplace_GetRef(InputData);
+
+		const UPCGSplineData* InputSplineData = Cast<const UPCGSplineData>(InputData.Data);
+		if (!InputSplineData || InputSplineData->SplineStruct.GetNumberOfSplineSegments() < 1)
+		{
+			continue;
+		}
+
+		bool bShouldReverse = true;
+		switch (Settings->Operation)
+		{
+		case EPCGReverseSplineOperation::ForceClockwise:
+			bShouldReverse = !PCGSplineDirection::IsClockwiseXY(InputSplineData);
+			break;
+		case EPCGReverseSplineOperation::ForceCounterClockwise:
+			bShouldReverse = PCGSplineDirection::IsClockwiseXY(InputSplineData);
+			break;
+		case EPCGReverseSplineOperation::Reverse:
+			bShouldReverse = true;
+			break;
+		default:
+			PCGLog::LogErrorOnGraph(LOCTEXT("InvalidOperation", "Invalid operation enum value"), InContext);
+			return true;
+		}
+
+		if (bShouldReverse)
+		{
+			Output.Data = PCGSplineDirection::Reverse(InputSplineData, InContext);
+		}
+	}
+
+	return true;
+}
+
+#undef LOCTEXT_NAMESPACE
diff --git a/Engine/Plugins/PCG/Source/PCG/Public/Elements/PCGReverseSpline.h b/Engine/Plugins/PCG/Source/PCG/Public/Elements/PCGReverseSpline.h
deleted file mode 100644
diff --git a/Engine/Plugins/PCG/Source/PCG/Public/Elements/PCGSplineDirection.h b/Engine/Plugins/PCG/Source/PCG/Public/Elements/PCGSplineDirection.h
new file mode 100644
--- /dev/null
+++ b/Engine/Plugins/PCG/Source/PCG/Public/Elements/PCGSplineDirection.h
@@ -0,0 +1,69 @@
+// Copyright Epic Games, Inc. All Rights Reserved.
+
+#pragma once
+
+#include "PCGSettings.h"
+
+#include "PCGSplineDirection.generated.h"
+
+class UPCGSplineData;
+
+namespace PCGSplineDirection
+{
+	/** Test spline controls points to know if they are clockwise on the XY plane (clockwise around the Z axis). */
+	PCG_API bool IsClockwiseXY(const UPCGSplineData* InputSplineData);
+
+	/** Return the same spline data, but with its control points in reverse order. */
+	PCG_API UPCGSplineData* Reverse(const UPCGSplineData* InputSplineData, FPCGContext* Context);
+}
+
+UENUM()
+enum class EPCGReverseSplineOperation
+{
+	/** Will reverse the spline points regardless of their orientation. */
+	Reverse,
+
+	/** Will reverse the spline points if they are counter clockwise. */
+	ForceClockwise,
+
+	/** Will reverse the spline points if they are clockwise. */
+	ForceCounterClockwise
+};
+
+/**
+* Direct the order of a spline's control points.
+* This can be conditional to force a given orientation (clockwise or counter clockwise).
+*/
+UCLASS(MinimalAPI, BlueprintType, ClassGroup = (Procedural))
+class UPCGReverseSplineSettings : public UPCGSettings
+{
+	GENERATED_BODY()
+
+public:
+	//~Begin UPCGSettings interface
+#if WITH_EDITOR
+	virtual FName GetDefaultNodeName() const override;
+	virtual FText GetDefaultNodeTitle() const override;
+	virtual EPCGSettingsType GetType() const override { return EPCGSettingsType::Spatial; }
+	virtual EPCGChangeType GetChangeTypeForProperty(const FName& InPropertyName) const override;
+#endif
+
+	virtual FString GetAdditionalTitleInformation() const override;
+
+protected:
+	virtual TArray<FPCGPinProperties> InputPinProperties() const override;
+	virtual TArray<FPCGPinProperties> OutputPinProperties() const override;
+	virtual FPCGElementPtr CreateElement() const override;
+	//~End UPCGSettings interface
+
+public:
+	UPROPERTY(EditAnywhere, BlueprintReadWrite, Category = "Settings")
+	EPCGReverseSplineOperation Operation = EPCGReverseSplineOperation::Reverse;
+};
+
+class FPCGSplineDirectionElement : public IPCGElement
+{
+protected:
+	virtual bool ExecuteInternal(FPCGContext* InContext) const override;
+	virtual EPCGElementExecutionLoopMode ExecutionLoopMode(const UPCGSettings* Settings) const override { return EPCGElementExecutionLoopMode::SinglePrimaryPin; }
+};
