[PCG] Fixed a bug in Create Points Sphere causing generated points to improperly self orient

#rnx

diff --git a/Engine/Plugins/PCG/Source/PCG/Private/Elements/PCGCreatePointsSphere.cpp b/Engine/Plugins/PCG/Source/PCG/Private/Elements/PCGCreatePointsSphere.cpp
--- a/Engine/Plugins/PCG/Source/PCG/Private/Elements/PCGCreatePointsSphere.cpp
+++ b/Engine/Plugins/PCG/Source/PCG/Private/Elements/PCGCreatePointsSphere.cpp
@@ -459,11 +459,15 @@ bool FPCGCreatePointsSphereElement::ExecuteInternal(FPCGContext* InContext) cons
 		switch (Settings->PointOrientation)
 		{
 			case EPCGSpherePointOrientation::Radial:
-				LocalPointOrientation = FRotationMatrix::MakeFromZ(OutLocalPointLocations[Index]).Rotator();
-				break;
+			{
+				LocalPointOrientation = FRotationMatrix::MakeFromZ(OutLocalPointLocations[Index] - Settings->Origin).Rotator();
+			}
+			break;
 			case EPCGSpherePointOrientation::Centric:
-				LocalPointOrientation = FRotationMatrix::MakeFromZ(-OutLocalPointLocations[Index]).Rotator();
-				break;
+			{
+				LocalPointOrientation = FRotationMatrix::MakeFromZ(Settings->Origin - OutLocalPointLocations[Index]).Rotator();
+			}
+			break;
 			case EPCGSpherePointOrientation::None:
 				break;
 			default:
