[PCG] Do not early out of Cleanup for non-generated/generating components if they have managed resources (regression from 27747133)

- This can happen if a Graph is cancelled after having generated resources but before completing

diff --git a/Engine/Plugins/PCG/Source/PCG/Private/PCGComponent.cpp b/Engine/Plugins/PCG/Source/PCG/Private/PCGComponent.cpp
--- a/Engine/Plugins/PCG/Source/PCG/Private/PCGComponent.cpp
+++ b/Engine/Plugins/PCG/Source/PCG/Private/PCGComponent.cpp
@@ -790,7 +790,12 @@ void UPCGComponent::OnProcessGraphAborted(bool bQuiet, bool bCleanupUnusedResour
 
 void UPCGComponent::Cleanup()
 {
-	if ((!bGenerated && !IsGenerating()) || !GetSubsystem() || IsCleaningUp())
+	if (!GetSubsystem() || IsCleaningUp())
+	{
+		return;
+	}
+
+	if (!bGenerated && !IsGenerating() && GeneratedResources.IsEmpty())
 	{
 		return;
 	}
@@ -931,7 +936,12 @@ void UPCGComponent::CleanupLocal(bool bRemoveComponents, bool bSave, const TArra
 
 FPCGTaskId UPCGComponent::CleanupInternal(bool bRemoveComponents, const TArray<FPCGTaskId>& Dependencies)
 {
-	if ((!bGenerated && !IsGenerating()) || !GetSubsystem() || IsCleaningUp())
+	if (!GetSubsystem() || IsCleaningUp())
+	{
+		return InvalidPCGTaskId;
+	}
+
+	if (!bGenerated && !IsGenerating() && GeneratedResources.IsEmpty())
 	{
 		return InvalidPCGTaskId;
 	}
