[PCG] Remove ensure() for invalid edge types and add clear error message for user.

#rnx

diff --git a/Engine/Plugins/PCG/Source/PCG/Private/Graph/PCGGraphCompilerGPU.cpp b/Engine/Plugins/PCG/Source/PCG/Private/Graph/PCGGraphCompilerGPU.cpp
--- a/Engine/Plugins/PCG/Source/PCG/Private/Graph/PCGGraphCompilerGPU.cpp
+++ b/Engine/Plugins/PCG/Source/PCG/Private/Graph/PCGGraphCompilerGPU.cpp
@@ -864,7 +864,7 @@ UPCGComputeGraph* FPCGGraphCompilerGPU::CompileComputeGraph(
 
 			if (!UpstreamDI)
 			{
-				ensure(false);
+				UE_LOG(LogPCG, Warning, TEXT("Encountered missing data interface, skipping."))
 				continue;
 			}
 
diff --git a/Engine/Plugins/PCG/Source/PCG/Private/PCGSettings.cpp b/Engine/Plugins/PCG/Source/PCG/Private/PCGSettings.cpp
--- a/Engine/Plugins/PCG/Source/PCG/Private/PCGSettings.cpp
+++ b/Engine/Plugins/PCG/Source/PCG/Private/PCGSettings.cpp
@@ -237,7 +237,36 @@ bool UPCGSettings::IsKernelValid(FPCGContext* InContext, bool bQuiet) const
 		}
 	}
 
-	return true;
+	// Validate types of incident edges to make sure we catch invalid cases like Spatial -> Point.
+	bool bAllEdgesValid = true;
+
+	if (const UPCGNode* Node = Cast<UPCGNode>(GetOuter()))
+	{
+		for (const UPCGPin* InputPin : Node->GetInputPins())
+		{
+			if (!InputPin)
+			{
+				continue;
+			}
+
+			for (const UPCGEdge* InputEdge : InputPin->Edges)
+			{
+				const UPCGPin* UpstreamPin = InputEdge ? InputEdge->GetOtherPin(InputPin) : nullptr;
+				if (UpstreamPin && InputPin->GetRequiredTypeConversion(UpstreamPin) != EPCGTypeConversion::NoConversionRequired)
+				{
+					PCG_KERNEL_VALIDATION_ERR(InContext, this, bQuiet, FText::Format(
+						LOCTEXT("InvalidInputPinEdge", "Unsupported connected upstream pin '{0}' on node '{1}' with type {2}. Recreate the edge to add required conversion nodes."),
+						FText::FromName(UpstreamPin->Properties.Label),
+						Node->GetNodeTitle(EPCGNodeTitleType::ListView),
+						StaticEnum<EPCGDataType>() ? StaticEnum<EPCGDataType>()->GetDisplayNameTextByValue(static_cast<int64>(UpstreamPin->Properties.AllowedTypes)) : FText::FromString(TEXT("MISSING"))));
+
+					bAllEdgesValid = false;
+				}
+			}
+		}
+	}
+
+	return bAllEdgesValid;
 }
 
 bool UPCGSettings::ComputeOutputPinDataDesc(const FName& OutputPinLabel, const UPCGDataBinding* InBinding, FPCGDataCollectionDesc& OutDesc) const
