[PCG] GPU - add last high priority shader API features for Custom HLSL nodes.

* Out_CopyElementFrom_In() function to automate the job of copying all data from input to output, which is arduous for user to do manually.
* Expose Thread Count Multiplier to shader which is otherwise inaccessible.
* Expose GetThreadData() for all kernel types because its needed for secondary pins (not just for Custom kernels).

#rnx

diff --git a/Engine/Plugins/PCG/Source/PCG/Private/Compute/DataInterfaces/PCGCustomKernelDataInterface.cpp b/Engine/Plugins/PCG/Source/PCG/Private/Compute/DataInterfaces/PCGCustomKernelDataInterface.cpp
--- a/Engine/Plugins/PCG/Source/PCG/Private/Compute/DataInterfaces/PCGCustomKernelDataInterface.cpp
+++ b/Engine/Plugins/PCG/Source/PCG/Private/Compute/DataInterfaces/PCGCustomKernelDataInterface.cpp
@@ -5,6 +5,7 @@
 #include "PCGComponent.h"
 #include "PCGModule.h"
 #include "PCGSettings.h"
+#include "Compute/Elements/PCGCustomHLSL.h"
 
 #include "RenderGraphBuilder.h"
 #include "RenderGraphResources.h"
@@ -21,6 +22,10 @@ void UPCGCustomKernelDataInterface::GetSupportedInputs(TArray<FShaderFunctionDef
 		.SetName(TEXT("GetNumThreads"))
 		.AddReturnType(FShaderValueType::Get(EShaderFundamentalType::Int, 3));
 
+	OutFunctions.AddDefaulted_GetRef()
+		.SetName(TEXT("GetThreadCountMultiplier"))
+		.AddReturnType(FShaderValueType::Get(EShaderFundamentalType::Uint));
+
 	OutFunctions.AddDefaulted_GetRef()
 		.SetName(TEXT("GetSeed"))
 		.AddReturnType(FShaderValueType::Get(EShaderFundamentalType::Uint));
@@ -45,6 +50,7 @@ void UPCGCustomKernelDataInterface::GetSupportedInputs(TArray<FShaderFunctionDef
 
 BEGIN_SHADER_PARAMETER_STRUCT(FPCGCustomKernelDataInterfaceParameters, )
 	SHADER_PARAMETER(FIntVector3, NumThreads)
+	SHADER_PARAMETER(uint32, ThreadCountMultiplier)
 	SHADER_PARAMETER(uint32, Seed)
 	SHADER_PARAMETER(uint32, SeedSettings)
 	SHADER_PARAMETER(uint32, SeedComponent)
@@ -66,6 +72,7 @@ void UPCGCustomKernelDataInterface::GetHLSL(FString& OutHLSL, FString const& InD
 
 	OutHLSL += FString::Format(TEXT(
 		"int3 {DataInterfaceName}_NumThreads;\n"
+		"uint {DataInterfaceName}_ThreadCountMultiplier;\n"
 		"uint {DataInterfaceName}_Seed;\n"
 		"uint {DataInterfaceName}_SeedSettings;\n"
 		"uint {DataInterfaceName}_SeedComponent;\n"
@@ -73,6 +80,7 @@ void UPCGCustomKernelDataInterface::GetHLSL(FString& OutHLSL, FString const& InD
 		"float3 {DataInterfaceName}_ComponentBoundsMax;\n"
 		"\n"
 		"int3 GetNumThreads_{DataInterfaceName}()\n{\n\treturn {DataInterfaceName}_NumThreads;\n}\n\n"
+		"uint GetThreadCountMultiplier_{DataInterfaceName}()\n{\n\treturn {DataInterfaceName}_ThreadCountMultiplier;\n}\n\n"
 		"uint GetSeed_{DataInterfaceName}()\n{\n\treturn {DataInterfaceName}_Seed;\n}\n\n"
 		"uint GetSettingsSeed_{DataInterfaceName}()\n{\n\treturn {DataInterfaceName}_SeedSettings;\n}\n\n"
 		"uint GetComponentSeed_{DataInterfaceName}()\n{\n\treturn {DataInterfaceName}_SeedComponent;\n}\n\n"
@@ -91,7 +99,16 @@ UComputeDataProvider* UPCGCustomKernelDataInterface::CreateDataProvider(TObjectP
 
 	UPCGCustomComputeKernelDataProvider* Provider = NewObject<UPCGCustomComputeKernelDataProvider>();
 	Provider->ThreadCount = Settings->ComputeKernelThreadCount(Binding);
-	
+
+	if (const UPCGCustomHLSLSettings* CustomHLSLSettings = Cast<UPCGCustomHLSLSettings>(Settings))
+	{
+		Provider->ThreadCountMultiplier = CustomHLSLSettings->GetThreadCountMultiplier();
+	}
+	else
+	{
+		Provider->ThreadCountMultiplier = 1;
+	}
+
 	Provider->Seed = static_cast<uint32>(Settings->GetSeed(Binding->SourceComponent.Get()));
 	Provider->SeedSettings = static_cast<uint32>(Settings->Seed);
 	Provider->SeedComponent = static_cast<uint32>(Binding->SourceComponent->Seed);
@@ -103,7 +120,7 @@ UComputeDataProvider* UPCGCustomKernelDataInterface::CreateDataProvider(TObjectP
 
 FComputeDataProviderRenderProxy* UPCGCustomComputeKernelDataProvider::GetRenderProxy()
 {
-	return new FPCGCustomComputeKernelDataProviderProxy(ThreadCount, Seed, SeedSettings, SeedComponent, SourceComponentBounds);
+	return new FPCGCustomComputeKernelDataProviderProxy(ThreadCount, ThreadCountMultiplier, Seed, SeedSettings, SeedComponent, SourceComponentBounds);
 }
 
 bool FPCGCustomComputeKernelDataProviderProxy::IsValid(FValidationData const& InValidationData) const
@@ -136,6 +153,7 @@ void FPCGCustomComputeKernelDataProviderProxy::GatherDispatchData(FDispatchData
 		ensure(InDispatchData.bUnifiedDispatch || InDispatchData.NumInvocations == 1);
 		Parameters.NumThreads.X = ThreadCount;
 		Parameters.NumThreads.Y = Parameters.NumThreads.Z = 1;
+		Parameters.ThreadCountMultiplier = ThreadCountMultiplier;
 
 		// Seed for the node
 		Parameters.Seed = Seed;
diff --git a/Engine/Plugins/PCG/Source/PCG/Private/Compute/DataInterfaces/PCGCustomKernelDataInterface.h b/Engine/Plugins/PCG/Source/PCG/Private/Compute/DataInterfaces/PCGCustomKernelDataInterface.h
--- a/Engine/Plugins/PCG/Source/PCG/Private/Compute/DataInterfaces/PCGCustomKernelDataInterface.h
+++ b/Engine/Plugins/PCG/Source/PCG/Private/Compute/DataInterfaces/PCGCustomKernelDataInterface.h
@@ -48,6 +48,7 @@ public:
 	//~ End UComputeDataProvider Interface
 
 	int32 ThreadCount = -1;
+	uint32 ThreadCountMultiplier = 0;
 
 	uint32 Seed = 42;
 	uint32 SeedSettings = 42;
@@ -59,8 +60,9 @@ public:
 class FPCGCustomComputeKernelDataProviderProxy : public FComputeDataProviderRenderProxy
 {
 public:
-	FPCGCustomComputeKernelDataProviderProxy(int32 InThreadCount, uint32 InSeed, uint32 InSeedSettings, uint32 InSeedComponent, const FBox& InSourceComponentBounds)
+	FPCGCustomComputeKernelDataProviderProxy(int32 InThreadCount, uint32 InThreadCountMultiplier, uint32 InSeed, uint32 InSeedSettings, uint32 InSeedComponent, const FBox& InSourceComponentBounds)
 		: ThreadCount(InThreadCount)
+		, ThreadCountMultiplier(InThreadCountMultiplier)
 		, Seed(InSeed)
 		, SeedSettings(InSeedSettings)
 		, SeedComponent(InSeedComponent)
@@ -76,7 +78,8 @@ public:
 protected:
 	using FParameters = FPCGCustomKernelDataInterfaceParameters;
 
-	int32 ThreadCount;
+	int32 ThreadCount = -1;
+	uint32 ThreadCountMultiplier = 0;
 
 	uint32 Seed = 42;
 	uint32 SeedSettings = 42;
diff --git a/Engine/Plugins/PCG/Source/PCG/Private/Compute/Elements/PCGCustomHLSL.cpp b/Engine/Plugins/PCG/Source/PCG/Private/Compute/Elements/PCGCustomHLSL.cpp
--- a/Engine/Plugins/PCG/Source/PCG/Private/Compute/Elements/PCGCustomHLSL.cpp
+++ b/Engine/Plugins/PCG/Source/PCG/Private/Compute/Elements/PCGCustomHLSL.cpp
@@ -912,6 +912,10 @@ void UPCGCustomHLSLSettings::UpdateOutputDeclarations()
 		}
 
 		OutputDeclarations += FString::Format(TEXT(
+			"void {0}_InitializePoint(uint DataIndex, uint ElementIndex);\n"
+			"void {0}_CopyElementFrom_<input pin>(uint TargetDataIndex, uint TargetElementIndex, uint SourceDataIndex, uint SourceElementIndex);\n"
+			"bool {0}_RemovePoint(uint DataIndex, uint ElementIndex);\n"
+			"\n"
 			"void {0}_SetPosition(uint DataIndex, uint ElementIndex, float3 Position);\n"
 			"void {0}_SetRotation(uint DataIndex, uint ElementIndex, float4 Rotation);\n"
 			"void {0}_SetScale(uint DataIndex, uint ElementIndex, float3 Scale);\n"
@@ -921,9 +925,7 @@ void UPCGCustomHLSLSettings::UpdateOutputDeclarations()
 			"void {0}_SetDensity(uint DataIndex, uint ElementIndex, float Density);\n"
 			"void {0}_SetSeed(uint DataIndex, uint ElementIndex, int Seed);\n"
 			"void {0}_SetSteepness(uint DataIndex, uint ElementIndex, float Steepness);\n"
-			"void {0}_SetPointTransform(uint DataIndex, uint ElementIndex, float4x4 Transform);\n"
-			"bool {0}_RemovePoint(uint DataIndex, uint ElementIndex);\n"
-			"void {0}_InitializePoint(uint DataIndex, uint ElementIndex);\n"),
+			"void {0}_SetPointTransform(uint DataIndex, uint ElementIndex, float4x4 Transform);\n"),
 			{ bMultiPin ? PCGHLSLElement::PinDeclTemplateStr : PointDataPins[0] });
 
 		OutputDeclarations += TEXT("\n");
@@ -954,14 +956,15 @@ void UPCGCustomHLSLSettings::UpdateHelperDeclarations()
 {
 	HelperDeclarations.Reset();
 
-	// Helper funcs category
+	// Helper functions category
 	{
 		HelperDeclarations += TEXT(
 			"/*** HELPER FUNCTIONS ***/\n"
 			"\n"
-			"int3 GetNumThreads();\n");
+			"int3 GetNumThreads();\n"
+			"uint GetThreadCountMultiplier();\n");
 
-		if (KernelType == EPCGKernelType::Custom)
+		// Get thread data - useful in all kernel types for secondary pins.
 		{
 			TArray<FString> PinNames;
 
@@ -1464,6 +1467,15 @@ bool UPCGCustomHLSLSettings::IsKernelValid(FPCGContext* InContext, bool bQuiet)
 			}
 			return false;
 		}
+
+		if (!GetGeneratedFunctions(/*OutFunctions=*/nullptr, ErrorTextPtr))
+		{
+			if (ErrorTextPtr)
+			{
+				PCG_KERNEL_VALIDATION_ERR(InContext, this, bQuiet, *ErrorTextPtr);
+			}
+			return false;
+		}
 	}
 
 	return true;
@@ -1635,13 +1647,138 @@ bool UPCGCustomHLSLSettings::AreKernelAttributesValid(FPCGContext* InContext, FT
 	return true;
 }
 
+bool UPCGCustomHLSLSettings::GetGeneratedFunctions(FString* OutFunctions, FText* OutErrorText) const
+{
+	TSet<TPair<FString, FString>> PinsRequiringElementCopy;
+
+	// Point processor nodes always copy input to output at the beginning of the kernel.
+	if (KernelType == EPCGKernelType::PointProcessor && GetPointProcessingInputPin() && GetFirstPointOutputPin())
+	{
+		PinsRequiringElementCopy.Emplace({ GetPointProcessingInputPin()->Properties.Label.ToString(), GetFirstPointOutputPin()->Properties.Label.ToString() });
+	}
+
+	const FString Source = ShaderFunctions + TEXT("\n") + ShaderSource;
+
+	/** Regex for parsing usages of CopyElementFrom function. */
+	constexpr TCHAR CopyElementPattern[] = { TEXT(/*TargetPin*/"([a-zA-Z0-9]+)" "_CopyElementFrom_" /*SourcePin*/"([a-zA-Z0-9]+)\\s*\\(") };
+	constexpr int CopyElementTargetPinCaptureGroup = 1;
+	constexpr int CopyElementSourcePinCaptureGroup = 2;
+
+	FRegexMatcher ModuleMatcher(FRegexPattern(CopyElementPattern), Source);
+
+	while (ModuleMatcher.FindNext())
+	{
+		const FString SourcePinStr = ModuleMatcher.GetCaptureGroup(CopyElementSourcePinCaptureGroup);
+		const UPCGPin* SourcePin = GetInputPin(FName(*SourcePinStr));
+		if (!SourcePin)
+		{
+#if PCG_KERNEL_LOGGING_ENABLED
+			if (OutErrorText)
+			{
+				*OutErrorText = FText::Format(LOCTEXT("InvalidCopySourcePin", "Invalid input pin '{0}' used in copy element expression."), FText::FromString(SourcePinStr));
+			}
+#endif
+			return false;
+		}
+
+		const FString TargetPinStr = ModuleMatcher.GetCaptureGroup(CopyElementTargetPinCaptureGroup);
+		const UPCGPin* TargetPin = GetOutputPin(FName(*TargetPinStr));
+		if (!TargetPin)
+		{
+#if PCG_KERNEL_LOGGING_ENABLED
+			if (OutErrorText)
+			{
+				*OutErrorText = FText::Format(LOCTEXT("InvalidCopyTargetPin", "Invalid output pin '{0}' used in copy element expression."), FText::FromString(TargetPinStr));
+			}
+#endif
+			return false;
+		}
+
+		if (SourcePin->Properties.AllowedTypes != TargetPin->Properties.AllowedTypes)
+		{
+#if PCG_KERNEL_LOGGING_ENABLED
+			if (OutErrorText)
+			{
+				*OutErrorText = LOCTEXT("InvalidCopyTargetPin", "CopyFromElement only allowed across pins of matching type.");
+			}
+#endif
+			return false;
+		}
+
+		PinsRequiringElementCopy.Emplace({ SourcePinStr, TargetPinStr });
+	}
+
+	if (!OutFunctions)
+	{
+		return true;
+	}
+
+	FString Functions = TEXT("// Generated functions\n");
+
+	for (TPair<FString, FString>& Pair : PinsRequiringElementCopy)
+	{
+		Functions += FString::Format(TEXT(
+			"\n"
+			"// Copy element from pin {0} to pin {1}.\n"
+			"void {1}_CopyElementFrom_{0}(uint {1}_DataIndex, uint {1}_ElementIndex, uint {0}_DataIndex, uint {0}_ElementIndex)\n"
+			"{\n"
+			"    // Copy elements across pins. Loop over all attribute headers, if the address is non-zero, then copy it from pin {0} to pin {1}.\n"
+			"    int NumAttributesRemaining = (int){1}_GetDataNumAttributesInternal({1}_DataIndex);\n"
+			"\n"
+			"    for (int AttributeId = 0; AttributeId < {2}; ++AttributeId)\n"
+			"    {\n"
+			"        const uint {1}_Stride = {1}_GetAttributeStrideInternal({1}_DataIndex, AttributeId);\n"
+			"\n"
+			"        if ({1}_Stride == 0)\n"
+			"        {\n"
+			"            // No output attribute to write to.\n"
+			"            continue;\n"
+			"        }\n"
+			"\n"
+			"        const uint {0}_Stride = {0}_GetAttributeStrideInternal({0}_DataIndex, AttributeId);\n"
+			"        const uint {1}_ElementAddress = {1}_GetElementAddressInternal({1}_DataIndex, {1}_ElementIndex, AttributeId);\n"
+			"\n"
+			"        if ({0}_Stride == {1}_Stride)\n"
+			"        {\n"
+			"            const uint {0}_ElementAddress = {0}_GetElementAddressInternal({0}_DataIndex, {0}_ElementIndex, AttributeId);\n"
+			"\n"
+			"            for (uint Offset = 0; Offset < {1}_Stride; Offset += 4)\n"
+			"            {\n"
+			"                {1}_StoreBufferInternal({1}_ElementAddress + Offset, {0}_LoadBufferInternal({0}_ElementAddress + Offset));\n"
+			"            }\n"
+			"\n"
+			"        }\n"
+			"        else\n"
+			"        {\n"
+			"            for (uint Offset = 0; Offset < {1}_Stride; Offset += 4)\n"
+			"            {\n"
+			"                // Initialize output data (only header part of buffer will be initialized).\n"
+			"                {1}_StoreBufferInternal({1}_ElementAddress + Offset, 0u);\n"
+			"            }\n"
+			"        }\n"
+			"\n"
+			"        if (--NumAttributesRemaining <= 0)\n"
+			"        {\n"
+			"            break; // We can early-out when we've looked at all the possible attributes\n"
+			"        }\n"
+			"    }\n"
+			"}\n"),
+			{ Pair.Key, Pair.Value, PCGComputeConstants::MAX_NUM_ATTRS });
+	}
+
+	*OutFunctions = Functions;
+	return true;
+}
+
 FString UPCGCustomHLSLSettings::GetCookedKernelSource(const TMap<FName, FPCGKernelAttributeIDAndType>& GlobalAttributeLookupTable) const
 {
 	const FIntVector GroupSize(PCGComputeConstants::THREAD_GROUP_SIZE, 1, 1);
 
-	// FIXME: Create source range mappings so that we can go from error location to our source.
 	FString Source = ShaderSource;
-	FString Functions = (TEXT("// User kernel functions\n") + ShaderFunctions);
+
+	FString Functions;
+	GetGeneratedFunctions(&Functions);
+	Functions += (TEXT("\n\n// User kernel functions\n") + ShaderFunctions);
 
 #if PLATFORM_WINDOWS
 	// Remove old-school stuff.
@@ -1729,49 +1866,8 @@ FString UPCGCustomHLSLSettings::GetCookedKernelSource(const TMap<FName, FPCGKern
 			// TODO pass in IDs of attributes that are actually present.
 			KernelSpecificPreamble += FString::Format(TEXT(
 				"\n"
-				"    // Initialize output data elements. Loop over all attribute headers, if the address is non-zero, then copy it from pin {0} to pin {1}.\n"
-				"    {\n"
-				"        const uint NumAttributes = {1}_GetDataNumAttributesInternal({1}_DataIndex);\n"
-				"        uint NumAttributesProcessed = 0;\n"
-				"\n"
-				"        for (int AttributeId = 0; AttributeId < 128; ++AttributeId)\n"
-				"        {\n"
-				"            const uint {1}_Stride = {1}_GetAttributeStrideInternal({1}_DataIndex, AttributeId);\n"
-				"\n"
-				"            if ({1}_Stride == 0)\n"
-				"            {\n"
-				"                // No output attribute to write to.\n"
-				"                continue;\n"
-				"            }\n"
-				"\n"
-				"            const uint {0}_Stride = {0}_GetAttributeStrideInternal({0}_DataIndex, AttributeId);\n"
-				"            const uint {1}_ElementAddress = {1}_GetElementAddressInternal({1}_DataIndex, ElementIndex, AttributeId);\n"
-				"\n"
-				"            if ({0}_Stride == {1}_Stride)\n"
-				"            {\n"
-				"                const uint {0}_ElementAddress = {0}_GetElementAddressInternal({0}_DataIndex, ElementIndex, AttributeId);\n"
-				"\n"
-				"                for (uint Offset = 0; Offset < {1}_Stride; Offset += 4)\n"
-				"                {\n"
-				"                    {1}_StoreBufferInternal({1}_ElementAddress + Offset, {0}_LoadBufferInternal({0}_ElementAddress + Offset));\n"
-				"                }\n"
-				"\n"
-				"            }\n"
-				"            else\n"
-				"            {\n"
-				"                for (uint Offset = 0; Offset < {1}_Stride; Offset += 4)\n"
-				"                {\n"
-				"                    // Initialize output data (only header part of buffer will be initialized).\n"
-				"                    {1}_StoreBufferInternal({1}_ElementAddress + Offset, 0u);\n"
-				"                }\n"
-				"            }\n"
-				"\n"
-				"            if (++NumAttributesProcessed >= NumAttributes)\n"
-				"            {\n"
-				"                break; // We can early-out when we've looked at all the possible attributes\n"
-				"            }\n"
-				"        }\n"
-				"    }\n"),
+				"// Point processor always initializes outputs by copying input data elements.\n"
+				"{1}_CopyElementFrom_{0}({1}_DataIndex, ElementIndex, {0}_DataIndex, ElementIndex);\n"),
 				{ InputPin->Properties.Label.ToString(), OutputPin->Properties.Label.ToString() });
 		}
 	}
diff --git a/Engine/Plugins/PCG/Source/PCG/Private/Compute/Elements/PCGCustomHLSL.h b/Engine/Plugins/PCG/Source/PCG/Private/Compute/Elements/PCGCustomHLSL.h
--- a/Engine/Plugins/PCG/Source/PCG/Private/Compute/Elements/PCGCustomHLSL.h
+++ b/Engine/Plugins/PCG/Source/PCG/Private/Compute/Elements/PCGCustomHLSL.h
@@ -79,6 +79,7 @@ protected:
 	virtual FPCGElementPtr CreateElement() const override;
 	//~End UPCGSettings interface
 
+public:
 #if WITH_EDITOR
 	//~Begin IPCGNodeSourceTextProvider interface
 	FString GetShaderText() const override;
@@ -90,6 +91,10 @@ protected:
 	//~End IPCGNodeSourceTextProvider interface
 #endif
 
+	/** Get user multiplier value applied to the dispatch thread count. */
+	uint32 GetThreadCountMultiplier() const { return IsThreadCountMultiplierInUse() ? static_cast<uint32>(ThreadCountMultiplier) : 1u; }
+
+protected:
 	/** Gets the GPU pin properties for the output pin with the given label. */
 	const FPCGPinPropertiesGPU* GetOutputPinPropertiesGPU(const FName& InPinLabel) const;
 
@@ -126,6 +131,9 @@ protected:
 	/** Will the ThreadCountMultiplier value be applied when calculating the dispatch thread count. */
 	bool IsThreadCountMultiplierInUse() const { return KernelType == EPCGKernelType::Custom && DispatchThreadCount != EPCGDispatchThreadCount::Fixed; }
 
+	/** Generate any functions that require code-gen for the shader source. */
+	virtual bool GetGeneratedFunctions(FString* OutFunctions, FText* OutErrorText = nullptr) const;
+
 protected:
 	UPROPERTY(BlueprintReadWrite, EditAnywhere, Category = "Settings")
 	EPCGKernelType KernelType = EPCGKernelType::PointProcessor;
