[PCG] Fix FPCGPoint Crc calculations (FArchiveCrc32 is neither Loading nor Saving so consider we are writing to it when archive is not loading)

#rnx

diff --git a/Engine/Plugins/PCG/Source/PCG/Private/PCGPoint.cpp b/Engine/Plugins/PCG/Source/PCG/Private/PCGPoint.cpp
--- a/Engine/Plugins/PCG/Source/PCG/Private/PCGPoint.cpp
+++ b/Engine/Plugins/PCG/Source/PCG/Private/PCGPoint.cpp
@@ -58,7 +58,7 @@ bool FPCGPoint::Serialize(FStructuredArchive::FSlot Slot)
 	
 	const FPCGPoint Default;
 	EPCGPointSerializeFields SerializeMask = EPCGPointSerializeFields::None;
-	if (UnderlyingArchive.IsSaving())
+	if (!UnderlyingArchive.IsLoading())
 	{
 		if (Density != Default.Density)
 		{
