[PCG] Fix crash when enabling GPU execution on multiple selected SM Spawners, each of which having multiple incident edges.

Details:

* All our GPU pins are explicitly set to Single Connection (SC)
* However the case in the jira starts with SM Spawners with multiple connected edges, then enables GPU on them one at a time, and one or more compilations will get triggered while one or more SM Spawners still have multiple edges.
* We have a function that injects merge elements to pre-merge on CPU before GPU execution, but it has not been tested since prehistoric times because we just enforced SC on all GPU nodes which normally prevented this case.
* This function was not in a good state and was being called from the wrong place (it invalidated the task successor map).
* I moved it to the right place and things almost work, but the current compute graph element virtual pin setup requires output pins, and the injected merge element does not have a node/pins, so an ensure gets triggered.

So I've demoted the ensure to a warning, and future work can fix this up to add support for Multiple Connection to GPU pins.

#rnx

diff --git a/Engine/Plugins/PCG/Source/PCG/Private/Graph/PCGGraphCompilerGPU.cpp b/Engine/Plugins/PCG/Source/PCG/Private/Graph/PCGGraphCompilerGPU.cpp
--- a/Engine/Plugins/PCG/Source/PCG/Private/Graph/PCGGraphCompilerGPU.cpp
+++ b/Engine/Plugins/PCG/Source/PCG/Private/Graph/PCGGraphCompilerGPU.cpp
@@ -377,10 +377,14 @@ void FPCGGraphCompilerGPU::SetupVirtualPins(
 					}
 				}
 
-				if (ensure(OutputPin))
+				if (OutputPin)
 				{
 					OutOutputCPUPinToVirtualPin.Add(OutputPin, VirtualLabel);
 				}
+				else
+				{
+					UE_LOG(LogPCG, Warning, TEXT("Edge dropped during compilation due to missing upstream output pin."))
+				}
 			}
 		}
 
@@ -1224,6 +1228,10 @@ void FPCGGraphCompilerGPU::CreateGPUNodes(FPCGGraphCompiler& InOutCompiler, UPCG
 		return;
 	}
 
+	// For input pins at CPU -> GPU boundary, inject gather elements to pre-combine data on CPU side
+	// before passing to GPU.
+	CreateGatherTasksAtGPUInputs(InGraph, GPUCompatibleTaskIds, InOutCompiledTasks);
+
 	FTaskToSuccessors TaskSuccessors;
 	TaskSuccessors.Reserve(InOutCompiledTasks.Num());
 	for (FPCGTaskId TaskId = 0; TaskId < InOutCompiledTasks.Num(); ++TaskId)
@@ -1236,10 +1244,6 @@ void FPCGGraphCompilerGPU::CreateGPUNodes(FPCGGraphCompiler& InOutCompiler, UPCG
 		}
 	}
 
-	// For input pins at CPU -> GPU boundary, inject gather elements to pre-combine data on CPU side
-	// before passing to GPU.
-	CreateGatherTasksAtGPUInputs(InGraph, GPUCompatibleTaskIds, InOutCompiledTasks);
-
 	TArray<TSet<FPCGTaskId>> NodeSubsetsToConvertToCFGraph;
 	CollectGPUNodeSubsets(InOutCompiledTasks, TaskSuccessors, GPUCompatibleTaskIds, NodeSubsetsToConvertToCFGraph);
 
