[PCG] Fix delegates that break when compiling BPs

* Delegates needs to be copied in the ApplyToComponent but it actually breaks dynamic delegates that are re-hooked prior to ApplyToComponent call. So we need to re-hook them after the copy.

#rnx

diff --git a/Engine/Plugins/PCG/Source/PCG/Private/PCGComponent.cpp b/Engine/Plugins/PCG/Source/PCG/Private/PCGComponent.cpp
--- a/Engine/Plugins/PCG/Source/PCG/Private/PCGComponent.cpp
+++ b/Engine/Plugins/PCG/Source/PCG/Private/PCGComponent.cpp
@@ -31,6 +31,7 @@
 #include "Algo/AnyOf.h"
 #include "Algo/Transform.h"
 #include "Components/InstancedStaticMeshComponent.h"
+#include "Engine/BlueprintGeneratedClass.h"
 #include "Engine/Engine.h"
 #include "Engine/Level.h"
 #include "Kismet/GameplayStatics.h"
@@ -3751,6 +3752,15 @@ void FPCGComponentInstanceData::ApplyToComponent(UActorComponent* Component, con
 			PCGComponent->OnPCGGraphGeneratedExternal = ConstructionSourceComponent->OnPCGGraphGeneratedExternal;
 			PCGComponent->OnPCGGraphCleanedExternal = ConstructionSourceComponent->OnPCGGraphCleanedExternal;
 
+			// But if the owner has changed and BP had dynamic delegates, we need to restore them.
+			if (AActor* Owner = PCGComponent->GetOwner())
+			{
+				if (Owner != ConstructionSourceComponent->GetOwner())
+				{
+					UBlueprintGeneratedClass::BindDynamicDelegates(Owner->GetClass(), Owner);
+				}
+			}
+
 			PCGComponent->bWasGeneratedThisSession = ConstructionSourceComponent->bWasGeneratedThisSession;
 			PCGComponent->InspectionCounter = ConstructionSourceComponent->InspectionCounter;
 #endif
