[PCG] GPU - fix GetPointTransform and SetPointTransform for non-uniform scales and non-identity rotations.

#rnx

diff --git a/Engine/Plugins/PCG/Shaders/Private/PCGDataCollectionDataInterface.ush b/Engine/Plugins/PCG/Shaders/Private/PCGDataCollectionDataInterface.ush
--- a/Engine/Plugins/PCG/Shaders/Private/PCGDataCollectionDataInterface.ush
+++ b/Engine/Plugins/PCG/Shaders/Private/PCGDataCollectionDataInterface.ush
@@ -332,7 +332,7 @@ uint GetSeed_{DataInterfaceName}(uint DataIndex, uint ElementIndex)
 float4x4 GetPointTransform_{DataInterfaceName}(uint InDataIndex, uint InElementIndex)
 {
 	const FQuat Rotation = GetRotation_{DataInterfaceName}(InDataIndex, InElementIndex);
-	const half3x3 RotationMatrix = transpose(QuatToMatrix(Rotation));
+	const half3x3 RotationMatrix = QuatToMatrix(Rotation);
 	const float3 Position = GetPosition_{DataInterfaceName}(InDataIndex, InElementIndex);
 	const float3 Scale = GetScale_{DataInterfaceName}(InDataIndex, InElementIndex);
 	
@@ -341,9 +341,9 @@ float4x4 GetPointTransform_{DataInterfaceName}(uint InDataIndex, uint InElementI
 	const float3 Axis2 = Scale.z * RotationMatrix[2];
 
 	return float4x4(
-		Axis0.x, Axis0.y, Axis0.z, Position.x,
-		Axis1.x, Axis1.y, Axis1.z, Position.y,
-		Axis2.x, Axis2.y, Axis2.z, Position.z,
+		Axis0.x, Axis1.x, Axis2.x, Position.x,
+		Axis0.y, Axis1.y, Axis2.y, Position.y,
+		Axis0.z, Axis1.z, Axis2.z, Position.z,
 		(float3)0.0, 1.0);
 }
 
@@ -401,7 +401,7 @@ void SetSeed_{DataInterfaceName}(uint DataIndex, uint ElementIndex, uint Seed)
 
 void SetPointTransform_{DataInterfaceName}(uint InDataIndex, uint InElementIndex, float4x4 Transform)
 {
-	const float3 Scale = float3(length(Transform._m00_m01_m02), length(Transform._m10_m11_m12), length(Transform._m20_m21_m22));
+	const float3 Scale = float3(length(Transform._m00_m10_m20), length(Transform._m01_m11_m21), length(Transform._m02_m12_m22));
 	SetScale_{DataInterfaceName}(InDataIndex, InElementIndex, Scale);
 
 	// Assumes Transform axes have the correct handedness. We could do more work to fix this up too if required.
