[PCG] Add metadata to CRC computation for Spline Interior Surface Data

* Otherwise just changing metadata on the spline does not dirty the surface data downstream

#rnx

diff --git a/Engine/Plugins/PCG/Source/PCG/Private/Data/PCGSplineInteriorSurfaceData.cpp b/Engine/Plugins/PCG/Source/PCG/Private/Data/PCGSplineInteriorSurfaceData.cpp
--- a/Engine/Plugins/PCG/Source/PCG/Private/Data/PCGSplineInteriorSurfaceData.cpp
+++ b/Engine/Plugins/PCG/Source/PCG/Private/Data/PCGSplineInteriorSurfaceData.cpp
@@ -31,7 +31,10 @@ void UPCGSplineInteriorSurfaceData::AddToCrc(FArchiveCrc32& Ar, bool bFullDataCr
 {
 	Super::AddToCrc(Ar, bFullDataCrc);
 
-	// Implementation note: no metadata in this data yet.
+	if (Metadata)
+	{
+		Metadata->AddToCrc(Ar, bFullDataCrc);
+	}
 
 	uint32 UniqueTypeID = StaticClass()->GetDefaultObject()->GetUniqueID();
 	Ar << UniqueTypeID;
