[PCG] Remove unnecessary Refresh() that was causing the source editor text box to lose focus, introducing significant friction to shader editing.

#rnx

diff --git a/Engine/Plugins/PCG/Source/PCGEditor/Private/Widgets/SPCGEditorNodeSource.cpp b/Engine/Plugins/PCG/Source/PCGEditor/Private/Widgets/SPCGEditorNodeSource.cpp
--- a/Engine/Plugins/PCG/Source/PCGEditor/Private/Widgets/SPCGEditorNodeSource.cpp
+++ b/Engine/Plugins/PCG/Source/PCGEditor/Private/Widgets/SPCGEditorNodeSource.cpp
@@ -98,6 +98,8 @@ void SPCGEditorNodeSource::SetTextProviderObject(UObject* InProviderObject)
 	ShaderFunctionsText = GetShaderFunctionsAsText();
 	ShaderSourceText = GetShaderTextAsText();
 
+	SyntaxHighlighterShaderText->ClearCompilerMessages();
+
 	Refresh();
 }
 
@@ -243,17 +245,6 @@ void SPCGEditorNodeSource::SetShaderSourceText() const
 void SPCGEditorNodeSource::OnDiagnosticsUpdated(const FPCGCompilerDiagnostics& InDiagnostics) const
 {
 	SyntaxHighlighterShaderText->SetCompilerMessages(InDiagnostics);
-
-	PCGHelpers::ExecuteOnGameThread(UE_SOURCE_LOCATION, [ShaderTextBoxPtr = TWeakPtr<SPCGNodeSourceTextBox>(ShaderTextTextBox)]()
-	{
-		if (ShaderTextBoxPtr.IsValid())
-		{
-			if (TSharedPtr<SPCGNodeSourceTextBox> ShaderTextBox = ShaderTextBoxPtr.Pin())
-			{
-				ShaderTextBox->Refresh();
-			}
-		}
-	});
 }
 
 TSharedRef<SWidget> SPCGEditorNodeSource::ConstructNonExpandableHeaderWidget(const SExpandableArea::FArguments& InArgs) const
