[PCG] Fix procedural instance flushing on material changes.

Combined with the previous flush fix, the flushes seem to be eliminated for a range of simple actions like applying/modifying materials.

#rnx

diff --git a/Engine/Plugins/PCG/Source/PCGEditor/Private/PCGEditorModule.cpp b/Engine/Plugins/PCG/Source/PCGEditor/Private/PCGEditorModule.cpp
--- a/Engine/Plugins/PCG/Source/PCGEditor/Private/PCGEditorModule.cpp
+++ b/Engine/Plugins/PCG/Source/PCGEditor/Private/PCGEditorModule.cpp
@@ -113,6 +113,8 @@ void FPCGEditorModule::ShutdownModule()
 		{
 			return Delegate.GetHandle() == ShouldDisableCPUThrottlingDelegateHandle;
 		});
+
+		GEditor->OnSceneMaterialsModifiedEvent().RemoveAll(this);
 	}
 
 	if (FModuleManager::Get().IsModuleLoaded("LevelEditor"))
@@ -142,6 +144,8 @@ void FPCGEditorModule::OnLevelEditorCreated(TSharedPtr<ILevelEditor> InLevelEdit
 			GEditor->ShouldDisableCPUThrottlingDelegates.Add(UEditorEngine::FShouldDisableCPUThrottling::CreateRaw(this, &FPCGEditorModule::ShouldDisableCPUThrottling));
 			ShouldDisableCPUThrottlingDelegateHandle = GEditor->ShouldDisableCPUThrottlingDelegates.Last().GetHandle();
 		}
+
+		GEditor->OnSceneMaterialsModifiedEvent().AddRaw(this, &FPCGEditorModule::OnSceneMaterialsModified);
 	}
 }
 
@@ -169,6 +173,23 @@ void FPCGEditorModule::OnEditorModeIDChanged(const FEditorModeID& EditorModeID,
 	}
 }
 
+void FPCGEditorModule::OnSceneMaterialsModified()
+{
+	if (!GetDefault<UPCGEditorProjectSettings>()->bAutoRefreshGPUStaticMeshSpawners)
+	{
+		return;
+	}
+
+	// Currently, there is no explicit persistence of instance data in the GPU scene and procedural instances are lost when the GPU Scene is flushed.
+	// TODO: This function is a stop gap that refreshes PCG Components owning procedural instances, and should be removed later.
+	if (UPCGSubsystem* Subsystem = UPCGSubsystem::GetSubsystemForCurrentWorld())
+	{
+		Subsystem->RefreshAllComponentsFiltered(
+			[](UPCGComponent* InComponent) { return InComponent->AreProceduralInstancesInUse(); },
+			EPCGChangeType::Structural);
+	}
+}
+
 bool FPCGEditorModule::SupportsDynamicReloading()
 {
 	return true;
diff --git a/Engine/Plugins/PCG/Source/PCGEditor/Public/PCGEditorModule.h b/Engine/Plugins/PCG/Source/PCGEditor/Public/PCGEditorModule.h
--- a/Engine/Plugins/PCG/Source/PCGEditor/Public/PCGEditorModule.h
+++ b/Engine/Plugins/PCG/Source/PCGEditor/Public/PCGEditorModule.h
@@ -46,6 +46,11 @@ protected:
 	void RegisterOnEditorModeChange();
 	void OnEditorModeIDChanged(const FEditorModeID& EditorModeID, bool bIsEntering);
 
+	/** [EXPERIMENTAL] Used to refresh procedural instances when materials are modified which can otherwise be lost.
+	* Note: This function subject to change/removal without deprecation.
+	*/
+	void OnSceneMaterialsModified();
+
 	bool ShouldDisableCPUThrottling();
 
 	TArray<TSharedRef<IAssetTypeActions>> RegisteredAssetTypeActions;
diff --git a/Engine/Plugins/PCG/Source/PCGEditor/Public/PCGEditorSettings.h b/Engine/Plugins/PCG/Source/PCGEditor/Public/PCGEditorSettings.h
--- a/Engine/Plugins/PCG/Source/PCGEditor/Public/PCGEditorSettings.h
+++ b/Engine/Plugins/PCG/Source/PCGEditor/Public/PCGEditorSettings.h
@@ -23,6 +23,13 @@ public:
 	/** Default Builder Settings to use when running the PCGWorldPartitionBuilder */
 	UPROPERTY(EditAnywhere, config, Category = Builder)
 	TSoftObjectPtr<UPCGBuilderSettings> DefaultBuilderSetting;
+
+	/** [EXPERIMENTAL] Whether to automatically refresh components that use GPU Static Mesh Spawners when materials are modified. Temporary workaround for
+	* issues where instances can be lost.
+	* Note: This setting is subject to change/removal without deprecation.
+	*/
+	UPROPERTY(EditAnywhere, Config, Category = Workflow)
+	bool bAutoRefreshGPUStaticMeshSpawners = true;
 };
 
 UCLASS(config=EditorPerProjectUserSettings)
diff --git a/Engine/Source/Editor/UnrealEd/Classes/Editor/EditorEngine.h b/Engine/Source/Editor/UnrealEd/Classes/Editor/EditorEngine.h
--- a/Engine/Source/Editor/UnrealEd/Classes/Editor/EditorEngine.h
+++ b/Engine/Source/Editor/UnrealEd/Classes/Editor/EditorEngine.h
@@ -3112,6 +3112,10 @@ public:
 	/** This function should be called to notify the editor that new materials were added to our scene or some materials were modified */
 	UNREALED_API void OnSceneMaterialsModified();
 
+	/** Called when a scene material is added or modified. */
+	DECLARE_EVENT(UEditorEngine, FSceneMaterialsModifiedEvent);
+	FSceneMaterialsModifiedEvent& OnSceneMaterialsModifiedEvent() { return SceneMaterialsModifiedEvent; }
+
 	/** Call this function to change the feature level and to override the material quality platform of the editor and PIE worlds */
 	UNREALED_API void SetPreviewPlatform(const FPreviewPlatformInfo& NewPreviewPlatform, bool bSaveSettings);
 
@@ -3178,6 +3182,9 @@ private:
 	/** Minimized Windows during PIE */
 	TArray<TWeakPtr<SWindow>> MinimizedWindowsDuringPIE;
 
+	/** Delegate broadcast when a scene material is added or modified. */
+	FSceneMaterialsModifiedEvent SceneMaterialsModifiedEvent;
+
 public:
 	/**
 	 * Get a Subsystem of specified type
diff --git a/Engine/Source/Editor/UnrealEd/Private/EditorEngine.cpp b/Engine/Source/Editor/UnrealEd/Private/EditorEngine.cpp
--- a/Engine/Source/Editor/UnrealEd/Private/EditorEngine.cpp
+++ b/Engine/Source/Editor/UnrealEd/Private/EditorEngine.cpp
@@ -6221,6 +6221,7 @@ bool UEditorEngine::IsOfflineShaderCompilerAvailable(UWorld* World)
 
 void UEditorEngine::OnSceneMaterialsModified()
 {
+	SceneMaterialsModifiedEvent.Broadcast();
 }
 
 void UEditorEngine::OnEffectivePreviewShaderPlatformChange()
