[PCG] StaticMeshSpawner full data crc on inputs

#rnx

diff --git a/Engine/Plugins/PCG/Source/PCG/Private/Elements/PCGStaticMeshSpawner.cpp b/Engine/Plugins/PCG/Source/PCG/Private/Elements/PCGStaticMeshSpawner.cpp
--- a/Engine/Plugins/PCG/Source/PCG/Private/Elements/PCGStaticMeshSpawner.cpp
+++ b/Engine/Plugins/PCG/Source/PCG/Private/Elements/PCGStaticMeshSpawner.cpp
@@ -587,6 +587,17 @@ bool FPCGStaticMeshSpawnerElement::CanExecuteOnlyOnMainThread(FPCGContext* Conte
 	return !Context || Context->CurrentPhase == EPCGExecutionPhase::Execute || Context->CurrentPhase == EPCGExecutionPhase::PrepareData;
 }
 
+void FPCGStaticMeshSpawnerElement::GetDependenciesCrc(const FPCGDataCollection& InInput, const UPCGSettings* InSettings, UPCGComponent* InComponent, FPCGCrc& OutCrc) const
+{
+	TRACE_CPUPROFILER_EVENT_SCOPE_TEXT(*FString::Printf(TEXT("FPCGStaticMeshSpawnerElement::GetDependenciesCrc (%s)"), InSettings ? *InSettings->GetFName().ToString() : TEXT("")));
+		
+	// Make sure Inputs have a full data crc
+	FPCGDataCollection InputCopy = InInput;
+	InputCopy.ComputeCrcs(/*bFullDataCrc=*/true);
+
+	IPCGElement::GetDependenciesCrc(InputCopy, InSettings, InComponent, OutCrc);
+}
+
 void FPCGStaticMeshSpawnerElement::SpawnStaticMeshInstances(FPCGStaticMeshSpawnerContext* Context, const FPCGMeshInstanceList& InstanceList, AActor* TargetActor, const FPCGPackedCustomData* InPackedCustomData) const
 {
 	// Populate the (H)ISM from the previously prepared entries
diff --git a/Engine/Plugins/PCG/Source/PCG/Public/Elements/PCGStaticMeshSpawner.h b/Engine/Plugins/PCG/Source/PCG/Public/Elements/PCGStaticMeshSpawner.h
--- a/Engine/Plugins/PCG/Source/PCG/Public/Elements/PCGStaticMeshSpawner.h
+++ b/Engine/Plugins/PCG/Source/PCG/Public/Elements/PCGStaticMeshSpawner.h
@@ -135,6 +135,7 @@ class PCG_API FPCGStaticMeshSpawnerElement : public IPCGElement
 public:
 	virtual bool CanExecuteOnlyOnMainThread(FPCGContext* Context) const override;
 	virtual bool IsCacheable(const UPCGSettings* InSettings) const override { return false; }
+	virtual void GetDependenciesCrc(const FPCGDataCollection& InInput, const UPCGSettings* InSettings, UPCGComponent* InComponent, FPCGCrc& OutCrc) const override;
 
 protected:
 	virtual FPCGContext* CreateContext() override;
