[PCG] Attribute Filter optimization pass to help alleviate OOM crash

#rnx

diff --git a/Engine/Plugins/PCG/Source/PCG/Private/Elements/PCGAttributeFilter.cpp b/Engine/Plugins/PCG/Source/PCG/Private/Elements/PCGAttributeFilter.cpp
--- a/Engine/Plugins/PCG/Source/PCG/Private/Elements/PCGAttributeFilter.cpp
+++ b/Engine/Plugins/PCG/Source/PCG/Private/Elements/PCGAttributeFilter.cpp
@@ -536,6 +536,7 @@ bool FPCGAttributeFilterElementBase::DoFiltering(FPCGContext* Context, EPCGAttri
 		const TArray<FPCGPoint>* OriginalPoints = nullptr;
 		TArray<FPCGPoint>* InFilterPoints = nullptr;
 		TArray<FPCGPoint>* OutFilterPoints = nullptr;
+		TBitArray<TInlineAllocator<2048>> FilterBitArray;
 
 		const UPCGMetadata* OriginalMetadata = nullptr;
 		UPCGMetadata* InFilterMetadata = nullptr;
@@ -696,8 +697,8 @@ bool FPCGAttributeFilterElementBase::DoFiltering(FPCGContext* Context, EPCGAttri
 			OperationData.InFilterPoints = &InFilterPointData->GetMutablePoints();
 			OperationData.OutFilterPoints = &OutFilterPointData->GetMutablePoints();
 
-			OperationData.InFilterPoints->Reserve(OriginalPointData->GetPoints().Num());
-			OperationData.OutFilterPoints->Reserve(OriginalPointData->GetPoints().Num());
+			// Will be set individually in batches
+			OperationData.FilterBitArray.SetNumUninitialized(OriginalPointData->GetPoints().Num());
 
 			InFilterData = InFilterPointData;
 			OutFilterData = OutFilterPointData;
@@ -732,10 +733,10 @@ bool FPCGAttributeFilterElementBase::DoFiltering(FPCGContext* Context, EPCGAttri
 				return false;
 			}
 
-			TArray<Type, TInlineAllocator<PCGAttributeFilterConstants::ChunkSize>> TargetValues;
-			TArray<Type, TInlineAllocator<PCGAttributeFilterConstants::ChunkSize>> FirstThresholdValues;
-			TArray<Type, TInlineAllocator<PCGAttributeFilterConstants::ChunkSize>> SecondThresholdValues;
-			TArray<bool, TInlineAllocator<PCGAttributeFilterConstants::ChunkSize>> SkipTests;
+			TArray<Type, TFixedAllocator<PCGAttributeFilterConstants::ChunkSize>> TargetValues;
+			TArray<Type, TFixedAllocator<PCGAttributeFilterConstants::ChunkSize>> FirstThresholdValues;
+			TArray<Type, TFixedAllocator<PCGAttributeFilterConstants::ChunkSize>> SecondThresholdValues;
+			TArray<bool, TFixedAllocator<PCGAttributeFilterConstants::ChunkSize>> SkipTests;
 			TargetValues.SetNum(PCGAttributeFilterConstants::ChunkSize);
 			FirstThresholdValues.SetNum(PCGAttributeFilterConstants::ChunkSize);
 			SecondThresholdValues.SetNum(PCGAttributeFilterConstants::ChunkSize);
@@ -811,9 +812,7 @@ bool FPCGAttributeFilterElementBase::DoFiltering(FPCGContext* Context, EPCGAttri
 				{
 					if (OperationData.bIsInputPointData)
 					{
-						TArray<FPCGPoint>* Points = bInFilter ? OperationData.InFilterPoints : OperationData.OutFilterPoints;
-						check(Points && OperationData.OriginalPoints);
-						Points->Add((*OperationData.OriginalPoints)[Index]);
+						OperationData.FilterBitArray[Index] = bInFilter;
 					}
 					else
 					{
@@ -852,6 +851,29 @@ bool FPCGAttributeFilterElementBase::DoFiltering(FPCGContext* Context, EPCGAttri
 
 		if (PCGMetadataAttribute::CallbackWithRightType(TargetAccessor->GetUnderlyingType(), Operation))
 		{
+			if (OperationData.bIsInputPointData)
+			{
+				check(OperationData.OriginalPoints);
+
+				const int32 NumInFilterPoints = OperationData.FilterBitArray.CountSetBits();
+				const int32 NumOutFilterPoints = OperationData.OriginalPoints->Num() - NumInFilterPoints;
+
+				OperationData.InFilterPoints->Reserve(NumInFilterPoints);
+				OperationData.OutFilterPoints->Reserve(NumOutFilterPoints);
+
+				for (int32 Index = 0; Index < OperationData.FilterBitArray.Num(); ++Index)
+				{
+					if (OperationData.FilterBitArray[Index])
+					{
+						OperationData.InFilterPoints->Add((*OperationData.OriginalPoints)[Index]);
+					}
+					else
+					{
+						OperationData.OutFilterPoints->Add((*OperationData.OriginalPoints)[Index]);
+					}
+				}
+			}
+
 			FPCGTaggedData& InFilterOutput = Outputs.Add_GetRef(Input);
 			InFilterOutput.Pin = PCGPinConstants::DefaultInFilterLabel;
 			InFilterOutput.Data = InFilterData;
diff --git a/Engine/Plugins/PCG/Source/PCG/Public/Elements/PCGAttributeFilter.h b/Engine/Plugins/PCG/Source/PCG/Public/Elements/PCGAttributeFilter.h
--- a/Engine/Plugins/PCG/Source/PCG/Public/Elements/PCGAttributeFilter.h
+++ b/Engine/Plugins/PCG/Source/PCG/Public/Elements/PCGAttributeFilter.h
@@ -200,6 +200,7 @@ class FPCGAttributeFilterElementBase : public IPCGElement
 {
 protected:
 	virtual EPCGElementExecutionLoopMode ExecutionLoopMode(const UPCGSettings* Settings) const override { return EPCGElementExecutionLoopMode::SinglePrimaryPin; }
+	/** Performs the filter operation. Returns true when the operation is complete. */
 	bool DoFiltering(FPCGContext* Context, EPCGAttributeFilterOperator InOperation, const FPCGAttributePropertyInputSelector& TargetAttribute, bool bHasSpatialToPointDeprecation, bool bWarnOnDataMissingAttribute, const FPCGAttributeFilterThresholdSettings& FirstThreshold, const FPCGAttributeFilterThresholdSettings* SecondThreshold = nullptr) const;
 };
 
@@ -214,4 +215,3 @@ class FPCGAttributeFilterRangeElement : public FPCGAttributeFilterElementBase
 protected:
 	virtual bool ExecuteInternal(FPCGContext* Context) const override;
 };
-
