[PCG] Remove internal deprecation redirectors for 5.5 internal usage.

#rnx

diff --git a/Engine/Plugins/PCG/Config/BasePCG.ini b/Engine/Plugins/PCG/Config/BasePCG.ini
--- a/Engine/Plugins/PCG/Config/BasePCG.ini
+++ b/Engine/Plugins/PCG/Config/BasePCG.ini
@@ -101,17 +101,3 @@
 +EnumRedirects=(OldName="EPCGMedadataVectorOperation",NewName="EPCGMetadataVectorOperation")
 +EnumRedirects=(OldName="EPCGMedadataBitwiseOperation",NewName="EPCGMetadataBitwiseOperation")
 +EnumRedirects=(OldName="EPCGMedadataRotatorOperation",NewName="EPCGMetadataRotatorOperation")
-
-// START - TODO: Remove before 5.5
-+ClassRedirects=(OldName="/Script/PCG.PCGSplineSlicerSettings",NewName="/Script/PCG.PCGSubdivideSplineSettings")
-+ClassRedirects=(OldName="/Script/PCG.PCGSegmentSlicerSettings",NewName="/Script/PCG.PCGSubdivideSegmentSettings")
-+ClassRedirects=(OldName="/Script/PCG.PCGSlicingBaseSettings",NewName="/Script/PCG.PCGSubdivisionBaseSettings")
-+ClassRedirects=(OldName="/Script/PCG.PCGVolumeSlicerSettings",NewName="/Script/PCG.PCGDuplicateCrossSectionsSettings")
-
-+StructRedirects=(OldName="/Script/PCG.PCGSlicingSubmodule",NewName="/Script/PCG.PCGSubdivisionSubmodule")
-+StructRedirects=(OldName="/Script/PCG.PCGSlicingModuleAttributeNames",NewName="/Script/PCG.PCGSubdivisionModuleAttributeNames")
-
-+PropertyRedirects=(OldName="/Script/PCG.PCGSubdivideSegmentSettings.SlicingAxis",NewName="/Script/PCG.PCGSubdivideSegmentSettings.SubdivisionAxis")
-+PropertyRedirects=(OldName="/Script/PCG.PCGSubdivideSegmentSettings.bAcceptIncompleteSlicing",NewName="/Script/PCG.PCGSubdivideSegmentSettings.bAcceptIncompleteSubdivision")
-+PropertyRedirects=(OldName="/Script/PCG.PCGSubdivideSplineSettings.bAcceptIncompleteSlicing",NewName="/Script/PCG.PCGSubdivideSplineSettings.bAcceptIncompleteSubdivision")
-// END - TODO : Remove before 5.5
