Fixed crc caching issue in per-data caching related to the enabled flag not trigger a crc update.
Also made the 'treat editor viewport as generation source' bp-accessible

#rnx

diff --git a/Engine/Plugins/PCG/Source/PCG/Private/PCGElement.cpp b/Engine/Plugins/PCG/Source/PCG/Private/PCGElement.cpp
--- a/Engine/Plugins/PCG/Source/PCG/Private/PCGElement.cpp
+++ b/Engine/Plugins/PCG/Source/PCG/Private/PCGElement.cpp
@@ -282,7 +282,7 @@ void IPCGElement::PreExecute(FPCGContext* Context) const
 		// Implementation note: this supposes that the current node has only ONE required pin, and not multiple.
 		// For more complex cases (such as multiple required pins, whether cartesian or matching), the implementation should use common code instead to streamline this process -
 		// both when getting the results from the cache but also when writing them
-		if (ExecutionLoopMode(Settings) != EPCGElementExecutionLoopMode::NotALoop && IsCacheable(Settings))
+		if (ExecutionLoopMode(Settings) != EPCGElementExecutionLoopMode::NotALoop && IsCacheableInstance(Settings))
 		{
 			PreExecutePrimaryLoopElement(Context, Settings);
 		}
@@ -373,7 +373,7 @@ void IPCGElement::PostExecute(FPCGContext* Context) const
 
 	if (CVarPCGAllowPerDataCaching.GetValueOnAnyThread())
 	{
-		if (!Context->OutputData.bCancelExecution && ExecutionLoopMode(Settings) != EPCGElementExecutionLoopMode::NotALoop && IsCacheable(Settings))
+		if (!Context->OutputData.bCancelExecution && ExecutionLoopMode(Settings) != EPCGElementExecutionLoopMode::NotALoop && IsCacheableInstance(Settings))
 		{
 			PostExecutePrimaryLoopElement(Context, Settings);
 		}
diff --git a/Engine/Plugins/PCG/Source/PCG/Private/PCGSettings.cpp b/Engine/Plugins/PCG/Source/PCG/Private/PCGSettings.cpp
--- a/Engine/Plugins/PCG/Source/PCG/Private/PCGSettings.cpp
+++ b/Engine/Plugins/PCG/Source/PCG/Private/PCGSettings.cpp
@@ -207,13 +207,16 @@ void UPCGSettingsInterface::SetEnabled(bool bInEnabled)
 	if (bEnabled != bInEnabled)
 	{
 		bEnabled = bInEnabled;
-#if WITH_EDITOR
+
 		if (UPCGSettings* Settings = GetSettings())
 		{
+			Settings->CacheCrc();
+
+#if WITH_EDITOR
 			const EPCGChangeType ChangeType = Settings->GetChangeTypeForProperty(GET_MEMBER_NAME_CHECKED(UPCGSettingsInterface, bEnabled));
 			OnSettingsChangedDelegate.Broadcast(Settings, ChangeType);
-		}
 #endif
+		}
 	}
 }
 
diff --git a/Engine/Plugins/PCG/Source/PCG/Public/PCGWorldActor.h b/Engine/Plugins/PCG/Source/PCG/Public/PCGWorldActor.h
--- a/Engine/Plugins/PCG/Source/PCG/Public/PCGWorldActor.h
+++ b/Engine/Plugins/PCG/Source/PCG/Public/PCGWorldActor.h
@@ -66,7 +66,7 @@ public:
 	bool bUse2DGrid = true;
 
 	/** Allows any currently active editor viewport to act as a Runtime Generation Source. */
-	UPROPERTY(EditAnywhere, Category = RuntimeGeneration)
+	UPROPERTY(BlueprintReadWrite, EditAnywhere, Category = RuntimeGeneration)
 	bool bTreatEditorViewportAsGenerationSource = false;
 #endif
 
