[PCG] GPU SM Spawner GPU crash and bugs caught in UnrealFest demo level

#rnx

diff --git a/Engine/Plugins/PCG/Source/PCG/Private/Components/PCGProceduralISMComponent.cpp b/Engine/Plugins/PCG/Source/PCG/Private/Components/PCGProceduralISMComponent.cpp
--- a/Engine/Plugins/PCG/Source/PCG/Private/Components/PCGProceduralISMComponent.cpp
+++ b/Engine/Plugins/PCG/Source/PCG/Private/Components/PCGProceduralISMComponent.cpp
@@ -196,7 +196,6 @@ UPCGProceduralISMComponent::UPCGProceduralISMComponent(const FObjectInitializer&
 	bAffectDistanceFieldLighting = false;
 	bVisibleInRayTracing = false;
 
-	Mobility = EComponentMobility::Static;
 	BodyInstance.bSimulatePhysics = false;
 	SetGenerateOverlapEvents(false);
 
diff --git a/Engine/Plugins/PCG/Source/PCG/Private/Compute/Elements/PCGComputeGraphElement.cpp b/Engine/Plugins/PCG/Source/PCG/Private/Compute/Elements/PCGComputeGraphElement.cpp
--- a/Engine/Plugins/PCG/Source/PCG/Private/Compute/Elements/PCGComputeGraphElement.cpp
+++ b/Engine/Plugins/PCG/Source/PCG/Private/Compute/Elements/PCGComputeGraphElement.cpp
@@ -403,19 +403,30 @@ bool FPCGComputeGraphElement::SetupProceduralISMComponents(FPCGContext* InContex
 				}
 
 				const FString& MeshPathString = InBinding->GetStringTable()[StringKey];
-				
+				if (MeshPathString.IsEmpty())
+				{
+					continue;
+				}
+
+				if (ComponentsToCreate.Num() >= PCGComputeConstants::MAX_PRIMITIVE_COMPONENTS_PER_SPAWNER)
+				{
+					UE_LOG(LogPCG, Warning, TEXT("Attempted to emit too many primitive components, terminated after creating %d."), PCGComputeConstants::MAX_PRIMITIVE_COMPONENTS_PER_SPAWNER);
+					break;
+				}
+
+				UStaticMesh* StaticMesh = Cast<UStaticMesh>(FSoftObjectPath(MeshPathString).TryLoad());
+				if (!StaticMesh)
+				{
+					UE_LOG(LogPCG, Error, TEXT("Could not load static mesh from path '%s'."), *MeshPathString);
+					continue;
+				}
+
 				FPCGProceduralISMComponentDescriptor Descriptor;
 				Descriptor = SelectorByAttribute->TemplateDescriptor;
 				Descriptor.NumInstances = InputPointCount;
 				Descriptor.LocalBounds = LocalBounds;
 				Descriptor.NumCustomFloats = CustomFloatCount;
-				
-				Descriptor.StaticMesh = Cast<UStaticMesh>(FSoftObjectPath(MeshPathString).TryLoad());
-				if (!Descriptor.StaticMesh)
-				{
-					UE_LOG(LogPCG, Error, TEXT("Could not load static mesh from path '%s'."), *MeshPathString);
-					continue;
-				}
+				Descriptor.StaticMesh = StaticMesh;
 
 				PrimitiveStringKeys.Emplace(StringKey);
 				ComponentsToCreate.Add(MoveTemp(Descriptor));
@@ -453,6 +464,17 @@ bool FPCGComputeGraphElement::SetupProceduralISMComponents(FPCGContext* InContex
 
 			for (const FPCGMeshSelectorWeightedEntry& Entry : SelectorWeighted->MeshEntries)
 			{
+				if (Entry.Descriptor.StaticMesh.IsNull())
+				{
+					continue;
+				}
+
+				if (ComponentsToCreate.Num() >= PCGComputeConstants::MAX_PRIMITIVE_COMPONENTS_PER_SPAWNER)
+				{
+					UE_LOG(LogPCG, Warning, TEXT("Attempted to emit too many primitive components, terminated after creating %d."), PCGComputeConstants::MAX_PRIMITIVE_COMPONENTS_PER_SPAWNER);
+					break;
+				}
+
 				if (UStaticMesh* StaticMesh = Entry.Descriptor.StaticMesh.LoadSynchronous())
 				{
 					const float Weight = float(Entry.Weight) / TotalWeight;
@@ -462,10 +484,14 @@ bool FPCGComputeGraphElement::SetupProceduralISMComponents(FPCGContext* InContex
 
 					FPCGProceduralISMComponentDescriptor Descriptor;
 					Descriptor = Entry.Descriptor;
-					Descriptor.NumInstances = FMath::CeilToInt(InputPointCount * Weight);
 					Descriptor.LocalBounds = LocalBounds;
 					Descriptor.NumCustomFloats = CustomFloatCount;
 					Descriptor.StaticMesh = StaticMesh;
+
+					// Heuristic: The random draw can produce more or less values than expectation. 4% buffer was sufficient to cover a test case in the PCG test graph.
+					constexpr float VarianceFactor = 1.04f;
+					Descriptor.NumInstances = FMath::CeilToInt(InputPointCount * Weight * VarianceFactor);
+
 					ComponentsToCreate.Add(MoveTemp(Descriptor));
 
 					if (SpawnerSettings->bApplyMeshBoundsToPoints)
@@ -520,12 +546,6 @@ bool FPCGComputeGraphElement::SetupProceduralISMComponents(FPCGContext* InContex
 			Primitives.Primitives.Add(MISMC->GetComponent());
 
 			bAnyComponentsSetup = true;
-
-			if (Primitives.Primitives.Num() >= PCGComputeConstants::MAX_PRIMITIVE_COMPONENTS_PER_SPAWNER)
-			{
-				UE_LOG(LogPCG, Warning, TEXT("Attempted to emit too many primitive components, terminated after creating %d."), Primitives.Primitives.Num());
-				break;
-			}
 		}
 	}
 
diff --git a/Engine/Plugins/PCG/Source/PCG/Private/Compute/PCGComputeCommon.h b/Engine/Plugins/PCG/Source/PCG/Private/Compute/PCGComputeCommon.h
--- a/Engine/Plugins/PCG/Source/PCG/Private/Compute/PCGComputeCommon.h
+++ b/Engine/Plugins/PCG/Source/PCG/Private/Compute/PCGComputeCommon.h
@@ -57,7 +57,7 @@ namespace PCGComputeConstants
 
 	constexpr uint32 KernelExecutedFlag = 1 << 31;
 
-	constexpr uint32 MAX_PRIMITIVE_COMPONENTS_PER_SPAWNER = 64;
+	constexpr uint32 MAX_PRIMITIVE_COMPONENTS_PER_SPAWNER = 256;
 
 	constexpr uint32 THREAD_GROUP_SIZE = 64;
 }
