[PCG] Withhold graph change notifications while we toggle enabled states of a set of nodes.

Avoids a case where we do a compile & schedule for a partially disabled graph and end up with a confused compute graph element at execution time.

diff --git a/Engine/Plugins/PCG/Source/PCGEditor/Private/PCGEditor.cpp b/Engine/Plugins/PCG/Source/PCGEditor/Private/PCGEditor.cpp
--- a/Engine/Plugins/PCG/Source/PCGEditor/Private/PCGEditor.cpp
+++ b/Engine/Plugins/PCG/Source/PCGEditor/Private/PCGEditor.cpp
@@ -2116,6 +2116,14 @@ void FPCGEditor::OnToggleEnabled()
 	{
 		FScopedTransaction Transaction(*FPCGEditorCommon::ContextIdentifier, LOCTEXT("PCGEditorToggleEnableTransactionMessage", "PCG Editor: Toggle Enable Nodes"), nullptr);
 
+		UPCGGraph* PCGGraph = PCGEditorGraph ? PCGEditorGraph->GetPCGGraph() : nullptr;
+		if (!ensure(PCGGraph))
+		{
+			return;
+		}
+
+		PCGGraph->DisableNotificationsForEditor();
+
 		bool bChanged = false;
 		for (UObject* Object : GraphEditorWidget->GetSelectedNodes())
 		{
@@ -2137,6 +2145,8 @@ void FPCGEditor::OnToggleEnabled()
 			}
 		}
 
+		PCGGraph->EnableNotificationsForEditor();
+
 		if (bChanged)
 		{
 			GraphEditorWidget->NotifyGraphChanged();
