[PCG] Fix crash when setting GPU attribute to none-type
* Fix ensure when compiling GPU graphs with pins labeled 'None'

#rnx

diff --git a/Engine/Plugins/PCG/Source/PCG/Private/Compute/Elements/PCGCustomHLSL.cpp b/Engine/Plugins/PCG/Source/PCG/Private/Compute/Elements/PCGCustomHLSL.cpp
--- a/Engine/Plugins/PCG/Source/PCG/Private/Compute/Elements/PCGCustomHLSL.cpp
+++ b/Engine/Plugins/PCG/Source/PCG/Private/Compute/Elements/PCGCustomHLSL.cpp
@@ -618,7 +618,6 @@ void UPCGCustomHLSLSettings::SetShaderFunctionsText(const FString& NewFunctionsT
 		PostEditChangeProperty(PropertyChangedEvent);
 	}
 	
-
 	OnSettingsChangedDelegate.Broadcast(this, GetChangeTypeForProperty(PropertyName));
 }
 
@@ -1449,6 +1448,12 @@ bool UPCGCustomHLSLSettings::IsKernelValid(FPCGContext* InContext, bool bQuiet)
 			PCG_KERNEL_VALIDATION_ERR(InContext, this, bQuiet, LOCTEXT("InvalidAttributeNameNone", "'None' is not a valid GPU attribute name, check the 'Attributes to Create' array on your pins."));
 			return false;
 		}
+
+		if (AttributeKey.Type == EPCGKernelAttributeType::None)
+		{
+			PCG_KERNEL_VALIDATION_ERR(InContext, this, bQuiet, LOCTEXT("InvalidAttributeTypeNone", "'None' is not a valid GPU attribute type, check the 'Attributes to Create' array on your pins."));
+			return false;
+		}
 	}
 
 	if (InContext)
diff --git a/Engine/Plugins/PCG/Source/PCG/Private/Compute/PCGDataForGPU.h b/Engine/Plugins/PCG/Source/PCG/Private/Compute/PCGDataForGPU.h
--- a/Engine/Plugins/PCG/Source/PCG/Private/Compute/PCGDataForGPU.h
+++ b/Engine/Plugins/PCG/Source/PCG/Private/Compute/PCGDataForGPU.h
@@ -23,7 +23,7 @@ enum class EPCGUnpackDataCollectionResult
 UENUM()
 enum class EPCGKernelAttributeType : uint8
 {
-	None = 0,
+	None = 0 UMETA(Hidden),
 	Bool,
 	Int,
 	Float,
diff --git a/Engine/Plugins/PCG/Source/PCG/Private/Graph/PCGGraphCompilerGPU.cpp b/Engine/Plugins/PCG/Source/PCG/Private/Graph/PCGGraphCompilerGPU.cpp
--- a/Engine/Plugins/PCG/Source/PCG/Private/Graph/PCGGraphCompilerGPU.cpp
+++ b/Engine/Plugins/PCG/Source/PCG/Private/Graph/PCGGraphCompilerGPU.cpp
@@ -682,6 +682,11 @@ UPCGComputeGraph* FPCGGraphCompilerGPU::CompileComputeGraph(
 		{
 			for (const FPCGPinProperties& OutputPinProperties : Settings->AllOutputPinProperties())
 			{
+				if (OutputPinProperties.Label == NAME_None)
+				{
+					continue;
+				}
+
 				if (OutputPinDataInterfaces.Contains({ TaskId, OutputPinProperties.Label }))
 				{
 					ensure(false);
@@ -889,6 +894,11 @@ UPCGComputeGraph* FPCGGraphCompilerGPU::CompileComputeGraph(
 		// Always create a DI for every output pin, so kernel always has something to write to.
 		for (const FPCGPinProperties& OutputPinProperties : Settings->AllOutputPinProperties())
 		{
+			if (OutputPinProperties.Label == NAME_None)
+			{
+				continue;
+			}
+
 			UPCGComputeDataInterface** FoundDI = OutputPinDataInterfaces.Find({ TaskId, OutputPinProperties.Label });
 			if (!ensure(FoundDI) || !ensure(*FoundDI))
 			{
