[PCG] Fix localization keys

#rnx

diff --git a/Engine/Plugins/PCG/Source/PCG/Private/Compute/Elements/PCGCustomHLSL.cpp b/Engine/Plugins/PCG/Source/PCG/Private/Compute/Elements/PCGCustomHLSL.cpp
--- a/Engine/Plugins/PCG/Source/PCG/Private/Compute/Elements/PCGCustomHLSL.cpp
+++ b/Engine/Plugins/PCG/Source/PCG/Private/Compute/Elements/PCGCustomHLSL.cpp
@@ -1704,7 +1704,7 @@ bool UPCGCustomHLSLSettings::GetGeneratedFunctions(FString* OutFunctions, FText*
 #if PCG_KERNEL_LOGGING_ENABLED
 			if (OutErrorText)
 			{
-				*OutErrorText = LOCTEXT("InvalidCopyTargetPin", "CopyFromElement only allowed across pins of matching type.");
+				*OutErrorText = LOCTEXT("InvalidCopyPinTypes", "CopyFromElement only allowed across pins of matching type.");
 			}
 #endif
 			return false;
