[PCG] GPU - ignore 'already generated' check for GenerateOnLoad if PISMCs are in use because this type of generated data is never persisted.

#rnx

diff --git a/Engine/Plugins/PCG/Source/PCG/Private/PCGComponent.cpp b/Engine/Plugins/PCG/Source/PCG/Private/PCGComponent.cpp
--- a/Engine/Plugins/PCG/Source/PCG/Private/PCGComponent.cpp
+++ b/Engine/Plugins/PCG/Source/PCG/Private/PCGComponent.cpp
@@ -1546,7 +1546,10 @@ void UPCGComponent::BeginPlay()
 		GetSubsystem()->RegisterOrUpdatePCGComponent(this);
 	}
 
-	if (bActivated && !bGenerated && GenerationTrigger == EPCGComponentGenerationTrigger::GenerateOnLoad)
+	// Procedural instances are never persisted so always require generation.
+	const bool bAlreadyGenerated = bGenerated & !bProceduralInstancesInUse;
+
+	if (bActivated && !bAlreadyGenerated && GenerationTrigger == EPCGComponentGenerationTrigger::GenerateOnLoad)
 	{
 		GenerateInternal(/*bForce=*/false, EPCGHiGenGrid::Uninitialized, EPCGComponentGenerationTrigger::GenerateOnLoad, {});
 		bRuntimeGenerated = true;
