[PCG] Add upper bound limit for visualization of volume and surface data

* Update the surface sampler generic method to work even if we don't have a context (no random seed)

#rnx

diff --git a/Engine/Plugins/PCG/Source/PCG/Private/Elements/PCGSurfaceSampler.cpp b/Engine/Plugins/PCG/Source/PCG/Private/Elements/PCGSurfaceSampler.cpp
--- a/Engine/Plugins/PCG/Source/PCG/Private/Elements/PCGSurfaceSampler.cpp
+++ b/Engine/Plugins/PCG/Source/PCG/Private/Elements/PCGSurfaceSampler.cpp
@@ -46,12 +46,12 @@ namespace PCGSurfaceSampler
 
 	bool FSurfaceSamplerData::Initialize(const FPCGContext* Context, const FBox& InEffectiveGridBounds, const FTransform& InSurfaceTransform)
 	{
-		if (!Context || !InEffectiveGridBounds.IsValid)
+		if (!InEffectiveGridBounds.IsValid)
 		{
 			return false;
 		}
 
-		Seed = Context->GetSeed();
+		Seed = Context ? Context->GetSeed() : 42;
 
 		// Conceptually, we will break down the surface bounds in a N x M grid
 		InterstitialDistance = Params.PointExtents * 2;
diff --git a/Engine/Plugins/PCG/Source/PCGEditor/Private/DataVisualizations/PCGLandscapeDataVisualization.cpp b/Engine/Plugins/PCG/Source/PCGEditor/Private/DataVisualizations/PCGLandscapeDataVisualization.cpp
new file mode 100644
--- /dev/null
+++ b/Engine/Plugins/PCG/Source/PCGEditor/Private/DataVisualizations/PCGLandscapeDataVisualization.cpp
@@ -0,0 +1,34 @@
+// Copyright Epic Games, Inc. All Rights Reserved.
+
+#include "DataVisualizations/PCGLandscapeDataVisualization.h"
+
+#include "PCGEditorSettings.h"
+#include "Data/PCGLandscapeData.h"
+#include "Data/PCGPointData.h"
+#include "Elements/PCGSurfaceSampler.h"
+
+const UPCGPointData* IPCGLandscapeDataVisualization::CollapseToDebugPointData(FPCGContext* Context, const UPCGData* Data) const
+{
+	const UPCGLandscapeData* LandscapeData = Cast<UPCGLandscapeData>(Data);
+	
+	if (!LandscapeData)
+	{
+		return nullptr;
+	}
+	
+	const int32 NumPointsUpperBound = FMath::Max(1, GetDefault<UPCGEditorSettings>()->TargetNumPointsForDebug);
+	
+	const FBox Bounds = LandscapeData->GetBounds();
+	const double BoundsSurface = (Bounds.Max.X - Bounds.Min.X) * (Bounds.Max.Y - Bounds.Min.Y);
+	const double Extents = FMath::Max(1.0, FMath::Sqrt(BoundsSurface / NumPointsUpperBound) * 0.5);
+	
+	const PCGSurfaceSampler::FSurfaceSamplerParams SamplerParams
+	{
+		.PointsPerSquaredMeter = 99999999.0f, // Big number to make sure we saturate the random process to always select a cell to spawn in.
+		.PointExtents = FVector(Extents, Extents, Extents),
+		.Looseness = 0,
+		.bApplyDensityToPoints = false
+	};
+	
+	return PCGSurfaceSampler::SampleSurface(Context, LandscapeData, nullptr, Bounds, SamplerParams);
+}
diff --git a/Engine/Plugins/PCG/Source/PCGEditor/Private/DataVisualizations/PCGVolumeDataVisualization.cpp b/Engine/Plugins/PCG/Source/PCGEditor/Private/DataVisualizations/PCGVolumeDataVisualization.cpp
new file mode 100644
--- /dev/null
+++ b/Engine/Plugins/PCG/Source/PCGEditor/Private/DataVisualizations/PCGVolumeDataVisualization.cpp
@@ -0,0 +1,44 @@
+// Copyright Epic Games, Inc. All Rights Reserved.
+
+#include "DataVisualizations/PCGVolumeDataVisualization.h"
+
+#include "Data/PCGCollisionShapeData.h"
+#include "Data/PCGCollisionWrapperData.h"
+#include "Data/PCGPrimitiveData.h"
+#include "Data/PCGSpatialData.h"
+#include "Data/PCGVolumeData.h"
+#include "Elements/PCGVolumeSampler.h"
+
+#include "PCGEditorSettings.h"
+
+const UPCGPointData* IPCGPrimitiveDataVisualization::CollapseToDebugPointData(FPCGContext* Context, const UPCGData* Data) const
+{
+	const UPCGSpatialData* SpatialData = Cast<UPCGSpatialData>(Data);
+	
+	if (!SpatialData)
+	{
+		return nullptr;
+	}
+	
+	// Have special cases for primitive data that would potentially generate too many points if they were collapsed with default sizes.
+	// We'll change the size of the extents of the points to make sure we can't get much higher than our upper bound.
+	if (SpatialData->IsA<UPCGVolumeData>() || SpatialData->IsA<UPCGPrimitiveData>() || SpatialData->IsA<UPCGCollisionShapeData>() || SpatialData->IsA<UPCGCollisionWrapperData>())
+	{
+		const int32 NumPointsUpperBound = FMath::Max(1, GetDefault<UPCGEditorSettings>()->TargetNumPointsForDebug);
+		
+		const FBox Bounds = SpatialData->GetBounds();
+		const double BoundsVolume = Bounds.GetVolume();
+		// Pow of 1/3 is the cubic root.
+		const double VoxelSize = FMath::Max(1.0, FMath::Pow(BoundsVolume / NumPointsUpperBound, 1.0 / 3.0));
+		
+		PCGVolumeSampler::FVolumeSamplerParams SamplerParams;
+		SamplerParams.VoxelSize = FVector(VoxelSize, VoxelSize, VoxelSize);
+		SamplerParams.Bounds = Bounds;
+
+		return PCGVolumeSampler::SampleVolume(Context, SamplerParams, SpatialData);
+	}
+	else
+	{
+		return nullptr;
+	}
+}
diff --git a/Engine/Plugins/PCG/Source/PCGEditor/Private/PCGEditorModule.cpp b/Engine/Plugins/PCG/Source/PCGEditor/Private/PCGEditorModule.cpp
--- a/Engine/Plugins/PCG/Source/PCGEditor/Private/PCGEditorModule.cpp
+++ b/Engine/Plugins/PCG/Source/PCGEditor/Private/PCGEditorModule.cpp
@@ -8,8 +8,13 @@
 #include "PCGParamData.h"
 #include "PCGSubsystem.h"
 #include "PCGVolumeFactory.h"
+#include "Data/PCGCollisionShapeData.h"
+#include "Data/PCGCollisionWrapperData.h"
+#include "Data/PCGLandscapeData.h"
+#include "Data/PCGPrimitiveData.h"
 #include "Data/PCGSpatialData.h"
 #include "Data/PCGSplineData.h"
+#include "Data/PCGVolumeData.h"
 #include "Grid/PCGPartitionActor.h"
 #include "WorldPartitionBuilder/PCGWorldPartitionBuilder.h"
 
@@ -20,9 +25,11 @@
 #include "PCGEditorSettings.h"
 #include "PCGEditorStyle.h"
 #include "PCGEditorUtils.h"
+#include "DataVisualizations/PCGLandscapeDataVisualization.h"
 #include "DataVisualizations/PCGParamDataVisualization.h"
 #include "DataVisualizations/PCGSpatialDataVisualization.h"
 #include "DataVisualizations/PCGSplineDataVisualization.h"
+#include "DataVisualizations/PCGVolumeDataVisualization.h"
 #include "Details/PCGAttributePropertySelectorDetails.h"
 #include "Details/PCGBlueprintSettingsDetails.h"
 #include "Details/PCGComponentDetails.h"
@@ -705,6 +712,11 @@ void FPCGEditorModule::RegisterPCGDataVisualizations()
 	DataVisRegistry.InternalRegistry.Add(UPCGParamData::StaticClass(), MakeUnique<const IPCGParamDataVisualization>());
 	DataVisRegistry.InternalRegistry.Add(UPCGSpatialData::StaticClass(), MakeUnique<const IPCGSpatialDataVisualization>());
 	DataVisRegistry.InternalRegistry.Add(UPCGSplineData::StaticClass(), MakeUnique<const IPCGSplineDataVisualization>());
+	DataVisRegistry.InternalRegistry.Add(UPCGVolumeData::StaticClass(), MakeUnique<const IPCGPrimitiveDataVisualization>());
+    DataVisRegistry.InternalRegistry.Add(UPCGPrimitiveData::StaticClass(), MakeUnique<const IPCGPrimitiveDataVisualization>());
+    DataVisRegistry.InternalRegistry.Add(UPCGCollisionShapeData::StaticClass(), MakeUnique<const IPCGPrimitiveDataVisualization>());
+    DataVisRegistry.InternalRegistry.Add(UPCGCollisionWrapperData::StaticClass(), MakeUnique<const IPCGPrimitiveDataVisualization>());
+    DataVisRegistry.InternalRegistry.Add(UPCGLandscapeData::StaticClass(), MakeUnique<const IPCGLandscapeDataVisualization>());
 }
 
 void FPCGEditorModule::UnregisterPCGDataVisualizations()
diff --git a/Engine/Plugins/PCG/Source/PCGEditor/Public/DataVisualizations/PCGLandscapeDataVisualization.h b/Engine/Plugins/PCG/Source/PCGEditor/Public/DataVisualizations/PCGLandscapeDataVisualization.h
new file mode 100644
--- /dev/null
+++ b/Engine/Plugins/PCG/Source/PCGEditor/Public/DataVisualizations/PCGLandscapeDataVisualization.h
@@ -0,0 +1,18 @@
+// Copyright Epic Games, Inc. All Rights Reserved.
+
+#pragma once
+
+#include "PCGSpatialDataVisualization.h"
+
+class UPCGData;
+class UPCGPointData;
+struct FPCGContext;
+
+/** Overrides collapse behavior to target a given number of points generated, to avoid too coarse debug (points way bigger than data) or enormous amount of points (points way smaller than data). */
+class PCGEDITOR_API IPCGLandscapeDataVisualization : public IPCGSpatialDataVisualization
+{
+public:
+	// ~Begin IPCGSpatialDataVisualization interface
+	virtual const UPCGPointData* CollapseToDebugPointData(FPCGContext* Context, const UPCGData* Data) const override;
+	// ~End IPCGSpatialDataVisualization interface
+};
diff --git a/Engine/Plugins/PCG/Source/PCGEditor/Public/DataVisualizations/PCGVolumeDataVisualization.h b/Engine/Plugins/PCG/Source/PCGEditor/Public/DataVisualizations/PCGVolumeDataVisualization.h
new file mode 100644
--- /dev/null
+++ b/Engine/Plugins/PCG/Source/PCGEditor/Public/DataVisualizations/PCGVolumeDataVisualization.h
@@ -0,0 +1,18 @@
+// Copyright Epic Games, Inc. All Rights Reserved.
+
+#pragma once
+
+#include "PCGSpatialDataVisualization.h"
+
+class UPCGData;
+class UPCGPointData;
+struct FPCGContext;
+
+/** Overrides collapse behavior to target a given number of points generated, to avoid too coarse debug (points way bigger than data) or enormous amount of points (points way smaller than data). */
+class PCGEDITOR_API IPCGPrimitiveDataVisualization : public IPCGSpatialDataVisualization
+{
+public:
+	// ~Begin IPCGSpatialDataVisualization interface
+	virtual const UPCGPointData* CollapseToDebugPointData(FPCGContext* Context, const UPCGData* Data) const override;
+	// ~End IPCGSpatialDataVisualization interface
+};
diff --git a/Engine/Plugins/PCG/Source/PCGEditor/Public/PCGEditorSettings.h b/Engine/Plugins/PCG/Source/PCGEditor/Public/PCGEditorSettings.h
--- a/Engine/Plugins/PCG/Source/PCGEditor/Public/PCGEditorSettings.h
+++ b/Engine/Plugins/PCG/Source/PCGEditor/Public/PCGEditorSettings.h
@@ -230,4 +230,8 @@ public:
 	/** Overrides the tooltip on the pause button. The default is empty, which will use the default tooltip then. */
 	UPROPERTY(EditAnywhere, config, Category = "Editor Performance", meta = (MultiLine = true, DisplayName = "Overrides default tooltip on the pause button", EditCondition = "bShowPauseButton", EditConditionHides))
 	FString OverridePausedButtonTooltip = FString();
+
+	/** Target number of points when trying to debug volume data or landscape data, to avoid exploding the number of points if the volume/landscape is big, and to have a good representation if the volume/landscape is small. */
+	UPROPERTY(EditAnywhere, config, Category = "Editor Performance")
+	int32 TargetNumPointsForDebug = 64000;
 };
