[PCG] GPU Custom HLSL node - un-blueprint a bunch of member properties as they undo the efforts to put the experimental code in private folders.

#rnx

diff --git a/Engine/Plugins/PCG/Source/PCG/Private/Compute/Elements/PCGCustomHLSL.h b/Engine/Plugins/PCG/Source/PCG/Private/Compute/Elements/PCGCustomHLSL.h
--- a/Engine/Plugins/PCG/Source/PCG/Private/Compute/Elements/PCGCustomHLSL.h
+++ b/Engine/Plugins/PCG/Source/PCG/Private/Compute/Elements/PCGCustomHLSL.h
@@ -135,29 +135,29 @@ protected:
 	virtual bool GetGeneratedFunctions(FString* OutFunctions, FText* OutErrorText = nullptr) const;
 
 protected:
-	UPROPERTY(BlueprintReadWrite, EditAnywhere, Category = "Settings")
+	UPROPERTY(EditAnywhere, Category = "Settings")
 	EPCGKernelType KernelType = EPCGKernelType::PointProcessor;
 
-	UPROPERTY(BlueprintReadWrite, EditAnywhere, DisplayName = "Number of Points", Category = "Settings", meta = (EditCondition = "KernelType == EPCGKernelType::PointGenerator", EditConditionHides))
+	UPROPERTY(EditAnywhere, DisplayName = "Number of Points", Category = "Settings", meta = (EditCondition = "KernelType == EPCGKernelType::PointGenerator", EditConditionHides))
 	int PointCount = 256;
 
-	UPROPERTY(BlueprintReadWrite, EditAnywhere, Category = "Settings|Thread Count", meta = (EditCondition = "KernelType == EPCGKernelType::Custom", EditConditionHides))
+	UPROPERTY(EditAnywhere, Category = "Settings|Thread Count", meta = (EditCondition = "KernelType == EPCGKernelType::Custom", EditConditionHides))
 	EPCGDispatchThreadCount DispatchThreadCount = EPCGDispatchThreadCount::FromFirstOutputPin;
 
-	UPROPERTY(BlueprintReadWrite, EditAnywhere, Category = "Settings|Thread Count", meta = (EditCondition = "KernelType == EPCGKernelType::Custom && DispatchThreadCount != EPCGDispatchThreadCount::Fixed", EditConditionHides))
+	UPROPERTY(EditAnywhere, Category = "Settings|Thread Count", meta = (EditCondition = "KernelType == EPCGKernelType::Custom && DispatchThreadCount != EPCGDispatchThreadCount::Fixed", EditConditionHides))
 	int ThreadCountMultiplier = 1;
 
-	UPROPERTY(BlueprintReadWrite, EditAnywhere, Category = "Settings|Thread Count", meta = (EditCondition = "KernelType == EPCGKernelType::Custom && DispatchThreadCount == EPCGDispatchThreadCount::Fixed", EditConditionHides))
+	UPROPERTY(EditAnywhere, Category = "Settings|Thread Count", meta = (EditCondition = "KernelType == EPCGKernelType::Custom && DispatchThreadCount == EPCGDispatchThreadCount::Fixed", EditConditionHides))
 	int FixedThreadCount = 1;
 
-	UPROPERTY(BlueprintReadWrite, EditAnywhere, DisplayName = "Input Pins", Category = "Settings|Thread Count", meta = (EditCondition = "KernelType == EPCGKernelType::Custom && DispatchThreadCount == EPCGDispatchThreadCount::FromProductOfInputPins", EditConditionHides, GetOptions = "GetInputPinNames"))
+	UPROPERTY(EditAnywhere, DisplayName = "Input Pins", Category = "Settings|Thread Count", meta = (EditCondition = "KernelType == EPCGKernelType::Custom && DispatchThreadCount == EPCGDispatchThreadCount::FromProductOfInputPins", EditConditionHides, GetOptions = "GetInputPinNames"))
 	TArray<FName> ThreadCountInputPinLabels;
 
 public:
-	UPROPERTY(BlueprintReadWrite, EditAnywhere, Category = "Settings")
+	UPROPERTY(EditAnywhere, Category = "Settings")
 	TArray<FPCGPinProperties> InputPins = Super::DefaultPointInputPinProperties();
 
-	UPROPERTY(BlueprintReadWrite, EditAnywhere, Category = "Settings")
+	UPROPERTY(EditAnywhere, Category = "Settings")
 	TArray<FPCGPinPropertiesGPU> OutputPins = { FPCGPinPropertiesGPU(PCGPinConstants::DefaultOutputLabel, EPCGDataType::Point) };
 
 #if WITH_EDITOR
@@ -167,11 +167,11 @@ public:
 
 protected:
 	/** Optional functions that can be called from the source. Intended to be edited using the Node Source Editor window. */
-	UPROPERTY(BlueprintReadWrite, EditAnywhere, Category = "Settings", AdvancedDisplay, meta = (MultiLine = true))
+	UPROPERTY(EditAnywhere, Category = "Settings", AdvancedDisplay, meta = (MultiLine = true))
 	FString ShaderFunctions = "/** CUSTOM SHADER FUNCTIONS **/\n";
 
 	/** Shader code that forms the body of the kernel. Intended to be edited using the Node Source Editor window. */
-	UPROPERTY(BlueprintReadWrite, EditAnywhere, Category = "Settings", AdvancedDisplay, meta = (MultiLine = true))
+	UPROPERTY(EditAnywhere, Category = "Settings", AdvancedDisplay, meta = (MultiLine = true))
 	FString ShaderSource;
 
 	/** Inputs data accessors that can be used from the shader code. Intended to be viewed using the Node Source Editor window. */
