[PCG] GPU - add mute-able error on Custom HLSL node for when output pin data may not be written.

Catches cases where there's no evidence of writing to a not-automatically-initialized output pin (secondary output pins of point processor and point generator, and all output pins of custom kernels).

#rnx

diff --git a/Engine/Plugins/PCG/Source/PCG/Private/Compute/Elements/PCGCustomHLSL.cpp b/Engine/Plugins/PCG/Source/PCG/Private/Compute/Elements/PCGCustomHLSL.cpp
--- a/Engine/Plugins/PCG/Source/PCG/Private/Compute/Elements/PCGCustomHLSL.cpp
+++ b/Engine/Plugins/PCG/Source/PCG/Private/Compute/Elements/PCGCustomHLSL.cpp
@@ -16,6 +16,7 @@
 #include "Elements/Metadata/PCGMetadataElementCommon.h"
 
 #include "Internationalization/Regex.h"
+#include "Containers/StaticArray.h"
 
 #if WITH_EDITOR
 #include "ScopedTransaction.h"
@@ -1481,6 +1482,15 @@ bool UPCGCustomHLSLSettings::IsKernelValid(FPCGContext* InContext, bool bQuiet)
 			}
 			return false;
 		}
+
+		if (!bMuteUnwrittenPinDataErrors && !AreAllOutputPinsWritten(ErrorTextPtr))
+		{
+			if (ErrorTextPtr)
+			{
+				PCG_KERNEL_VALIDATION_ERR(InContext, this, bQuiet, *ErrorTextPtr);
+			}
+			return false;
+		}
 	}
 
 	return true;
@@ -1652,6 +1662,60 @@ bool UPCGCustomHLSLSettings::AreKernelAttributesValid(FPCGContext* InContext, FT
 	return true;
 }
 
+bool UPCGCustomHLSLSettings::AreAllOutputPinsWritten(FText* OutWarningText) const
+{
+	TRACE_CPUPROFILER_EVENT_SCOPE(UPCGCustomHLSLSettings::AreAllOutputPinsWritten);
+
+	const UPCGNode* Node = Cast<UPCGNode>(GetOuter());
+	if (!Node)
+	{
+		return true;
+	}
+
+	const FString AllSource = ShaderFunctions + ShaderSource;
+
+	// Non-custom kernels initialize the first output pin data automatically.
+	const UPCGPin* AutomaticallyInitializedPin = (KernelType != EPCGKernelType::Custom) ? GetFirstOutputPin() : nullptr;
+
+	for (const UPCGPin* Pin : Node->GetOutputPins())
+	{
+		if (!Pin || Pin == AutomaticallyInitializedPin || Pin->Properties.Label == NAME_None)
+		{
+			continue;
+		}
+
+		const TStaticArray<FString, 3> WriteStrings = { 
+			Pin->Properties.Label.ToString() + TEXT("_Set"), 
+			Pin->Properties.Label.ToString() + TEXT("_CopyElementFrom_"),
+			Pin->Properties.Label.ToString() + TEXT("_Initialize") };
+
+		bool bFoundDataWrite = false;
+
+		for (const FString& WriteString : WriteStrings)
+		{
+			if (AllSource.Contains(WriteString))
+			{
+				bFoundDataWrite = true;
+				break;
+			}
+		}
+
+		if (!bFoundDataWrite)
+		{
+			if (OutWarningText)
+			{
+				*OutWarningText = FText::Format(
+					LOCTEXT("PinMayNotBeWritten", "Data on pin '{0}' may be uninitialized. Add code to write to this data, or mute this error in the node settings."),
+					{ FText::FromName(Pin->Properties.Label) });
+			}
+
+			return false;
+		}
+	}
+
+	return true;
+}
+
 bool UPCGCustomHLSLSettings::GetGeneratedFunctions(FString* OutFunctions, FText* OutErrorText) const
 {
 	TSet<TPair<FString, FString>> PinsRequiringElementCopy;
diff --git a/Engine/Plugins/PCG/Source/PCG/Private/Compute/Elements/PCGCustomHLSL.h b/Engine/Plugins/PCG/Source/PCG/Private/Compute/Elements/PCGCustomHLSL.h
--- a/Engine/Plugins/PCG/Source/PCG/Private/Compute/Elements/PCGCustomHLSL.h
+++ b/Engine/Plugins/PCG/Source/PCG/Private/Compute/Elements/PCGCustomHLSL.h
@@ -127,6 +127,7 @@ protected:
 	const UPCGPin* GetFirstPointOutputPin() const;
 	int GetProcessingElemCountForInputPin(const UPCGPin* InputPin, const UPCGDataBinding* Binding) const;
 	bool AreKernelAttributesValid(FPCGContext* InContext, FText* OutErrorText) const;
+	bool AreAllOutputPinsWritten(FText* OutErrorText) const;
 
 	/** Will the ThreadCountMultiplier value be applied when calculating the dispatch thread count. */
 	bool IsThreadCountMultiplierInUse() const { return KernelType == EPCGKernelType::Custom && DispatchThreadCount != EPCGDispatchThreadCount::Fixed; }
@@ -166,6 +167,10 @@ public:
 #endif
 
 protected:
+	/** Mute uninitialized data warnings. */
+	UPROPERTY(EditAnywhere, Category = "Settings")
+	bool bMuteUnwrittenPinDataErrors = false;
+
 	/** Optional functions that can be called from the source. Intended to be edited using the Node Source Editor window. */
 	UPROPERTY(EditAnywhere, Category = "Settings", AdvancedDisplay, meta = (MultiLine = true))
 	FString ShaderFunctions = "/** CUSTOM SHADER FUNCTIONS **/\n";
