[PCG] GPU - fixes for gaps in attribute access.

* Fix GetBool and SetBool
* Add missing GetTransform and SetTransform
* Add support for FNames: will treat Names as 2 uint2, similar to StringKeys, but with their internal FName EntryID and Number.
* Soft Class Path is a Soft Object Path so do exactly the same thing, add support for it too.

#rnx

diff --git a/Engine/Plugins/PCG/Shaders/Private/PCGDataCollectionDataInterface.ush b/Engine/Plugins/PCG/Shaders/Private/PCGDataCollectionDataInterface.ush
--- a/Engine/Plugins/PCG/Shaders/Private/PCGDataCollectionDataInterface.ush
+++ b/Engine/Plugins/PCG/Shaders/Private/PCGDataCollectionDataInterface.ush
@@ -163,7 +163,7 @@ void SetAsExecutedInternal_{DataInterfaceName}()
 bool GetBool_{DataInterfaceName}(uint DataIndex, uint ElementIndex, int AttributeId)
 {
 	const uint AttributeAddress = GetFirstElementAddressInternal_{DataInterfaceName}(DataIndex, AttributeId);
-	return asbool({DataInterfaceName}_DataCollectionBuffer.Load(AttributeAddress + ElementIndex * 1 * 4));
+	return {DataInterfaceName}_DataCollectionBuffer.Load(AttributeAddress + ElementIndex * 1 * 4) != 0;
 }
 
 int GetInt_{DataInterfaceName}(uint DataIndex, uint ElementIndex, int AttributeId)
@@ -214,18 +214,38 @@ float4 GetQuat_{DataInterfaceName}(uint DataIndex, uint ElementIndex, int Attrib
 	return asfloat({DataInterfaceName}_DataCollectionBuffer.Load4(AttributeAddress + ElementIndex * 4 * 4));
 }
 
+float4x4 GetTransform_{DataInterfaceName}(uint DataIndex, uint ElementIndex, int AttributeId)
+{
+	const uint AttributeAddress = GetFirstElementAddressInternal_{DataInterfaceName}(DataIndex, AttributeId);
+	const uint StartAddress = ElementIndex * 16 * 4 + AttributeAddress;
+
+	return float4x4(
+		asfloat({DataInterfaceName}_DataCollectionBuffer.Load4(StartAddress)),
+		asfloat({DataInterfaceName}_DataCollectionBuffer.Load4(StartAddress + 4 * 4)),
+		asfloat({DataInterfaceName}_DataCollectionBuffer.Load4(StartAddress + 8 * 4)),
+		asfloat({DataInterfaceName}_DataCollectionBuffer.Load4(StartAddress + 12 * 4))
+	);
+}
+
 int GetStringKey_{DataInterfaceName}(uint DataIndex, uint ElementIndex, int AttributeId)
 {
 	// String keys are represented as ints.
 	return GetInt_{DataInterfaceName}(DataIndex, ElementIndex, AttributeId);
 }
 
+uint2 GetName_{DataInterfaceName}(uint DataIndex, uint ElementIndex, int AttributeId)
+{
+	// Names are composed of 2 uints.
+	const uint AttributeAddress = GetFirstElementAddressInternal_{DataInterfaceName}(DataIndex, AttributeId);
+	return {DataInterfaceName}_DataCollectionBuffer.Load2(AttributeAddress + ElementIndex * 2 * 4);
+}
+
 // #################### ATTRIBUTE SETTERS ##########################
 
 void SetBool_{DataInterfaceName}(uint DataIndex, uint ElementIndex, int AttributeId, bool Value)
 {
 	const uint AttributeAddress = GetFirstElementAddressInternal_{DataInterfaceName}(DataIndex, AttributeId);
-	{DataInterfaceName}_DataCollectionBuffer.Store(AttributeAddress + ElementIndex * 1 * 4, asuint(Value));
+	{DataInterfaceName}_DataCollectionBuffer.Store(AttributeAddress + ElementIndex * 1 * 4, Value ? 1u : 0u);
 }
 
 void SetInt_{DataInterfaceName}(uint DataIndex, uint ElementIndex, int AttributeId, int Value)
@@ -276,12 +296,29 @@ void SetQuat_{DataInterfaceName}(uint DataIndex, uint ElementIndex, int Attribut
 	{DataInterfaceName}_DataCollectionBuffer.Store4(AttributeAddress + ElementIndex * 4 * 4, asuint(Value));
 }
 
+void SetTransform_{DataInterfaceName}(uint DataIndex, uint ElementIndex, int AttributeId, float4x4 Value)
+{
+	const uint AttributeAddress = GetFirstElementAddressInternal_{DataInterfaceName}(DataIndex, AttributeId);
+	const uint StartAddress = ElementIndex * 16 * 4 + AttributeAddress;
+	{DataInterfaceName}_DataCollectionBuffer.Store4(StartAddress, asuint(Value[0]));
+	{DataInterfaceName}_DataCollectionBuffer.Store4(StartAddress + 4 * 4, asuint(Value[1]));
+	{DataInterfaceName}_DataCollectionBuffer.Store4(StartAddress + 8 * 4, asuint(Value[2]));
+	{DataInterfaceName}_DataCollectionBuffer.Store4(StartAddress + 12 * 4, asuint(Value[3]));
+}
+
 void SetStringKey_{DataInterfaceName}(uint DataIndex, uint ElementIndex, int AttributeId, int Value)
 {
 	// String keys are represented as ints.
 	SetInt_{DataInterfaceName}(DataIndex, ElementIndex, AttributeId, Value);
 }
 
+void SetName_{DataInterfaceName}(uint DataIndex, uint ElementIndex, int AttributeId, uint2 Value)
+{
+	// Names are composed of 2 uints.
+	const uint AttributeAddress = GetFirstElementAddressInternal_{DataInterfaceName}(DataIndex, AttributeId);
+	{DataInterfaceName}_DataCollectionBuffer.Store2(AttributeAddress + ElementIndex * 2 * 4, Value);
+}
+
 // #################### POINT ATTRIBUTE GETTERS ##########################
 
 float3 GetPosition_{DataInterfaceName}(uint DataIndex, uint ElementIndex)
diff --git a/Engine/Plugins/PCG/Source/PCG/Private/Compute/DataInterfaces/PCGDataCollectionDataInterface.cpp b/Engine/Plugins/PCG/Source/PCG/Private/Compute/DataInterfaces/PCGDataCollectionDataInterface.cpp
--- a/Engine/Plugins/PCG/Source/PCG/Private/Compute/DataInterfaces/PCGDataCollectionDataInterface.cpp
+++ b/Engine/Plugins/PCG/Source/PCG/Private/Compute/DataInterfaces/PCGDataCollectionDataInterface.cpp
@@ -152,6 +152,13 @@ void UPCGDataCollectionDataInterface::GetSupportedInputs(TArray<FShaderFunctionD
 			.AddParam(EShaderFundamentalType::Uint) // DataIndex
 			.AddParam(EShaderFundamentalType::Uint) // ElementIndex
 			.AddParam(EShaderFundamentalType::Uint); // AttributeId
+
+		OutFunctions.AddDefaulted_GetRef()
+			.SetName(TEXT("GetName"))
+			.AddReturnType(EShaderFundamentalType::Uint, 2) // Name represented by 2 uints
+			.AddParam(EShaderFundamentalType::Uint) // DataIndex
+			.AddParam(EShaderFundamentalType::Uint) // ElementIndex
+			.AddParam(EShaderFundamentalType::Uint); // AttributeId
 	}
 
 	// Attribute Setters
@@ -225,6 +232,13 @@ void UPCGDataCollectionDataInterface::GetSupportedInputs(TArray<FShaderFunctionD
 			.AddParam(EShaderFundamentalType::Uint) // ElementIndex
 			.AddParam(EShaderFundamentalType::Uint) // AttributeId
 			.AddParam(EShaderFundamentalType::Int); // String key represented by int
+		
+		OutFunctions.AddDefaulted_GetRef()
+			.SetName(TEXT("SetName"))
+			.AddParam(EShaderFundamentalType::Uint) // DataIndex
+			.AddParam(EShaderFundamentalType::Uint) // ElementIndex
+			.AddParam(EShaderFundamentalType::Uint) // AttributeId
+			.AddParam(EShaderFundamentalType::Uint, 2); // Name represented by 2 uints
 	}
 
 	// Point Attribute Getters
diff --git a/Engine/Plugins/PCG/Source/PCG/Private/Compute/Elements/PCGCustomHLSL.cpp b/Engine/Plugins/PCG/Source/PCG/Private/Compute/Elements/PCGCustomHLSL.cpp
--- a/Engine/Plugins/PCG/Source/PCG/Private/Compute/Elements/PCGCustomHLSL.cpp
+++ b/Engine/Plugins/PCG/Source/PCG/Private/Compute/Elements/PCGCustomHLSL.cpp
@@ -755,7 +755,7 @@ void UPCGCustomHLSLSettings::UpdateInputDeclarations()
 			"uint {0}_GetNumElements();\n"
 			"uint {0}_GetNumElements(uint DataIndex);\n"
 			"\n"
-			"// Valid types: bool, int, float, float2, float3, float4, Rotator (float3), Quat (float4), Transform (float4x4), StringKey (int)\n"
+			"// Valid types: bool, int, float, float2, float3, float4, Rotator (float3), Quat (float4), Transform (float4x4), StringKey (int), Name (uint2)\n"
 			"\n"
 			"<type> {0}_Get<type>(uint DataIndex, uint ElementIndex, int AttributeId);\n"
 			"<type> {0}_Get<type>(uint DataIndex, uint ElementIndex, 'AttributeName');\n"),
@@ -891,7 +891,7 @@ void UPCGCustomHLSLSettings::UpdateOutputDeclarations()
 		}
 
 		OutputDeclarations += FString::Format(TEXT(
-			"// Valid types: bool, int, float, float2, float3, float4, Rotator (float3), Quat (float4), Transform (float4x4), StringKey (int)\n"
+			"// Valid types: bool, int, float, float2, float3, float4, Rotator (float3), Quat (float4), Transform (float4x4), StringKey (int), Name (uint2)\n"
 			"\n"
 			"void {0}_Set<type>(uint DataIndex, uint ElementIndex, int AttributeId, <type> Value);\n"
 			"void {0}_Set<type>(uint DataIndex, uint ElementIndex, 'AttributeName', <type> Value);\n"),
diff --git a/Engine/Plugins/PCG/Source/PCG/Private/Compute/PCGDataBinding.cpp b/Engine/Plugins/PCG/Source/PCG/Private/Compute/PCGDataBinding.cpp
--- a/Engine/Plugins/PCG/Source/PCG/Private/Compute/PCGDataBinding.cpp
+++ b/Engine/Plugins/PCG/Source/PCG/Private/Compute/PCGDataBinding.cpp
@@ -125,6 +125,10 @@ void UPCGDataBinding::BuildStringTable()
 					{
 						StringTable.AddUnique(static_cast<const FPCGMetadataAttribute<FString>*>(AttributeBase)->GetValue(InValueKey));
 					}
+					else if (AttributeBase->GetTypeId() == PCG::Private::MetadataTypes<FSoftClassPath>::Id)
+                    {
+                    	StringTable.AddUnique(static_cast<const FPCGMetadataAttribute<FSoftClassPath>*>(AttributeBase)->GetValue(InValueKey).ToString());
+                    }
 					else
 					{
 						checkNoEntry();
diff --git a/Engine/Plugins/PCG/Source/PCG/Private/Compute/PCGDataForGPU.cpp b/Engine/Plugins/PCG/Source/PCG/Private/Compute/PCGDataForGPU.cpp
--- a/Engine/Plugins/PCG/Source/PCG/Private/Compute/PCGDataForGPU.cpp
+++ b/Engine/Plugins/PCG/Source/PCG/Private/Compute/PCGDataForGPU.cpp
@@ -62,9 +62,11 @@ namespace PCGDataForGPUHelpers
 		case EPCGMetadataTypes::Transform:
 			return EPCGKernelAttributeType::Transform;
 		case EPCGMetadataTypes::SoftObjectPath:
+		case EPCGMetadataTypes::SoftClassPath:
 		case EPCGMetadataTypes::String:
 			return EPCGKernelAttributeType::StringKey;
-		case EPCGMetadataTypes::SoftClassPath:
+		case EPCGMetadataTypes::Name:
+			return EPCGKernelAttributeType::Name;
 		default:
 			return EPCGKernelAttributeType::Invalid;
 		}
@@ -80,6 +82,7 @@ namespace PCGDataForGPUHelpers
 		case EPCGKernelAttributeType::StringKey:
 			return 4;
 		case EPCGKernelAttributeType::Float2:
+		case EPCGKernelAttributeType::Name:
 			return 8;
 		case EPCGKernelAttributeType::Float3:
 		case EPCGKernelAttributeType::Rotator:
@@ -238,6 +241,28 @@ namespace PCGDataForGPUHelpers
 			OutPackedDataCollection[InOutAddressUints++] = Value;
 			break;
 		}
+		case PCG::Private::MetadataTypes<FSoftClassPath>::Id:
+		{
+			// SCP path string stored as an integer for reading/writing in kernel, and accompanying string table in data description.
+			const FPCGMetadataAttribute<FSoftClassPath>* Attribute = static_cast<const FPCGMetadataAttribute<FSoftClassPath>*>(InAttributeBase);
+			const int32 Value = InStringTable.IndexOfByKey(Attribute->GetValue(ValueKey).ToString());
+			check(StrideBytes == 4);
+			OutPackedDataCollection[InOutAddressUints++] = Value;
+			break;
+		}
+		case PCG::Private::MetadataTypes<FName>::Id:
+        {
+        	// FName are defined with 2 int32, so packed them both.
+        	const FPCGMetadataAttribute<FName>* Attribute = static_cast<const FPCGMetadataAttribute<FName>*>(InAttributeBase);
+        	FName Value = Attribute->GetValue(ValueKey);
+        	check(StrideBytes == 8);
+			const FNameEntryId NameEntry = Value.GetDisplayIndex();
+			const uint32 Number = Value.GetNumber();
+
+        	OutPackedDataCollection[InOutAddressUints++] = NameEntry.ToUnstableInt();
+        	OutPackedDataCollection[InOutAddressUints++] = Number;
+        	break;
+        }
 		default:
 			return false;
 		}
@@ -289,6 +314,10 @@ namespace PCGDataForGPUHelpers
 		{
 			return Metadata->FindOrCreateAttribute<FString>(AttributeDesc.Name);
 		}
+		case EPCGKernelAttributeType::Name:
+		{
+			return Metadata->FindOrCreateAttribute<FName>(AttributeDesc.Name);
+		}
 		default:
 			return nullptr;
 		}
@@ -420,6 +449,20 @@ namespace PCGDataForGPUHelpers
 			}
 			break;
 		}
+		case EPCGKernelAttributeType::Name:
+		{
+			FPCGMetadataAttribute<FName>* Attribute = static_cast<FPCGMetadataAttribute<FName>*>(AttributeBase);
+
+			const FNameEntryId NameEntryId = FNameEntryId::FromUnstableInt(DataAsInt[ElementIndex]);
+			const uint32 Number = DataAsInt[ElementIndex + 1];
+			// Names need some validation
+			const FName Value = FName::CreateFromDisplayId(NameEntryId, Number);
+			if (Value.IsValid())
+			{
+				Attribute->SetValue(EntryKey, Value);
+			}
+			break;
+		}
 		default:
 			return false;
 		}
@@ -637,6 +680,10 @@ void FPCGDataDesc::InitializeAttributeDescs(const UPCGData* InData, const TMap<F
 						{
 							StringTableIndex = InStringTable.IndexOfByKey(static_cast<const FPCGMetadataAttribute<FString>*>(AttributeBase)->GetValue(ValueKey));
 						}
+						else if (AttributeBase->GetTypeId() == PCG::Private::MetadataTypes<FSoftClassPath>::Id)
+                        {
+                        	StringTableIndex = InStringTable.IndexOfByKey(static_cast<const FPCGMetadataAttribute<FSoftClassPath>*>(AttributeBase)->GetValue(ValueKey).ToString());
+                        }
 						else
 						{
 							// Should not get here if attribute type is string key.
@@ -665,6 +712,10 @@ void FPCGDataDesc::InitializeAttributeDescs(const UPCGData* InData, const TMap<F
 						{
 							StringTableIndex = InStringTable.IndexOfByKey(static_cast<const FPCGMetadataAttribute<FString>*>(AttributeBase)->GetValue(MetadataKey));
 						}
+						else if (AttributeBase->GetTypeId() == PCG::Private::MetadataTypes<FSoftClassPath>::Id)
+                        {
+                        	StringTableIndex = InStringTable.IndexOfByKey(static_cast<const FPCGMetadataAttribute<FSoftClassPath>*>(AttributeBase)->GetValue(MetadataKey).ToString());
+                        }
 						else
 						{
 							// Should not get here if attribute type is string key.
diff --git a/Engine/Plugins/PCG/Source/PCG/Private/Compute/PCGDataForGPU.h b/Engine/Plugins/PCG/Source/PCG/Private/Compute/PCGDataForGPU.h
--- a/Engine/Plugins/PCG/Source/PCG/Private/Compute/PCGDataForGPU.h
+++ b/Engine/Plugins/PCG/Source/PCG/Private/Compute/PCGDataForGPU.h
@@ -34,6 +34,7 @@ enum class EPCGKernelAttributeType : uint8
 	Quat,
 	Transform,
 	StringKey,
+	Name,
 
 	Invalid = std::numeric_limits<uint8>::max()
 };
