Fixed collapse points issues due to bounds consideration when it wasn't needed.

#rnx

diff --git a/Engine/Plugins/PCG/Source/PCG/Private/Elements/PCGCollapsePoints.cpp b/Engine/Plugins/PCG/Source/PCG/Private/Elements/PCGCollapsePoints.cpp
--- a/Engine/Plugins/PCG/Source/PCG/Private/Elements/PCGCollapsePoints.cpp
+++ b/Engine/Plugins/PCG/Source/PCG/Private/Elements/PCGCollapsePoints.cpp
@@ -35,7 +35,7 @@ namespace PCGCollapsePoints
 				const double WeightSum = PrimaryWeight + SecondaryWeight;
 
 				double Alpha = FMath::IsNearlyZero(WeightSum) ? 0.5 : (SecondaryWeight / WeightSum);
-				const FVector DeltaPosition = OutState.Points[SecondaryPointIndex].Transform.GetLocation() - OutState.Points[PrimaryPointIndex].Transform.GetLocation();
+				const FVector DeltaPosition = Settings.GetPositionFunc(OutState.Points[SecondaryPointIndex]) - Settings.GetPositionFunc(OutState.Points[PrimaryPointIndex]);
 
 				OutState.Points[PrimaryPointIndex].Transform.AddToTranslation(Alpha * DeltaPosition);
 				OutState.Weights[PrimaryPointIndex] = WeightSum;
@@ -43,7 +43,7 @@ namespace PCGCollapsePoints
 			}
 		}
 
-		void RebuildOctree(FPCGContext* InContext, FCollapsePointsState& OutState)
+		void RebuildOctree(FPCGContext* InContext, const FCollapsePointsSettings& Settings, FCollapsePointsState& OutState)
 		{
 			TRACE_CPUPROFILER_EVENT_SCOPE(PCGCollapsePointsElement::Algo::RebuildOctree);
 			check(OutState.SourceData && OutState.Merged.Num() == OutState.Points.Num());
@@ -53,7 +53,7 @@ namespace PCGCollapsePoints
 			{
 				if (OutState.Merged[PointIndex] == INDEX_NONE)
 				{
-					NewOctree.AddElement(FPCGPointRef(OutState.Points[PointIndex]));
+					NewOctree.AddElement(Settings.GetPointRefFunc(OutState.Points[PointIndex]));
 				}
 			}
 
@@ -95,8 +95,8 @@ namespace PCGCollapsePoints
 				bool bHasColocatedPoint = false;
 
 				const double Extents = UE_DOUBLE_SQRT_2 * Settings.DistanceThreshold;
-				FBoxCenterAndExtent SearchBounds(OutState.Points[PointIndex].Transform.GetLocation(), FVector(Extents, Extents, Extents));
-				OutState.Octree.FindElementsWithBoundsTest(SearchBounds, [PointIndex, &MinSqrDistance, &ClosestUnvisitedIndex, &bHasColocatedPoint, &OutState](const FPCGPointRef& PointRef)
+				FBoxCenterAndExtent SearchBounds = Settings.GetSearchBoundsFunc(OutState.Points[PointIndex], Extents);
+				OutState.Octree.FindElementsWithBoundsTest(SearchBounds, [&Settings, PointIndex, &MinSqrDistance, &ClosestUnvisitedIndex, &bHasColocatedPoint, &OutState](const FPCGPointRef& PointRef)
 				{
 					int NeighborIndex = PointRef.Point - OutState.Points.GetData();
 
@@ -105,7 +105,7 @@ namespace PCGCollapsePoints
 						return;
 					}
 
-					double SqrDistance = (OutState.Points[PointIndex].Transform.GetLocation() - OutState.Points[NeighborIndex].Transform.GetLocation()).SquaredLength();
+					double SqrDistance = (Settings.GetPositionFunc(OutState.Points[PointIndex]) - Settings.GetPositionFunc(OutState.Points[NeighborIndex])).SquaredLength();
 
 					if (FMath::IsNearlyZero(SqrDistance))
 					{
@@ -151,6 +151,39 @@ namespace PCGCollapsePoints
 			return true;
 		}
 	} // namespace Modes
+
+	namespace ComparisonModes
+	{
+		FVector GetPosition(const FPCGPoint& InPoint)
+		{
+			return InPoint.Transform.GetLocation();
+		}
+
+		FVector GetCenter(const FPCGPoint& InPoint)
+		{
+			return InPoint.Transform.TransformPosition(InPoint.GetLocalCenter());
+		}
+
+		FPCGPointRef GetPositionPointRef(const FPCGPoint& InPoint)
+		{
+			return FPCGPointRef(InPoint, FBox(FVector::ZeroVector, FVector::ZeroVector));
+		}
+
+		FPCGPointRef GetCenterPointRef(const FPCGPoint& InPoint)
+		{
+			return FPCGPointRef(InPoint, FBox(InPoint.GetLocalCenter(), InPoint.GetLocalCenter()));
+		}
+
+		FBoxCenterAndExtent GetPositionSearchBounds(const FPCGPoint& InPoint, const double& Extents)
+		{
+			return FBoxCenterAndExtent(InPoint.Transform.GetLocation(), FVector(Extents, Extents, Extents));
+		}
+
+		FBoxCenterAndExtent GetCenterSearchBounds(const FPCGPoint& InPoint, const double& Extents)
+		{
+			return FBoxCenterAndExtent(InPoint.Transform.TransformPosition(InPoint.GetLocalCenter()), FVector(Extents, Extents, Extents));
+		}
+	}
 } // namespace PCGCollapsePoints
 
 UPCGCollapsePointsSettings::UPCGCollapsePointsSettings()
@@ -203,6 +236,9 @@ bool FPCGCollapsePointsElement::PrepareDataInternal(FPCGContext* InContext) cons
 	Context->InitializePerExecutionState([Settings](const ContextType*, ExecStateType& OutState)
 	{
 		PCGCollapsePoints::FCollapsePointsSettings::PairSelectionFuncType PairSelectionFunc = nullptr;
+		PCGCollapsePoints::FCollapsePointsSettings::GetPositionFuncType GetPositionFunc = nullptr;
+		PCGCollapsePoints::FCollapsePointsSettings::GetPointRefFuncType GetPointRefFunc = nullptr;
+		PCGCollapsePoints::FCollapsePointsSettings::GetSearchBoundsFuncType GetSearchBoundsFunc = nullptr;
 
 		if (Settings->Mode == EPCGCollapseMode::PairwiseClosest)
 		{
@@ -217,19 +253,39 @@ bool FPCGCollapsePointsElement::PrepareDataInternal(FPCGContext* InContext) cons
 			return EPCGTimeSliceInitResult::AbortExecution;
 		}
 
+		if (Settings->ComparisonMode == EPCGCollapseComparisonMode::Position)
+		{
+			GetPositionFunc = PCGCollapsePoints::ComparisonModes::GetPosition;
+			GetPointRefFunc = PCGCollapsePoints::ComparisonModes::GetPositionPointRef;
+			GetSearchBoundsFunc = PCGCollapsePoints::ComparisonModes::GetPositionSearchBounds;
+		}
+		else if (Settings->ComparisonMode == EPCGCollapseComparisonMode::Center)
+		{
+			GetPositionFunc = PCGCollapsePoints::ComparisonModes::GetCenter;
+			GetPointRefFunc = PCGCollapsePoints::ComparisonModes::GetCenterPointRef;
+			GetSearchBoundsFunc = PCGCollapsePoints::ComparisonModes::GetCenterSearchBounds;
+		}
+		else
+		{
+			return EPCGTimeSliceInitResult::AbortExecution;
+		}
+
 		check(PairSelectionFunc);
 		OutState = PCGCollapsePoints::FCollapsePointsSettings
 		{
 			.Settings = Settings,
 			.PairSelectionFunc = PairSelectionFunc,
 			.MergeSelectionFunc = PCGCollapsePoints::Algo::MergePairs,
+			.GetPositionFunc = GetPositionFunc,
+			.GetPointRefFunc = GetPointRefFunc,
+			.GetSearchBoundsFunc = GetSearchBoundsFunc,
 			.DistanceThreshold = Settings->DistanceThreshold
 		};
 
 		return EPCGTimeSliceInitResult::Success;
 	});
 
-	Context->InitializePerIterationStates(Inputs.Num(), [&Inputs, Settings, Context](IterStateType& OutState, const ExecStateType&, const uint32 IterationIndex)
+	Context->InitializePerIterationStates(Inputs.Num(), [&Inputs, Settings, Context](IterStateType& OutState, const ExecStateType& ExecState, const uint32 IterationIndex)
 	{
 		const UPCGPointData* PointData = Cast<UPCGPointData>(Inputs[IterationIndex].Data);
 
@@ -371,7 +427,7 @@ bool FPCGCollapsePointsElement::PrepareDataInternal(FPCGContext* InContext) cons
 		}
 
 		// Rebuild octree
-		PCGCollapsePoints::Algo::RebuildOctree(Context, OutState);
+		PCGCollapsePoints::Algo::RebuildOctree(Context, ExecState, OutState);
 
 		Context->OutputData.TaggedData.Emplace_GetRef().Data = OutPointData;
 		return EPCGTimeSliceInitResult::Success;
@@ -417,7 +473,7 @@ bool FPCGCollapsePointsElement::ExecuteInternal(FPCGContext* InContext) const
 			// Merge based on selection
 			CollapseSettings.MergeSelectionFunc(InContext, CollapseSettings, CollapseState);
 			// Rebuild octree for next iteration
-			PCGCollapsePoints::Algo::RebuildOctree(InContext, CollapseState);
+			PCGCollapsePoints::Algo::RebuildOctree(InContext, CollapseSettings, CollapseState);
 
 			return false;
 		}
diff --git a/Engine/Plugins/PCG/Source/PCG/Public/Elements/PCGCollapsePoints.h b/Engine/Plugins/PCG/Source/PCG/Public/Elements/PCGCollapsePoints.h
--- a/Engine/Plugins/PCG/Source/PCG/Public/Elements/PCGCollapsePoints.h
+++ b/Engine/Plugins/PCG/Source/PCG/Public/Elements/PCGCollapsePoints.h
@@ -20,6 +20,13 @@ enum class EPCGCollapseMode : uint8
 	AbsoluteClosest UMETA(Hidden, Tooltip="At every iteration, will select the closest pair to merge. Performance warning.")
 };
 
+UENUM(BlueprintType)
+enum class EPCGCollapseComparisonMode : uint8
+{
+	Position UMETA(Tooltip="Uses point position only for distance testing, regardless of bounds."),
+	Center UMETA(Tooltip="Uses point centers (e.g. center of bounds) for distance testing, and ignore bounds otherwise."),
+};
+
 UENUM(BlueprintType)
 enum class EPCGCollapseVisitOrder : uint8
 {
@@ -61,6 +68,9 @@ public:
 	UPROPERTY(EditAnywhere, BlueprintReadWrite, Category = Settings, meta = (PCG_Overridable))
 	EPCGCollapseMode Mode = EPCGCollapseMode::PairwiseClosest;
 
+	UPROPERTY(EditAnywhere, BlueprintReadWrite, Category = Settings, meta = (PCG_Overridable))
+	EPCGCollapseComparisonMode ComparisonMode = EPCGCollapseComparisonMode::Position;
+
 	/** Determines order in which we will collapse points pair-wise. */
 	UPROPERTY(EditAnywhere, BlueprintReadWrite, Category = Settings, meta = (EditCondition="Mode == EPCGCollapseMode::PairwiseClosest", EditConditionHides, PCG_Overridable))
 	EPCGCollapseVisitOrder VisitOrder = EPCGCollapseVisitOrder::Ordered;
@@ -120,6 +130,15 @@ namespace PCGCollapsePoints
 		using MergeSelectionFuncType = void(*)(FPCGContext* /*InContext*/, const FCollapsePointsSettings& /*Settings*/, FCollapsePointsState& /*OutState*/);
 		MergeSelectionFuncType MergeSelectionFunc = nullptr;
 
+		using GetPositionFuncType = FVector(*)(const FPCGPoint&);
+		GetPositionFuncType GetPositionFunc = nullptr;
+
+		using GetPointRefFuncType = FPCGPointRef(*)(const FPCGPoint&);
+		GetPointRefFuncType GetPointRefFunc = nullptr;
+
+		using GetSearchBoundsFuncType = FBoxCenterAndExtent(*)(const FPCGPoint&, const double&);
+		GetSearchBoundsFuncType GetSearchBoundsFunc = nullptr;
+
 		double DistanceThreshold = 0.0;
 	};
 } // namespace PCGCollapsePoints
