[PCG] Remove delegate registration to PCGComponents and instead have the Component call into the Subsystem directly
- Will fix the crash and reduce the number of registrations

diff --git a/Engine/Plugins/PCG/Source/PCG/Private/PCGActorAndComponentMapping.cpp b/Engine/Plugins/PCG/Source/PCG/Private/PCGActorAndComponentMapping.cpp
--- a/Engine/Plugins/PCG/Source/PCG/Private/PCGActorAndComponentMapping.cpp
+++ b/Engine/Plugins/PCG/Source/PCG/Private/PCGActorAndComponentMapping.cpp
@@ -1005,15 +1005,6 @@ void FPCGActorAndComponentMapping::RegisterTracking(UPCGComponent* InComponent)
 	AlwaysTrackedKeysToComponentsMap.FindOrAdd(FPCGSelectionKey(EPCGActorFilter::Self)).Add(InComponent);
 
 	UpdateTracking(InComponent, /*bInShouldDirtyActors=*/ false);
-
-	// Add tracking for when the graph was generated/cleaned, only once
-	if (!InComponent->OnPCGGraphGeneratedDelegate.IsBoundToObject(this))
-	{
-		InComponent->OnPCGGraphGeneratedDelegate.AddRaw(this, &FPCGActorAndComponentMapping::OnPCGGraphGeneratedOrCleaned);
-		InComponent->OnPCGGraphCleanedDelegate.AddRaw(this, &FPCGActorAndComponentMapping::OnPCGGraphGeneratedOrCleaned);
-		InComponent->OnPCGGraphStartGeneratingDelegate.AddRaw(this, &FPCGActorAndComponentMapping::OnPCGGraphStartsGenerating);
-		InComponent->OnPCGGraphCancelledDelegate.AddRaw(this, &FPCGActorAndComponentMapping::OnPCGGraphCancelled);
-	}
 }
 
 void FPCGActorAndComponentMapping::UpdateTracking(UPCGComponent* InComponent, bool bInShouldDirtyActors, const TArray<FPCGSelectionKey>* ChangedKeys)
@@ -1109,24 +1100,6 @@ void FPCGActorAndComponentMapping::RemapTracking(const UPCGComponent* InOldCompo
 	{
 		It.Value.RemoveSwap(InOldComponent);
 	}
-
-	// Old component will probably die, but we'll force removing the delegates even if it is const.
-	if (UPCGComponent* MutableOldComponent = const_cast<UPCGComponent*>(InOldComponent))
-	{
-		MutableOldComponent->OnPCGGraphGeneratedDelegate.RemoveAll(this);
-		MutableOldComponent->OnPCGGraphCleanedDelegate.RemoveAll(this);
-		MutableOldComponent->OnPCGGraphStartGeneratingDelegate.RemoveAll(this);
-		MutableOldComponent->OnPCGGraphCancelledDelegate.RemoveAll(this);
-	}
-
-	// And just making sure we are not registering multiple times
-	if (!InNewComponent->OnPCGGraphGeneratedDelegate.IsBoundToObject(this))
-	{
-		InNewComponent->OnPCGGraphGeneratedDelegate.AddRaw(this, &FPCGActorAndComponentMapping::OnPCGGraphGeneratedOrCleaned);
-		InNewComponent->OnPCGGraphCleanedDelegate.AddRaw(this, &FPCGActorAndComponentMapping::OnPCGGraphGeneratedOrCleaned);
-		InNewComponent->OnPCGGraphStartGeneratingDelegate.AddRaw(this, &FPCGActorAndComponentMapping::OnPCGGraphStartsGenerating);
-		InNewComponent->OnPCGGraphCancelledDelegate.AddRaw(this, &FPCGActorAndComponentMapping::OnPCGGraphCancelled);
-	}
 }
 
 void FPCGActorAndComponentMapping::UnregisterTracking(UPCGComponent* InComponent, const TSet<FPCGSelectionKey>* OptionalKeysToUntrack)
@@ -1190,11 +1163,6 @@ void FPCGActorAndComponentMapping::UnregisterTracking(UPCGComponent* InComponent
 	}
 
 	UnregisterTracking(InComponent, nullptr);
-
-	InComponent->OnPCGGraphGeneratedDelegate.RemoveAll(this);
-	InComponent->OnPCGGraphCleanedDelegate.RemoveAll(this);
-	InComponent->OnPCGGraphStartGeneratingDelegate.RemoveAll(this);
-	InComponent->OnPCGGraphCancelledDelegate.RemoveAll(this);
 }
 
 bool FPCGActorAndComponentMapping::IsKeyTracked(const FPCGSelectionKey& InKey) const
@@ -1295,23 +1263,6 @@ void FPCGActorAndComponentMapping::UnregisterDelegates()
 			LevelInstanceSubsystem->OnLevelInstancesUpdated().RemoveAll(this);
 		}
 	}
-
-	TSet<UPCGComponent*> StillRegisteredComponents = GetAllRegisteredComponents();
-	// FIXME: When the changes to unregister will go in, we can enable this ensure. At the moment, we didn't unregister when the world is shutdown
-	// because the subsystem is dead before we destroy the component. So it is expected to have still all of those components registered at that point.
-	//if (!ensureMsgf(StillRegisteredComponents.IsEmpty(), TEXT("Found components not unregistered during PCG subsystem shutdown.")))
-	{
-		for (UPCGComponent* Component : StillRegisteredComponents)
-		{
-			if (IsValid(Component))
-			{
-				Component->OnPCGGraphGeneratedDelegate.RemoveAll(this);
-				Component->OnPCGGraphCleanedDelegate.RemoveAll(this);
-				Component->OnPCGGraphStartGeneratingDelegate.RemoveAll(this);
-				Component->OnPCGGraphCancelledDelegate.RemoveAll(this);
-			}
-		}
-	}
 }
 
 bool FPCGActorAndComponentMapping::ShouldDelayActor(AActor* InActor) const
diff --git a/Engine/Plugins/PCG/Source/PCG/Private/PCGComponent.cpp b/Engine/Plugins/PCG/Source/PCG/Private/PCGComponent.cpp
--- a/Engine/Plugins/PCG/Source/PCG/Private/PCGComponent.cpp
+++ b/Engine/Plugins/PCG/Source/PCG/Private/PCGComponent.cpp
@@ -468,6 +468,10 @@ FPCGTaskId UPCGComponent::GenerateInternal(bool bForce, EPCGHiGenGrid Grid, EPCG
 	if (CurrentGenerationTask != InvalidPCGTaskId)
 	{
 #if WITH_EDITOR
+		// Notify Subsystem first
+		GetSubsystem()->OnPCGGraphStartGenerating(this);
+		
+		// Notify Delegate next
 		OnPCGGraphStartGeneratingDelegate.Broadcast(this);
 #endif // WITH_EDITOR
 
@@ -640,8 +644,13 @@ void UPCGComponent::PostProcessGraph(const FBox& InNewBounds, bool bInGenerated,
 		bForceGenerateOnBPAddedToWorld = false;
 
 		bDirtyGenerated = false;
-		OnPCGGraphGeneratedDelegate.Broadcast(this);
 
+		// Notify Subsystem first
+		GetSubsystem()->OnPCGGraphGenerated(this);
+
+		// Notify Delegate next
+		OnPCGGraphGeneratedDelegate.Broadcast(this);
+		
 		UpdateDynamicTracking();
 #endif // WITH_EDITOR
 
@@ -738,6 +747,13 @@ void UPCGComponent::PostCleanupGraph()
 	GeneratedGraphOutput.Reset();
 
 #if WITH_EDITOR
+	// Notify Subsystem first
+	if (UPCGSubsystem* Subsystem = GetSubsystem())
+	{
+		Subsystem->OnPCGGraphCleaned(this);
+	}
+
+	// Notify Delegate next
 	OnPCGGraphCleanedDelegate.Broadcast(this);
 	bDirtyGenerated = false;
 
@@ -780,9 +796,13 @@ void UPCGComponent::OnProcessGraphAborted(bool bQuiet, bool bCleanupUnusedResour
 
 	StopGenerationInProgress();
 
+	// Notify Subsystem first
+	UPCGSubsystem* Subsystem = GetSubsystem();
+	Subsystem->OnPCGGraphCancelled(this);
+
+	// Notify Delegate next
 	OnPCGGraphCancelledDelegate.Broadcast(this);
 
-	UPCGSubsystem* Subsystem = GetSubsystem();
 	Subsystem->OnPCGComponentGenerationDone.Broadcast(Subsystem, this, EPCGGenerationStatus::Aborted);
 #endif
 
diff --git a/Engine/Plugins/PCG/Source/PCG/Private/PCGSubsystem.cpp b/Engine/Plugins/PCG/Source/PCG/Private/PCGSubsystem.cpp
--- a/Engine/Plugins/PCG/Source/PCG/Private/PCGSubsystem.cpp
+++ b/Engine/Plugins/PCG/Source/PCG/Private/PCGSubsystem.cpp
@@ -1853,6 +1853,26 @@ TSet<TObjectPtr<APCGPartitionActor>> UPCGSubsystem::GetPCGComponentPartitionActo
 	return ActorAndComponentMapping.GetPCGComponentPartitionActorMappings(InComponent); 
 }
 
+void UPCGSubsystem::OnPCGGraphCancelled(UPCGComponent* InComponent)
+{
+	ActorAndComponentMapping.OnPCGGraphCancelled(InComponent);
+}
+
+void UPCGSubsystem::OnPCGGraphStartGenerating(UPCGComponent* InComponent)
+{
+	ActorAndComponentMapping.OnPCGGraphStartsGenerating(InComponent);
+}
+
+void UPCGSubsystem::OnPCGGraphGenerated(UPCGComponent* InComponent)
+{
+	ActorAndComponentMapping.OnPCGGraphGeneratedOrCleaned(InComponent);
+}
+
+void UPCGSubsystem::OnPCGGraphCleaned(UPCGComponent* InComponent)
+{
+	ActorAndComponentMapping.OnPCGGraphGeneratedOrCleaned(InComponent);
+}
+
 #endif // WITH_EDITOR
 
 UPCGData* UPCGSubsystem::GetPCGData(FPCGTaskId InGraphExecutionTaskId)
diff --git a/Engine/Plugins/PCG/Source/PCG/Public/PCGSubsystem.h b/Engine/Plugins/PCG/Source/PCG/Public/PCGSubsystem.h
--- a/Engine/Plugins/PCG/Source/PCG/Public/PCGSubsystem.h
+++ b/Engine/Plugins/PCG/Source/PCG/Public/PCGSubsystem.h
@@ -357,6 +357,11 @@ private:
 		Generate
 	};
 
+	void OnPCGGraphCancelled(UPCGComponent* InComponent);
+	void OnPCGGraphStartGenerating(UPCGComponent* InComponent);
+	void OnPCGGraphGenerated(UPCGComponent* InComponent);
+	void OnPCGGraphCleaned(UPCGComponent* InComponent);
+
 	void CreatePartitionActorsWithinBounds(UPCGComponent* InComponent, const FBox& InBounds, const PCGHiGenGrid::FSizeArray& InGridSizes);
 	void UpdateMappingPCGComponentPartitionActor(UPCGComponent* InComponent);
 	TSet<TObjectPtr<APCGPartitionActor>> GetPCGComponentPartitionActorMappings(UPCGComponent* InComponent) const;
