Properly reset BP runaway counts in blueprint element.

#rnx

diff --git a/Engine/Plugins/PCG/Source/PCG/Private/Elements/PCGExecuteBlueprint.cpp b/Engine/Plugins/PCG/Source/PCG/Private/Elements/PCGExecuteBlueprint.cpp
--- a/Engine/Plugins/PCG/Source/PCG/Private/Elements/PCGExecuteBlueprint.cpp
+++ b/Engine/Plugins/PCG/Source/PCG/Private/Elements/PCGExecuteBlueprint.cpp
@@ -25,6 +25,11 @@
 
 #define LOCTEXT_NAMESPACE "PCGBlueprintElement"
 
+namespace PCGBlueprintConstants
+{
+	constexpr int32 RunawayResetFrequency = 1024;
+}
+
 #if WITH_EDITOR
 namespace PCGBlueprintHelper
 {
@@ -933,10 +938,20 @@ void UPCGBlueprintElement::PointLoop(FPCGContext& InContext, const UPCGPointData
 	const TArray<FPCGPoint>& InPoints = InData->GetPoints();
 	TArray<FPCGPoint>& OutPoints = OutData->GetMutablePoints();
 
+	bool bPreviousBPStateValue = true;
+	std::swap(bPreviousBPStateValue, InContext.AsyncState.bIsCallingBlueprint);
+
 	FPCGAsync::AsyncPointProcessing(&InContext, InPoints.Num(), OutPoints, [this, &InContext, InData, OutData, &InPoints](int32 Index, FPCGPoint& OutPoint)
 	{
+		if (Index % PCGBlueprintConstants::RunawayResetFrequency == 0)
+		{
+			GInitRunaway(); // Reset periodically the iteration count, because we know we're in a fixed size loop.
+		}
+
 		return PointLoopBody(InContext, InData, InPoints[Index], OutPoint, OutData->Metadata, Index);
 	});
+
+	std::swap(bPreviousBPStateValue, InContext.AsyncState.bIsCallingBlueprint);
 }
 
 void UPCGBlueprintElement::VariableLoop(FPCGContext& InContext, const UPCGPointData* InData, UPCGPointData*& OutData, UPCGPointData* OptionalOutData) const
@@ -960,10 +975,20 @@ void UPCGBlueprintElement::VariableLoop(FPCGContext& InContext, const UPCGPointD
 	const TArray<FPCGPoint>& InPoints = InData->GetPoints();
 	TArray<FPCGPoint>& OutPoints = OutData->GetMutablePoints();
 
+	bool bPreviousBPStateValue = true;
+	std::swap(bPreviousBPStateValue, InContext.AsyncState.bIsCallingBlueprint);
+
 	FPCGAsync::AsyncMultiPointProcessing(&InContext, InPoints.Num(), OutPoints, [this, &InContext, InData, OutData, &InPoints](int32 Index)
 	{
+		if (Index % PCGBlueprintConstants::RunawayResetFrequency == 0)
+		{
+			GInitRunaway(); // Reset periodically the iteration count, because we know we're in a fixed size loop.
+		}
+
 		return VariableLoopBody(InContext, InData, InPoints[Index], OutData->Metadata, Index);
 	});
+
+	std::swap(bPreviousBPStateValue, InContext.AsyncState.bIsCallingBlueprint);
 }
 
 void UPCGBlueprintElement::NestedLoop(FPCGContext& InContext, const UPCGPointData* InOuterData, const UPCGPointData* InInnerData, UPCGPointData*& OutData, UPCGPointData* OptionalOutData) const
@@ -989,12 +1014,22 @@ void UPCGBlueprintElement::NestedLoop(FPCGContext& InContext, const UPCGPointDat
 	const TArray<FPCGPoint>& InInnerPoints = InInnerData->GetPoints();
 	TArray<FPCGPoint>& OutPoints = OutData->GetMutablePoints();
 
+	bool bPreviousBPStateValue = true;
+	std::swap(bPreviousBPStateValue, InContext.AsyncState.bIsCallingBlueprint);
+
 	FPCGAsync::AsyncPointProcessing(&InContext, InOuterPoints.Num() * InInnerPoints.Num(), OutPoints, [this, &InContext, InOuterData, InInnerData, OutData, &InOuterPoints, &InInnerPoints](int32 Index, FPCGPoint& OutPoint)
 	{
+		if (Index % PCGBlueprintConstants::RunawayResetFrequency == 0)
+		{
+			GInitRunaway(); // Reset periodically the iteration count, because we know we're in a fixed size loop.
+		}
+
 		const int32 OuterIndex = Index / InInnerPoints.Num();
 		const int32 InnerIndex = Index % InInnerPoints.Num();
 		return NestedLoopBody(InContext, InOuterData, InInnerData, InOuterPoints[OuterIndex], InInnerPoints[InnerIndex], OutPoint, OutData->Metadata, OuterIndex, InnerIndex);
 	});
+
+	std::swap(bPreviousBPStateValue, InContext.AsyncState.bIsCallingBlueprint);
 }
 
 void UPCGBlueprintElement::IterationLoop(FPCGContext& InContext, int64 NumIterations, UPCGPointData*& OutData, const UPCGSpatialData* InA, const UPCGSpatialData* InB, UPCGPointData* OptionalOutData) const
@@ -1022,10 +1057,20 @@ void UPCGBlueprintElement::IterationLoop(FPCGContext& InContext, int64 NumIterat
 
 	TArray<FPCGPoint>& OutPoints = OutData->GetMutablePoints();
 
+	bool bPreviousBPStateValue = true;
+	std::swap(bPreviousBPStateValue, InContext.AsyncState.bIsCallingBlueprint);
+
 	FPCGAsync::AsyncPointProcessing(&InContext, NumIterations, OutPoints, [this, &InContext, InA, InB, OutData](int32 Index, FPCGPoint& OutPoint)
 	{
+		if (Index % PCGBlueprintConstants::RunawayResetFrequency == 0)
+		{
+			GInitRunaway(); // Reset periodically the iteration count, because we know we're in a fixed size loop.
+		}
+
 		return IterationLoopBody(InContext, Index, InA, InB, OutPoint, OutData->Metadata);
 	});
+
+	std::swap(bPreviousBPStateValue, InContext.AsyncState.bIsCallingBlueprint);
 }
 
 void FPCGBlueprintExecutionContext::AddExtraStructReferencedObjects(FReferenceCollector& Collector)
diff --git a/Engine/Plugins/PCG/Source/PCG/Private/Helpers/PCGAsync.cpp b/Engine/Plugins/PCG/Source/PCG/Private/Helpers/PCGAsync.cpp
--- a/Engine/Plugins/PCG/Source/PCG/Private/Helpers/PCGAsync.cpp
+++ b/Engine/Plugins/PCG/Source/PCG/Private/Helpers/PCGAsync.cpp
@@ -70,6 +70,7 @@ namespace FPCGAsync
 		// Setup [current, last, nb points] data per dispatch
 		TArray<UE::Tasks::TTask<int32>> AsyncTasks;
 		AsyncTasks.Reserve(NumFutures);
+		const bool bInitBPContext = (AsyncState && AsyncState->bIsCallingBlueprint);
 
 		// Launch the async tasks
 		for (int32 TaskIndex = 0; TaskIndex < NumFutures; ++TaskIndex)
@@ -77,8 +78,13 @@ namespace FPCGAsync
 			const int32 StartIndex = TaskIndex * IterationsPerTask;
 			const int32 EndIndex = StartIndex + IterationsPerTask;
 
-			AsyncTasks.Emplace(UE::Tasks::Launch(UE_SOURCE_LOCATION, [&IterationInnerLoop, StartIndex, EndIndex]()
+			AsyncTasks.Emplace(UE::Tasks::Launch(UE_SOURCE_LOCATION, [&IterationInnerLoop, StartIndex, EndIndex, bInitBPContext]()
 			{
+				if (bInitBPContext)
+				{
+					GInitRunaway(); // Reset counter as threads can run multiple workloads and might not be properly reset.
+				}
+
 				return IterationInnerLoop(StartIndex, EndIndex);
 			}));
 		}
@@ -171,6 +177,7 @@ namespace FPCGAsync
 
 		TArray<UE::Tasks::TTask<TArray<FPCGPoint>>> AsyncTasks;
 		AsyncTasks.Reserve(NumFutures);
+		const bool bInitBPContext = (AsyncState && AsyncState->bIsCallingBlueprint);
 
 		// Launch the async tasks
 		for (int32 TaskIndex = 0; TaskIndex < NumFutures; ++TaskIndex)
@@ -178,8 +185,13 @@ namespace FPCGAsync
 			const int32 StartIndex = TaskIndex * IterationsPerTask;
 			const int32 EndIndex = StartIndex + IterationsPerTask;
 
-			AsyncTasks.Emplace(UE::Tasks::Launch(UE_SOURCE_LOCATION, [&IterationInnerLoop, StartIndex, EndIndex]() -> TArray<FPCGPoint>
+			AsyncTasks.Emplace(UE::Tasks::Launch(UE_SOURCE_LOCATION, [&IterationInnerLoop, StartIndex, EndIndex, bInitBPContext]() -> TArray<FPCGPoint>
 			{
+				if (bInitBPContext)
+				{
+					GInitRunaway(); // Reset counter as threads can run multiple workloads and might not be properly reset.
+				}
+
 				return IterationInnerLoop(StartIndex, EndIndex);
 			}));
 		}
@@ -243,6 +255,7 @@ namespace FPCGAsync
 		// Setup [current, last, nb points] data per dispatch
 		TArray<UE::Tasks::TTask<TPair<int32, int32>>> AsyncTasks;
 		AsyncTasks.Reserve(NumFutures);
+		const bool bInitBPContext = (AsyncState && AsyncState->bIsCallingBlueprint);
 
 		// Launch the async tasks
 		for (int32 TaskIndex = 0; TaskIndex < NumFutures; ++TaskIndex)
@@ -250,8 +263,13 @@ namespace FPCGAsync
 			const int32 StartIndex = TaskIndex * IterationsPerTask;
 			const int32 EndIndex = StartIndex + IterationsPerTask;
 
-			AsyncTasks.Emplace(UE::Tasks::Launch(UE_SOURCE_LOCATION, [&IterationInnerLoop, StartIndex, EndIndex]() -> TPair<int32, int32>
+			AsyncTasks.Emplace(UE::Tasks::Launch(UE_SOURCE_LOCATION, [&IterationInnerLoop, StartIndex, EndIndex, bInitBPContext]() -> TPair<int32, int32>
 			{
+				if (bInitBPContext)
+				{
+					GInitRunaway(); // Reset counter as threads can run multiple workloads and might not be properly reset.
+				}
+
 				return IterationInnerLoop(StartIndex, EndIndex);
 			}));
 		}
diff --git a/Engine/Plugins/PCG/Source/PCG/Public/Helpers/PCGAsyncState.h b/Engine/Plugins/PCG/Source/PCG/Public/Helpers/PCGAsyncState.h
--- a/Engine/Plugins/PCG/Source/PCG/Public/Helpers/PCGAsyncState.h
+++ b/Engine/Plugins/PCG/Source/PCG/Public/Helpers/PCGAsyncState.h
@@ -53,6 +53,9 @@ struct FPCGAsyncState
 	/** True if currently inside a PCGAsync scope - will prevent further async processing. */
 	bool bIsRunningAsyncCall = false;
 
+	/** True if the async process will call BP functions, which will require a BP context reset. */
+	bool bIsCallingBlueprint = false;
+
 	/** True if we can run out of tick */
 	bool bIsRunningOutOfTick = false;
 
