Fixed a significant amount of additional titles that weren't properly localized.

#rnx

diff --git a/Engine/Plugins/PCG/Source/PCG/Private/Compute/Elements/PCGCustomHLSL.cpp b/Engine/Plugins/PCG/Source/PCG/Private/Compute/Elements/PCGCustomHLSL.cpp
--- a/Engine/Plugins/PCG/Source/PCG/Private/Compute/Elements/PCGCustomHLSL.cpp
+++ b/Engine/Plugins/PCG/Source/PCG/Private/Compute/Elements/PCGCustomHLSL.cpp
@@ -522,7 +522,7 @@ FString UPCGCustomHLSLSettings::GetAdditionalTitleInformation() const
 {
 	if (const UEnum* EnumPtr = StaticEnum<EPCGKernelType>())
 	{
-		return EnumPtr->GetNameStringByValue(static_cast<int>(KernelType));
+		return EnumPtr->GetDisplayNameTextByValue(static_cast<int64>(KernelType)).ToString();
 	}
 
 	return FString();
diff --git a/Engine/Plugins/PCG/Source/PCG/Private/Elements/Metadata/PCGMetadataBitwiseOpElement.cpp b/Engine/Plugins/PCG/Source/PCG/Private/Elements/Metadata/PCGMetadataBitwiseOpElement.cpp
--- a/Engine/Plugins/PCG/Source/PCG/Private/Elements/Metadata/PCGMetadataBitwiseOpElement.cpp
+++ b/Engine/Plugins/PCG/Source/PCG/Private/Elements/Metadata/PCGMetadataBitwiseOpElement.cpp
@@ -90,7 +90,7 @@ FString UPCGMetadataBitwiseSettings::GetAdditionalTitleInformation() const
 {
 	if (const UEnum* EnumPtr = StaticEnum<EPCGMetadataBitwiseOperation>())
 	{
-		return FString("Bitwise: ") + EnumPtr->GetNameStringByValue(static_cast<int>(Operation));
+		return FText::Format(NSLOCTEXT("PCGMetadataBitwiseSettings", "BitwiseOperation", "Bitwise: {0}"), EnumPtr->GetDisplayNameTextByValue(static_cast<int64>(Operation))).ToString();
 	}
 	else
 	{
diff --git a/Engine/Plugins/PCG/Source/PCG/Private/Elements/Metadata/PCGMetadataBooleanOpElement.cpp b/Engine/Plugins/PCG/Source/PCG/Private/Elements/Metadata/PCGMetadataBooleanOpElement.cpp
--- a/Engine/Plugins/PCG/Source/PCG/Private/Elements/Metadata/PCGMetadataBooleanOpElement.cpp
+++ b/Engine/Plugins/PCG/Source/PCG/Private/Elements/Metadata/PCGMetadataBooleanOpElement.cpp
@@ -95,7 +95,7 @@ FString UPCGMetadataBooleanSettings::GetAdditionalTitleInformation() const
 {
 	if (const UEnum* EnumPtr = StaticEnum<EPCGMetadataBooleanOperation>())
 	{
-		return FString("Boolean: ") + EnumPtr->GetNameStringByValue(static_cast<int>(Operation));
+		return FText::Format(NSLOCTEXT("PCGMetadataBooleanSettings", "BooleanOperation", "Boolean: {0}"), EnumPtr->GetDisplayNameTextByValue(static_cast<int64>(Operation))).ToString();
 	}
 	else
 	{
diff --git a/Engine/Plugins/PCG/Source/PCG/Private/Elements/Metadata/PCGMetadataCompareOpElement.cpp b/Engine/Plugins/PCG/Source/PCG/Private/Elements/Metadata/PCGMetadataCompareOpElement.cpp
--- a/Engine/Plugins/PCG/Source/PCG/Private/Elements/Metadata/PCGMetadataCompareOpElement.cpp
+++ b/Engine/Plugins/PCG/Source/PCG/Private/Elements/Metadata/PCGMetadataCompareOpElement.cpp
@@ -121,7 +121,7 @@ FString UPCGMetadataCompareSettings::GetAdditionalTitleInformation() const
 {
 	if (const UEnum* EnumPtr = StaticEnum<EPCGMetadataCompareOperation>())
 	{
-		return FString("Compare: ") + EnumPtr->GetNameStringByValue(static_cast<int>(Operation));
+		return FText::Format(NSLOCTEXT("PCGMetadataCompareSettings", "CompareWithOperation", "Compare: {0}"), EnumPtr->GetDisplayNameTextByValue(static_cast<int64>(Operation))).ToString();
 	}
 	else
 	{
diff --git a/Engine/Plugins/PCG/Source/PCG/Private/Elements/Metadata/PCGMetadataMathsOpElement.cpp b/Engine/Plugins/PCG/Source/PCG/Private/Elements/Metadata/PCGMetadataMathsOpElement.cpp
--- a/Engine/Plugins/PCG/Source/PCG/Private/Elements/Metadata/PCGMetadataMathsOpElement.cpp
+++ b/Engine/Plugins/PCG/Source/PCG/Private/Elements/Metadata/PCGMetadataMathsOpElement.cpp
@@ -281,7 +281,7 @@ FString UPCGMetadataMathsSettings::GetAdditionalTitleInformation() const
 {
 	if (const UEnum* EnumPtr = StaticEnum<EPCGMetadataMathsOperation>())
 	{
-		return EnumPtr->GetNameStringByValue(static_cast<int>(Operation));
+		return EnumPtr->GetDisplayNameTextByValue(static_cast<int64>(Operation)).ToString();
 	}
 	else
 	{
diff --git a/Engine/Plugins/PCG/Source/PCG/Private/Elements/Metadata/PCGMetadataRotatorOpElement.cpp b/Engine/Plugins/PCG/Source/PCG/Private/Elements/Metadata/PCGMetadataRotatorOpElement.cpp
--- a/Engine/Plugins/PCG/Source/PCG/Private/Elements/Metadata/PCGMetadataRotatorOpElement.cpp
+++ b/Engine/Plugins/PCG/Source/PCG/Private/Elements/Metadata/PCGMetadataRotatorOpElement.cpp
@@ -198,7 +198,7 @@ FString UPCGMetadataRotatorSettings::GetAdditionalTitleInformation() const
 {
 	if (const UEnum* EnumPtr = StaticEnum<EPCGMetadataRotatorOperation>())
 	{
-		return FString("Rotator: ") + EnumPtr->GetNameStringByValue(static_cast<int>(Operation));
+		return FText::Format(NSLOCTEXT("PCGMetadataRotatorSettings", "RotatorOperation", "Rotator: {0}"), EnumPtr->GetDisplayNameTextByValue(static_cast<int64>(Operation))).ToString();
 	}
 	else
 	{
diff --git a/Engine/Plugins/PCG/Source/PCG/Private/Elements/Metadata/PCGMetadataStringOpElement.cpp b/Engine/Plugins/PCG/Source/PCG/Private/Elements/Metadata/PCGMetadataStringOpElement.cpp
--- a/Engine/Plugins/PCG/Source/PCG/Private/Elements/Metadata/PCGMetadataStringOpElement.cpp
+++ b/Engine/Plugins/PCG/Source/PCG/Private/Elements/Metadata/PCGMetadataStringOpElement.cpp
@@ -73,7 +73,7 @@ FString UPCGMetadataStringOpSettings::GetAdditionalTitleInformation() const
 {
 	if (const UEnum* EnumPtr = StaticEnum<EPCGMetadataStringOperation>())
 	{
-		return FString("String: ") + EnumPtr->GetNameStringByValue(static_cast<int>(Operation));
+		return FText::Format(NSLOCTEXT("PCGMetadataStringOpSettings", "StringOperation", "String: {0}"), EnumPtr->GetDisplayNameTextByValue(static_cast<int64>(Operation))).ToString();
 	}
 	else
 	{
diff --git a/Engine/Plugins/PCG/Source/PCG/Private/Elements/Metadata/PCGMetadataTransformOpElement.cpp b/Engine/Plugins/PCG/Source/PCG/Private/Elements/Metadata/PCGMetadataTransformOpElement.cpp
--- a/Engine/Plugins/PCG/Source/PCG/Private/Elements/Metadata/PCGMetadataTransformOpElement.cpp
+++ b/Engine/Plugins/PCG/Source/PCG/Private/Elements/Metadata/PCGMetadataTransformOpElement.cpp
@@ -161,7 +161,7 @@ FString UPCGMetadataTransformSettings::GetAdditionalTitleInformation() const
 {
 	if (const UEnum* EnumPtr = StaticEnum<EPCGMetadataTransformOperation>())
 	{
-		return FString("Transform: ") + EnumPtr->GetNameStringByValue(static_cast<int>(Operation));
+		return FText::Format(NSLOCTEXT("PCGMetadataTransformSettings", "TransformOperation", "Transform: {0}"), EnumPtr->GetDisplayNameTextByValue(static_cast<int64>(Operation))).ToString();
 	}
 	else
 	{
diff --git a/Engine/Plugins/PCG/Source/PCG/Private/Elements/Metadata/PCGMetadataTrigOpElement.cpp b/Engine/Plugins/PCG/Source/PCG/Private/Elements/Metadata/PCGMetadataTrigOpElement.cpp
--- a/Engine/Plugins/PCG/Source/PCG/Private/Elements/Metadata/PCGMetadataTrigOpElement.cpp
+++ b/Engine/Plugins/PCG/Source/PCG/Private/Elements/Metadata/PCGMetadataTrigOpElement.cpp
@@ -107,7 +107,7 @@ FString UPCGMetadataTrigSettings::GetAdditionalTitleInformation() const
 {
 	if (const UEnum* EnumPtr = StaticEnum<EPCGMetadataTrigOperation>())
 	{
-		return FString("Trig: ") + EnumPtr->GetNameStringByValue(static_cast<int>(Operation));
+		return FText::Format(NSLOCTEXT("PCGMetadataTrigSettings", "TrigOperation", "Trig: {0}"), EnumPtr->GetDisplayNameTextByValue(static_cast<int64>(Operation))).ToString();
 	}
 	else
 	{
diff --git a/Engine/Plugins/PCG/Source/PCG/Private/Elements/Metadata/PCGMetadataVectorOpElement.cpp b/Engine/Plugins/PCG/Source/PCG/Private/Elements/Metadata/PCGMetadataVectorOpElement.cpp
--- a/Engine/Plugins/PCG/Source/PCG/Private/Elements/Metadata/PCGMetadataVectorOpElement.cpp
+++ b/Engine/Plugins/PCG/Source/PCG/Private/Elements/Metadata/PCGMetadataVectorOpElement.cpp
@@ -323,7 +323,7 @@ FString UPCGMetadataVectorSettings::GetAdditionalTitleInformation() const
 {
 	if (const UEnum* EnumPtr = StaticEnum<EPCGMetadataVectorOperation>())
 	{
-		return FString("Vector: ") + EnumPtr->GetNameStringByValue(static_cast<int>(Operation));
+		return FText::Format(NSLOCTEXT("PCGMetadataVectorSettings", "VectorOperation", "Vector: {0}"), EnumPtr->GetDisplayNameTextByValue(static_cast<int64>(Operation))).ToString();
 	}
 	else
 	{
diff --git a/Engine/Plugins/PCG/Source/PCG/Private/Elements/PCGActorSelector.cpp b/Engine/Plugins/PCG/Source/PCG/Private/Elements/PCGActorSelector.cpp
--- a/Engine/Plugins/PCG/Source/PCG/Private/Elements/PCGActorSelector.cpp
+++ b/Engine/Plugins/PCG/Source/PCG/Private/Elements/PCGActorSelector.cpp
@@ -507,11 +507,11 @@ FText FPCGActorSelectorSettings::GetTaskNameSuffix() const
 	{
 		if (ActorSelection == EPCGActorSelection::ByClass)
 		{
-			return FText::Format(FText::FromString(TEXT("Class: {0}")), (ActorSelectionClass.Get() ? ActorSelectionClass->GetDisplayNameText() : FText::FromName(NAME_None)));
+			return FText::Format(NSLOCTEXT("PCGActorSelectorSettings", "ClassLabel", "Class: {0}"), (ActorSelectionClass.Get() ? ActorSelectionClass->GetDisplayNameText() : FText::FromName(NAME_None)));
 		}
 		else if (ActorSelection == EPCGActorSelection::ByTag)
 		{
-			return FText::Format(FText::FromString(TEXT("Tag: {0}")), FText::FromName(ActorSelectionTag));
+			return FText::Format(NSLOCTEXT("PCGActorSelectorSettings", "TagLabel", "Tag: {0}"), FText::FromName(ActorSelectionTag));
 		}
 	}
 	else if(const UEnum* EnumPtr = StaticEnum<EPCGActorFilter>())
diff --git a/Engine/Plugins/PCG/Source/PCG/Private/Elements/PCGAttributeReduceElement.cpp b/Engine/Plugins/PCG/Source/PCG/Private/Elements/PCGAttributeReduceElement.cpp
--- a/Engine/Plugins/PCG/Source/PCG/Private/Elements/PCGAttributeReduceElement.cpp
+++ b/Engine/Plugins/PCG/Source/PCG/Private/Elements/PCGAttributeReduceElement.cpp
@@ -214,7 +214,7 @@ FString UPCGAttributeReduceSettings::GetAdditionalTitleInformation() const
 {
 	if (const UEnum* EnumPtr = StaticEnum<EPCGAttributeReduceOperation>())
 	{
-		const FString OperationName = EnumPtr->GetNameStringByValue(static_cast<int>(Operation));
+		const FText OperationName = EnumPtr->GetDisplayNameTextByValue(static_cast<int64>(Operation));
 
 		FName InputAttributeName = InputSource.GetName();
 		if (InputAttributeName == NAME_None)
@@ -224,11 +224,11 @@ FString UPCGAttributeReduceSettings::GetAdditionalTitleInformation() const
 
 		if (InputAttributeName != OutputAttributeName && OutputAttributeName != NAME_None)
 		{
-			return FString::Printf(TEXT("Reduce %s to %s: %s"), *InputAttributeName.ToString(), *OutputAttributeName.ToString(), *OperationName);
+			return FText::Format(LOCTEXT("ReduceInputToOutputWithOperation", "Reduce {0} to {1}: {2}"), FText::FromName(InputAttributeName), FText::FromName(OutputAttributeName), OperationName).ToString();
 		}
 		else
 		{
-			return FString::Printf(TEXT("Reduce %s: %s"), *InputAttributeName.ToString(), *OperationName);
+			return FText::Format(LOCTEXT("ReduceInplaceWithOperation", "Reduce {0}: {1}"), FText::FromName(InputAttributeName), OperationName).ToString();
 		}
 	}
 	else
diff --git a/Engine/Plugins/PCG/Source/PCG/Private/Elements/PCGAttributeSelectElement.cpp b/Engine/Plugins/PCG/Source/PCG/Private/Elements/PCGAttributeSelectElement.cpp
--- a/Engine/Plugins/PCG/Source/PCG/Private/Elements/PCGAttributeSelectElement.cpp
+++ b/Engine/Plugins/PCG/Source/PCG/Private/Elements/PCGAttributeSelectElement.cpp
@@ -213,7 +213,7 @@ FString UPCGAttributeSelectSettings::GetAdditionalTitleInformation() const
 	{
 		if (const UEnum* EnumAxisPtr = StaticEnum<EPCGAttributeSelectAxis>())
 		{
-			const FString OperationName = EnumOpPtr->GetNameStringByValue(static_cast<int>(Operation));
+			const FText OperationName = EnumOpPtr->GetDisplayNameTextByValue(static_cast<int64>(Operation));
 			FString AxisName;
 			if (Axis == EPCGAttributeSelectAxis::CustomAxis)
 			{
@@ -221,7 +221,7 @@ FString UPCGAttributeSelectSettings::GetAdditionalTitleInformation() const
 			}
 			else
 			{
-				AxisName = EnumAxisPtr->GetNameStringByValue(static_cast<int>(Axis));
+				AxisName = EnumAxisPtr->GetDisplayNameTextByValue(static_cast<int64>(Axis)).ToString();
 			}
 
 			FName InputAttributeName = InputSource.GetName();
@@ -232,11 +232,11 @@ FString UPCGAttributeSelectSettings::GetAdditionalTitleInformation() const
 
 			if (InputAttributeName != OutputAttributeName && OutputAttributeName != NAME_None)
 			{
-				return FString::Printf(TEXT("Select %s to %s: %s on %s"), *InputAttributeName.ToString(), *OutputAttributeName.ToString(), *OperationName, *AxisName);
+				return FText::Format(LOCTEXT("SelectInputToOutputWithOperationAndAxis", "Select {0} to {1}: {2} on {3}"), FText::FromName(InputAttributeName), FText::FromName(OutputAttributeName), OperationName, FText::FromString(AxisName)).ToString();
 			}
 			else
 			{
-				return FString::Printf(TEXT("Select %s: %s on %s"), *InputAttributeName.ToString(), *OperationName, *AxisName);
+				return FText::Format(LOCTEXT("SelectWithOperationAndAxis", "Select {0}: {1} on {2}"), FText::FromName(InputAttributeName), OperationName, FText::FromString(AxisName)).ToString();
 			}
 		}
 	}
diff --git a/Engine/Plugins/PCG/Source/PCG/Private/Elements/PCGDifferenceElement.cpp b/Engine/Plugins/PCG/Source/PCG/Private/Elements/PCGDifferenceElement.cpp
--- a/Engine/Plugins/PCG/Source/PCG/Private/Elements/PCGDifferenceElement.cpp
+++ b/Engine/Plugins/PCG/Source/PCG/Private/Elements/PCGDifferenceElement.cpp
@@ -96,7 +96,7 @@ FString UPCGDifferenceSettings::GetAdditionalTitleInformation() const
 	{
 		if (const UEnum* EnumPtr = StaticEnum<EPCGDifferenceDensityFunction>())
 		{
-			return EnumPtr->GetNameStringByValue(static_cast<int>(DensityFunction));
+			return EnumPtr->GetDisplayNameTextByValue(static_cast<int64>(DensityFunction)).ToString();
 		}
 		else
 		{
diff --git a/Engine/Plugins/PCG/Source/PCG/Private/Elements/PCGFilterByIndex.cpp b/Engine/Plugins/PCG/Source/PCG/Private/Elements/PCGFilterByIndex.cpp
--- a/Engine/Plugins/PCG/Source/PCG/Private/Elements/PCGFilterByIndex.cpp
+++ b/Engine/Plugins/PCG/Source/PCG/Private/Elements/PCGFilterByIndex.cpp
@@ -36,7 +36,7 @@ FText UPCGFilterByIndexSettings::GetNodeTooltipText() const
 
 FString UPCGFilterByIndexSettings::GetAdditionalTitleInformation() const
 {
-	FString NodeName = TEXT("Indices: ");
+	FString NodeName = LOCTEXT("FilterIndices", "Indices: ").ToString();
 	NodeName += SelectedIndices.Len() <= PCGFilterByIndexConstants::IndexExpressionTruncation
 		? SelectedIndices
 		: SelectedIndices.Left(PCGFilterByIndexConstants::IndexExpressionTruncation - 3) + TEXT("...");
diff --git a/Engine/Plugins/PCG/Source/PCG/Private/Elements/PCGFilterByTag.cpp b/Engine/Plugins/PCG/Source/PCG/Private/Elements/PCGFilterByTag.cpp
--- a/Engine/Plugins/PCG/Source/PCG/Private/Elements/PCGFilterByTag.cpp
+++ b/Engine/Plugins/PCG/Source/PCG/Private/Elements/PCGFilterByTag.cpp
@@ -48,7 +48,7 @@ FString UPCGFilterByTagSettings::GetAdditionalTitleInformation() const
 		: PCGHelpers::GetStringArrayFromCommaSeparatedList(SelectedTags);
 	PRAGMA_ENABLE_DEPRECATION_WARNINGS
 
-	const FString Prefix = (Operation == EPCGFilterByTagOperation::KeepTagged ? TEXT("Tag (Keep):") : TEXT("Tag (Remove):"));
+	const FString Prefix = (Operation == EPCGFilterByTagOperation::KeepTagged ? LOCTEXT("FilteredTagKeep", "Tag (Keep):") : LOCTEXT("FilteredTagRemove", "Tag (Remove):")).ToString();
 
 	if (Tags.IsEmpty())
 	{
@@ -60,7 +60,8 @@ FString UPCGFilterByTagSettings::GetAdditionalTitleInformation() const
 	}
 	else
 	{
-		return FString::Printf(TEXT("%s (multiple)"), *Prefix);
+		const FText MultipleEntriesIndicator = LOCTEXT("FilterMultipleTags", "(multiple)");
+		return FString::Printf(TEXT("%s %s"), *Prefix, *MultipleEntriesIndicator.ToString());
 	}
 }
 
diff --git a/Engine/Plugins/PCG/Source/PCG/Private/Elements/PCGFilterElementsByIndex.cpp b/Engine/Plugins/PCG/Source/PCG/Private/Elements/PCGFilterElementsByIndex.cpp
--- a/Engine/Plugins/PCG/Source/PCG/Private/Elements/PCGFilterElementsByIndex.cpp
+++ b/Engine/Plugins/PCG/Source/PCG/Private/Elements/PCGFilterElementsByIndex.cpp
@@ -100,10 +100,10 @@ FString UPCGFilterElementsByIndexSettings::GetAdditionalTitleInformation() const
 		return Super::GetAdditionalTitleInformation();
 	}
 
-	FString NodeName = TEXT("Indices: ");
+	FString NodeName = LOCTEXT("ListOfIndices", "Indices: ").ToString();
 	if (SelectedIndices == FString(TEXT(":")))
 	{
-		NodeName += TEXT("All");
+		NodeName += LOCTEXT("AllIndicesSelected", "All").ToString();
 	}
 	else
 	{
diff --git a/Engine/Plugins/PCG/Source/PCG/Private/Elements/PCGGetPropertyFromObjectPath.cpp b/Engine/Plugins/PCG/Source/PCG/Private/Elements/PCGGetPropertyFromObjectPath.cpp
--- a/Engine/Plugins/PCG/Source/PCG/Private/Elements/PCGGetPropertyFromObjectPath.cpp
+++ b/Engine/Plugins/PCG/Source/PCG/Private/Elements/PCGGetPropertyFromObjectPath.cpp
@@ -87,7 +87,7 @@ FString UPCGGetPropertyFromObjectPathSettings::GetAdditionalTitleInformation() c
 		else if (ObjectPathsToExtract.Num() == 1)
 		{
 			Path = ObjectPathsToExtract[0].ToString();
-			Path = !Path.IsEmpty() ? Path : TEXT("None");
+			Path = !Path.IsEmpty() ? Path : LOCTEXT("EmptyString", "None").ToString();
 		}
 		else
 		{
diff --git a/Engine/Plugins/PCG/Source/PCG/Private/Elements/PCGReplaceTags.cpp b/Engine/Plugins/PCG/Source/PCG/Private/Elements/PCGReplaceTags.cpp
--- a/Engine/Plugins/PCG/Source/PCG/Private/Elements/PCGReplaceTags.cpp
+++ b/Engine/Plugins/PCG/Source/PCG/Private/Elements/PCGReplaceTags.cpp
@@ -46,17 +46,11 @@ FString UPCGReplaceTagsSettings::GetAdditionalTitleInformation() const
 		: PCGHelpers::GetStringArrayFromCommaSeparatedList(ReplacedTags);
 	PRAGMA_ENABLE_DEPRECATION_WARNINGS
 
-	if (TagsToReplace.Num() == 1 && TagsToProcess.Num() == 1)
-	{
-		return FString::Printf(TEXT("%s -> %s"), *TagsToReplace[0], *TagsToProcess[0]);
-	}
-	else if (TagsToReplace.Num() > 1 && TagsToProcess.Num() == 1)
-	{
-		return FString::Printf(TEXT("Multiple -> %s"), *TagsToProcess[0]);
-	}
-	else if (TagsToReplace.Num() > 1 && TagsToProcess.Num() > 1)
+	const FText MultipleEntry = LOCTEXT("MultipleEntry", "Multiple");
+
+	if(TagsToReplace.Num() >= 1 && TagsToProcess.Num() >= 1)
 	{
-		return FString(TEXT("Multiple -> Multiple"));
+		return FString::Printf(TEXT("%s -> %s"), (TagsToReplace.Num() > 1 ? *MultipleEntry.ToString() : *TagsToReplace[0]), (TagsToProcess.Num() > 1 ? *MultipleEntry.ToString() : *TagsToProcess[0]));
 	}
 	else
 	{
