Update GPU widget following feedback.

#rnx

diff --git a/Engine/Plugins/PCG/Source/PCGEditor/Private/SPCGEditorGraphNode.cpp b/Engine/Plugins/PCG/Source/PCGEditor/Private/SPCGEditorGraphNode.cpp
--- a/Engine/Plugins/PCG/Source/PCGEditor/Private/SPCGEditorGraphNode.cpp
+++ b/Engine/Plugins/PCG/Source/PCGEditor/Private/SPCGEditorGraphNode.cpp
@@ -692,18 +692,16 @@ TArray<FOverlayWidgetInfo> SPCGEditorGraphNode::GetOverlayWidgets(bool bSelected
 
 		if (PCGEditorGraphNode->GetPCGNode() && PCGEditorGraphNode->GetPCGNode()->GetSettings() && PCGEditorGraphNode->GetPCGNode()->GetSettings()->ShouldExecuteOnGPU())
 		{
-			const FLinearColor LightGray(0.7f, 0.7f, 0.7f);
 			const float BorderRadius = 7.0f;
-			const float BorderStroke = 1.5f;
-
+			const float BorderStroke = 1.0f;
 			FText GPUText = FText::FromString(TEXT("GPU"));
-			FLinearColor Tint = LightGray;
-			FLinearColor TextColor = LightGray;
-			FLinearColor BackgroundColor = FColor::White;
+			const FLinearColor BorderColor(0.5f, 0.5f, 0.5f, 0.5f);
+			FLinearColor TextColor(0.5f, 0.5f, 0.5f, 0.8f);
+			
 			const FSlateBrush* BorderBrush = new FSlateRoundedBoxBrush(
-				FLinearColor::Black,
+				FLinearColor::Transparent,
 				BorderRadius,
-				LightGray,
+				BorderColor,
 				BorderStroke);
 
 			TSharedPtr<SWidget> GPUUsageLabel =
@@ -713,8 +711,7 @@ TArray<FOverlayWidgetInfo> SPCGEditorGraphNode::GetOverlayWidgets(bool bSelected
 				[
 					SNew(SBorder)
 					.BorderImage(BorderBrush)
-					.Padding(FMargin(6, 3))
-					.ColorAndOpacity(Tint)
+					.Padding(FMargin(4, 3))
 					[
 						SNew(STextBlock)
 						.TextStyle(FPCGEditorStyle::Get(), "PCG.Node.AdditionalOverlayWidgetText")
