[PCG] Fixed a bug causing empty point data to bypass a check in Filter Elements By Index node, causing an ensure to fire.

Also, some minor issues with properly initializing metadata during the process.

#rnx

diff --git a/Engine/Plugins/PCG/Source/PCG/Private/Elements/PCGFilterElementsByIndex.cpp b/Engine/Plugins/PCG/Source/PCG/Private/Elements/PCGFilterElementsByIndex.cpp
--- a/Engine/Plugins/PCG/Source/PCG/Private/Elements/PCGFilterElementsByIndex.cpp
+++ b/Engine/Plugins/PCG/Source/PCG/Private/Elements/PCGFilterElementsByIndex.cpp
@@ -45,6 +45,11 @@ namespace PCGFilterElementsByIndex
 		{
 			check(InContext);
 
+			if (ElementCount == 0)
+			{
+				return PCGIndexing::FPCGIndexCollection::Invalid();
+			}
+
 			PCGIndexing::FPCGIndexCollection FilteredIndices(ElementCount);
 
 			// Parse the indices and switch through possible issues.
@@ -322,6 +327,7 @@ bool FPCGFilterElementsByIndexElement::ExecuteInternal(FPCGContext* InContext) c
 			if (Settings->bOutputDiscardedElements)
 			{
 				UPCGPointData* OutFilterPointData = FPCGContext::NewObject_AnyThread<UPCGPointData>(InContext);
+				OutFilterPointData->InitializeFromData(InputPointData);
 				OutFilterPointData->GetMutablePoints() = std::move(OutFilterPoints);
 				AddNewOutput(OutFilterPointData, Constants::OutFilterOutputPinLabel);
 			}
@@ -332,8 +338,8 @@ bool FPCGFilterElementsByIndexElement::ExecuteInternal(FPCGContext* InContext) c
 			check(InputMetadata);
 
 			const int32 NumEntries = InputMetadata->GetItemCountForChild();
-			TArray<PCGMetadataEntryKey, TInlineAllocator<256>> InEntryKeys;
-			TArray<PCGMetadataEntryKey, TInlineAllocator<256>> OutEntryKeys;
+			TArray<PCGMetadataEntryKey> InEntryKeys;
+			TArray<PCGMetadataEntryKey> OutEntryKeys;
 			InEntryKeys.Reserve(NumEntries);
 			if (Settings->bOutputDiscardedElements)
 			{
@@ -359,7 +365,7 @@ bool FPCGFilterElementsByIndexElement::ExecuteInternal(FPCGContext* InContext) c
 
 			UPCGMetadata* InFilterMetadata = InFilterParamData->Metadata;
 			check(InFilterMetadata);
-			InFilterMetadata->SetAttributes(InEntryKeys, InputMetadata);
+			InFilterMetadata->InitializeAsCopy(InputMetadata, &InEntryKeys);
 
 			if (Settings->bOutputDiscardedElements)
 			{
@@ -367,7 +373,7 @@ bool FPCGFilterElementsByIndexElement::ExecuteInternal(FPCGContext* InContext) c
 				UPCGMetadata* OutFilterMetadata = OutFilterParamData->Metadata;
 				check(OutFilterMetadata);
 				AddNewOutput(OutFilterParamData, Constants::OutFilterOutputPinLabel);
-				OutFilterMetadata->SetAttributes(OutEntryKeys, InputMetadata);
+				OutFilterMetadata->InitializeAsCopy(InputMetadata, &OutEntryKeys);
 			}
 		}
 		else
diff --git a/Engine/Plugins/PCG/Source/PCG/Private/Helpers/Parsing/PCGIndexing.cpp b/Engine/Plugins/PCG/Source/PCG/Private/Helpers/Parsing/PCGIndexing.cpp
--- a/Engine/Plugins/PCG/Source/PCG/Private/Helpers/Parsing/PCGIndexing.cpp
+++ b/Engine/Plugins/PCG/Source/PCG/Private/Helpers/Parsing/PCGIndexing.cpp
@@ -96,6 +96,11 @@ namespace PCGIndexing
 
 	bool FPCGIndexCollection::ContainsIndex(const int32 Index) const
 	{
+		if (!IsValid())
+		{
+			return false;
+		}
+
 		for (const FPCGIndexRange& Range : IndexRanges)
 		{
 			if (Range.ContainsIndex(Index))
