Hide 'Invalid' GPU attribute type from UI which caused a crash when selected.

#rnx

diff --git a/Engine/Plugins/PCG/Source/PCG/Private/Compute/PCGDataForGPU.h b/Engine/Plugins/PCG/Source/PCG/Private/Compute/PCGDataForGPU.h
--- a/Engine/Plugins/PCG/Source/PCG/Private/Compute/PCGDataForGPU.h
+++ b/Engine/Plugins/PCG/Source/PCG/Private/Compute/PCGDataForGPU.h
@@ -36,7 +36,7 @@ enum class EPCGKernelAttributeType : uint8
 	StringKey,
 	Name,
 
-	Invalid = std::numeric_limits<uint8>::max()
+	Invalid = std::numeric_limits<uint8>::max() UMETA(Hidden),
 };
 
 USTRUCT()
