[PCG] PCGLog::GetComponentOwnerName() will return full owner object path when called in commandlet

#rnx

diff --git a/Engine/Plugins/PCG/Source/PCG/Private/PCGCommon.cpp b/Engine/Plugins/PCG/Source/PCG/Private/PCGCommon.cpp
--- a/Engine/Plugins/PCG/Source/PCG/Private/PCGCommon.cpp
+++ b/Engine/Plugins/PCG/Source/PCG/Private/PCGCommon.cpp
@@ -1,9 +1,34 @@
 // Copyright Epic Games, Inc. All Rights Reserved.
 
 #include "PCGCommon.h"
+#include "PCGComponent.h"
 #include "PCGSubsystem.h"
 #include "PCGWorldActor.h"
 
+namespace PCGLog
+{
+	FString GetComponentOwnerName(const UPCGComponent* InComponent, bool bUseLabel, FString InDefaultName)
+	{
+		if (InComponent && InComponent->GetOwner())
+		{
+			if (IsRunningCommandlet())
+			{
+				return InComponent->GetOwner()->GetPathName();
+			}
+			else if(bUseLabel)
+			{
+				return InComponent->GetOwner()->GetActorNameOrLabel();
+			}
+			else
+			{
+				return InComponent->GetOwner()->GetName();
+			}
+		}
+
+		return InDefaultName;
+	}
+}
+
 namespace PCGFeatureSwitches
 {
 	TAutoConsoleVariable<bool> CVarCheckSamplerMemory{
diff --git a/Engine/Plugins/PCG/Source/PCG/Private/PCGContext.cpp b/Engine/Plugins/PCG/Source/PCG/Private/PCGContext.cpp
--- a/Engine/Plugins/PCG/Source/PCG/Private/PCGContext.cpp
+++ b/Engine/Plugins/PCG/Source/PCG/Private/PCGContext.cpp
@@ -1,6 +1,7 @@
 // Copyright Epic Games, Inc. All Rights Reserved.
 
 #include "PCGContext.h"
+#include "PCGCommon.h"
 #include "PCGComponent.h"
 #include "PCGGraph.h"
 #include "PCGParamData.h"
@@ -81,7 +82,7 @@ int FPCGContext::GetSeed() const
 
 FString FPCGContext::GetComponentName() const
 {
-	return SourceComponent.IsValid() && SourceComponent->GetOwner() ? SourceComponent->GetOwner()->GetFName().ToString() : TEXT("Non-PCG Component");
+	return PCGLog::GetComponentOwnerName(SourceComponent.Get(), /*bUseLabel=*/false, /*InDefaultName=*/TEXT("Non-PCG Component"));
 }
 
 AActor* FPCGContext::GetTargetActor(const UPCGSpatialData* InSpatialData) const
diff --git a/Engine/Plugins/PCG/Source/PCG/Private/Utils/PCGGraphExecutionLogging.cpp b/Engine/Plugins/PCG/Source/PCG/Private/Utils/PCGGraphExecutionLogging.cpp
--- a/Engine/Plugins/PCG/Source/PCG/Private/Utils/PCGGraphExecutionLogging.cpp
+++ b/Engine/Plugins/PCG/Source/PCG/Private/Utils/PCGGraphExecutionLogging.cpp
@@ -2,6 +2,7 @@
 
 #include "PCGGraphExecutionLogging.h"
 
+#include "PCGCommon.h"
 #include "PCGComponent.h"
 #include "PCGContext.h"
 #include "PCGGraph.h"
@@ -43,19 +44,6 @@ namespace PCGGraphExecutionLogging
 #endif
 	}
 
-#if !(UE_BUILD_SHIPPING || UE_BUILD_TEST) || USE_LOGGING_IN_SHIPPING
-	FString GetOwnerName(const UPCGComponent* InComponent)
-	{
-		return (InComponent && InComponent->GetOwner()) ?
-#if WITH_EDITOR
-			InComponent->GetOwner()->GetActorLabel() :
-#else
-			InComponent->GetOwner()->GetName() :
-#endif
-			FString(TEXT("MISSINGOWNER"));
-	}
-#endif // !(UE_BUILD_SHIPPING || UE_BUILD_TEST) || USE_LOGGING_IN_SHIPPING
-
 	void LogGraphTask(FPCGTaskId TaskId, const FPCGGraphTask& Task, const TSet<FPCGTaskId>* SuccessorIds)
 	{
 #if !(UE_BUILD_SHIPPING || UE_BUILD_TEST) || USE_LOGGING_IN_SHIPPING
@@ -154,7 +142,7 @@ namespace PCGGraphExecutionLogging
 		}
 
 		UE_LOG(LogPCG, Display, TEXT("[%s/%s] --- SCHEDULE GRAPH %s ---"),
-			(SourceComponent && SourceComponent->GetOwner()) ? *SourceComponent->GetOwner()->GetName() : TEXT("MISSINGCOMPONENT"),
+			*PCGLog::GetComponentOwnerName(SourceComponent),
 			(SourceComponent && SourceComponent->GetGraph()) ? *SourceComponent->GetGraph()->GetName() : TEXT("MISSINGGRAPH"),
 			InScheduledGraph ? *InScheduledGraph->GetName() : TEXT("MISSINGGRAPH)"));
 #endif // !(UE_BUILD_SHIPPING || UE_BUILD_TEST) || USE_LOGGING_IN_SHIPPING
@@ -175,7 +163,7 @@ namespace PCGGraphExecutionLogging
 		}
 
 		UE_LOG(LogPCG, Display, TEXT("[%s/%s] --- SCHEDULE GRAPH FOR DEPENDENCY, from stack: %s"),
-			(InComponent && InComponent->GetOwner()) ? *InComponent->GetOwner()->GetName() : TEXT("MISSINGCOMPONENT"),
+			*PCGLog::GetComponentOwnerName(InComponent),
 			(InComponent && InComponent->GetGraph()) ? *InComponent->GetGraph()->GetName() : TEXT("MISSINGGRAPH"),
 			*FromStackString);
 #endif // !(UE_BUILD_SHIPPING || UE_BUILD_TEST) || USE_LOGGING_IN_SHIPPING
@@ -196,7 +184,7 @@ namespace PCGGraphExecutionLogging
 		}
 
 		UE_LOG(LogPCG, Warning, TEXT("[%s/%s] Failed to schedule dependency, from stack: %s"),
-			(InComponent && InComponent->GetOwner()) ? *InComponent->GetOwner()->GetName() : TEXT("MISSINGCOMPONENT"),
+			*PCGLog::GetComponentOwnerName(InComponent),
 			(InComponent && InComponent->GetGraph()) ? *InComponent->GetGraph()->GetName() : TEXT("MISSINGGRAPH"),
 			*FromStackString);
 #endif // !(UE_BUILD_SHIPPING || UE_BUILD_TEST) || USE_LOGGING_IN_SHIPPING
@@ -225,7 +213,7 @@ namespace PCGGraphExecutionLogging
 		}
 
 		UE_LOG(LogPCG, Display, TEXT("[%s/%s] UPCGComponent::PostProcessGraph"),
-			(InSourceComponent && InSourceComponent->GetOwner()) ? *InSourceComponent->GetOwner()->GetName() : TEXT("MISSINGCOMPONENT"),
+			*PCGLog::GetComponentOwnerName(InSourceComponent),
 			(InSourceComponent && InSourceComponent->GetGraph()) ? *InSourceComponent->GetGraph()->GetName() : TEXT("MISSINGGRAPH"));
 #endif // !(UE_BUILD_SHIPPING || UE_BUILD_TEST) || USE_LOGGING_IN_SHIPPING
 	}
@@ -241,7 +229,7 @@ namespace PCGGraphExecutionLogging
 		for (const UPCGComponent* Component : CancelledComponents)
 		{
 			UE_LOG(LogPCG, Display, TEXT("[%s/%s] Component cancelled"),
-				(Component && Component->GetOwner()) ? *Component->GetOwner()->GetName() : TEXT("MISSINGCOMPONENT"),
+				*PCGLog::GetComponentOwnerName(Component),
 				(Component && Component->GetGraph()) ? *Component->GetGraph()->GetName() : TEXT("MISSINGGRAPH"));
 		}
 #endif // !(UE_BUILD_SHIPPING || UE_BUILD_TEST) || USE_LOGGING_IN_SHIPPING
@@ -256,7 +244,7 @@ namespace PCGGraphExecutionLogging
 		}
 
 		UE_LOG(LogPCG, Log, TEXT("[%s/%s] Change origin ignored: '%s'"),
-			(InComponent && InComponent->GetOwner()) ? *InComponent->GetOwner()->GetName() : TEXT("MISSINGCOMPONENT"),
+			*PCGLog::GetComponentOwnerName(InComponent),
 			(InComponent && InComponent->GetGraph()) ? *InComponent->GetGraph()->GetName() : TEXT("MISSINGGRAPH"),
 			InObject ? *InObject->GetName() : TEXT("MISSINGOBJECT"));
 #endif // !(UE_BUILD_SHIPPING || UE_BUILD_TEST) || USE_LOGGING_IN_SHIPPING
@@ -301,7 +289,7 @@ namespace PCGGraphExecutionLogging
 		}
 
 		UE_LOG(LogPCG, Log, TEXT("         [%s/%s] %s\t\tEXECUTE"),
-			*Task.SourceComponent->GetOwner()->GetName(),
+			*PCGLog::GetComponentOwnerName(Task.SourceComponent.Get()),
 			Task.SourceComponent->GetGraph() ? *Task.SourceComponent->GetGraph()->GetName() : TEXT("MISSINGGRAPH"),
 			*FString::Printf(TEXT("%" UINT64_FMT "'%s'"), Task.NodeId, Task.Node ? *Task.Node->GetNodeTitle(EPCGNodeTitleType::ListView).ToString() : TEXT("")));
 #endif // !(UE_BUILD_SHIPPING || UE_BUILD_TEST) || USE_LOGGING_IN_SHIPPING
@@ -316,7 +304,7 @@ namespace PCGGraphExecutionLogging
 		}
 
 		UE_LOG(LogPCG, Warning, TEXT("[%s/%s] %s\t\tCACHING DISABLED"),
-			*Task.SourceComponent->GetOwner()->GetName(),
+			*PCGLog::GetComponentOwnerName(Task.SourceComponent.Get()),
 			Task.SourceComponent->GetGraph() ? *Task.SourceComponent->GetGraph()->GetName() : TEXT("MISSINGGRAPH"),
 			*FString::Printf(TEXT("%" UINT64_FMT "'%s'"), Task.NodeId, Task.Node ? *Task.Node->GetNodeTitle(EPCGNodeTitleType::ListView).ToString() : TEXT("")));
 #endif // !(UE_BUILD_SHIPPING || UE_BUILD_TEST) || USE_LOGGING_IN_SHIPPING
@@ -381,7 +369,7 @@ namespace PCGGraphExecutionLogging
 		check(InContext);
 
 		UE_LOG(LogPCG, Log, TEXT("[GRIDLINKING] [%s] STORE. GenerationGridSize=%u, FromGridSize=%u, ToGridSize=%u, Path=%s, DataItems=%d"),
-			*GetOwnerName(InContext->SourceComponent.Get()),
+			*PCGLog::GetComponentOwnerName(InContext->SourceComponent.Get(), /*bUseLabel=*/true),
 			PCGHiGenGrid::IsValidGrid(InGenerationGrid) ? PCGHiGenGrid::GridToGridSize(InGenerationGrid) : PCGHiGenGrid::UnboundedGridSize(),
 			InFromGridSize,
 			InToGridSize,
@@ -400,7 +388,7 @@ namespace PCGGraphExecutionLogging
 		check(InContext);
 
 		UE_LOG(LogPCG, Log, TEXT("[GRIDLINKING] [%s] RETRIEVE. GenerationGridSize=%u, FromGridSize=%u, ToGridSize=%u, Path=%s"),
-			*GetOwnerName(InContext->SourceComponent.Get()),
+			*PCGLog::GetComponentOwnerName(InContext->SourceComponent.Get(), /*bUseLabel=*/true),
 			PCGHiGenGrid::IsValidGrid(InGenerationGrid) ? PCGHiGenGrid::GridToGridSize(InGenerationGrid) : PCGHiGenGrid::UnboundedGridSize(),
 			InFromGridSize,
 			InToGridSize,
@@ -418,7 +406,7 @@ namespace PCGGraphExecutionLogging
 		check(InContext);
 
 		UE_LOG(LogPCG, Log, TEXT("[GRIDLINKING] [%s] RETRIEVE: SUCCESS. Path=%s, DataItems=%d"),
-			*GetOwnerName(InContext->SourceComponent.Get()),
+			*PCGLog::GetComponentOwnerName(InContext->SourceComponent.Get(), /*bUseLabel=*/true),
 			*InResourcePath,
 			InDataItemCount);
 #endif // !(UE_BUILD_SHIPPING || UE_BUILD_TEST) || USE_LOGGING_IN_SHIPPING
@@ -434,8 +422,8 @@ namespace PCGGraphExecutionLogging
 		check(InContext);
 
 		UE_LOG(LogPCG, Log, TEXT("[GRIDLINKING] [%s] RETRIEVE: SCHEDULE GRAPH. Component=%s, Path=%s"),
-			*GetOwnerName(InContext->SourceComponent.Get()),
-			*GetOwnerName(InScheduledComponent),
+			*PCGLog::GetComponentOwnerName(InContext->SourceComponent.Get(), /*bUseLabel=*/true),
+			*PCGLog::GetComponentOwnerName(InScheduledComponent, /*bUseLabel=*/true),
 			*InResourcePath);
 #endif // !(UE_BUILD_SHIPPING || UE_BUILD_TEST) || USE_LOGGING_IN_SHIPPING
 	}
@@ -450,8 +438,8 @@ namespace PCGGraphExecutionLogging
 		check(InContext);
 
 		UE_LOG(LogPCG, Log, TEXT("[GRIDLINKING] [%s] RETRIEVE: WAIT FOR SCHEDULED GRAPH. Component=%s, Path=%s"),
-			*GetOwnerName(InContext->SourceComponent.Get()),
-			*GetOwnerName(InWaitOnComponent),
+			*PCGLog::GetComponentOwnerName(InContext->SourceComponent.Get(), /*bUseLabel=*/true),
+			*PCGLog::GetComponentOwnerName(InWaitOnComponent, /*bUseLabel=*/true),
 			*InResourcePath);
 #endif // !(UE_BUILD_SHIPPING || UE_BUILD_TEST) || USE_LOGGING_IN_SHIPPING
 	}
@@ -466,8 +454,8 @@ namespace PCGGraphExecutionLogging
 		check(InContext);
 
 		UE_LOG(LogPCG, Log, TEXT("[GRIDLINKING] [%s] RETRIEVE: WOKEN BY Component=%s"),
-			*GetOwnerName(InContext->SourceComponent.Get()),
-			*GetOwnerName(InWokenBy));
+			*PCGLog::GetComponentOwnerName(InContext->SourceComponent.Get(), /*bUseLabel=*/true),
+			*PCGLog::GetComponentOwnerName(InWokenBy), /*bUseLabel=*/true);
 #endif // !(UE_BUILD_SHIPPING || UE_BUILD_TEST) || USE_LOGGING_IN_SHIPPING
 	}
 
@@ -481,7 +469,7 @@ namespace PCGGraphExecutionLogging
 		check(InContext);
 
 		UE_LOG(LogPCG, Warning, TEXT("[GRIDLINKING] [%s] RETRIEVE: FAILED: No overlapping local component found. This may be expected. Path=%s"),
-			*GetOwnerName(InContext->SourceComponent.Get()),
+			*PCGLog::GetComponentOwnerName(InContext->SourceComponent.Get(), /*bUseLabel=*/true),
 			*InResourcePath);
 #endif // !(UE_BUILD_SHIPPING || UE_BUILD_TEST) || USE_LOGGING_IN_SHIPPING
 	}
@@ -496,8 +484,8 @@ namespace PCGGraphExecutionLogging
 		check(InContext);
 
 		UE_LOG(LogPCG, Warning, TEXT("[GRIDLINKING] [%s] RETRIEVE: FAILED: No data found on local component. Component=%s, Path=%s"),
-			*GetOwnerName(InContext->SourceComponent.Get()),
-			*GetOwnerName(InComponent),
+			*PCGLog::GetComponentOwnerName(InContext->SourceComponent.Get(), /*bUseLabel=*/true),
+			*PCGLog::GetComponentOwnerName(InComponent, /*bUseLabel=*/true),
 			*InResourcePath);
 #endif // !(UE_BUILD_SHIPPING || UE_BUILD_TEST) || USE_LOGGING_IN_SHIPPING
 	}
diff --git a/Engine/Plugins/PCG/Source/PCG/Public/PCGCommon.h b/Engine/Plugins/PCG/Source/PCG/Public/PCGCommon.h
--- a/Engine/Plugins/PCG/Source/PCG/Public/PCGCommon.h
+++ b/Engine/Plugins/PCG/Source/PCG/Public/PCGCommon.h
@@ -11,6 +11,14 @@ static const FPCGTaskId InvalidPCGTaskId = (uint64)-1;
 
 using FPCGPinId = uint64;
 
+class UPCGComponent;
+
+namespace PCGLog
+{
+	/** Returns the component's owner name. In the case of commandlets, owner full object path will be returned. */
+	FString GetComponentOwnerName(const UPCGComponent* InComponent, bool bUseLabel = false, FString InDefaultName = TEXT("MISSINGOWNER"));
+}
+
 namespace PCGPinIdHelpers
 {
 	/** Pin active bitmask stored in uint64, so 64 flags available. */
