Fixed corruption issue  (which could also cause crashes) when switching static mesh spawners to gpu when there were between 65 and 256 primitives. Improved behavior of attribute value in rejected points so we don't try to build potentially bad values.

#rnx

diff --git a/Engine/Plugins/PCG/Shaders/Private/Elements/PCGStaticMeshSpawner.usf b/Engine/Plugins/PCG/Shaders/Private/Elements/PCGStaticMeshSpawner.usf
--- a/Engine/Plugins/PCG/Shaders/Private/Elements/PCGStaticMeshSpawner.usf
+++ b/Engine/Plugins/PCG/Shaders/Private/Elements/PCGStaticMeshSpawner.usf
@@ -74,13 +74,6 @@ void Main(uint3 GroupId : SV_GroupID, uint GroupIndex : SV_GroupIndex)
 		}
 	}
 
-	const int InstanceIndexInt = InstanceData_GetIndexToWriteTo(PrimitiveIndex);
-	if (InstanceIndexInt < 0)
-	{
-		// Saturated instances for this primitive. No-op.
-		return;
-	}
-
 	uint Out_DataIndex, Out_ElementIndex;
 	if (!Out_GetThreadData(ThreadIndex, Out_DataIndex, Out_ElementIndex))
 	{
@@ -88,37 +81,39 @@ void Main(uint3 GroupId : SV_GroupID, uint GroupIndex : SV_GroupIndex)
 		return;
 	}
 
-	const uint InstanceIndex = InstanceIndexInt;
-
 	if (In_IsPointRemoved(In_DataIndex, In_ElementIndex))
 	{
-		// Cull the instance by setting the transform to 0.
-		InstanceData_WriteInstanceLocalToWorld(InstanceIndex, (float4x4)0);
-		
 		// Propagate the 'Removed' status to the output point if input point is removed. Otherwise, this point will not be culled.
 		Out_RemovePoint(Out_DataIndex, Out_ElementIndex);
 
 		return;
 	}
 
-	// Write instance transform.
-	const float4x4 Transform = In_GetPointTransform(In_DataIndex, In_ElementIndex);
-	InstanceData_WriteInstanceLocalToWorld(InstanceIndex, transpose(Transform));
+	// Get instance index to write to. Will return a negative value if we have saturated instances for this primitive, in which case we will not push the data.
+	const int InstanceIndexInt = InstanceData_GetIndexToWriteTo(PrimitiveIndex);
 
-	// Write instance custom float data.
-	const uint NumAttributes = SMSpawner_GetNumAttributes();
-	for (uint AttributeIndex = 0; AttributeIndex < NumAttributes; ++AttributeIndex)
+	// Write instance transform, only if not culled by primitive instance index
+	if(InstanceIndexInt >= 0)
 	{
-		const uint4 AttributeIdOffsetStride = SMSpawner_GetAttributeIdOffsetStride(AttributeIndex);
-		const uint AttributeId = AttributeIdOffsetStride[0];
-		const uint AttributeOffset = AttributeIdOffsetStride[1];
-		const uint AttributeStride = AttributeIdOffsetStride[2];
+		const uint InstanceIndex = InstanceIndexInt;
+		const float4x4 Transform = In_GetPointTransform(In_DataIndex, In_ElementIndex);
+		InstanceData_WriteInstanceLocalToWorld(InstanceIndex, transpose(Transform));
 
-		for (uint FloatIndex = 0; FloatIndex < AttributeStride; ++FloatIndex)
+		// Write instance custom float data.
+		const uint NumAttributes = SMSpawner_GetNumAttributes();
+		for (uint AttributeIndex = 0; AttributeIndex < NumAttributes; ++AttributeIndex)
 		{
-			const uint AttributeValueAddress = In_GetElementAddressInternal(In_DataIndex, In_ElementIndex, AttributeId);
-			// Assumes floating point data. TODO add type information and deal with different types in different ways?
-			InstanceData_WriteCustomFloatRaw(InstanceIndex, AttributeOffset + FloatIndex, In_LoadBufferInternal(AttributeValueAddress + 4 * FloatIndex));
+			const uint4 AttributeIdOffsetStride = SMSpawner_GetAttributeIdOffsetStride(AttributeIndex);
+			const uint AttributeId = AttributeIdOffsetStride[0];
+			const uint AttributeOffset = AttributeIdOffsetStride[1];
+			const uint AttributeStride = AttributeIdOffsetStride[2];
+
+			for (uint FloatIndex = 0; FloatIndex < AttributeStride; ++FloatIndex)
+			{
+				const uint AttributeValueAddress = In_GetElementAddressInternal(In_DataIndex, In_ElementIndex, AttributeId);
+				// Assumes floating point data. TODO add type information and deal with different types in different ways?
+				InstanceData_WriteCustomFloatRaw(InstanceIndex, AttributeOffset + FloatIndex, In_LoadBufferInternal(AttributeValueAddress + 4 * FloatIndex));
+			}
 		}
 	}
 
diff --git a/Engine/Plugins/PCG/Shaders/Private/PCGInstanceDataInterface.ush b/Engine/Plugins/PCG/Shaders/Private/PCGInstanceDataInterface.ush
--- a/Engine/Plugins/PCG/Shaders/Private/PCGInstanceDataInterface.ush
+++ b/Engine/Plugins/PCG/Shaders/Private/PCGInstanceDataInterface.ush
@@ -2,7 +2,7 @@
 
 #include "/Engine/Private/Common.ush"
 
-#define MAX_PRIMITIVES 64
+#define MAX_PRIMITIVES 256 // Should be the same as MAX_PRIMITIVE_COMPONENTS_PER_SPAWNER
 
 uint4 {DataInterfaceName}_PrimitiveParameters[MAX_PRIMITIVES]; // (NumInstancesAllocated, InstanceOffset, Unused, Unused)
 
diff --git a/Engine/Plugins/PCG/Source/PCG/Private/Compute/PCGDataForGPU.cpp b/Engine/Plugins/PCG/Source/PCG/Private/Compute/PCGDataForGPU.cpp
--- a/Engine/Plugins/PCG/Source/PCG/Private/Compute/PCGDataForGPU.cpp
+++ b/Engine/Plugins/PCG/Source/PCG/Private/Compute/PCGDataForGPU.cpp
@@ -1264,6 +1264,12 @@ EPCGUnpackDataCollectionResult FPCGDataCollectionDesc::UnpackDataCollection(cons
 
 					ParallelFor(NumElements, [PackedData, &InStringTable, &OutPoints, Metadata, AttributeBase, &AttributeDesc, AddressUints, AttributeStrideUints](int32 ElementIndex)
 					{
+						// Points that were removed will not have their metadata set at all, and most likely not zero initialized either, so this could be garbage otherwise.
+						if (!FMath::IsFinite(OutPoints[ElementIndex].Density))
+						{
+							return;
+						}
+
 						Metadata->InitializeOnSet(OutPoints[ElementIndex].MetadataEntry);
 						ensure(PCGDataForGPUHelpers::UnpackAttributeHelper(PackedData, InStringTable, AddressUints + ElementIndex * AttributeStrideUints, AttributeBase, AttributeDesc, OutPoints[ElementIndex].MetadataEntry));
 					});
@@ -1349,4 +1355,4 @@ uint32 FPCGDataCollectionDesc::ComputeDataElementCount(EPCGDataType InDataType)
 void FPCGDataCollectionDesc::Combine(const FPCGDataCollectionDesc& Other)
 {
 	DataDescs.Append(Other.DataDescs);
-}
\ No newline at end of file
+}
