[PCG] Avoid PCGComponent refresh when inside a editor loading stack

#rnx

diff --git a/Engine/Plugins/PCG/Source/PCG/Private/PCGActorAndComponentMapping.cpp b/Engine/Plugins/PCG/Source/PCG/Private/PCGActorAndComponentMapping.cpp
--- a/Engine/Plugins/PCG/Source/PCG/Private/PCGActorAndComponentMapping.cpp
+++ b/Engine/Plugins/PCG/Source/PCG/Private/PCGActorAndComponentMapping.cpp
@@ -1351,7 +1351,8 @@ void FPCGActorAndComponentMapping::OnActorAdded(AActor* InActor)
 			
 	// Implementation note: We delay adding because OnActorAdded fires before an actor's properties are set,
 	// so the actor is not ready for processing until the next tick.
-	OnActorAdded_Internal(InActor, /*bShouldDirty=*/ true);
+	// Only dirty if we aren't currently inside a loading path
+	OnActorAdded_Internal(InActor, /*bShouldDirty=*/ !UE::GetIsEditorLoadingPackage());
 }
 
 void FPCGActorAndComponentMapping::OnPackageReloaded(const EPackageReloadPhase InPackageReloadPhase, FPackageReloadedEvent* InPackageReloadedEvent)
@@ -1548,8 +1549,8 @@ void FPCGActorAndComponentMapping::OnActorDeleted(AActor* InActor)
 		return;
 	}
 
-	// Implementation note: since this is called only for actors directly in the current level, the depth here is 0.
-	OnActorDeleted_Internal(InActor, /*bShouldDirty=*/true);
+	// Only dirty if we aren't currently inside a loading path
+	OnActorDeleted_Internal(InActor, /*bShouldDirty=*/ !UE::GetIsEditorLoadingPackage());
 }
 
 void FPCGActorAndComponentMapping::OnActorDeleted_Internal(AActor* InActor, bool bShouldDirty)
