[PCG] APCGPartitionActor/APCGWorldActor Expose GetPCGGridSize() to Blueprint

#rnx

diff --git a/Engine/Plugins/PCG/Source/PCG/Public/Grid/PCGPartitionActor.h b/Engine/Plugins/PCG/Source/PCG/Public/Grid/PCGPartitionActor.h
--- a/Engine/Plugins/PCG/Source/PCG/Public/Grid/PCGPartitionActor.h
+++ b/Engine/Plugins/PCG/Source/PCG/Public/Grid/PCGPartitionActor.h
@@ -60,6 +60,10 @@ public:
 
 	FBox GetFixedBounds() const;
 	FIntVector GetGridCoord() const;
+
+	UFUNCTION(BlueprintCallable, Category = "PCG", meta= (DisplayName="Get PCG Grid Size"))
+	int64 BP_GetPCGGridSize() const { return (int64)GetPCGGridSize(); }
+
 	uint32 GetPCGGridSize() const { return PCGGridSize; }
 	FPCGGridDescriptor GetGridDescriptor() const;
 
diff --git a/Engine/Plugins/PCG/Source/PCG/Public/PCGWorldActor.h b/Engine/Plugins/PCG/Source/PCG/Public/PCGWorldActor.h
--- a/Engine/Plugins/PCG/Source/PCG/Public/PCGWorldActor.h
+++ b/Engine/Plugins/PCG/Source/PCG/Public/PCGWorldActor.h
@@ -36,6 +36,9 @@ public:
 	//~End AActor Interface
 #endif
 
+	UFUNCTION(BlueprintCallable, Category = "PCG", meta=(DisplayName="Get PCG Grid Size"))
+	int64 BP_GetPCGGridSize() const { return (int64)PartitionGridSize; }
+
 	//~ Begin UObject Interface.
 #if WITH_EDITOR
 	virtual void PostEditChangeProperty(FPropertyChangedEvent& PropertyChangedEvent) override;
