[PCG] Flush out visual logs for invalid/missing component to fix stuck errors.

[RN] [PCG] Fixed cases where graph node errors would stick onto nodes and not clear properly.

diff --git a/Engine/Plugins/PCG/Source/PCG/Private/Utils/PCGNodeVisualLogs.cpp b/Engine/Plugins/PCG/Source/PCG/Private/Utils/PCGNodeVisualLogs.cpp
--- a/Engine/Plugins/PCG/Source/PCG/Private/Utils/PCGNodeVisualLogs.cpp
+++ b/Engine/Plugins/PCG/Source/PCG/Private/Utils/PCGNodeVisualLogs.cpp
@@ -3,6 +3,7 @@
 #include "Utils/PCGNodeVisualLogs.h"
 
 #include "PCGComponent.h"
+#include "PCGModule.h"
 
 #include "Algo/Find.h"
 #include "GameFramework/Actor.h"
@@ -202,8 +203,16 @@ void FPCGNodeVisualLogs::ClearLogs(const FPCGStack& InPCGStack)
 		TArray<FPCGStack> StacksToRemove;
 		for (const TPair<FPCGStack, FPCGPerNodeVisualLogs>& Entry : StackToLogs)
 		{
-			if (Entry.Key.BeginsWith(InPCGStack))
+			// Always take every opportunity to flush messages logged against invalid/dead components.
+			const bool bComponentValid = IsValid(Entry.Key.GetRootComponent());
+
+			if (!bComponentValid || Entry.Key.BeginsWith(InPCGStack))
 			{
+				if (!bComponentValid)
+				{
+					UE_LOG(LogPCG, Verbose, TEXT("Cleared out logs for null component."));
+				}
+
 				StacksToRemove.Add(Entry.Key);
 
 				for (const FPCGStackFrame& Frame : Entry.Key.GetStackFrames())
diff --git a/Engine/Plugins/PCG/Source/PCGEditor/Private/PCGEditor.cpp b/Engine/Plugins/PCG/Source/PCGEditor/Private/PCGEditor.cpp
--- a/Engine/Plugins/PCG/Source/PCGEditor/Private/PCGEditor.cpp
+++ b/Engine/Plugins/PCG/Source/PCGEditor/Private/PCGEditor.cpp
@@ -3115,6 +3115,11 @@ void FPCGEditor::OnComponentUnregistered(UPCGComponent* Component)
 	{
 		DebugObjectTreeWidget->RequestRefresh();
 	}
+
+	if (UPCGSubsystem* Subsystem = Component ? Component->GetSubsystem() : nullptr)
+	{
+		Subsystem->GetNodeVisualLogsMutable().ClearLogs(Component);
+	}
 }
 
 void FPCGEditor::OnComponentGenerationDone(UPCGSubsystem* Subsystem, UPCGComponent* Component, EPCGGenerationStatus Status)
