[PCG] Fix ShouldComputeFullOutputDataCrc not always getting a full crc
- PCGData Crc: add a new member to know if cached Crc is a full crc or not
- Calling ComputeCrc will recompute the Crc if we are asking for a full crc and the current one isn't and the data type supports it
- Add SupportsFullDataCrc() virtual to avoid multiple crc computes on PCGData classes that don't distinguish between Crc and FullCrc
- PCGCollapseElement will not require a full output data crc if it is just forwarding data
- Remove usage of PropagateCrcThroughBooleanData and always go through that path for composite data

#rnx

diff --git a/Engine/Plugins/PCG/Source/PCG/Private/Data/PCGDifferenceData.cpp b/Engine/Plugins/PCG/Source/PCG/Private/Data/PCGDifferenceData.cpp
--- a/Engine/Plugins/PCG/Source/PCG/Private/Data/PCGDifferenceData.cpp
+++ b/Engine/Plugins/PCG/Source/PCG/Private/Data/PCGDifferenceData.cpp
@@ -137,26 +137,19 @@ FPCGCrc UPCGDifferenceData::ComputeCrc(bool bFullDataCrc) const
 {
 	FArchiveCrc32 Ar;
 
-	if (PropagateCrcThroughBooleanData())
-	{
-		AddToCrc(Ar, bFullDataCrc);
+	AddToCrc(Ar, bFullDataCrc);
 
-		// Chain together CRCs of operands
-		check(GetSource());
-		uint32 SourceCrc = GetSource()->GetOrComputeCrc(bFullDataCrc).GetValue();
-		Ar << SourceCrc;
+	// Chain together CRCs of operands
+	check(GetSource());
+	uint32 SourceCrc = GetSource()->GetOrComputeCrc(bFullDataCrc).GetValue();
+	Ar << SourceCrc;
 
-		if (GetDifference())
-		{
-			uint32 DifferenceCrc = GetDifference()->GetOrComputeCrc(bFullDataCrc).GetValue();
-			Ar << DifferenceCrc;
-		}
-	}
-	else
+	if (GetDifference())
 	{
-		AddUIDToCrc(Ar);
+		uint32 DifferenceCrc = GetDifference()->GetOrComputeCrc(bFullDataCrc).GetValue();
+		Ar << DifferenceCrc;
 	}
-
+	
 	return FPCGCrc(Ar.GetCrc());
 }
 
diff --git a/Engine/Plugins/PCG/Source/PCG/Private/Data/PCGIntersectionData.cpp b/Engine/Plugins/PCG/Source/PCG/Private/Data/PCGIntersectionData.cpp
--- a/Engine/Plugins/PCG/Source/PCG/Private/Data/PCGIntersectionData.cpp
+++ b/Engine/Plugins/PCG/Source/PCG/Private/Data/PCGIntersectionData.cpp
@@ -78,22 +78,15 @@ FPCGCrc UPCGIntersectionData::ComputeCrc(bool bFullDataCrc) const
 {
 	FArchiveCrc32 Ar;
 
-	if (PropagateCrcThroughBooleanData())
-	{
-		AddToCrc(Ar, bFullDataCrc);
-
-		// Chain together CRCs of operands
-		check(GetA() && GetB());
-		uint32 CrcA = GetA()->GetOrComputeCrc(bFullDataCrc).GetValue();
-		uint32 CrcB = GetB()->GetOrComputeCrc(bFullDataCrc).GetValue();
-		Ar << CrcA;
-		Ar << CrcB;
-	}
-	else
-	{
-		AddUIDToCrc(Ar);
-	}
+	AddToCrc(Ar, bFullDataCrc);
 
+	// Chain together CRCs of operands
+	check(GetA() && GetB());
+	uint32 CrcA = GetA()->GetOrComputeCrc(bFullDataCrc).GetValue();
+	uint32 CrcB = GetB()->GetOrComputeCrc(bFullDataCrc).GetValue();
+	Ar << CrcA;
+	Ar << CrcB;
+	
 	return FPCGCrc(Ar.GetCrc());
 }
 
diff --git a/Engine/Plugins/PCG/Source/PCG/Private/Data/PCGProjectionData.cpp b/Engine/Plugins/PCG/Source/PCG/Private/Data/PCGProjectionData.cpp
--- a/Engine/Plugins/PCG/Source/PCG/Private/Data/PCGProjectionData.cpp
+++ b/Engine/Plugins/PCG/Source/PCG/Private/Data/PCGProjectionData.cpp
@@ -74,23 +74,16 @@ FPCGCrc UPCGProjectionData::ComputeCrc(bool bFullDataCrc) const
 {
 	FArchiveCrc32 Ar;
 
-	if (PropagateCrcThroughBooleanData())
-	{
-		AddToCrc(Ar, bFullDataCrc);
-
-		// Chain together CRCs of operands
-		check(Source && Target);
-		uint32 CrcSource = Source->GetOrComputeCrc(bFullDataCrc).GetValue();
-		uint32 CrcTarget = Target->GetOrComputeCrc(bFullDataCrc).GetValue();
+	AddToCrc(Ar, bFullDataCrc);
 
-		Ar << CrcSource;
-		Ar << CrcTarget;
-	}
-	else
-	{
-		AddUIDToCrc(Ar);
-	}
+	// Chain together CRCs of operands
+	check(Source && Target);
+	uint32 CrcSource = Source->GetOrComputeCrc(bFullDataCrc).GetValue();
+	uint32 CrcTarget = Target->GetOrComputeCrc(bFullDataCrc).GetValue();
 
+	Ar << CrcSource;
+	Ar << CrcTarget;
+	
 	return FPCGCrc(Ar.GetCrc());
 }
 
diff --git a/Engine/Plugins/PCG/Source/PCG/Private/Data/PCGUnionData.cpp b/Engine/Plugins/PCG/Source/PCG/Private/Data/PCGUnionData.cpp
--- a/Engine/Plugins/PCG/Source/PCG/Private/Data/PCGUnionData.cpp
+++ b/Engine/Plugins/PCG/Source/PCG/Private/Data/PCGUnionData.cpp
@@ -89,28 +89,21 @@ FPCGCrc UPCGUnionData::ComputeCrc(bool bFullDataCrc) const
 {
 	FArchiveCrc32 Ar;
 
-	if (PropagateCrcThroughBooleanData())
-	{
-		AddToCrc(Ar, bFullDataCrc);
+	AddToCrc(Ar, bFullDataCrc);
 
-		// Chain together CRCs of operands
-		int32 NumOperands = Data.Num();
-		Ar << NumOperands;
+	// Chain together CRCs of operands
+	int32 NumOperands = Data.Num();
+	Ar << NumOperands;
 
-		for (const TObjectPtr<const UPCGSpatialData>& Datum : Data)
+	for (const TObjectPtr<const UPCGSpatialData>& Datum : Data)
+	{
+		if (Datum)
 		{
-			if (Datum)
-			{
-				uint32 DatumCrc = Datum->GetOrComputeCrc(bFullDataCrc).GetValue();
-				Ar << DatumCrc;
-			}
+			uint32 DatumCrc = Datum->GetOrComputeCrc(bFullDataCrc).GetValue();
+			Ar << DatumCrc;
 		}
 	}
-	else
-	{
-		AddUIDToCrc(Ar);
-	}
-
+	
 	return FPCGCrc(Ar.GetCrc());
 }
 
diff --git a/Engine/Plugins/PCG/Source/PCG/Private/Elements/PCGCollapseElement.cpp b/Engine/Plugins/PCG/Source/PCG/Private/Elements/PCGCollapseElement.cpp
--- a/Engine/Plugins/PCG/Source/PCG/Private/Elements/PCGCollapseElement.cpp
+++ b/Engine/Plugins/PCG/Source/PCG/Private/Elements/PCGCollapseElement.cpp
@@ -40,10 +40,17 @@ TArray<FPCGPinProperties> UPCGConvertToPointDataSettings::InputPinProperties() c
 	return PinProperties;
 }
 
-bool FPCGCollapseElement::ExecuteInternal(FPCGContext* Context) const
+bool FPCGCollapseElement::ShouldComputeFullOutputDataCrc(FPCGContext* InContext) const
+{
+	FPCGCollapseContext* Context = static_cast<FPCGCollapseContext*>(InContext);
+	return Context && Context->bShouldComputeFullOutputDataCrc;
+}
+
+bool FPCGCollapseElement::ExecuteInternal(FPCGContext* InContext) const
 {
 	TRACE_CPUPROFILER_EVENT_SCOPE(FPCGCollapseElement::Execute);
-	check(Context);
+	check(InContext);
+	FPCGCollapseContext* Context = static_cast<FPCGCollapseContext*>(InContext);
 
 	TArray<FPCGTaggedData> Inputs = Context->InputData.GetInputsByPin(PCGPinConstants::DefaultInputLabel);
 	TArray<FPCGTaggedData>& Outputs = Context->OutputData.TaggedData;
@@ -61,6 +68,11 @@ bool FPCGCollapseElement::ExecuteInternal(FPCGContext* Context) const
 		{
 			// Currently we support collapsing to point data only, but at some point in the future that might be different
 			Output.Data = Cast<const UPCGSpatialData>(Input.Data)->ToPointData(Context);
+
+			if (Output.Data != Input.Data)
+			{
+				Context->bShouldComputeFullOutputDataCrc = true;
+			}
 		}
 		else if (const UPCGParamData* ParamData = Cast<UPCGParamData>(Input.Data))
 		{
@@ -85,6 +97,7 @@ bool FPCGCollapseElement::ExecuteInternal(FPCGContext* Context) const
 			}
 
 			Output.Data = PointData;
+			Context->bShouldComputeFullOutputDataCrc = true;
 		}
 	}
 
diff --git a/Engine/Plugins/PCG/Source/PCG/Private/PCGData.cpp b/Engine/Plugins/PCG/Source/PCG/Private/PCGData.cpp
--- a/Engine/Plugins/PCG/Source/PCG/Private/PCGData.cpp
+++ b/Engine/Plugins/PCG/Source/PCG/Private/PCGData.cpp
@@ -17,11 +17,6 @@
 
 #include UE_INLINE_GENERATED_CPP_BY_NAME(PCGData)
 
-static TAutoConsoleVariable<bool> CVarCachePropagateCrcThroughBooleanData(
-	TEXT("pcg.Cache.PropagateCrcThroughBooleanData"),
-	false,
-	TEXT("Whether intersection, union, difference, projection combine Crc values from operands. If false they fall back to using data UID."));
-
 UPCGData::UPCGData(const FObjectInitializer& ObjectInitializer)
 	: Super(ObjectInitializer)
 {
@@ -33,12 +28,16 @@ UPCGData::UPCGData(const FObjectInitializer& ObjectInitializer)
 
 FPCGCrc UPCGData::GetOrComputeCrc(bool bFullDataCrc) const
 {
-	if (Crc.IsValid())
+	const bool bComputeFullDataCrc = bFullDataCrc && SupportsFullDataCrc();
+	
+	// If we already valid Crc and it is a full data crc or we aren't being requested one, return the current Crc
+	if (Crc.IsValid() && (bIsFullDataCrc || !bComputeFullDataCrc))
 	{
 		return Crc;
 	}
 
-	Crc = ComputeCrc(bFullDataCrc);
+	Crc = ComputeCrc(bComputeFullDataCrc);
+	bIsFullDataCrc = bComputeFullDataCrc;
 
 	return Crc;
 }
@@ -58,11 +57,6 @@ void UPCGData::AddToCrc(FArchiveCrc32& Ar, bool bFullDataCrc) const
 	LastSelector.AddToCrc(Ar);
 }
 
-bool UPCGData::PropagateCrcThroughBooleanData() const
-{
-	return CVarCachePropagateCrcThroughBooleanData.GetValueOnAnyThread();
-}
-
 void UPCGData::VisitDataNetwork(TFunctionRef<void(const UPCGData*)> Action) const
 {
 	Action(this);
diff --git a/Engine/Plugins/PCG/Source/PCG/Public/Data/PCGPointData.h b/Engine/Plugins/PCG/Source/PCG/Public/Data/PCGPointData.h
--- a/Engine/Plugins/PCG/Source/PCG/Public/Data/PCGPointData.h
+++ b/Engine/Plugins/PCG/Source/PCG/Public/Data/PCGPointData.h
@@ -92,6 +92,10 @@ protected:
 	virtual UPCGSpatialData* CopyInternal(FPCGContext* Context) const override;
 	// ~End UPCGSpatialData interface
 
+	// ~Begin UPCGData interface
+	virtual bool SupportsFullDataCrc() const override { return true; }
+	// ~End UPCGData interface
+
 public:
 	/** Initializes a single point based on the given actor */
 	void InitializeFromActor(AActor* InActor, bool* bOutOptionalSanitizedTagAttributeName = nullptr);
diff --git a/Engine/Plugins/PCG/Source/PCG/Public/Elements/PCGCollapseElement.h b/Engine/Plugins/PCG/Source/PCG/Public/Elements/PCGCollapseElement.h
--- a/Engine/Plugins/PCG/Source/PCG/Public/Elements/PCGCollapseElement.h
+++ b/Engine/Plugins/PCG/Source/PCG/Public/Elements/PCGCollapseElement.h
@@ -2,6 +2,7 @@
 
 #pragma once
 
+#include "PCGContext.h"
 #include "PCGSettings.h"
 
 #include "PCGCollapseElement.generated.h"
@@ -48,11 +49,17 @@ protected:
 	//~End UPCGSettings interface
 };
 
-class FPCGCollapseElement : public IPCGElement
+struct FPCGCollapseContext : public FPCGContext
+{
+public:
+	bool bShouldComputeFullOutputDataCrc = false;
+};
+
+class FPCGCollapseElement : public IPCGElementWithCustomContext<FPCGCollapseContext>
 {
 public:
 	// Might be sampling spline/landscape or other external data, worth computing a full CRC in case we can halt change propagation/re-executions
-	virtual bool ShouldComputeFullOutputDataCrc(FPCGContext* Context) const override { return true; }
+	virtual bool ShouldComputeFullOutputDataCrc(FPCGContext* Context) const override;
 
 protected:
 	virtual bool ExecuteInternal(FPCGContext* Context) const override;
diff --git a/Engine/Plugins/PCG/Source/PCG/Public/PCGData.h b/Engine/Plugins/PCG/Source/PCG/Public/PCGData.h
--- a/Engine/Plugins/PCG/Source/PCG/Public/PCGData.h
+++ b/Engine/Plugins/PCG/Source/PCG/Public/PCGData.h
@@ -74,13 +74,16 @@ protected:
 	/** Adds this data to Crc. Fallback implementation writes object instance UID. */
 	virtual void AddToCrc(FArchiveCrc32& Ar, bool bFullDataCrc) const;
 
-	/** Whether intersection, union, difference combine Crc values from operands. If false they fall back to using data UID. */
-	bool PropagateCrcThroughBooleanData() const;
+	/** Whether the data distinguishes between regular crc and full crc */
+	virtual bool SupportsFullDataCrc() const { return false; }
 
 	void AddUIDToCrc(FArchiveCrc32& Ar) const;
 
 private:
 	void InitUID();
+	
+	/** If the Crc cache contains a full data Crc, if data type supports it */
+	mutable bool bIsFullDataCrc = false;
 
 	/** Serves unique ID values to instances of this object. */
 	static inline std::atomic<uint64> UIDCounter{ 1 };
diff --git a/Engine/Plugins/PCGInterops/PCGGeometryScriptInterop/Source/PCGGeometryScriptInterop/Public/Data/PCGDynamicMeshData.h b/Engine/Plugins/PCGInterops/PCGGeometryScriptInterop/Source/PCGGeometryScriptInterop/Public/Data/PCGDynamicMeshData.h
--- a/Engine/Plugins/PCGInterops/PCGGeometryScriptInterop/Source/PCGGeometryScriptInterop/Public/Data/PCGDynamicMeshData.h
+++ b/Engine/Plugins/PCGInterops/PCGGeometryScriptInterop/Source/PCGGeometryScriptInterop/Public/Data/PCGDynamicMeshData.h
@@ -64,6 +64,10 @@ public:
 	void InitializeDynamicMeshComponentFromData(UDynamicMeshComponent* InComponent) const;
 
 protected:
+	// ~Begin UPCGData interface
+	virtual bool SupportsFullDataCrc() const override { return true; }
+	// ~End UPCGData interface
+
 	// ~Begin UPCGSpatialData interface
 	virtual UPCGSpatialData* CopyInternal(FPCGContext* Context) const override;
 public:
