[PCG] Fix ensure happening when cancelling tasks (with pcg.GraphMultithreading enabled)
- It was possible for newly scheduled tasks to add successors on an already cancelled dependency, those scheduled tasks would also be cancelled but were leaving an empty TaskSuccessors entry behind

#rnx

diff --git a/Engine/Plugins/PCG/Source/PCG/Private/Graph/PCGGraphExecutor.cpp b/Engine/Plugins/PCG/Source/PCG/Private/Graph/PCGGraphExecutor.cpp
--- a/Engine/Plugins/PCG/Source/PCG/Private/Graph/PCGGraphExecutor.cpp
+++ b/Engine/Plugins/PCG/Source/PCG/Private/Graph/PCGGraphExecutor.cpp
@@ -858,22 +858,27 @@ void FPCGGraphExecutor::ExecuteTasksEnded()
 	// keep a cache of intermediate results.
 	if (GetNonScheduledRemainingTaskCount() == 0)
 	{
+		// Make sure no scheduling can happen here (and check condition again)
+		PCGGraphExecutor::TScopeLock ScopeLock(ScheduleLock);
+		if (GetNonScheduledRemainingTaskCount() == 0)
 		{
-			PCGGraphExecutor::TScopeLock ScopeLock(TasksLock);
-			if (!ensure(TaskSuccessors.IsEmpty()))
 			{
-				TaskSuccessors.Reset();
+				PCGGraphExecutor::TScopeLock ScopeChildLock(TasksLock);
+				if (!ensure(TaskSuccessors.IsEmpty()))
+				{
+					TaskSuccessors.Reset();
+				}
 			}
-		}
 
-		ClearResults();
+			ClearResults();
 
 #if WITH_EDITOR
-		// Call the notification update here to prevent it from sticking around - needed because we early out before this
-		UpdateGenerationNotification();
+			// Call the notification update here to prevent it from sticking around - needed because we early out before this
+			UpdateGenerationNotification();
 #endif
 
-		PCGGraphExecutionLogging::LogGraphExecuteFrameFinished();
+			PCGGraphExecutionLogging::LogGraphExecuteFrameFinished();
+		}
 	}
 
 	// Purge things from cache if memory usage is too high
@@ -2043,6 +2048,10 @@ void FPCGGraphExecutor::RemoveTaskFromInputSuccessorsNoLock(FPCGTaskId Cancelled
 		if (TSet<FPCGTaskId>* Successors = TaskSuccessors.Find(Input.TaskId))
 		{
 			Successors->Remove(CancelledTask);
+			if (Successors->IsEmpty())
+			{
+				TaskSuccessors.Remove(Input.TaskId);
+			}
 		}
 	}
 }
@@ -2270,7 +2279,6 @@ void FPCGGraphExecutor::StoreResults(FPCGTaskId InTaskId, const FPCGDataCollecti
 void FPCGGraphExecutor::ClearResults()
 {
 	{
-		PCGGraphExecutor::TScopeLock ScopeLock(ScheduleLock);
 		// Only reset if we have no more scheduled tasks, to avoid breaking dependencies
 		if (ScheduledTasks.IsEmpty())
 		{
