[PCG] Error out if a PCGBuilderSettings is specified but fails to load

#rnx

diff --git a/Engine/Plugins/PCG/Source/PCGEditor/Private/WorldPartitionBuilder/PCGWorldPartitionBuilder.cpp b/Engine/Plugins/PCG/Source/PCGEditor/Private/WorldPartitionBuilder/PCGWorldPartitionBuilder.cpp
--- a/Engine/Plugins/PCG/Source/PCGEditor/Private/WorldPartitionBuilder/PCGWorldPartitionBuilder.cpp
+++ b/Engine/Plugins/PCG/Source/PCGEditor/Private/WorldPartitionBuilder/PCGWorldPartitionBuilder.cpp
@@ -221,7 +221,7 @@ private:
 	std::atomic<bool> bErrorOccurred = false;
 };
 
-FPCGWorldPartitionBuilderArgs FPCGWorldPartitionBuilderArgs::InitializeFrom(const FPCGWorldPartitionCommandlineArgs& CommandlineArgs)
+bool FPCGWorldPartitionBuilderArgs::InitializeFrom(const FPCGWorldPartitionCommandlineArgs& CommandlineArgs, FPCGWorldPartitionBuilderArgs& OutArgs)
 {
 	FPCGWorldPartitionBuilderArgs BuilderArgs;
 	TSoftObjectPtr<UPCGBuilderSettings> SettingsPath = CommandlineArgs.SettingAsset.IsSet() ? CommandlineArgs.SettingAsset.GetValue() : GetDefault<UPCGEditorProjectSettings>()->DefaultBuilderSetting;
@@ -253,7 +253,8 @@ FPCGWorldPartitionBuilderArgs FPCGWorldPartitionBuilderArgs::InitializeFrom(cons
 		}
 		else
 		{
-			UE_LOG(LogPCGWorldPartitionBuilder, Warning, TEXT("Failed to load PCGBuilderSettings '%s'"), *SettingsPath.ToString());
+			UE_LOG(LogPCGWorldPartitionBuilder, Error, TEXT("Failed to load PCGBuilderSettings '%s'"), *SettingsPath.ToString());
+			return false;
 		}
 	}
 	
@@ -375,7 +376,8 @@ FPCGWorldPartitionBuilderArgs FPCGWorldPartitionBuilderArgs::InitializeFrom(cons
 		BuilderArgs.IterativeCellSize,
 		bFromCommandline ? FromCommandline : FromDefault);
 
-	return BuilderArgs;
+	OutArgs = BuilderArgs;
+	return true;
 }
 
 UPCGWorldPartitionBuilder::UPCGWorldPartitionBuilder(const FObjectInitializer& ObjectInitializer)
@@ -469,7 +471,10 @@ bool UPCGWorldPartitionBuilder::PreRun(UWorld* World, FPackageSourceControlHelpe
 		return false;
 	}
 
-	Args = FPCGWorldPartitionBuilderArgs::InitializeFrom(CommandlineArgs);
+	if (!FPCGWorldPartitionBuilderArgs::InitializeFrom(CommandlineArgs, Args))
+	{
+		return false;
+	}
 
 	// Set base class member
 	IterativeCellSize = Args.IterativeCellSize;
diff --git a/Engine/Plugins/PCG/Source/PCGEditor/Private/WorldPartitionBuilder/PCGWorldPartitionBuilder.h b/Engine/Plugins/PCG/Source/PCGEditor/Private/WorldPartitionBuilder/PCGWorldPartitionBuilder.h
--- a/Engine/Plugins/PCG/Source/PCGEditor/Private/WorldPartitionBuilder/PCGWorldPartitionBuilder.h
+++ b/Engine/Plugins/PCG/Source/PCGEditor/Private/WorldPartitionBuilder/PCGWorldPartitionBuilder.h
@@ -91,7 +91,7 @@ struct FPCGWorldPartitionBuilderArgs
 	/** Data layers that need to be unloaded when running builder. */
 	TArray<TObjectPtr<UDataLayerAsset>> ExcludedDataLayers;
 
-	static FPCGWorldPartitionBuilderArgs InitializeFrom(const FPCGWorldPartitionCommandlineArgs& CommandlineArgs);
+	static bool InitializeFrom(const FPCGWorldPartitionCommandlineArgs& CommandlineArgs, FPCGWorldPartitionBuilderArgs& OutArgs);
 };
 
 /**
