Added a get subgraph/recursion depth node

diff --git a/Engine/Plugins/PCG/Source/PCG/Private/Elements/PCGGetLoopIndex.cpp b/Engine/Plugins/PCG/Source/PCG/Private/Elements/PCGGetLoopIndex.cpp
--- a/Engine/Plugins/PCG/Source/PCG/Private/Elements/PCGGetLoopIndex.cpp
+++ b/Engine/Plugins/PCG/Source/PCG/Private/Elements/PCGGetLoopIndex.cpp
@@ -47,7 +47,7 @@ bool FPCGGetLoopIndexElement::ExecuteInternal(FPCGContext* Context) const
 
 	if (!ensure(Context->Stack))
 	{
-		PCGE_LOG(Error, LogOnly, LOCTEXT("ContextHasNoExecutionStack", "The execution context is malformed and has no context."));
+		PCGE_LOG(Error, LogOnly, LOCTEXT("ContextHasNoExecutionStack", "The execution context is malformed and has no call stack."));
 		return true;
 	}
 
diff --git a/Engine/Plugins/PCG/Source/PCG/Private/Elements/PCGGetSubgraphDepth.cpp b/Engine/Plugins/PCG/Source/PCG/Private/Elements/PCGGetSubgraphDepth.cpp
new file mode 100644
--- /dev/null
+++ b/Engine/Plugins/PCG/Source/PCG/Private/Elements/PCGGetSubgraphDepth.cpp
@@ -0,0 +1,118 @@
+// Copyright Epic Games, Inc. All Rights Reserved.
+
+#include "Elements/PCGGetSubgraphDepth.h"
+
+#include "PCGGraph.h"
+#include "PCGParamData.h"
+#include "PCGPin.h"
+#include "Graph/PCGStackContext.h"
+#include "Metadata/PCGMetadata.h"
+
+#include UE_INLINE_GENERATED_CPP_BY_NAME(PCGGetSubgraphDepth)
+
+#define LOCTEXT_NAMESPACE "PCGGetSubgraphDepthElement"
+
+namespace PCGGetSubgraphDepthConstants
+{
+	static const FName DepthAttributeName = TEXT("Depth");
+}
+
+#if WITH_EDITOR
+FText UPCGGetSubgraphDepthSettings::GetNodeTooltipText() const
+{
+	return LOCTEXT("GetSubgraphDepthTooltip", "Returns the call depth of this graph.");
+}
+
+EPCGChangeType UPCGGetSubgraphDepthSettings::GetChangeTypeForProperty(const FName& InPropertyName) const
+{
+	return Super::GetChangeTypeForProperty(InPropertyName) | EPCGChangeType::Cosmetic;
+}
+#endif // WITH_EDITOR
+
+TArray<FPCGPinProperties> UPCGGetSubgraphDepthSettings::OutputPinProperties() const
+{
+	TArray<FPCGPinProperties> PinProperties;
+	PinProperties.Emplace(PCGPinConstants::DefaultOutputLabel, EPCGDataType::Param);
+
+	return PinProperties;
+}
+
+FString UPCGGetSubgraphDepthSettings::GetAdditionalTitleInformation() const
+{
+	if (const UEnum* EnumPtr = StaticEnum<EPCGSubgraphDepthMode>())
+	{
+		return FText::Format(LOCTEXT("AdditionalTitle", "Get {0}"), EnumPtr->GetDisplayNameTextByValue(static_cast<int64>(Mode))).ToString();
+	}
+	else
+	{
+		return FString();
+	}
+}
+
+FPCGElementPtr UPCGGetSubgraphDepthSettings::CreateElement() const
+{
+	return MakeShared<FPCGGetSubgraphDepthElement>();
+}
+
+bool FPCGGetSubgraphDepthElement::ExecuteInternal(FPCGContext* Context) const
+{
+	TRACE_CPUPROFILER_EVENT_SCOPE(FPCGGetSubgraphDepthElement::Execute);
+
+	check(Context);
+
+	const UPCGGetSubgraphDepthSettings* Settings = Context->GetInputSettings<UPCGGetSubgraphDepthSettings>();
+	check(Settings);
+
+	bool bRecursive = Settings->Mode == EPCGSubgraphDepthMode::RecursiveDepth;
+	const UPCGGraph* ThisGraph = nullptr;
+
+	if (bRecursive)
+	{
+		ThisGraph = Context->Node ? Context->Node->GetGraph() : nullptr;
+
+		if (!ThisGraph)
+		{
+			PCGE_LOG(Error, GraphAndLog, LOCTEXT("NodeIsNotInAGraph", "Unable to retrieve originating graph, so impossible to find recursive depth."));
+			return true;
+		}
+	}
+
+	if (!ensure(Context->Stack))
+	{
+		PCGE_LOG(Error, LogOnly, LOCTEXT("ContextHasNoExecutionStack", "The execution context is malformed and has no call stack."));
+		return true;
+	}
+
+	const FPCGStack* Stack = Context->Stack;
+	const TArray<FPCGStackFrame>& StackFrames = Stack->GetStackFrames();
+	int Depth = 0;
+
+	for (const FPCGStackFrame& StackFrame : StackFrames)
+	{
+		if (const UPCGGraph* Graph = Cast<UPCGGraph>(StackFrame.Object.Get()))
+		{
+			if (!bRecursive || Graph == ThisGraph)
+			{
+				++Depth;
+			}
+		}
+	}
+
+	// Finally, since we count everything, including the root and self, we need to remove one.
+	Depth = FMath::Max(0, Depth - 1);
+
+	UPCGParamData* DepthParamData = FPCGContext::NewObject_AnyThread<UPCGParamData>(Context);
+	check(DepthParamData && DepthParamData->Metadata);
+
+	FPCGTaggedData& OutputData = Context->OutputData.TaggedData.Emplace_GetRef();
+	OutputData.Data = DepthParamData;
+
+	FPCGMetadataAttribute<int>* DepthAttribute = DepthParamData->Metadata->CreateAttribute<int>(PCGGetSubgraphDepthConstants::DepthAttributeName, Depth, /*bAllowInterpolation=*/false, /*bOverrideParent=*/false);
+	check(DepthAttribute);
+
+	DepthParamData->Metadata->AddEntry();
+
+	return true;
+}
+
+#undef LOCTEXT_NAMESPACE
\ No newline at end of file
diff --git a/Engine/Plugins/PCG/Source/PCG/Public/Elements/PCGGetSubgraphDepth.h b/Engine/Plugins/PCG/Source/PCG/Public/Elements/PCGGetSubgraphDepth.h
new file mode 100644
--- /dev/null
+++ b/Engine/Plugins/PCG/Source/PCG/Public/Elements/PCGGetSubgraphDepth.h
@@ -0,0 +1,55 @@
+// Copyright Epic Games, Inc. All Rights Reserved.
+
+#pragma once
+
+#include "PCGContext.h"
+#include "PCGSettings.h"
+
+#include "PCGGetSubgraphDepth.generated.h"
+
+UENUM(BlueprintType)
+enum class EPCGSubgraphDepthMode : uint8
+{
+	Depth UMETA(Tooltip="Depth of the dynamic subgraph with respect to the top level graph."),
+	RecursiveDepth UMETA(Tooltip="Recursive depth of the current subgraph, e.g. the number of times this graph is in the execution stack.")
+};
+
+/** Returns the current call or recursion depth in the execution stack. */
+UCLASS(MinimalAPI, BlueprintType, ClassGroup=(Procedural))
+class UPCGGetSubgraphDepthSettings : public UPCGSettings
+{
+	GENERATED_BODY()
+
+public:
+	//~Begin UPCGSettings interface implementation
+#if WITH_EDITOR
+	virtual FName GetDefaultNodeName() const override { return FName(TEXT("GetSubgraphDepth")); }
+	virtual FText GetDefaultNodeTitle() const override { return NSLOCTEXT("PCGGetSubgraphDepthElement", "NodeTitle", "Get Subgraph Depth"); }
+	virtual FText GetNodeTooltipText() const override;
+	virtual EPCGSettingsType GetType() const override { return EPCGSettingsType::Generic; }
+	virtual EPCGChangeType GetChangeTypeForProperty(const FName& InPropertyName) const override;
+#endif // WITH_EDITOR
+	
+	virtual bool HasFlippedTitleLines() const override { return true; }
+	virtual FString GetAdditionalTitleInformation() const override;
+
+protected:
+	virtual TArray<FPCGPinProperties> InputPinProperties() const override { return TArray<FPCGPinProperties>(); }
+	virtual TArray<FPCGPinProperties> OutputPinProperties() const override;
+
+	virtual FPCGElementPtr CreateElement() const override;
+	// ~End UPCGSettings interface implementation
+
+public:
+	UPROPERTY(BlueprintReadWrite, EditAnywhere, Category = Settings)
+	EPCGSubgraphDepthMode Mode = EPCGSubgraphDepthMode::Depth;
+};
+
+class FPCGGetSubgraphDepthElement : public IPCGElement
+{
+public:
+	virtual bool IsCacheable(const UPCGSettings* InSettings) const { return false; }
+
+protected:
+	virtual bool ExecuteInternal(FPCGContext* Context) const override;
+};
\ No newline at end of file
