[PCG] Improve World Ray Hit filtering

* Add support for Class filtering and filtering with a list of given actors
* Rename EPCGWorldQueryFilterByTag by EPCGWorldQueryFilter and added a redirector.
* Added more options and tooltip for the filtering. Backward compatible with the previous values.
* Add utility function to extract/set an array of values from a data and selector
* Factorized the Filter Actors pins + extraction into the common world query params
* Implemented the logic in all nodes that uses WRH params (WRHQ, World Volumetric Query, World Raycast and Pathfinding)

#rnx

diff --git a/Engine/Plugins/PCG/Config/BasePCG.ini b/Engine/Plugins/PCG/Config/BasePCG.ini
--- a/Engine/Plugins/PCG/Config/BasePCG.ini
+++ b/Engine/Plugins/PCG/Config/BasePCG.ini
@@ -101,3 +101,4 @@
 +EnumRedirects=(OldName="EPCGMedadataVectorOperation",NewName="EPCGMetadataVectorOperation")
 +EnumRedirects=(OldName="EPCGMedadataBitwiseOperation",NewName="EPCGMetadataBitwiseOperation")
 +EnumRedirects=(OldName="EPCGMedadataRotatorOperation",NewName="EPCGMetadataRotatorOperation")
++EnumRedirects=(OldName="EPCGWorldQueryFilterByTag",NewName="EPCGWorldQueryFilter")
diff --git a/Engine/Plugins/PCG/Source/PCG/Private/Data/PCGWorldData.cpp b/Engine/Plugins/PCG/Source/PCG/Private/Data/PCGWorldData.cpp
--- a/Engine/Plugins/PCG/Source/PCG/Private/Data/PCGWorldData.cpp
+++ b/Engine/Plugins/PCG/Source/PCG/Private/Data/PCGWorldData.cpp
@@ -13,6 +13,7 @@
 #include "Helpers/PCGBlueprintHelpers.h"
 #include "Helpers/PCGHelpers.h"
 #include "Helpers/PCGWorldQueryHelpers.h"
+#include "Metadata/Accessors/PCGAttributeAccessorHelpers.h"
 
 #include "LandscapeProxy.h"
 #include "Components/BrushComponent.h"
@@ -22,9 +23,40 @@
 
 #include UE_INLINE_GENERATED_CPP_BY_NAME(PCGWorldData)
 
+#define LOCTEXT_NAMESPACE "PCGWorldData"
+
+void FPCGWorldCommonQueryParams::AddFilterPinIfNeeded(TArray<FPCGPinProperties>& PinProperties) const
+{
+	if (ActorFilterFromInput != EPCGWorldQueryFilter::None)
+	{
+		FPCGPinProperties& FilterActors = PinProperties.Emplace_GetRef(PCGWorldRayHitConstants::FilterActorPinLabel, EPCGDataType::PointOrParam);
+#if WITH_EDITOR
+		FilterActors.Tooltip = LOCTEXT("ActorFilterFromInputTooltip", "All hit actors will be filtered against this list. Can be 1 list or N lists, N being the number of data in Origins pin.");
+#endif // WITH_EDITOR
+	}
+}
+
+bool FPCGWorldCommonQueryParams::ExtractActorFiltersIfNeeded(const UPCGData* InData, TArray<TSoftObjectPtr<AActor>>& OutArray, FPCGContext* InContext) const
+{
+	if (ActorFilterFromInput != EPCGWorldQueryFilter::None)
+	{
+		TArray<FSoftObjectPath> FilterActors;
+
+		if (!PCGAttributeAccessorHelpers::ExtractAllValues(InData, ActorFilterInputSource, FilterActors, InContext))
+		{
+			return false;
+		}
+
+		OutArray.Empty(FilterActors.Num());
+		Algo::Transform(FilterActors, OutArray, [](const FSoftObjectPath& ActorPath) { return TSoftObjectPtr<AActor>(ActorPath); });
+	}
+
+	return true;
+}
+
 void FPCGWorldCommonQueryParams::Initialize()
 {
-	if (ActorTagFilter == EPCGWorldQueryFilterByTag::NoTagFilter)
+	if (ActorTagFilter == EPCGWorldQueryFilter::None)
 	{
 		ParsedActorTagsList.Reset();
 	}
@@ -118,7 +150,7 @@ bool UPCGWorldVolumetricData::SamplePoint(const FTransform& InTransform, const F
 	TArray<FOverlapResult> Overlaps;
 	/*bool bOverlaps =*/ World->OverlapMultiByObjectType(Overlaps, InTransform.TransformPosition(InBounds.GetCenter()), InTransform.GetRotation(), ObjectQueryParams, CollisionShape, Params);
 
-	TOptional<FOverlapResult> Overlap = PCGWorldQueryHelpers::FilterOverlapResults(&QueryParams, OriginatingComponent, Overlaps);
+	TOptional<FOverlapResult> Overlap = PCGWorldQueryHelpers::FilterOverlapResults(&QueryParams, OriginatingComponent, Overlaps, FilterActors);
 
 	// If searched for overlap and found one, or didn't search and didn't find one, set the point and return true. Otherwise, return false.
 	if (Overlap.IsSet() == QueryParams.bSearchForOverlap)
@@ -232,7 +264,7 @@ bool UPCGWorldRayHitData::SamplePoint(const FTransform& InTransform, const FBox&
 	TArray<FHitResult> Hits;
 	if (World->LineTraceMultiByObjectType(Hits, RayOrigin, RayEnd, ObjectQueryParams, Params))
 	{
-		TOptional<FHitResult> HitResult = PCGWorldQueryHelpers::FilterRayHitResults(&QueryParams, OriginatingComponent, Hits);
+		TOptional<FHitResult> HitResult = PCGWorldQueryHelpers::FilterRayHitResults(&QueryParams, OriginatingComponent, Hits, FilterActors);
 
 		if (HitResult.IsSet())
 		{
@@ -296,4 +328,6 @@ UPCGSpatialData* UPCGWorldRayHitData::CopyInternal(FPCGContext* Context) const
 	NewData->QueryParams.Initialize();
 
 	return NewData;
-}
\ No newline at end of file
+}
+
+#undef LOCTEXT_NAMESPACE
\ No newline at end of file
diff --git a/Engine/Plugins/PCG/Source/PCG/Private/Elements/PCGPathfindingElement.cpp b/Engine/Plugins/PCG/Source/PCG/Private/Elements/PCGPathfindingElement.cpp
--- a/Engine/Plugins/PCG/Source/PCG/Private/Elements/PCGPathfindingElement.cpp
+++ b/Engine/Plugins/PCG/Source/PCG/Private/Elements/PCGPathfindingElement.cpp
@@ -90,6 +90,11 @@ TArray<FPCGPinProperties> UPCGPathfindingSettings::InputPinProperties() const
 		Properties.Emplace_GetRef(PCGPathfindingElement::Constants::GoalLocationsInputPinLabel, EPCGDataType::PointOrParam, /*bInAllowMultipleConnections=*/false, /*bAllowMultipleData=*/false).SetRequiredPin();
 	}
 
+	if (bUsePathTraces)
+	{
+		PathTraceParams.AddFilterPinIfNeeded(Properties);
+	}
+
 	return Properties;
 }
 
@@ -114,18 +119,23 @@ FPCGElementPtr UPCGPathfindingSettings::CreateElement() const
 }
 
 #if WITH_EDITOR
-EPCGChangeType UPCGPathfindingSettings::GetChangeTypeForProperty(const FName& InPropertyName) const
+EPCGChangeType UPCGPathfindingSettings::GetChangeTypeForProperty(FPropertyChangedEvent& PropertyChangedEvent) const
 {
 	EPCGChangeType LocalChangeType = EPCGChangeType::Cosmetic;
 
-	if (InPropertyName == GET_MEMBER_NAME_CHECKED(UPCGPathfindingSettings, bOutputAsSpline)
-		|| InPropertyName == GET_MEMBER_NAME_CHECKED(UPCGPathfindingSettings, bStartLocationsAsInput)
-		|| InPropertyName == GET_MEMBER_NAME_CHECKED(UPCGPathfindingSettings, bGoalLocationsAsInput))
+	const FName PropertyName = PropertyChangedEvent.GetPropertyName();
+	const FName MemberPropertyName = PropertyChangedEvent.GetMemberPropertyName();
+
+	if (PropertyName == GET_MEMBER_NAME_CHECKED(UPCGPathfindingSettings, bOutputAsSpline)
+		|| PropertyName == GET_MEMBER_NAME_CHECKED(UPCGPathfindingSettings, bStartLocationsAsInput)
+		|| PropertyName == GET_MEMBER_NAME_CHECKED(UPCGPathfindingSettings, bGoalLocationsAsInput)
+		|| (PropertyName == GET_MEMBER_NAME_CHECKED(UPCGPathfindingSettings, PathTraceParams) &&
+			MemberPropertyName == GET_MEMBER_NAME_CHECKED(FPCGWorldCommonQueryParams, ActorFilterFromInput)))
 	{
 		LocalChangeType |= EPCGChangeType::Structural;
 	}
 
-	return Super::GetChangeTypeForProperty(InPropertyName) | LocalChangeType;
+	return Super::GetChangeTypeForProperty(PropertyChangedEvent) | LocalChangeType;
 }
 #endif // WITH_EDITOR
 
@@ -146,13 +156,14 @@ bool FPCGPathfindingElement::PrepareDataInternal(FPCGContext* InContext) const
 	check(Settings);
 
 	const TArray<FPCGTaggedData> PointInputs = Context->InputData.GetInputsByPin(PCGPinConstants::DefaultInputLabel);
+	const TArray<FPCGTaggedData> FilterActorInputData = Context->InputData.GetInputsByPin(PCGWorldRayHitConstants::FilterActorPinLabel);
 
 	if (PointInputs.IsEmpty())
 	{
 		return true;
 	}
 
-	const EPCGTimeSliceInitResult ExecResult = Context->InitializePerExecutionState([&PointInputs, Settings](const ContextType* SlicedContext, ExecStateType& OutState)
+	const EPCGTimeSliceInitResult ExecResult = Context->InitializePerExecutionState([&PointInputs, &FilterActorInputData, Settings](const ContextType* SlicedContext, ExecStateType& OutState)
 	{
 		TRACE_CPUPROFILER_EVENT_SCOPE(FPCGClusterElement::PrepareData::InitializePerExecutionState);
 
@@ -160,6 +171,12 @@ bool FPCGPathfindingElement::PrepareDataInternal(FPCGContext* InContext) const
 		{
 			return EPCGTimeSliceInitResult::NoOperation;
 		}
+		
+		if (Settings->bUsePathTraces && FilterActorInputData.Num() > 1 && FilterActorInputData.Num() != PointInputs.Num())
+		{
+			PCGLog::LogErrorOnGraph(LOCTEXT("InvalidFilterActorInputCount", "Filter Actor input data count must be 1 or match the In pin data count."), SlicedContext);
+			return EPCGTimeSliceInitResult::NoOperation;
+		}
 
 		OutState.Settings = PCGSpatialAlgo::AStar::FSearchSettings
 		{
@@ -271,7 +288,7 @@ bool FPCGPathfindingElement::PrepareDataInternal(FPCGContext* InContext) const
 		return true;
 	}
 
-	Context->InitializePerIterationStates(PointInputs.Num(), [&PointInputs, Settings, Context](IterStateType& OutState, const ExecStateType& ExecutionState, const uint32 IterationIndex)
+	Context->InitializePerIterationStates(PointInputs.Num(), [&PointInputs, &FilterActorInputData, Settings, Context](IterStateType& OutState, const ExecStateType& ExecutionState, const uint32 IterationIndex)
 	{
 		TRACE_CPUPROFILER_EVENT_SCOPE(FPCGClusterElement::PrepareData::InitializePerIterationStates);
 
@@ -308,6 +325,22 @@ bool FPCGPathfindingElement::PrepareDataInternal(FPCGContext* InContext) const
 
 		if (Settings->bUsePathTraces)
 		{
+			const UPCGData* FilterActorsData = !FilterActorInputData.IsEmpty() ? FilterActorInputData[IterationIndex % FilterActorInputData.Num()].Data : nullptr;
+
+			// Accept only point or param data
+			if (FilterActorsData && !FilterActorsData->IsA<UPCGPointData>() && !FilterActorsData->IsA<UPCGParamData>())
+			{
+				PCGLog::InputOutput::LogTypedDataNotFoundWarning(EPCGDataType::PointOrParam, PCGWorldRayHitConstants::FilterActorPinLabel, Context);
+				return EPCGTimeSliceInitResult::NoOperation;
+			}
+
+			// --- Gather filtered actors ---
+			TArray<TSoftObjectPtr<AActor>> CachedFilterActors;
+			if (!Settings->PathTraceParams.ExtractActorFiltersIfNeeded(FilterActorsData, CachedFilterActors, Context))
+			{
+				return EPCGTimeSliceInitResult::NoOperation;
+			}
+			
 			if (UWorld* World = Context->SourceComponent.Get() ? Context->SourceComponent->GetWorld() : nullptr)
 			{
 				FPCGWorldRaycastQueryParams PathTraceParams = Settings->PathTraceParams;
@@ -318,12 +351,12 @@ bool FPCGPathfindingElement::PrepareDataInternal(FPCGContext* InContext) const
 				FCollisionQueryParams Params;
 				Params.bTraceComplex = PathTraceParams.bTraceComplex;
 
-				LineTraceTest = [World, ObjectQueryParams = MoveTemp(ObjectQueryParams), Params = MoveTemp(Params), OriginatingComponent, PathTraceParams = MoveTemp(PathTraceParams)](const FVector& StartPosition, const FVector& EndPosition) -> bool
+				LineTraceTest = [World, ObjectQueryParams = MoveTemp(ObjectQueryParams), Params = MoveTemp(Params), OriginatingComponent, PathTraceParams = MoveTemp(PathTraceParams), CachedFilterActors = MoveTemp(CachedFilterActors)](const FVector& StartPosition, const FVector& EndPosition) -> bool
 				{
 					TArray<FHitResult> OutHits;
 					if (World->LineTraceMultiByObjectType(OutHits, StartPosition, EndPosition, ObjectQueryParams, Params))
 					{
-						TOptional<FHitResult> HitResult = PCGWorldQueryHelpers::FilterRayHitResults(&PathTraceParams, OriginatingComponent, OutHits);
+						TOptional<FHitResult> HitResult = PCGWorldQueryHelpers::FilterRayHitResults(&PathTraceParams, OriginatingComponent, OutHits,CachedFilterActors);
 						return !HitResult.IsSet();
 					}
 					else
diff --git a/Engine/Plugins/PCG/Source/PCG/Private/Elements/PCGWorldQuery.cpp b/Engine/Plugins/PCG/Source/PCG/Private/Elements/PCGWorldQuery.cpp
--- a/Engine/Plugins/PCG/Source/PCG/Private/Elements/PCGWorldQuery.cpp
+++ b/Engine/Plugins/PCG/Source/PCG/Private/Elements/PCGWorldQuery.cpp
@@ -17,6 +17,30 @@
 
 #define LOCTEXT_NAMESPACE "PCGWorldQuery"
 
+namespace PCGWorldQuery
+{
+	template <typename WorldQueryDataType>
+	void ExtractActorFiltersIfNeeded(const FPCGWorldCommonQueryParams& QueryParams, FPCGContext* Context, WorldQueryDataType* WorldQueryData)
+	{
+		check(WorldQueryData);
+		
+		if (QueryParams.ActorFilterFromInput != EPCGWorldQueryFilter::None)
+		{
+			const TArray<FPCGTaggedData> ActorFilterTaggedData = Context->InputData.GetInputsByPin(PCGWorldRayHitConstants::FilterActorPinLabel);
+			if (ActorFilterTaggedData.Num() > 1)
+			{
+				PCGLog::InputOutput::LogFirstInputOnlyWarning(PCGWorldRayHitConstants::FilterActorPinLabel, Context);
+			}
+
+			const UPCGData* ActorFilterData = ActorFilterTaggedData.IsEmpty() ? nullptr : ActorFilterTaggedData[0].Data;
+			if (ActorFilterData && !QueryParams.ExtractActorFiltersIfNeeded(ActorFilterData, WorldQueryData->FilterActors, Context))
+			{
+				PCGLog::LogWarningOnGraph(LOCTEXT("FailExtractActorFilters", "Failed to extract actor filters."));
+			}
+		}
+	}
+}
+
 #if WITH_EDITOR
 FText UPCGWorldQuerySettings::GetNodeTooltipText() const
 {
@@ -24,6 +48,13 @@ FText UPCGWorldQuerySettings::GetNodeTooltipText() const
 }
 #endif
 
+TArray<FPCGPinProperties> UPCGWorldQuerySettings::InputPinProperties() const
+{
+	TArray<FPCGPinProperties> PinProperties;
+	QueryParams.AddFilterPinIfNeeded(PinProperties);
+	return PinProperties;
+}
+
 TArray<FPCGPinProperties> UPCGWorldQuerySettings::OutputPinProperties() const
 {
 	TArray<FPCGPinProperties> PinProperties;
@@ -37,6 +68,22 @@ FPCGElementPtr UPCGWorldQuerySettings::CreateElement() const
 	return MakeShared<FPCGWorldVolumetricQueryElement>();
 }
 
+#if WITH_EDITOR
+EPCGChangeType UPCGWorldQuerySettings::GetChangeTypeForProperty(FPropertyChangedEvent& PropertyChangedEvent) const
+{
+	EPCGChangeType Result = Super::GetChangeTypeForProperty(PropertyChangedEvent);
+
+	if (PropertyChangedEvent.GetPropertyName() == GET_MEMBER_NAME_CHECKED(UPCGWorldQuerySettings, QueryParams) &&
+		PropertyChangedEvent.GetMemberPropertyName() == GET_MEMBER_NAME_CHECKED(FPCGWorldCommonQueryParams, ActorFilterFromInput))
+	{
+		// This can add/remove a pin, so we need a structural change
+		Result |= EPCGChangeType::Structural;
+	}
+
+	return Result;
+}
+#endif // WITH_EDITOR
+
 bool FPCGWorldVolumetricQueryElement::ExecuteInternal(FPCGContext* Context) const
 {
 	TRACE_CPUPROFILER_EVENT_SCOPE(FPCGWorldVolumetricQueryElement::Execute);
@@ -57,6 +104,8 @@ bool FPCGWorldVolumetricQueryElement::ExecuteInternal(FPCGContext* Context) cons
 	Data->QueryParams.Initialize();
 	Data->OriginatingComponent = Context->SourceComponent;
 
+	PCGWorldQuery::ExtractActorFiltersIfNeeded(Settings->QueryParams, Context, Data);
+
 	FPCGTaggedData& Output = Context->OutputData.TaggedData.Emplace_GetRef();
 	Output.Data = Data;
 
@@ -70,6 +119,13 @@ FText UPCGWorldRayHitSettings::GetNodeTooltipText() const
 }
 #endif
 
+TArray<FPCGPinProperties> UPCGWorldRayHitSettings::InputPinProperties() const
+{
+	TArray<FPCGPinProperties> PinProperties;
+	QueryParams.AddFilterPinIfNeeded(PinProperties);
+	return PinProperties;
+}
+
 TArray<FPCGPinProperties> UPCGWorldRayHitSettings::OutputPinProperties() const
 {
 	TArray<FPCGPinProperties> PinProperties;
@@ -83,6 +139,22 @@ FPCGElementPtr UPCGWorldRayHitSettings::CreateElement() const
 	return MakeShared<FPCGWorldRayHitQueryElement>();
 }
 
+#if WITH_EDITOR
+EPCGChangeType UPCGWorldRayHitSettings::GetChangeTypeForProperty(FPropertyChangedEvent& PropertyChangedEvent) const
+{
+	EPCGChangeType Result = Super::GetChangeTypeForProperty(PropertyChangedEvent);
+
+	if (PropertyChangedEvent.GetPropertyName() == GET_MEMBER_NAME_CHECKED(UPCGWorldRayHitSettings, QueryParams) &&
+		PropertyChangedEvent.GetMemberPropertyName() == GET_MEMBER_NAME_CHECKED(FPCGWorldCommonQueryParams, ActorFilterFromInput))
+	{
+		// This can add/remove a pin, so we need a structural change
+		Result |= EPCGChangeType::Structural;
+	}
+
+	return Result;
+}
+#endif // WITH_EDITOR
+
 bool FPCGWorldRayHitQueryElement::ExecuteInternal(FPCGContext* Context) const
 {
 	TRACE_CPUPROFILER_EVENT_SCOPE(FPCGWorldRayHitQueryElement::Execute);
@@ -189,6 +261,8 @@ bool FPCGWorldRayHitQueryElement::ExecuteInternal(FPCGContext* Context) const
 		Data->QueryParams.bApplyMetadataFromLandscape = false;
 	}
 
+	PCGWorldQuery::ExtractActorFiltersIfNeeded(Settings->QueryParams, Context, Data);
+
 	FPCGTaggedData& Output = Context->OutputData.TaggedData.Emplace_GetRef();
 	Output.Data = Data;
 
diff --git a/Engine/Plugins/PCG/Source/PCG/Private/Elements/PCGWorldRaycast.cpp b/Engine/Plugins/PCG/Source/PCG/Private/Elements/PCGWorldRaycast.cpp
--- a/Engine/Plugins/PCG/Source/PCG/Private/Elements/PCGWorldRaycast.cpp
+++ b/Engine/Plugins/PCG/Source/PCG/Private/Elements/PCGWorldRaycast.cpp
@@ -57,6 +57,8 @@ TArray<FPCGPinProperties> UPCGWorldRaycastElementSettings::InputPinProperties()
 	BoundingShapeInputPin.Tooltip = LOCTEXT("BoundingShapeInputPinTooltip", "All projected points must be contained within this shape. If this input is omitted then bounds will be taken from the actor so that points are contained within actor bounds.");
 #endif // WITH_EDITOR
 
+	WorldQueryParams.AddFilterPinIfNeeded(PinProperties);
+
 	return PinProperties;
 }
 
@@ -70,6 +72,22 @@ FPCGElementPtr UPCGWorldRaycastElementSettings::CreateElement() const
 	return MakeShared<FPCGWorldRaycastElement>();
 }
 
+#if WITH_EDITOR
+EPCGChangeType UPCGWorldRaycastElementSettings::GetChangeTypeForProperty(FPropertyChangedEvent& PropertyChangedEvent) const
+{
+	EPCGChangeType Result = Super::GetChangeTypeForProperty(PropertyChangedEvent);
+
+	if (PropertyChangedEvent.GetPropertyName() == GET_MEMBER_NAME_CHECKED(UPCGWorldRaycastElementSettings, WorldQueryParams) &&
+		PropertyChangedEvent.GetMemberPropertyName() == GET_MEMBER_NAME_CHECKED(FPCGWorldCommonQueryParams, ActorFilterFromInput))
+	{
+		// This can add/remove a pin, so we need a structural change
+		Result |= EPCGChangeType::Structural;
+	}
+
+	return Result;
+}
+#endif // WITH_EDITOR
+
 bool FPCGWorldRaycastElement::PrepareDataInternal(FPCGContext* InContext) const
 {
 	ContextType* Context = static_cast<ContextType*>(InContext);
@@ -79,8 +97,9 @@ bool FPCGWorldRaycastElement::PrepareDataInternal(FPCGContext* InContext) const
 
 	const TArray<FPCGTaggedData> OriginsInputData = Context->InputData.GetInputsByPin(PCGWorldRaycastElementConstants::OriginsInputLabel);
 	const TArray<FPCGTaggedData> EndPointsInputData = Context->InputData.GetInputsByPin(PCGWorldRaycastElementConstants::EndPointsInputLabel);
+	const TArray<FPCGTaggedData> FilterActorInputData = Context->InputData.GetInputsByPin(PCGWorldRayHitConstants::FilterActorPinLabel);
 
-	const EPCGTimeSliceInitResult ExecutionInitResult = Context->InitializePerExecutionState([&OriginsInputData, &EndPointsInputData, Settings](ContextType* Context, ExecStateType& OutState)
+	const EPCGTimeSliceInitResult ExecutionInitResult = Context->InitializePerExecutionState([&OriginsInputData, &EndPointsInputData, &FilterActorInputData, Settings](ContextType* Context, ExecStateType& OutState)
 	{
 		// With no origins data, early out
 		if (OriginsInputData.IsEmpty())
@@ -105,6 +124,12 @@ bool FPCGWorldRaycastElement::PrepareDataInternal(FPCGContext* InContext) const
 			PCGLog::LogErrorOnGraph(LOCTEXT("InvalidEndPointDataCount", "End point input data count must match the Origins input data count."), Context);
 			return EPCGTimeSliceInitResult::NoOperation;
 		}
+		
+		if (FilterActorInputData.Num() > 1 && FilterActorInputData.Num() != OriginsInputData.Num())
+		{
+			PCGLog::LogErrorOnGraph(LOCTEXT("InvalidFilterActorInputCount", "Filter Actor input data count must be 1 or match the Origins input data count."), Context);
+			return EPCGTimeSliceInitResult::NoOperation;
+		}
 
 		OutState.CollisionQueryParams.bTraceComplex = Settings->WorldQueryParams.bTraceComplex;
 		OutState.CollisionQueryParams.bReturnPhysicalMaterial = Settings->WorldQueryParams.bGetReferenceToPhysicalMaterial;
@@ -119,44 +144,33 @@ bool FPCGWorldRaycastElement::PrepareDataInternal(FPCGContext* InContext) const
 		return true;
 	}
 
-	Context->InitializePerIterationStates(OriginsInputData.Num(), [Context, Settings, OriginsInputData, EndPointsInputData](IterStateType& OutState, const ExecStateType&, int32 Index)
+	Context->InitializePerIterationStates(OriginsInputData.Num(), [Context, Settings, &OriginsInputData, &EndPointsInputData, &FilterActorInputData](IterStateType& OutState, const ExecStateType&, int32 Index)
 	{
 		const UPCGData* OriginsData = OriginsInputData[Index].Data;
+		const UPCGData* FilterActorsData = !FilterActorInputData.IsEmpty() ? FilterActorInputData[Index % FilterActorInputData.Num()].Data : nullptr;
 
 		// Accept only point or param data
-		if (!OriginsData || !OriginsData->ConstMetadata())
+		if (!OriginsData || (!OriginsData->IsA<UPCGPointData>() && !OriginsData->IsA<UPCGParamData>()))
 		{
-			PCGLog::LogErrorOnGraph(FText::Format(LOCTEXT("NoDataFoundOnOriginsInput", "No valid data found on Origins input {0}"), Index), Context);
+			PCGLog::InputOutput::LogTypedDataNotFoundWarning(EPCGDataType::PointOrParam, PCGWorldRaycastElementConstants::OriginsInputLabel, Context);
 			return EPCGTimeSliceInitResult::NoOperation;
 		}
 
-		// Given a selector, will populate a given array view with attributes
-		auto CacheArray = [Settings, Context]<typename Type>(const UPCGData* InData, const FPCGAttributePropertyInputSelector& InSelector, TArray<Type>& OutArray)
+		// Accept only point or param data
+		if (FilterActorsData && !FilterActorsData->IsA<UPCGPointData>() && !FilterActorsData->IsA<UPCGParamData>())
 		{
-			const FPCGAttributePropertyInputSelector Selector = InSelector.CopyAndFixLast(InData);
-
-			TUniquePtr<const IPCGAttributeAccessor> Accessor = PCGAttributeAccessorHelpers::CreateConstAccessor(InData, Selector);
-			TUniquePtr<const IPCGAttributeAccessorKeys> Keys = PCGAttributeAccessorHelpers::CreateConstKeys(InData, Selector);
-			if (!Accessor.IsValid() || !Keys.IsValid())
-			{
-				PCGLog::LogWarningOnGraph(FText::Format(LOCTEXT("AttributeDoesNotExist", "Attribute/property '{0}' does not exist - skipped."), Selector.GetDisplayText()), Context);
-				return false;
-			}
-
-			OutArray.Empty();
-			OutArray.SetNumUninitialized(Keys->GetNum());
-
-			if (!Accessor->GetRange<Type>(OutArray, 0, *Keys, EPCGAttributeAccessorFlags::AllowBroadcast | EPCGAttributeAccessorFlags::AllowConstructible))
-			{
-				PCGLog::LogErrorOnGraph(FText::Format(LOCTEXT("InvalidAttributeValues", "Could not get the range of values for attribute/property '{0}'."), Selector.GetDisplayText()), Context);
-				return false;
-			}
+			PCGLog::InputOutput::LogTypedDataNotFoundWarning(EPCGDataType::PointOrParam, PCGWorldRayHitConstants::FilterActorPinLabel, Context);
+			return EPCGTimeSliceInitResult::NoOperation;
+		}
 
-			return true;
-		};
+		// --- Gather filtered actors ---
+		if (!Settings->WorldQueryParams.ExtractActorFiltersIfNeeded(FilterActorsData, OutState.CachedFilterActors, Context))
+		{
+			return EPCGTimeSliceInitResult::NoOperation;
+		}
 
 		// --- Find the ray origin ---
-		if (!CacheArray(OriginsData, Settings->OriginInputAttribute, OutState.CachedRayOrigins))
+		if (!PCGAttributeAccessorHelpers::ExtractAllValues(OriginsData, Settings->OriginInputAttribute, OutState.CachedRayOrigins, Context))
 		{
 			return EPCGTimeSliceInitResult::NoOperation;
 		}
@@ -173,7 +187,7 @@ bool FPCGWorldRaycastElement::PrepareDataInternal(FPCGContext* InContext) const
 			}
 
 			TArray<FVector> EndPoints;
-			if (!CacheArray(EndPointsData, Settings->EndPointAttribute, EndPoints))
+			if (!PCGAttributeAccessorHelpers::ExtractAllValues(EndPointsData, Settings->EndPointAttribute, EndPoints, Context))
 			{
 				return EPCGTimeSliceInitResult::NoOperation;
 			}
@@ -200,9 +214,7 @@ bool FPCGWorldRaycastElement::PrepareDataInternal(FPCGContext* InContext) const
 		}
 		else if (Settings->bOverrideRayDirections) // Use a selector to determine the (non-normalized) directions
 		{
-			const FPCGAttributePropertyInputSelector DirectionSelector = Settings->RayDirectionAttribute.CopyAndFixLast(OriginsData);
-
-			if (!CacheArray(OriginsData, DirectionSelector, OutState.CachedRayVectors))
+			if (!PCGAttributeAccessorHelpers::ExtractAllValues(OriginsData, Settings->RayDirectionAttribute, OutState.CachedRayVectors, Context))
 			{
 				return EPCGTimeSliceInitResult::NoOperation;
 			}
@@ -234,7 +246,7 @@ bool FPCGWorldRaycastElement::PrepareDataInternal(FPCGContext* InContext) const
 			{
 				const FPCGAttributePropertyInputSelector LengthSelector = Settings->RayLengthAttribute.CopyAndFixLast(OriginsData);
 
-				if (!CacheArray(OriginsData, LengthSelector, OutState.CachedRayLengths))
+				if (!PCGAttributeAccessorHelpers::ExtractAllValues(OriginsData, LengthSelector, OutState.CachedRayLengths, Context))
 				{
 					return EPCGTimeSliceInitResult::NoOperation;
 				}
@@ -333,7 +345,7 @@ bool FPCGWorldRaycastElement::ExecuteInternal(FPCGContext* InContext) const
 
 			TArray<FHitResult> Hits;
 			World->LineTraceMultiByObjectType(Hits, Origin, Origin + RayVector, ExecState.CollisionObjectQueryParams, ExecState.CollisionQueryParams);
-			const TOptional<FHitResult>& HitResult = PCGWorldQueryHelpers::FilterRayHitResults(&WorldQueryParams, Context->SourceComponent, Hits);
+			const TOptional<FHitResult>& HitResult = PCGWorldQueryHelpers::FilterRayHitResults(&WorldQueryParams, Context->SourceComponent, Hits, IterState.CachedFilterActors);
 
 			// If there's a hit, and it's within bounds, keep the point
 			if (HitResult.IsSet() && (Settings->bUnbounded || PCGHelpers::IsInsideBounds(ExecState.Bounds, HitResult.GetValue().ImpactPoint)))
diff --git a/Engine/Plugins/PCG/Source/PCG/Private/Helpers/PCGWorldQueryHelpers.cpp b/Engine/Plugins/PCG/Source/PCG/Private/Helpers/PCGWorldQueryHelpers.cpp
--- a/Engine/Plugins/PCG/Source/PCG/Private/Helpers/PCGWorldQueryHelpers.cpp
+++ b/Engine/Plugins/PCG/Source/PCG/Private/Helpers/PCGWorldQueryHelpers.cpp
@@ -10,6 +10,7 @@
 #include "Helpers/PCGHelpers.h"
 
 #include "LandscapeProxy.h"
+#include "Algo/AnyOf.h"
 #include "Components/BrushComponent.h"
 #include "Components/PrimitiveComponent.h"
 #include "Engine/HitResult.h"
@@ -38,7 +39,8 @@ namespace PCGWorldQueryHelpers
 	bool FilterCommonQueryResults(
 		const FPCGWorldCommonQueryParams* QueryParams,
 		const UPrimitiveComponent* TriggeredComponent,
-		const TWeakObjectPtr<UPCGComponent> OriginatingComponent)
+		const TWeakObjectPtr<UPCGComponent> OriginatingComponent,
+		const TArrayView<const TSoftObjectPtr<AActor>> FilteredActorReferences)
 	{
 		check(QueryParams && TriggeredComponent);
 
@@ -73,26 +75,52 @@ namespace PCGWorldQueryHelpers
 		}
 
 		// Additional filter as provided in the QueryParams base class
-		if (QueryParams->ActorTagFilter != EPCGWorldQueryFilterByTag::NoTagFilter)
+		if (QueryParams->ActorTagFilter != EPCGWorldQueryFilter::None
+			|| QueryParams->ActorClassFilter != EPCGWorldQueryFilter::None
+			|| QueryParams->ActorFilterFromInput != EPCGWorldQueryFilter::None)
 		{
-			if (AActor* Actor = TriggeredComponent->GetOwner())
+			AActor* Actor = TriggeredComponent->GetOwner();
+			if (!Actor)
 			{
-				bool bFoundMatch = false;
-				for (const FName& Tag : Actor->Tags)
-				{
-					if (QueryParams->ParsedActorTagsList.Contains(Tag))
-					{
-						bFoundMatch = true;
-						break;
-					}
-				}
+				return false;
+			}
+			
+			bool bSoftInclude = false;
+			bool bHasIncludeFilter = false;
+			bool bHardExclude = false;
+
+			// Exclude and requires will force the exclude if the criteria doesn't match (Hard exclude)
+			// Include and requires will have a soft include if the criteria does match
+			// Also keep track if we have an include filter. If we didn't, and we have not a hard exclude, it should be included (if we had just exclude filters).
+			auto UpdateIncludeExcludeStates = [&bSoftInclude, &bHardExclude, &bHasIncludeFilter](const EPCGWorldQueryFilter Filter, const bool bFoundMatch)
+			{
+				bHardExclude |= ((Filter == EPCGWorldQueryFilter::Exclude && bFoundMatch) || (Filter == EPCGWorldQueryFilter::Require && !bFoundMatch));
+				
+				const bool bIsIncludeFilter = (Filter == EPCGWorldQueryFilter::Include) || (Filter == EPCGWorldQueryFilter::Require);
+				bSoftInclude |= (bIsIncludeFilter && bFoundMatch);
+				bHasIncludeFilter |= bIsIncludeFilter;
+			};
+
+			if (QueryParams->ActorTagFilter != EPCGWorldQueryFilter::None)
+			{
+				const bool bFoundMatch = Algo::AnyOf(Actor->Tags, [QueryParams](const FName Tag) { return QueryParams->ParsedActorTagsList.Contains(Tag); });
+				UpdateIncludeExcludeStates(QueryParams->ActorTagFilter, bFoundMatch);
+			}
 
-				if (bFoundMatch != (QueryParams->ActorTagFilter == EPCGWorldQueryFilterByTag::IncludeTagged))
-				{
-					return false;
-				}
+			// No need to check if we already know we are going to discard it.
+			if (!bHardExclude && QueryParams->ActorClassFilter != EPCGWorldQueryFilter::None)
+			{
+				const bool bFoundMatch = Actor->GetClass() && Actor->GetClass()->IsChildOf(QueryParams->ActorClass);
+				UpdateIncludeExcludeStates(QueryParams->ActorClassFilter, bFoundMatch);
 			}
-			else if (QueryParams->ActorTagFilter == EPCGWorldQueryFilterByTag::IncludeTagged)
+
+			if (!bHardExclude && QueryParams->ActorFilterFromInput != EPCGWorldQueryFilter::None)
+			{
+				const bool bFoundMatch = Algo::AnyOf(FilteredActorReferences, [Actor](const TSoftObjectPtr<AActor>& ActorRef) { return ActorRef == Actor; });
+				UpdateIncludeExcludeStates(QueryParams->ActorFilterFromInput, bFoundMatch);
+			}
+
+			if (bHardExclude || (!bSoftInclude && bHasIncludeFilter))
 			{
 				return false;
 			}
@@ -119,7 +147,8 @@ namespace PCGWorldQueryHelpers
 	TOptional<FHitResult> FilterRayHitResults(
 		const FPCGWorldRaycastQueryParams* QueryParams,
 		const TWeakObjectPtr<UPCGComponent> OriginatingComponent,
-		const TArray<FHitResult>& HitResults)
+		const TArray<FHitResult>& HitResults,
+		const TArrayView<const TSoftObjectPtr<AActor>> FilteredActorReferences)
 	{
 		check(QueryParams);
 
@@ -127,7 +156,7 @@ namespace PCGWorldQueryHelpers
 		{
 			const UPrimitiveComponent* HitComponent = Hit.GetComponent();
 
-			if (!FilterCommonQueryResults(QueryParams, HitComponent, OriginatingComponent))
+			if (!FilterCommonQueryResults(QueryParams, HitComponent, OriginatingComponent, FilteredActorReferences))
 			{
 				continue;
 			}
@@ -151,7 +180,8 @@ namespace PCGWorldQueryHelpers
 	TOptional<FOverlapResult> FilterOverlapResults(
 		const FPCGWorldVolumetricQueryParams* QueryParams,
 		const TWeakObjectPtr<UPCGComponent> OriginatingComponent,
-		const TArray<FOverlapResult>& OverlapResults)
+		const TArray<FOverlapResult>& OverlapResults,
+		const TArrayView<const TSoftObjectPtr<AActor>> FilteredActorReferences)
 	{
 		check(QueryParams);
 
@@ -159,7 +189,7 @@ namespace PCGWorldQueryHelpers
 		{
 			const UPrimitiveComponent* OverlappedComponent = Overlap.GetComponent();
 
-			if (!FilterCommonQueryResults(QueryParams, OverlappedComponent, OriginatingComponent))
+			if (!FilterCommonQueryResults(QueryParams, OverlappedComponent, OriginatingComponent, FilteredActorReferences))
 			{
 				continue;
 			}
diff --git a/Engine/Plugins/PCG/Source/PCG/Private/Utils/PCGLogErrors.cpp b/Engine/Plugins/PCG/Source/PCG/Private/Utils/PCGLogErrors.cpp
--- a/Engine/Plugins/PCG/Source/PCG/Private/Utils/PCGLogErrors.cpp
+++ b/Engine/Plugins/PCG/Source/PCG/Private/Utils/PCGLogErrors.cpp
@@ -44,11 +44,13 @@ namespace PCGLog::Metadata
 {
 	namespace Format
 	{
-		const FTextFormat CreateAccessorFailure = LOCTEXT("CreateAccessorFailure", "Couldn't create accessor. Attribute '{0}' was not found.");
-		const FTextFormat CreateAttributeFailure = LOCTEXT("CreateAttributeFailure", "Couldn't create attribute '{0}' of type '{1}'.");
-		const FTextFormat GetAttributeFailure = LOCTEXT("GetAttributeFailure", "Couldn't retrieve attribute '{0}' value.");
-		const FTextFormat GetTypedAttributeFailure = LOCTEXT("GetTypedAttributeFailure", "Couldn't retrieve attribute '{0}' value. Expected type: {1}, Actual Type: {2}.");
-		const FTextFormat GetTypedAttributeFailureNoAccessor = LOCTEXT("GetTypedAttributeFailureNoAccessor", "Couldn't retrieve attribute '{0}' value of type: '{1}.");
+		const FTextFormat CreateAccessorFailure = LOCTEXT("CreateAccessorFailure", "Could not create accessor. Attribute '{0}' was not found.");
+		const FTextFormat CreateAttributeFailure = LOCTEXT("CreateAttributeFailure", "Could not create attribute '{0}' of type '{1}'.");
+		const FTextFormat GetAttributeFailure = LOCTEXT("GetAttributeFailure", "Could not retrieve attribute '{0}' value.");
+		const FTextFormat GetTypedAttributeFailure = LOCTEXT("GetTypedAttributeFailure", "Could not retrieve attribute '{0}' value. Expected type: {1}, Actual Type: {2}.");
+		const FTextFormat GetTypedAttributeFailureNoAccessor = LOCTEXT("GetTypedAttributeFailureNoAccessor", "Could not retrieve attribute '{0}' value of type: '{1}.");
+		const FTextFormat SetTypedAttributeFailure = LOCTEXT("SetTypedAttributeFailure", "Could not set value in attribute '{0}'. Attribute type: {1}, Value Type: {2}.");
+		const FTextFormat SetTypedAttributeFailureNoAccessor = LOCTEXT("SetTypedAttributeFailureNoAccessor", "Could not set value in attribute '{0}' value of type: '{1}.");
 		const FTextFormat IncomparableTypesFailure = LOCTEXT("IncomparableTypesFailure", "Attributes '{0}' and '{1}' are incomparable. Ensure they are either of the same or compatible types.");
 	}
 
diff --git a/Engine/Plugins/PCG/Source/PCG/Public/Data/PCGWorldData.h b/Engine/Plugins/PCG/Source/PCG/Public/Data/PCGWorldData.h
--- a/Engine/Plugins/PCG/Source/PCG/Public/Data/PCGWorldData.h
+++ b/Engine/Plugins/PCG/Source/PCG/Public/Data/PCGWorldData.h
@@ -17,13 +17,27 @@ class UPCGMetadata;
 class UPCGComponent;
 
 UENUM()
-enum class EPCGWorldQueryFilterByTag
+enum class UE_DEPRECATED(5.6, "Not used anymore, replaced by EPCGWorldQueryFilter.") EPCGWorldQueryFilterByTag
 {
 	NoTagFilter,
 	IncludeTagged,
 	ExcludeTagged
 };
 
+UENUM(Blueprintable)
+enum class EPCGWorldQueryFilter : uint8
+{
+	None    UMETA(Tooltip = "Filter disabled"),
+	Include UMETA(Tooltip = "Includes the actor if no other filter explicitly filters it out (either by exclusion or by requiring an unmet criteria)."),
+	Exclude UMETA(Tooltip = "Always exclude an actor if it matches this filter."),
+	Require UMETA(Tooltip = "Requires the actor to match to this filter to be included."),
+
+	// Deprecation values to support deserializing from EPCGWorldQueryFilterByTag, since it is deserialized by name. To be removed when EPCGWorldQueryFilterByTag is removed.
+	NoTagFilter = None UMETA(Hidden),
+	IncludeTagged = Include UMETA(Hidden),
+	ExcludeTagged = Exclude UMETA(Hidden)
+};
+
 UENUM()
 enum class EPCGWorldQuerySelectLandscapeHits : uint8
 {
@@ -36,6 +50,8 @@ namespace PCGWorldRayHitConstants
 {
 	UE_DEPRECATED(5.5, "Please use 'PCGWorldQueryConstants::PhysicalMaterialReferenceAttribute' instead.")
 	const FName PhysicalMaterialReferenceAttribute = PCGWorldQueryConstants::PhysicalMaterialReferenceAttribute;
+	
+	const FName FilterActorPinLabel = TEXT("FilterActors");
 }
 
 USTRUCT(BlueprintType)
@@ -62,10 +78,24 @@ struct FPCGWorldCommonQueryParams
 	bool bTraceComplex = false;
 
 	UPROPERTY(BlueprintReadWrite, EditAnywhere, Category = "Data|Filtering", meta = (PCG_Overridable))
-	EPCGWorldQueryFilterByTag ActorTagFilter = EPCGWorldQueryFilterByTag::NoTagFilter;
+	EPCGWorldQueryFilter ActorTagFilter = EPCGWorldQueryFilter::None;
 
-	UPROPERTY(BlueprintReadWrite, EditAnywhere, Category = "Data|Filtering", meta = (PCG_Overridable, EditCondition = "ActorTagFilter != EPCGWorldQueryFilterByTag::NoTagFilter"))
+	UPROPERTY(BlueprintReadWrite, EditAnywhere, Category = "Data|Filtering", meta = (PCG_Overridable, EditCondition = "ActorTagFilter != EPCGWorldQueryFilter::None"))
 	FString ActorTagsList;
+	
+	UPROPERTY(BlueprintReadWrite, EditAnywhere, Category = "Data|Filtering", meta = (PCG_Overridable))
+	EPCGWorldQueryFilter ActorClassFilter = EPCGWorldQueryFilter::None;
+	
+	UPROPERTY(BlueprintReadWrite, EditAnywhere, Category = "Data|Filtering", meta = (PCG_Overridable, EditCondition = "ActorClassFilter != EPCGWorldQueryFilter::None"))
+	TSubclassOf<AActor> ActorClass;
+
+	/** Will add an input pin to pass a list of actor references for filtering if this value is not set to None. */
+	UPROPERTY(BlueprintReadWrite, EditAnywhere, Category = "Data|Filtering", meta = (PCG_Overridable))
+	EPCGWorldQueryFilter ActorFilterFromInput = EPCGWorldQueryFilter::None;
+
+	/** Input source for the attribute to read from the Filter Actor pin. */
+	UPROPERTY(BlueprintReadWrite, EditAnywhere, Category = "Data|Filtering", meta = (PCG_Overridable, EditCondition = "ActorFilterFromInput != EPCGWorldQueryFilter::None"))
+	FPCGAttributePropertyInputSelector ActorFilterInputSource;
 
 	UPROPERTY(BlueprintReadWrite, EditAnywhere, Category = "Data|Filtering", meta = (PCG_Overridable, PCG_OverrideAliases="bIgnoreLandscapeHits"))
 	EPCGWorldQuerySelectLandscapeHits SelectLandscapeHits = EPCGWorldQuerySelectLandscapeHits::Include;
@@ -85,6 +115,12 @@ struct FPCGWorldCommonQueryParams
 	UPROPERTY()
 	TSet<FName> ParsedActorTagsList;
 
+	/** Utility function to add a filter pin with the right label and tooltip to the pin properties. */
+	PCG_API void AddFilterPinIfNeeded(TArray<FPCGPinProperties>& PinProperties) const;
+
+	/** Utility function to extract the Actor Filters from the incoming data using the selector in this struct. Return true if the extraction succeeded. */
+	PCG_API bool ExtractActorFiltersIfNeeded(const UPCGData* InData, TArray<TSoftObjectPtr<AActor>>& OutArray, FPCGContext* InContext = nullptr) const;
+
 protected:
 	/** Sets up the data we need to efficiently perform the queries */
 	void Initialize();
@@ -160,7 +196,7 @@ struct FPCGWorldRaycastQueryParams : public FPCGWorldCommonQueryParams
 };
 
 template<>
-struct TStructOpsTypeTraits<FPCGWorldRaycastQueryParams> : public TStructOpsTypeTraitsBase2<FPCGWorldRaycastQueryParams>
+struct TStructOpsTypeTraits<FPCGWorldRaycastQueryParams> : public TStructOpsTypeTraits<FPCGWorldCommonQueryParams>
 {
 	enum
 	{
@@ -182,7 +218,7 @@ struct FPCGWorldVolumetricQueryParams : public FPCGWorldCommonQueryParams
 };
 
 template<>
-struct TStructOpsTypeTraits<FPCGWorldVolumetricQueryParams> : public TStructOpsTypeTraitsBase2<FPCGWorldVolumetricQueryParams>
+struct TStructOpsTypeTraits<FPCGWorldVolumetricQueryParams> : public TStructOpsTypeTraits<FPCGWorldCommonQueryParams>
 {
 	enum
 	{
@@ -216,7 +252,7 @@ struct FPCGWorldRayHitQueryParams : public FPCGWorldRaycastQueryParams
 };
 
 template<>
-struct TStructOpsTypeTraits<FPCGWorldRayHitQueryParams> : public TStructOpsTypeTraitsBase2<FPCGWorldRayHitQueryParams>
+struct TStructOpsTypeTraits<FPCGWorldRayHitQueryParams> : public TStructOpsTypeTraits<FPCGWorldCommonQueryParams>
 {
 	enum
 	{
@@ -257,6 +293,9 @@ public:
 
 	UPROPERTY(BlueprintReadWrite, EditAnywhere, Category = Data, meta = (ShowOnlyInnerProperties))
 	FPCGWorldVolumetricQueryParams QueryParams;
+	
+	UPROPERTY()
+    TArray<TSoftObjectPtr<AActor>> FilterActors;
 };
 
 /** Executes collision queries against world collision. */
@@ -301,4 +340,7 @@ public:
 
 	UPROPERTY(BlueprintReadWrite, EditAnywhere, Category = Data, meta = (ShowOnlyInnerProperties))
 	FPCGWorldRayHitQueryParams QueryParams;
+	
+	UPROPERTY()
+	TArray<TSoftObjectPtr<AActor>> FilterActors;
 };
diff --git a/Engine/Plugins/PCG/Source/PCG/Public/Elements/PCGPathfindingElement.h b/Engine/Plugins/PCG/Source/PCG/Public/Elements/PCGPathfindingElement.h
--- a/Engine/Plugins/PCG/Source/PCG/Public/Elements/PCGPathfindingElement.h
+++ b/Engine/Plugins/PCG/Source/PCG/Public/Elements/PCGPathfindingElement.h
@@ -51,7 +51,7 @@ protected:
 	virtual TArray<FPCGPinProperties> OutputPinProperties() const override;
 	virtual FPCGElementPtr CreateElement() const override;
 #if WITH_EDITOR
-	virtual EPCGChangeType GetChangeTypeForProperty(const FName& InPropertyName) const override;
+	virtual EPCGChangeType GetChangeTypeForProperty(FPropertyChangedEvent& PropertyChangedEvent) const override;
 #endif
 	//~End UPCGSettings interface
 
diff --git a/Engine/Plugins/PCG/Source/PCG/Public/Elements/PCGWorldQuery.h b/Engine/Plugins/PCG/Source/PCG/Public/Elements/PCGWorldQuery.h
--- a/Engine/Plugins/PCG/Source/PCG/Public/Elements/PCGWorldQuery.h
+++ b/Engine/Plugins/PCG/Source/PCG/Public/Elements/PCGWorldQuery.h
@@ -22,10 +22,14 @@ public:
 #endif
 
 protected:
-	virtual TArray<FPCGPinProperties> InputPinProperties() const override { return TArray<FPCGPinProperties>(); }
+	virtual TArray<FPCGPinProperties> InputPinProperties() const override;
 	virtual TArray<FPCGPinProperties> OutputPinProperties() const override;
 
 	virtual FPCGElementPtr CreateElement() const override;
+
+#if WITH_EDITOR
+	virtual EPCGChangeType GetChangeTypeForProperty(FPropertyChangedEvent& PropertyChangedEvent) const override;
+#endif // WITH_EDITOR
 	//~End UPCGSettings interface
 
 public:
@@ -56,10 +60,14 @@ public:
 #endif
 
 protected:
-	virtual TArray<FPCGPinProperties> InputPinProperties() const override { return TArray<FPCGPinProperties>(); }
+	virtual TArray<FPCGPinProperties> InputPinProperties() const override;
 	virtual TArray<FPCGPinProperties> OutputPinProperties() const override;
 
 	virtual FPCGElementPtr CreateElement() const override;
+
+#if WITH_EDITOR
+	virtual EPCGChangeType GetChangeTypeForProperty(FPropertyChangedEvent& PropertyChangedEvent) const override;
+#endif // WITH_EDITOR
 	//~End UPCGSettings interface
 
 public:
diff --git a/Engine/Plugins/PCG/Source/PCG/Public/Elements/PCGWorldRaycast.h b/Engine/Plugins/PCG/Source/PCG/Public/Elements/PCGWorldRaycast.h
--- a/Engine/Plugins/PCG/Source/PCG/Public/Elements/PCGWorldRaycast.h
+++ b/Engine/Plugins/PCG/Source/PCG/Public/Elements/PCGWorldRaycast.h
@@ -33,6 +33,9 @@ protected:
 	virtual TArray<FPCGPinProperties> InputPinProperties() const override;
 	virtual TArray<FPCGPinProperties> OutputPinProperties() const override;
 	virtual FPCGElementPtr CreateElement() const override;
+#if WITH_EDITOR
+	virtual EPCGChangeType GetChangeTypeForProperty(FPropertyChangedEvent& PropertyChangedEvent) const override;
+#endif // WITH_EDITOR
 	//~End UPCGSettings interface
 
 public:
@@ -99,6 +102,7 @@ namespace FPCGWorldRaycast
 		TArray<FVector> CachedRayOrigins;
 		TArray<FVector> CachedRayVectors;
 		TArray<double> CachedRayLengths;
+		TArray<TSoftObjectPtr<AActor>> CachedFilterActors;
 	};
 }
 
diff --git a/Engine/Plugins/PCG/Source/PCG/Public/Helpers/PCGWorldQueryHelpers.h b/Engine/Plugins/PCG/Source/PCG/Public/Helpers/PCGWorldQueryHelpers.h
--- a/Engine/Plugins/PCG/Source/PCG/Public/Helpers/PCGWorldQueryHelpers.h
+++ b/Engine/Plugins/PCG/Source/PCG/Public/Helpers/PCGWorldQueryHelpers.h
@@ -2,11 +2,14 @@
 
 #pragma once
 
+#include "Containers/ArrayView.h"
 #include "Engine/HitResult.h"
 #include "Misc/Optional.h"
+#include "UObject/SoftObjectPtr.h"
 
 #include "PCGWorldQueryHelpers.generated.h"
 
+class AActor;
 class UPCGComponent;
 class UPCGMetadata;
 struct FHitResult;
@@ -49,19 +52,22 @@ namespace PCGWorldQueryHelpers
 	bool FilterCommonQueryResults(
 		const FPCGWorldCommonQueryParams* QueryParams,
 		const UPrimitiveComponent* TriggeredComponent,
-		const TWeakObjectPtr<UPCGComponent> OriginatingComponent);
+		const TWeakObjectPtr<UPCGComponent> OriginatingComponent,
+		const TArrayView<const TSoftObjectPtr<AActor>> FilteredActorReferences);
 
 	/** Filters through an array of hit results, testing them against raycast query parameters. Optional extra filter predicate. */
 	TOptional<FHitResult> FilterRayHitResults(
 		const FPCGWorldRaycastQueryParams* QueryParams,
 		TWeakObjectPtr<UPCGComponent> OriginatingComponent,
-		const TArray<FHitResult>& HitResults);
+		const TArray<FHitResult>& HitResults,
+		const TArrayView<const TSoftObjectPtr<AActor>> FilteredActorReferences);
 
 	/** Filters through an array of overlap results, testing them against raycast query parameters. Optional extra filter predicate. */
 	TOptional<FOverlapResult> FilterOverlapResults(
 		const FPCGWorldVolumetricQueryParams* QueryParams,
 		TWeakObjectPtr<UPCGComponent> OriginatingComponent,
-		const TArray<FOverlapResult>& OverlapResults);
+		const TArray<FOverlapResult>& OverlapResults,
+		const TArrayView<const TSoftObjectPtr<AActor>> FilteredActorReferences);
 
 	/** Creates hit result attributes based off query params. Can be called before ApplyRayHitMetadata. */
 	bool CreateRayHitAttributes(const FPCGWorldRaycastQueryParams& QueryParams, UPCGMetadata* OutMetadata);
@@ -76,4 +82,4 @@ namespace PCGWorldQueryHelpers
 		FPCGPoint& OutPoint,
 		UPCGMetadata* OutMetadata,
 		TWeakObjectPtr<UWorld> World);
-}
\ No newline at end of file
+}
diff --git a/Engine/Plugins/PCG/Source/PCG/Public/Metadata/Accessors/PCGAttributeAccessorHelpers.h b/Engine/Plugins/PCG/Source/PCG/Public/Metadata/Accessors/PCGAttributeAccessorHelpers.h
--- a/Engine/Plugins/PCG/Source/PCG/Public/Metadata/Accessors/PCGAttributeAccessorHelpers.h
+++ b/Engine/Plugins/PCG/Source/PCG/Public/Metadata/Accessors/PCGAttributeAccessorHelpers.h
@@ -2,6 +2,7 @@
 
 #pragma once
 
+#include "Metadata/PCGAttributePropertySelector.h"
 #include "Metadata/PCGMetadataAttributeTpl.h"
 #include "Metadata/PCGMetadataAttributeTraits.h"
 #include "Metadata/Accessors/IPCGAttributeAccessor.h"
@@ -12,17 +13,16 @@
 #include "UObject/NameTypes.h"
 
 enum class EPCGExtraProperties : uint8;
+struct FPCGContext;
+struct FPCGDataCollection;
+struct FPCGSettingsOverridableParam;
+class FProperty;
 class IPCGAttributeAccessor;
 class IPCGAttributeAccessorKeyIterator;
 class IPCGAttributeAccessorKeys;
-class FProperty;
 class UClass;
 class UPCGData;
 class UStruct;
-struct FPCGAttributePropertySelector;
-struct FPCGDataCollection;
-struct FPCGSettingsOverridableParam;
-
 
 namespace PCGAttributeAccessorHelpers
 {
@@ -93,6 +93,89 @@ namespace PCGAttributeAccessorHelpers
 	PCG_API TUniquePtr<const IPCGAttributeAccessorKeys> CreateConstKeys(const UPCGData* InData, const FPCGAttributePropertySelector& InSelector);
 	PCG_API TUniquePtr<IPCGAttributeAccessorKeys> CreateKeys(UPCGData* InData, const FPCGAttributePropertySelector& InSelector);
 
+	/** 
+	 * Utility function to create the accessor and keys and extract all the values for all the keys. Also, automatically resolve the selector for @Last.
+	 *
+	 * @param InData      Input data to extract the value from. Will also be used to resolve @Last on the selector.
+	 * @param InSelector  Selector to know which value to extract. Will be resolved if it is @Last.
+	 * @param OutArray    Output array where the values will be written.
+	 * @param Context     Optional context for logging. Can be null.
+	 * @param GetFlags    Flags for the get accessor. By default, is broadcast and constructible.
+	 * @return true if it succeeded.
+	 */
+	template <typename T, typename AllocatorType>
+	bool ExtractAllValues(const UPCGData* InData, const FPCGAttributePropertyInputSelector& InSelector, TArray<T, AllocatorType>& OutArray, FPCGContext* Context = nullptr, EPCGAttributeAccessorFlags GetFlags = EPCGAttributeAccessorFlags::AllowBroadcastAndConstructible)
+	{
+		// TODO: Provide a version that take advantage of compressed attributes (Keys + Values)
+		
+		const FPCGAttributePropertyInputSelector Selector = InSelector.CopyAndFixLast(InData);
+
+		TUniquePtr<const IPCGAttributeAccessor> Accessor = CreateConstAccessor(InData, Selector);
+		TUniquePtr<const IPCGAttributeAccessorKeys> Keys = CreateConstKeys(InData, Selector);
+		if (!Accessor.IsValid() || !Keys.IsValid())
+		{
+			PCGLog::Metadata::LogFailToCreateAccessorError(Selector, Context);
+			return false;
+		}
+			
+		if constexpr (std::is_trivially_copyable_v<T>)
+		{
+			OutArray.SetNumUninitialized(Keys->GetNum());
+		}
+		else
+		{
+			OutArray.SetNum(Keys->GetNum());
+		}
+
+		if (!Accessor->GetRange<T>(OutArray, 0, *Keys, GetFlags))
+		{
+			PCGLog::Metadata::LogFailToGetAttributeError<T>(Selector, Accessor.Get(), Context);
+			return false;
+		}
+
+		return true;
+	}
+	
+	/** 
+	 * Utility function to create the accessor and keys and write all the values for all the keys coming from InValues. Also, automatically resolve the selector for @Last.
+	 *
+	 * @param OutputData      Input data to extract the value from. Will also be used to resolve @Source on the selector.
+	 * @param OutputSelector  Selector to know which value to write to. Will be resolved if it is @Source.
+	 * @param InValues        Input view where the values will be read from. InValues MUST be the same size as the keys that will be created by CreateKeys on the OutputData. 
+	 * @param SourceSelector  Optional source selector for @Source resolve. Can be null.
+	 * @param Context         Optional context for logging. Can be null.
+	 * @param SetFlags        Flags for the set accessor. By default, is broadcast and constructible.
+	 * @return true if it succeeded.
+	 */
+	template <typename T>
+	bool WriteAllValues(UPCGData* OutputData, const FPCGAttributePropertyOutputSelector& OutputSelector, TArrayView<const T> InValues, const FPCGAttributePropertyInputSelector* SourceSelector = nullptr, FPCGContext* Context = nullptr, EPCGAttributeAccessorFlags SetFlags = EPCGAttributeAccessorFlags::AllowBroadcastAndConstructible)
+	{
+		// TODO: Provide a version that take advantage of compressed attributes (Keys + Values)
+		
+		const FPCGAttributePropertyOutputSelector Selector = OutputSelector.CopyAndFixSource(SourceSelector, OutputData);
+
+		TUniquePtr<IPCGAttributeAccessor> Accessor = CreateAccessor(OutputData, Selector);
+		TUniquePtr<IPCGAttributeAccessorKeys> Keys = CreateKeys(OutputData, Selector);
+		if (!Accessor.IsValid() || !Keys.IsValid())
+		{
+			PCGLog::Metadata::LogFailToCreateAccessorError(Selector, Context);
+			return false;
+		}
+
+		if (!ensureMsgf(Keys->GetNum() == InValues.GetNum(), TEXT("Number of values passed (%d) mismatches with the number of keys (%d)"), Keys->GetNum(), InValues.GetNum()))
+		{
+			return false;
+		}
+
+		if (!Accessor->SetRange<T>(InValues, 0, *Keys, SetFlags))
+		{
+			PCGLog::Metadata::LogFailToSetAttributeError<T>(Selector, Accessor.Get(), Context);
+			return false;
+		}
+
+		return true;
+	}
+
 	namespace Private
 	{
 		// Use a lambda to have this code more likely to be inlined in SortByAttribute.
diff --git a/Engine/Plugins/PCG/Source/PCG/Public/Utils/PCGLogErrors.h b/Engine/Plugins/PCG/Source/PCG/Public/Utils/PCGLogErrors.h
--- a/Engine/Plugins/PCG/Source/PCG/Public/Utils/PCGLogErrors.h
+++ b/Engine/Plugins/PCG/Source/PCG/Public/Utils/PCGLogErrors.h
@@ -40,6 +40,8 @@ namespace PCGLog
 			PCG_API extern const FTextFormat CreateAttributeFailure;
 			PCG_API extern const FTextFormat GetTypedAttributeFailure;
 			PCG_API extern const FTextFormat GetTypedAttributeFailureNoAccessor;
+			PCG_API extern const FTextFormat SetTypedAttributeFailure;
+			PCG_API extern const FTextFormat SetTypedAttributeFailureNoAccessor;
 		}
 
 		// Errors
@@ -86,6 +88,31 @@ namespace PCGLog
 		{
 			return LogFailToGetAttributeError<T>(Selector.GetDisplayText(), Accessor, InContext);
 		}
+		
+		template <typename T>
+		void LogFailToSetAttributeError(const FText& AttributeName, const IPCGAttributeAccessor* Accessor, const FPCGContext* InContext = nullptr)
+		{
+			if (Accessor)
+			{
+				PCGLog::LogErrorOnGraph(FText::Format(Format::SetTypedAttributeFailure, AttributeName, PCG::Private::GetTypeNameText(Accessor->GetUnderlyingType()), PCG::Private::GetTypeNameText<T>()), InContext);
+			}
+			else
+			{
+				PCGLog::LogErrorOnGraph(FText::Format(Format::SetTypedAttributeFailureNoAccessor, AttributeName, PCG::Private::GetTypeNameText<T>()), InContext);
+			}
+		}
+
+		template <typename T>
+		void LogFailToSetAttributeError(FName AttributeName, const IPCGAttributeAccessor* Accessor, const FPCGContext* InContext = nullptr)
+		{
+			LogFailToSetAttributeError<T>(FText::FromName(AttributeName), Accessor, InContext);
+		}
+
+		template <typename T>
+		void LogFailToSetAttributeError(const FPCGAttributePropertySelector& Selector, const IPCGAttributeAccessor* Accessor, const FPCGContext* InContext = nullptr)
+		{
+			return LogFailToSetAttributeError<T>(Selector.GetDisplayText(), Accessor, InContext);
+		}
 	}
 
 	namespace Parsing
