[PCG] Update attribute names to selectors when it makes sense, for newly created nodes

* Copy Attributes, Get Attribute from point Index, Match and Set Attributes, Get Property from Object path, Get Actor Property

#rnx

diff --git a/Engine/Plugins/PCG/Source/PCG/Private/Elements/PCGAttributeGetFromPointIndexElement.cpp b/Engine/Plugins/PCG/Source/PCG/Private/Elements/PCGAttributeGetFromPointIndexElement.cpp
--- a/Engine/Plugins/PCG/Source/PCG/Private/Elements/PCGAttributeGetFromPointIndexElement.cpp
+++ b/Engine/Plugins/PCG/Source/PCG/Private/Elements/PCGAttributeGetFromPointIndexElement.cpp
@@ -7,6 +7,7 @@
 #include "PCGPin.h"
 #include "PCGParamData.h"
 #include "Data/PCGPointData.h"
+#include "Helpers/PCGHelpers.h"
 #include "Metadata/Accessors/IPCGAttributeAccessor.h"
 #include "Metadata/Accessors/PCGAttributeAccessorKeys.h"
 #include "Metadata/Accessors/PCGAttributeAccessorHelpers.h"
@@ -29,17 +30,30 @@ FText UPCGAttributeGetFromPointIndexSettings::GetDefaultNodeTitle() const
 void UPCGAttributeGetFromPointIndexSettings::ApplyDeprecation(UPCGNode* InOutNode)
 {
 	if (DataVersion < FPCGCustomVersion::UpdateAttributePropertyInputSelector
-		&& (OutputAttributeName == NAME_None))
+		&& (OutputAttributeName.GetName() == NAME_None))
 	{
 		// Previous behavior of the output attribute for this node was:
 		// None => SameName
-		OutputAttributeName = PCGMetadataAttributeConstants::SourceNameAttributeName;
+		OutputAttributeName.SetAttributeName(PCGMetadataAttributeConstants::SourceNameAttributeName);
 	}
 
 	Super::ApplyDeprecation(InOutNode);
 }
 #endif
 
+UPCGAttributeGetFromPointIndexSettings::UPCGAttributeGetFromPointIndexSettings()
+{
+	// Default was None, but for new objects it will be @Source
+	if (PCGHelpers::IsNewObjectAndNotDefault(this))
+	{
+		OutputAttributeName.SetAttributeName(PCGMetadataAttributeConstants::SourceAttributeName);
+	}
+	else
+	{
+		OutputAttributeName.SetAttributeName(NAME_None);
+	}
+}
+
 void UPCGAttributeGetFromPointIndexSettings::PostLoad()
 {
 	Super::PostLoad();
@@ -109,9 +123,10 @@ bool FPCGAttributeGetFromPointIndexElement::ExecuteInternal(FPCGContext* Context
 		TArray<FPCGTaggedData>& Outputs = Context->OutputData.TaggedData;
 		const FPCGPoint& Point = PointData->GetPoints()[Index];
 
-		FPCGAttributePropertyInputSelector InputSource = Settings->InputSource.CopyAndFixLast(PointData);
-
-		const FName OutputAttributeName = (Settings->OutputAttributeName == PCGMetadataAttributeConstants::SourceNameAttributeName) ? InputSource.GetName() : Settings->OutputAttributeName;
+		const FPCGAttributePropertyInputSelector InputSource = Settings->InputSource.CopyAndFixLast(PointData);
+		const FPCGAttributePropertyOutputSelector OutputTarget = Settings->OutputAttributeName.CopyAndFixSource(&InputSource);
+		
+		const FName OutputAttributeName = OutputTarget.GetName();
 
 		TUniquePtr<const IPCGAttributeAccessor> Accessor = PCGAttributeAccessorHelpers::CreateConstAccessor(PointData, InputSource);
 		FPCGAttributeAccessorKeysPoints PointKey(Point);
diff --git a/Engine/Plugins/PCG/Source/PCG/Private/Elements/PCGCopyAttributes.cpp b/Engine/Plugins/PCG/Source/PCG/Private/Elements/PCGCopyAttributes.cpp
--- a/Engine/Plugins/PCG/Source/PCG/Private/Elements/PCGCopyAttributes.cpp
+++ b/Engine/Plugins/PCG/Source/PCG/Private/Elements/PCGCopyAttributes.cpp
@@ -9,6 +9,7 @@
 #include "Data/PCGPointData.h"
 #include "Data/PCGSpatialData.h"
 #include "Elements/PCGCopyPoints.h"
+#include "Helpers/PCGHelpers.h"
 #include "Helpers/PCGMetadataHelpers.h"
 
 // Keep loctext namespace to its previous value to avoid to redo loc for this.
@@ -17,11 +18,19 @@
 
 UPCGCopyAttributesSettings::UPCGCopyAttributesSettings()
 {
-	// Previous default object was: None for source and target attribute, density for point property
-	// and Property to attribute
-	// Recreate the same default
-	InputSource.SetPointProperty(EPCGPointProperties::Density);
-	OutputTarget.SetAttributeName(NAME_None);
+	if (PCGHelpers::IsNewObjectAndNotDefault(this))
+	{
+		InputSource.SetAttributeName(PCGMetadataAttributeConstants::LastAttributeName);
+		OutputTarget.SetAttributeName(PCGMetadataAttributeConstants::SourceAttributeName);
+	}
+	else
+	{
+		// Previous default object was: None for source and target attribute, density for point property
+		// and Property to attribute
+		// Recreate the same default
+		InputSource.SetPointProperty(EPCGPointProperties::Density);
+		OutputTarget.SetAttributeName(NAME_None);
+	}
 }
 
 FString UPCGCopyAttributesSettings::GetAdditionalTitleInformation() const
diff --git a/Engine/Plugins/PCG/Source/PCG/Private/Elements/PCGGetActorProperty.cpp b/Engine/Plugins/PCG/Source/PCG/Private/Elements/PCGGetActorProperty.cpp
--- a/Engine/Plugins/PCG/Source/PCG/Private/Elements/PCGGetActorProperty.cpp
+++ b/Engine/Plugins/PCG/Source/PCG/Private/Elements/PCGGetActorProperty.cpp
@@ -18,10 +18,15 @@
 UPCGGetActorPropertySettings::UPCGGetActorPropertySettings(const FObjectInitializer& ObjectInitializer)
 {
 	// To not break previous behavior, we keep this false by default, but will force it to true for new nodes as it makes more sense to process them all,
-	// as we have no guarantee on the order components are found.
+	// as we have no guarantee on the order components are found. Also the new default is @Source for the output attribute
 	if (PCGHelpers::IsNewObjectAndNotDefault(this))
 	{
 		bProcessAllComponents = true;
+		OutputAttributeName.SetAttributeName(PCGMetadataAttributeConstants::SourceAttributeName);
+	}
+	else
+	{
+		OutputAttributeName.SetAttributeName(NAME_None);
 	}
 }
 
@@ -291,7 +296,9 @@ bool FPCGGetActorPropertyElement::ExecuteInternal(FPCGContext* InContext) const
 
 		for (UObject* ObjectToInspect : ObjectsToInspect)
 		{
-			PCGPropertyHelpers::FExtractorParameters Parameters(ObjectToInspect, ObjectToInspect->GetClass(), Settings->PropertyName.ToString(), Settings->OutputAttributeName, Settings->bForceObjectAndStructExtraction, /*bPropertyNeedsToBeVisible=*/true);
+			const FName OutputAttributeName = Settings->OutputAttributeName.GetName();
+
+			PCGPropertyHelpers::FExtractorParameters Parameters(ObjectToInspect, ObjectToInspect->GetClass(), Settings->PropertyName.ToString(), OutputAttributeName, Settings->bForceObjectAndStructExtraction, /*bPropertyNeedsToBeVisible=*/true);
 
 			// Don't care for object traversed in non-editor build, since it is only useful for tracking.
 			TSet<FSoftObjectPath>* ObjectTraversedPtr = nullptr;
diff --git a/Engine/Plugins/PCG/Source/PCG/Private/Elements/PCGGetPropertyFromObjectPath.cpp b/Engine/Plugins/PCG/Source/PCG/Private/Elements/PCGGetPropertyFromObjectPath.cpp
--- a/Engine/Plugins/PCG/Source/PCG/Private/Elements/PCGGetPropertyFromObjectPath.cpp
+++ b/Engine/Plugins/PCG/Source/PCG/Private/Elements/PCGGetPropertyFromObjectPath.cpp
@@ -5,10 +5,24 @@
 #include "PCGComponent.h"
 #include "PCGParamData.h"
 #include "Helpers/PCGDynamicTrackingHelpers.h"
+#include "Helpers/PCGHelpers.h"
 #include "Helpers/PCGPropertyHelpers.h"
 
 #define LOCTEXT_NAMESPACE "PCGGetPropertyFromObjectPathElement"
 
+UPCGGetPropertyFromObjectPathSettings::UPCGGetPropertyFromObjectPathSettings()
+{
+	// Default was None, but for new objects it will be @Source
+	if (PCGHelpers::IsNewObjectAndNotDefault(this))
+	{
+		OutputAttributeName.SetAttributeName(PCGMetadataAttributeConstants::SourceAttributeName);
+	}
+	else
+	{
+		OutputAttributeName.SetAttributeName(NAME_None);
+	}
+}
+
 #if WITH_EDITOR
 void UPCGGetPropertyFromObjectPathSettings::GetStaticTrackedKeys(FPCGSelectionKeyToSettingsMap& OutKeysToSettings, TArray<TObjectPtr<const UPCGGraph>>& OutVisitedGraphs) const
 {
@@ -204,7 +218,9 @@ bool FPCGGetPropertyFromObjectPathElement::ExecuteInternal(FPCGContext* Context)
 			continue;
 		}
 
-		PCGPropertyHelpers::FExtractorParameters Parameters(Object, Object->GetClass(), Settings->PropertyName.ToString(), Settings->OutputAttributeName, Settings->bForceObjectAndStructExtraction, /*bPropertyNeedsToBeVisible=*/true);
+		const FName OutputAttributeName = Settings->OutputAttributeName.GetName();
+
+		PCGPropertyHelpers::FExtractorParameters Parameters(Object, Object->GetClass(), Settings->PropertyName.ToString(), OutputAttributeName, Settings->bForceObjectAndStructExtraction, /*bPropertyNeedsToBeVisible=*/true);
 
 		if (UPCGParamData* ParamData = PCGPropertyHelpers::ExtractPropertyAsAttributeSet(Parameters, Context))
 		{
diff --git a/Engine/Plugins/PCG/Source/PCG/Private/Elements/PCGMatchAndSetAttributes.cpp b/Engine/Plugins/PCG/Source/PCG/Private/Elements/PCGMatchAndSetAttributes.cpp
--- a/Engine/Plugins/PCG/Source/PCG/Private/Elements/PCGMatchAndSetAttributes.cpp
+++ b/Engine/Plugins/PCG/Source/PCG/Private/Elements/PCGMatchAndSetAttributes.cpp
@@ -8,6 +8,7 @@
 #include "Data/PCGPointData.h"
 #include "Elements/Metadata/PCGMetadataElementCommon.h"
 #include "Helpers/PCGBlueprintHelpers.h"
+#include "Helpers/PCGHelpers.h"
 #include "Helpers/PCGPointDataPartition.h"
 #include "Metadata/PCGMetadataAttributeTpl.h"
 #include "Metadata/PCGMetadataAttributeTraits.h"
@@ -29,6 +30,22 @@ namespace PCGMatchAndSetAttributesConstants
 	const FName MaxDistanceLabel = TEXT("Max Match Distance");
 }
 
+UPCGMatchAndSetAttributesSettings::UPCGMatchAndSetAttributesSettings()
+{
+	// Weight attribute is None by default
+	WeightAttribute.SetAttributeName(NAME_None);
+	
+	// The match attribute will be @Last, but leave None for not new object as it was the previous default.
+	if (PCGHelpers::IsNewObjectAndNotDefault(this))
+	{
+		MatchAttribute.SetAttributeName(PCGMetadataAttributeConstants::LastAttributeName);
+	}
+	else
+	{
+		MatchAttribute.SetAttributeName(NAME_None);
+	}
+}
+
 #if WITH_EDITOR
 FName UPCGMatchAndSetAttributesSettings::GetDefaultNodeName() const
 {
@@ -560,9 +577,9 @@ public:
 		AttributeSetPartition.Initialize(Context,
 			ParamData,
 			Settings->bMatchAttributes,
-			Settings->MatchAttribute,
+			Settings->MatchAttribute.CopyAndFixLast(ParamData).GetName(),
 			Settings->bUseWeightAttribute,
-			Settings->WeightAttribute,
+			Settings->WeightAttribute.CopyAndFixLast(ParamData).GetName(),
 			Settings->bFindNearest,
 			Settings->MaxDistanceMode,
 			&Settings->MaxDistanceForNearestMatch);
@@ -629,6 +646,19 @@ public:
 		TArray<EPCGMetadataTypes> ParamAttributeTypes;
 		ParamData->ConstMetadata()->GetAttributes(ParamAttributeNames, ParamAttributeTypes);
 
+		const FPCGAttributePropertyInputSelector MatchAttributeSelector = Settings->MatchAttribute.CopyAndFixLast(ParamData);
+		const FPCGAttributePropertyInputSelector WeightAttributeSelector = Settings->WeightAttribute.CopyAndFixLast(ParamData);
+
+		// Validate that the selector are basic attributes
+		if (!MatchAttributeSelector.IsBasicAttribute() || !WeightAttributeSelector.IsBasicAttribute())
+		{
+			PCGLog::LogWarningOnGraph(FText::Format(LOCTEXT("NotBasicAttribute", "Match attribute {0} and/or weight attribute {1} are not basic attributes. Not yet supported, we will only use the basic attribute."),
+				MatchAttributeSelector.GetDisplayText(), WeightAttributeSelector.GetDisplayText()), Context);
+		}
+
+		const FName MatchAttribute = MatchAttributeSelector.GetName();
+		const FName WeightAttribute = WeightAttributeSelector.GetName();
+
 		UPCGMetadata* OutMetadata = OutData->MutableMetadata();
 		check(OutMetadata);
 
@@ -636,8 +666,8 @@ public:
 		{
 			const FName AttributeName = ParamAttributeNames[AttributeIndex];
 
-			if ((Settings->bMatchAttributes && AttributeName == Settings->MatchAttribute) ||
-				(Settings->bUseWeightAttribute && AttributeName == Settings->WeightAttribute))
+			if ((Settings->bMatchAttributes && AttributeName == MatchAttribute) ||
+				(Settings->bUseWeightAttribute && AttributeName == WeightAttribute))
 			{
 				continue;
 			}
diff --git a/Engine/Plugins/PCG/Source/PCG/Private/Helpers/PCGPropertyHelpers.cpp b/Engine/Plugins/PCG/Source/PCG/Private/Helpers/PCGPropertyHelpers.cpp
--- a/Engine/Plugins/PCG/Source/PCG/Private/Helpers/PCGPropertyHelpers.cpp
+++ b/Engine/Plugins/PCG/Source/PCG/Private/Helpers/PCGPropertyHelpers.cpp
@@ -359,7 +359,9 @@ UPCGParamData* PCGPropertyHelpers::ExtractPropertyAsAttributeSet(const PCGProper
 		else
 		{
 			// For non struct/object, there is just a single property to extract with no shenanigans for address indirection.
-			const FName AttributeName = (bIgnoreOutputAttributeName || Parameters.OutputAttributeName == PCGMetadataAttributeConstants::SourceNameAttributeName) ? Property->GetFName() : Parameters.OutputAttributeName;
+			const bool bIsSourceName = Parameters.OutputAttributeName == PCGMetadataAttributeConstants::SourceNameAttributeName || Parameters.OutputAttributeName == PCGMetadataAttributeConstants::SourceAttributeName;
+			
+			const FName AttributeName = (bIgnoreOutputAttributeName || bIsSourceName) ? Property->GetFName() : Parameters.OutputAttributeName;
 			ExtractableProperties.Emplace(AttributeName.ToString(), Property);
 			// Identity
 			AddressFunc = [](const void* InAddress) { return InAddress; };
diff --git a/Engine/Plugins/PCG/Source/PCG/Private/Metadata/PCGAttributePropertySelector.cpp b/Engine/Plugins/PCG/Source/PCG/Private/Metadata/PCGAttributePropertySelector.cpp
--- a/Engine/Plugins/PCG/Source/PCG/Private/Metadata/PCGAttributePropertySelector.cpp
+++ b/Engine/Plugins/PCG/Source/PCG/Private/Metadata/PCGAttributePropertySelector.cpp
@@ -399,6 +399,13 @@ bool FPCGAttributePropertyInputSelector::SerializeFromMismatchedTag(const struct
 		FPCGAttributePropertyInputSelector::StaticStruct()->SerializeItem(Slot, this, nullptr);
 		return true;
 	}
+	else if (Tag.GetType().GetName() == NAME_NameProperty)
+	{
+		FName Value;
+		Slot << Value;
+		SetAttributeName(Value);
+		return true;
+	}
 
 	return false;
 }
@@ -459,6 +466,13 @@ bool FPCGAttributePropertyOutputSelector::SerializeFromMismatchedTag(const struc
 		FPCGAttributePropertyOutputSelector::StaticStruct()->SerializeItem(Slot, this, nullptr);
 		return true;
 	}
+	else if (Tag.GetType().GetName() == NAME_NameProperty)
+	{
+		FName Value;
+		Slot << Value;
+		SetAttributeName(Value);
+		return true;
+	}
 
 	return false;
 }
diff --git a/Engine/Plugins/PCG/Source/PCG/Private/Tests/Elements/PCGGetActorPropertyTest.cpp b/Engine/Plugins/PCG/Source/PCG/Private/Tests/Elements/PCGGetActorPropertyTest.cpp
--- a/Engine/Plugins/PCG/Source/PCG/Private/Tests/Elements/PCGGetActorPropertyTest.cpp
+++ b/Engine/Plugins/PCG/Source/PCG/Private/Tests/Elements/PCGGetActorPropertyTest.cpp
@@ -48,7 +48,7 @@ bool VerifyAttributeValue(FPCGTestBaseClass* TestInstance, PCGTestsCommon::FTest
 	}
 
 	Settings->PropertyName = PropertyName;
-	Settings->OutputAttributeName = AttributeNames[0];
+	Settings->OutputAttributeName.SetAttributeName(AttributeNames[0]);
 	Settings->bForceObjectAndStructExtraction = AttributeNames.Num() > 1;
 
 	// Add 2 nodes, PropertyToParamDataNode and a Trivial node (just there for the connection)
diff --git a/Engine/Plugins/PCG/Source/PCG/Public/Elements/PCGAttributeGetFromPointIndexElement.h b/Engine/Plugins/PCG/Source/PCG/Public/Elements/PCGAttributeGetFromPointIndexElement.h
--- a/Engine/Plugins/PCG/Source/PCG/Public/Elements/PCGAttributeGetFromPointIndexElement.h
+++ b/Engine/Plugins/PCG/Source/PCG/Public/Elements/PCGAttributeGetFromPointIndexElement.h
@@ -27,6 +27,8 @@ class UPCGAttributeGetFromPointIndexSettings : public UPCGSettings
 	GENERATED_BODY()
 
 public:
+	UPCGAttributeGetFromPointIndexSettings();
+
 	//~Begin UObject interface
 	virtual void PostLoad() override;
 	//~End UObject interface
@@ -53,8 +55,8 @@ public:
 	UPROPERTY(BlueprintReadWrite, EditAnywhere, Category = Settings, meta = (PCG_Overridable))
 	int32 Index = 0;
 
-	UPROPERTY(BlueprintReadWrite, EditAnywhere, Category = Settings)
-	FName OutputAttributeName = NAME_None;
+	UPROPERTY(BlueprintReadWrite, EditAnywhere, Category = Settings, meta = (PCG_Overridable, PCG_DiscardPropertySelection, PCG_DiscardExtraSelection))
+	FPCGAttributePropertyOutputSelector OutputAttributeName;
 
 #if WITH_EDITORONLY_DATA
 	UPROPERTY()
diff --git a/Engine/Plugins/PCG/Source/PCG/Public/Elements/PCGGetActorProperty.h b/Engine/Plugins/PCG/Source/PCG/Public/Elements/PCGGetActorProperty.h
--- a/Engine/Plugins/PCG/Source/PCG/Public/Elements/PCGGetActorProperty.h
+++ b/Engine/Plugins/PCG/Source/PCG/Public/Elements/PCGGetActorProperty.h
@@ -79,9 +79,9 @@ public:
 	UPROPERTY(BlueprintReadWrite, EditAnywhere, Category = Settings, meta = (PCG_Overridable))
 	bool bForceObjectAndStructExtraction = false;
 
-	/** By default, attribute name will be None, but it can be overridden by this name. Use @SourceName to use the property name (only works when not extracting). Will be ignored if multiple properties are extracted at the same time. */
-	UPROPERTY(BlueprintReadWrite, EditAnywhere, Category = Settings, meta = (PCG_Overridable, EditCondition = "!bForceObjectAndStructExtraction", EditConditionHides))
-	FName OutputAttributeName = NAME_None;
+	/** In the case of multiple properties being extracted, will be ignored. */
+	UPROPERTY(BlueprintReadWrite, EditAnywhere, Category = Settings, meta = (PCG_Overridable, EditCondition = "!bForceObjectAndStructExtraction", EditConditionHides, PCG_DiscardPropertySelection, PCG_DiscardExtraSelection))
+	FPCGAttributePropertyOutputSelector OutputAttributeName;
 
 	/** Controls whether an actor reference attribute will be added to the result */
 	UPROPERTY(BlueprintReadWrite, EditAnywhere, Category = Settings, meta = (PCG_Overridable))
diff --git a/Engine/Plugins/PCG/Source/PCG/Public/Elements/PCGGetPropertyFromObjectPath.h b/Engine/Plugins/PCG/Source/PCG/Public/Elements/PCGGetPropertyFromObjectPath.h
--- a/Engine/Plugins/PCG/Source/PCG/Public/Elements/PCGGetPropertyFromObjectPath.h
+++ b/Engine/Plugins/PCG/Source/PCG/Public/Elements/PCGGetPropertyFromObjectPath.h
@@ -16,6 +16,8 @@ class UPCGGetPropertyFromObjectPathSettings : public UPCGSettings
 	GENERATED_BODY()
 
 public:
+	UPCGGetPropertyFromObjectPathSettings();
+
 	//~Begin UPCGSettings interface
 #if WITH_EDITOR
 	virtual FName GetDefaultNodeName() const override;
@@ -55,9 +57,9 @@ public:
 	UPROPERTY(BlueprintReadWrite, EditAnywhere, Category = Settings)
 	bool bForceObjectAndStructExtraction = false;
 
-	/** By default, attribute name will be None, but it can be overridden by this name. Use @SourceName to use the property name (only works when not extracting). In the case of multiple properties being extracted, will be ignored. */
-	UPROPERTY(BlueprintReadWrite, EditAnywhere, Category = Settings, meta = (PCG_Overridable, EditCondition = "!bForceObjectAndStructExtraction", EditConditionHides))
-	FName OutputAttributeName = NAME_None;
+	/** In the case of multiple properties being extracted, will be ignored. */
+	UPROPERTY(BlueprintReadWrite, EditAnywhere, Category = Settings, meta = (PCG_Overridable, EditCondition = "!bForceObjectAndStructExtraction", EditConditionHides, PCG_DiscardPropertySelection, PCG_DiscardExtraSelection))
+	FPCGAttributePropertyOutputSelector OutputAttributeName;
 
 	/** By default, object loading is asynchronous, can force it synchronous if needed. */
 	UPROPERTY(BlueprintReadWrite, EditAnywhere, Category = "Settings|Debug")
diff --git a/Engine/Plugins/PCG/Source/PCG/Public/Elements/PCGMatchAndSetAttributes.h b/Engine/Plugins/PCG/Source/PCG/Public/Elements/PCGMatchAndSetAttributes.h
--- a/Engine/Plugins/PCG/Source/PCG/Public/Elements/PCGMatchAndSetAttributes.h
+++ b/Engine/Plugins/PCG/Source/PCG/Public/Elements/PCGMatchAndSetAttributes.h
@@ -45,6 +45,8 @@ class UPCGMatchAndSetAttributesSettings : public UPCGSettings
 	GENERATED_BODY()
 
 public:
+	UPCGMatchAndSetAttributesSettings();
+
 	// ~Begin UPCGSettings interface
 #if WITH_EDITOR
 	virtual FName GetDefaultNodeName() const override;
@@ -72,8 +74,8 @@ public:
 	FPCGAttributePropertyInputSelector InputAttribute;
 
 	/** Attribute from the attribute set to match against. */
-	UPROPERTY(BlueprintReadWrite, EditAnywhere, Category = Settings, meta = (EditCondition = "bMatchAttributes", PCG_Overridable))
-	FName MatchAttribute = NAME_None;
+	UPROPERTY(BlueprintReadWrite, EditAnywhere, Category = Settings, meta = (EditCondition = "bMatchAttributes", PCG_Overridable, PCG_DiscardPropertySelection))
+	FPCGAttributePropertyInputSelector MatchAttribute;
 
 	/** Controls whether points that have no valid match in the attribute set are kept as is (default values) or removed from the output. */
 	UPROPERTY(BlueprintReadWrite, EditAnywhere, Category = Settings, meta = (EditCondition = "bMatchAttributes", PCG_Overridable))
@@ -107,8 +109,8 @@ public:
 	bool bUseWeightAttribute = false;
 
 	/** Attribute to weight more or less some entries from the attribute set. */
-	UPROPERTY(BlueprintReadWrite, EditAnywhere, Category = Settings, meta = (DisplayName="Match Weight Attribute", EditCondition = "bUseWeightAttribute", PCG_Overridable))
-	FName WeightAttribute = NAME_None;
+	UPROPERTY(BlueprintReadWrite, EditAnywhere, Category = Settings, meta = (DisplayName="Match Weight Attribute", EditCondition = "bUseWeightAttribute", PCG_Overridable, PCG_DiscardPropertySelection))
+	FPCGAttributePropertyInputSelector WeightAttribute;
 
 	/** Controls whether we will emit a warning and return nothing if there is no provided attribute set. */
 	UPROPERTY(BlueprintReadWrite, EditAnywhere, Category = Settings)
diff --git a/Engine/Plugins/PCG/Source/PCG/Public/Helpers/PCGPropertyHelpers.h b/Engine/Plugins/PCG/Source/PCG/Public/Helpers/PCGPropertyHelpers.h
--- a/Engine/Plugins/PCG/Source/PCG/Public/Helpers/PCGPropertyHelpers.h
+++ b/Engine/Plugins/PCG/Source/PCG/Public/Helpers/PCGPropertyHelpers.h
@@ -72,7 +72,7 @@ namespace PCGPropertyHelpers
 
 		TArray<FPCGAttributePropertySelector> PropertySelectors;
 
-		// Optional name of the attribute that will receive the extracted property. If None, will take the property name. 
+		// Optional name of the attribute that will receive the extracted property. If @Source or @SourceName, will take the property name. 
 		// Also not used for Structs/Object extraction, as we will create multiple attributes, and they will be the name of all the extracted members.
 		FName OutputAttributeName = NAME_None;
 
diff --git a/Engine/Plugins/PCG/Source/PCG/Public/Metadata/PCGAttributePropertySelector.h b/Engine/Plugins/PCG/Source/PCG/Public/Metadata/PCGAttributePropertySelector.h
--- a/Engine/Plugins/PCG/Source/PCG/Public/Metadata/PCGAttributePropertySelector.h
+++ b/Engine/Plugins/PCG/Source/PCG/Public/Metadata/PCGAttributePropertySelector.h
@@ -41,6 +41,7 @@ struct PCG_API FPCGAttributePropertySelector
 	GENERATED_BODY()
 
 public:
+	FPCGAttributePropertySelector() = default;
 	virtual ~FPCGAttributePropertySelector() = default;
 
 	bool operator==(const FPCGAttributePropertySelector& Other) const;
@@ -129,6 +130,23 @@ public:
 		return Selector;
 	}
 
+	// For deprecation purposes
+	explicit FPCGAttributePropertySelector(const FName InName)
+	{
+		SetAttributeName(InName);
+	}
+	
+	FPCGAttributePropertySelector& operator=(const FName InName)
+	{
+		SetAttributeName(InName);
+		return *this;
+	}
+
+	explicit operator FName() const
+	{
+		return GetName();
+	}
+
 	void ImportFromOtherSelector(const FPCGAttributePropertySelector& InOther);
 
 	virtual void AddToCrc(FArchiveCrc32& Ar) const;
