[PCG] Support anonymous attributes (with name 'None') in compute graphs

#rnx

diff --git a/Engine/Plugins/PCG/Source/PCG/Private/Compute/Elements/PCGComputeGraphElement.cpp b/Engine/Plugins/PCG/Source/PCG/Private/Compute/Elements/PCGComputeGraphElement.cpp
--- a/Engine/Plugins/PCG/Source/PCG/Private/Compute/Elements/PCGComputeGraphElement.cpp
+++ b/Engine/Plugins/PCG/Source/PCG/Private/Compute/Elements/PCGComputeGraphElement.cpp
@@ -367,11 +367,6 @@ bool FPCGComputeGraphElement::SetupProceduralISMComponents(FPCGContext* InContex
 		if (const UPCGMeshSelectorByAttribute* SelectorByAttribute = Cast<UPCGMeshSelectorByAttribute>(SpawnerSettings->MeshSelectorParameters))
 		{
 			const FName SelectedName = SelectorByAttribute->AttributeName;
-			if (SelectedName == NAME_None)
-			{
-				UE_LOG(LogPCG, Error, TEXT("Invalid mesh selector attribute specified '%s'."), *SelectedName.ToString());
-				continue;
-			}
 
 			const FPCGKernelAttributeIDAndType* FoundAttribute = InBinding->GetAttributeLookupTable().Find(SelectedName);
 			if (!FoundAttribute)
diff --git a/Engine/Plugins/PCG/Source/PCG/Private/Compute/Elements/PCGCustomHLSL.cpp b/Engine/Plugins/PCG/Source/PCG/Private/Compute/Elements/PCGCustomHLSL.cpp
--- a/Engine/Plugins/PCG/Source/PCG/Private/Compute/Elements/PCGCustomHLSL.cpp
+++ b/Engine/Plugins/PCG/Source/PCG/Private/Compute/Elements/PCGCustomHLSL.cpp
@@ -1444,15 +1444,15 @@ bool UPCGCustomHLSLSettings::IsKernelValid(FPCGContext* InContext, bool bQuiet)
 	// Validate attributes
 	for (const FPCGKernelAttributeKey& AttributeKey : KernelAttributeKeys)
 	{
-		if (AttributeKey.Name == NAME_None)
+		if (AttributeKey.Type == EPCGKernelAttributeType::None || AttributeKey.Type == EPCGKernelAttributeType::Invalid)
 		{
-			PCG_KERNEL_VALIDATION_ERR(InContext, this, bQuiet, LOCTEXT("InvalidAttributeNameNone", "'None' is not a valid GPU attribute name, check the 'Attributes to Create' array on your pins."));
-			return false;
-		}
+			const UEnum* AttributeTypeEnum = StaticEnum<EPCGKernelAttributeType>();
+			check(AttributeTypeEnum);
 
-		if (AttributeKey.Type == EPCGKernelAttributeType::None)
-		{
-			PCG_KERNEL_VALIDATION_ERR(InContext, this, bQuiet, LOCTEXT("InvalidAttributeTypeNone", "'None' is not a valid GPU attribute type, check the 'Attributes to Create' array on your pins."));
+			PCG_KERNEL_VALIDATION_ERR(InContext, this, bQuiet, FText::Format(
+				LOCTEXT("InvalidAttributeTypeNone", "Attribute '{0}' has invalid GPU attribute type '{1}', check the 'Attributes to Create' array on your pins."),
+				FText::FromName(AttributeKey.Name),
+				FText::FromString(AttributeTypeEnum->GetNameStringByValue(static_cast<int64>(AttributeKey.Type)))));
 			return false;
 		}
 	}
diff --git a/Engine/Plugins/PCG/Source/PCG/Private/Compute/PCGDataForGPU.cpp b/Engine/Plugins/PCG/Source/PCG/Private/Compute/PCGDataForGPU.cpp
--- a/Engine/Plugins/PCG/Source/PCG/Private/Compute/PCGDataForGPU.cpp
+++ b/Engine/Plugins/PCG/Source/PCG/Private/Compute/PCGDataForGPU.cpp
@@ -23,15 +23,15 @@ namespace PCGDataForGPUConstants
 {
 	const static FPCGKernelAttributeDesc PointPropertyDescs[NUM_POINT_PROPERTIES] =
 	{
-		FPCGKernelAttributeDesc(POINT_POSITION_ATTRIBUTE_ID,   EPCGKernelAttributeType::Float3, NAME_None),
-		FPCGKernelAttributeDesc(POINT_ROTATION_ATTRIBUTE_ID,   EPCGKernelAttributeType::Quat,   NAME_None),
-		FPCGKernelAttributeDesc(POINT_SCALE_ATTRIBUTE_ID,      EPCGKernelAttributeType::Float3, NAME_None),
-		FPCGKernelAttributeDesc(POINT_BOUNDS_MIN_ATTRIBUTE_ID, EPCGKernelAttributeType::Float3, NAME_None),
-		FPCGKernelAttributeDesc(POINT_BOUNDS_MAX_ATTRIBUTE_ID, EPCGKernelAttributeType::Float3, NAME_None),
-		FPCGKernelAttributeDesc(POINT_COLOR_ATTRIBUTE_ID,      EPCGKernelAttributeType::Float4, NAME_None),
-		FPCGKernelAttributeDesc(POINT_DENSITY_ATTRIBUTE_ID,    EPCGKernelAttributeType::Float,  NAME_None),
-		FPCGKernelAttributeDesc(POINT_SEED_ATTRIBUTE_ID,       EPCGKernelAttributeType::Int,    NAME_None),
-		FPCGKernelAttributeDesc(POINT_STEEPNESS_ATTRIBUTE_ID,  EPCGKernelAttributeType::Float,  NAME_None)
+		FPCGKernelAttributeDesc(POINT_POSITION_ATTRIBUTE_ID,   EPCGKernelAttributeType::Float3, TEXT("$Position")),
+		FPCGKernelAttributeDesc(POINT_ROTATION_ATTRIBUTE_ID,   EPCGKernelAttributeType::Quat,   TEXT("$Rotation")),
+		FPCGKernelAttributeDesc(POINT_SCALE_ATTRIBUTE_ID,      EPCGKernelAttributeType::Float3, TEXT("$Scale")),
+		FPCGKernelAttributeDesc(POINT_BOUNDS_MIN_ATTRIBUTE_ID, EPCGKernelAttributeType::Float3, TEXT("$BoundsMin")),
+		FPCGKernelAttributeDesc(POINT_BOUNDS_MAX_ATTRIBUTE_ID, EPCGKernelAttributeType::Float3, TEXT("$BoundsMax")),
+		FPCGKernelAttributeDesc(POINT_COLOR_ATTRIBUTE_ID,      EPCGKernelAttributeType::Float4, TEXT("$Color")),
+		FPCGKernelAttributeDesc(POINT_DENSITY_ATTRIBUTE_ID,    EPCGKernelAttributeType::Float,  TEXT("$Density")),
+		FPCGKernelAttributeDesc(POINT_SEED_ATTRIBUTE_ID,       EPCGKernelAttributeType::Int,    TEXT("$Seed")),
+		FPCGKernelAttributeDesc(POINT_STEEPNESS_ATTRIBUTE_ID,  EPCGKernelAttributeType::Float,  TEXT("$Steepness"))
 	};
 }
 
@@ -648,13 +648,6 @@ void FPCGDataDesc::InitializeAttributeDescs(const UPCGData* InData, const TMap<F
 				break;
 			}
 
-			if (AttributeName == NAME_None)
-			{
-				// TODO: Would be nice to include the pin label for debug purposes
-				UE_LOG(LogPCG, Warning, TEXT("'%s' is not a valid name for a kernel attribute. Attribute will be skipped."), *AttributeName.ToString());
-				continue;
-			}
-
 			TArray<int32> UniqueStringKeys;
 
 			if (AttributeType == EPCGKernelAttributeType::StringKey)
diff --git a/Engine/Plugins/PCG/Source/PCG/Private/Elements/PCGStaticMeshSpawner.cpp b/Engine/Plugins/PCG/Source/PCG/Private/Elements/PCGStaticMeshSpawner.cpp
--- a/Engine/Plugins/PCG/Source/PCG/Private/Elements/PCGStaticMeshSpawner.cpp
+++ b/Engine/Plugins/PCG/Source/PCG/Private/Elements/PCGStaticMeshSpawner.cpp
@@ -75,15 +75,7 @@ bool UPCGStaticMeshSpawnerSettings::IsKernelValid(FPCGContext* InContext, bool b
 			}
 		}
 	}
-	else if (UPCGMeshSelectorByAttribute* SelectorByAttribute = Cast<UPCGMeshSelectorByAttribute>(MeshSelectorParameters))
-	{
-		if (SelectorByAttribute->AttributeName == NAME_None)
-		{
-			PCG_KERNEL_VALIDATION_ERR(InContext, this, bQuiet, LOCTEXT("NoMeshAttribute", "Mesh selector attribute name is 'None'."));
-			return false;
-		}
-	}
-	else
+	else if (!MeshSelectorParameters || !MeshSelectorParameters->IsA<UPCGMeshSelectorByAttribute>())
 	{
 		PCG_KERNEL_VALIDATION_ERR(InContext, this, bQuiet, LOCTEXT("InvalidMeshSelector", "Currently GPU Static Mesh Spawner nodes must use PCGMeshSelectorWeighted or UPCGMeshSelectorByAttribute as the mesh selector type."));
 		return false;
