Fixed stochastic repetition issue.

#rnx

diff --git a/Engine/Plugins/PCG/Source/PCG/Public/Elements/Grammar/PCGSubdivisionBase.h b/Engine/Plugins/PCG/Source/PCG/Public/Elements/Grammar/PCGSubdivisionBase.h
--- a/Engine/Plugins/PCG/Source/PCG/Public/Elements/Grammar/PCGSubdivisionBase.h
+++ b/Engine/Plugins/PCG/Source/PCG/Public/Elements/Grammar/PCGSubdivisionBase.h
@@ -7,6 +7,8 @@
 #include "Grammar/PCGGrammar.h"
 #include "Utils/PCGLogErrors.h"
 
+#include "Algo/Copy.h"
+
 #include "PCGSubdivisionBase.generated.h"
 
 class UPCGData;
@@ -168,7 +170,7 @@ namespace PCGSubdivisionBase
 		const T* Module = nullptr;
 		double ExtraScale = 0.0;
 		bool bIsValid = true;
-		bool bSkipExpansion = false;
+		bool bSkipExpansion = false; // FIXME used to mark a token so it cannot be repeated - TODO rename this to bCanBeRepeated
 	};
 
 	PCGGrammar::FTokenizedGrammar GetTokenizedGrammar(FPCGContext* InContext, const FString& InGrammar, const FModuleInfoMap& InModulesInfo, double& OutMinSize);
@@ -198,224 +200,204 @@ namespace PCGSubdivisionBase
 		}
 
 		// Working data sets
-		TArray<TModuleInstance<T>*> PriorityModules;
-		TArray<TModuleInstance<T>*> StochasticModules;
 		TArray<TModuleInstance<T>> ExpandedModules;
-		TArray<TModuleInstance<T>> PreviousModules;
+
+		auto ExpandSubmodule = [&ExpandedModules](const T& Submodule)
+		{
+			if (!Submodule.IsValid())
+			{
+				return;
+			}
+
+			int NumConcreteModules = Submodule.GetNumRepeat();
+			int NumRepeatModules = 0;
+
+			if (NumConcreteModules == PCGGrammar::InfiniteRepetition)
+			{
+				NumConcreteModules = 0;
+				NumRepeatModules = 1;
+			}
+			else if (NumConcreteModules == PCGGrammar::AtLeastOneRepetition)
+			{
+				NumConcreteModules = 1;
+				NumRepeatModules = 1;
+			}
+
+			for (int C = 0; C < NumConcreteModules; ++C)
+			{
+				TModuleInstance<T>& ExpandedSubmodule = ExpandedModules.Emplace_GetRef(&Submodule);
+				ExpandedSubmodule.bSkipExpansion = false;
+			}
+
+			for (int R = 0; R < NumRepeatModules; ++R)
+			{
+				TModuleInstance<T>& ExpandedSubmodule = ExpandedModules.Emplace_GetRef(&Submodule);
+				ExpandedSubmodule.bSkipExpansion = true;
+			}
+		};
 
 		bool bNotDone = true;
 		while (bNotDone)
 		{
 			bNotDone = false;
-			// 1. Expand/Visit sequences
+
+			// 1. Expand "concrete" symbols, e.g. root, sequence, stochastic, priority
 			// Implementation note; since we've already consumed the min size, there's no need to update anything when expanding nodes, only when replacing with another choice.
-			// We'll keep track of whether we have certain types of nodes at this point so we can short-circuit some passes
 			ExpandedModules.Reset();
 			ExpandedModules.Reserve(CurrentModules.Num());
-			bool bHasPriorityOrStochasticModules = false;
+
+			int NumExpandedModules = 0;
 
 			for (TModuleInstance<T>& CurrentModule : CurrentModules)
 			{
 				if (!CurrentModule.bIsValid)
 				{
-					continue; // discard this module (reclaim?)
+					// Module was discarded, could be culled. Normally should have no size
+					continue;
 				}
-
-				if ((CurrentModule.Module->GetType() == PCGGrammar::EModuleType::Root || CurrentModule.Module->GetType() == PCGGrammar::EModuleType::Sequence))
+				else if (CurrentModule.bSkipExpansion)
 				{
-					// When we duplicate sequences, we can't expand them directly otherwise they won't be repeatable
-					if (CurrentModule.bSkipExpansion)
-					{
-						// Copy module as-is, but remove the expansion limitation so we break it down next iteration
-						TModuleInstance<T>& ExpandedModule = ExpandedModules.Add_GetRef(CurrentModule);
-						ExpandedModule.bSkipExpansion = false;
-					}
-					else
-					{
-						for (const T& Submodule : CurrentModule.Module->Submodules)
-						{
-							if (Submodule.IsValid())
-							{
-								// Expand module if it has a positive number of iterations (>= 1) or if it the unit size would fit.
-								// This is where we do a difference also between the * (0+) and + cases; the + will have a force one repetition here, while the * will not.
-								int NumRepeats = Submodule.GetNumRepeat();
-								if (Submodule.GetNumRepeat() == PCGGrammar::AtLeastOneRepetition)
-								{
-									NumRepeats = 1;
-								}
-								else if (NumRepeats == PCGGrammar::InfiniteRepetition && RemainingLength >= Submodule.GetMinConcreteSize())
-								{
-									NumRepeats = 1;
-									RemainingLength -= Submodule.GetMinConcreteSize();
-								}
-
-								for (int R = 0; R < NumRepeats; ++R)
-								{
-									TModuleInstance<T>& ExpandedModule = ExpandedModules.Emplace_GetRef(&Submodule);
-								}
-
-								if (NumRepeats > 0)
-								{
-									bNotDone = true;
-									bHasPriorityOrStochasticModules |= (Submodule.GetType() == PCGGrammar::EModuleType::Priority || Submodule.GetType() == PCGGrammar::EModuleType::Stochastic);
-								}
-							}
-						}
-					}
-				}
-				else if (CurrentModule.Module->GetType() == PCGGrammar::EModuleType::Priority || CurrentModule.Module->GetType() == PCGGrammar::EModuleType::Stochastic)
-				{
-					bHasPriorityOrStochasticModules = true;
-
-					// Copy module as-is, but remove the expansion limitation so we break it down next iteration
-					TModuleInstance<T>& ExpandedModule = ExpandedModules.Add_GetRef(CurrentModule);
-					ExpandedModule.bSkipExpansion = false;
+					// Do not expand repetition modules until they are concretized
+					ExpandedModules.Add(CurrentModule);
 				}
-				else
+				else if (CurrentModule.Module->GetType() == PCGGrammar::EModuleType::Root || CurrentModule.Module->GetType() == PCGGrammar::EModuleType::Sequence)
 				{
-					// Copy module as-is, but do NOT remove the expansion limitation, otherwise we will have issues with literals with repetitions.
-					TModuleInstance<T>& ExpandedModule = ExpandedModules.Add_GetRef(CurrentModule);
-				}
-			}
-
-			CurrentModules = MoveTemp(ExpandedModules);
-
-			// 2- Make a copy of the current set of modules, because we'll use this for repeatable modules
-			PreviousModules = CurrentModules;
-
-			// 3- Perform choices
-			if (bHasPriorityOrStochasticModules)
-			{
-				// 3.1 - Gather priority & stochastic modules (pointers here)
-				PriorityModules.Reset();
-				StochasticModules.Reset();
-
-				for (TModuleInstance<T>& CurrentModule : CurrentModules)
-				{
-					check(CurrentModule.Module);
-					if (CurrentModule.Module->GetType() == PCGGrammar::EModuleType::Priority)
+					for (const T& Submodule : CurrentModule.Module->Submodules)
 					{
-						PriorityModules.Add(&CurrentModule);
+						ExpandSubmodule(Submodule);
 					}
-					else if (CurrentModule.Module->GetType() == PCGGrammar::EModuleType::Stochastic)
-					{
-						StochasticModules.Add(&CurrentModule);
-					}
-				}
 
-				bNotDone |= (!PriorityModules.IsEmpty() || !StochasticModules.IsEmpty());
-
-				// 3.2 - Apply priority choices first (from left to right)
-				for (TModuleInstance<T>* PriorityModule : PriorityModules)
+					++NumExpandedModules;
+				}
+				else if (CurrentModule.Module->GetType() == PCGGrammar::EModuleType::Priority)
 				{
-					bool bWasReplaced = false;
 					// Replace current module by the first of its childen that has a min size that fits in with the remaining length
-					for (const T& PrioritySubmodule : PriorityModule->Module->Submodules)
+					bool bModuleExpanded = false;
+					for (const T& Submodule : CurrentModule.Module->Submodules)
 					{
-						if (!PrioritySubmodule.IsValid())
+						if (!Submodule.IsValid())
 						{
 							continue;
 						}
 
-						const double DeltaMinSize = (PrioritySubmodule.GetMinConcreteSize() - PriorityModule->Module->GetConcreteUnitSize());
+						const double DeltaMinSize = (Submodule.GetMinSize() - CurrentModule.Module->GetUnitSize());
 						check(DeltaMinSize >= 0);
 						if (RemainingLength >= DeltaMinSize)
 						{
-							PriorityModule->Module = &PrioritySubmodule;
+							bModuleExpanded = true;
+							ExpandSubmodule(Submodule);
 							RemainingLength -= DeltaMinSize;
-							bWasReplaced = true;
 							break;
 						}
 					}
 
-					if (!bWasReplaced)
-					{
-						// TODO Kill this module + reclaim size(?)
-						PriorityModule->bIsValid = false;
-					}
+					// Implementation note: if the module doesn't pick anything (here) it will be removed automatically.
+					check(bModuleExpanded || CurrentModule.Module->GetMinSize() == 0);
+					++NumExpandedModules;
 				}
-
-				// 3.3 Process stochastic modules
-				// TODO random pick order of stochastic modules
-				for (TModuleInstance<T>* StochasticModule : StochasticModules)
+				else if (CurrentModule.Module->GetType() == PCGGrammar::EModuleType::Stochastic)
 				{
+					// TODO: investigate if doing expansion of stochastic modules in a random order would be a better idea.
 					// Replace current module by a random pick according to total weights of "valid" choices we can still make
-					int TotalWeight = 0;
-					for (const T& StochasticSubmodule : StochasticModule->Module->Submodules)
+					int TotalValidWeight = 0;
+
+					// Compute total weight for all valid permutations
+					for (const T& Submodule : CurrentModule.Module->Submodules)
 					{
-						if (!StochasticSubmodule.IsValid())
+						if (!Submodule.IsValid())
 						{
 							continue;
 						}
 
-						const double DeltaMinSize = (StochasticSubmodule.GetMinConcreteSize() - StochasticModule->Module->GetConcreteUnitSize());
+						const double DeltaMinSize = Submodule.GetMinSize() - CurrentModule.Module->GetUnitSize();
 						check(DeltaMinSize >= 0);
 						if (DeltaMinSize <= RemainingLength)
 						{
-							TotalWeight += StochasticSubmodule.GetWeight();
+							TotalValidWeight += Submodule.GetWeight();
 						}
 					}
 
-					if (TotalWeight == 0)
-					{
-						// Kill module + reclaim size(?)
-						StochasticModule->bIsValid = false;
-						continue;
-					}
+					int WeightPick = RandomStream.RandRange(0, TotalValidWeight - 1);
+					bool bModuleExpanded = false;
 
-					int WeightPick = RandomStream.RandRange(0, TotalWeight - 1);
-					bool bWasReplaced = false;
-
-					for (const T& StochasticSubmodule : StochasticModule->Module->Submodules)
+					for (const T& Submodule : CurrentModule.Module->Submodules)
 					{
-						if (!StochasticSubmodule.IsValid())
+						if (!Submodule.IsValid())
 						{
 							continue;
 						}
 
-						const double DeltaMinSize = (StochasticSubmodule.GetMinConcreteSize() - StochasticModule->Module->GetConcreteUnitSize());
+						const double DeltaMinSize = Submodule.GetMinSize() - CurrentModule.Module->GetUnitSize();
 						check(DeltaMinSize >= 0);
 						if (DeltaMinSize <= RemainingLength)
 						{
-							if (StochasticSubmodule.GetWeight() > WeightPick)
+							if (Submodule.GetWeight() > WeightPick)
 							{
-								bWasReplaced = true;
+								bModuleExpanded = true;
+								ExpandSubmodule(Submodule);
 								RemainingLength -= DeltaMinSize;
-								StochasticModule->Module = &StochasticSubmodule;
 								break;
 							}
 							else
 							{
-								WeightPick -= StochasticSubmodule.GetWeight();
+								WeightPick -= Submodule.GetWeight();
 							}
 						}
 					}
 
-					check(bWasReplaced);
+					check(bModuleExpanded);
+					++NumExpandedModules;
+				}
+				else // Literals - copy as-is
+				{
+					ExpandedModules.Add(CurrentModule);
 				}
 			}
 
-			// 4. Add arbitrary repetitions, requires to keep data from previous iteration
-			// Implementation note: since we add one module only in each iteration, it will definitely be inefficient when there is a long chain of things to duplicate.
-			check(PreviousModules.Num() == CurrentModules.Num());
-			int NumAddedRepetitions = 0;
+			// Move expanded modules to current modules
+			CurrentModules = MoveTemp(ExpandedModules);
+
+			bNotDone |= (NumExpandedModules > 0);
 
-			for (int ModuleIndex = 0; ModuleIndex < PreviousModules.Num(); ++ModuleIndex)
+			// 2. Concretize repetitions as needed
+			int NumConcretizedModules = 0;
+			for (int ModuleIndex = 0; ModuleIndex < CurrentModules.Num(); ++ModuleIndex)
 			{
-				TModuleInstance<T>& PreviousModule = PreviousModules[ModuleIndex];
-				if (!PreviousModule.bSkipExpansion && PreviousModule.Module->GetNumRepeat() < 0 && PreviousModule.Module->GetMinConcreteSize() <= RemainingLength)
+				TModuleInstance<T>& CurrentModule = CurrentModules[ModuleIndex];
+				if (!CurrentModule.bSkipExpansion)
+				{
+					continue;
+				}
+
+				check(CurrentModule.Module->GetMinConcreteSize() >= CurrentModule.Module->GetUnitSize());
+				check(CurrentModule.Module->GetUnitSize() >= CurrentModule.Module->GetMinSize());
+				if(CurrentModule.Module->GetMinConcreteSize() > RemainingLength)
+				{
+					CurrentModule.bIsValid = false;
+				}
+				else
 				{
-					// Mark the module to skip expansion, otherwise we can't chain for repeats
-					PreviousModule.bSkipExpansion = true;
-					CurrentModules.Insert(PreviousModule, ModuleIndex + NumAddedRepetitions + 1);
-					++NumAddedRepetitions;
-					RemainingLength -= PreviousModule.Module->GetMinConcreteSize();
+					// Update length
+					RemainingLength -= CurrentModule.Module->GetUnitSize();
+
+					// Duplicate this module, mark it non-repeatable
+					TModuleInstance<T> ModuleToDuplicate = CurrentModule;
+					ModuleToDuplicate.bSkipExpansion = false;
+
+					// Insert left
+					CurrentModules.Insert(ModuleToDuplicate, ModuleIndex);
+					++ModuleIndex;
+					++NumConcretizedModules;
 				}
 			}
 
-			bNotDone |= (NumAddedRepetitions > 0);
+			bNotDone |= (NumConcretizedModules > 0);
 		}
 
+		// Remove invalid modules
+		TArray<TModuleInstance<T>> PreviousModules = MoveTemp(CurrentModules);
+		Algo::CopyIf(PreviousModules, CurrentModules, [](const TModuleInstance<T>& InModule) { return InModule.bIsValid; });
+
 #if WITH_EDITOR
 		// Perform some early validation and see if there's a mismatch on the reported size and the one actually placed
 		double CountedLength = RemainingLength;
