[PCG] Make the steepness to 1 for all new points in the Create Points

#rnx

diff --git a/Engine/Plugins/PCG/Source/PCG/Private/Elements/PCGCreatePoints.cpp b/Engine/Plugins/PCG/Source/PCG/Private/Elements/PCGCreatePoints.cpp
--- a/Engine/Plugins/PCG/Source/PCG/Private/Elements/PCGCreatePoints.cpp
+++ b/Engine/Plugins/PCG/Source/PCG/Private/Elements/PCGCreatePoints.cpp
@@ -15,8 +15,14 @@
 
 UPCGCreatePointsSettings::UPCGCreatePointsSettings()
 {
-	// Add one default point in the array
-	PointsToCreate.Add(FPCGPoint());
+	// Add one default point in the array, with a steepness of 1 for new nodes.
+	FPCGPoint DefaultPoint{};
+	if (PCGHelpers::IsNewObjectAndNotDefault(this))
+	{
+		DefaultPoint.Steepness = 1.0f;
+	}
+	
+	PointsToCreate.Add(DefaultPoint);
 }
 
 void UPCGCreatePointsSettings::PostLoad()
@@ -34,6 +40,20 @@ PRAGMA_ENABLE_DEPRECATION_WARNINGS
 #endif
 }
 
+#if WITH_EDITOR
+void UPCGCreatePointsSettings::PostEditChangeChainProperty(struct FPropertyChangedChainEvent& PropertyChangedEvent)
+{
+	Super::PostEditChangeChainProperty(PropertyChangedEvent);
+
+	if (PropertyChangedEvent.GetPropertyName() == GET_MEMBER_NAME_CHECKED(UPCGCreatePointsSettings, PointsToCreate)
+		&& PropertyChangedEvent.ChangeType == EPropertyChangeType::ArrayAdd && !PointsToCreate.IsEmpty())
+	{
+		// Force the steepness to 1 to any new added points to the array.
+		PointsToCreate.Last().Steepness = 1.0f;	
+	}
+}
+#endif // WITH_EDITOR
+
 TArray<FPCGPinProperties> UPCGCreatePointsSettings::InputPinProperties() const
 {
 	return TArray<FPCGPinProperties>();
diff --git a/Engine/Plugins/PCG/Source/PCG/Public/Elements/PCGCreatePoints.h b/Engine/Plugins/PCG/Source/PCG/Public/Elements/PCGCreatePoints.h
--- a/Engine/Plugins/PCG/Source/PCG/Public/Elements/PCGCreatePoints.h
+++ b/Engine/Plugins/PCG/Source/PCG/Public/Elements/PCGCreatePoints.h
@@ -29,6 +29,9 @@ public:
 
 	//~Begin UObject interface
 	virtual void PostLoad() override;
+#if WITH_EDITOR
+	virtual void PostEditChangeChainProperty(struct FPropertyChangedChainEvent& PropertyChangedEvent) override;
+#endif // WITH_EDITOR
 	//~End UObject interface
 
 	//~Begin UPCGSettings interface
