[PCG] Demote missing data layer instance error to display as we can run builder on multiple maps (which don't have the same data layers)
- Add missing settings log for included/excluded data layers

#rnx

diff --git a/Engine/Plugins/PCG/Source/PCGEditor/Private/WorldPartitionBuilder/PCGWorldPartitionBuilder.cpp b/Engine/Plugins/PCG/Source/PCGEditor/Private/WorldPartitionBuilder/PCGWorldPartitionBuilder.cpp
--- a/Engine/Plugins/PCG/Source/PCGEditor/Private/WorldPartitionBuilder/PCGWorldPartitionBuilder.cpp
+++ b/Engine/Plugins/PCG/Source/PCGEditor/Private/WorldPartitionBuilder/PCGWorldPartitionBuilder.cpp
@@ -376,6 +376,21 @@ bool FPCGWorldPartitionBuilderArgs::InitializeFrom(const FPCGWorldPartitionComma
 		BuilderArgs.IterativeCellSize,
 		bFromCommandline ? FromCommandline : FromDefault);
 
+	// Log DataLayers
+	TArray<FString> DataLayersToString;
+	Algo::Transform(BuilderArgs.IncludedDataLayers, DataLayersToString, [](const TSoftObjectPtr<UDataLayerAsset>& DataLayerAsset) { return DataLayerAsset.ToString(); });
+
+	UE_LOG(LogPCGWorldPartitionBuilder, Display, TEXT("IncludedDataLayers '%s' (%s)"),
+		DataLayersToString.IsEmpty() ? TEXT("None") : *FString::Join(DataLayersToString, TEXT(",")),
+		FromDefault);
+
+	DataLayersToString.Empty();
+	Algo::Transform(BuilderArgs.ExcludedDataLayers, DataLayersToString, [](const TSoftObjectPtr<UDataLayerAsset>& DataLayerAsset) { return DataLayerAsset.ToString(); });
+
+	UE_LOG(LogPCGWorldPartitionBuilder, Display, TEXT("ExcludedDataLayers '%s' (%s)"),
+		DataLayersToString.IsEmpty() ? TEXT("None") : *FString::Join(DataLayersToString, TEXT(",")),
+		FromDefault);
+
 	OutArgs = BuilderArgs;
 	return true;
 }
@@ -485,7 +500,7 @@ bool UPCGWorldPartitionBuilder::PreRun(UWorld* World, FPackageSourceControlHelpe
 	// Setup DataLayers
 	if (UDataLayerManager* DataLayerManager = UDataLayerManager::GetDataLayerManager(World))
 	{
-		auto FillDataLayerShortNames = [DataLayerManager](const TArray<TObjectPtr<UDataLayerAsset>>& InDataLayerAssets, TSet<FName>& OutDataLayerShortNames)
+		auto FillDataLayerShortNames = [World, DataLayerManager](const TArray<TObjectPtr<UDataLayerAsset>>& InDataLayerAssets, TSet<FName>& OutDataLayerShortNames)
 		{
 			for(const TObjectPtr<UDataLayerAsset>& DataLayer : InDataLayerAssets)
 			{
@@ -494,7 +509,7 @@ bool UPCGWorldPartitionBuilder::PreRun(UWorld* World, FPackageSourceControlHelpe
 					const UDataLayerInstance* DataLayerInstance = DataLayerManager->GetDataLayerInstance(DataLayer);
 					if (DataLayerInstance == nullptr)
 					{
-						UE_LOG(LogPCGWorldPartitionBuilder, Error, TEXT("Missing UDataLayerInstance for %s."), *DataLayer->GetName());
+						UE_LOG(LogPCGWorldPartitionBuilder, Display, TEXT("World '%s' does not contain a DataLayerInstance for '%s'"), *World->GetPathName(), *DataLayer->GetPathName());
 					}
 					else
 					{
