[PCG] Fix SM Spawner mesh selector / instance data attribute validation not working.

Also fixed kernel not writing data when there were no instances to create.

#rnx

diff --git a/Engine/Plugins/PCG/Shaders/Private/Elements/PCGStaticMeshSpawner.usf b/Engine/Plugins/PCG/Shaders/Private/Elements/PCGStaticMeshSpawner.usf
--- a/Engine/Plugins/PCG/Shaders/Private/Elements/PCGStaticMeshSpawner.usf
+++ b/Engine/Plugins/PCG/Shaders/Private/Elements/PCGStaticMeshSpawner.usf
@@ -26,22 +26,33 @@ void Main(uint3 GroupId : SV_GroupID, uint GroupIndex : SV_GroupIndex)
 	}
 
 	const uint ThreadIndex = GetUnWrappedDispatchThreadId(GroupId, GroupIndex, 64);
-	if (ThreadIndex >= InstanceData_GetNumInstancesAllPrimitives())
+
+	uint In_DataIndex, In_ElementIndex;
+	if (!In_GetThreadData(ThreadIndex, In_DataIndex, In_ElementIndex))
 	{
 		return;
 	}
 
-	uint In_DataIndex, In_ElementIndex;
-	if (!In_GetThreadData(ThreadIndex, In_DataIndex, In_ElementIndex))
+	uint Out_DataIndex, Out_ElementIndex;
+	if (!Out_GetThreadData(ThreadIndex, Out_DataIndex, Out_ElementIndex))
 	{
+		// Should never happen.
 		return;
 	}
 
-	const int NumPrimitives = InstanceData_GetNumPrimitives();
+	if (In_IsPointRemoved(In_DataIndex, In_ElementIndex))
+	{
+		// Propagate the 'Removed' status to the output point if input point is removed. Otherwise, this point will not be culled.
+		Out_RemovePoint(Out_DataIndex, Out_ElementIndex);
 
-	uint PrimitiveIndex = 0;
+		return;
+	}
 
-	if (NumPrimitives > 1)
+	uint PrimitiveIndex = (uint)-1;
+
+	const int NumPrimitives = InstanceData_GetNumPrimitives();
+
+	if (NumPrimitives > 0)
 	{
 		const uint SelectorAttributeId = SMSpawner_GetSelectorAttributeId();
 		
@@ -64,7 +75,7 @@ void Main(uint3 GroupId : SV_GroupID, uint GroupIndex : SV_GroupIndex)
 
 			const float RandomDraw = FRand(Seed);
 
-			for (; PrimitiveIndex < uint(NumPrimitives) - 1; ++PrimitiveIndex)
+			for (PrimitiveIndex = 0; PrimitiveIndex < uint(NumPrimitives) - 1; ++PrimitiveIndex)
 			{
 				if (RandomDraw < SMSpawner_GetPrimitiveSelectionCDF(PrimitiveIndex))
 				{
@@ -74,26 +85,11 @@ void Main(uint3 GroupId : SV_GroupID, uint GroupIndex : SV_GroupIndex)
 		}
 	}
 
-	uint Out_DataIndex, Out_ElementIndex;
-	if (!Out_GetThreadData(ThreadIndex, Out_DataIndex, Out_ElementIndex))
-	{
-		// Should never happen.
-		return;
-	}
-
-	if (In_IsPointRemoved(In_DataIndex, In_ElementIndex))
-	{
-		// Propagate the 'Removed' status to the output point if input point is removed. Otherwise, this point will not be culled.
-		Out_RemovePoint(Out_DataIndex, Out_ElementIndex);
-
-		return;
-	}
-
 	// Get instance index to write to. Will return a negative value if we have saturated instances for this primitive, in which case we will not push the data.
-	const int InstanceIndexInt = InstanceData_GetIndexToWriteTo(PrimitiveIndex);
+	const int InstanceIndexInt = (PrimitiveIndex != (uint)-1) ? InstanceData_GetIndexToWriteTo(PrimitiveIndex) : -1;
 
 	// Write instance transform, only if not culled by primitive instance index
-	if(InstanceIndexInt >= 0)
+	if (InstanceIndexInt >= 0)
 	{
 		const uint InstanceIndex = InstanceIndexInt;
 		const float4x4 Transform = In_GetPointTransform(In_DataIndex, In_ElementIndex);
@@ -127,7 +123,7 @@ void Main(uint3 GroupId : SV_GroupID, uint GroupIndex : SV_GroupIndex)
 	Out_SetSteepness(Out_DataIndex, Out_ElementIndex, In_GetSteepness(In_DataIndex, In_ElementIndex));
 
 	// Either apply SM bounds to points or pass bounds through.
-	if (SMSpawner_ShouldApplyBounds())
+	if (PrimitiveIndex != (uint)-1 && SMSpawner_ShouldApplyBounds())
 	{
 		Out_SetBoundsMin(Out_DataIndex, Out_ElementIndex, SMSpawner_GetPrimitiveMeshBoundsMin(PrimitiveIndex));
 		Out_SetBoundsMax(Out_DataIndex, Out_ElementIndex, SMSpawner_GetPrimitiveMeshBoundsMax(PrimitiveIndex));
@@ -142,7 +138,7 @@ void Main(uint3 GroupId : SV_GroupID, uint GroupIndex : SV_GroupIndex)
 
 	if (SMSpawner_GetSelectedMeshAttributeId() != (uint)-1)
 	{
-		const int OutputMeshStringKey = SMSpawner_GetPrimitiveStringKey(PrimitiveIndex);
+		const int OutputMeshStringKey = (PrimitiveIndex != (uint)-1) ? SMSpawner_GetPrimitiveStringKey(PrimitiveIndex) : -1;
 		Out_SetStringKey(Out_DataIndex, Out_ElementIndex, SMSpawner_GetSelectedMeshAttributeId(), OutputMeshStringKey);
 	}
 }
diff --git a/Engine/Plugins/PCG/Source/PCG/Private/Compute/Elements/PCGComputeGraphElement.cpp b/Engine/Plugins/PCG/Source/PCG/Private/Compute/Elements/PCGComputeGraphElement.cpp
--- a/Engine/Plugins/PCG/Source/PCG/Private/Compute/Elements/PCGComputeGraphElement.cpp
+++ b/Engine/Plugins/PCG/Source/PCG/Private/Compute/Elements/PCGComputeGraphElement.cpp
@@ -355,6 +355,21 @@ bool FPCGComputeGraphElement::SetupProceduralISMComponents(FPCGContext* InContex
 			TArray<FName> AttributeNames;
 			SpawnerSettings->InstanceDataPackerParameters->GetAttributeNames(&AttributeNames);
 
+			for (FName AttributeName : AttributeNames)
+			{
+				if (AttributeName == NAME_None)
+				{
+					UE_LOG(LogPCG, Error, TEXT("Invalid instance data attribute specified '%s'."), *AttributeName.ToString());
+					continue;
+				}
+
+				if (!InBinding->GetAttributeLookupTable().Find(AttributeName) && !InBinding->GetInputAttributeTable().Find(AttributeName))
+				{
+					UE_LOG(LogPCG, Error, TEXT("Instance data attribute '%s' not found."), *AttributeName.ToString());
+					continue;
+				}
+			}
+
 			PCGDataForGPUHelpers::ComputeCustomFloatPacking(AttributeNames, InBinding, InputDataDesc, CustomFloatCount, AttributeIdOffsetStrides);
 		}
 
@@ -369,12 +384,24 @@ bool FPCGComputeGraphElement::SetupProceduralISMComponents(FPCGContext* InContex
 			const FName SelectedName = SelectorByAttribute->AttributeName;
 
 			const FPCGKernelAttributeIDAndType* FoundAttribute = InBinding->GetAttributeLookupTable().Find(SelectedName);
+
+			if (!FoundAttribute)
+			{
+				FoundAttribute = InBinding->GetInputAttributeTable().Find(SelectedName);
+			}
+
 			if (!FoundAttribute)
 			{
 				UE_LOG(LogPCG, Error, TEXT("Mesh selector attribute '%s' not found."), *SelectedName.ToString());
 				continue;
 			}
 
+			if (FoundAttribute->Type != EPCGKernelAttributeType::StringKey)
+			{
+				UE_LOG(LogPCG, Error, TEXT("Mesh selector attribute '%s' must be of type String Key."), *SelectedName.ToString());
+				continue;
+			}
+
 			SelectorAttributeId = FoundAttribute->Id;
 
 			// Compute how many unique incoming values we will receive.
diff --git a/Engine/Plugins/PCG/Source/PCG/Private/Compute/PCGDataBinding.cpp b/Engine/Plugins/PCG/Source/PCG/Private/Compute/PCGDataBinding.cpp
--- a/Engine/Plugins/PCG/Source/PCG/Private/Compute/PCGDataBinding.cpp
+++ b/Engine/Plugins/PCG/Source/PCG/Private/Compute/PCGDataBinding.cpp
@@ -76,6 +76,43 @@ void UPCGDataBinding::InitializeInputData(const FPCGDataCollection& InComputeGra
 	DataForGPU.InputPinLabelAliases = Graph->InputPinLabelAliases;
 
 	BuildStringTable();
+
+	BuildInputAttributeTable(InComputeGraphElementInputData);
+}
+
+void UPCGDataBinding::BuildInputAttributeTable(const FPCGDataCollection& InComputeGraphElementInputData)
+{
+	const UEnum* PCGKernelAttributeTypeEnum = StaticEnum<EPCGKernelAttributeType>();
+	check(PCGKernelAttributeTypeEnum);
+
+	const FPCGDataCollectionDesc Desc = FPCGDataCollectionDesc::BuildFromDataCollection(InComputeGraphElementInputData, GetAttributeLookupTable(), GetStringTable());
+	for (const FPCGDataDesc& DataDesc : Desc.DataDescs)
+	{
+		for (const FPCGKernelAttributeDesc& AttrDesc : DataDesc.AttributeDescs)
+		{
+			// Skip non-attributes like point properties.
+			if (AttrDesc.Name == NAME_None)
+			{
+				continue;
+			}
+
+			if (FPCGKernelAttributeIDAndType* AttrIDAndType = InputAttributeTable.Find(AttrDesc.Name))
+			{
+				if (AttrIDAndType->Type != AttrDesc.Type)
+				{
+					UE_LOG(LogPCG, Warning, TEXT("Encountered attribute '%s' multiple times in input data with different types (%s vs %s). Attributes currently must have a single type in compute graphs."),
+						*AttrDesc.Name.ToString(),
+						*PCGKernelAttributeTypeEnum->GetNameStringByValue(static_cast<int64>(AttrDesc.Type)),
+						*PCGKernelAttributeTypeEnum->GetNameStringByValue(static_cast<int64>(AttrIDAndType->Type)));
+				}
+			}
+			else
+			{
+				const FPCGKernelAttributeIDAndType IDAndType(AttrDesc.Index, AttrDesc.Type);
+				InputAttributeTable.Add(AttrDesc.Name, IDAndType);
+			}
+		}
+	}
 }
 
 void UPCGDataBinding::BuildStringTable()
diff --git a/Engine/Plugins/PCG/Source/PCG/Private/Compute/PCGDataBinding.h b/Engine/Plugins/PCG/Source/PCG/Private/Compute/PCGDataBinding.h
--- a/Engine/Plugins/PCG/Source/PCG/Private/Compute/PCGDataBinding.h
+++ b/Engine/Plugins/PCG/Source/PCG/Private/Compute/PCGDataBinding.h
@@ -51,6 +51,9 @@ public:
 	/** Get the global attribute information for this graph instance. */
 	const TMap<FName, FPCGKernelAttributeIDAndType>& GetAttributeLookupTable() const { return GlobalAttributeLookupTable; }
 
+	/** Get the table of attributes present in the compute graph element input data. */
+	const TMap<FName, FPCGKernelAttributeIDAndType>& GetInputAttributeTable() const { return InputAttributeTable; }
+
 	/** Get the set of all unique strings entering the compute graph. */
 	const TArray<FString>& GetStringTable() const { return StringTable; }
 
@@ -60,6 +63,8 @@ public:
 private:
 	void InitializeInputData(const FPCGDataCollection& InComputeGraphElementInputData);
 
+	void BuildInputAttributeTable(const FPCGDataCollection& InComputeGraphElementInputData);
+
 	void BuildStringTable();
 
 	/** Some attributes in table may be missing type information that we could not infer statically. Fill in missing types from the execution-time data. */
@@ -85,9 +90,14 @@ public:
 	TMap<TObjectPtr<const UPCGSettings>, FPCGSpawnerPrimitives> MeshSpawnersToPrimitives;
 
 protected:
+	/** Attributes read/written/created by kernels. */
 	UPROPERTY()
 	TMap<FName /* Attribute name */, FPCGKernelAttributeIDAndType> GlobalAttributeLookupTable;
 
+	/** Attributes present across all input data. */
+	UPROPERTY()
+	TMap<FName /* Attribute name */, FPCGKernelAttributeIDAndType> InputAttributeTable;
+
 	UPROPERTY()
 	TArray<FString> StringTable;
 };
diff --git a/Engine/Plugins/PCG/Source/PCG/Private/Compute/PCGDataForGPU.cpp b/Engine/Plugins/PCG/Source/PCG/Private/Compute/PCGDataForGPU.cpp
--- a/Engine/Plugins/PCG/Source/PCG/Private/Compute/PCGDataForGPU.cpp
+++ b/Engine/Plugins/PCG/Source/PCG/Private/Compute/PCGDataForGPU.cpp
@@ -746,6 +746,26 @@ void FPCGDataDesc::InitializeAttributeDescs(const UPCGData* InData, const TMap<F
 	}
 }
 
+FPCGDataCollectionDesc FPCGDataCollectionDesc::BuildFromDataCollection(
+	const FPCGDataCollection& InDataCollection,
+	const TMap<FName, FPCGKernelAttributeIDAndType>& InAttributeLookupTable,
+	const TArray<FString>& InStringTable)
+{
+	FPCGDataCollectionDesc CollectionDesc;
+
+	for (const FPCGTaggedData& Data : InDataCollection.TaggedData)
+	{
+		if (!Data.Data || !PCGComputeHelpers::IsTypeAllowedInDataCollection(Data.Data->GetDataType()))
+		{
+			continue;
+		}
+
+		CollectionDesc.DataDescs.Emplace(Data.Data, InAttributeLookupTable, InStringTable);
+	}
+
+	return CollectionDesc;
+}
+
 FPCGDataCollectionDesc FPCGDataCollectionDesc::BuildFromInputDataCollectionAndInputPinLabel(
 	const FPCGDataCollection& InDataCollection,
 	FName InputPinLabel,
diff --git a/Engine/Plugins/PCG/Source/PCG/Private/Compute/PCGDataForGPU.h b/Engine/Plugins/PCG/Source/PCG/Private/Compute/PCGDataForGPU.h
--- a/Engine/Plugins/PCG/Source/PCG/Private/Compute/PCGDataForGPU.h
+++ b/Engine/Plugins/PCG/Source/PCG/Private/Compute/PCGDataForGPU.h
@@ -126,6 +126,11 @@ private:
 
 struct FPCGDataCollectionDesc
 {
+	static FPCGDataCollectionDesc BuildFromDataCollection(
+		const FPCGDataCollection& InDataCollection,
+		const TMap<FName, FPCGKernelAttributeIDAndType>& InAttributeLookupTable,
+		const TArray<FString>& InStringTable);
+
 	static FPCGDataCollectionDesc BuildFromInputDataCollectionAndInputPinLabel(
 		const FPCGDataCollection& InDataCollection,
 		FName InputPinLabel,
diff --git a/Engine/Plugins/PCG/Source/PCG/Private/Elements/PCGStaticMeshSpawner.cpp b/Engine/Plugins/PCG/Source/PCG/Private/Elements/PCGStaticMeshSpawner.cpp
--- a/Engine/Plugins/PCG/Source/PCG/Private/Elements/PCGStaticMeshSpawner.cpp
+++ b/Engine/Plugins/PCG/Source/PCG/Private/Elements/PCGStaticMeshSpawner.cpp
@@ -176,25 +176,6 @@ bool UPCGStaticMeshSpawnerSettings::ComputeOutputPinDataDesc(const UPCGPin* Outp
 
 void UPCGStaticMeshSpawnerSettings::GetKernelAttributeKeys(TArray<FPCGKernelAttributeKey>& OutKeys) const
 {
-	if (InstanceDataPackerParameters)
-	{
-		TArray<FName> AttributeNames;
-		if (ensure(InstanceDataPackerParameters->GetAttributeNames(&AttributeNames)))
-		{
-			for (const FName& AttributeName : AttributeNames)
-			{
-				// We don't know the type statically before execution, leave unset.
-				OutKeys.AddUnique(FPCGKernelAttributeKey(AttributeName, EPCGKernelAttributeType::None));
-			}
-		}
-	}
-
-	if (const UPCGMeshSelectorByAttribute* Selector = Cast<UPCGMeshSelectorByAttribute>(MeshSelectorParameters))
-	{
-		// Add an attribute key for the given attribute so we register it as being read.
-		OutKeys.AddUnique(FPCGKernelAttributeKey(Selector->AttributeName, EPCGKernelAttributeType::None));
-	}
-
 	// Write out attribute.
 	OutKeys.AddUnique(FPCGKernelAttributeKey(OutAttributeName, EPCGKernelAttributeType::StringKey));
 }
