[PCG] Put FNames in string table to make robust against crashes when unpacking invalid name IDs.

#rnx

diff --git a/Engine/Plugins/PCG/Shaders/Private/PCGDataCollectionDataInterface.ush b/Engine/Plugins/PCG/Shaders/Private/PCGDataCollectionDataInterface.ush
--- a/Engine/Plugins/PCG/Shaders/Private/PCGDataCollectionDataInterface.ush
+++ b/Engine/Plugins/PCG/Shaders/Private/PCGDataCollectionDataInterface.ush
@@ -233,11 +233,11 @@ int GetStringKey_{DataInterfaceName}(uint DataIndex, uint ElementIndex, int Attr
 	return GetInt_{DataInterfaceName}(DataIndex, ElementIndex, AttributeId);
 }
 
+// TODO: left as uint2 which we locked into for 5.5.0. Future version should change this to int.
 uint2 GetName_{DataInterfaceName}(uint DataIndex, uint ElementIndex, int AttributeId)
 {
-	// Names are composed of 2 uints.
-	const uint AttributeAddress = GetFirstElementAddressInternal_{DataInterfaceName}(DataIndex, AttributeId);
-	return {DataInterfaceName}_DataCollectionBuffer.Load2(AttributeAddress + ElementIndex * 2 * 4);
+	// Names are represented as ints.
+	return uint2((uint)GetInt_{DataInterfaceName}(DataIndex, ElementIndex, AttributeId), 0);
 }
 
 // #################### ATTRIBUTE SETTERS ##########################
@@ -312,11 +312,11 @@ void SetStringKey_{DataInterfaceName}(uint DataIndex, uint ElementIndex, int Att
 	SetInt_{DataInterfaceName}(DataIndex, ElementIndex, AttributeId, Value);
 }
 
+// TODO: Value left as uint2 which we locked into for 5.5.0. Future version should change this to int.
 void SetName_{DataInterfaceName}(uint DataIndex, uint ElementIndex, int AttributeId, uint2 Value)
 {
-	// Names are composed of 2 uints.
-	const uint AttributeAddress = GetFirstElementAddressInternal_{DataInterfaceName}(DataIndex, AttributeId);
-	{DataInterfaceName}_DataCollectionBuffer.Store2(AttributeAddress + ElementIndex * 2 * 4, Value);
+	// Names are represented as ints.
+	SetInt_{DataInterfaceName}(DataIndex, ElementIndex, AttributeId, (int)Value.x);
 }
 
 // #################### POINT ATTRIBUTE GETTERS ##########################
diff --git a/Engine/Plugins/PCG/Source/PCG/Private/Compute/Elements/PCGComputeGraphElement.cpp b/Engine/Plugins/PCG/Source/PCG/Private/Compute/Elements/PCGComputeGraphElement.cpp
--- a/Engine/Plugins/PCG/Source/PCG/Private/Compute/Elements/PCGComputeGraphElement.cpp
+++ b/Engine/Plugins/PCG/Source/PCG/Private/Compute/Elements/PCGComputeGraphElement.cpp
@@ -412,7 +412,14 @@ bool FPCGComputeGraphElement::SetupProceduralISMComponents(FPCGContext* InContex
 				{
 					if (AttributeDesc.Name == SelectedName)
 					{
-						UniqueStringKeys.Append(AttributeDesc.UniqueStringKeys);
+						if (AttributeDesc.Type == EPCGKernelAttributeType::StringKey)
+						{
+							UniqueStringKeys.Append(AttributeDesc.UniqueStringKeys);
+						}
+						else
+						{
+							UE_LOG(LogPCG, Error, TEXT("Attribute '%s' not usable for mesh selection, only attributes of type String Key are supported."), *AttributeDesc.Name.ToString())
+						}
 					}
 				}
 			}
diff --git a/Engine/Plugins/PCG/Source/PCG/Private/Compute/Elements/PCGCustomHLSL.cpp b/Engine/Plugins/PCG/Source/PCG/Private/Compute/Elements/PCGCustomHLSL.cpp
--- a/Engine/Plugins/PCG/Source/PCG/Private/Compute/Elements/PCGCustomHLSL.cpp
+++ b/Engine/Plugins/PCG/Source/PCG/Private/Compute/Elements/PCGCustomHLSL.cpp
@@ -361,7 +361,12 @@ bool UPCGCustomHLSLSettings::ComputeOutputPinDataDesc(const UPCGPin* OutputPin,
 
 	for (const FPCGDataDesc& DataDesc : PinDesc.DataDescs)
 	{
-		if (DataDesc.AttributeDescs.FindByPredicate([](const FPCGKernelAttributeDesc& InAttributeDesc) { return InAttributeDesc.Type == EPCGKernelAttributeType::StringKey; }))
+		auto HasStringKey = [](const FPCGKernelAttributeDesc& InAttributeDesc)
+		{
+			return InAttributeDesc.Type == EPCGKernelAttributeType::StringKey || InAttributeDesc.Type == EPCGKernelAttributeType::Name;
+		};
+
+		if (DataDesc.AttributeDescs.FindByPredicate(HasStringKey))
 		{
 			bOutputHasStringKeys = true;
 			break;
@@ -384,7 +389,7 @@ bool UPCGCustomHLSLSettings::ComputeOutputPinDataDesc(const UPCGPin* OutputPin,
 			{
 				for (FPCGKernelAttributeDesc& AttributeDesc : DataDesc.AttributeDescs)
 				{
-					bFoundStringKeyAttribute |= (AttributeDesc.Type == EPCGKernelAttributeType::StringKey);
+					bFoundStringKeyAttribute |= (AttributeDesc.Type == EPCGKernelAttributeType::StringKey || AttributeDesc.Type == EPCGKernelAttributeType::Name);
 				}
 			}
 
@@ -400,7 +405,7 @@ bool UPCGCustomHLSLSettings::ComputeOutputPinDataDesc(const UPCGPin* OutputPin,
 			{
 				for (FPCGKernelAttributeDesc& AttributeDesc : DataDesc.AttributeDescs)
 				{
-					if (AttributeDesc.Type != EPCGKernelAttributeType::StringKey)
+					if (AttributeDesc.Type != EPCGKernelAttributeType::StringKey && AttributeDesc.Type != EPCGKernelAttributeType::Name)
 					{
 						continue;
 					}
@@ -415,7 +420,8 @@ bool UPCGCustomHLSLSettings::ComputeOutputPinDataDesc(const UPCGPin* OutputPin,
 						{
 							for (const FPCGKernelAttributeDesc& InputAttributeDesc : InputDataDesc.AttributeDescs)
 							{
-								if (InputAttributeDesc.Type == EPCGKernelAttributeType::StringKey && InputAttributeDesc.Name == AttributeDesc.Name)
+								const bool bIsString = InputAttributeDesc.Type == EPCGKernelAttributeType::StringKey || InputAttributeDesc.Type == EPCGKernelAttributeType::Name;
+								if (bIsString && InputAttributeDesc.Name == AttributeDesc.Name)
 								{
 									AttributeDesc.UniqueStringKeys.Append(InputAttributeDesc.UniqueStringKeys);
 									bFoundMatchingAttribute = true;
@@ -437,7 +443,7 @@ bool UPCGCustomHLSLSettings::ComputeOutputPinDataDesc(const UPCGPin* OutputPin,
 							{
 								for (const FPCGKernelAttributeDesc& InputAttributeDesc : InputDataDesc.AttributeDescs)
 								{
-									if (InputAttributeDesc.Type == EPCGKernelAttributeType::StringKey)
+									if (InputAttributeDesc.Type == EPCGKernelAttributeType::StringKey || InputAttributeDesc.Type == EPCGKernelAttributeType::Name)
 									{
 										AttributeDesc.UniqueStringKeys.Append(InputAttributeDesc.UniqueStringKeys);
 									}
diff --git a/Engine/Plugins/PCG/Source/PCG/Private/Compute/PCGDataBinding.cpp b/Engine/Plugins/PCG/Source/PCG/Private/Compute/PCGDataBinding.cpp
--- a/Engine/Plugins/PCG/Source/PCG/Private/Compute/PCGDataBinding.cpp
+++ b/Engine/Plugins/PCG/Source/PCG/Private/Compute/PCGDataBinding.cpp
@@ -143,7 +143,7 @@ void UPCGDataBinding::BuildStringTable()
 			const FName AttributeName = AttributeNames[AttributeIndex];
 			const EPCGKernelAttributeType AttributeType = PCGDataForGPUHelpers::GetAttributeTypeFromMetadataType(AttributeTypes[AttributeIndex]);
 
-			if (AttributeType == EPCGKernelAttributeType::StringKey)
+			if (AttributeType == EPCGKernelAttributeType::StringKey || AttributeType == EPCGKernelAttributeType::Name)
 			{
 				const FPCGMetadataAttributeBase* AttributeBase = Metadata->GetConstAttribute(AttributeName);
 				if (!AttributeBase)
@@ -163,9 +163,13 @@ void UPCGDataBinding::BuildStringTable()
 						StringTable.AddUnique(static_cast<const FPCGMetadataAttribute<FString>*>(AttributeBase)->GetValue(InValueKey));
 					}
 					else if (AttributeBase->GetTypeId() == PCG::Private::MetadataTypes<FSoftClassPath>::Id)
-                    {
-                    	StringTable.AddUnique(static_cast<const FPCGMetadataAttribute<FSoftClassPath>*>(AttributeBase)->GetValue(InValueKey).ToString());
-                    }
+					{
+						StringTable.AddUnique(static_cast<const FPCGMetadataAttribute<FSoftClassPath>*>(AttributeBase)->GetValue(InValueKey).ToString());
+					}
+					else if (AttributeBase->GetTypeId() == PCG::Private::MetadataTypes<FName>::Id)
+					{
+						StringTable.AddUnique(static_cast<const FPCGMetadataAttribute<FName>*>(AttributeBase)->GetValue(InValueKey).ToString());
+					}
 					else
 					{
 						checkNoEntry();
diff --git a/Engine/Plugins/PCG/Source/PCG/Private/Compute/PCGDataForGPU.cpp b/Engine/Plugins/PCG/Source/PCG/Private/Compute/PCGDataForGPU.cpp
--- a/Engine/Plugins/PCG/Source/PCG/Private/Compute/PCGDataForGPU.cpp
+++ b/Engine/Plugins/PCG/Source/PCG/Private/Compute/PCGDataForGPU.cpp
@@ -80,9 +80,9 @@ namespace PCGDataForGPUHelpers
 		case EPCGKernelAttributeType::Int:
 		case EPCGKernelAttributeType::Float:
 		case EPCGKernelAttributeType::StringKey:
+		case EPCGKernelAttributeType::Name:
 			return 4;
 		case EPCGKernelAttributeType::Float2:
-		case EPCGKernelAttributeType::Name:
 			return 8;
 		case EPCGKernelAttributeType::Float3:
 		case EPCGKernelAttributeType::Rotator:
@@ -252,15 +252,11 @@ namespace PCGDataForGPUHelpers
 		}
 		case PCG::Private::MetadataTypes<FName>::Id:
         {
-        	// FName are defined with 2 int32, so packed them both.
+        	// FNames are currently stored in string table so use same logic as string.
         	const FPCGMetadataAttribute<FName>* Attribute = static_cast<const FPCGMetadataAttribute<FName>*>(InAttributeBase);
-        	FName Value = Attribute->GetValue(ValueKey);
-        	check(StrideBytes == 8);
-			const FNameEntryId NameEntry = Value.GetDisplayIndex();
-			const uint32 Number = Value.GetNumber();
-
-        	OutPackedDataCollection[InOutAddressUints++] = NameEntry.ToUnstableInt();
-        	OutPackedDataCollection[InOutAddressUints++] = Number;
+			const int32 Value = InStringTable.IndexOfByKey(Attribute->GetValue(ValueKey).ToString());
+			check(StrideBytes == 4);
+			OutPackedDataCollection[InOutAddressUints++] = Value;
         	break;
         }
 		default:
@@ -453,13 +449,11 @@ namespace PCGDataForGPUHelpers
 		{
 			FPCGMetadataAttribute<FName>* Attribute = static_cast<FPCGMetadataAttribute<FName>*>(AttributeBase);
 
-			const FNameEntryId NameEntryId = FNameEntryId::FromUnstableInt(DataAsInt[ElementIndex]);
-			const uint32 Number = DataAsInt[ElementIndex + 1];
-			// Names need some validation
-			const FName Value = FName::CreateFromDisplayId(NameEntryId, Number);
-			if (Value.IsValid())
+			// FNames currently stored in string table.
+			const int32 StringKey = DataAsInt[ElementIndex];
+			if (InStringTable.IsValidIndex(StringKey))
 			{
-				Attribute->SetValue(EntryKey, Value);
+				Attribute->SetValue(EntryKey, *InStringTable[StringKey]);
 			}
 			break;
 		}
@@ -650,7 +644,7 @@ void FPCGDataDesc::InitializeAttributeDescs(const UPCGData* InData, const TMap<F
 
 			TArray<int32> UniqueStringKeys;
 
-			if (AttributeType == EPCGKernelAttributeType::StringKey)
+			if (AttributeType == EPCGKernelAttributeType::StringKey || AttributeType == EPCGKernelAttributeType::Name)
 			{
 				const FPCGMetadataAttributeBase* AttributeBase = Metadata->GetConstAttribute(AttributeName);
 				check(AttributeBase);
@@ -674,9 +668,13 @@ void FPCGDataDesc::InitializeAttributeDescs(const UPCGData* InData, const TMap<F
 							StringTableIndex = InStringTable.IndexOfByKey(static_cast<const FPCGMetadataAttribute<FString>*>(AttributeBase)->GetValue(ValueKey));
 						}
 						else if (AttributeBase->GetTypeId() == PCG::Private::MetadataTypes<FSoftClassPath>::Id)
-                        {
-                        	StringTableIndex = InStringTable.IndexOfByKey(static_cast<const FPCGMetadataAttribute<FSoftClassPath>*>(AttributeBase)->GetValue(ValueKey).ToString());
-                        }
+						{
+							StringTableIndex = InStringTable.IndexOfByKey(static_cast<const FPCGMetadataAttribute<FSoftClassPath>*>(AttributeBase)->GetValue(ValueKey).ToString());
+						}
+						else if (AttributeBase->GetTypeId() == PCG::Private::MetadataTypes<FName>::Id)
+						{
+							StringTableIndex = InStringTable.IndexOfByKey(static_cast<const FPCGMetadataAttribute<FName>*>(AttributeBase)->GetValue(ValueKey).ToString());
+						}
 						else
 						{
 							// Should not get here if attribute type is string key.
@@ -706,9 +704,13 @@ void FPCGDataDesc::InitializeAttributeDescs(const UPCGData* InData, const TMap<F
 							StringTableIndex = InStringTable.IndexOfByKey(static_cast<const FPCGMetadataAttribute<FString>*>(AttributeBase)->GetValue(MetadataKey));
 						}
 						else if (AttributeBase->GetTypeId() == PCG::Private::MetadataTypes<FSoftClassPath>::Id)
-                        {
-                        	StringTableIndex = InStringTable.IndexOfByKey(static_cast<const FPCGMetadataAttribute<FSoftClassPath>*>(AttributeBase)->GetValue(MetadataKey).ToString());
-                        }
+						{
+							StringTableIndex = InStringTable.IndexOfByKey(static_cast<const FPCGMetadataAttribute<FSoftClassPath>*>(AttributeBase)->GetValue(MetadataKey).ToString());
+						}
+						else if (AttributeBase->GetTypeId() == PCG::Private::MetadataTypes<FName>::Id)
+						{
+							StringTableIndex = InStringTable.IndexOfByKey(static_cast<const FPCGMetadataAttribute<FName>*>(AttributeBase)->GetValue(MetadataKey).ToString());
+						}
 						else
 						{
 							// Should not get here if attribute type is string key.
