[PCG] Fix use seed as attribute in Subdivide Segment to work the same as Grammar and FlipAxis Attribute

#rnx

diff --git a/Engine/Plugins/PCG/Source/PCG/Private/Elements/Grammar/PCGSubdivideSegment.cpp b/Engine/Plugins/PCG/Source/PCG/Private/Elements/Grammar/PCGSubdivideSegment.cpp
--- a/Engine/Plugins/PCG/Source/PCG/Private/Elements/Grammar/PCGSubdivideSegment.cpp
+++ b/Engine/Plugins/PCG/Source/PCG/Private/Elements/Grammar/PCGSubdivideSegment.cpp
@@ -10,6 +10,7 @@
 #include "Metadata/Accessors/IPCGAttributeAccessor.h"
 #include "Metadata/Accessors/PCGAttributeAccessorHelpers.h"
 #include "Metadata/Accessors/PCGAttributeAccessorKeys.h"
+#include "Metadata/Accessors/PCGCustomAccessor.h"
 
 #define LOCTEXT_NAMESPACE "PCGSubdivideSegmentElement"
 
@@ -31,7 +32,6 @@ public:
 
 		FVector SubdivisionDirection;
 		FVector PerpendicularSubdivisionDirection;
-		int32 AdditionalSeed = 0;
 
 		const UPCGSubdivideSegmentSettings* Settings = nullptr;
 		FPCGContext* Context = nullptr;
@@ -40,7 +40,7 @@ public:
 		UPCGMetadata* OutputMetadata = nullptr;
 	};
 
-	static void Process(FParameters& InOutParameters, const FString& InGrammar, const bool bFlipAxis, int32 Index)
+	static void Process(FParameters& InOutParameters, const FString& InGrammar, const bool bFlipAxis, int32 AdditionalSeed, int32 Index)
 	{
 		if (!InOutParameters.CachedModules.Contains(InGrammar))
 		{
@@ -79,7 +79,7 @@ public:
 
 		TArray<PCGSubdivisionBase::TModuleInstance<PCGGrammar::FTokenizedModule>> ModulesInstances;
 		double RemainingSubdivide;
-		const bool bSubdivideSuccess = PCGSubdivisionBase::Subdivide(*CurrentTokenizedGrammar.ModuleGrammar, Size, ModulesInstances, RemainingSubdivide, InOutParameters.Context, InOutParameters.AdditionalSeed);
+		const bool bSubdivideSuccess = PCGSubdivisionBase::Subdivide(*CurrentTokenizedGrammar.ModuleGrammar, Size, ModulesInstances, RemainingSubdivide, InOutParameters.Context, AdditionalSeed);
 
 		if (!bSubdivideSuccess)
 		{
@@ -243,32 +243,40 @@ bool FPCGSegmentSubdivisionElement::ExecuteInternal(FPCGContext* InContext) cons
 
 		TUniquePtr<const IPCGAttributeAccessor> GrammarAccessor;
 		TUniquePtr<const IPCGAttributeAccessor> FlipAxisAccessor;
-		TUniquePtr<const IPCGAttributeAccessorKeys> Keys;
+		TUniquePtr<const IPCGAttributeAccessor> SeedAccessor;
+		const FPCGAttributeAccessorKeysPoints Keys(InputPointData->GetPoints());
 
-		if (Settings->GrammarSelection.bGrammarAsAttribute)
-		{
-			const FPCGAttributePropertyInputSelector Selector = Settings->GrammarSelection.GrammarAttribute.CopyAndFixLast(InputPointData);
-			GrammarAccessor = PCGAttributeAccessorHelpers::CreateConstAccessor(InputPointData, Selector);
-			Keys = PCGAttributeAccessorHelpers::CreateConstKeys(InputPointData, Selector);
-			if (!GrammarAccessor || !Keys)
-			{
-				PCGLog::Metadata::LogFailToCreateAccessorError(Selector, InContext);
-				continue;
-			}
-		}
+		const bool bShouldUseAccessors = Settings->GrammarSelection.bGrammarAsAttribute || Settings->bFlipAxisAsAttribute || Settings->bUseSeedAttribute;
+		constexpr int32 DefaultAdditionalSeed = 0;
 
-		if (Settings->bFlipAxisAsAttribute)
+		if (bShouldUseAccessors)
 		{
-			const FPCGAttributePropertyInputSelector Selector = Settings->FlipAxisAttribute.CopyAndFixLast(InputPointData);
-			FlipAxisAccessor = PCGAttributeAccessorHelpers::CreateConstAccessor(InputPointData, Selector);
-			if (!Keys)
+			// Create accessors for grammar, flip axis and seed as soon as either is needed. For the ones that aren't needed, we will create a constant accessor to uniformize the code below, even though there is a small performance cost.
+			auto CreateAccessor = [InputPointData, InContext]<typename T>(TUniquePtr<const IPCGAttributeAccessor>&OutAccessor, bool bValueAsAttribute, const FPCGAttributePropertyInputSelector & InputSelector, const T& ConstantValue) -> bool
 			{
-				Keys = PCGAttributeAccessorHelpers::CreateConstKeys(InputPointData, Selector);
-			}
+				if (bValueAsAttribute)
+				{
+					const FPCGAttributePropertyInputSelector Selector = InputSelector.CopyAndFixLast(InputPointData);
+					OutAccessor = PCGAttributeAccessorHelpers::CreateConstAccessor(InputPointData, Selector);
+
+					if (!OutAccessor)
+					{
+						PCGLog::Metadata::LogFailToCreateAccessorError(Selector, InContext);
+						return false;
+					}
+				}
+				else
+				{
+					OutAccessor = MakeUnique<FPCGConstantValueAccessor<T>>(ConstantValue);
+				}
+
+				return true;
+			};
 
-			if (!FlipAxisAccessor || !Keys)
+			if (!CreateAccessor(GrammarAccessor, Settings->GrammarSelection.bGrammarAsAttribute, Settings->GrammarSelection.GrammarAttribute, Settings->GrammarSelection.GrammarString)
+				|| !CreateAccessor(FlipAxisAccessor, Settings->bFlipAxisAsAttribute, Settings->FlipAxisAttribute, Settings->bShouldFlipAxis)
+				|| !CreateAccessor(SeedAccessor, Settings->bUseSeedAttribute, Settings->SeedAttribute, DefaultAdditionalSeed))
 			{
-				PCGLog::Metadata::LogFailToCreateAccessorError(Selector, InContext);
 				continue;
 			}
 		}
@@ -307,60 +315,21 @@ bool FPCGSegmentSubdivisionElement::ExecuteInternal(FPCGContext* InContext) cons
 			continue;
 		}
 
-		// Set seed if required
-		Parameters.AdditionalSeed = 0;
-		if (Settings->bUseSeedAttribute)
-		{
-			const FPCGAttributePropertyInputSelector Selector = Settings->SeedAttribute.CopyAndFixLast(InputPointData);
-			TUniquePtr<const IPCGAttributeAccessor> SeedAccessor = PCGAttributeAccessorHelpers::CreateConstAccessor(InputPointData, Selector);
-			if (!Keys)
-			{
-				Keys = PCGAttributeAccessorHelpers::CreateConstKeys(InputPointData, Selector);
-			}
-
-			if (!SeedAccessor || !Keys)
-			{
-				PCGLog::Metadata::LogFailToCreateAccessorError(Selector, InContext);
-			}
-			// Otherwise, get the value, if it fails, the attribute wasn't compatible
-			else if (!SeedAccessor->Get(Parameters.AdditionalSeed, FPCGAttributeAccessorKeysEntries(PCGInvalidEntryKey), EPCGAttributeAccessorFlags::AllowBroadcastAndConstructible))
-			{
-				PCGLog::Metadata::LogFailToGetAttributeError<int32>(Selector, SeedAccessor.Get(), InContext);
-			}
-		}
-
-		if (Settings->GrammarSelection.bGrammarAsAttribute && Settings->bFlipAxisAsAttribute)
-		{
-			auto Process = [&Parameters](const FString& InGrammar, const bool bFlipAxis, int32 Index) -> void
-			{
-				PCGSubdivideSegmentHelpers::Process(Parameters, InGrammar, bFlipAxis, Index);
-			};
-
-			PCGMetadataElementCommon::ApplyOnMultiAccessors<FString, bool>(*Keys, { GrammarAccessor.Get(), FlipAxisAccessor.Get() }, Process);
-		}
-		else if (Settings->GrammarSelection.bGrammarAsAttribute)
-		{
-			auto Process = [&Parameters, bShouldFlipAxis = Settings->bShouldFlipAxis](const FString& InGrammar, int32 Index) -> void
-			{
-				PCGSubdivideSegmentHelpers::Process(Parameters, InGrammar, bShouldFlipAxis, Index);
-			};
-
-			PCGMetadataElementCommon::ApplyOnAccessor<FString>(*Keys, *GrammarAccessor, Process);
-		}
-		else if (Settings->bFlipAxisAsAttribute)
+		if (bShouldUseAccessors)
 		{
-			auto Process = [&Parameters, Grammar = Settings->GrammarSelection.GrammarString](const bool bFlipAxis, int32 Index) -> void
+			check(GrammarAccessor && FlipAxisAccessor && SeedAccessor);
+			auto Process = [&Parameters](const FString& InGrammar, const bool bFlipAxis, int32 AdditionalSeed, int32 Index) -> void
 			{
-				PCGSubdivideSegmentHelpers::Process(Parameters, Grammar, bFlipAxis, Index);
+				PCGSubdivideSegmentHelpers::Process(Parameters, InGrammar, bFlipAxis, AdditionalSeed, Index);
 			};
 
-			PCGMetadataElementCommon::ApplyOnAccessor<bool>(*Keys, *FlipAxisAccessor, Process);
+			PCGMetadataElementCommon::ApplyOnMultiAccessors<FString, bool, int32>(Keys, { GrammarAccessor.Get(), FlipAxisAccessor.Get(), SeedAccessor.Get()}, Process);
 		}
 		else
 		{
 			for (int32 SegmentIndex = 0; SegmentIndex < Parameters.InPoints->Num(); ++SegmentIndex)
 			{
-				PCGSubdivideSegmentHelpers::Process(Parameters, Settings->GrammarSelection.GrammarString, Settings->bShouldFlipAxis, SegmentIndex);
+				PCGSubdivideSegmentHelpers::Process(Parameters, Settings->GrammarSelection.GrammarString, Settings->bShouldFlipAxis, DefaultAdditionalSeed, SegmentIndex);
 			}
 		}
 
