[PCG] Add context for Instance property bag changed to only refresh if it is the currently displayed instance

#rnx

diff --git a/Engine/Plugins/PCG/Source/PCGEditor/Private/Details/PCGInstancedPropertyBagOverrideDetails.h b/Engine/Plugins/PCG/Source/PCGEditor/Private/Details/PCGInstancedPropertyBagOverrideDetails.h
--- a/Engine/Plugins/PCG/Source/PCGEditor/Private/Details/PCGInstancedPropertyBagOverrideDetails.h
+++ b/Engine/Plugins/PCG/Source/PCGEditor/Private/Details/PCGInstancedPropertyBagOverrideDetails.h
@@ -2,6 +2,8 @@
 
 #pragma once
 
+#include "PCGGraph.h"
+
 #include "IDetailCustomization.h"
 #include "IPropertyTypeCustomization.h"
 #include "InstancedStructDetails.h"
@@ -10,7 +12,6 @@
 #include "Templates/SharedPointer.h"
 #include "UObject/WeakObjectPtrTemplates.h"
 
-class UPCGGraphInstance;
 class UPCGSettings;
 struct FInstancedPropertyBag;
 
@@ -48,5 +49,13 @@ public:
 	~FPCGOverrideInstancedPropertyBagDataDetails();
 
 	virtual void GenerateHeaderRowContent(FDetailWidgetRow& NodeRow) override;
-	void OnStructChange(const FInstancedPropertyBag&) { OnStructLayoutChanges(); }
+	void OnStructChange(const FInstancedPropertyBag& Instance)
+	{
+		// Only react to changes that happened on this property bag.
+		const UPCGGraphInstance* OwnerPtr = Owner.Get();
+		if (IsValid(OwnerPtr) && OwnerPtr->GetUserParametersStruct() == &Instance)
+		{
+			OnStructLayoutChanges();
+		}
+	}
 };
\ No newline at end of file
