[PCG] Fix attribute function captures breaking for nested functions
* Also demoted overeager ensure in GPU graph compiler that is hit when adding pins with unsupported types

#rnx

diff --git a/Engine/Plugins/PCG/Source/PCG/Private/Compute/Elements/PCGCustomHLSL.cpp b/Engine/Plugins/PCG/Source/PCG/Private/Compute/Elements/PCGCustomHLSL.cpp
--- a/Engine/Plugins/PCG/Source/PCG/Private/Compute/Elements/PCGCustomHLSL.cpp
+++ b/Engine/Plugins/PCG/Source/PCG/Private/Compute/Elements/PCGCustomHLSL.cpp
@@ -43,7 +43,7 @@ namespace PCGHLSLElement
 	constexpr int AttributeNameCaptureGroup = 4;
 
 	/** Regex pattern used to detect and parse attribute function usage in kernels. */
-	constexpr TCHAR AttributeFunctionPattern[] = { TEXT("([a-zA-Z0-9]+)_(Get|Set)(.*)\\(.*'([a-zA-Z0-9 -_\\/]+)'.*") };
+	constexpr TCHAR AttributeFunctionPattern[] = { TEXT("([a-zA-Z0-9]+)_(Get|Set)(.*?)[\\s]*?\\([^,]*?,[^,]*?,[^,]*?'([a-zA-Z0-9 -_\\/]*?)'") };
 
 	constexpr TCHAR AttributeFunctionGet[] = { TEXT("Get") };
 	constexpr TCHAR AttributeFunctionSet[] = { TEXT("Set") };
diff --git a/Engine/Plugins/PCG/Source/PCG/Private/Graph/PCGGraphCompilerGPU.cpp b/Engine/Plugins/PCG/Source/PCG/Private/Graph/PCGGraphCompilerGPU.cpp
--- a/Engine/Plugins/PCG/Source/PCG/Private/Graph/PCGGraphCompilerGPU.cpp
+++ b/Engine/Plugins/PCG/Source/PCG/Private/Graph/PCGGraphCompilerGPU.cpp
@@ -873,7 +873,7 @@ UPCGComputeGraph* FPCGGraphCompilerGPU::CompileComputeGraph(
 
 			if (!UpstreamDI)
 			{
-				UE_LOG(LogPCG, Warning, TEXT("Encountered missing data interface, skipping."))
+				UE_LOG(LogPCG, Warning, TEXT("Failed to find data interface associated with pin '%s'."), *Input.UpstreamPin->Label.ToString());
 				continue;
 			}
 
@@ -900,8 +900,9 @@ UPCGComputeGraph* FPCGGraphCompilerGPU::CompileComputeGraph(
 			}
 
 			UPCGComputeDataInterface** FoundDI = OutputPinDataInterfaces.Find({ TaskId, OutputPinProperties.Label });
-			if (!ensure(FoundDI) || !ensure(*FoundDI))
+			if (!FoundDI || !*FoundDI)
 			{
+				UE_LOG(LogPCG, Warning, TEXT("Failed to find data interface associated with pin '%s'."), *OutputPinProperties.Label.ToString());
 				continue;
 			}
 
@@ -936,8 +937,9 @@ UPCGComputeGraph* FPCGGraphCompilerGPU::CompileComputeGraph(
 
 					// DIs for output pins should have all been created.
 					UPCGComputeDataInterface** FoundOutputDI = OutputPinDataInterfaces.Find({ TaskId, OutputPinLabel });
-					if (!ensure(FoundOutputDI) || !ensure(*FoundOutputDI))
+					if (!FoundOutputDI || !*FoundOutputDI)
 					{
+						UE_LOG(LogPCG, Warning, TEXT("Failed to find data interface associated with pin '%s'."), *OutputPinLabel.ToString());
 						continue;
 					}
 
