[PCG] GPU compiler refactor and tidy pass

#rnx

diff --git a/Engine/Plugins/PCG/Source/PCG/Private/Graph/PCGGraphCompilerGPU.cpp b/Engine/Plugins/PCG/Source/PCG/Private/Graph/PCGGraphCompilerGPU.cpp
--- a/Engine/Plugins/PCG/Source/PCG/Private/Graph/PCGGraphCompilerGPU.cpp
+++ b/Engine/Plugins/PCG/Source/PCG/Private/Graph/PCGGraphCompilerGPU.cpp
@@ -456,10 +456,12 @@ void FPCGGraphCompilerGPU::WireGPUGraphNode(
 			// TODO is pinless fine with skipping?
 			if (AddedInput.DownstreamPin.IsSet())
 			{
-				const FName* FoundVirtualLabel = InOriginalToVirtualPin.Find({ GPUTaskId, AddedInput.DownstreamPin->Label, /*bIsInputPin=*/true });
-				if (ensure(FoundVirtualLabel))
+				// Get the compute graph virtual pin corresponding to this pin, if one was created.
+				const FName* VirtualPinLabel = InOriginalToVirtualPin.Find({ GPUTaskId, AddedInput.DownstreamPin->Label, /*bIsInputPin=*/true });
+
+				if (ensure(VirtualPinLabel))
 				{
-					AddedInput.DownstreamPin->Label = *FoundVirtualLabel;
+					AddedInput.DownstreamPin->Label = *VirtualPinLabel;
 				}
 			}
 		}
@@ -499,12 +501,13 @@ void FPCGGraphCompilerGPU::WireGPUGraphNode(
 
 				if (DownstreamCPUNode.Inputs[SuccessorInputIndex].UpstreamPin.IsSet())
 				{
-					const FNodePin PinKey = { GPUTaskId, InputCopy.UpstreamPin->Label, /*Pin is input*/false };
-					const FName* FoundVirtualPinLabel = InOriginalToVirtualPin.Find(PinKey);
-					if (ensure(FoundVirtualPinLabel))
+					// Get the compute graph virtual pin corresponding to this pin, if one was created.
+					const FName* VirtualPinLabel = InOriginalToVirtualPin.Find({ GPUTaskId, InputCopy.UpstreamPin->Label, /*Pin is input*/false });
+
+					if (ensure(VirtualPinLabel))
 					{
 						// Wire to the existing virtual output pin.
-						InputCopy.UpstreamPin->Label = *FoundVirtualPinLabel;
+						InputCopy.UpstreamPin->Label = *VirtualPinLabel;
 					}
 				}
 
@@ -674,12 +677,47 @@ UPCGComputeGraph* FPCGGraphCompilerGPU::CompileComputeGraph(
 
 	TMap<TPair</* Node task*/FPCGTaskId, /*Node output pin*/FName>, UPCGComputeDataInterface*> OutputPinDataInterfaces;
 
-	// Create all the output data interfaces.
+	// Create DIs for output pins (either output pins within this compute graph or outside).
 	for (FPCGTaskId TaskId : InCollapsedTasks)
 	{
-		// Create DIs for all output pins, because the kernels currently need their outputs to be bound to valid resources.
+		// Create DIs for all output pins regardless of outbound connections, because the kernels currently need their outputs to be bound to valid resources.
 		if (const UPCGSettings* Settings = InOutCompiledTasks[TaskId].Node ? InOutCompiledTasks[TaskId].Node->GetSettings() : nullptr)
 		{
+			// Record attributes.
+			TArray<FPCGKernelAttributeKey> AttributeKeys;
+			Settings->GetKernelAttributeKeys(AttributeKeys);
+
+			for (const FPCGKernelAttributeKey& Key : AttributeKeys)
+			{
+				FPCGKernelAttributeIDAndType* FoundEntry = ComputeGraph->GlobalAttributeLookupTable.Find(Key.Name);
+				if (FoundEntry)
+				{
+					// If we already registered this attribute, populate the type if not populated already, or otherwise ensure type matches.
+					if (Key.Type != EPCGKernelAttributeType::None)
+					{
+						if (FoundEntry->Type == EPCGKernelAttributeType::None)
+						{
+							FoundEntry->Type = Key.Type;
+						}
+						else if (Key.Type != FoundEntry->Type)
+						{
+							// TODO: This is more a limitation of the current system rather than an issue with the user's data. Log for bug reports etc but remove this
+							// later when we properly flow attribute table and string table down graph branches.
+							UE_LOG(LogPCG, Warning, TEXT("Attribute '%s' multiply defined with different types (existing: %d, new: %d)."), *Key.Name.ToString(), FoundEntry->Type, Key.Type);
+						}
+					}
+				}
+				else
+				{
+					const FPCGKernelAttributeIDAndType IDAndType(ComputeGraph->GlobalAttributeLookupTable.Num() + PCGComputeConstants::NUM_RESERVED_ATTRS, Key.Type);
+					ComputeGraph->GlobalAttributeLookupTable.Add(Key.Name, IDAndType);
+				}
+			}
+
+			// Record strings.
+			Settings->AddStaticCreatedStrings(ComputeGraph->StringTable);
+
+			// Create all the output data interfaces.
 			for (const FPCGPinProperties& OutputPinProperties : Settings->AllOutputPinProperties())
 			{
 				if (OutputPinProperties.Label == NAME_None)
@@ -693,6 +731,7 @@ UPCGComputeGraph* FPCGGraphCompilerGPU::CompileComputeGraph(
 					continue;
 				}
 
+				// Request readback if any downstream task is outside of this compute graph.
 				bool bRequiresReadback = false;
 				if (const TArray<FPCGTaskId>* Successors = InTaskSuccessors.Find(TaskId))
 				{
@@ -700,7 +739,7 @@ UPCGComputeGraph* FPCGGraphCompilerGPU::CompileComputeGraph(
 					{
 						for (const FPCGGraphTaskInput& Input : InOutCompiledTasks[Successor].Inputs)
 						{
-							if (Input.UpstreamPin.IsSet() && (*Input.UpstreamPin == OutputPinProperties) && !InCollapsedTasks.Contains(Successor))
+							if (Input.TaskId == TaskId && Input.UpstreamPin && !InCollapsedTasks.Contains(Successor))
 							{
 								bRequiresReadback = true;
 								break;
@@ -711,42 +750,30 @@ UPCGComputeGraph* FPCGGraphCompilerGPU::CompileComputeGraph(
 
 				if (UPCGComputeDataInterface* OutputDI = CreateDataInterface(TaskId, bRequiresReadback, OutputPinProperties))
 				{
-					OutputDI->SetOutputPin(OutputPinProperties.Label);
+					// Get the compute graph virtual pin corresponding to this pin, if one was created.
+					const FName* VirtualPinLabel = InOriginalToVirtualPin.Find({ TaskId, OutputPinProperties.Label, /*bIsInputPin*/false });
+					OutputDI->SetOutputPin(OutputPinProperties.Label, VirtualPinLabel);
+
 					ComputeGraph->DataInterfaces.Add(OutputDI);
-					OutputPinDataInterfaces.Add({ TaskId, OutputPinProperties.Label }, OutputDI);
 
-					// TODO: This seems like a once-per-settings process rather than repeating for each output pin?
-					TArray<FPCGKernelAttributeKey> AttributeKeys;
-					Settings->GetKernelAttributeKeys(AttributeKeys);
+					OutputPinDataInterfaces.Add({ TaskId, OutputPinProperties.Label }, OutputDI);
 
-					for (const FPCGKernelAttributeKey& Key : AttributeKeys)
+					// Iterate over downstream connections and register each downstream pin.
+					if (const TArray<FPCGTaskId>* Successors = InTaskSuccessors.Find(TaskId))
 					{
-						FPCGKernelAttributeIDAndType* FoundEntry = ComputeGraph->GlobalAttributeLookupTable.Find(Key.Name);
-						if (FoundEntry)
+						for (FPCGTaskId Successor : *Successors)
 						{
-							// If we already registered this attribute, populate the type if not populated already, or otherwise ensure type matches.
-							if (Key.Type != EPCGKernelAttributeType::None)
+							for (const FPCGGraphTaskInput& Input : InOutCompiledTasks[Successor].Inputs)
 							{
-								if (FoundEntry->Type == EPCGKernelAttributeType::None)
+								if (Input.TaskId == TaskId && Input.UpstreamPin && Input.DownstreamPin)
 								{
-									FoundEntry->Type = Key.Type;
-								}
-								else if (Key.Type != FoundEntry->Type)
-								{
-									// TODO: This is more a limitation of the current system rather than an issue with the user's data. Log for bug reports etc but remove this
-									// later when we properly flow attribute table and string table down graph branches.
-									UE_LOG(LogPCG, Warning, TEXT("Attribute '%s' multiply defined with different types (existing: %d, new: %d)."), *Key.Name.ToString(), FoundEntry->Type, Key.Type);
+									// Get the compute graph virtual pin corresponding to this pin, if one was created.
+									const FName* VirtualInputPinLabel = InOriginalToVirtualPin.Find({ TaskId, Input.DownstreamPin->Label, /*bIsInputPin=*/true });
+									OutputDI->AddDownstreamInputPin(Input.DownstreamPin->Label, VirtualInputPinLabel);
 								}
 							}
 						}
-						else
-						{
-							const FPCGKernelAttributeIDAndType IDAndType(ComputeGraph->GlobalAttributeLookupTable.Num() + PCGComputeConstants::NUM_RESERVED_ATTRS, Key.Type);
-							ComputeGraph->GlobalAttributeLookupTable.Add(Key.Name, IDAndType);
-						}
 					}
-
-					Settings->AddStaticCreatedStrings(ComputeGraph->StringTable);
 				}
 			}
 		}
@@ -778,55 +805,93 @@ UPCGComputeGraph* FPCGGraphCompilerGPU::CompileComputeGraph(
 				continue;
 			}
 
-			if (OutputPinDataInterfaces.Contains({ Input.TaskId, Input.UpstreamPin->Label }))
+			UPCGComputeDataInterface* OutputDI = nullptr;
+
+			if (UPCGComputeDataInterface** FoundOutputDI = OutputPinDataInterfaces.Find({ Input.TaskId, Input.UpstreamPin->Label }))
 			{
-				// Skip if already created.
-				continue;
+				OutputDI = *FoundOutputDI;
 			}
-
-			if (UPCGComputeDataInterface* OutputDI = CreateDataInterface(Input.TaskId, /*bRequiresReadback=*/false, *Input.UpstreamPin))
+			else if (UPCGComputeDataInterface* NewOutputDI = CreateDataInterface(Input.TaskId, /*bRequiresReadback=*/false, *Input.UpstreamPin))
 			{
-				OutputDI->SetOutputPin(Input.UpstreamPin->Label);
+				OutputDI = NewOutputDI;
+
+				// Get the compute graph virtual pin corresponding to this pin, if one was created.
+				const FName* VirtualPinLabel = InOriginalToVirtualPin.Find({ Input.TaskId, Input.UpstreamPin->Label, /*bIsInputPin*/false });
+				OutputDI->SetOutputPin(Input.UpstreamPin->Label, VirtualPinLabel);
+
 				ComputeGraph->DataInterfaces.Add(OutputDI);
+
 				OutputPinDataInterfaces.Add({ Input.TaskId, Input.UpstreamPin->Label }, OutputDI);
 			}
+
+			if (OutputDI)
+			{
+				// Get the compute graph virtual pin corresponding to this pin, if one was created.
+				const FName* VirtualPinLabel = InOriginalToVirtualPin.Find({ TaskId, Input.DownstreamPin->Label, /*bIsInputPin=*/true });
+				OutputDI->AddDownstreamInputPin(Input.DownstreamPin->Label, VirtualPinLabel);
+			}
 		}
 	}
 
-	TSet<FPCGTaskId> RemainingTasks = InCollapsedTasks;
+	// Build compute graph by traversing the set of collapsed tasks in execution order (order of queued kernel invocations matters).
+	TSet<FPCGTaskId> RemainingTasks;
+	TArray<FPCGTaskId> ReadyTasks;
+	RemainingTasks.Reserve(InCollapsedTasks.Num());
+	ReadyTasks.Reserve(InCollapsedTasks.Num());
 
-	while (!RemainingTasks.IsEmpty())
+	for (FPCGTaskId TaskId : InCollapsedTasks)
 	{
-		// Find a ready task
-		FPCGTaskId TaskId = InvalidPCGTaskId;
-		for (FPCGTaskId RemainingTask : RemainingTasks)
+		bool bTaskReady = true;
+		for (const FPCGGraphTaskInput& Input : InOutCompiledTasks[TaskId].Inputs)
 		{
-			// TODO: use 'QueueSuccessors' pattern rather than brute force searching for ready tasks
-			bool bReady = true;
-			for (const FPCGGraphTaskInput& Input : InOutCompiledTasks[RemainingTask].Inputs)
-			{
-				if (RemainingTasks.Contains(Input.TaskId))
-				{
-					bReady = false;
-					break;
-				}
-			}
-
-			if (bReady)
+			if (InCollapsedTasks.Contains(Input.TaskId))
 			{
-				TaskId = RemainingTask;
+				bTaskReady = false;
 				break;
 			}
 		}
 
-		if (TaskId == InvalidPCGTaskId)
+		if (bTaskReady)
 		{
-			// Tasks in RemainingTasks are unreachable?
-			ensure(false);
-			break;
+			ReadyTasks.Add(TaskId);
 		}
+		else
+		{
+			RemainingTasks.Add(TaskId);
+		}
+	}
+
+	while (!ReadyTasks.IsEmpty())
+	{
+		FPCGTaskId TaskId = ReadyTasks.Pop();
+
+		// Queue up any successors that are ready to go (do not depend on any remaining kernel).
+		if (const TArray<FPCGTaskId>* Successors = InTaskSuccessors.Find(TaskId))
+		{
+			for (const FPCGTaskId Successor : *Successors)
+			{
+				if (!RemainingTasks.Contains(Successor))
+				{
+					continue;
+				}
+
+				bool bTaskReady = true;
+				for (const FPCGGraphTaskInput& Input : InOutCompiledTasks[Successor].Inputs)
+				{
+					if (RemainingTasks.Contains(Input.TaskId))
+					{
+						bTaskReady = false;
+						break;
+					}
+				}
 
-		RemainingTasks.Remove(TaskId);
+				if (bTaskReady)
+				{
+					ReadyTasks.Add(Successor);
+					RemainingTasks.Remove(Successor);
+				}
+			}
+		}
 
 		const UPCGNode* Node = InOutCompiledTasks[TaskId].Node;
 
@@ -845,8 +910,7 @@ UPCGComputeGraph* FPCGGraphCompilerGPU::CompileComputeGraph(
 		InputDataInterfaceIndexAndPin.Reserve(Settings->InputPinProperties().Num());
 		OutputDataInterfaceIndexAndPin.Reserve(Settings->OutputPinProperties().Num());
 
-		// Add DIs (PCG -> CF transcoding).
-
+		// Setup input pin DIs.
 		for (const FPCGGraphTaskInput& Input : InOutCompiledTasks[TaskId].Inputs)
 		{
 			// Currently the new compute graph tasks are wired into the graph in parallel to each GPU node task. The GPU node tasks
@@ -858,40 +922,32 @@ UPCGComputeGraph* FPCGGraphCompilerGPU::CompileComputeGraph(
 				continue;
 			}
 
-			if (!Input.UpstreamPin.IsSet())
+			if (!Input.UpstreamPin || !Input.DownstreamPin)
 			{
 				// Execution-only dependencies not supported currently. Unclear if this should ever be supported for GPU graphs.
 				// Writes followed by reads will be protected via barriers added by RDG.
 				continue;
 			}
 
-			UPCGComputeDataInterface* UpstreamDI = nullptr;
-			if (UPCGComputeDataInterface** FoundUpstreamDI = OutputPinDataInterfaces.Find({ Input.TaskId, Input.UpstreamPin->Label }))
-			{
-				UpstreamDI = *FoundUpstreamDI;
-			}
 
-			if (!UpstreamDI)
+			UPCGComputeDataInterface** FoundDI = OutputPinDataInterfaces.Find({ Input.TaskId, Input.UpstreamPin->Label });
+			if (!FoundDI || !*FoundDI)
 			{
 				UE_LOG(LogPCG, Warning, TEXT("Failed to find data interface associated with pin '%s'."), *Input.UpstreamPin->Label.ToString());
 				continue;
 			}
 
-			const int Index = ComputeGraph->DataInterfaces.Find(UpstreamDI);
+			const int Index = ComputeGraph->DataInterfaces.Find(*FoundDI);
 			if (Index == INDEX_NONE)
 			{
 				ensure(false);
 				continue;
 			}
 
-			const FName DownstreamInputPinLabel = Input.DownstreamPin->Label;
-			InputDataInterfaceIndexAndPin.Emplace(Index, DownstreamInputPinLabel);
-
-			const bool bIsInputPin = true;
-			UpstreamDI->AddDownstreamInputPin(DownstreamInputPinLabel, InOriginalToVirtualPin.Find({ TaskId, DownstreamInputPinLabel, bIsInputPin }));
+			InputDataInterfaceIndexAndPin.Emplace(Index, Input.DownstreamPin->Label);
 		}
 
-		// Always create a DI for every output pin, so kernel always has something to write to.
+		// Setup output pin DIs. Always bind a DI to every output pin, so kernel always has something to write to.
 		for (const FPCGPinProperties& OutputPinProperties : Settings->AllOutputPinProperties())
 		{
 			if (OutputPinProperties.Label == NAME_None)
@@ -916,85 +972,58 @@ UPCGComputeGraph* FPCGGraphCompilerGPU::CompileComputeGraph(
 			OutputDataInterfaceIndexAndPin.Emplace(Index, OutputPinProperties.Label);
 		}
 
-		// Make sure every downstream input pin is registered with the upstream DI.
-		if (const TArray<FPCGTaskId>* Successors = InTaskSuccessors.Find(TaskId))
+		// Kernel data interface.
 		{
-			for (FPCGTaskId Successor : *Successors)
-			{
-				for (const FPCGGraphTaskInput& SuccessorInput : InOutCompiledTasks[Successor].Inputs)
-				{
-					if (SuccessorInput.TaskId != TaskId)
-					{
-						continue;
-					}
+			UPCGCustomKernelDataInterface* KernelDI = InOutContext.NewObject_AnyThread<UPCGCustomKernelDataInterface>(ComputeGraph);
+			KernelDI->Settings = Settings;
 
-					if (!SuccessorInput.UpstreamPin.IsSet() || !SuccessorInput.DownstreamPin.IsSet())
-					{
-						continue;
-					}
+			const int32 KernelDIIndex = ComputeGraph->DataInterfaces.Num();
+			ComputeGraph->DataInterfaces.Add(KernelDI);
+			InputDataInterfaceIndexAndPin.Emplace(KernelDIIndex, NAME_None);
+		}
 
-					FName OutputPinLabel = SuccessorInput.UpstreamPin->Label;
+		// Additional DIs created by settings.
+		{
+			TArray<TObjectPtr<UComputeDataInterface>> AdditionalInputDIs, AdditionalOutputDIs;
+			Settings->CreateAdditionalInputDataInterfaces(InOutContext, ComputeGraph, AdditionalInputDIs);
+			Settings->CreateAdditionalOutputDataInterfaces(InOutContext, ComputeGraph, AdditionalOutputDIs);
 
-					// DIs for output pins should have all been created.
-					UPCGComputeDataInterface** FoundOutputDI = OutputPinDataInterfaces.Find({ TaskId, OutputPinLabel });
-					if (!FoundOutputDI || !*FoundOutputDI)
-					{
-						UE_LOG(LogPCG, Warning, TEXT("Failed to find data interface associated with pin '%s'."), *OutputPinLabel.ToString());
-						continue;
-					}
+			auto AddAdditionalDataInterface = [ComputeGraph](TObjectPtr<UComputeDataInterface> DataInterface) -> int32
+			{
+				if (ensure(DataInterface))
+				{
+					const int32 DataInterfaceIndex = ComputeGraph->DataInterfaces.Num();
+					ComputeGraph->DataInterfaces.Add(DataInterface);
 
-					// If the map has an entry for this task and output pin label, then its a GPU->CPU readback.
-					const bool bIsInputPin = false;
-					(*FoundOutputDI)->SetOutputPin(OutputPinLabel, InOriginalToVirtualPin.Find({ TaskId, OutputPinLabel, bIsInputPin }));
+					return DataInterfaceIndex;
 				}
-			}
-		}
-
-		UPCGCustomKernelDataInterface* KernelDI = InOutContext.NewObject_AnyThread<UPCGCustomKernelDataInterface>(ComputeGraph);
-		KernelDI->Settings = Settings;
-
-		const int32 KernelDIIndex = ComputeGraph->DataInterfaces.Num();
-		ComputeGraph->DataInterfaces.Add(KernelDI);
-		InputDataInterfaceIndexAndPin.Emplace(KernelDIIndex, NAME_None);
 
-		TArray<TObjectPtr<UComputeDataInterface>> AdditionalInputDIs, AdditionalOutputDIs;
-		Settings->CreateAdditionalInputDataInterfaces(InOutContext, ComputeGraph, AdditionalInputDIs);
-		Settings->CreateAdditionalOutputDataInterfaces(InOutContext, ComputeGraph, AdditionalOutputDIs);
+				return static_cast<int32>(INDEX_NONE);
+			};
 
-		auto AddAdditionalDataInterface = [ComputeGraph](TObjectPtr<UComputeDataInterface> DataInterface) -> int32
-		{
-			if (ensure(DataInterface))
+			for (TObjectPtr<UComputeDataInterface> DataInterface : AdditionalInputDIs)
 			{
-				const int32 DataInterfaceIndex = ComputeGraph->DataInterfaces.Num();
-				ComputeGraph->DataInterfaces.Add(DataInterface);
+				const int32 DataInterfaceIndex = AddAdditionalDataInterface(DataInterface);
 
-				return DataInterfaceIndex;
+				if (DataInterfaceIndex != INDEX_NONE)
+				{
+					InputDataInterfaceIndexAndPin.Emplace(DataInterfaceIndex, NAME_None);
+				}
 			}
 
-			return static_cast<int32>(INDEX_NONE);
-		};
-
-		for (TObjectPtr<UComputeDataInterface> DataInterface : AdditionalInputDIs)
-		{
-			const int32 DataInterfaceIndex = AddAdditionalDataInterface(DataInterface);
-
-			if (DataInterfaceIndex != INDEX_NONE)
+			for (TObjectPtr<UComputeDataInterface> DataInterface : AdditionalOutputDIs)
 			{
-				InputDataInterfaceIndexAndPin.Emplace(DataInterfaceIndex, NAME_None);
-			}
-		}
+				const int32 DataInterfaceIndex = AddAdditionalDataInterface(DataInterface);
 
-		for (TObjectPtr<UComputeDataInterface> DataInterface : AdditionalOutputDIs)
-		{
-			const int32 DataInterfaceIndex = AddAdditionalDataInterface(DataInterface);
-
-			if (DataInterfaceIndex != INDEX_NONE)
-			{
-				OutputDataInterfaceIndexAndPin.Emplace(DataInterfaceIndex, NAME_None);
+				if (DataInterfaceIndex != INDEX_NONE)
+				{
+					OutputDataInterfaceIndexAndPin.Emplace(DataInterfaceIndex, NAME_None);
+				}
 			}
 		}
 
-		// TODO once we support cooking for different platforms/configs, don't create the interface if logging is not present.
+		// Debug data interface.
+		// TODO skip depending on target platform.
 		if (Settings->bPrintShaderDebugValues)
 		{
 			UPCGDebugDataInterface* DebugDI = InOutContext.NewObject_AnyThread<UPCGDebugDataInterface>(ComputeGraph);
@@ -1076,7 +1105,7 @@ UPCGComputeGraph* FPCGGraphCompilerGPU::CompileComputeGraph(
 		}
 
 		{
-			UPCGComputeKernelSource* KernelSource = InOutContext.NewObject_AnyThread<UPCGComputeKernelSource>(KernelWithBindings.Kernel); // is outer to kernel fine?
+			UPCGComputeKernelSource* KernelSource = InOutContext.NewObject_AnyThread<UPCGComputeKernelSource>(KernelWithBindings.Kernel);
 			KernelWithBindings.Kernel->KernelSource = KernelSource;
 
 			// These could be exposed through PCGSettings API later when the need arises (and/or when GPU feature matures).
@@ -1182,6 +1211,8 @@ UPCGComputeGraph* FPCGGraphCompilerGPU::CompileComputeGraph(
 		AddAllEdgesForKernel(KernelIndex, /*bInEdgesAreInputs=*/false);
 	}
 
+	ensureMsgf(RemainingTasks.IsEmpty(), TEXT("PCG GPU graph compiler did not consume all tasks, %d/%d remaining."), RemainingTasks.Num(), InCollapsedTasks.Num());
+
 	// Register all virtual pin aliases with the corresponding pins for downstream usage.
 	for (const TPair<TTuple<FPCGTaskId, FName, bool>, FName>& Mapping : InOriginalToVirtualPin)
 	{
