[PCG - Slate] UX improvements to the Attribute List View

* Add the option to set a double click callback on resize border in SHeaderRow
* Resize the column to match content width on double click on the resize border of the column (like in Excel)
* Right click option on the column to do the same resize
* Add caching on column width and visibility so we can keep the same column width and visibility when inspecting a node and coming back to it.
* Add option to reset the column width, and expend all columns to maximal width
* Resize Search bar to not take too much space.

diff --git a/Engine/Plugins/PCG/Source/PCGEditor/Private/SPCGEditorGraphAttributeListView.cpp b/Engine/Plugins/PCG/Source/PCGEditor/Private/SPCGEditorGraphAttributeListView.cpp
--- a/Engine/Plugins/PCG/Source/PCGEditor/Private/SPCGEditorGraphAttributeListView.cpp
+++ b/Engine/Plugins/PCG/Source/PCGEditor/Private/SPCGEditorGraphAttributeListView.cpp
@@ -55,7 +55,7 @@ namespace PCGEditorGraphAttributeListView
 		return Subsystem && Subsystem->IsGraphCacheDebuggingEnabled();
 	}
 
-	float CalculateColumnWidth(const FText& InText)
+	float CalculateColumnWidth(const FText& InText, bool bClampToMaxColumnWidth = true)
 	{
 		check(FSlateApplication::Get().GetRenderer());
 		const TSharedRef<FSlateFontMeasure> FontMeasure = FSlateApplication::Get().GetRenderer()->GetFontMeasureService();
@@ -64,7 +64,7 @@ namespace PCGEditorGraphAttributeListView
 		const float TextWidth = FontMeasure->Measure(InText, FontInfo).X;
 		constexpr float ColumnPadding = 22.0f; // TODO: Grab padding from header style
 		const float ColumnWidth = TextWidth + ColumnPadding;
-		return FMath::Min(ColumnWidth, MaxColumnWidth);
+		return bClampToMaxColumnWidth ? FMath::Min(ColumnWidth, MaxColumnWidth) : ColumnWidth;
 	}
 }
 
@@ -131,12 +131,15 @@ TSharedRef<SWidget> SPCGListViewItemRow::GenerateWidgetForColumn(const FName& Co
 	const TSharedPtr<SPCGEditorGraphAttributeListView> SharedAttributeListView = AttributeListView.Pin();
 	check(SharedAttributeListView.IsValid());
 
+	// We will only compute cell width for types are that not numbers.
+	bool bShouldComputeCellWidth = false;
+
 	if (FPCGColumnData* PCGColumnData = SharedAttributeListView->PCGColumnData.Find(ColumnId))
 	{
 		if (PCGColumnData->DataAccessor.IsValid() && PCGColumnData->DataKeys.IsValid())
 		{
 			int32 Index = InternalItem->Index;
-			auto Callback = [&PCGColumnData, Index, &RowText] (auto Dummy)
+			auto Callback = [&PCGColumnData, Index, &RowText, &bShouldComputeCellWidth] (auto Dummy)
 			{
 				using ValueType = decltype(Dummy);
 				ValueType Value{};
@@ -149,10 +152,12 @@ TSharedRef<SWidget> SPCGListViewItemRow::GenerateWidgetForColumn(const FName& Co
 					else if constexpr (PCG::Private::IsOfTypes<ValueType, FString>())
 					{
 						RowText = FText::FromString(Value);
+						bShouldComputeCellWidth = true;
 					}
 					else if constexpr (PCG::Private::IsOfTypes<ValueType, FName>())
 					{
 						RowText = FText::FromName(Value);
+						bShouldComputeCellWidth = true;
 					}
 					else if constexpr (FTextAsNumberIsValid<ValueType>::value)
 					{
@@ -161,10 +166,12 @@ TSharedRef<SWidget> SPCGListViewItemRow::GenerateWidgetForColumn(const FName& Co
 					else if constexpr (PCG::Private::IsOfTypes<ValueType, FSoftObjectPath>())
 					{
 						RowText = FText::FromString(Value.ToString());
+						bShouldComputeCellWidth = true;
 					}
 					else if constexpr (PCG::Private::IsOfTypes<ValueType, FSoftClassPath>())
 					{
 						RowText = FText::FromString(Value.ToString());
+						bShouldComputeCellWidth = true;
 					}
 					else
 					{
@@ -178,6 +185,16 @@ TSharedRef<SWidget> SPCGListViewItemRow::GenerateWidgetForColumn(const FName& Co
 		}
 	}
 
+	if (bShouldComputeCellWidth)
+	{
+		float& CurrentColumnWidth = SharedAttributeListView->ColumnsMaxWidthMapping[ColumnId];
+		const float TextWidth = PCGEditorGraphAttributeListView::CalculateColumnWidth(RowText, /*bClampToMaxColumnWidth=*/false);
+		if (TextWidth > CurrentColumnWidth)
+		{
+			CurrentColumnWidth = TextWidth;
+		}
+	}
+
 	return SNew(STextBlock)
 		.Text(RowText)
 		.OverflowPolicy(ETextOverflowPolicy::Ellipsis)
@@ -274,6 +291,8 @@ void SPCGEditorGraphAttributeListView::Construct(const FArguments& InArgs, TShar
 
 	TextFilter = MakeShareable(new FTextFilterExpressionEvaluator(ETextFilterExpressionEvaluatorMode::Complex));
 
+	ColumnWidthVisibilityCache.Empty(PCGEditorGraphAttributeListView::MaxNodeColumnWidthCachedItems);
+
 	ListViewHeader = CreateHeaderRowWidget();
 
 	ListViewCommands = MakeShareable(new FUICommandList);
@@ -380,7 +399,7 @@ void SPCGEditorGraphAttributeListView::Construct(const FArguments& InArgs, TShar
 		];
 
 	SAssignNew(SearchBoxWidget, SSearchBox)
-		.MinDesiredWidth(300.0f)
+		.MinDesiredWidth(100.0f)
 		.InitialText(ActiveFilterText)
 		.OnTextChanged(this, &SPCGEditorGraphAttributeListView::OnFilterTextChanged)
 		.OnTextCommitted(this, &SPCGEditorGraphAttributeListView::OnFilterTextCommitted)
@@ -426,7 +445,8 @@ void SPCGEditorGraphAttributeListView::Construct(const FArguments& InArgs, TShar
 				DataComboBox->AsShared()
 			]
 			+SHorizontalBox::Slot()
-			.AutoWidth()
+			.MinWidth(100.0)
+			.MaxWidth(300.0)
 			.Padding(1.0f, 0.0f)
 			[
 				SearchBoxWidget->AsShared()
@@ -453,10 +473,8 @@ void SPCGEditorGraphAttributeListView::Construct(const FArguments& InArgs, TShar
 				AdditionalOperationsButton->AsShared()
 			]
 			+SHorizontalBox::Slot()
-			.FillWidth(1.0f)
-			+SHorizontalBox::Slot()
 			.AutoWidth()
-			.HAlign(EHorizontalAlignment::HAlign_Right)
+			.HAlign(EHorizontalAlignment::HAlign_Left)
 			.VAlign(EVerticalAlignment::VAlign_Center)
 			.Padding(4.0f, 0.0f)
 			[
@@ -578,6 +596,8 @@ void SPCGEditorGraphAttributeListView::SetNodeBeingInspected(UPCGEditorGraphNode
 		return;
 	}
 
+	CacheColumnWidthVisibility();
+
 	PCGEditorGraphNode = InPCGEditorGraphNode;
 
 	if (PCGEditorGraphNode.IsValid())
@@ -669,6 +689,7 @@ void SPCGEditorGraphAttributeListView::RefreshAttributeList()
 	ListView->SetItemsSource(&EmptyList);
 
 	PCGColumnData.Empty();
+	ColumnsMaxWidthMapping.Empty();
 	ListViewItems.Empty();
 	ListViewHeader->ClearColumns();
 	InfoTextBlock->SetText(FText::GetEmpty());
@@ -734,6 +755,8 @@ void SPCGEditorGraphAttributeListView::RefreshAttributeList()
 	ListView->SetItemsSource(&ListViewItems);
 	ListView->RequestListRefresh();
 
+	RestoreColumnWidthVisibility();
+
 	LaunchUpdateTask();
 }
 
@@ -913,6 +936,24 @@ TSharedRef<SWidget> SPCGEditorGraphAttributeListView::OnGenerateAdditionalOperat
 		NAME_None,
 		EUserInterfaceActionType::Button);
 
+	MenuBuilder.AddSeparator();
+
+	MenuBuilder.AddMenuEntry(
+		LOCTEXT("ResetDefaultWidth", "Reset columns to default width"),
+		LOCTEXT("ResetDefaultWidthTooltip", "Resets all the columns width to their default."),
+		FSlateIcon(),
+		FUIAction(FExecuteAction::CreateSP(this, &SPCGEditorGraphAttributeListView::ResetColumnsWidthToDefault)),
+		NAME_None,
+		EUserInterfaceActionType::Button);
+	
+	MenuBuilder.AddMenuEntry(
+		LOCTEXT("ExpendMaxWidth", "Expend columns to max width"),
+		LOCTEXT("ResetDefaultWidthTooltip", "Expends all the columns width to their maximum to see their content."),
+		FSlateIcon(),
+		FUIAction(FExecuteAction::CreateSP(this, &SPCGEditorGraphAttributeListView::ExpendAllColumnToMaxWidth)),
+		NAME_None,
+		EUserInterfaceActionType::Button);
+
 	return MenuBuilder.MakeWidget();
 }
 
@@ -994,6 +1035,7 @@ void SPCGEditorGraphAttributeListView::OnSelectionChangedPin(TSharedPtr<FPinComb
 
 	if (InSelectInfo != ESelectInfo::Direct)
 	{
+		CacheColumnWidthVisibility();
 		RefreshAttributeList();
 	}
 }
@@ -1012,6 +1054,7 @@ void SPCGEditorGraphAttributeListView::OnSelectionChanged(TSharedPtr<FName> Item
 {
 	if (SelectInfo != ESelectInfo::Direct)
 	{
+		CacheColumnWidthVisibility();
 		RefreshAttributeList();
 	}
 }
@@ -1199,11 +1242,14 @@ void SPCGEditorGraphAttributeListView::AddColumn(const UPCGData* InPCGData, cons
 	ColumnData.DataAccessor = InColumnInfo.Accessor;
 	ColumnData.DataKeys = InAccessorKeys;
 
+	const float ColumnWidth = InColumnInfo.Width < 0.0f ? PCGEditorGraphAttributeListView::CalculateColumnWidth(InColumnInfo.Label) : InColumnInfo.Width;
+	ColumnsMaxWidthMapping.FindOrAdd(InColumnInfo.Id) = ColumnWidth;
+
 	SHeaderRow::FColumn::FArguments Arguments;
 	Arguments.ColumnId(InColumnInfo.Id);
 	Arguments.DefaultLabel(InColumnInfo.Label);
 	Arguments.DefaultTooltip(InColumnInfo.Tooltip);
-	Arguments.ManualWidth(InColumnInfo.Width < 0.0f ? PCGEditorGraphAttributeListView::CalculateColumnWidth(InColumnInfo.Label) : InColumnInfo.Width);
+	Arguments.ManualWidth(ColumnWidth);
 	Arguments.HAlignHeader(HAlign_Center);
 	Arguments.HAlignCell((EHorizontalAlignment)InColumnInfo.CellAlignment);
 	Arguments.SortMode(this, &SPCGEditorGraphAttributeListView::GetColumnSortMode, InColumnInfo.Id);
@@ -1214,12 +1260,112 @@ void SPCGEditorGraphAttributeListView::AddColumn(const UPCGData* InPCGData, cons
 	[
 		GenerateColumnMenu(InColumnInfo.Id)
 	];
+	Arguments.OnColumnSplitterDoubleClick_Lambda([WeakThis = AsWeak(), ColumnId = InColumnInfo.Id](const FGeometry&, const FPointerEvent&) -> FReply
+	{
+		if (SPCGEditorGraphAttributeListView* This = static_cast<SPCGEditorGraphAttributeListView*>(WeakThis.Pin().Get()))
+		{
+			This->ResizeColumnToMaxWidth(ColumnId);
+			return FReply::Handled();
+		}
+		
+		return FReply::Unhandled();
+	});
 
 	SHeaderRow::FColumn* NewColumn = new SHeaderRow::FColumn(Arguments);
 	NewColumn->bIsVisible = !HiddenAttributes.Contains(InColumnInfo.Id);
 	ListViewHeader->AddColumn(*NewColumn);
 }
 
+void SPCGEditorGraphAttributeListView::ResizeColumnToMaxWidth(const FName ColumnId)
+{
+	ListViewHeader->SetColumnWidth(ColumnId, ColumnsMaxWidthMapping[ColumnId]);
+}
+
+void SPCGEditorGraphAttributeListView::ResetColumnsWidthToDefault()
+{
+	if (!ListViewHeader.IsValid())
+	{
+		return;
+	}
+	
+	for (const SHeaderRow::FColumn& Column : ListViewHeader->GetColumns())
+	{
+		ListViewHeader->SetColumnWidth(Column.ColumnId, PCGEditorGraphAttributeListView::CalculateColumnWidth(Column.DefaultText.Get()));
+	}
+}
+
+void SPCGEditorGraphAttributeListView::ExpendAllColumnToMaxWidth()
+{
+	if (!ListViewHeader.IsValid())
+	{
+		return;
+	}
+	
+	for (const SHeaderRow::FColumn& Column : ListViewHeader->GetColumns())
+	{
+		ResizeColumnToMaxWidth(Column.ColumnId);
+	}
+}
+
+void SPCGEditorGraphAttributeListView::CacheColumnWidthVisibility()
+{
+	if (!PCGEditorGraphNode.IsValid() || !ListViewHeader.IsValid() || ListViewHeader->GetColumns().IsEmpty())
+	{
+		return;
+	}
+	
+	const TObjectKey<UPCGEditorGraphNodeBase> NodeKey(PCGEditorGraphNode.Get());
+	
+	FNodeKeyToColumnWidthVisibilityMap* It = Algo::FindBy(ColumnWidthVisibilityCache, NodeKey, [](const FNodeKeyToColumnWidthVisibilityMap& It)
+	{
+		return It.Get<0>();
+	});
+
+	if (!It)
+	{
+		FNodeKeyToColumnWidthVisibilityMap KeyAndCache;
+		KeyAndCache.Get<0>() = NodeKey;
+		It = &ColumnWidthVisibilityCache.Add_GetRef(std::move(KeyAndCache));
+	}
+
+	check(It);
+	for (const SHeaderRow::FColumn& Column : ListViewHeader->GetColumns())
+	{
+		TTuple<float, bool>& WidthVisiblity = It->Get<1>().FindOrAdd(Column.ColumnId);
+		WidthVisiblity.Get<0>() = Column.GetWidth();
+		WidthVisiblity.Get<1>() = Column.bIsVisible;
+	}
+}
+
+void SPCGEditorGraphAttributeListView::RestoreColumnWidthVisibility()
+{
+	if (!PCGEditorGraphNode.IsValid() || !ListViewHeader.IsValid() || ListViewHeader->GetColumns().IsEmpty())
+	{
+		return;
+	}
+	
+	const TObjectKey<UPCGEditorGraphNodeBase> NodeKey(PCGEditorGraphNode.Get());
+	
+	const FNodeKeyToColumnWidthVisibilityMap* It = Algo::FindBy(ColumnWidthVisibilityCache, NodeKey, [](const FNodeKeyToColumnWidthVisibilityMap& It)
+	{
+		return It.Get<0>();
+	});
+
+	if (!It)
+	{
+		return;
+	}
+
+	for (const SHeaderRow::FColumn& Column : ListViewHeader->GetColumns())
+	{
+		if (const TTuple<float, bool>* WidthVisibility = It->Get<1>().Find(Column.ColumnId))
+		{
+			ListViewHeader->SetColumnWidth(Column.ColumnId, WidthVisibility->Get<0>());
+			ListViewHeader->SetShowGeneratedColumn(Column.ColumnId, WidthVisibility->Get<1>());
+		}
+	}
+}
+
 TSharedRef<SWidget> SPCGEditorGraphAttributeListView::GenerateColumnMenu(FName ColumnId)
 {
 	FMenuBuilder MenuBuilder(/*bShouldCloseAfterMenuSelection=*/true, nullptr);
@@ -1239,6 +1385,20 @@ TSharedRef<SWidget> SPCGEditorGraphAttributeListView::GenerateColumnMenu(FName C
 		NAME_None,
 		EUserInterfaceActionType::Button);
 
+	MenuBuilder.AddMenuEntry(
+		LOCTEXT("ResizeColumnToMaxWidth", "Resize Column width to match content"),
+		LOCTEXT("CopyAttributeNameToClipboardTooltip", "Resizes column width to match the content width."),
+		FSlateIcon(),
+		FUIAction(FExecuteAction::CreateLambda([WeakThis = AsWeak(), ColumnId]()
+		{
+			if (SPCGEditorGraphAttributeListView* This = static_cast<SPCGEditorGraphAttributeListView*>(WeakThis.Pin().Get()))
+			{
+				This->ResizeColumnToMaxWidth(ColumnId);
+			}
+		})),
+		NAME_None,
+		EUserInterfaceActionType::Button);
+
 	return MenuBuilder.MakeWidget();
 }
 
diff --git a/Engine/Plugins/PCG/Source/PCGEditor/Private/SPCGEditorGraphAttributeListView.h b/Engine/Plugins/PCG/Source/PCGEditor/Private/SPCGEditorGraphAttributeListView.h
--- a/Engine/Plugins/PCG/Source/PCGEditor/Private/SPCGEditorGraphAttributeListView.h
+++ b/Engine/Plugins/PCG/Source/PCGEditor/Private/SPCGEditorGraphAttributeListView.h
@@ -7,9 +7,10 @@
 #include "Metadata/Accessors/IPCGAttributeAccessor.h"
 #include "Metadata/Accessors/PCGAttributeAccessorKeys.h"
 
+#include "Containers/RingBuffer.h"
 #include "Misc/TextFilterExpressionEvaluator.h"
 #include "Tasks/Task.h"
-#include "UObject/StrongObjectPtr.h"
+#include "UObject/ObjectKey.h"
 #include "Widgets/Views/SListView.h"
 #include "Widgets/Views/STableRow.h"
 
@@ -40,6 +41,7 @@ struct FSlateBrush;
 namespace PCGEditorGraphAttributeListView
 {
 	constexpr float MaxColumnWidth = 200.0f;
+	constexpr int32 MaxNodeColumnWidthCachedItems = 256;
 }
 
 struct FPCGListViewItem
@@ -159,6 +161,12 @@ private:
 
 	void OnGenerateUpdated(UPCGComponent* InPCGComponent);
 
+	void ResizeColumnToMaxWidth(FName InColumnId);
+	void ResetColumnsWidthToDefault();
+	void ExpendAllColumnToMaxWidth();
+	void CacheColumnWidthVisibility();
+	void RestoreColumnWidthVisibility();
+
 	void RefreshAttributeList();
 	void RefreshPinComboBox();
 	void RefreshDataComboBox();
@@ -266,6 +274,13 @@ private:
 	TArray<FName> HiddenAttributes;
 
 	TMap<FName, FPCGColumnData> PCGColumnData;
+	
+	/** Contains info of the max width to display the full content of the column. Computed only for string types. */
+	TMap<FName, float> ColumnsMaxWidthMapping;
+
+	/** Ring buffer to keep the latest inspected nodes column width and visibility. */
+	using FNodeKeyToColumnWidthVisibilityMap = TTuple<TObjectKey<UPCGEditorGraphNodeBase>, TMap<FName, TTuple<float, bool>>>;
+	TRingBuffer<FNodeKeyToColumnWidthVisibilityMap> ColumnWidthVisibilityCache;
 
 	FText ActiveFilterText;
 
diff --git a/Engine/Source/Runtime/Slate/Private/Widgets/Views/SHeaderRow.cpp b/Engine/Source/Runtime/Slate/Private/Widgets/Views/SHeaderRow.cpp
--- a/Engine/Source/Runtime/Slate/Private/Widgets/Views/SHeaderRow.cpp
+++ b/Engine/Source/Runtime/Slate/Private/Widgets/Views/SHeaderRow.cpp
@@ -879,6 +879,7 @@ void SHeaderRow::RegenerateWidgets()
 					SizingGrip->SetOnMouseButtonDown(FPointerEventHandler::CreateLambda(SizingGrip_OnMouseButtonDown));
 					SizingGrip->SetOnMouseButtonUp(FPointerEventHandler::CreateLambda(SizingGrip_OnMouseButtonUp));
 					SizingGrip->SetOnMouseMove(FPointerEventHandler::CreateLambda(SizingGrip_OnMouseMove));
+					SizingGrip->SetOnMouseDoubleClick(SomeColumn.OnColumnSplitterDoubleClick);
 
 					auto GetColumnWidthAsOptionalSize = [&SomeColumn, SplitterHandleSizeCopy = SplitterHandleSize]() -> FOptionalSize
 					{
diff --git a/Engine/Source/Runtime/Slate/Public/Widgets/Views/SHeaderRow.h b/Engine/Source/Runtime/Slate/Public/Widgets/Views/SHeaderRow.h
--- a/Engine/Source/Runtime/Slate/Public/Widgets/Views/SHeaderRow.h
+++ b/Engine/Source/Runtime/Slate/Public/Widgets/Views/SHeaderRow.h
@@ -147,6 +147,9 @@ public:
 			SLATE_NAMED_SLOT( FArguments, MenuContent )
 			SLATE_EVENT( FOnGetContent, OnGetMenuContent )
 
+			/** Allow to provide a custom handler when a splitter is double-clicked in Manual column size mode. */
+			SLATE_EVENT( FPointerEventHandler, OnColumnSplitterDoubleClick )
+
 			SLATE_ARGUMENT( EHorizontalAlignment, HAlignCell )
 			SLATE_ARGUMENT( EVerticalAlignment, VAlignCell )
 
@@ -171,6 +174,7 @@ public:
 			, HeaderContent( InArgs._HeaderContent )
 			, HeaderMenuContent( InArgs._MenuContent )
 			, OnGetMenuContent( InArgs._OnGetMenuContent )
+			, OnColumnSplitterDoubleClick( InArgs._OnColumnSplitterDoubleClick )
 			, HeaderHAlignment( InArgs._HAlignHeader )
 			, HeaderVAlignment( InArgs._VAlignHeader )
 			, HeaderContentPadding( InArgs._HeaderContentPadding )
@@ -254,6 +258,7 @@ public:
 		TAlwaysValidWidget HeaderMenuContent;
 
 		FOnGetContent OnGetMenuContent;
+		FPointerEventHandler OnColumnSplitterDoubleClick;
 
 		EHorizontalAlignment HeaderHAlignment;
 		EVerticalAlignment HeaderVAlignment;
