[PCG] Improvement on the actor filtering from input for World Query

* Use set of ObjectKey instead of array of soft pointers
* Cache actor filters as Object Keys on World Data, with same mecanism as point data octree.
* Add helper function to extract the set directly, when we don't need to store the soft pointers.

#rnx

diff --git a/Engine/Plugins/PCG/Source/PCG/Private/Data/PCGWorldData.cpp b/Engine/Plugins/PCG/Source/PCG/Private/Data/PCGWorldData.cpp
--- a/Engine/Plugins/PCG/Source/PCG/Private/Data/PCGWorldData.cpp
+++ b/Engine/Plugins/PCG/Source/PCG/Private/Data/PCGWorldData.cpp
@@ -25,6 +25,28 @@
 
 #define LOCTEXT_NAMESPACE "PCGWorldData"
 
+namespace PCGWorldCommonQueryParams
+{
+	template<class T, class Container, typename Func>
+	bool ExtractActorFiltersIfNeeded(const FPCGWorldCommonQueryParams& InParams, const UPCGData* InData, Container& OutContainer, FPCGContext* InContext, Func TransformLambda)
+	{
+		if (InParams.ActorFilterFromInput != EPCGWorldQueryFilter::None)
+		{
+			TArray<FSoftObjectPath> FilterActors;
+
+			if (!PCGAttributeAccessorHelpers::ExtractAllValues(InData, InParams.ActorFilterInputSource, FilterActors, InContext))
+			{
+				return false;
+			}
+
+			OutContainer.Empty(FilterActors.Num());
+			Algo::Transform(FilterActors, OutContainer, TransformLambda);
+		}
+
+		return true;
+	}
+}
+
 void FPCGWorldCommonQueryParams::AddFilterPinIfNeeded(TArray<FPCGPinProperties>& PinProperties) const
 {
 	if (ActorFilterFromInput != EPCGWorldQueryFilter::None)
@@ -38,20 +60,14 @@ void FPCGWorldCommonQueryParams::AddFilterPinIfNeeded(TArray<FPCGPinProperties>&
 
 bool FPCGWorldCommonQueryParams::ExtractActorFiltersIfNeeded(const UPCGData* InData, TArray<TSoftObjectPtr<AActor>>& OutArray, FPCGContext* InContext) const
 {
-	if (ActorFilterFromInput != EPCGWorldQueryFilter::None)
-	{
-		TArray<FSoftObjectPath> FilterActors;
-
-		if (!PCGAttributeAccessorHelpers::ExtractAllValues(InData, ActorFilterInputSource, FilterActors, InContext))
-		{
-			return false;
-		}
-
-		OutArray.Empty(FilterActors.Num());
-		Algo::Transform(FilterActors, OutArray, [](const FSoftObjectPath& ActorPath) { return TSoftObjectPtr<AActor>(ActorPath); });
-	}
+	return PCGWorldCommonQueryParams::ExtractActorFiltersIfNeeded<TSoftObjectPtr<AActor>>(*this, InData, OutArray, InContext,
+		[](const FSoftObjectPath& ActorPath) -> TSoftObjectPtr<AActor> { return TSoftObjectPtr<AActor>(ActorPath); });
+}
 
-	return true;
+bool FPCGWorldCommonQueryParams::ExtractLoadedActorFiltersIfNeeded(const UPCGData* InData, TSet<TObjectKey<AActor>>& OutSet, FPCGContext* InContext) const
+{
+	return PCGWorldCommonQueryParams::ExtractActorFiltersIfNeeded<TObjectKey<AActor>>(*this, InData, OutSet, InContext,
+	[](const FSoftObjectPath& ActorPath) -> TObjectKey<AActor> { return TObjectKey<AActor>(Cast<AActor>(ActorPath.ResolveObject())); });
 }
 
 void FPCGWorldCommonQueryParams::Initialize()
@@ -127,6 +143,22 @@ void FPCGWorldRayHitQueryParams::PostSerialize(const FArchive& Ar)
 #endif
 }
 
+const TSet<TObjectKey<AActor>>& FPCGWorldQueryActorFilterCache::GetCachedFilterActors() const
+{
+	if (!FilterActors.IsEmpty() && CachedFilterActorsDirty)
+	{
+		FScopeLock Lock(&CachedFilterActorsLock);
+		if (CachedFilterActorsDirty)
+		{
+			CachedFilterActors.Empty(FilterActors.Num());
+			Algo::Transform(FilterActors, CachedFilterActors, [](const TSoftObjectPtr<AActor>& SoftPtr) { return TObjectKey<AActor>(SoftPtr.Get()); });
+			CachedFilterActorsDirty = false;
+		}
+	}
+
+	return CachedFilterActors;
+}
+
 void UPCGWorldVolumetricData::Initialize(UWorld* InWorld, const FBox& InBounds)
 {
 	Super::Initialize(InBounds);
@@ -150,7 +182,7 @@ bool UPCGWorldVolumetricData::SamplePoint(const FTransform& InTransform, const F
 	TArray<FOverlapResult> Overlaps;
 	/*bool bOverlaps =*/ World->OverlapMultiByObjectType(Overlaps, InTransform.TransformPosition(InBounds.GetCenter()), InTransform.GetRotation(), ObjectQueryParams, CollisionShape, Params);
 
-	TOptional<FOverlapResult> Overlap = PCGWorldQueryHelpers::FilterOverlapResults(&QueryParams, OriginatingComponent, Overlaps, FilterActors);
+	TOptional<FOverlapResult> Overlap = PCGWorldQueryHelpers::FilterOverlapResults(&QueryParams, OriginatingComponent, Overlaps, ActorFilter.GetCachedFilterActors());
 
 	// If searched for overlap and found one, or didn't search and didn't find one, set the point and return true. Otherwise, return false.
 	if (Overlap.IsSet() == QueryParams.bSearchForOverlap)
@@ -224,6 +256,7 @@ UPCGSpatialData* UPCGWorldVolumetricData::CopyInternal(FPCGContext* Context) con
 	NewVolumetricData->OriginatingComponent = OriginatingComponent;
 	NewVolumetricData->QueryParams = QueryParams;
 	NewVolumetricData->QueryParams.Initialize();
+	NewVolumetricData->ActorFilter.GetFilterActorsMutable() = ActorFilter.GetFilterActorsConst();
 
 	return NewVolumetricData;
 }
@@ -264,7 +297,7 @@ bool UPCGWorldRayHitData::SamplePoint(const FTransform& InTransform, const FBox&
 	TArray<FHitResult> Hits;
 	if (World->LineTraceMultiByObjectType(Hits, RayOrigin, RayEnd, ObjectQueryParams, Params))
 	{
-		TOptional<FHitResult> HitResult = PCGWorldQueryHelpers::FilterRayHitResults(&QueryParams, OriginatingComponent, Hits, FilterActors);
+		TOptional<FHitResult> HitResult = PCGWorldQueryHelpers::FilterRayHitResults(&QueryParams, OriginatingComponent, Hits, ActorFilter.GetCachedFilterActors());
 
 		if (HitResult.IsSet())
 		{
@@ -326,6 +359,7 @@ UPCGSpatialData* UPCGWorldRayHitData::CopyInternal(FPCGContext* Context) const
 	NewData->Bounds = Bounds;
 	NewData->QueryParams = QueryParams;
 	NewData->QueryParams.Initialize();
+	NewData->ActorFilter.GetFilterActorsMutable() = ActorFilter.GetFilterActorsConst();
 
 	return NewData;
 }
diff --git a/Engine/Plugins/PCG/Source/PCG/Private/Elements/PCGPathfindingElement.cpp b/Engine/Plugins/PCG/Source/PCG/Private/Elements/PCGPathfindingElement.cpp
--- a/Engine/Plugins/PCG/Source/PCG/Private/Elements/PCGPathfindingElement.cpp
+++ b/Engine/Plugins/PCG/Source/PCG/Private/Elements/PCGPathfindingElement.cpp
@@ -335,8 +335,8 @@ bool FPCGPathfindingElement::PrepareDataInternal(FPCGContext* InContext) const
 			}
 
 			// --- Gather filtered actors ---
-			TArray<TSoftObjectPtr<AActor>> CachedFilterActors;
-			if (!Settings->PathTraceParams.ExtractActorFiltersIfNeeded(FilterActorsData, CachedFilterActors, Context))
+			TSet<TObjectKey<AActor>> CachedFilterActors;
+			if (!Settings->PathTraceParams.ExtractLoadedActorFiltersIfNeeded(FilterActorsData, CachedFilterActors, Context))
 			{
 				return EPCGTimeSliceInitResult::NoOperation;
 			}
@@ -356,7 +356,7 @@ bool FPCGPathfindingElement::PrepareDataInternal(FPCGContext* InContext) const
 					TArray<FHitResult> OutHits;
 					if (World->LineTraceMultiByObjectType(OutHits, StartPosition, EndPosition, ObjectQueryParams, Params))
 					{
-						TOptional<FHitResult> HitResult = PCGWorldQueryHelpers::FilterRayHitResults(&PathTraceParams, OriginatingComponent, OutHits,CachedFilterActors);
+						TOptional<FHitResult> HitResult = PCGWorldQueryHelpers::FilterRayHitResults(&PathTraceParams, OriginatingComponent, OutHits, CachedFilterActors);
 						return !HitResult.IsSet();
 					}
 					else
diff --git a/Engine/Plugins/PCG/Source/PCG/Private/Elements/PCGWorldQuery.cpp b/Engine/Plugins/PCG/Source/PCG/Private/Elements/PCGWorldQuery.cpp
--- a/Engine/Plugins/PCG/Source/PCG/Private/Elements/PCGWorldQuery.cpp
+++ b/Engine/Plugins/PCG/Source/PCG/Private/Elements/PCGWorldQuery.cpp
@@ -33,7 +33,7 @@ namespace PCGWorldQuery
 			}
 
 			const UPCGData* ActorFilterData = ActorFilterTaggedData.IsEmpty() ? nullptr : ActorFilterTaggedData[0].Data;
-			if (ActorFilterData && !QueryParams.ExtractActorFiltersIfNeeded(ActorFilterData, WorldQueryData->FilterActors, Context))
+			if (ActorFilterData && !QueryParams.ExtractActorFiltersIfNeeded(ActorFilterData, WorldQueryData->ActorFilter.GetFilterActorsMutable(), Context))
 			{
 				PCGLog::LogWarningOnGraph(LOCTEXT("FailExtractActorFilters", "Failed to extract actor filters."));
 			}
diff --git a/Engine/Plugins/PCG/Source/PCG/Private/Elements/PCGWorldRaycast.cpp b/Engine/Plugins/PCG/Source/PCG/Private/Elements/PCGWorldRaycast.cpp
--- a/Engine/Plugins/PCG/Source/PCG/Private/Elements/PCGWorldRaycast.cpp
+++ b/Engine/Plugins/PCG/Source/PCG/Private/Elements/PCGWorldRaycast.cpp
@@ -164,7 +164,7 @@ bool FPCGWorldRaycastElement::PrepareDataInternal(FPCGContext* InContext) const
 		}
 
 		// --- Gather filtered actors ---
-		if (!Settings->WorldQueryParams.ExtractActorFiltersIfNeeded(FilterActorsData, OutState.CachedFilterActors, Context))
+		if (!Settings->WorldQueryParams.ExtractLoadedActorFiltersIfNeeded(FilterActorsData, OutState.CachedFilterActors, Context))
 		{
 			return EPCGTimeSliceInitResult::NoOperation;
 		}
diff --git a/Engine/Plugins/PCG/Source/PCG/Private/Helpers/PCGWorldQueryHelpers.cpp b/Engine/Plugins/PCG/Source/PCG/Private/Helpers/PCGWorldQueryHelpers.cpp
--- a/Engine/Plugins/PCG/Source/PCG/Private/Helpers/PCGWorldQueryHelpers.cpp
+++ b/Engine/Plugins/PCG/Source/PCG/Private/Helpers/PCGWorldQueryHelpers.cpp
@@ -40,7 +40,7 @@ namespace PCGWorldQueryHelpers
 		const FPCGWorldCommonQueryParams* QueryParams,
 		const UPrimitiveComponent* TriggeredComponent,
 		const TWeakObjectPtr<UPCGComponent> OriginatingComponent,
-		const TArrayView<const TSoftObjectPtr<AActor>> FilteredActorReferences)
+		const TSet<TObjectKey<AActor>>& FilteredActorReferences)
 	{
 		check(QueryParams && TriggeredComponent);
 
@@ -116,7 +116,7 @@ namespace PCGWorldQueryHelpers
 
 			if (!bHardExclude && QueryParams->ActorFilterFromInput != EPCGWorldQueryFilter::None)
 			{
-				const bool bFoundMatch = Algo::AnyOf(FilteredActorReferences, [Actor](const TSoftObjectPtr<AActor>& ActorRef) { return ActorRef == Actor; });
+				const bool bFoundMatch = FilteredActorReferences.Contains(TObjectKey<AActor>(Actor));
 				UpdateIncludeExcludeStates(QueryParams->ActorFilterFromInput, bFoundMatch);
 			}
 
@@ -148,7 +148,7 @@ namespace PCGWorldQueryHelpers
 		const FPCGWorldRaycastQueryParams* QueryParams,
 		const TWeakObjectPtr<UPCGComponent> OriginatingComponent,
 		const TArray<FHitResult>& HitResults,
-		const TArrayView<const TSoftObjectPtr<AActor>> FilteredActorReferences)
+		const TSet<TObjectKey<AActor>>& FilteredActorReferences)
 	{
 		check(QueryParams);
 
@@ -181,7 +181,7 @@ namespace PCGWorldQueryHelpers
 		const FPCGWorldVolumetricQueryParams* QueryParams,
 		const TWeakObjectPtr<UPCGComponent> OriginatingComponent,
 		const TArray<FOverlapResult>& OverlapResults,
-		const TArrayView<const TSoftObjectPtr<AActor>> FilteredActorReferences)
+		const TSet<TObjectKey<AActor>>& FilteredActorReferences)
 	{
 		check(QueryParams);
 
diff --git a/Engine/Plugins/PCG/Source/PCG/Public/Data/PCGWorldData.h b/Engine/Plugins/PCG/Source/PCG/Public/Data/PCGWorldData.h
--- a/Engine/Plugins/PCG/Source/PCG/Public/Data/PCGWorldData.h
+++ b/Engine/Plugins/PCG/Source/PCG/Public/Data/PCGWorldData.h
@@ -7,6 +7,7 @@
 #include "Helpers/PCGWorldQueryHelpers.h"
 
 #include "Engine/EngineTypes.h"
+#include "UObject/ObjectKey.h"
 
 #include "PCGWorldData.generated.h"
 
@@ -121,6 +122,9 @@ struct FPCGWorldCommonQueryParams
 	/** Utility function to extract the Actor Filters from the incoming data using the selector in this struct. Return true if the extraction succeeded. */
 	PCG_API bool ExtractActorFiltersIfNeeded(const UPCGData* InData, TArray<TSoftObjectPtr<AActor>>& OutArray, FPCGContext* InContext = nullptr) const;
 
+	/** Utility function to extract the Actor Filters (from loaded actors) from the incoming data using the selector in this struct. Return true if the extraction succeeded. */
+	PCG_API bool ExtractLoadedActorFiltersIfNeeded(const UPCGData* InData, TSet<TObjectKey<AActor>>& OutSet, FPCGContext* InContext = nullptr) const;
+
 protected:
 	/** Sets up the data we need to efficiently perform the queries */
 	void Initialize();
@@ -196,7 +200,7 @@ struct FPCGWorldRaycastQueryParams : public FPCGWorldCommonQueryParams
 };
 
 template<>
-struct TStructOpsTypeTraits<FPCGWorldRaycastQueryParams> : public TStructOpsTypeTraits<FPCGWorldCommonQueryParams>
+struct TStructOpsTypeTraits<FPCGWorldRaycastQueryParams> : public TStructOpsTypeTraitsBase2<FPCGWorldRaycastQueryParams>
 {
 	enum
 	{
@@ -218,7 +222,7 @@ struct FPCGWorldVolumetricQueryParams : public FPCGWorldCommonQueryParams
 };
 
 template<>
-struct TStructOpsTypeTraits<FPCGWorldVolumetricQueryParams> : public TStructOpsTypeTraits<FPCGWorldCommonQueryParams>
+struct TStructOpsTypeTraits<FPCGWorldVolumetricQueryParams> : public TStructOpsTypeTraitsBase2<FPCGWorldVolumetricQueryParams>
 {
 	enum
 	{
@@ -252,7 +256,7 @@ struct FPCGWorldRayHitQueryParams : public FPCGWorldRaycastQueryParams
 };
 
 template<>
-struct TStructOpsTypeTraits<FPCGWorldRayHitQueryParams> : public TStructOpsTypeTraits<FPCGWorldCommonQueryParams>
+struct TStructOpsTypeTraits<FPCGWorldRayHitQueryParams> : public TStructOpsTypeTraitsBase2<FPCGWorldRayHitQueryParams>
 {
 	enum
 	{
@@ -260,6 +264,36 @@ struct TStructOpsTypeTraits<FPCGWorldRayHitQueryParams> : public TStructOpsTypeT
 	};
 };
 
+USTRUCT()
+struct FPCGWorldQueryActorFilterCache
+{
+	GENERATED_BODY();
+	
+public:
+	PCG_API const TSet<TObjectKey<AActor>>& GetCachedFilterActors() const;
+	TArray<TSoftObjectPtr<AActor>>& GetFilterActorsMutable() { CachedFilterActorsDirty = true; return FilterActors; }
+	const TArray<TSoftObjectPtr<AActor>>& GetFilterActorsConst() const { return FilterActors; }
+	
+private:
+	UPROPERTY()
+	TArray<TSoftObjectPtr<AActor>> FilterActors;
+
+	// Transient data to cache filter actors
+	mutable TSet<TObjectKey<AActor>> CachedFilterActors;
+	mutable FCriticalSection CachedFilterActorsLock;
+	mutable bool CachedFilterActorsDirty = false;
+};
+
+// Not copyable because of the lock, but FilterActors will be copied in the Copy Internal of the data.
+template<>
+struct TStructOpsTypeTraits<FPCGWorldQueryActorFilterCache> : public TStructOpsTypeTraitsBase2<FPCGWorldQueryActorFilterCache>
+{
+	enum
+	{
+		WithCopy = false,
+	};
+};
+
 /** Queries volume for presence of world collision or not. Can be used to voxelize environment. */
 UCLASS(MinimalAPI, BlueprintType, ClassGroup = (Procedural))
 class UPCGWorldVolumetricData : public UPCGVolumeData
@@ -295,7 +329,7 @@ public:
 	FPCGWorldVolumetricQueryParams QueryParams;
 	
 	UPROPERTY()
-    TArray<TSoftObjectPtr<AActor>> FilterActors;
+	FPCGWorldQueryActorFilterCache ActorFilter;
 };
 
 /** Executes collision queries against world collision. */
@@ -342,5 +376,5 @@ public:
 	FPCGWorldRayHitQueryParams QueryParams;
 	
 	UPROPERTY()
-	TArray<TSoftObjectPtr<AActor>> FilterActors;
+	FPCGWorldQueryActorFilterCache ActorFilter;
 };
diff --git a/Engine/Plugins/PCG/Source/PCG/Public/Elements/PCGWorldRaycast.h b/Engine/Plugins/PCG/Source/PCG/Public/Elements/PCGWorldRaycast.h
--- a/Engine/Plugins/PCG/Source/PCG/Public/Elements/PCGWorldRaycast.h
+++ b/Engine/Plugins/PCG/Source/PCG/Public/Elements/PCGWorldRaycast.h
@@ -102,7 +102,7 @@ namespace FPCGWorldRaycast
 		TArray<FVector> CachedRayOrigins;
 		TArray<FVector> CachedRayVectors;
 		TArray<double> CachedRayLengths;
-		TArray<TSoftObjectPtr<AActor>> CachedFilterActors;
+		TSet<TObjectKey<AActor>> CachedFilterActors;
 	};
 }
 
diff --git a/Engine/Plugins/PCG/Source/PCG/Public/Helpers/PCGWorldQueryHelpers.h b/Engine/Plugins/PCG/Source/PCG/Public/Helpers/PCGWorldQueryHelpers.h
--- a/Engine/Plugins/PCG/Source/PCG/Public/Helpers/PCGWorldQueryHelpers.h
+++ b/Engine/Plugins/PCG/Source/PCG/Public/Helpers/PCGWorldQueryHelpers.h
@@ -5,6 +5,7 @@
 #include "Containers/ArrayView.h"
 #include "Engine/HitResult.h"
 #include "Misc/Optional.h"
+#include "UObject/ObjectKey.h"
 #include "UObject/SoftObjectPtr.h"
 
 #include "PCGWorldQueryHelpers.generated.h"
@@ -53,21 +54,21 @@ namespace PCGWorldQueryHelpers
 		const FPCGWorldCommonQueryParams* QueryParams,
 		const UPrimitiveComponent* TriggeredComponent,
 		const TWeakObjectPtr<UPCGComponent> OriginatingComponent,
-		const TArrayView<const TSoftObjectPtr<AActor>> FilteredActorReferences);
+		const TSet<TObjectKey<AActor>>& FilteredActorReferences);
 
 	/** Filters through an array of hit results, testing them against raycast query parameters. Optional extra filter predicate. */
 	TOptional<FHitResult> FilterRayHitResults(
 		const FPCGWorldRaycastQueryParams* QueryParams,
 		TWeakObjectPtr<UPCGComponent> OriginatingComponent,
 		const TArray<FHitResult>& HitResults,
-		const TArrayView<const TSoftObjectPtr<AActor>> FilteredActorReferences);
+		const TSet<TObjectKey<AActor>>& FilteredActorReferences);
 
 	/** Filters through an array of overlap results, testing them against raycast query parameters. Optional extra filter predicate. */
 	TOptional<FOverlapResult> FilterOverlapResults(
 		const FPCGWorldVolumetricQueryParams* QueryParams,
 		TWeakObjectPtr<UPCGComponent> OriginatingComponent,
 		const TArray<FOverlapResult>& OverlapResults,
-		const TArrayView<const TSoftObjectPtr<AActor>> FilteredActorReferences);
+		const TSet<TObjectKey<AActor>>& FilteredActorReferences);
 
 	/** Creates hit result attributes based off query params. Can be called before ApplyRayHitMetadata. */
 	bool CreateRayHitAttributes(const FPCGWorldRaycastQueryParams& QueryParams, UPCGMetadata* OutMetadata);
