[PCG] Fix loc

#rnx

diff --git a/Engine/Plugins/PCG/Source/PCGEditor/Private/SPCGEditorGraphAttributeListView.cpp b/Engine/Plugins/PCG/Source/PCGEditor/Private/SPCGEditorGraphAttributeListView.cpp
--- a/Engine/Plugins/PCG/Source/PCGEditor/Private/SPCGEditorGraphAttributeListView.cpp
+++ b/Engine/Plugins/PCG/Source/PCGEditor/Private/SPCGEditorGraphAttributeListView.cpp
@@ -948,7 +948,7 @@ TSharedRef<SWidget> SPCGEditorGraphAttributeListView::OnGenerateAdditionalOperat
 	
 	MenuBuilder.AddMenuEntry(
 		LOCTEXT("ExpendMaxWidth", "Expend columns to max width"),
-		LOCTEXT("ResetDefaultWidthTooltip", "Expends all the columns width to their maximum to see their content."),
+		LOCTEXT("ExpendMaxWidthTooltip", "Expends all the columns width to their maximum to see their content."),
 		FSlateIcon(),
 		FUIAction(FExecuteAction::CreateSP(this, &SPCGEditorGraphAttributeListView::ExpendAllColumnToMaxWidth)),
 		NAME_None,
@@ -1387,7 +1387,7 @@ TSharedRef<SWidget> SPCGEditorGraphAttributeListView::GenerateColumnMenu(FName C
 
 	MenuBuilder.AddMenuEntry(
 		LOCTEXT("ResizeColumnToMaxWidth", "Resize Column width to match content"),
-		LOCTEXT("CopyAttributeNameToClipboardTooltip", "Resizes column width to match the content width."),
+		LOCTEXT("ResizeColumnToMaxWidthTooltip", "Resizes column width to match the content width."),
 		FSlateIcon(),
 		FUIAction(FExecuteAction::CreateLambda([WeakThis = AsWeak(), ColumnId]()
 		{
