[PCG] Flesh out PISMC descriptor to add missing fields, for example the Cast Shadow setting which will be critical for microscattering.

#rnx

diff --git a/Engine/Plugins/PCG/Source/PCG/Private/Components/PCGProceduralISMComponent.cpp b/Engine/Plugins/PCG/Source/PCG/Private/Components/PCGProceduralISMComponent.cpp
--- a/Engine/Plugins/PCG/Source/PCG/Private/Components/PCGProceduralISMComponent.cpp
+++ b/Engine/Plugins/PCG/Source/PCG/Private/Components/PCGProceduralISMComponent.cpp
@@ -120,6 +120,28 @@ FPCGProceduralISMComponentDescriptor& FPCGProceduralISMComponentDescriptor::oper
 	OverlayMaterial = Other.OverlayMaterial.LoadSynchronous();
 	StaticMesh = Other.StaticMesh.LoadSynchronous();
 	ComponentTags = Other.ComponentTags;
+	Mobility = Other.Mobility;
+	VirtualTextureRenderPassType = Other.VirtualTextureRenderPassType;
+	LightingChannels = Other.LightingChannels;
+	CustomDepthStencilWriteMask = Other.CustomDepthStencilWriteMask;
+	VirtualTextureCullMips = Other.VirtualTextureCullMips;
+	TranslucencySortPriority = Other.TranslucencySortPriority;
+	CustomDepthStencilValue = Other.CustomDepthStencilValue;
+	bCastShadow = Other.bCastShadow;
+	bEmissiveLightSource = Other.bEmissiveLightSource;
+	bCastDynamicShadow = Other.bCastDynamicShadow;
+	bCastStaticShadow = Other.bCastStaticShadow;
+	bCastContactShadow = Other.bCastContactShadow;
+	bCastShadowAsTwoSided = Other.bCastShadowAsTwoSided;
+	bCastHiddenShadow = Other.bCastHiddenShadow;
+	bReceivesDecals = Other.bReceivesDecals;
+	bUseAsOccluder = Other.bUseAsOccluder;
+	bRenderCustomDepth = Other.bRenderCustomDepth;
+	bEvaluateWorldPositionOffset = Other.bEvaluateWorldPositionOffset;
+	bReverseCulling = Other.bReverseCulling;
+	WorldPositionOffsetDisableDistance = Other.WorldPositionOffsetDisableDistance;
+	ShadowCacheInvalidationBehavior = Other.ShadowCacheInvalidationBehavior;
+	DetailMode = Other.DetailMode;
 
 	Algo::Transform(Other.OverrideMaterials, OverrideMaterials, [](const TSoftObjectPtr<UMaterialInterface>& OverrideMaterial)
 	{
@@ -145,6 +167,28 @@ void FPCGProceduralISMComponentDescriptor::InitComponent(UPCGProceduralISMCompon
 	InProceduralISMComponent->SetLocalBounds(LocalBounds);
 	InProceduralISMComponent->SetCullDistances(InstanceStartCullDistance, InstanceEndCullDistance);
 	InProceduralISMComponent->ComponentTags = ComponentTags;
+	InProceduralISMComponent->Mobility = Mobility;
+	InProceduralISMComponent->VirtualTextureRenderPassType = VirtualTextureRenderPassType;
+	InProceduralISMComponent->LightingChannels = LightingChannels;
+	InProceduralISMComponent->CustomDepthStencilWriteMask = CustomDepthStencilWriteMask;
+	InProceduralISMComponent->VirtualTextureCullMips = VirtualTextureCullMips;
+	InProceduralISMComponent->TranslucencySortPriority = TranslucencySortPriority;
+	InProceduralISMComponent->CustomDepthStencilValue = CustomDepthStencilValue;
+	InProceduralISMComponent->CastShadow = bCastShadow;
+	InProceduralISMComponent->bEmissiveLightSource = bEmissiveLightSource;
+	InProceduralISMComponent->bCastDynamicShadow = bCastDynamicShadow;
+	InProceduralISMComponent->bCastStaticShadow = bCastStaticShadow;
+	InProceduralISMComponent->bCastContactShadow = bCastContactShadow;
+	InProceduralISMComponent->bCastShadowAsTwoSided = bCastShadowAsTwoSided;
+	InProceduralISMComponent->bCastHiddenShadow = bCastHiddenShadow;
+	InProceduralISMComponent->bReceivesDecals = bReceivesDecals;
+	InProceduralISMComponent->bUseAsOccluder = bUseAsOccluder;
+	InProceduralISMComponent->bRenderCustomDepth = bRenderCustomDepth;
+	InProceduralISMComponent->bEvaluateWorldPositionOffset = bEvaluateWorldPositionOffset;
+	InProceduralISMComponent->bReverseCulling = bReverseCulling;
+	InProceduralISMComponent->WorldPositionOffsetDisableDistance = WorldPositionOffsetDisableDistance;
+	InProceduralISMComponent->ShadowCacheInvalidationBehavior = ShadowCacheInvalidationBehavior;
+	InProceduralISMComponent->DetailMode = DetailMode;
 }
 
 void FPCGProceduralISMComponentDescriptor::InitFrom(UPCGProceduralISMComponent* InProceduralISMComponent)
@@ -157,6 +201,29 @@ void FPCGProceduralISMComponentDescriptor::InitFrom(UPCGProceduralISMComponent*
 	NumCustomFloats = InProceduralISMComponent->GetNumCustomDataFloats();
 	InProceduralISMComponent->GetLocalBounds(LocalBounds);
 	InProceduralISMComponent->GetCullDistances(InstanceStartCullDistance, InstanceEndCullDistance);
+	ComponentTags = InProceduralISMComponent->ComponentTags;
+	Mobility = InProceduralISMComponent->Mobility;
+	VirtualTextureRenderPassType = InProceduralISMComponent->VirtualTextureRenderPassType;
+	LightingChannels = InProceduralISMComponent->LightingChannels;
+	CustomDepthStencilWriteMask = InProceduralISMComponent->CustomDepthStencilWriteMask;
+	VirtualTextureCullMips = InProceduralISMComponent->VirtualTextureCullMips;
+	TranslucencySortPriority = InProceduralISMComponent->TranslucencySortPriority;
+	CustomDepthStencilValue = InProceduralISMComponent->CustomDepthStencilValue;
+	bCastShadow = InProceduralISMComponent->CastShadow;
+	bEmissiveLightSource = InProceduralISMComponent->bEmissiveLightSource;
+	bCastDynamicShadow = InProceduralISMComponent->bCastDynamicShadow;
+	bCastStaticShadow = InProceduralISMComponent->bCastStaticShadow;
+	bCastContactShadow = InProceduralISMComponent->bCastContactShadow;
+	bCastShadowAsTwoSided = InProceduralISMComponent->bCastShadowAsTwoSided;
+	bCastHiddenShadow = InProceduralISMComponent->bCastHiddenShadow;
+	bReceivesDecals = InProceduralISMComponent->bReceivesDecals;
+	bUseAsOccluder = InProceduralISMComponent->bUseAsOccluder;
+	bRenderCustomDepth = InProceduralISMComponent->bRenderCustomDepth;
+	bEvaluateWorldPositionOffset = InProceduralISMComponent->bEvaluateWorldPositionOffset;
+	bReverseCulling = InProceduralISMComponent->bReverseCulling;
+	WorldPositionOffsetDisableDistance = InProceduralISMComponent->WorldPositionOffsetDisableDistance;
+	ShadowCacheInvalidationBehavior = InProceduralISMComponent->ShadowCacheInvalidationBehavior;
+	DetailMode = InProceduralISMComponent->DetailMode;
 }
 
 uint32 FPCGProceduralISMComponentDescriptor::ComputeHash() const
@@ -185,7 +252,30 @@ bool FPCGProceduralISMComponentDescriptor::operator==(const FPCGProceduralISMCom
 		&& NumCustomFloats == Other.NumCustomFloats
 		&& LocalBounds == Other.LocalBounds
 		&& InstanceStartCullDistance == Other.InstanceStartCullDistance
-		&& InstanceEndCullDistance == Other.InstanceEndCullDistance;
+		&& InstanceEndCullDistance == Other.InstanceEndCullDistance
+		&& ComponentTags == Other.ComponentTags
+		&& Mobility == Other.Mobility
+		&& VirtualTextureRenderPassType == Other.VirtualTextureRenderPassType
+		&& GetLightingChannelMaskForStruct(LightingChannels) == GetLightingChannelMaskForStruct(Other.LightingChannels)
+		&& CustomDepthStencilWriteMask == Other.CustomDepthStencilWriteMask
+		&& VirtualTextureCullMips == Other.VirtualTextureCullMips
+		&& TranslucencySortPriority == Other.TranslucencySortPriority
+		&& CustomDepthStencilValue == Other.CustomDepthStencilValue
+		&& bCastShadow == Other.bCastShadow
+		&& bEmissiveLightSource == Other.bEmissiveLightSource
+		&& bCastDynamicShadow == Other.bCastDynamicShadow
+		&& bCastStaticShadow == Other.bCastStaticShadow
+		&& bCastContactShadow == Other.bCastContactShadow
+		&& bCastShadowAsTwoSided == Other.bCastShadowAsTwoSided
+		&& bCastHiddenShadow == Other.bCastHiddenShadow
+		&& bReceivesDecals == Other.bReceivesDecals
+		&& bUseAsOccluder == Other.bUseAsOccluder
+		&& bRenderCustomDepth == Other.bRenderCustomDepth
+		&& bEvaluateWorldPositionOffset == Other.bEvaluateWorldPositionOffset
+		&& bReverseCulling == Other.bReverseCulling
+		&& WorldPositionOffsetDisableDistance == Other.WorldPositionOffsetDisableDistance
+		&& ShadowCacheInvalidationBehavior == Other.ShadowCacheInvalidationBehavior
+		&& DetailMode == Other.DetailMode;
 }
 
 UPCGProceduralISMComponent::UPCGProceduralISMComponent(const FObjectInitializer& ObjectInitializer)
diff --git a/Engine/Plugins/PCG/Source/PCG/Private/Components/PCGProceduralISMComponent.h b/Engine/Plugins/PCG/Source/PCG/Private/Components/PCGProceduralISMComponent.h
--- a/Engine/Plugins/PCG/Source/PCG/Private/Components/PCGProceduralISMComponent.h
+++ b/Engine/Plugins/PCG/Source/PCG/Private/Components/PCGProceduralISMComponent.h
@@ -192,22 +192,88 @@ public:
 	TArray<TObjectPtr<URuntimeVirtualTexture>> RuntimeVirtualTextures;
 
 	UPROPERTY()
-	int32 NumInstances = 0;
+	int32 NumInstances;
 
 	UPROPERTY()
-	int32 NumCustomFloats = 0;
+	int32 NumCustomFloats;
 
 	UPROPERTY()
 	FBox LocalBounds;
 
 	UPROPERTY()
-	int32 InstanceStartCullDistance = 0;
+	int32 InstanceStartCullDistance;
 
 	UPROPERTY()
-	int32 InstanceEndCullDistance = 0;
+	int32 InstanceEndCullDistance;
 
 	UPROPERTY()
 	TArray<FName> ComponentTags;
+
+	UPROPERTY()
+	TEnumAsByte<EComponentMobility::Type> Mobility;
+		
+	UPROPERTY()
+	ERuntimeVirtualTextureMainPassType VirtualTextureRenderPassType;
+	
+	UPROPERTY()
+	FLightingChannels LightingChannels;
+
+	UPROPERTY()
+	ERendererStencilMask CustomDepthStencilWriteMask;
+
+	UPROPERTY()
+	int32 VirtualTextureCullMips;
+
+	UPROPERTY()
+	int32 TranslucencySortPriority;
+
+	UPROPERTY()
+	int32 CustomDepthStencilValue;
+
+	UPROPERTY()
+	uint8 bCastShadow : 1;
+
+	UPROPERTY()
+	uint8 bEmissiveLightSource : 1;
+		
+	UPROPERTY()
+	uint8 bCastDynamicShadow : 1;
+
+	UPROPERTY()
+	uint8 bCastStaticShadow : 1;
+
+	UPROPERTY()
+	uint8 bCastContactShadow : 1;
+
+	UPROPERTY()
+	uint8 bCastShadowAsTwoSided : 1;
+
+	UPROPERTY()
+	uint8 bCastHiddenShadow : 1;
+
+	UPROPERTY()
+	uint8 bReceivesDecals : 1;
+
+	UPROPERTY()
+	uint8 bUseAsOccluder : 1;
+
+	UPROPERTY()
+	uint8 bRenderCustomDepth : 1;
+
+	UPROPERTY()
+	uint8 bEvaluateWorldPositionOffset : 1;
+
+	UPROPERTY()
+	uint8 bReverseCulling : 1;
+
+	UPROPERTY()
+	int32 WorldPositionOffsetDisableDistance;
+
+	UPROPERTY()
+	EShadowCacheInvalidationBehavior ShadowCacheInvalidationBehavior;
+
+	UPROPERTY()
+	TEnumAsByte<enum EDetailMode> DetailMode;
 };
 
 struct FPCGProceduralISMCBuilderParameters
