[PCG] Fix attribute function regex regression
* Remove support for linebreaks in attribute functions, since they introduce some complex edge cases that we can't support right now

#rnx

diff --git a/Engine/Plugins/PCG/Source/PCG/Private/Compute/Elements/PCGCustomHLSL.cpp b/Engine/Plugins/PCG/Source/PCG/Private/Compute/Elements/PCGCustomHLSL.cpp
--- a/Engine/Plugins/PCG/Source/PCG/Private/Compute/Elements/PCGCustomHLSL.cpp
+++ b/Engine/Plugins/PCG/Source/PCG/Private/Compute/Elements/PCGCustomHLSL.cpp
@@ -43,7 +43,7 @@ namespace PCGHLSLElement
 	constexpr int AttributeNameCaptureGroup = 4;
 
 	/** Regex pattern used to detect and parse attribute function usage in kernels. */
-	constexpr TCHAR AttributeFunctionPattern[] = { TEXT("([a-zA-Z0-9]+)_(Get|Set)(.*?)[\\s]*?\\([^,]*?,[^,]*?,[^,]*?'([a-zA-Z0-9 -_\\/]*?)'") };
+	constexpr TCHAR AttributeFunctionPattern[] = { TEXT("([a-zA-Z0-9]+)_(Get|Set)(.*?)[\\s]*?\\(.*?,.*?,.*?'([a-zA-Z0-9 -_\\/]*?)'") };
 
 	constexpr TCHAR AttributeFunctionGet[] = { TEXT("Get") };
 	constexpr TCHAR AttributeFunctionSet[] = { TEXT("Set") };
