Fixed an issue where preview components could get dirtied when their graph was regen.

#rnx

diff --git a/Engine/Plugins/PCG/Source/PCG/Private/PCGComponent.cpp b/Engine/Plugins/PCG/Source/PCG/Private/PCGComponent.cpp
--- a/Engine/Plugins/PCG/Source/PCG/Private/PCGComponent.cpp
+++ b/Engine/Plugins/PCG/Source/PCG/Private/PCGComponent.cpp
@@ -608,6 +608,19 @@ void UPCGComponent::PostProcessGraph(const FBox& InNewBounds, bool bInGenerated,
 							continue;
 						}
 					}
+					// TODO: instead of doing this (moving to transient then moving back), it might be better to flow down a 'bMarkDirty' in the Flatten call.
+					else if(IsInPreviewMode())
+					{
+						// As a safety procedure, we're going to move all the data back to the transient package prior to doing the flatten,
+						// since it will dirty the source package otherwise
+						OutputTaggedData.Data->VisitDataNetwork([this](const UPCGData* InData)
+						{
+							if (InData)
+							{
+								const_cast<UPCGData*>(InData)->Rename(nullptr, GetTransientPackage(), IsInPreviewMode() ? REN_DoNotDirty : REN_None);
+							}
+						});
+					}
 
 					// Flatten data
 					OutputTaggedData.Data->VisitDataNetwork([](const UPCGData* InData)
