Added some flexibility in how we retrieve depth

#rnx

diff --git a/Engine/Plugins/PCG/Source/PCG/Private/Elements/PCGGetSubgraphDepth.cpp b/Engine/Plugins/PCG/Source/PCG/Private/Elements/PCGGetSubgraphDepth.cpp
--- a/Engine/Plugins/PCG/Source/PCG/Private/Elements/PCGGetSubgraphDepth.cpp
+++ b/Engine/Plugins/PCG/Source/PCG/Private/Elements/PCGGetSubgraphDepth.cpp
@@ -41,7 +41,14 @@ FString UPCGGetSubgraphDepthSettings::GetAdditionalTitleInformation() const
 {
 	if (const UEnum* EnumPtr = StaticEnum<EPCGSubgraphDepthMode>())
 	{
-		return FText::Format(LOCTEXT("AdditionalTitle", "Get {0}"), EnumPtr->GetDisplayNameTextByValue(static_cast<int64>(Mode))).ToString();
+		if (Mode != EPCGSubgraphDepthMode::RecursiveDepth || DistanceRelativeToUpstreamGraph == 0)
+		{
+			return FText::Format(LOCTEXT("AdditionalTitle", "Get {0}"), EnumPtr->GetDisplayNameTextByValue(static_cast<int64>(Mode))).ToString();
+		}
+		else
+		{
+			return FText::Format(LOCTEXT("RelativeDistanceTitle", "Get Recursive Depth from {0}-Upstream Graph."), FText::AsNumber(DistanceRelativeToUpstreamGraph)).ToString();
+		}
 	}
 	else
 	{
@@ -64,18 +71,7 @@ bool FPCGGetSubgraphDepthElement::ExecuteInternal(FPCGContext* Context) const
 	check(Settings);
 
 	bool bRecursive = Settings->Mode == EPCGSubgraphDepthMode::RecursiveDepth;
-	const UPCGGraph* ThisGraph = nullptr;
-
-	if (bRecursive)
-	{
-		ThisGraph = Context->Node ? Context->Node->GetGraph() : nullptr;
-
-		if (!ThisGraph)
-		{
-			PCGE_LOG(Error, GraphAndLog, LOCTEXT("NodeIsNotInAGraph", "Unable to retrieve originating graph, so impossible to find recursive depth."));
-			return true;
-		}
-	}
+	const UPCGGraph* TargetGraph = nullptr;
 
 	if (!ensure(Context->Stack))
 	{
@@ -85,19 +81,47 @@ bool FPCGGetSubgraphDepthElement::ExecuteInternal(FPCGContext* Context) const
 
 	const FPCGStack* Stack = Context->Stack;
 	const TArray<FPCGStackFrame>& StackFrames = Stack->GetStackFrames();
+
+	// Find target graph first
+	int SkipGraphCount = 0;
 	int Depth = 0;
 
-	for (const FPCGStackFrame& StackFrame : StackFrames)
+	for (int FrameIndex = StackFrames.Num() - 1; FrameIndex >= 0; --FrameIndex)
 	{
+		const FPCGStackFrame& StackFrame = StackFrames[FrameIndex];
 		if (const UPCGGraph* Graph = Cast<UPCGGraph>(StackFrame.Object.Get()))
 		{
-			if (!bRecursive || Graph == ThisGraph)
+			if (!bRecursive)
 			{
 				++Depth;
 			}
+			else
+			{
+				if (!TargetGraph)
+				{
+					if (SkipGraphCount == Settings->DistanceRelativeToUpstreamGraph)
+					{
+						TargetGraph = Graph;
+						Depth = 1;
+					}
+					else
+					{
+						++SkipGraphCount;
+					}
+				}
+				else if (TargetGraph == Graph)
+				{
+					++Depth;
+				}
+			}
 		}
 	}
 
+	if (Depth == 0 && !Settings->bQuietInvalidDepthQueries)
+	{
+		PCGE_LOG(Warning, GraphAndLog, LOCTEXT("InvalidDepthQuery", "Get depth query failed because of the relative distance to the upstream graph."));
+	}
+
 	// Finally, since we count everything, including the root and self, we need to remove one.
 	Depth = FMath::Max(0, Depth - 1);
 
diff --git a/Engine/Plugins/PCG/Source/PCG/Public/Elements/PCGGetSubgraphDepth.h b/Engine/Plugins/PCG/Source/PCG/Public/Elements/PCGGetSubgraphDepth.h
--- a/Engine/Plugins/PCG/Source/PCG/Public/Elements/PCGGetSubgraphDepth.h
+++ b/Engine/Plugins/PCG/Source/PCG/Public/Elements/PCGGetSubgraphDepth.h
@@ -43,6 +43,13 @@ protected:
 public:
 	UPROPERTY(BlueprintReadWrite, EditAnywhere, Category = Settings)
 	EPCGSubgraphDepthMode Mode = EPCGSubgraphDepthMode::Depth;
+
+	// In the case of recursive depth, it is possible to target the current graph (0), the parent graph (1) or other graphs upstream (2+).
+	UPROPERTY(BlueprintReadWrite, EditAnywhere, Category = Settings, meta = (ClampMin = "0", EditCondition = "Mode == EPCGSubgraphDepthMode::RecursiveDepth", EditConditionHides, PCG_Overridable))
+	int DistanceRelativeToUpstreamGraph = 0;
+
+	UPROPERTY(BlueprintReadWrite, EditAnywhere, Category = Settings, meta = (EditCondition = "Mode == EPCGSubgraphDepthMode::RecursiveDepth", EditConditionHides))
+	bool bQuietInvalidDepthQueries = false;
 };
 
 class FPCGGetSubgraphDepthElement : public IPCGElement
