[PCG] DataLayer
- New nodes: PartitionByActorDataLayers, GetActorDataLayers
- CreateTargetActor/SpawnActor now support different DataLayer sources (Self or DataLayerReferences as inputs)

diff --git a/Engine/Plugins/PCG/Source/PCG/Private/Elements/PCGCreateTargetActor.cpp b/Engine/Plugins/PCG/Source/PCG/Private/Elements/PCGCreateTargetActor.cpp
--- a/Engine/Plugins/PCG/Source/PCG/Private/Elements/PCGCreateTargetActor.cpp
+++ b/Engine/Plugins/PCG/Source/PCG/Private/Elements/PCGCreateTargetActor.cpp
@@ -11,6 +11,7 @@
 #include "PCGParamData.h"
 #include "Data/PCGPointData.h"
 #include "Helpers/PCGActorHelpers.h"
+#include "Helpers/PCGDataLayerHelpers.h"
 #include "Helpers/PCGHelpers.h"
 
 #include "Engine/Blueprint.h"
@@ -111,10 +112,23 @@ FPCGElementPtr UPCGCreateTargetActor::CreateElement() const
 	return MakeShared<FPCGCreateTargetActorElement>();
 }
 
+#if WITH_EDITOR
+EPCGChangeType UPCGCreateTargetActor::GetChangeTypeForProperty(const FName& InPropertyName) const
+{
+	EPCGChangeType ChangeType = Super::GetChangeTypeForProperty(InPropertyName);
+
+	ChangeType |= DataLayerSettings.GetChangeTypeForProperty(InPropertyName);
+
+	return ChangeType;
+}
+#endif
+
 TArray<FPCGPinProperties> UPCGCreateTargetActor::InputPinProperties() const
 {
 	TArray<FPCGPinProperties> PinProperties;
 	PinProperties.Add(PCGObjectPropertyOverrideHelpers::CreateObjectPropertiesOverridePin(PCGCreateTargetActorConstants::ActorPropertyOverridesLabel, PCGCreateTargetActorConstants::ActorPropertyOverridesTooltip));
+	PinProperties.Append(DataLayerSettings.InputPinProperties());
+
 	return PinProperties;
 }
 
@@ -370,7 +384,8 @@ bool FPCGCreateTargetActorElement::ExecuteInternal(FPCGContext* Context) const
 	UPCGActorHelpers::FSpawnDefaultActorParams SpawnDefaultActorParams(TargetActor->GetWorld(), Settings->TemplateActorClass, Transform, SpawnParams);
 	
 #if WITH_EDITOR
-	SpawnDefaultActorParams.DataLayerInstances = TargetActor->GetDataLayerInstances();
+	int32 DataLayerCrc;
+	SpawnDefaultActorParams.DataLayerInstances = PCGDataLayerHelpers::GetDataLayerInstancesAndCrc(Context, Settings->DataLayerSettings, TargetActor, DataLayerCrc);
 #endif
 
 	UPCGManagedActors* ReusableResource = nullptr;
@@ -385,6 +400,13 @@ bool FPCGCreateTargetActorElement::ExecuteInternal(FPCGContext* Context) const
 			// but for future reference if we want to have better reusing and stable generation across, we need to fix this.
 			FPCGDataCollection EmptyCollection;
 			GetDependenciesCrc(EmptyCollection, Settings, Context->SourceComponent.Get(), Context->DependenciesCrc);
+
+#if WITH_EDITOR
+			if (DataLayerCrc != 0)
+			{
+				Context->DependenciesCrc.Combine(DataLayerCrc);
+			}
+#endif
 		}
 
 		if (Context->DependenciesCrc.IsValid())
diff --git a/Engine/Plugins/PCG/Source/PCG/Private/Elements/PCGGetActorDataLayers.cpp b/Engine/Plugins/PCG/Source/PCG/Private/Elements/PCGGetActorDataLayers.cpp
new file mode 100644
--- /dev/null
+++ b/Engine/Plugins/PCG/Source/PCG/Private/Elements/PCGGetActorDataLayers.cpp
@@ -0,0 +1,86 @@
+// Copyright Epic Games, Inc. All Rights Reserved.
+
+#include "Elements/PCGGetActorDataLayers.h"
+
+#include "PCGContext.h"
+#include "PCGData.h"
+#include "PCGParamData.h"
+#include "PCGPin.h"
+#include "Data/PCGPointData.h"
+#include "Helpers/PCGDataLayerHelpers.h"
+#include "Utils/PCGLogErrors.h"
+
+#define LOCTEXT_NAMESPACE "PCGGetActorDataLayers"
+
+UPCGGetActorDataLayersSettings::UPCGGetActorDataLayersSettings()
+{
+	ActorReferenceAttribute.SetAttributeName(PCGPointDataConstants::ActorReferenceAttribute);
+	DataLayerReferenceAttribute.SetAttributeName(PCGDataLayerHelpers::Constants::DataLayerReferenceAttribute);
+}
+
+TArray<FPCGPinProperties> UPCGGetActorDataLayersSettings::InputPinProperties() const
+{
+	TArray<FPCGPinProperties> PinProperties;
+	FPCGPinProperties& InputPin = PinProperties.Emplace_GetRef(PCGPinConstants::DefaultInputLabel, EPCGDataType::PointOrParam);
+	InputPin.SetRequiredPin();
+	return PinProperties;
+}
+
+TArray<FPCGPinProperties> UPCGGetActorDataLayersSettings::OutputPinProperties() const
+{
+	TArray<FPCGPinProperties> PinProperties;
+	PinProperties.Emplace(PCGPinConstants::DefaultOutputLabel, EPCGDataType::Param);
+	return PinProperties;
+}
+
+FString UPCGGetActorDataLayersSettings::GetAdditionalTitleInformation() const
+{
+	return FString::Printf(TEXT("%s -> %s"), *ActorReferenceAttribute.GetDisplayText().ToString(), *DataLayerReferenceAttribute.GetDisplayText().ToString());
+}
+
+FPCGElementPtr UPCGGetActorDataLayersSettings::CreateElement() const
+{
+	return MakeShared<FPCGGetActorDataLayersElement>();
+}
+
+bool FPCGGetActorDataLayersElement::ExecuteInternal(FPCGContext* Context) const
+{
+	const UPCGGetActorDataLayersSettings* Settings = Context->GetInputSettings<UPCGGetActorDataLayersSettings>();
+	check(Settings);
+
+	TArray<FPCGTaggedData> Inputs = Context->InputData.GetInputsByPin(PCGPinConstants::DefaultInputLabel);
+	TArray<FPCGTaggedData>& Outputs = Context->OutputData.TaggedData;
+
+	for (const FPCGTaggedData& Input : Inputs)
+	{
+		if (!Input.Data)
+		{
+			continue;
+		}
+
+		// Add PCGParam Output with DataLayers
+		FPCGTaggedData& DataLayerSetOutput = Context->OutputData.TaggedData.Emplace_GetRef();
+		UPCGParamData* OutputParamData = FPCGContext::NewObject_AnyThread<UPCGParamData>(Context);
+		DataLayerSetOutput.Pin = PCGPinConstants::DefaultOutputLabel;
+		DataLayerSetOutput.Data = OutputParamData;
+
+#if WITH_EDITOR
+		FPCGMetadataAttribute<FSoftObjectPath>* DataLayersAttribute = OutputParamData->MutableMetadata()->CreateAttribute(Settings->DataLayerReferenceAttribute.GetName(), FSoftObjectPath(), /*bAllowsInterpolation=*/false, /*bOverrideParent=*/false);
+
+		TArray<FSoftObjectPath> DataLayerAssets = PCGDataLayerHelpers::GetDataLayerAssetsFromActorReferences(Context, Input.Data, Settings->ActorReferenceAttribute);
+		for (const FSoftObjectPath& ActorDataLayer : DataLayerAssets)
+		{
+			PCGMetadataEntryKey Entry = OutputParamData->MutableMetadata()->AddEntry();
+			DataLayersAttribute->SetValue(Entry, ActorDataLayer);
+		}
+#endif
+	}
+
+#if !WITH_EDITOR
+	PCGLog::LogErrorOnGraph(LOCTEXT("GetActorDataLayersUnsupported", "Get Actor Data Layers is unsupported at runtime"), Context);
+#endif
+
+	return true;
+}
+
+#undef LOCTEXT_NAMESPACE
diff --git a/Engine/Plugins/PCG/Source/PCG/Private/Elements/PCGPartitionByActorDataLayers.cpp b/Engine/Plugins/PCG/Source/PCG/Private/Elements/PCGPartitionByActorDataLayers.cpp
new file mode 100644
--- /dev/null
+++ b/Engine/Plugins/PCG/Source/PCG/Private/Elements/PCGPartitionByActorDataLayers.cpp
@@ -0,0 +1,240 @@
+// Copyright Epic Games, Inc. All Rights Reserved.
+
+#include "Elements/PCGPartitionByActorDataLayers.h"
+
+#include "PCGContext.h"
+#include "PCGParamData.h"
+#include "PCGPin.h"
+#include "Data/PCGPointData.h"
+#include "Helpers/PCGDataLayerHelpers.h"
+#include "Utils/PCGLogErrors.h"
+
+#if WITH_EDITOR
+#include "Helpers/PCGHelpers.h"
+#include "Metadata/PCGMetadataAccessor.h"
+#include "Metadata/PCGMetadataPartitionCommon.h"
+#include "Metadata/Accessors/IPCGAttributeAccessor.h"
+#include "Metadata/Accessors/PCGAttributeAccessorHelpers.h"
+
+#include "Algo/RemoveIf.h"
+#include "GameFramework/Actor.h"
+#include "Serialization/ArchiveCrc32.h"
+#include "WorldPartition/DataLayer/DataLayerAsset.h"
+#endif // WITH_EDITOR
+
+#define LOCTEXT_NAMESPACE "PCGPartitionByActorDataLayers"
+
+namespace PCGPartitionByActorDataLayers
+{
+	namespace Constants
+	{
+		const FName DataLayerPartitionsLabel = TEXT("DataLayerPartitions");
+	}
+
+#if WITH_EDITOR
+	TArray<const UDataLayerAsset*> GetDataLayersFromActor(AActor* InActor, const TArray<TSoftObjectPtr<UDataLayerAsset>>& InIncludedDataLayerAssets, const TArray<TSoftObjectPtr<UDataLayerAsset>>& InExcludedDataLayerAssets)
+	{
+		TArray<const UDataLayerAsset*> DataLayerAssets = PCGDataLayerHelpers::GetDatalayerAssetsForActor(InActor);
+
+		int32 RemoveIndex = Algo::RemoveIf(DataLayerAssets, [&InIncludedDataLayerAssets, &InExcludedDataLayerAssets](const UDataLayerAsset* InDataLayerAsset)
+		{
+			return (!InIncludedDataLayerAssets.IsEmpty() && !InIncludedDataLayerAssets.Contains(InDataLayerAsset))
+				|| InExcludedDataLayerAssets.Contains(InDataLayerAsset);
+		});
+
+		DataLayerAssets.SetNum(RemoveIndex);
+		DataLayerAssets.Sort([](const UDataLayerAsset& A, const UDataLayerAsset& B)
+		{
+			return A.GetPathName() < B.GetPathName();
+		});
+
+		return DataLayerAssets;
+	}
+
+	int32 GetDataLayersCrc(const TArray<const UDataLayerAsset*>& InDataLayerAssets)
+	{
+		FArchiveCrc32 Ar;
+		for (const UDataLayerAsset* DataLayerAsset : InDataLayerAssets)
+		{
+			UDataLayerAsset* NonConstDataLayerAsset = const_cast<UDataLayerAsset*>(DataLayerAsset);
+			Ar << NonConstDataLayerAsset;
+		}
+		return Ar.GetCrc();
+	}
+#endif
+}
+
+UPCGPartitionByActorDataLayersSettings::UPCGPartitionByActorDataLayersSettings()
+{
+	ActorReferenceAttribute.SetAttributeName(PCGPointDataConstants::ActorReferenceAttribute);
+	IncludedDataLayers.Attribute.SetAttributeName(PCGDataLayerHelpers::Constants::IncludedDataLayersAttribute);
+	ExcludedDataLayers.Attribute.SetAttributeName(PCGDataLayerHelpers::Constants::ExcludedDataLayersAttribute);
+	DataLayerReferenceAttribute.SetAttributeName(PCGDataLayerHelpers::Constants::DataLayerReferenceAttribute);
+}
+
+TArray<FPCGPinProperties> UPCGPartitionByActorDataLayersSettings::InputPinProperties() const
+{
+	TArray<FPCGPinProperties> PinProperties;
+	FPCGPinProperties& InputPin = PinProperties.Emplace_GetRef(PCGPinConstants::DefaultInputLabel, EPCGDataType::Point);
+	InputPin.SetRequiredPin();
+
+	if (IncludedDataLayers.bAsInput)
+	{
+		FPCGPinProperties& IncludedDataLayersPin = PinProperties.Emplace_GetRef(PCGDataLayerHelpers::Constants::IncludedDataLayersAttribute, EPCGDataType::PointOrParam);
+		IncludedDataLayersPin.SetRequiredPin();
+	}
+
+	if (ExcludedDataLayers.bAsInput)
+	{
+		FPCGPinProperties& ExcludedDataLayersPin = PinProperties.Emplace_GetRef(PCGDataLayerHelpers::Constants::ExcludedDataLayersAttribute, EPCGDataType::Param);
+		ExcludedDataLayersPin.SetRequiredPin();
+	}
+
+	return PinProperties;
+}
+
+TArray<FPCGPinProperties> UPCGPartitionByActorDataLayersSettings::OutputPinProperties() const
+{
+	TArray<FPCGPinProperties> PinProperties;
+	PinProperties.Emplace(PCGPinConstants::DefaultOutputLabel, EPCGDataType::Point);
+	PinProperties.Emplace(PCGPartitionByActorDataLayers::Constants::DataLayerPartitionsLabel, EPCGDataType::Param);
+
+	return PinProperties;
+}
+
+FPCGElementPtr UPCGPartitionByActorDataLayersSettings::CreateElement() const
+{
+	return MakeShared<FPCGPartitionByActorDataLayersElement>();
+}
+
+#if WITH_EDITOR
+EPCGChangeType UPCGPartitionByActorDataLayersSettings::GetChangeTypeForProperty(const FName& InPropertyName) const
+{
+	EPCGChangeType ChangeType = Super::GetChangeTypeForProperty(InPropertyName);
+
+	if (InPropertyName == GET_MEMBER_NAME_CHECKED(FPCGDataLayerReferenceSelector, bAsInput))
+	{
+		ChangeType |= EPCGChangeType::Structural;
+	}
+
+	return ChangeType;
+}
+#endif
+
+bool FPCGPartitionByActorDataLayersElement::ExecuteInternal(FPCGContext* Context) const
+{
+	const UPCGPartitionByActorDataLayersSettings* Settings = Context->GetInputSettings<UPCGPartitionByActorDataLayersSettings>();
+	check(Settings);
+	
+	TArray<FPCGTaggedData> Inputs = Context->InputData.GetInputsByPin(PCGPinConstants::DefaultInputLabel);
+	TArray<FPCGTaggedData>& Outputs = Context->OutputData.TaggedData;
+
+#if WITH_EDITOR
+	TArray<TSoftObjectPtr<UDataLayerAsset>> IncludedDataLayerAssets = PCGDataLayerHelpers::GetDataLayerAssetsFromInput(Context, PCGDataLayerHelpers::Constants::IncludedDataLayersAttribute, Settings->IncludedDataLayers);
+	TArray<TSoftObjectPtr<UDataLayerAsset>> ExcludedDataLayerAssets = PCGDataLayerHelpers::GetDataLayerAssetsFromInput(Context, PCGDataLayerHelpers::Constants::ExcludedDataLayersAttribute, Settings->ExcludedDataLayers);
+
+	TMap<int32, UPCGPointData*> DataLayerCrcToOutputData;
+
+	for (const FPCGTaggedData& Input : Inputs)
+	{
+		// @todo_pcg: support ParamData
+		const UPCGPointData* InData = Cast<UPCGPointData>(Input.Data);
+		if (!InData)
+		{
+			continue;
+		}
+
+		// @todo_pcg: partitioning by actor reference will cause performance issues if we have too many actors, we should have a way in the partition api to specify a lambda to compare values where
+		// we could compare the actor DataLayers Crcs directly
+		FPCGAttributePropertySelector PartitionAttributeSource = Settings->ActorReferenceAttribute.CopyAndFixLast(InData);
+		TArray<UPCGData*> PartitionDataArray = PCGMetadataPartitionCommon::AttributePartition(InData, PartitionAttributeSource, Context);
+
+		for (UPCGData* PartitionData : PartitionDataArray)
+		{
+			UPCGPointData* PartitionPointData = CastChecked<UPCGPointData>(PartitionData);
+
+			const TUniquePtr<const IPCGAttributeAccessor> Accessor = PCGAttributeAccessorHelpers::CreateConstAccessor(PartitionData, PartitionAttributeSource);
+			const TUniquePtr<const IPCGAttributeAccessorKeys> Keys = PCGAttributeAccessorHelpers::CreateConstKeys(PartitionData, PartitionAttributeSource);
+
+			if (!Accessor || !Keys)
+			{
+				PCGLog::Metadata::LogFailToCreateAccessorError(PartitionAttributeSource, Context);
+				continue;
+			}
+
+			if (Keys->GetNum() < 1)
+			{
+				continue;
+			}
+
+			FSoftObjectPath ActorSoftPath;
+			if (!Accessor->Get<FSoftObjectPath>(ActorSoftPath, *Keys.Get(), EPCGAttributeAccessorFlags::AllowBroadcastAndConstructible))
+			{
+				PCGLog::Metadata::LogFailToGetAttributeError(PartitionAttributeSource, Context);
+				continue;
+			}
+
+			AActor* Actor = Cast<AActor>(ActorSoftPath.ResolveObject());
+			if (!Actor)
+			{
+				PCGLog::LogErrorOnGraph(FText::Format(LOCTEXT("UnresolvedActor", "Could not resolve actor path '{0}'."), FText::FromString(ActorSoftPath.ToString())), Context);
+				continue;
+			}
+
+			// Compute DL Crc
+			TArray<const UDataLayerAsset*> ActorDataLayers = PCGPartitionByActorDataLayers::GetDataLayersFromActor(Actor, IncludedDataLayerAssets, ExcludedDataLayerAssets);
+			int32 DataLayerCrc = PCGPartitionByActorDataLayers::GetDataLayersCrc(ActorDataLayers);
+
+			UPCGPointData* PartitionOutputData = nullptr;
+			if (UPCGPointData** FoundOutputData = DataLayerCrcToOutputData.Find(DataLayerCrc))
+			{
+				PartitionOutputData = *FoundOutputData;
+			}
+			else
+			{
+				PartitionOutputData = FPCGContext::NewObject_AnyThread<UPCGPointData>(Context);
+				PartitionOutputData->InitializeFromData(PartitionPointData);
+				DataLayerCrcToOutputData.Add(DataLayerCrc, PartitionOutputData);
+
+				FPCGTaggedData& Output = Context->OutputData.TaggedData.Emplace_GetRef();
+				Output.Pin = PCGPinConstants::DefaultOutputLabel;
+				Output.Data = PartitionOutputData;
+
+				// Add PCGParam Output with DataLayers
+				FPCGTaggedData& DataLayerSetOutput = Context->OutputData.TaggedData.Emplace_GetRef();
+				UPCGParamData* ParamData = FPCGContext::NewObject_AnyThread<UPCGParamData>(Context);
+				DataLayerSetOutput.Pin = PCGPartitionByActorDataLayers::Constants::DataLayerPartitionsLabel;
+				DataLayerSetOutput.Data = ParamData;
+
+				FPCGMetadataAttribute<FSoftObjectPath>* DataLayersAttribute = ParamData->MutableMetadata()->CreateAttribute(Settings->DataLayerReferenceAttribute.GetName(), FSoftObjectPath(), /*bAllowsInterpolation=*/false, /*bOverrideParent=*/false);
+				for (const UDataLayerAsset* ActorDataLayer : ActorDataLayers)
+				{
+					PCGMetadataEntryKey Entry = ParamData->MutableMetadata()->AddEntry();
+					DataLayersAttribute->SetValue(Entry, FSoftObjectPath(ActorDataLayer));
+				}
+			}
+
+			check(PartitionOutputData);
+			TArray<FPCGPoint>& MutablePoints = PartitionOutputData->GetMutablePoints();
+			MutablePoints.Append(PartitionPointData->GetPoints());
+		}
+	}
+#else
+	for (const FPCGTaggedData& Input : Inputs)
+	{
+		Outputs.Add(Input);
+		
+		FPCGTaggedData& DataLayerSetOutput = Context->OutputData.TaggedData.Emplace_GetRef();
+		UPCGParamData* ParamData = FPCGContext::NewObject_AnyThread<UPCGParamData>(Context);
+		DataLayerSetOutput.Pin = PCGPartitionByActorDataLayers::Constants::DataLayerPartitionsLabel;
+		DataLayerSetOutput.Data = ParamData;
+	}
+
+	PCGLog::LogErrorOnGraph(LOCTEXT("PartitionByActorDataLayersUnsupported", "Partition by Actor Data Layers is unsupported at runtime"), Context);
+#endif // WITH_EDITOR
+
+	return true;
+}
+
+#undef LOCTEXT_NAMESPACE
+
diff --git a/Engine/Plugins/PCG/Source/PCG/Private/Elements/PCGSpawnActor.cpp b/Engine/Plugins/PCG/Source/PCG/Private/Elements/PCGSpawnActor.cpp
--- a/Engine/Plugins/PCG/Source/PCG/Private/Elements/PCGSpawnActor.cpp
+++ b/Engine/Plugins/PCG/Source/PCG/Private/Elements/PCGSpawnActor.cpp
@@ -286,10 +286,19 @@ EPCGChangeType UPCGSpawnActorSettings::GetChangeTypeForProperty(const FName& InP
 		ChangeType |= EPCGChangeType::Structural;
 	}
 	
+	ChangeType |= DataLayerSettings.GetChangeTypeForProperty(InPropertyName);
+
 	return ChangeType;
 }
 #endif // WITH_EDITOR
 
+TArray<FPCGPinProperties> UPCGSpawnActorSettings::InputPinProperties() const
+{
+	TArray<FPCGPinProperties> PinProperties = DefaultPointInputPinProperties();
+	PinProperties.Append(DataLayerSettings.InputPinProperties());
+	return PinProperties;
+}
+
 TObjectPtr<UPCGGraphInterface> UPCGSpawnActorNode::GetSubgraphInterface() const
 {
 	TObjectPtr<UPCGSpawnActorSettings> Settings = Cast<UPCGSpawnActorSettings>(GetSettings());
@@ -823,7 +832,8 @@ void FPCGSpawnActorElement::SpawnActors(FPCGSubgraphContext* Context, AActor* Ta
 	SpawnDefaultActorParams.bForceStaticMobility = false; // Always respect the actor's mobility
 
 #if WITH_EDITOR
-	SpawnDefaultActorParams.DataLayerInstances = TargetActor->GetDataLayerInstances();
+	int32 DataLayerCrc;
+	SpawnDefaultActorParams.DataLayerInstances = PCGDataLayerHelpers::GetDataLayerInstancesAndCrc(Context, Settings->DataLayerSettings, TargetActor, DataLayerCrc);
 #endif
 
 	const bool bForceCallGenerate = (Settings->GenerationTrigger == EPCGSpawnActorGenerationTrigger::ForceGenerate);
@@ -847,6 +857,13 @@ void FPCGSpawnActorElement::SpawnActors(FPCGSubgraphContext* Context, AActor* Ta
 
 		GetDependenciesCrc(SingleInputCollection, Settings, Context->SourceComponent.Get(), InputDependenciesCrc);
 
+#if WITH_EDITOR
+		if (DataLayerCrc != 0)
+		{
+			InputDependenciesCrc.Combine(DataLayerCrc);
+		}
+#endif
+
 		if (InputDependenciesCrc.IsValid())
 		{
 			Context->SourceComponent->ForEachManagedResource([&ReusedManagedActorsResources, &InputDependenciesCrc, &Context](UPCGManagedResource* InResource)
diff --git a/Engine/Plugins/PCG/Source/PCG/Private/Helpers/PCGDataLayerHelpers.cpp b/Engine/Plugins/PCG/Source/PCG/Private/Helpers/PCGDataLayerHelpers.cpp
new file mode 100644
--- /dev/null
+++ b/Engine/Plugins/PCG/Source/PCG/Private/Helpers/PCGDataLayerHelpers.cpp
@@ -0,0 +1,317 @@
+// Copyright Epic Games, Inc. All Rights Reserved.
+
+#include "Helpers/PCGDataLayerHelpers.h"
+
+#if WITH_EDITOR
+#include "PCGContext.h"
+#include "PCGParamData.h"
+#include "Data/PCGPointData.h"
+#include "Metadata/PCGAttributePropertySelector.h"
+#include "Metadata/Accessors/PCGAttributeAccessorHelpers.h"
+#include "Utils/PCGLogErrors.h"
+
+#include "Algo/RemoveIf.h"
+#include "Algo/Transform.h"
+#include "GameFramework/Actor.h"
+#include "LevelInstance/LevelInstanceInterface.h"
+#include "LevelInstance/LevelInstanceSubsystem.h"
+#include "Serialization/ArchiveCrc32.h"
+#include "WorldPartition/DataLayer/DataLayerAsset.h"
+#include "WorldPartition/DataLayer/DataLayerManager.h"
+#include "WorldPartition/DataLayer/DataLayerInstance.h"
+#include "WorldPartition/DataLayer/DataLayerInstanceWithAsset.h"
+#endif // WITH_EDITOR
+
+#define LOCTEXT_NAMESPACE "PCGDataLayerHelpers"
+
+namespace PCGDataLayerHelpers
+{
+#if WITH_EDITOR
+	namespace Private
+	{
+		TArray<const UDataLayerInstance*> GetDataLayerInstancesFromAssets(FPCGContext* Context, const TArray<TSoftObjectPtr<UDataLayerAsset>>& DataLayerAssets, UDataLayerManager* DataLayerManager)
+		{
+			if (!DataLayerManager)
+			{
+				return {};
+			}
+
+			TArray<const UDataLayerInstance*> DataLayerInstances;
+			for (const TSoftObjectPtr<UDataLayerAsset>& DataLayerAssetPtr : DataLayerAssets)
+			{
+				if (UDataLayerAsset* DataLayerAsset = DataLayerAssetPtr.Get())
+				{
+					if (const UDataLayerInstance* DataLayerInstance = DataLayerManager->GetDataLayerInstance(DataLayerAsset))
+					{
+						DataLayerInstances.Add(DataLayerInstance);
+					}
+					else
+					{
+						PCGLog::LogWarningOnGraph(FText::Format(LOCTEXT("NoDataLayerInstanceFound", "No DataLayerInstance using DataLayerAsset '{0}' found in World"), FText::FromString(DataLayerAssetPtr.ToString())), Context);
+					}
+				}
+				else
+				{
+					PCGLog::LogWarningOnGraph(FText::Format(LOCTEXT("FailedToResolveDataLayerAsset", "Could not resolve DataLayerAsset '{0}', this probably means your World does not have a DataLayerInstance using it"), FText::FromString(DataLayerAssetPtr.ToString())), Context);
+				}
+			}
+
+			return DataLayerInstances;
+		}
+
+		TArray<const UDataLayerInstance*> GetDataLayerInstancesFromDataLayerReferences(FPCGContext* Context, const FPCGDataLayerSettings& DataLayerSettings, UWorld* World)
+		{
+			const TArray<FPCGTaggedData> DataLayersInputs = Context->InputData.GetInputsByPin(PCGDataLayerHelpers::Constants::DataLayerReferenceAttribute);
+
+			if (DataLayersInputs.IsEmpty())
+			{
+				return {};
+			}
+
+			TSet<FSoftObjectPath> DataLayerSoftObjectPaths;
+			for (const FPCGTaggedData& DataLayersInput : DataLayersInputs)
+			{
+				if (!DataLayersInput.Data)
+				{
+					PCGLog::LogWarningOnGraph(LOCTEXT("DataLayersWrongType", "Data layers input is not of type attribute set."), Context);
+					continue;
+				}
+
+				TArray<FSoftObjectPath> SoftObjectPaths;
+				if (PCGAttributeAccessorHelpers::ExtractAllValues(DataLayersInput.Data, DataLayerSettings.DataLayerReferenceAttribute, SoftObjectPaths, Context, EPCGAttributeAccessorFlags::AllowBroadcastAndConstructible, /*bQuiet=*/true))
+				{
+					DataLayerSoftObjectPaths.Append(SoftObjectPaths);
+				}
+			}
+
+			TArray<TSoftObjectPtr<UDataLayerAsset>> DataLayerAssetPtrs;
+			Algo::Transform(DataLayerSoftObjectPaths, DataLayerAssetPtrs, [](const FSoftObjectPath& DataLayerPath) { return TSoftObjectPtr<UDataLayerAsset>(DataLayerPath); });
+
+			return GetDataLayerInstancesFromAssets(Context, DataLayerAssetPtrs, UDataLayerManager::GetDataLayerManager(World));
+		}
+
+		void FilterDataLayerInstances(FPCGContext* Context, const FPCGDataLayerSettings& DataLayerSettings, TArray<const UDataLayerInstance*>& DataLayerInstances)
+		{
+			TArray<TSoftObjectPtr<UDataLayerAsset>> ExcludedDataLayerAssets = GetDataLayerAssetsFromInput(Context, PCGDataLayerHelpers::Constants::ExcludedDataLayersAttribute, DataLayerSettings.ExcludedDataLayers);
+			TArray<TSoftObjectPtr<UDataLayerAsset>> IncludedDataLayerAssets = GetDataLayerAssetsFromInput(Context, PCGDataLayerHelpers::Constants::IncludedDataLayersAttribute, DataLayerSettings.IncludedDataLayers);
+
+			int32 RemoveIndex = Algo::RemoveIf(DataLayerInstances, [&ExcludedDataLayerAssets, &IncludedDataLayerAssets](const UDataLayerInstance* DataLayerInstance)
+			{
+				if (const UDataLayerInstanceWithAsset* DataLayerInstanceWithAsset = Cast<UDataLayerInstanceWithAsset>(DataLayerInstance))
+				{
+					return (!IncludedDataLayerAssets.IsEmpty() && !IncludedDataLayerAssets.Contains(DataLayerInstanceWithAsset->GetAsset()))
+						|| ExcludedDataLayerAssets.Contains(DataLayerInstanceWithAsset->GetAsset());
+
+				}
+
+				return false;
+			});
+			DataLayerInstances.SetNum(RemoveIndex);
+		}
+
+		void AddDataLayerInstances(FPCGContext* Context, const FPCGDataLayerSettings& DataLayerSettings, UWorld* World, TArray<const UDataLayerInstance*>& DataLayerInstances)
+		{
+			TArray<TSoftObjectPtr<UDataLayerAsset>> AddDataLayerAssets = GetDataLayerAssetsFromInput(Context, PCGDataLayerHelpers::Constants::AddDataLayersAttribute, DataLayerSettings.AddDataLayers);
+			TArray<const UDataLayerInstance*> AddDataLayerInstances = GetDataLayerInstancesFromAssets(Context, AddDataLayerAssets, UDataLayerManager::GetDataLayerManager(World));
+
+			for (const UDataLayerInstance* AddDataLayerInstance : AddDataLayerInstances)
+			{
+				DataLayerInstances.AddUnique(AddDataLayerInstance);
+			}
+		}
+	} // Private
+
+	TArray<TSoftObjectPtr<UDataLayerAsset>> GetDataLayerAssetsFromInput(FPCGContext* Context, FName InputPinName, const FPCGAttributePropertyInputSelector& InputSelector)
+	{
+		const TArray<FPCGTaggedData> DataLayerAssetsInputs = Context->InputData.GetInputsByPin(InputPinName);
+
+		if (DataLayerAssetsInputs.IsEmpty())
+		{
+			PCGLog::LogWarningOnGraph(NSLOCTEXT("PCGDataLayerHelpers", "NoDataLayerAssets", "No data was found on the data layer assets pin."), Context);
+			return {};
+		}
+
+		TSet<FSoftObjectPath> DataLayerSoftObjectPaths;
+		for (const FPCGTaggedData& DataLayerAssetsInput : DataLayerAssetsInputs)
+		{
+			if (!DataLayerAssetsInput.Data)
+			{
+				PCGLog::LogWarningOnGraph(LOCTEXT("DataLayersWrongType", "Data layer assets input is not of type attribute set."), Context);
+				continue;
+			}
+
+			TArray<FSoftObjectPath> SoftObjectPaths;
+			if (!PCGAttributeAccessorHelpers::ExtractAllValues(DataLayerAssetsInput.Data, InputSelector, SoftObjectPaths, Context, EPCGAttributeAccessorFlags::AllowBroadcastAndConstructible, /*bQuiet=*/true))
+			{
+				continue;
+			}
+
+			DataLayerSoftObjectPaths.Append(SoftObjectPaths);
+		}
+				
+		TArray<TSoftObjectPtr<UDataLayerAsset>> ReturnValue;
+		Algo::Transform(DataLayerSoftObjectPaths, ReturnValue, [](const FSoftObjectPath& AssetPath) { return TSoftObjectPtr<UDataLayerAsset>(AssetPath); });
+
+		return ReturnValue;
+	}
+
+	TArray<TSoftObjectPtr<UDataLayerAsset>> GetDataLayerAssetsFromInput(FPCGContext* Context, FName InputPinName, const FPCGDataLayerReferenceSelector& DataLayerSelector)
+	{
+		if (DataLayerSelector.bAsInput)
+		{
+			return GetDataLayerAssetsFromInput(Context, InputPinName, DataLayerSelector.Attribute);
+		}
+
+		return DataLayerSelector.DataLayers;
+	}
+		
+	TArray<FSoftObjectPath> GetDataLayerAssetsFromActorReferences(FPCGContext* Context, const UPCGData* Data, const FPCGAttributePropertyInputSelector& ActorReferenceAttribute)
+	{			
+		check(Data);
+			
+		TArray<FSoftObjectPath> ActorReferences;
+		if (!PCGAttributeAccessorHelpers::ExtractAllValues<FSoftObjectPath>(Data, ActorReferenceAttribute, ActorReferences, Context, EPCGAttributeAccessorFlags::AllowBroadcastAndConstructible, /*bQuiet=*/true))
+		{
+			return {};
+		}
+
+		TSet<const UDataLayerAsset*> DataLayerAssets;
+		for (const FSoftObjectPath& ActorSoftPath : ActorReferences)
+		{
+			AActor* Actor = Cast<AActor>(ActorSoftPath.ResolveObject());
+			if (!Actor)
+			{
+				PCGLog::LogErrorOnGraph(FText::Format(LOCTEXT("UnresolvedActor", "Could not resolve actor path '{0}'."), FText::FromString(ActorSoftPath.ToString())), Context);
+				continue;
+			}
+
+			DataLayerAssets.Append(GetDatalayerAssetsForActor(Actor));
+		}
+		
+		TArray<FSoftObjectPath> DataLayerSoftObjectPaths;
+		Algo::Transform(DataLayerAssets, DataLayerSoftObjectPaths, [](const UDataLayerAsset* DataLayerAsset) { return FSoftObjectPath(const_cast<UDataLayerAsset*>(DataLayerAsset)); });
+		
+		// Sort for determinism
+		DataLayerSoftObjectPaths.Sort([](const FSoftObjectPath& A, const FSoftObjectPath& B)
+		{
+			return A.LexicalLess(B);
+		});
+
+		return DataLayerSoftObjectPaths;
+	}
+				
+	TArray<const UDataLayerInstance*> GetDataLayerInstancesAndCrc(FPCGContext* Context, const FPCGDataLayerSettings& DataLayerSettings, AActor* DefaultDataLayerSource, int32& OutCrc)
+	{
+		TArray<const UDataLayerInstance*> DataLayerInstances;
+
+		switch (DataLayerSettings.DataLayerSourceType)
+		{
+		case EPCGDataLayerSource::Self:
+		{
+			DataLayerInstances = DefaultDataLayerSource->GetDataLayerInstances();
+			break;
+		}
+		case EPCGDataLayerSource::DataLayerReferences:
+		{
+			DataLayerInstances = Private::GetDataLayerInstancesFromDataLayerReferences(Context, DataLayerSettings, DefaultDataLayerSource->GetWorld());
+			break;
+		}
+		}
+
+		Private::FilterDataLayerInstances(Context, DataLayerSettings, DataLayerInstances);
+
+		Private::AddDataLayerInstances(Context, DataLayerSettings, DefaultDataLayerSource->GetWorld(), DataLayerInstances);
+
+		// Sort for determinism
+		DataLayerInstances.Sort([](const UDataLayerInstance& A, const UDataLayerInstance& B)
+		{
+			return A.GetPathName() < B.GetPathName();
+		});
+
+		// Crc DataLayerInstances
+		FArchiveCrc32 Ar;
+		for (const UDataLayerInstance* DataLayerInstance : DataLayerInstances)
+		{
+			UDataLayerInstance* NonConstInstance = const_cast<UDataLayerInstance*>(DataLayerInstance);
+			Ar << NonConstInstance;
+		}
+		OutCrc = Ar.GetCrc();
+
+		return DataLayerInstances;
+	}
+
+	TArray<const UDataLayerAsset*> GetDatalayerAssetsForActor(AActor* InActor)
+	{
+		ULevelInstanceSubsystem* LevelInstanceSubsystem = UWorld::GetSubsystem<ULevelInstanceSubsystem>(InActor->GetWorld());
+		ILevelInstanceInterface* LevelInstance = LevelInstanceSubsystem ? LevelInstanceSubsystem->GetOwningLevelInstance(InActor->GetLevel()) : nullptr;
+		if (AActor* LevelInstanceActor = Cast<AActor>(LevelInstance))
+		{
+			TSet<const UDataLayerAsset*> DataLayerAssets;
+			DataLayerAssets.Append(LevelInstanceActor->GetDataLayerAssets(true));
+			DataLayerAssets.Append(InActor->GetDataLayerAssets(false));
+			return DataLayerAssets.Array();
+		}
+		else
+		{
+			return InActor->GetDataLayerAssets(true);
+		}
+	}
+
+#endif // WITH_EDITOR
+}
+
+FPCGDataLayerSettings::FPCGDataLayerSettings()
+{
+	DataLayerReferenceAttribute.SetAttributeName(PCGDataLayerHelpers::Constants::DataLayerReferenceAttribute);
+	IncludedDataLayers.Attribute.SetAttributeName(PCGDataLayerHelpers::Constants::IncludedDataLayersAttribute);
+	ExcludedDataLayers.Attribute.SetAttributeName(PCGDataLayerHelpers::Constants::ExcludedDataLayersAttribute);
+	AddDataLayers.Attribute.SetAttributeName(PCGDataLayerHelpers::Constants::AddDataLayersAttribute);
+}
+
+TArray<FPCGPinProperties> FPCGDataLayerSettings::InputPinProperties() const
+{
+	TArray<FPCGPinProperties> PinProperties;
+
+	if (DataLayerSourceType == EPCGDataLayerSource::DataLayerReferences)
+	{
+		FPCGPinProperties& IncludedDataLayersPin = PinProperties.Emplace_GetRef(PCGDataLayerHelpers::Constants::DataLayerReferenceAttribute, EPCGDataType::Param);
+		IncludedDataLayersPin.SetRequiredPin();
+	}
+
+	if (IncludedDataLayers.bAsInput)
+	{
+		FPCGPinProperties& IncludedDataLayersPin = PinProperties.Emplace_GetRef(PCGDataLayerHelpers::Constants::IncludedDataLayersAttribute, EPCGDataType::Param);
+		IncludedDataLayersPin.SetRequiredPin();
+	}
+
+	if (ExcludedDataLayers.bAsInput)
+	{
+		FPCGPinProperties& ExcludedDataLayersPin = PinProperties.Emplace_GetRef(PCGDataLayerHelpers::Constants::ExcludedDataLayersAttribute, EPCGDataType::Param);
+		ExcludedDataLayersPin.SetRequiredPin();
+	}
+
+	if (AddDataLayers.bAsInput)
+	{
+		FPCGPinProperties& AddDataLayersPin = PinProperties.Emplace_GetRef(PCGDataLayerHelpers::Constants::AddDataLayersAttribute, EPCGDataType::Param);
+		AddDataLayersPin.SetRequiredPin();
+	}
+
+	return PinProperties;
+}
+
+#if WITH_EDITOR
+EPCGChangeType FPCGDataLayerSettings::GetChangeTypeForProperty(const FName& InPropertyName) const
+{
+	EPCGChangeType ChangeType = EPCGChangeType::None;
+
+	if (InPropertyName == GET_MEMBER_NAME_CHECKED(FPCGDataLayerReferenceSelector, bAsInput)
+		|| InPropertyName == GET_MEMBER_NAME_CHECKED(FPCGDataLayerSettings, DataLayerSourceType))
+	{
+		ChangeType |= EPCGChangeType::Structural;
+	}
+
+	return ChangeType;
+}
+#endif
+
+#undef LOCTEXT_NAMESPACE
\ No newline at end of file
diff --git a/Engine/Plugins/PCG/Source/PCG/Public/Elements/PCGCreateTargetActor.h b/Engine/Plugins/PCG/Source/PCG/Public/Elements/PCGCreateTargetActor.h
--- a/Engine/Plugins/PCG/Source/PCG/Public/Elements/PCGCreateTargetActor.h
+++ b/Engine/Plugins/PCG/Source/PCG/Public/Elements/PCGCreateTargetActor.h
@@ -3,8 +3,11 @@
 #pragma once
 
 #include "PCGSettings.h"
+#include "Helpers/PCGDataLayerHelpers.h"
 #include "Metadata/PCGObjectPropertyOverride.h"
 
+#include "WorldPartition/DataLayer/DataLayerAsset.h"
+
 #include "PCGCreateTargetActor.generated.h"
 
 class AActor;
@@ -40,9 +43,10 @@ protected:
 	virtual TArray<FPCGPinProperties> OutputPinProperties() const override;
 
 	virtual FPCGElementPtr CreateElement() const override;
+#if WITH_EDITOR
+	virtual EPCGChangeType GetChangeTypeForProperty(const FName& InPropertyName) const override;
 	//~End UPCGSettings interface
 
-#if WITH_EDITOR
 	void OnBlueprintChanged(UBlueprint* Blueprint);
 	void SetupBlueprintEvent();
 	void TeardownBlueprintEvent();
@@ -78,6 +82,9 @@ public:
 	UPROPERTY(BlueprintReadWrite, EditAnywhere, Category = Settings)
 	TArray<FName> PostProcessFunctionNames;
 
+	UPROPERTY(BlueprintReadWrite, EditAnywhere, Category = Settings, meta = (PCG_Overridable))
+	FPCGDataLayerSettings DataLayerSettings;
+
 	PCG_API void SetTemplateActorClass(const TSubclassOf<AActor>& InTemplateActorClass);
 	PCG_API void SetAllowTemplateActorEditing(bool bInAllowTemplateActorEditing);
 	const TSubclassOf<AActor>& GetTemplateActorClass() const { return TemplateActorClass; }
diff --git a/Engine/Plugins/PCG/Source/PCG/Public/Elements/PCGGetActorDataLayers.h b/Engine/Plugins/PCG/Source/PCG/Public/Elements/PCGGetActorDataLayers.h
new file mode 100644
--- /dev/null
+++ b/Engine/Plugins/PCG/Source/PCG/Public/Elements/PCGGetActorDataLayers.h
@@ -0,0 +1,48 @@
+// Copyright Epic Games, Inc. All Rights Reserved.
+
+#pragma once
+
+#include "PCGSettings.h"
+
+#include "PCGGetActorDataLayers.generated.h"
+
+UCLASS(MinimalAPI, BlueprintType, ClassGroup = (Procedural))
+class UPCGGetActorDataLayersSettings : public UPCGSettings
+{
+	GENERATED_BODY()
+
+public:
+	UPCGGetActorDataLayersSettings();
+
+	//~Begin UPCGSettings interface
+#if WITH_EDITOR
+	virtual EPCGSettingsType GetType() const override { return EPCGSettingsType::DataLayers; }
+	virtual FName GetDefaultNodeName() const override { return FName(TEXT("PCGGetActorDataLayers")); }
+	virtual FText GetDefaultNodeTitle() const override { return NSLOCTEXT("PCGGetActorDataLayersSettings", "NodeTitle", "Get Actor Data Layers"); }
+#endif
+	virtual FString GetAdditionalTitleInformation() const override;
+
+protected:
+	virtual TArray<FPCGPinProperties> InputPinProperties() const override;
+	virtual TArray<FPCGPinProperties> OutputPinProperties() const override;
+	virtual FPCGElementPtr CreateElement() const override;
+	//~End UPCGSettings interface
+
+public:
+	/** Actor reference Attribute to use from input to resolve actors */
+	UPROPERTY(BlueprintReadWrite, EditAnywhere, Category = Default, meta = (PCG_Overridable))
+	FPCGAttributePropertyInputSelector ActorReferenceAttribute;
+
+	/** Data Layer reference Attribute to use as output */
+	UPROPERTY(BlueprintReadWrite, EditAnywhere, Category = Default, meta = (PCG_Overridable))
+	FPCGAttributePropertyOutputSelector DataLayerReferenceAttribute;
+};
+
+class FPCGGetActorDataLayersElement : public IPCGElement
+{
+public:
+	virtual bool IsCacheable(const UPCGSettings* InSettings) const override { return false; }
+	virtual bool CanExecuteOnlyOnMainThread(FPCGContext* Context) const override { return true; }
+protected:
+	virtual bool ExecuteInternal(FPCGContext* Context) const override;
+};
\ No newline at end of file
diff --git a/Engine/Plugins/PCG/Source/PCG/Public/Elements/PCGPartitionByActorDataLayers.h b/Engine/Plugins/PCG/Source/PCG/Public/Elements/PCGPartitionByActorDataLayers.h
new file mode 100644
--- /dev/null
+++ b/Engine/Plugins/PCG/Source/PCG/Public/Elements/PCGPartitionByActorDataLayers.h
@@ -0,0 +1,59 @@
+// Copyright Epic Games, Inc. All Rights Reserved.
+
+#pragma once
+
+#include "PCGSettings.h"
+#include "Helpers/PCGDataLayerHelpers.h"
+
+#include "PCGPartitionByActorDataLayers.generated.h"
+
+class UDataLayerAsset;
+
+UCLASS(MinimalAPI, BlueprintType, ClassGroup = (Procedural))
+class UPCGPartitionByActorDataLayersSettings : public UPCGSettings
+{
+	GENERATED_BODY()
+
+	UPCGPartitionByActorDataLayersSettings();
+
+	//~Begin UPCGSettings interface
+#if WITH_EDITOR
+	virtual EPCGSettingsType GetType() const override { return EPCGSettingsType::DataLayers; }
+	virtual FName GetDefaultNodeName() const override { return FName(TEXT("PartitionByActorDataLayers")); }
+	virtual FText GetDefaultNodeTitle() const override { return NSLOCTEXT("PCGPartitionByActorDataLayersSettings", "NodeTitle", "Partition by Actor Data Layers"); }
+#endif
+
+protected:
+	virtual TArray<FPCGPinProperties> InputPinProperties() const override;
+	virtual TArray<FPCGPinProperties> OutputPinProperties() const override;
+	virtual FPCGElementPtr CreateElement() const override;
+#if WITH_EDITOR
+	virtual EPCGChangeType GetChangeTypeForProperty(const FName& InPropertyName) const override;
+#endif
+	//~End UPCGSettings interface
+		
+public:
+	UPROPERTY(BlueprintReadWrite, EditAnywhere, Category = Settings, meta = (PCG_Overridable))
+	FPCGAttributePropertyInputSelector ActorReferenceAttribute;
+		
+	/** Data Layer reference Attribute to use as output for Data Layer Partitions */
+	UPROPERTY(BlueprintReadWrite, EditAnywhere, Category = Settings, meta = (PCG_Overridable))
+	FPCGAttributePropertyOutputSelector DataLayerReferenceAttribute;
+
+	/** When left empty, all Data Layers are included, if any Data Layers are specified, only those will be included */
+	UPROPERTY(BlueprintReadWrite, EditAnywhere, Category = Settings)
+	FPCGDataLayerReferenceSelector IncludedDataLayers;
+
+	/** Specified Data Layers will get excluded */
+	UPROPERTY(BlueprintReadWrite, EditAnywhere, Category = Settings)
+	FPCGDataLayerReferenceSelector ExcludedDataLayers;
+};
+
+class FPCGPartitionByActorDataLayersElement : public IPCGElement
+{
+public:
+	virtual bool IsCacheable(const UPCGSettings* InSettings) const override { return false; }
+	virtual bool CanExecuteOnlyOnMainThread(FPCGContext* Context) const override { return true; }
+protected:
+	virtual bool ExecuteInternal(FPCGContext* Context) const override;
+};
\ No newline at end of file
diff --git a/Engine/Plugins/PCG/Source/PCG/Public/Elements/PCGSpawnActor.h b/Engine/Plugins/PCG/Source/PCG/Public/Elements/PCGSpawnActor.h
--- a/Engine/Plugins/PCG/Source/PCG/Public/Elements/PCGSpawnActor.h
+++ b/Engine/Plugins/PCG/Source/PCG/Public/Elements/PCGSpawnActor.h
@@ -3,6 +3,7 @@
 #pragma once
 
 #include "PCGSubgraph.h"
+#include "Helpers/PCGDataLayerHelpers.h"
 #include "Metadata/PCGObjectPropertyOverride.h"
 
 #include "PCGSpawnActor.generated.h"
@@ -88,6 +89,9 @@ public:
 	UPROPERTY(BlueprintReadWrite, EditAnywhere, Category = "Settings|Advanced")
 	bool bWarnOnIdenticalSpawn = true;
 
+	UPROPERTY(BlueprintReadWrite, EditAnywhere, Category = Settings, meta = (PCG_Overridable))
+	FPCGDataLayerSettings DataLayerSettings;
+
 protected:
 	UPROPERTY(BlueprintReadOnly, EditAnywhere, Category = Settings, meta = (OnlyPlaceable, DisallowCreateNew))
 	TSubclassOf<AActor> TemplateActorClass = nullptr;
@@ -133,7 +137,7 @@ protected:
 #if WITH_EDITOR	
 	virtual EPCGChangeType GetChangeTypeForProperty(const FName& InPropertyName) const override;
 #endif
-	virtual TArray<FPCGPinProperties> InputPinProperties() const override { return DefaultPointInputPinProperties(); }
+	virtual TArray<FPCGPinProperties> InputPinProperties() const override;
 	virtual TArray<FPCGPinProperties> OutputPinProperties() const override { return DefaultPointOutputPinProperties(); }
 	virtual FPCGElementPtr CreateElement() const override;
 	//~End UPCGSettings interface
diff --git a/Engine/Plugins/PCG/Source/PCG/Public/Helpers/PCGDataLayerHelpers.h b/Engine/Plugins/PCG/Source/PCG/Public/Helpers/PCGDataLayerHelpers.h
new file mode 100644
--- /dev/null
+++ b/Engine/Plugins/PCG/Source/PCG/Public/Helpers/PCGDataLayerHelpers.h
@@ -0,0 +1,97 @@
+// Copyright Epic Games, Inc. All Rights Reserved.
+
+#pragma once
+
+#include "PCGPin.h"
+#include "Metadata/PCGAttributePropertySelector.h"
+
+#include "Engine/EngineTypes.h"
+
+class UDataLayerAsset;
+
+#if WITH_EDITOR
+class AActor;
+class UDataLayerInstance;
+class UPCGData;
+
+struct FPCGDataLayerSettings;
+struct FPCGContext;
+#endif
+
+#include "PCGDataLayerHelpers.generated.h"
+
+namespace PCGDataLayerHelpers
+{
+	namespace Constants
+	{
+		const FName IncludedDataLayersAttribute = TEXT("IncludedDataLayers");
+		const FName ExcludedDataLayersAttribute = TEXT("ExcludedDataLayers");
+		const FName AddDataLayersAttribute = TEXT("AddDataLayers");
+		const FName DataLayerReferenceAttribute = TEXT("DataLayerReference");
+	}
+
+#if WITH_EDITOR
+	TArray<TSoftObjectPtr<UDataLayerAsset>> GetDataLayerAssetsFromInput(FPCGContext* Context, FName InputPinName, const FPCGDataLayerReferenceSelector& DataLayerSelector);
+	TArray<TSoftObjectPtr<UDataLayerAsset>> GetDataLayerAssetsFromInput(FPCGContext* Context, FName InputPinName, const FPCGAttributePropertyInputSelector& InputSelector);
+	TArray<FSoftObjectPath> GetDataLayerAssetsFromActorReferences(FPCGContext* Context, const UPCGData* ParamData, const FPCGAttributePropertyInputSelector& ActorReferenceAttribute);
+	TArray<const UDataLayerInstance*> GetDataLayerInstancesAndCrc(FPCGContext* Context, const FPCGDataLayerSettings& DataLayerSettings, AActor* DefaulDataLayerSource, int32& OutCrc);
+	TArray<const UDataLayerAsset*> GetDatalayerAssetsForActor(AActor* InActor);
+#endif
+};
+
+UENUM()
+enum class EPCGDataLayerSource : uint8
+{
+	Self,
+	DataLayerReferences,
+};
+
+USTRUCT(BlueprintType)
+struct FPCGDataLayerReferenceSelector
+{
+	GENERATED_BODY()
+
+	/** Set it to true to get Data Layers through input attribute set. */
+	UPROPERTY(BlueprintReadWrite, EditAnywhere, Category = Default)
+	bool bAsInput = false;
+
+	UPROPERTY(BlueprintReadWrite, EditAnywhere, Category = Default, meta = (EditCondition = "bAsInput", EditConditionHides))
+	FPCGAttributePropertyInputSelector Attribute;
+
+	UPROPERTY(BlueprintReadWrite, EditAnywhere, Category = Default, meta = (EditCondition = "!bAsInput", EditConditionHides))
+	TArray<TSoftObjectPtr<UDataLayerAsset>> DataLayers;
+};
+
+
+USTRUCT(BlueprintType)
+struct FPCGDataLayerSettings
+{
+	GENERATED_BODY()
+
+	FPCGDataLayerSettings();
+
+	/** What source should be used to assign Data Layers to the spawned actor */
+	UPROPERTY(BlueprintReadWrite, EditAnywhere, Category = Default)
+	EPCGDataLayerSource DataLayerSourceType = EPCGDataLayerSource::Self;
+		
+	UPROPERTY(BlueprintReadWrite, EditAnywhere, Category = Default, meta = (EditCondition = "DataLayerSourceType==EPCGDataLayerSource::DataLayerReferences", EditConditionHides, DisplayAfter = DataLayerSourceType))
+	FPCGAttributePropertyInputSelector DataLayerReferenceAttribute;
+		
+	/** When left empty, all Data Layers from the Data Layer Source are included, if any Data Layers are specified, only those will be included */
+	UPROPERTY(BlueprintReadWrite, EditAnywhere, Category = Default)
+	FPCGDataLayerReferenceSelector IncludedDataLayers;
+		
+	/** Specified Data Layers will get excluded from the Data Layer Source */
+	UPROPERTY(BlueprintReadWrite, EditAnywhere, Category = Default)
+	FPCGDataLayerReferenceSelector ExcludedDataLayers;
+		
+	/** Specified Data Layers will get added */
+	UPROPERTY(BlueprintReadWrite, EditAnywhere, Category = Default)
+	FPCGDataLayerReferenceSelector AddDataLayers;
+
+	TArray<FPCGPinProperties> InputPinProperties() const;
+
+#if WITH_EDITOR
+	EPCGChangeType GetChangeTypeForProperty(const FName& InPropertyName) const;
+#endif
+};
\ No newline at end of file
diff --git a/Engine/Plugins/PCG/Source/PCG/Public/Metadata/Accessors/PCGAttributeAccessorHelpers.h b/Engine/Plugins/PCG/Source/PCG/Public/Metadata/Accessors/PCGAttributeAccessorHelpers.h
--- a/Engine/Plugins/PCG/Source/PCG/Public/Metadata/Accessors/PCGAttributeAccessorHelpers.h
+++ b/Engine/Plugins/PCG/Source/PCG/Public/Metadata/Accessors/PCGAttributeAccessorHelpers.h
@@ -104,19 +104,33 @@ namespace PCGAttributeAccessorHelpers
 	 * @return true if it succeeded.
 	 */
 	template <typename T, typename AllocatorType>
-	bool ExtractAllValues(const UPCGData* InData, const FPCGAttributePropertyInputSelector& InSelector, TArray<T, AllocatorType>& OutArray, FPCGContext* Context = nullptr, EPCGAttributeAccessorFlags GetFlags = EPCGAttributeAccessorFlags::AllowBroadcastAndConstructible)
+	bool ExtractAllValues(const UPCGData* InData, const FPCGAttributePropertyInputSelector& InSelector, TArray<T, AllocatorType>& OutArray, FPCGContext* Context = nullptr, EPCGAttributeAccessorFlags GetFlags = EPCGAttributeAccessorFlags::AllowBroadcastAndConstructible, bool bQuiet = false)
 	{
 		// TODO: Provide a version that take advantage of compressed attributes (Keys + Values)
 		
 		const FPCGAttributePropertyInputSelector Selector = InSelector.CopyAndFixLast(InData);
 
-		TUniquePtr<const IPCGAttributeAccessor> Accessor = CreateConstAccessor(InData, Selector);
+		TUniquePtr<const IPCGAttributeAccessor> Accessor = CreateConstAccessor(InData, Selector, bQuiet);
+		if (!Accessor.IsValid())
+		{
+			if (!bQuiet)
+			{
+				PCGLog::Metadata::LogFailToCreateAccessorError(Selector, Context);
+			}
+			return false;
+		}
+
 		TUniquePtr<const IPCGAttributeAccessorKeys> Keys = CreateConstKeys(InData, Selector);
-		if (!Accessor.IsValid() || !Keys.IsValid())
+		if (!Keys.IsValid())
 		{
 			PCGLog::Metadata::LogFailToCreateAccessorError(Selector, Context);
 			return false;
 		}
+
+		if (Keys->GetNum() < 1)
+		{
+			return true;
+		}
 			
 		if constexpr (std::is_trivially_copyable_v<T>)
 		{
diff --git a/Engine/Plugins/PCG/Source/PCG/Public/PCGSettings.h b/Engine/Plugins/PCG/Source/PCG/Public/PCGSettings.h
--- a/Engine/Plugins/PCG/Source/PCG/Public/PCGSettings.h
+++ b/Engine/Plugins/PCG/Source/PCG/Public/PCGSettings.h
@@ -71,6 +71,7 @@ enum class EPCGSettingsType : uint8
 	Reroute,
 	GPU,
 	DynamicMesh,
+	DataLayers,
 };
 
 #if WITH_EDITOR
