[PCG] Break out compute graph build function into DI creation/setup and graph wiring

#rnx

diff --git a/Engine/Plugins/PCG/Source/PCG/Private/Graph/PCGGraphCompilerGPU.cpp b/Engine/Plugins/PCG/Source/PCG/Private/Graph/PCGGraphCompilerGPU.cpp
--- a/Engine/Plugins/PCG/Source/PCG/Private/Graph/PCGGraphCompilerGPU.cpp
+++ b/Engine/Plugins/PCG/Source/PCG/Private/Graph/PCGGraphCompilerGPU.cpp
@@ -428,15 +428,14 @@ void FPCGGraphCompilerGPU::SetupVirtualPins(
 	}
 }
 
-void FPCGGraphCompilerGPU::WireGPUGraphNode(
+void FPCGGraphCompilerGPU::WireComputeGraphTask(
 	FPCGTaskId InGPUGraphTaskId,
 	const TSet<FPCGTaskId>& InCollapsedTasks,
 	TArray<FPCGGraphTask>& InOutCompiledTasks,
 	const FTaskToSuccessors& InTaskSuccessors,
-	const FOriginalToVirtualPin& InOriginalToVirtualPin,
-	const TMap<TSoftObjectPtr<const UPCGPin>, FName>& InOutputCPUPinToVirtualPin)
+	const FOriginalToVirtualPin& InOriginalToVirtualPin)
 {
-	TRACE_CPUPROFILER_EVENT_SCOPE(FPCGGraphCompilerGPU::WireGPUGraphNode);
+	TRACE_CPUPROFILER_EVENT_SCOPE(FPCGGraphCompilerGPU::WireComputeGraphTask);
 
 	FPCGGraphTask& GPUGraphTask = InOutCompiledTasks[InGPUGraphTaskId];
 
@@ -517,12 +516,12 @@ void FPCGGraphCompilerGPU::WireGPUGraphNode(
 	}
 }
 
-void FPCGGraphCompilerGPU::BuildGPUGraphTask(
+void FPCGGraphCompilerGPU::BuildComputeGraphTask(
 	FPCGGPUCompilationContext& InOutContext,
+	FPCGTaskId InGPUGraphTaskId,
 	UPCGGraph* InGraph,
 	uint32 InGridSize,
 	uint32 InComputeGraphIndex,
-	FPCGTaskId InGPUGraphTaskId,
 	const TSet<FPCGTaskId>& InCollapsedTasks,
 	const TSet<FPCGTaskId>& InAllGPUCompatibleTasks,
 	const FTaskToSuccessors& InTaskSuccessors,
@@ -530,7 +529,7 @@ void FPCGGraphCompilerGPU::BuildGPUGraphTask(
 	const FOriginalToVirtualPin& InOriginalToVirtualPin,
 	const TMap<TSoftObjectPtr<const UPCGPin>, FName>& InOutputCPUPinToVirtualPin)
 {
-	TRACE_CPUPROFILER_EVENT_SCOPE(FPCGGraphCompilerGPU::BuildGPUGraphTask);
+	TRACE_CPUPROFILER_EVENT_SCOPE(FPCGGraphCompilerGPU::BuildComputeGraphTask);
 
 	check(InGraph);
 
@@ -543,15 +542,27 @@ void FPCGGraphCompilerGPU::BuildGPUGraphTask(
 		// TODO: There is a race condition here where two threads can try to find this graph at the same time and both fail, so both create the graph.
 		// This is fine, as they'll both create the same graph and place it at the same compute graph index. However it would be ideal perf-wise to
 		// avoid this, as creating the compute graph can be expensive.
-		ComputeGraph = CompileComputeGraph(
-			InOutContext,
-			InGraph,
-			InCollapsedTasks,
-			InAllGPUCompatibleTasks,
-			InTaskSuccessors,
-			InOutCompiledTasks,
-			InOriginalToVirtualPin,
-			InOutputCPUPinToVirtualPin);
+		
+		// Create a new compute graph.
+		UObject* ComputeGraphOuter = InOutContext.GetGraphCompiler().IsCooking() ? InGraph : Cast<UObject>(GetTransientPackage());
+		const FName GraphName = MakeUniqueObjectName(InGraph, UPCGComputeGraph::StaticClass(), InGraph->GetFName());
+		ComputeGraph = InOutContext.NewObject_AnyThread<UPCGComputeGraph>(ComputeGraphOuter, GraphName);
+
+		ComputeGraph->OutputCPUPinToInputGPUPinAlias = InOutputCPUPinToVirtualPin;
+
+		// Not incredibly useful for us - DG adds GetComponentSource()->GetComponentClass() object which allows it to bind at execution time by class.
+		// But execution code requires it currently.
+		ComputeGraph->Bindings.Add(UPCGDataBinding::StaticClass());
+
+		FPinToDataInterface OutputPinDataInterfaces;
+
+		CreateDataInterfaces(InOutContext, InOutCompiledTasks, InCollapsedTasks, InAllGPUCompatibleTasks, InTaskSuccessors, InOriginalToVirtualPin, ComputeGraph, OutputPinDataInterfaces);
+
+		CompileComputeGraph(InOutContext, InOutCompiledTasks, InCollapsedTasks, InAllGPUCompatibleTasks, InTaskSuccessors, InOriginalToVirtualPin, OutputPinDataInterfaces, ComputeGraph);
+
+		// Remove empty strings (at execution time empty string is placed in table slot 0).
+		// NOTE: This can scramble order but order is not important at this stage (key values created at execution time in data binding).
+		ComputeGraph->StringTable.RemoveAllSwap([](const FString& InString) { return InString.IsEmpty(); });
 	}
 
 	const int32 ComputeGraphIndex = InOutContext.AddCompiledComputeGraph(ComputeGraph);
@@ -578,32 +589,21 @@ void FPCGGraphCompilerGPU::BuildGPUGraphTask(
 	}
 }
 
-UPCGComputeGraph* FPCGGraphCompilerGPU::CompileComputeGraph(
+void FPCGGraphCompilerGPU::CreateDataInterfaces(
 	FPCGGPUCompilationContext& InOutContext,
-	UPCGGraph* InGraph,
+	const TArray<FPCGGraphTask>& InCompiledTasks,
 	const TSet<FPCGTaskId>& InCollapsedTasks,
 	const TSet<FPCGTaskId>& InAllGPUCompatibleTasks,
 	const FTaskToSuccessors& InTaskSuccessors,
-	TArray<FPCGGraphTask>& InOutCompiledTasks,
 	const FOriginalToVirtualPin& InOriginalToVirtualPin,
-	const TMap<TSoftObjectPtr<const UPCGPin>, FName>& InOutputCPUPinToVirtualPin)
+	UPCGComputeGraph* InOutComputeGraph,
+	FPinToDataInterface& InOutPinToDataInterface)
 {
-	TRACE_CPUPROFILER_EVENT_SCOPE(FPCGGraphCompilerGPU::CompileComputeGraph);
-
-	// Create a new compute graph.
-	UObject* ComputeGraphOuter = InOutContext.GetGraphCompiler().IsCooking() ? InGraph : Cast<UObject>(GetTransientPackage());
-	const FName GraphName = MakeUniqueObjectName(InGraph, UPCGComputeGraph::StaticClass(), InGraph->GetFName());
-	UPCGComputeGraph* ComputeGraph = InOutContext.NewObject_AnyThread<UPCGComputeGraph>(ComputeGraphOuter, GraphName);
-
-	ComputeGraph->OutputCPUPinToInputGPUPinAlias = InOutputCPUPinToVirtualPin;
-
-	// Not incredibly useful for us - DG adds GetComponentSource()->GetComponentClass() object which allows it to bind at execution time by class.
-	// But execution code requires it currently.
-	ComputeGraph->Bindings.Add(UPCGDataBinding::StaticClass());
+	TRACE_CPUPROFILER_EVENT_SCOPE(FPCGGraphCompilerGPU::CreateDataInterfaces);
 
 	// Create data interfaces which allow kernels to read or write data. Each data interface is associated with a node output pin.
 	// For CPU->GPU edges, an upload data interface is created. For GPU->CPU edges, a readback data interface is created.
-	auto CreateDataInterface = [&InOutContext, &InCollapsedTasks, &InOutCompiledTasks, ComputeGraph](FPCGTaskId InTaskId, bool bRequiresReadback, const FPCGPinProperties& InOutputPinProperties) -> UPCGComputeDataInterface*
+	auto CreateDataInterface = [&InOutContext, &InCollapsedTasks, &InCompiledTasks, InOutComputeGraph](FPCGTaskId InTaskId, bool bRequiresReadback, const FPCGPinProperties& InOutputPinProperties) -> UPCGComputeDataInterface*
 	{
 		TRACE_CPUPROFILER_EVENT_SCOPE(CreateDataInterface);
 
@@ -612,8 +612,8 @@ UPCGComputeGraph* FPCGGraphCompilerGPU::CompileComputeGraph(
 		EPCGDataType PinType = InOutputPinProperties.AllowedTypes;
 
 		// Dynamically typed pins could have a different type
-		const UPCGSettings* Settings = InOutCompiledTasks[InTaskId].Node ? InOutCompiledTasks[InTaskId].Node->GetSettings() : nullptr;
-		const UPCGPin* Pin = InOutCompiledTasks[InTaskId].Node ? InOutCompiledTasks[InTaskId].Node->GetOutputPin(InOutputPinProperties.Label) : nullptr;
+		const UPCGSettings* Settings = InCompiledTasks[InTaskId].Node ? InCompiledTasks[InTaskId].Node->GetSettings() : nullptr;
+		const UPCGPin* Pin = InCompiledTasks[InTaskId].Node ? InCompiledTasks[InTaskId].Node->GetOutputPin(InOutputPinProperties.Label) : nullptr;
 		if (Settings && Pin)
 		{
 			PinType = Settings->GetCurrentPinTypes(Pin);
@@ -631,16 +631,16 @@ UPCGComputeGraph* FPCGGraphCompilerGPU::CompileComputeGraph(
 			if (bUpstreamIsGPUTask)
 			{
 				// Provides data for GPU -> GPU and GPU -> CPU edges.
-				DataInterfacePCGData = InOutContext.NewObject_AnyThread<UPCGDataCollectionDataInterface>(ComputeGraph);
+				DataInterfacePCGData = InOutContext.NewObject_AnyThread<UPCGDataCollectionDataInterface>(InOutComputeGraph);
 			}
 			else
 			{
 				// Provides data for CPU -> GPU edge.
-				DataInterfacePCGData = InOutContext.NewObject_AnyThread<UPCGDataCollectionUploadDataInterface>(ComputeGraph);
+				DataInterfacePCGData = InOutContext.NewObject_AnyThread<UPCGDataCollectionUploadDataInterface>(InOutComputeGraph);
 			}
 
 			check(DataInterfacePCGData);
-			DataInterfacePCGData->SetProducerSettings(InOutCompiledTasks[InTaskId].Node ? InOutCompiledTasks[InTaskId].Node->GetSettings() : nullptr);
+			DataInterfacePCGData->SetProducerSettings(InCompiledTasks[InTaskId].Node ? InCompiledTasks[InTaskId].Node->GetSettings() : nullptr);
 			DataInterfacePCGData->SetRequiresReadback(bRequiresReadback);
 
 			DataInterface = DataInterfacePCGData;
@@ -648,19 +648,19 @@ UPCGComputeGraph* FPCGGraphCompilerGPU::CompileComputeGraph(
 		}
 		case EPCGDataType::Texture:
 		{
-			DataInterface = InOutContext.NewObject_AnyThread<UPCGTextureDataInterface>(ComputeGraph);
+			DataInterface = InOutContext.NewObject_AnyThread<UPCGTextureDataInterface>(InOutComputeGraph);
 			break;
 		}
 		case EPCGDataType::Landscape:
 		{
-			DataInterface = InOutContext.NewObject_AnyThread<UPCGLandscapeDataInterface>(ComputeGraph);
+			DataInterface = InOutContext.NewObject_AnyThread<UPCGLandscapeDataInterface>(InOutComputeGraph);
 			break;
 		}
 		default:
 		{
 			UE_LOG(LogPCG, Error, TEXT("Unsupported connected upstream pin '%s' on node '%s' with type %s. Consider adding a conversion to a supported type such as Point."),
 				*InOutputPinProperties.Label.ToString(),
-				InOutCompiledTasks[InTaskId].Node ? *InOutCompiledTasks[InTaskId].Node->GetNodeTitle(EPCGNodeTitleType::ListView).ToString() : TEXT("MISSING"),
+				InCompiledTasks[InTaskId].Node ? *InCompiledTasks[InTaskId].Node->GetNodeTitle(EPCGNodeTitleType::ListView).ToString() : TEXT("MISSING"),
 				StaticEnum<EPCGDataType>() ? *StaticEnum<EPCGDataType>()->GetNameStringByValue(static_cast<int64>(PinType)) : TEXT("MISSING")
 			);
 			break;
@@ -675,13 +675,11 @@ UPCGComputeGraph* FPCGGraphCompilerGPU::CompileComputeGraph(
 		return DataInterface;
 	};
 
-	TMap<TPair</* Node task*/FPCGTaskId, /*Node output pin*/FName>, UPCGComputeDataInterface*> OutputPinDataInterfaces;
-
 	// Create DIs for output pins (either output pins within this compute graph or outside).
 	for (FPCGTaskId TaskId : InCollapsedTasks)
 	{
 		// Create DIs for all output pins regardless of outbound connections, because the kernels currently need their outputs to be bound to valid resources.
-		if (const UPCGSettings* Settings = InOutCompiledTasks[TaskId].Node ? InOutCompiledTasks[TaskId].Node->GetSettings() : nullptr)
+		if (const UPCGSettings* Settings = InCompiledTasks[TaskId].Node ? InCompiledTasks[TaskId].Node->GetSettings() : nullptr)
 		{
 			// Record attributes.
 			TArray<FPCGKernelAttributeKey> AttributeKeys;
@@ -689,7 +687,7 @@ UPCGComputeGraph* FPCGGraphCompilerGPU::CompileComputeGraph(
 
 			for (const FPCGKernelAttributeKey& Key : AttributeKeys)
 			{
-				FPCGKernelAttributeIDAndType* FoundEntry = ComputeGraph->GlobalAttributeLookupTable.Find(Key.Name);
+				FPCGKernelAttributeIDAndType* FoundEntry = InOutComputeGraph->GlobalAttributeLookupTable.Find(Key.Name);
 				if (FoundEntry)
 				{
 					// If we already registered this attribute, populate the type if not populated already, or otherwise ensure type matches.
@@ -709,13 +707,13 @@ UPCGComputeGraph* FPCGGraphCompilerGPU::CompileComputeGraph(
 				}
 				else
 				{
-					const FPCGKernelAttributeIDAndType IDAndType(ComputeGraph->GlobalAttributeLookupTable.Num() + PCGComputeConstants::NUM_RESERVED_ATTRS, Key.Type);
-					ComputeGraph->GlobalAttributeLookupTable.Add(Key.Name, IDAndType);
+					const FPCGKernelAttributeIDAndType IDAndType(InOutComputeGraph->GlobalAttributeLookupTable.Num() + PCGComputeConstants::NUM_RESERVED_ATTRS, Key.Type);
+					InOutComputeGraph->GlobalAttributeLookupTable.Add(Key.Name, IDAndType);
 				}
 			}
 
 			// Record strings.
-			Settings->AddStaticCreatedStrings(ComputeGraph->StringTable);
+			Settings->AddStaticCreatedStrings(InOutComputeGraph->StringTable);
 
 			// Create all the output data interfaces.
 			for (const FPCGPinProperties& OutputPinProperties : Settings->AllOutputPinProperties())
@@ -725,7 +723,7 @@ UPCGComputeGraph* FPCGGraphCompilerGPU::CompileComputeGraph(
 					continue;
 				}
 
-				if (OutputPinDataInterfaces.Contains({ TaskId, OutputPinProperties.Label }))
+				if (InOutPinToDataInterface.Contains({ TaskId, OutputPinProperties.Label }))
 				{
 					ensure(false);
 					continue;
@@ -737,7 +735,7 @@ UPCGComputeGraph* FPCGGraphCompilerGPU::CompileComputeGraph(
 				{
 					for (FPCGTaskId Successor : *Successors)
 					{
-						for (const FPCGGraphTaskInput& Input : InOutCompiledTasks[Successor].Inputs)
+						for (const FPCGGraphTaskInput& Input : InCompiledTasks[Successor].Inputs)
 						{
 							if (Input.TaskId == TaskId && Input.UpstreamPin && !InCollapsedTasks.Contains(Successor))
 							{
@@ -754,16 +752,16 @@ UPCGComputeGraph* FPCGGraphCompilerGPU::CompileComputeGraph(
 					const FName* VirtualPinLabel = InOriginalToVirtualPin.Find({ TaskId, OutputPinProperties.Label, /*bIsInputPin*/false });
 					OutputDI->SetOutputPin(OutputPinProperties.Label, VirtualPinLabel);
 
-					ComputeGraph->DataInterfaces.Add(OutputDI);
+					InOutComputeGraph->DataInterfaces.Add(OutputDI);
 
-					OutputPinDataInterfaces.Add({ TaskId, OutputPinProperties.Label }, OutputDI);
+					InOutPinToDataInterface.Add({ TaskId, OutputPinProperties.Label }, OutputDI);
 
 					// Iterate over downstream connections and register each downstream pin.
 					if (const TArray<FPCGTaskId>* Successors = InTaskSuccessors.Find(TaskId))
 					{
 						for (FPCGTaskId Successor : *Successors)
 						{
-							for (const FPCGGraphTaskInput& Input : InOutCompiledTasks[Successor].Inputs)
+							for (const FPCGGraphTaskInput& Input : InCompiledTasks[Successor].Inputs)
 							{
 								if (Input.TaskId == TaskId && Input.UpstreamPin && Input.DownstreamPin)
 								{
@@ -779,7 +777,7 @@ UPCGComputeGraph* FPCGGraphCompilerGPU::CompileComputeGraph(
 		}
 
 		// Create any DIs for upstream nodes outside of this compute graph.
-		for (const FPCGGraphTaskInput& Input : InOutCompiledTasks[TaskId].Inputs)
+		for (const FPCGGraphTaskInput& Input : InCompiledTasks[TaskId].Inputs)
 		{
 			// Never wire directly to the GPU-compatible tasks, these are collapsed into compute graphs (and will be culled in a next step).
 			if (InAllGPUCompatibleTasks.Contains(Input.TaskId))
@@ -792,11 +790,11 @@ UPCGComputeGraph* FPCGGraphCompilerGPU::CompileComputeGraph(
 				continue;
 			}
 
-			if (const UPCGNode* DownstreamNode = InOutCompiledTasks[TaskId].Node)
+			if (const UPCGNode* DownstreamNode = InCompiledTasks[TaskId].Node)
 			{
 				if (const UPCGPin* InputPin = DownstreamNode->GetInputPin(Input.DownstreamPin->Label))
 				{
-					ComputeGraph->PinsReceivingDataFromCPU.Add(InputPin);
+					InOutComputeGraph->PinsReceivingDataFromCPU.Add(InputPin);
 				}
 			}
 
@@ -807,7 +805,7 @@ UPCGComputeGraph* FPCGGraphCompilerGPU::CompileComputeGraph(
 
 			UPCGComputeDataInterface* OutputDI = nullptr;
 
-			if (UPCGComputeDataInterface** FoundOutputDI = OutputPinDataInterfaces.Find({ Input.TaskId, Input.UpstreamPin->Label }))
+			if (UPCGComputeDataInterface** FoundOutputDI = InOutPinToDataInterface.Find({ Input.TaskId, Input.UpstreamPin->Label }))
 			{
 				OutputDI = *FoundOutputDI;
 			}
@@ -819,9 +817,9 @@ UPCGComputeGraph* FPCGGraphCompilerGPU::CompileComputeGraph(
 				const FName* VirtualPinLabel = InOriginalToVirtualPin.Find({ Input.TaskId, Input.UpstreamPin->Label, /*bIsInputPin*/false });
 				OutputDI->SetOutputPin(Input.UpstreamPin->Label, VirtualPinLabel);
 
-				ComputeGraph->DataInterfaces.Add(OutputDI);
+				InOutComputeGraph->DataInterfaces.Add(OutputDI);
 
-				OutputPinDataInterfaces.Add({ Input.TaskId, Input.UpstreamPin->Label }, OutputDI);
+				InOutPinToDataInterface.Add({ Input.TaskId, Input.UpstreamPin->Label }, OutputDI);
 			}
 
 			if (OutputDI)
@@ -832,6 +830,19 @@ UPCGComputeGraph* FPCGGraphCompilerGPU::CompileComputeGraph(
 			}
 		}
 	}
+}
+
+void FPCGGraphCompilerGPU::CompileComputeGraph(
+	FPCGGPUCompilationContext& InOutContext,
+	const TArray<FPCGGraphTask>& InCompiledTasks,
+	const TSet<FPCGTaskId>& InCollapsedTasks,
+	const TSet<FPCGTaskId>& InAllGPUCompatibleTasks,
+	const FTaskToSuccessors& InTaskSuccessors,
+	const FOriginalToVirtualPin& InOriginalToVirtualPin,
+	const FPinToDataInterface& InOutputPinToDataInterface,
+	UPCGComputeGraph* InOutComputeGraph)
+{
+	TRACE_CPUPROFILER_EVENT_SCOPE(FPCGGraphCompilerGPU::CompileComputeGraph);
 
 	// Build compute graph by traversing the set of collapsed tasks in execution order (order of queued kernel invocations matters).
 	TSet<FPCGTaskId> RemainingTasks;
@@ -842,7 +853,7 @@ UPCGComputeGraph* FPCGGraphCompilerGPU::CompileComputeGraph(
 	for (FPCGTaskId TaskId : InCollapsedTasks)
 	{
 		bool bTaskReady = true;
-		for (const FPCGGraphTaskInput& Input : InOutCompiledTasks[TaskId].Inputs)
+		for (const FPCGGraphTaskInput& Input : InCompiledTasks[TaskId].Inputs)
 		{
 			if (InCollapsedTasks.Contains(Input.TaskId))
 			{
@@ -876,7 +887,7 @@ UPCGComputeGraph* FPCGGraphCompilerGPU::CompileComputeGraph(
 				}
 
 				bool bTaskReady = true;
-				for (const FPCGGraphTaskInput& Input : InOutCompiledTasks[Successor].Inputs)
+				for (const FPCGGraphTaskInput& Input : InCompiledTasks[Successor].Inputs)
 				{
 					if (RemainingTasks.Contains(Input.TaskId))
 					{
@@ -893,14 +904,14 @@ UPCGComputeGraph* FPCGGraphCompilerGPU::CompileComputeGraph(
 			}
 		}
 
-		const UPCGNode* Node = InOutCompiledTasks[TaskId].Node;
+		const UPCGNode* Node = InCompiledTasks[TaskId].Node;
 
 		const UPCGSettings* Settings = Node ? Node->GetSettings() : nullptr;
 		check(Settings && Settings->bEnabled && Settings->ShouldExecuteOnGPU());
 
 		if (Settings->IsA<UPCGStaticMeshSpawnerSettings>())
 		{
-			ComputeGraph->StaticMeshSpawners.Add(Settings);
+			InOutComputeGraph->StaticMeshSpawners.Add(Settings);
 		}
 
 		// One data interface may be connected to multiple downstream pins. Because each entry in these arrays produces its own bindings and graph edges,
@@ -911,7 +922,7 @@ UPCGComputeGraph* FPCGGraphCompilerGPU::CompileComputeGraph(
 		OutputDataInterfaceIndexAndPin.Reserve(Settings->OutputPinProperties().Num());
 
 		// Setup input pin DIs.
-		for (const FPCGGraphTaskInput& Input : InOutCompiledTasks[TaskId].Inputs)
+		for (const FPCGGraphTaskInput& Input : InCompiledTasks[TaskId].Inputs)
 		{
 			// Currently the new compute graph tasks are wired into the graph in parallel to each GPU node task. The GPU node tasks
 			// will be culled at the end, leaving only the compute graphs. Only create DIs for tasks within this compute graph (in InCollapsedTasks),
@@ -929,15 +940,14 @@ UPCGComputeGraph* FPCGGraphCompilerGPU::CompileComputeGraph(
 				continue;
 			}
 
-
-			UPCGComputeDataInterface** FoundDI = OutputPinDataInterfaces.Find({ Input.TaskId, Input.UpstreamPin->Label });
+			UPCGComputeDataInterface* const* FoundDI = InOutputPinToDataInterface.Find({ Input.TaskId, Input.UpstreamPin->Label });
 			if (!FoundDI || !*FoundDI)
 			{
 				UE_LOG(LogPCG, Warning, TEXT("Failed to find data interface associated with pin '%s'."), *Input.UpstreamPin->Label.ToString());
 				continue;
 			}
 
-			const int Index = ComputeGraph->DataInterfaces.Find(*FoundDI);
+			const int Index = InOutComputeGraph->DataInterfaces.Find(*FoundDI);
 			if (Index == INDEX_NONE)
 			{
 				ensure(false);
@@ -955,14 +965,14 @@ UPCGComputeGraph* FPCGGraphCompilerGPU::CompileComputeGraph(
 				continue;
 			}
 
-			UPCGComputeDataInterface** FoundDI = OutputPinDataInterfaces.Find({ TaskId, OutputPinProperties.Label });
+			UPCGComputeDataInterface* const* FoundDI = InOutputPinToDataInterface.Find({ TaskId, OutputPinProperties.Label });
 			if (!FoundDI || !*FoundDI)
 			{
 				UE_LOG(LogPCG, Warning, TEXT("Failed to find data interface associated with pin '%s'."), *OutputPinProperties.Label.ToString());
 				continue;
 			}
 
-			const int Index = ComputeGraph->DataInterfaces.Find(*FoundDI);
+			const int Index = InOutComputeGraph->DataInterfaces.Find(*FoundDI);
 			if (Index == INDEX_NONE)
 			{
 				ensure(false);
@@ -974,26 +984,26 @@ UPCGComputeGraph* FPCGGraphCompilerGPU::CompileComputeGraph(
 
 		// Kernel data interface.
 		{
-			UPCGCustomKernelDataInterface* KernelDI = InOutContext.NewObject_AnyThread<UPCGCustomKernelDataInterface>(ComputeGraph);
+			UPCGCustomKernelDataInterface* KernelDI = InOutContext.NewObject_AnyThread<UPCGCustomKernelDataInterface>(InOutComputeGraph);
 			KernelDI->Settings = Settings;
 
-			const int32 KernelDIIndex = ComputeGraph->DataInterfaces.Num();
-			ComputeGraph->DataInterfaces.Add(KernelDI);
+			const int32 KernelDIIndex = InOutComputeGraph->DataInterfaces.Num();
+			InOutComputeGraph->DataInterfaces.Add(KernelDI);
 			InputDataInterfaceIndexAndPin.Emplace(KernelDIIndex, NAME_None);
 		}
 
 		// Additional DIs created by settings.
 		{
 			TArray<TObjectPtr<UComputeDataInterface>> AdditionalInputDIs, AdditionalOutputDIs;
-			Settings->CreateAdditionalInputDataInterfaces(InOutContext, ComputeGraph, AdditionalInputDIs);
-			Settings->CreateAdditionalOutputDataInterfaces(InOutContext, ComputeGraph, AdditionalOutputDIs);
+			Settings->CreateAdditionalInputDataInterfaces(InOutContext, InOutComputeGraph, AdditionalInputDIs);
+			Settings->CreateAdditionalOutputDataInterfaces(InOutContext, InOutComputeGraph, AdditionalOutputDIs);
 
-			auto AddAdditionalDataInterface = [ComputeGraph](TObjectPtr<UComputeDataInterface> DataInterface) -> int32
+			auto AddAdditionalDataInterface = [InOutComputeGraph](TObjectPtr<UComputeDataInterface> DataInterface) -> int32
 			{
 				if (ensure(DataInterface))
 				{
-					const int32 DataInterfaceIndex = ComputeGraph->DataInterfaces.Num();
-					ComputeGraph->DataInterfaces.Add(DataInterface);
+					const int32 DataInterfaceIndex = InOutComputeGraph->DataInterfaces.Num();
+					InOutComputeGraph->DataInterfaces.Add(DataInterface);
 
 					return DataInterfaceIndex;
 				}
@@ -1026,18 +1036,18 @@ UPCGComputeGraph* FPCGGraphCompilerGPU::CompileComputeGraph(
 		// TODO skip depending on target platform.
 		if (Settings->bPrintShaderDebugValues)
 		{
-			UPCGDebugDataInterface* DebugDI = InOutContext.NewObject_AnyThread<UPCGDebugDataInterface>(ComputeGraph);
+			UPCGDebugDataInterface* DebugDI = InOutContext.NewObject_AnyThread<UPCGDebugDataInterface>(InOutComputeGraph);
 			DebugDI->SetDebugBufferSize(Settings->DebugBufferSize);
 
-			const int32 DebugDIIndex = ComputeGraph->DataInterfaces.Num();
-			ComputeGraph->DataInterfaces.Add(DebugDI);
+			const int32 DebugDIIndex = InOutComputeGraph->DataInterfaces.Num();
+			InOutComputeGraph->DataInterfaces.Add(DebugDI);
 			OutputDataInterfaceIndexAndPin.Emplace(DebugDIIndex, NAME_None);
 		}
 
-		ComputeGraph->bLogDataDescriptions |= Settings->bDumpDataDescriptions;
+		InOutComputeGraph->bLogDataDescriptions |= Settings->bDumpDataDescriptions;
 
 		// Now that all data interfaces added, create the (trivial) binding mapping. All map to primary binding, index 0.
-		ComputeGraph->DataInterfaceToBinding.SetNumZeroed(ComputeGraph->DataInterfaces.Num());
+		InOutComputeGraph->DataInterfaceToBinding.SetNumZeroed(InOutComputeGraph->DataInterfaces.Num());
 
 		struct FInterfaceBinding
 		{
@@ -1056,14 +1066,14 @@ UPCGComputeGraph* FPCGGraphCompilerGPU::CompileComputeGraph(
 
 		FKernelWithDataBindings KernelWithBindings;
 
-		KernelWithBindings.Kernel = InOutContext.NewObject_AnyThread<UComputeKernel>(ComputeGraph);
-		const int KernelIndex = ComputeGraph->KernelInvocations.Num();
-		ComputeGraph->KernelInvocations.Add(KernelWithBindings.Kernel);
-		ComputeGraph->KernelToNode.Add(Node);
+		KernelWithBindings.Kernel = InOutContext.NewObject_AnyThread<UComputeKernel>(InOutComputeGraph);
+		const int KernelIndex = InOutComputeGraph->KernelInvocations.Num();
+		InOutComputeGraph->KernelInvocations.Add(KernelWithBindings.Kernel);
+		InOutComputeGraph->KernelToNode.Add(Node);
 
-		auto SetupAllInputBindings = [&KernelWithBindings, ComputeGraph](const TPair<int, FName>& InDataInterfaceIndex)
+		auto SetupAllInputBindings = [&KernelWithBindings, InOutComputeGraph](const TPair<int, FName>& InDataInterfaceIndex)
 		{
-			const UComputeDataInterface* DataInterface = ComputeGraph->DataInterfaces[InDataInterfaceIndex.Key];
+			const UComputeDataInterface* DataInterface = InOutComputeGraph->DataInterfaces[InDataInterfaceIndex.Key];
 			TArray<FShaderFunctionDefinition> Functions;
 			DataInterface->GetSupportedInputs(Functions);
 
@@ -1077,9 +1087,9 @@ UPCGComputeGraph* FPCGGraphCompilerGPU::CompileComputeGraph(
 			}
 		};
 
-		auto SetupAllOutputBindings = [&KernelWithBindings, ComputeGraph](const TPair<int, FName>& InDataInterfaceIndex)
+		auto SetupAllOutputBindings = [&KernelWithBindings, InOutComputeGraph](const TPair<int, FName>& InDataInterfaceIndex)
 		{
-			const UComputeDataInterface* DataInterface = ComputeGraph->DataInterfaces[InDataInterfaceIndex.Key];
+			const UComputeDataInterface* DataInterface = InOutComputeGraph->DataInterfaces[InDataInterfaceIndex.Key];
 			TArray<FShaderFunctionDefinition> Functions;
 			DataInterface->GetSupportedOutputs(Functions);
 
@@ -1112,7 +1122,7 @@ UPCGComputeGraph* FPCGGraphCompilerGPU::CompileComputeGraph(
 			KernelSource->EntryPoint = TEXT("Main");
 			KernelSource->GroupSize = FIntVector(PCGComputeConstants::THREAD_GROUP_SIZE, 1, 1);
 
-			KernelSource->SetSource(Settings->GetCookedKernelSource(ComputeGraph->GlobalAttributeLookupTable));
+			KernelSource->SetSource(Settings->GetCookedKernelSource(InOutComputeGraph->GlobalAttributeLookupTable));
 
 			if (Settings->bDumpCookedHLSL)
 			{
@@ -1123,7 +1133,7 @@ UPCGComputeGraph* FPCGGraphCompilerGPU::CompileComputeGraph(
 			{
 				UE_LOG(LogPCG, Warning, TEXT("ATTRIBUTE LOOK-UP TABLE [%s]"), *Settings->GetDefaultNodeTitle().ToString());
 
-				for (const TPair<FName, FPCGKernelAttributeIDAndType>& Pair : ComputeGraph->GlobalAttributeLookupTable)
+				for (const TPair<FName, FPCGKernelAttributeIDAndType>& Pair : InOutComputeGraph->GlobalAttributeLookupTable)
 				{
 					UE_LOG(LogPCG, Warning, TEXT("%s: %s (%s)"),
 						*FString::FromInt(Pair.Value.Id),
@@ -1166,7 +1176,7 @@ UPCGComputeGraph* FPCGGraphCompilerGPU::CompileComputeGraph(
 			}
 		}
 
-		auto AddAllEdgesForKernel = [&KernelWithBindings, ComputeGraph](int32 InKernelIndex, bool bInEdgesAreInputs)
+		auto AddAllEdgesForKernel = [&KernelWithBindings, InOutComputeGraph](int32 InKernelIndex, bool bInEdgesAreInputs)
 		{
 			TArray<FInterfaceBinding>& Bindings = bInEdgesAreInputs ? KernelWithBindings.InputDataBindings : KernelWithBindings.OutputDataBindings;
 
@@ -1176,15 +1186,15 @@ UPCGComputeGraph* FPCGGraphCompilerGPU::CompileComputeGraph(
 			{
 				FInterfaceBinding& Binding = Bindings[BindingIndex];
 
-				FComputeGraphEdge& Edge = ComputeGraph->GraphEdges.Emplace_GetRef();
+				FComputeGraphEdge& Edge = InOutComputeGraph->GraphEdges.Emplace_GetRef();
 				Edge.KernelIndex = InKernelIndex;
 				Edge.KernelBindingIndex = BindingIndex;
-				Edge.DataInterfaceIndex = ComputeGraph->DataInterfaces.IndexOfByPredicate([&Binding](const UComputeDataInterface* In) { return Binding.DataInterface == In; });
+				Edge.DataInterfaceIndex = InOutComputeGraph->DataInterfaces.IndexOfByPredicate([&Binding](const UComputeDataInterface* In) { return Binding.DataInterface == In; });
 				check(Edge.DataInterfaceIndex != INDEX_NONE);
 				Edge.DataInterfaceBindingIndex = Binding.DataInterfaceBindingIndex;
 				Edge.bKernelInput = bInEdgesAreInputs;
 
-				UComputeDataInterface* DataInterface = ComputeGraph->DataInterfaces[Edge.DataInterfaceIndex];
+				UComputeDataInterface* DataInterface = InOutComputeGraph->DataInterfaces[Edge.DataInterfaceIndex];
 				check(DataInterface);
 
 				if (Binding.PinLabel != NAME_None)
@@ -1221,24 +1231,18 @@ UPCGComputeGraph* FPCGGraphCompilerGPU::CompileComputeGraph(
 		{
 			const FPCGTaskId TaskId = Mapping.Get<0>().Get<0>();
 
-			if (const UPCGNode* Node = InOutCompiledTasks[TaskId].Node)
+			if (const UPCGNode* Node = InCompiledTasks[TaskId].Node)
 			{
 				const FName& OriginalLabel = Mapping.Get<0>().Get<1>();
 
 				if (const UPCGPin* Pin = Node->GetInputPin(OriginalLabel))
 				{
 					const FName& VirtualLabel = Mapping.Get<1>();
-					ComputeGraph->InputPinLabelAliases.FindOrAdd(Pin) = VirtualLabel;
+					InOutComputeGraph->InputPinLabelAliases.FindOrAdd(Pin) = VirtualLabel;
 				}
 			}
 		}
 	}
-
-	// Remove empty strings (at execution time empty string is placed in table slot 0).
-	// NOTE: This can scramble order but order is not important at this stage (key values created at execution time in data binding).
-	ComputeGraph->StringTable.RemoveAllSwap([](const FString& InString) { return InString.IsEmpty(); });
-
-	return ComputeGraph;
 }
 
 void FPCGGraphCompilerGPU::CreateGPUNodes(FPCGGraphCompiler& InOutCompiler, UPCGGraph* InGraph, uint32 InGridSize, TArray<FPCGGraphTask>& InOutCompiledTasks)
@@ -1313,8 +1317,7 @@ void FPCGGraphCompilerGPU::CreateGPUNodes(FPCGGraphCompiler& InOutCompiler, UPCG
 			OutputCPUPinToVirtualPin[ComputeGraphIndex]);
 	}
 
-	// Do actual collapsing now, one subset at a time. Each collapse will do all fixup of task ids? That will invalidate
-	// ids in NodeSubsetsToConvertToCFGraph, so may need remap table. But can ignore this for now.
+	// Build each compute graph.
 	for (uint32 ComputeGraphIndex = 0; ComputeGraphIndex < NumComputeGraphs; ++ComputeGraphIndex)
 	{
 		TSet<FPCGTaskId>& NodeSubsetToConvertToCFGraph = NodeSubsetsToConvertToCFGraph[ComputeGraphIndex];
@@ -1325,7 +1328,7 @@ void FPCGGraphCompilerGPU::CreateGPUNodes(FPCGGraphCompiler& InOutCompiler, UPCG
 			continue;
 		}
 
-		// Add a new compute graph task. Then kill the original GPU tasks.
+		// Add a new compute graph task. The original GPU tasks will then be culled later below.
 		const FPCGTaskId ComputeGraphTaskId = InOutCompiledTasks.Num();
 		FPCGGraphTask& ComputeGraphTask = InOutCompiledTasks.Emplace_GetRef();
 		ComputeGraphTask.NodeId = ComputeGraphTaskId;
@@ -1339,21 +1342,20 @@ void FPCGGraphCompilerGPU::CreateGPUNodes(FPCGGraphCompiler& InOutCompiler, UPCG
 		}
 
 		// Wire in the compute graph task, side by side with the individual GPU tasks, which will be culled below.
-		WireGPUGraphNode(
+		WireComputeGraphTask(
 			ComputeGraphTaskId,
 			NodeSubsetToConvertToCFGraph,
 			InOutCompiledTasks,
 			TaskSuccessors,
-			OriginalToVirtualPin[ComputeGraphIndex],
-			OutputCPUPinToVirtualPin[ComputeGraphIndex]);
+			OriginalToVirtualPin[ComputeGraphIndex]);
 
 		// Generate a compute graph from all of the individual GPU tasks.
-		BuildGPUGraphTask(
+		BuildComputeGraphTask(
 			Context,
+			ComputeGraphTaskId,
 			InGraph,
 			InGridSize,
 			ComputeGraphIndex,
-			ComputeGraphTaskId,
 			NodeSubsetToConvertToCFGraph,
 			GPUCompatibleTaskIds,
 			TaskSuccessors,
diff --git a/Engine/Plugins/PCG/Source/PCG/Private/Graph/PCGGraphCompilerGPU.h b/Engine/Plugins/PCG/Source/PCG/Private/Graph/PCGGraphCompilerGPU.h
--- a/Engine/Plugins/PCG/Source/PCG/Private/Graph/PCGGraphCompilerGPU.h
+++ b/Engine/Plugins/PCG/Source/PCG/Private/Graph/PCGGraphCompilerGPU.h
@@ -5,6 +5,7 @@
 #include "PCGCommon.h"
 
 class FPCGGraphCompiler;
+class UPCGComputeDataInterface;
 class UPCGComputeGraph;
 class UPCGGraph;
 class UPCGPin;
@@ -23,6 +24,7 @@ private:
 	using FNodePin = TTuple<FPCGTaskId, /*Pin label*/FName, /*Pin is input*/bool>;
 	using FOriginalToVirtualPin = TMap<FNodePin, /*Virtual pin label*/FName>;
 	using FTaskToSuccessors = TMap<FPCGTaskId, TArray<FPCGTaskId>>;
+	using FPinToDataInterface = TMap<TPair</*Node task*/FPCGTaskId, /*Pin label*/FName>, UPCGComputeDataInterface*>;
 
 	/** Identifies connected sets of GPU nodes, giving each a non - zero ID value. */
 	static void LabelConnectedGPUNodeIslands(
@@ -52,22 +54,21 @@ private:
 		FOriginalToVirtualPin& OutOriginalToVirtualPin,
 		TMap<TSoftObjectPtr<const UPCGPin>, FName>& OutOutputCPUPinToVirtualPin);
 
-	/** Wires in a compute graph element alongside each set of GPU compatible nodes. The tasks for each node will be culled later. */
-	static void WireGPUGraphNode(
+	/** Wires in a compute graph task alongside each set of GPU compatible nodes. The tasks for each node will be culled later. */
+	static void WireComputeGraphTask(
 		FPCGTaskId InGPUGraphTaskId,
 		const TSet<FPCGTaskId>& InCollapsedTasks,
 		TArray<FPCGGraphTask>& InOutCompiledTasks,
 		const TMap<FPCGTaskId, TArray<FPCGTaskId>>& InTaskSuccessors,
-		const FOriginalToVirtualPin& InOriginalToVirtualPin,
-		const TMap<TSoftObjectPtr<const UPCGPin>, FName>& InOutputCPUPinToVirtualPin);
+		const FOriginalToVirtualPin& InOriginalToVirtualPin);
 	
-	/** Compiles a compute graph. */
-	static void BuildGPUGraphTask(
+	/** Creates compute graph with all resources and data interfaces and sets up the compute graph task. */
+	static void BuildComputeGraphTask(
 		FPCGGPUCompilationContext& InOutContext,
+		FPCGTaskId InGPUGraphTaskId,
 		UPCGGraph* InGraph,
 		uint32 InGridSize,
 		uint32 InComputeGraphIndex,
-		FPCGTaskId InGPUGraphTaskId,
 		const TSet<FPCGTaskId>& InCollapsedTasks,
 		const TSet<FPCGTaskId>& InAllGPUCompatibleTasks,
 		const FTaskToSuccessors& InTaskSuccessors,
@@ -75,14 +76,26 @@ private:
 		const FOriginalToVirtualPin& InOriginalToVirtualPin,
 		const TMap<TSoftObjectPtr<const UPCGPin>, FName>& InOutputCPUPinToVirtualPin);
 
-	static UPCGComputeGraph* CompileComputeGraph(
+	/** Creates and configures all data interfaces for the graph. */
+	static void CreateDataInterfaces(
 		FPCGGPUCompilationContext& InOutContext,
-		UPCGGraph* InGraph,
+		const TArray<FPCGGraphTask>& InCompiledTasks,
 		const TSet<FPCGTaskId>& InCollapsedTasks,
 		const TSet<FPCGTaskId>& InAllGPUCompatibleTasks,
 		const FTaskToSuccessors& InTaskSuccessors,
-		TArray<FPCGGraphTask>& InOutCompiledTasks,
 		const FOriginalToVirtualPin& InOriginalToVirtualPin,
-		const TMap<TSoftObjectPtr<const UPCGPin>, FName>& InOutputCPUPinToVirtualPin);
+		UPCGComputeGraph* InOutComputeGraph, // In
+		FPinToDataInterface& InOutPinToDataInterface);
+
+	/** Builds the compute graph from the input tasks and data interfaces. */
+	static void CompileComputeGraph(
+		FPCGGPUCompilationContext& InOutContext,
+		const TArray<FPCGGraphTask>& InCompiledTasks,
+		const TSet<FPCGTaskId>& InCollapsedTasks,
+		const TSet<FPCGTaskId>& InAllGPUCompatibleTasks,
+		const FTaskToSuccessors& InTaskSuccessors,
+		const FOriginalToVirtualPin& InOriginalToVirtualPin,
+		const FPinToDataInterface& InOutputPinToDataInterface,
+		UPCGComputeGraph* InOutComputeGraph);
 };
 #endif // WITH_EDITOR
