[PCG] Make unpacking GPU data collections significantly faster
* Fixed an issue where maximum buffer sizes were cut in half
* Improved logging when exceeding maximum buffer size
* Fixed crash when packing transforms with non-normalized rotations

#rnx

diff --git a/Engine/Plugins/PCG/Source/PCG/Private/Compute/DataInterfaces/PCGDataCollectionDataInterface.cpp b/Engine/Plugins/PCG/Source/PCG/Private/Compute/DataInterfaces/PCGDataCollectionDataInterface.cpp
--- a/Engine/Plugins/PCG/Source/PCG/Private/Compute/DataInterfaces/PCGDataCollectionDataInterface.cpp
+++ b/Engine/Plugins/PCG/Source/PCG/Private/Compute/DataInterfaces/PCGDataCollectionDataInterface.cpp
@@ -550,7 +550,7 @@ bool UPCGDataCollectionDataProvider::ProcessReadBackData(FPCGComputeGraphContext
 	}
 
 	FPCGDataCollection DataFromGPU;
-	const EPCGUnpackDataCollectionResult Result = PinDesc.UnpackDataCollection(RawReadbackData, OutputPinLabelAlias, InContext->DataBinding->GetStringTable(), DataFromGPU);
+	const EPCGUnpackDataCollectionResult Result = PinDesc.UnpackDataCollection(InContext, RawReadbackData, OutputPinLabelAlias, InContext->DataBinding->GetStringTable(), DataFromGPU);
 
 	if (Result == EPCGUnpackDataCollectionResult::DataMismatch)
 	{
diff --git a/Engine/Plugins/PCG/Source/PCG/Private/Compute/DataInterfaces/PCGDataCollectionDataInterface.h b/Engine/Plugins/PCG/Source/PCG/Private/Compute/DataInterfaces/PCGDataCollectionDataInterface.h
--- a/Engine/Plugins/PCG/Source/PCG/Private/Compute/DataInterfaces/PCGDataCollectionDataInterface.h
+++ b/Engine/Plugins/PCG/Source/PCG/Private/Compute/DataInterfaces/PCGDataCollectionDataInterface.h
@@ -126,7 +126,7 @@ public:
 	FReadbackCallback AsyncReadbackCallback_RenderThread;
 
 protected:
-	int SizeBytes = -1;
+	uint32 SizeBytes = 0;
 
 	EPCGReadbackMode ReadbackMode;
 
diff --git a/Engine/Plugins/PCG/Source/PCG/Private/Compute/PCGComputeCommon.cpp b/Engine/Plugins/PCG/Source/PCG/Private/Compute/PCGComputeCommon.cpp
--- a/Engine/Plugins/PCG/Source/PCG/Private/Compute/PCGComputeCommon.cpp
+++ b/Engine/Plugins/PCG/Source/PCG/Private/Compute/PCGComputeCommon.cpp
@@ -304,7 +304,13 @@ namespace PCGComputeHelpers
 
 		// If buffer size exceeds a proportion of total graphics memory, then it is deemed too large. Using this as a heuristic as there
 		// is no RHI API to obtain available graphics memory outside of D3D12.
-		const uint64 BudgetBytes = static_cast<uint64>(TextureMemStats.TotalGraphicsMemory * CVarMaxGPUBufferSizeProportion.GetValueOnAnyThread());
+		const uint64 GPUBudgetBytes = static_cast<uint64>(TextureMemStats.TotalGraphicsMemory * CVarMaxGPUBufferSizeProportion.GetValueOnAnyThread());
+
+		// Buffer size also cannot exceed the max size of a uint32, otherwise the size will get truncated and other systems will break.
+		// TODO: This limits the maximum number of points to around 46 million. Support uint64 max instead of uint32 to get up to 2 billion points.
+		const uint64 MaxBufferSize = TNumericLimits<uint32>::Max();
+		const uint64 BudgetBytes = FMath::Min(GPUBudgetBytes, MaxBufferSize);
+
 		const bool bBufferTooLarge = TextureMemStats.TotalGraphicsMemory > 0 && InBufferSizeBytes > BudgetBytes;
 
 		if (bBufferTooLarge && bInLogError)
diff --git a/Engine/Plugins/PCG/Source/PCG/Private/Compute/PCGDataForGPU.cpp b/Engine/Plugins/PCG/Source/PCG/Private/Compute/PCGDataForGPU.cpp
--- a/Engine/Plugins/PCG/Source/PCG/Private/Compute/PCGDataForGPU.cpp
+++ b/Engine/Plugins/PCG/Source/PCG/Private/Compute/PCGDataForGPU.cpp
@@ -14,6 +14,7 @@
 #include "Data/PCGPointData.h"
 #include "Helpers/PCGAsync.h"
 #include "Metadata/PCGMetadata.h"
+#include "Metadata/Accessors/PCGAttributeAccessorHelpers.h"
 
 #include "Async/ParallelFor.h"
 
@@ -61,7 +62,7 @@ namespace PCGDataForGPUHelpers
 			return EPCGKernelAttributeType::Quat;
 		case EPCGMetadataTypes::Transform:
 			return EPCGKernelAttributeType::Transform;
-		case EPCGMetadataTypes::SoftObjectPath:
+		case EPCGMetadataTypes::SoftObjectPath: // TODO: This collapses all StringKey types into String attributes, meaning we'll lose the original CPU type when doing readback.
 		case EPCGMetadataTypes::SoftClassPath:
 		case EPCGMetadataTypes::String:
 			return EPCGKernelAttributeType::StringKey;
@@ -202,8 +203,14 @@ namespace PCGDataForGPUHelpers
 		case PCG::Private::MetadataTypes<FTransform>::Id:
 		{
 			const FPCGMetadataAttribute<FTransform>* Attribute = static_cast<const FPCGMetadataAttribute<FTransform>*>(InAttributeBase);
-			const FTransform Value = Attribute->GetValue(ValueKey);
-			const FMatrix Matrix = Value.ToMatrixWithScale();
+			const FTransform Transform = Attribute->GetValue(ValueKey);
+
+			const bool bIsRotationNormalized = Transform.IsRotationNormalized();
+			ensureMsgf(bIsRotationNormalized, TEXT("Tried to pack transform for GPU data collection, but the transform's rotation is not normalized. Using identity instead."));
+
+			// Note: ToMatrixWithScale() crashes if the transform is not normalized.
+			const FMatrix Matrix = bIsRotationNormalized ? Transform.ToMatrixWithScale() : FMatrix::Identity;
+
 			check(StrideBytes == 64);
 			OutPackedDataCollection[InOutAddressUints++] = FMath::AsUInt(static_cast<float>(Matrix.M[0][0]));
 			OutPackedDataCollection[InOutAddressUints++] = FMath::AsUInt(static_cast<float>(Matrix.M[0][1]));
@@ -221,6 +228,7 @@ namespace PCGDataForGPUHelpers
 			OutPackedDataCollection[InOutAddressUints++] = FMath::AsUInt(static_cast<float>(Matrix.M[3][1]));
 			OutPackedDataCollection[InOutAddressUints++] = FMath::AsUInt(static_cast<float>(Matrix.M[3][2]));
 			OutPackedDataCollection[InOutAddressUints++] = FMath::AsUInt(static_cast<float>(Matrix.M[3][3]));
+
 			break;
 		}
 		case PCG::Private::MetadataTypes<FString>::Id:
@@ -319,148 +327,214 @@ namespace PCGDataForGPUHelpers
 		}
 	}
 
-	bool UnpackAttributeHelper(const void* InPackedData, const TArray<FString>& InStringTable, uint32 ElementIndex, FPCGMetadataAttributeBase* AttributeBase, const FPCGKernelAttributeDesc& AttributeDesc, PCGMetadataEntryKey EntryKey)
+	bool UnpackAttributeHelper(FPCGContext* InContext, const void* InPackedData, const FPCGKernelAttributeDesc& InAttributeDesc, const TArray<FString>& InStringTable, uint32 InAddressUints, uint32 InNumElements, UPCGData* OutData)
 	{
-		check(InPackedData && AttributeBase);
+		check(InPackedData);
+		check(OutData);
 
 		const float* DataAsFloat = static_cast<const float*>(InPackedData);
 		const int32* DataAsInt = static_cast<const int32*>(InPackedData);
 
-		switch (AttributeDesc.Type)
+		FPCGAttributePropertyOutputSelector Selector = FPCGAttributePropertySelector::CreateAttributeSelector<FPCGAttributePropertyOutputSelector>(InAttributeDesc.Name);
+
+		switch (InAttributeDesc.Type)
 		{
 		case EPCGKernelAttributeType::Bool:
 		{
-			FPCGMetadataAttribute<bool>* Attribute = static_cast<FPCGMetadataAttribute<bool>*>(AttributeBase);
+			TArray<bool> Values;
+			Values.SetNumUninitialized(InNumElements);
 
-			const bool Value = static_cast<bool>(DataAsFloat[ElementIndex]);
-			Attribute->SetValue(EntryKey, Value);
+			for (uint32 ElementIndex = 0; ElementIndex < InNumElements; ++ElementIndex)
+			{
+				const uint32 PackedElementIndex = InAddressUints + ElementIndex;
+				Values[ElementIndex] = static_cast<bool>(DataAsFloat[PackedElementIndex]);
+			}
+
+			ensure(PCGAttributeAccessorHelpers::WriteAllValues<bool>(OutData, Selector, Values, /*SourceSelector=*/nullptr, InContext));
 			break;
 		}
 		case EPCGKernelAttributeType::Int:
 		{
-			FPCGMetadataAttribute<int>* Attribute = static_cast<FPCGMetadataAttribute<int>*>(AttributeBase);
+			TArray<int> Values;
+			Values.SetNumUninitialized(InNumElements);
+
+			for (uint32 ElementIndex = 0; ElementIndex < InNumElements; ++ElementIndex)
+			{
+				const uint32 PackedElementIndex = InAddressUints + ElementIndex;
+				Values[ElementIndex] = DataAsInt[PackedElementIndex];
+			}
 
-			const int32 Value = DataAsInt[ElementIndex];
-			Attribute->SetValue(EntryKey, Value);
+			ensure(PCGAttributeAccessorHelpers::WriteAllValues<int>(OutData, Selector, Values, /*SourceSelector=*/nullptr, InContext));
 			break;
 		}
 		case EPCGKernelAttributeType::Float:
 		{
-			FPCGMetadataAttribute<float>* Attribute = static_cast<FPCGMetadataAttribute<float>*>(AttributeBase);
+			TArray<float> Values;
+			Values.SetNumUninitialized(InNumElements);
 
-			const float Value = DataAsFloat[ElementIndex];
-			Attribute->SetValue(EntryKey, Value);
+			for (uint32 ElementIndex = 0; ElementIndex < InNumElements; ++ElementIndex)
+			{
+				const uint32 PackedElementIndex = InAddressUints + ElementIndex;
+				Values[ElementIndex] = DataAsFloat[PackedElementIndex];
+			}
+
+			ensure(PCGAttributeAccessorHelpers::WriteAllValues<float>(OutData, Selector, Values, /*SourceSelector=*/nullptr, InContext));
 			break;
 		}
 		case EPCGKernelAttributeType::Float2:
 		{
-			FPCGMetadataAttribute<FVector2D>* Attribute = static_cast<FPCGMetadataAttribute<FVector2D>*>(AttributeBase);
+			TArray<FVector2D> Values;
+			Values.SetNumUninitialized(InNumElements);
+
+			for (uint32 ElementIndex = 0; ElementIndex < InNumElements; ++ElementIndex)
+			{
+				const uint32 PackedElementIndex = InAddressUints + ElementIndex * 2;
+				Values[ElementIndex].X = DataAsFloat[PackedElementIndex + 0];
+				Values[ElementIndex].Y = DataAsFloat[PackedElementIndex + 1];
+			}
 
-			FVector2D Value;
-			Value.X = DataAsFloat[ElementIndex + 0];
-			Value.Y = DataAsFloat[ElementIndex + 1];
-			Attribute->SetValue(EntryKey, Value);
+			ensure(PCGAttributeAccessorHelpers::WriteAllValues<FVector2D>(OutData, Selector, Values, /*SourceSelector=*/nullptr, InContext));
 			break;
 		}
 		case EPCGKernelAttributeType::Float3:
 		{
-			FPCGMetadataAttribute<FVector>* Attribute = static_cast<FPCGMetadataAttribute<FVector>*>(AttributeBase);
+			TArray<FVector> Values;
+			Values.SetNumUninitialized(InNumElements);
 
-			FVector Value;
-			Value.X = DataAsFloat[ElementIndex + 0];
-			Value.Y = DataAsFloat[ElementIndex + 1];
-			Value.Z = DataAsFloat[ElementIndex + 2];
-			Attribute->SetValue(EntryKey, Value);
+			for (uint32 ElementIndex = 0; ElementIndex < InNumElements; ++ElementIndex)
+			{
+				const uint32 PackedElementIndex = InAddressUints + ElementIndex * 3;
+				Values[ElementIndex].X = DataAsFloat[PackedElementIndex + 0];
+				Values[ElementIndex].Y = DataAsFloat[PackedElementIndex + 1];
+				Values[ElementIndex].Z = DataAsFloat[PackedElementIndex + 2];
+			}
+
+			ensure(PCGAttributeAccessorHelpers::WriteAllValues<FVector>(OutData, Selector, Values, /*SourceSelector=*/nullptr, InContext));
 			break;
 		}
 		case EPCGKernelAttributeType::Float4:
 		{
-			FPCGMetadataAttribute<FVector4>* Attribute = static_cast<FPCGMetadataAttribute<FVector4>*>(AttributeBase);
-
-			FVector4 Value;
-			Value.X = DataAsFloat[ElementIndex + 0];
-			Value.Y = DataAsFloat[ElementIndex + 1];
-			Value.Z = DataAsFloat[ElementIndex + 2];
-			Value.W = DataAsFloat[ElementIndex + 3];
-			Attribute->SetValue(EntryKey, Value);
+			TArray<FVector4> Values;
+			Values.SetNumUninitialized(InNumElements);
+
+			for (uint32 ElementIndex = 0; ElementIndex < InNumElements; ++ElementIndex)
+			{
+				const uint32 PackedElementIndex = InAddressUints + ElementIndex * 4;
+				Values[ElementIndex].X = DataAsFloat[PackedElementIndex + 0];
+				Values[ElementIndex].Y = DataAsFloat[PackedElementIndex + 1];
+				Values[ElementIndex].Z = DataAsFloat[PackedElementIndex + 2];
+				Values[ElementIndex].W = DataAsFloat[PackedElementIndex + 3];
+			}
+
+			ensure(PCGAttributeAccessorHelpers::WriteAllValues<FVector4>(OutData, Selector, Values, /*SourceSelector=*/nullptr, InContext));
 			break;
 		}
 		case EPCGKernelAttributeType::Rotator:
 		{
-			FPCGMetadataAttribute<FRotator>* Attribute = static_cast<FPCGMetadataAttribute<FRotator>*>(AttributeBase);
+			TArray<FRotator> Values;
+			Values.SetNumUninitialized(InNumElements);
+
+			for (uint32 ElementIndex = 0; ElementIndex < InNumElements; ++ElementIndex)
+			{
+				const uint32 PackedElementIndex = InAddressUints + ElementIndex * 3;
+				Values[ElementIndex].Pitch = DataAsFloat[PackedElementIndex + 0];
+				Values[ElementIndex].Yaw = DataAsFloat[PackedElementIndex + 1];
+				Values[ElementIndex].Roll = DataAsFloat[PackedElementIndex + 2];
+			}
 
-			FRotator Value;
-			Value.Pitch = DataAsFloat[ElementIndex + 0];
-			Value.Yaw = DataAsFloat[ElementIndex + 1];
-			Value.Roll = DataAsFloat[ElementIndex + 2];
-			Attribute->SetValue(EntryKey, Value);
+			ensure(PCGAttributeAccessorHelpers::WriteAllValues<FRotator>(OutData, Selector, Values, /*SourceSelector=*/nullptr, InContext));
 			break;
 		}
 		case EPCGKernelAttributeType::Quat:
 		{
-			FPCGMetadataAttribute<FQuat>* Attribute = static_cast<FPCGMetadataAttribute<FQuat>*>(AttributeBase);
-
-			FQuat Value;
-			Value.X = DataAsFloat[ElementIndex + 0];
-			Value.Y = DataAsFloat[ElementIndex + 1];
-			Value.Z = DataAsFloat[ElementIndex + 2];
-			Value.W = DataAsFloat[ElementIndex + 3];
-			Attribute->SetValue(EntryKey, Value);
+			TArray<FQuat> Values;
+			Values.SetNumUninitialized(InNumElements);
+
+			for (uint32 ElementIndex = 0; ElementIndex < InNumElements; ++ElementIndex)
+			{
+				const uint32 PackedElementIndex = InAddressUints + ElementIndex * 4;
+				Values[ElementIndex].X = DataAsFloat[PackedElementIndex + 0];
+				Values[ElementIndex].Y = DataAsFloat[PackedElementIndex + 1];
+				Values[ElementIndex].Z = DataAsFloat[PackedElementIndex + 2];
+				Values[ElementIndex].W = DataAsFloat[PackedElementIndex + 3];
+			}
+
+			ensure(PCGAttributeAccessorHelpers::WriteAllValues<FQuat>(OutData, Selector, Values, /*SourceSelector=*/nullptr, InContext));
 			break;
 		}
 		case EPCGKernelAttributeType::Transform:
 		{
-			FPCGMetadataAttribute<FTransform>* Attribute = static_cast<FPCGMetadataAttribute<FTransform>*>(AttributeBase);
+			TArray<FTransform> Values;
+			Values.SetNumUninitialized(InNumElements);
 
 			FMatrix Matrix;
-			Matrix.M[0][0] = DataAsFloat[ElementIndex + 0];
-			Matrix.M[0][1] = DataAsFloat[ElementIndex + 1];
-			Matrix.M[0][2] = DataAsFloat[ElementIndex + 2];
-			Matrix.M[0][3] = DataAsFloat[ElementIndex + 3];
-			Matrix.M[1][0] = DataAsFloat[ElementIndex + 4];
-			Matrix.M[1][1] = DataAsFloat[ElementIndex + 5];
-			Matrix.M[1][2] = DataAsFloat[ElementIndex + 6];
-			Matrix.M[1][3] = DataAsFloat[ElementIndex + 7];
-			Matrix.M[2][0] = DataAsFloat[ElementIndex + 8];
-			Matrix.M[2][1] = DataAsFloat[ElementIndex + 9];
-			Matrix.M[2][2] = DataAsFloat[ElementIndex + 10];
-			Matrix.M[2][3] = DataAsFloat[ElementIndex + 11];
-			Matrix.M[3][0] = DataAsFloat[ElementIndex + 12];
-			Matrix.M[3][1] = DataAsFloat[ElementIndex + 13];
-			Matrix.M[3][2] = DataAsFloat[ElementIndex + 14];
-			Matrix.M[3][3] = DataAsFloat[ElementIndex + 15];
-
-			const FTransform Value(Matrix);
-			Attribute->SetValue(EntryKey, Value);
+
+			for (uint32 ElementIndex = 0; ElementIndex < InNumElements; ++ElementIndex)
+			{
+				const uint32 PackedElementIndex = InAddressUints + ElementIndex * 16;
+
+				Matrix.M[0][0] = DataAsFloat[PackedElementIndex + 0];
+				Matrix.M[0][1] = DataAsFloat[PackedElementIndex + 1];
+				Matrix.M[0][2] = DataAsFloat[PackedElementIndex + 2];
+				Matrix.M[0][3] = DataAsFloat[PackedElementIndex + 3];
+				Matrix.M[1][0] = DataAsFloat[PackedElementIndex + 4];
+				Matrix.M[1][1] = DataAsFloat[PackedElementIndex + 5];
+				Matrix.M[1][2] = DataAsFloat[PackedElementIndex + 6];
+				Matrix.M[1][3] = DataAsFloat[PackedElementIndex + 7];
+				Matrix.M[2][0] = DataAsFloat[PackedElementIndex + 8];
+				Matrix.M[2][1] = DataAsFloat[PackedElementIndex + 9];
+				Matrix.M[2][2] = DataAsFloat[PackedElementIndex + 10];
+				Matrix.M[2][3] = DataAsFloat[PackedElementIndex + 11];
+				Matrix.M[3][0] = DataAsFloat[PackedElementIndex + 12];
+				Matrix.M[3][1] = DataAsFloat[PackedElementIndex + 13];
+				Matrix.M[3][2] = DataAsFloat[PackedElementIndex + 14];
+				Matrix.M[3][3] = DataAsFloat[PackedElementIndex + 15];
+
+				new (&Values[ElementIndex]) FTransform(Matrix);
+			}
+
+			ensure(PCGAttributeAccessorHelpers::WriteAllValues<FTransform>(OutData, Selector, Values, /*SourceSelector=*/nullptr, InContext));
 			break;
 		}
 		case EPCGKernelAttributeType::StringKey:
 		{
-			FPCGMetadataAttribute<FString>* Attribute = static_cast<FPCGMetadataAttribute<FString>*>(AttributeBase);
-			
-			const int32 StringKey = DataAsInt[ElementIndex];
-			if (InStringTable.IsValidIndex(StringKey))
+			check(!InStringTable.IsEmpty());
+
+			TArray<FString> Values;
+			Values.Reserve(InNumElements);
+
+			for (uint32 ElementIndex = 0; ElementIndex < InNumElements; ++ElementIndex)
 			{
-				Attribute->SetValue(EntryKey, InStringTable[StringKey]);
+				const uint32 PackedElementIndex = InAddressUints + ElementIndex;
+				const int32 StringKey = InStringTable.IsValidIndex(DataAsInt[PackedElementIndex]) ? DataAsInt[PackedElementIndex] : 0;
+				Values.Add(InStringTable[StringKey]);
 			}
+
+			ensure(PCGAttributeAccessorHelpers::WriteAllValues<FString>(OutData, Selector, Values, /*SourceSelector=*/nullptr, InContext));
 			break;
 		}
 		case EPCGKernelAttributeType::Name:
 		{
-			FPCGMetadataAttribute<FName>* Attribute = static_cast<FPCGMetadataAttribute<FName>*>(AttributeBase);
+			check(!InStringTable.IsEmpty());
 
-			// FNames currently stored in string table.
-			const int32 StringKey = DataAsInt[ElementIndex];
-			if (InStringTable.IsValidIndex(StringKey))
+			TArray<FName> Values;
+			Values.SetNumUninitialized(InNumElements);
+
+			for (uint32 ElementIndex = 0; ElementIndex < InNumElements; ++ElementIndex)
 			{
-				Attribute->SetValue(EntryKey, *InStringTable[StringKey]);
+				const uint32 PackedElementIndex = InAddressUints + ElementIndex;
+
+				// FNames currently stored in string table.
+				const int32 StringKey = InStringTable.IsValidIndex(DataAsInt[PackedElementIndex]) ? DataAsInt[PackedElementIndex] : 0;
+				Values[ElementIndex] = *InStringTable[StringKey];
 			}
+
+			ensure(PCGAttributeAccessorHelpers::WriteAllValues<FName>(OutData, Selector, Values, /*SourceSelector=*/nullptr, InContext));
 			break;
 		}
 		default:
 			return false;
 		}
-
 		return true;
 	}
 
@@ -598,6 +672,14 @@ uint64 FPCGDataDesc::ComputePackedSize() const
 	return DataSizeBytes;
 }
 
+bool FPCGDataDesc::HasMetadataAttributes() const
+{
+	return AttributeDescs.ContainsByPredicate([](const FPCGKernelAttributeDesc& AttributeDesc)
+	{
+		return AttributeDesc.Index >= NUM_RESERVED_ATTRS;
+	});
+}
+
 void FPCGDataDesc::InitializeAttributeDescs(const UPCGData* InData, const TMap<FName, FPCGKernelAttributeIDAndType>& InGlobalAttributeLookupTable, const TArray<FString>& InStringTable)
 {
 	TRACE_CPUPROFILER_EVENT_SCOPE(FPCGDataDesc::InitializeAttributeDescs);
@@ -908,6 +990,7 @@ void FPCGDataCollectionDesc::PackDataCollection(const FPCGDataCollection& InData
 	const TArray<FPCGTaggedData> InputData = InDataCollection.GetInputsByPin(InPin);
 
 	const uint32 PackedDataCollectionSizeBytes = ComputePackedSizeBytes();
+	check(PackedDataCollectionSizeBytes >= 0);
 
 	// Uninitialized is fine, all data is initialized explicitly.
 	OutPackedDataCollection.SetNumUninitialized(PackedDataCollectionSizeBytes >> 2);
@@ -1082,8 +1165,10 @@ void FPCGDataCollectionDesc::PackDataCollection(const FPCGDataCollection& InData
 	}
 }
 
-EPCGUnpackDataCollectionResult FPCGDataCollectionDesc::UnpackDataCollection(const TArray<uint8>& InPackedData, FName InPin, const TArray<FString>& InStringTable, FPCGDataCollection& OutDataCollection) const
+EPCGUnpackDataCollectionResult FPCGDataCollectionDesc::UnpackDataCollection(FPCGContext* InContext, const TArray<uint8>& InPackedData, FName InPin, const TArray<FString>& InStringTable, FPCGDataCollection& OutDataCollection) const
 {
+	TRACE_CPUPROFILER_EVENT_SCOPE(FPCGDataCollectionDesc::UnpackDataCollection);
+
 	if (InPackedData.IsEmpty())
 	{
 		ensureMsgf(false, TEXT("Tried to unpack a GPU data collection, but the readback buffer was empty."));
@@ -1110,36 +1195,51 @@ EPCGUnpackDataCollectionResult FPCGDataCollectionDesc::UnpackDataCollection(cons
 
 	for (uint32 DataIndex = 0; DataIndex < NumData; ++DataIndex)
 	{
-		TRACE_CPUPROFILER_EVENT_SCOPE(FPCGDataCollectionDesc::UnpackDataItem);
-
 		const uint32 DataHeaderAddress = (DATA_COLLECTION_HEADER_SIZE_BYTES + DATA_HEADER_SIZE_BYTES * DataIndex) / sizeof(uint32);
 
 		const uint32 TypeId =        DataAsUint[DataHeaderAddress + 0];
 		const uint32 NumAttributes = DataAsUint[DataHeaderAddress + 1];
 		const uint32 NumElements =   DataAsUint[DataHeaderAddress + 2];
 
-		const TArray<FPCGKernelAttributeDesc>& AttributeDescs = DataDescs[DataIndex].AttributeDescs;
+		const FPCGDataDesc& DataDesc = DataDescs[DataIndex];
+		const TArray<FPCGKernelAttributeDesc>& AttributeDescs = DataDesc.AttributeDescs;
 		check(NumAttributes == AttributeDescs.Num());
 
 		if (TypeId == POINT_DATA_TYPE_ID)
 		{
-			UPCGPointData* OutPointData = nullptr;
-			{
-				TRACE_CPUPROFILER_EVENT_SCOPE(InitalizeOutput);
-				OutPointData = NewObject<UPCGPointData>();
-				OutPointData->GetMutablePoints().SetNumUninitialized(NumElements);
-			}
+			TRACE_CPUPROFILER_EVENT_SCOPE(UnpackPointDataItem);
 
+			UPCGPointData* OutPointData = NewObject<UPCGPointData>();
 			UPCGMetadata* Metadata = OutPointData->MutableMetadata();
 			TArray<FPCGPoint>& OutPoints = OutPointData->GetMutablePoints();
 
+			OutPoints.SetNumUninitialized(NumElements);
+
 			{
-				// This can take an age as it populates the memory hierarchy.
-				TRACE_CPUPROFILER_EVENT_SCOPE(MetadataEntry);
-				ParallelFor(NumElements, [&OutPoints](int32 ElementIndex)
+				TRACE_CPUPROFILER_EVENT_SCOPE(InitializeOutput);
+
+				// We only need to add the entry keys if there are actually attributes to unpack.
+				if (DataDesc.HasMetadataAttributes())
 				{
-					OutPoints[ElementIndex].MetadataEntry = -1;
-				});
+					TArray<PCGMetadataEntryKey*> ParentEntryKeys;
+					ParentEntryKeys.Reserve(NumElements);
+
+					for (uint32 ElementIndex = 0; ElementIndex < NumElements; ++ElementIndex)
+					{
+						PCGMetadataEntryKey& EntryKey = OutPoints[ElementIndex].MetadataEntry;
+						EntryKey = PCGInvalidEntryKey;
+						ParentEntryKeys.Add(&EntryKey);
+					}
+
+					Metadata->AddEntriesInPlace(ParentEntryKeys);
+				}
+				else
+				{
+					ParallelFor(NumElements, [&OutPoints](int32 ElementIndex)
+					{
+						OutPoints[ElementIndex].MetadataEntry = -1;
+					});
+				}
 			}
 
 			FPCGTaggedData& OutTaggedData = OutData.Emplace_GetRef();
@@ -1149,15 +1249,14 @@ EPCGUnpackDataCollectionResult FPCGDataCollectionDesc::UnpackDataCollection(cons
 			// Loop over attributes.
 			for (const FPCGKernelAttributeDesc& AttributeDesc : AttributeDescs)
 			{
-				TRACE_CPUPROFILER_EVENT_SCOPE(UnpackAttribute);
-
 				const uint32 AttributeId = AttributeDesc.Index;
 				const uint32 AttributeStrideUints = PCGDataForGPUHelpers::GetAttributeTypeStrideBytes(AttributeDesc.Type) >> 2;
-
 				const uint32 AddressUints = GetElementDataStartAddressUints(DataAsUint, DataIndex, AttributeId);
 
 				if (AttributeId < NUM_RESERVED_ATTRS)
 				{
+					TRACE_CPUPROFILER_EVENT_SCOPE(UnpackPointProperty);
+
 					// We tried hoisting this decision to a lambda but it didn't appear to help.
 					switch (AttributeId)
 					{
@@ -1274,20 +1373,12 @@ EPCGUnpackDataCollectionResult FPCGDataCollectionDesc::UnpackDataCollection(cons
 				}
 				else
 				{
-					FPCGMetadataAttributeBase* AttributeBase = PCGDataForGPUHelpers::CreateAttributeFromAttributeDesc(Metadata, AttributeDesc);
-					check(AttributeBase);
+					TRACE_CPUPROFILER_EVENT_SCOPE(UnpackAttribute);
 
-					ParallelFor(NumElements, [PackedData, &InStringTable, &OutPoints, Metadata, AttributeBase, &AttributeDesc, AddressUints, AttributeStrideUints](int32 ElementIndex)
+					if (FPCGMetadataAttributeBase* AttributeBase = PCGDataForGPUHelpers::CreateAttributeFromAttributeDesc(Metadata, AttributeDesc))
 					{
-						// Points that were removed will not have their metadata set at all, and most likely not zero initialized either, so this could be garbage otherwise.
-						if (!FMath::IsFinite(OutPoints[ElementIndex].Density))
-						{
-							return;
-						}
-
-						Metadata->InitializeOnSet(OutPoints[ElementIndex].MetadataEntry);
-						ensure(PCGDataForGPUHelpers::UnpackAttributeHelper(PackedData, InStringTable, AddressUints + ElementIndex * AttributeStrideUints, AttributeBase, AttributeDesc, OutPoints[ElementIndex].MetadataEntry));
-					});
+						ensure(PCGDataForGPUHelpers::UnpackAttributeHelper(InContext, PackedData, AttributeDesc, InStringTable, AddressUints, NumElements, OutPointData));
+					}
 				}
 			}
 
@@ -1303,46 +1394,45 @@ EPCGUnpackDataCollectionResult FPCGDataCollectionDesc::UnpackDataCollection(cons
 				return true;
 			};
 
-			FPCGAsync::AsyncPointProcessing(/*Context=*/nullptr, OutPoints.Num(), OutPoints, DiscardInvalidPoints);
+			{
+				TRACE_CPUPROFILER_EVENT_SCOPE(DiscardInvalidPoints);
+				FPCGAsync::AsyncPointProcessing(InContext, OutPoints.Num(), OutPoints, DiscardInvalidPoints);
+			}
 		}
 		else if (TypeId == PARAM_DATA_TYPE_ID)
 		{
+			TRACE_CPUPROFILER_EVENT_SCOPE(UnpackParamDataItem);
+
 			UPCGParamData* OutParamData = NewObject<UPCGParamData>();
 			UPCGMetadata* Metadata = OutParamData->MutableMetadata();
 
-			FPCGTaggedData& OutTaggedData = OutData.Emplace_GetRef();
-			OutTaggedData.Data = OutParamData;
-			OutTaggedData.Pin = InPin;
+			{
+				TRACE_CPUPROFILER_EVENT_SCOPE(InitializeOutput);
 
-			TArray<TTuple</*EntryKey=*/int64, /*ParentEntryKey=*/int64>> AllMetadataEntries;
-			AllMetadataEntries.SetNumUninitialized(NumElements);
+				TArray<TTuple</*EntryKey=*/int64, /*ParentEntryKey=*/int64>> AllMetadataEntries;
+				AllMetadataEntries.SetNumUninitialized(NumElements);
 
-			ParallelFor(NumElements, [&](int32 ElementIndex)
-			{
-				AllMetadataEntries[ElementIndex] = MakeTuple(Metadata->AddEntryPlaceholder(), PCGInvalidEntryKey);
-			});
+				ParallelFor(NumElements, [&](int32 ElementIndex)
+				{
+					AllMetadataEntries[ElementIndex] = MakeTuple(Metadata->AddEntryPlaceholder(), PCGInvalidEntryKey);
+				});
 
-			Metadata->AddDelayedEntries(AllMetadataEntries);
+				Metadata->AddDelayedEntries(AllMetadataEntries);
+			}
+
+			FPCGTaggedData& OutTaggedData = OutData.Emplace_GetRef();
+			OutTaggedData.Data = OutParamData;
+			OutTaggedData.Pin = InPin;
 
 			// Loop over attributes.
 			for (const FPCGKernelAttributeDesc& AttributeDesc : AttributeDescs)
 			{
-				TRACE_CPUPROFILER_EVENT_SCOPE(WriteAttribute);
-
-				const uint32 AddressUints = GetElementDataStartAddressUints(DataAsUint, DataIndex, AttributeDesc.Index);
-				const uint32 AttributeStrideUints = PCGDataForGPUHelpers::GetAttributeTypeStrideBytes(AttributeDesc.Type) / sizeof(uint32);
+				TRACE_CPUPROFILER_EVENT_SCOPE(UnpackAttribute);
 
 				if (FPCGMetadataAttributeBase* AttributeBase = PCGDataForGPUHelpers::CreateAttributeFromAttributeDesc(Metadata, AttributeDesc))
 				{
-					const uint32 NumPackedUints = InPackedData.Num() / 4;
-
-					ParallelFor(NumElements, [PackedData, NumPackedUints, AddressUints, AttributeStrideUints, &AttributeDesc, AttributeBase, &InStringTable](int32 ElementIndex)
-					{
-						const uint32 PackedDataElementIndex = AddressUints + ElementIndex * AttributeStrideUints;
-						check(PackedDataElementIndex + AttributeStrideUints <= NumPackedUints);
-
-						ensure(PCGDataForGPUHelpers::UnpackAttributeHelper(PackedData, InStringTable, PackedDataElementIndex, AttributeBase, AttributeDesc, ElementIndex));
-					});
+					const uint32 AddressUints = GetElementDataStartAddressUints(DataAsUint, DataIndex, AttributeDesc.Index);
+					ensure(PCGDataForGPUHelpers::UnpackAttributeHelper(InContext, PackedData, AttributeDesc, InStringTable, AddressUints, NumElements, OutParamData));
 				}
 			}
 		}
diff --git a/Engine/Plugins/PCG/Source/PCG/Private/Compute/PCGDataForGPU.h b/Engine/Plugins/PCG/Source/PCG/Private/Compute/PCGDataForGPU.h
--- a/Engine/Plugins/PCG/Source/PCG/Private/Compute/PCGDataForGPU.h
+++ b/Engine/Plugins/PCG/Source/PCG/Private/Compute/PCGDataForGPU.h
@@ -11,6 +11,7 @@ class UPCGDataBinding;
 class UPCGMetadata;
 class UPCGPin;
 class UPCGSettings;
+struct FPCGContext;
 enum class EPCGMetadataTypes : uint8;
 
 enum class EPCGUnpackDataCollectionResult
@@ -115,6 +116,7 @@ struct FPCGDataDesc
 	FPCGDataDesc(const UPCGData* InData, const TMap<FName, FPCGKernelAttributeIDAndType>& InGlobalAttributeLookupTable, const TArray<FString>& InStringTable);
 
 	uint64 ComputePackedSize() const;
+	bool HasMetadataAttributes() const;
 
 	EPCGDataType Type = EPCGDataType::Point;
 	TArray<FPCGKernelAttributeDesc> AttributeDescs;
@@ -149,7 +151,7 @@ struct FPCGDataCollectionDesc
 	void PackDataCollection(const FPCGDataCollection& InDataCollection, FName InPin, const TArray<FString>& InStringTable, TArray<uint32>& OutPackedDataCollection) const;
 
 	/** Unpack a buffer of 8-bit uints to a data collection. */
-	EPCGUnpackDataCollectionResult UnpackDataCollection(const TArray<uint8>& InPackedData, FName InPin, const TArray<FString>& InStringTable, FPCGDataCollection& OutDataCollection) const;
+	EPCGUnpackDataCollectionResult UnpackDataCollection(FPCGContext* InContext, const TArray<uint8>& InPackedData, FName InPin, const TArray<FString>& InStringTable, FPCGDataCollection& OutDataCollection) const;
 
 	/** Compute total number of processing elements of the given type. */
 	uint32 ComputeDataElementCount(EPCGDataType InDataType) const;
diff --git a/Engine/Plugins/PCG/Source/PCG/Public/Metadata/Accessors/PCGAttributeAccessorHelpers.h b/Engine/Plugins/PCG/Source/PCG/Public/Metadata/Accessors/PCGAttributeAccessorHelpers.h
--- a/Engine/Plugins/PCG/Source/PCG/Public/Metadata/Accessors/PCGAttributeAccessorHelpers.h
+++ b/Engine/Plugins/PCG/Source/PCG/Public/Metadata/Accessors/PCGAttributeAccessorHelpers.h
@@ -176,7 +176,7 @@ namespace PCGAttributeAccessorHelpers
 			return false;
 		}
 
-		if (!ensureMsgf(Keys->GetNum() == InValues.GetNum(), TEXT("Number of values passed (%d) mismatches with the number of keys (%d)"), Keys->GetNum(), InValues.GetNum()))
+		if (!ensureMsgf(Keys->GetNum() == InValues.Num(), TEXT("Number of values passed (%d) mismatches with the number of keys (%d)"), Keys->GetNum(), InValues.Num()))
 		{
 			return false;
 		}
