[PCG] Spline to Dynamic mesh node

* Copied code from the Spline component to create a poly line out of a spline

#rnx

diff --git a/Engine/Plugins/PCG/Source/PCG/Private/Data/PCGSplineStruct.cpp b/Engine/Plugins/PCG/Source/PCG/Private/Data/PCGSplineStruct.cpp
--- a/Engine/Plugins/PCG/Source/PCG/Private/Data/PCGSplineStruct.cpp
+++ b/Engine/Plugins/PCG/Source/PCG/Private/Data/PCGSplineStruct.cpp
@@ -344,4 +344,121 @@ FQuat FPCGSplineStruct::GetQuaternionAtSplineInputKey(float InKey, ESplineCoordi
 FVector FPCGSplineStruct::GetScaleAtSplineInputKey(float InKey) const
 {
 	return SplineCurves.Scale.Eval(InKey, FVector(1.0f));
+}
+
+/** Taken from USplineComponent::ConvertSplineSegmentToPolyLine. */
+bool FPCGSplineStruct::ConvertSplineSegmentToPolyLine(int32 SplinePointStartIndex, ESplineCoordinateSpace::Type CoordinateSpace, const float MaxSquareDistanceFromSpline, TArray<FVector>& OutPoints) const
+{
+	OutPoints.Empty();
+	
+	TArray<double> DummyDistances;
+
+    const double StartDist = GetDistanceAlongSplineAtSplinePoint(SplinePointStartIndex);
+    const double StopDist = GetDistanceAlongSplineAtSplinePoint(SplinePointStartIndex + 1);
+
+    const int32 NumLines = 2; // Dichotomic subdivision of the spline segment
+    double Dist = StopDist - StartDist;
+    double SubstepSize = Dist / NumLines;
+    if (SubstepSize == 0.0)
+    {
+    	// There is no distance to cover, so handle the segment with a single point
+    	OutPoints.Add(GetLocationAtDistanceAlongSpline(StopDist, CoordinateSpace));
+    	return true;
+    }
+
+    double SubstepStartDist = StartDist;
+    for (int32 i = 0; i < NumLines; ++i)
+    {
+    	double SubstepEndDist = SubstepStartDist + SubstepSize;
+    	TArray<FVector> NewPoints;
+    	DummyDistances.Reset();
+    	// Recursively sub-divide each segment until the requested precision is reached :
+    	if (DivideSplineIntoPolylineRecursiveWithDistancesHelper(SubstepStartDist, SubstepEndDist, CoordinateSpace, MaxSquareDistanceFromSpline, NewPoints, DummyDistances))
+    	{
+    		if (OutPoints.Num() > 0)
+    		{
+    			check(OutPoints.Last() == NewPoints[0]); // our last point must be the same as the new segment's first
+    			OutPoints.RemoveAt(OutPoints.Num() - 1);
+    		}
+    		OutPoints.Append(NewPoints);
+    	}
+
+    	SubstepStartDist = SubstepEndDist;
+    }
+
+    return (OutPoints.Num() > 0);
+}
+
+/** Taken from USplineComponent::ConvertSplineToPolyLine. */
+bool FPCGSplineStruct::ConvertSplineToPolyLine(ESplineCoordinateSpace::Type CoordinateSpace, const float MaxSquareDistanceFromSpline, TArray<FVector>& OutPoints) const
+{
+	int32 NumSegments = GetNumberOfSplineSegments();
+	OutPoints.Empty();
+	OutPoints.Reserve(NumSegments * 2); // We sub-divide each segment in at least 2 sub-segments, so let's start with this amount of points
+
+	TArray<FVector> SegmentPoints;
+	for (int32 SegmentIndex = 0; SegmentIndex < NumSegments; ++SegmentIndex)
+	{
+		if (ConvertSplineSegmentToPolyLine(SegmentIndex, CoordinateSpace, MaxSquareDistanceFromSpline, SegmentPoints))
+		{
+			if (OutPoints.Num() > 0)
+			{
+				check(OutPoints.Last() == SegmentPoints[0]); // our last point must be the same as the new segment's first
+				OutPoints.RemoveAt(OutPoints.Num() - 1);
+			}
+			OutPoints.Append(SegmentPoints);
+		}
+	}
+
+	return (OutPoints.Num() > 0);
+}
+
+/** Taken from USplineComponent::DivideSplineIntoPolylineRecursiveWithDistancesHelper. */
+bool FPCGSplineStruct::DivideSplineIntoPolylineRecursiveWithDistancesHelper(float StartDistanceAlongSpline, float EndDistanceAlongSpline, ESplineCoordinateSpace::Type CoordinateSpace, const float MaxSquareDistanceFromSpline, TArray<FVector>& OutPoints, TArray<double>& OutDistancesAlongSpline) const
+{
+	double Dist = EndDistanceAlongSpline - StartDistanceAlongSpline;
+	if (Dist <= 0.0f)
+	{
+		return false;
+	}
+	double MiddlePointDistancAlongSpline = StartDistanceAlongSpline + Dist / 2.0f;
+	FVector Samples[3];
+	Samples[0] = GetLocationAtDistanceAlongSpline(StartDistanceAlongSpline, CoordinateSpace);
+	Samples[1] = GetLocationAtDistanceAlongSpline(MiddlePointDistancAlongSpline, CoordinateSpace);
+	Samples[2] = GetLocationAtDistanceAlongSpline(EndDistanceAlongSpline, CoordinateSpace);
+
+	if (FMath::PointDistToSegmentSquared(Samples[1], Samples[0], Samples[2]) > MaxSquareDistanceFromSpline)
+	{
+		TArray<FVector> NewPoints[2];
+		TArray<double> NewDistancesAlongSpline[2];
+		DivideSplineIntoPolylineRecursiveWithDistancesHelper(StartDistanceAlongSpline, MiddlePointDistancAlongSpline, CoordinateSpace, MaxSquareDistanceFromSpline, NewPoints[0], NewDistancesAlongSpline[0]);
+		DivideSplineIntoPolylineRecursiveWithDistancesHelper(MiddlePointDistancAlongSpline, EndDistanceAlongSpline, CoordinateSpace, MaxSquareDistanceFromSpline, NewPoints[1], NewDistancesAlongSpline[1]);
+		if ((NewPoints[0].Num() > 0) && (NewPoints[1].Num() > 0))
+		{
+			check(NewPoints[0].Last() == NewPoints[1][0]);
+			check(NewDistancesAlongSpline[0].Last() == NewDistancesAlongSpline[1][0]);
+			NewPoints[0].RemoveAt(NewPoints[0].Num() - 1);
+			NewDistancesAlongSpline[0].RemoveAt(NewDistancesAlongSpline[0].Num() - 1);
+		}
+		NewPoints[0].Append(NewPoints[1]);
+		NewDistancesAlongSpline[0].Append(NewDistancesAlongSpline[1]);
+		OutPoints.Append(NewPoints[0]);
+		OutDistancesAlongSpline.Append(NewDistancesAlongSpline[0]);
+	}
+	else
+	{
+		// The middle point is close enough to the other 2 points, let's keep those and stop the recursion :
+		OutPoints.Add(Samples[0]);
+		OutDistancesAlongSpline.Add(StartDistanceAlongSpline);
+		// For a constant spline, the end can be the exact same as the start; in this case, just add the point once
+		if (Samples[0] != Samples[2])
+		{
+			OutPoints.Add(Samples[2]);
+			OutDistancesAlongSpline.Add(EndDistanceAlongSpline);
+		}
+		
+	}
+
+	check(OutPoints.Num() == OutDistancesAlongSpline.Num())
+	return (OutPoints.Num() > 0);
 }
\ No newline at end of file
diff --git a/Engine/Plugins/PCG/Source/PCG/Public/Data/PCGSplineStruct.h b/Engine/Plugins/PCG/Source/PCG/Public/Data/PCGSplineStruct.h
--- a/Engine/Plugins/PCG/Source/PCG/Public/Data/PCGSplineStruct.h
+++ b/Engine/Plugins/PCG/Source/PCG/Public/Data/PCGSplineStruct.h
@@ -41,6 +41,12 @@ struct PCG_API FPCGSplineStruct
 	FVector::FReal GetDistanceAlongSplineAtSplinePoint(int32 PointIndex) const;
 	FVector GetLocationAtDistanceAlongSpline(FVector::FReal Distance, ESplineCoordinateSpace::Type CoordinateSpace) const;
 	FTransform GetTransformAtDistanceAlongSpline(FVector::FReal Distance, ESplineCoordinateSpace::Type CoordinateSpace, bool bUseScale = false) const;
+
+	/** Given a threshold, returns a list of vertices along the spline segment that, treated as a list of segments (polyline), matches the spline shape. Taken from USplineComponent. */
+	bool ConvertSplineSegmentToPolyLine(int32 SplinePointStartIndex, ESplineCoordinateSpace::Type CoordinateSpace, const float MaxSquareDistanceFromSpline, TArray<FVector>& OutPoints) const;
+
+	/** Given a threshold, returns a list of vertices along the spline that, treated as a list of segments (polyline), matches the spline shape. Taken from USplineComponent. */
+	bool ConvertSplineToPolyLine(ESplineCoordinateSpace::Type CoordinateSpace, const float MaxSquareDistanceFromSpline, TArray<FVector>& OutPoints) const;
 	
 	float FindInputKeyClosestToWorldLocation(const FVector& WorldLocation) const;	
 
@@ -65,4 +71,8 @@ struct PCG_API FPCGSplineStruct
 
 	UPROPERTY()
 	FBoxSphereBounds Bounds = FBoxSphereBounds(EForceInit::ForceInit);
+
+private:
+	// Internal helper function called by ConvertSplineSegmentToPolyLine -- assumes the input is within a half-segment, so testing the distance to midpoint will be an accurate guide to subdivision. Taken from USplineComponent.
+	bool DivideSplineIntoPolylineRecursiveWithDistancesHelper(float StartDistanceAlongSpline, float EndDistanceAlongSpline, ESplineCoordinateSpace::Type CoordinateSpace, const float MaxSquareDistanceFromSpline, TArray<FVector>& OutPoints, TArray<double>& OutDistancesAlongSpline) const;
 };
\ No newline at end of file
diff --git a/Engine/Plugins/PCGInterops/PCGGeometryScriptInterop/Source/PCGGeometryScriptInterop/PCGGeometryScriptInterop.Build.cs b/Engine/Plugins/PCGInterops/PCGGeometryScriptInterop/Source/PCGGeometryScriptInterop/PCGGeometryScriptInterop.Build.cs
--- a/Engine/Plugins/PCGInterops/PCGGeometryScriptInterop/Source/PCGGeometryScriptInterop/PCGGeometryScriptInterop.Build.cs
+++ b/Engine/Plugins/PCGInterops/PCGGeometryScriptInterop/Source/PCGGeometryScriptInterop/PCGGeometryScriptInterop.Build.cs
@@ -17,7 +17,8 @@ namespace UnrealBuildTool.Rules
 					"Projects",
 					"RenderCore",
 					"RHI",
-					"PCG"
+					"PCG",
+					"ModelingOperators"
 				});
 
 			PrivateDependencyModuleNames.AddRange(
diff --git a/Engine/Plugins/PCGInterops/PCGGeometryScriptInterop/Source/PCGGeometryScriptInterop/Private/Elements/PCGSplineToMesh.cpp b/Engine/Plugins/PCGInterops/PCGGeometryScriptInterop/Source/PCGGeometryScriptInterop/Private/Elements/PCGSplineToMesh.cpp
new file mode 100644
--- /dev/null
+++ b/Engine/Plugins/PCGInterops/PCGGeometryScriptInterop/Source/PCGGeometryScriptInterop/Private/Elements/PCGSplineToMesh.cpp
@@ -0,0 +1,104 @@
+// Copyright Epic Games, Inc. All Rights Reserved.
+
+#include "Elements/PCGSplineToMesh.h"
+
+#include "PCGContext.h"
+#include "Data/PCGDynamicMeshData.h"
+#include "Data/PCGSplineData.h"
+
+#include "CurveOps/TriangulateCurvesOp.h"
+
+#define LOCTEXT_NAMESPACE "PCGSplineToMeshElement"
+
+#if WITH_EDITOR
+FName UPCGSplineToMeshSettings::GetDefaultNodeName() const
+{
+	return FName(TEXT("SplineToMesh"));
+}
+
+FText UPCGSplineToMeshSettings::GetDefaultNodeTitle() const
+{
+	return LOCTEXT("NodeTitle", "Spline To Mesh");
+}
+
+FText UPCGSplineToMeshSettings::GetNodeTooltipText() const
+{
+	return LOCTEXT("NodeTooltip", "Converts a closed spline into a mesh.");
+}
+#endif // WITH_EDITOR
+
+FPCGElementPtr UPCGSplineToMeshSettings::CreateElement() const
+{
+	return MakeShared<FPCGSplineToMeshElement>();
+}
+
+TArray<FPCGPinProperties> UPCGSplineToMeshSettings::InputPinProperties() const
+{
+	TArray<FPCGPinProperties> Properties;
+	Properties.Emplace_GetRef(PCGPinConstants::DefaultInputLabel, EPCGDataType::Spline).SetRequiredPin();	
+	return Properties;
+}
+
+bool FPCGSplineToMeshElement::ExecuteInternal(FPCGContext* InContext) const
+{
+	TRACE_CPUPROFILER_EVENT_SCOPE(FPCGSplineToMeshElement::Execute);
+
+	check(InContext);
+
+	const UPCGSplineToMeshSettings* Settings = InContext->GetInputSettings<UPCGSplineToMeshSettings>();
+	check(Settings);
+
+	const TArray<FPCGTaggedData> Inputs = InContext->InputData.GetInputsByPin(PCGPinConstants::DefaultInputLabel);
+	for (const FPCGTaggedData& Input : Inputs)
+	{
+		const UPCGSplineData* InputSplineData = Cast<const UPCGSplineData>(Input.Data);
+		if (!InputSplineData)
+		{
+			PCGLog::InputOutput::LogTypedDataNotFoundWarning(EPCGDataType::Spline, PCGPinConstants::DefaultInputLabel, InContext);
+			continue;
+		}
+
+		UE::Geometry::FTriangulateCurvesOp TriangulateCurvesOp;
+		TriangulateCurvesOp.Thickness = Settings->Thickness;
+		TriangulateCurvesOp.bFlipResult = Settings->bFlipResult;
+		TriangulateCurvesOp.FlattenMethod = Settings->FlattenMethod;
+		TriangulateCurvesOp.CurveOffset = Settings->CurveOffset;
+		TriangulateCurvesOp.OffsetOpenMethod = Settings->OpenCurves;
+		TriangulateCurvesOp.OffsetJoinMethod = Settings->JoinMethod;
+		TriangulateCurvesOp.OpenEndShape = Settings->EndShapes;
+		TriangulateCurvesOp.MiterLimit = Settings->MiterLimit;
+		TriangulateCurvesOp.bFlipResult = Settings->bFlipResult;
+		
+		if (Settings->CurveOffset == 0.0)
+		{
+			TriangulateCurvesOp.OffsetClosedMethod = EOffsetClosedCurvesMethod::DoNotOffset;
+		}
+		else
+		{
+			TriangulateCurvesOp.OffsetClosedMethod = Settings->OffsetClosedCurves;
+		}
+		
+		TArray<FVector> SplinePoints;
+		InputSplineData->SplineStruct.ConvertSplineToPolyLine(ESplineCoordinateSpace::World, Settings->ErrorTolerance, SplinePoints);
+		TriangulateCurvesOp.AddWorldCurve(SplinePoints, InputSplineData->IsClosed(), InputSplineData->SplineStruct.Transform);
+		TriangulateCurvesOp.CalculateResult(nullptr);
+
+		TUniquePtr<UE::Geometry::FDynamicMesh3>	DynamicMesh = TriangulateCurvesOp.ExtractResult();
+
+		if (!DynamicMesh || DynamicMesh->TriangleCount() == 0)
+		{
+			PCGLog::LogErrorOnGraph(LOCTEXT("TriangulationFailed", "Triangulation failed"), InContext);
+			continue;
+		}
+		
+		UE::Geometry::FDynamicMesh3* DynamicMeshPtr = DynamicMesh.Release();
+		UPCGDynamicMeshData* DynamicMeshData = FPCGContext::NewObject_AnyThread<UPCGDynamicMeshData>(InContext);
+		DynamicMeshData->Initialize(std::move(*DynamicMeshPtr));
+
+		InContext->OutputData.TaggedData.Emplace_GetRef(Input).Data = DynamicMeshData;
+	}
+	
+	return true;
+}
+
+#undef LOCTEXT_NAMESPACE
diff --git a/Engine/Plugins/PCGInterops/PCGGeometryScriptInterop/Source/PCGGeometryScriptInterop/Public/Elements/PCGSplineToMesh.h b/Engine/Plugins/PCGInterops/PCGGeometryScriptInterop/Source/PCGGeometryScriptInterop/Public/Elements/PCGSplineToMesh.h
new file mode 100644
--- /dev/null
+++ b/Engine/Plugins/PCGInterops/PCGGeometryScriptInterop/Source/PCGGeometryScriptInterop/Public/Elements/PCGSplineToMesh.h
@@ -0,0 +1,78 @@
+// Copyright Epic Games, Inc. All Rights Reserved.
+
+#pragma once
+
+#include "Elements/PCGDynamicMeshBaseElement.h"
+
+#include "CurveOps/TriangulateCurvesOp.h"
+
+#include "PCGSplineToMesh.generated.h"
+
+UCLASS(MinimalAPI, BlueprintType, ClassGroup = (Procedural))
+class UPCGSplineToMeshSettings : public UPCGDynamicMeshBaseSettings
+{
+	GENERATED_BODY()
+
+public:
+	//~Begin UPCGSettings interface
+#if WITH_EDITOR
+	virtual FName GetDefaultNodeName() const override;
+	virtual FText GetDefaultNodeTitle() const override;
+    virtual FText GetNodeTooltipText() const override;
+#endif
+
+protected:
+	virtual TArray<FPCGPinProperties> InputPinProperties() const override;
+	virtual FPCGElementPtr CreateElement() const override;
+	//~End UPCGSettings interface
+	
+public:
+	// This is a 1:1 match with the Geometry Script struct options for this operation
+
+	/** How far to allow the triangulation boundary can deviate from the spline curve before we add more vertices */
+	UPROPERTY(EditAnywhere, Category = Spline, meta = (ClampMin = 0.001))
+	double ErrorTolerance = 1.0;
+	
+	/** Whether and how to flatten the curves. If curves are flattened, they can also be offset. */
+	UPROPERTY(EditAnywhere, Category = Spline)
+	EFlattenCurveMethod FlattenMethod = EFlattenCurveMethod::DoNotFlatten;
+
+	/** If > 0, Extrude the triangulation by this amount. */
+	UPROPERTY(EditAnywhere, Category = Mesh, meta = (ClampMin = 0.0))
+	double Thickness = 0.0;
+
+	/** Whether to flip the facing direction of the generated mesh. */
+	UPROPERTY(EditAnywhere, Category = Mesh)
+	bool bFlipResult = false;
+
+	/** How to handle open curves: Either offset them, or treat them as closed curves. */
+	UPROPERTY(EditAnywhere, Category = Offset, meta = (EditCondition = "FlattenMethod != EFlattenCurveMethod::DoNotFlatten", EditConditionHides))
+	EOffsetOpenCurvesMethod OpenCurves = EOffsetOpenCurvesMethod::TreatAsClosed;
+
+	/** How much offset to apply to curves. */
+	UPROPERTY(EditAnywhere, Category = Offset, meta = (EditCondition = "FlattenMethod != EFlattenCurveMethod::DoNotFlatten", EditConditionHides))
+	double CurveOffset = 0.0;
+
+	/** Whether and how to apply offset to closed curves. */
+	UPROPERTY(EditAnywhere, Category = Offset, meta = (EditCondition = "FlattenMethod != EFlattenCurveMethod::DoNotFlatten && CurveOffset != 0", EditConditionHides))
+	EOffsetClosedCurvesMethod OffsetClosedCurves = EOffsetClosedCurvesMethod::OffsetOuterSide;
+
+	/** The shape of the ends of offset curves. */
+	UPROPERTY(EditAnywhere, Category = Offset, meta = (EditCondition = "FlattenMethod != EFlattenCurveMethod::DoNotFlatten && OpenCurves != EOffsetOpenCurvesMethod::TreatAsClosed && CurveOffset != 0", EditConditionHides))
+	EOpenCurveEndShapes EndShapes = EOpenCurveEndShapes::Square;
+
+	/** The shape of joins between segments of an offset curve. */
+	UPROPERTY(EditAnywhere, Category = Offset, meta = (EditCondition = "FlattenMethod != EFlattenCurveMethod::DoNotFlatten && CurveOffset != 0", EditConditionHides))
+	EOffsetJoinMethod JoinMethod = EOffsetJoinMethod::Square;
+
+	/** How far a miter join can extend before it is replaced by a square join. */
+	UPROPERTY(EditAnywhere, Category = Offset, meta = (ClampMin = 0.0, EditCondition = "FlattenMethod != EFlattenCurveMethod::DoNotFlatten && CurveOffset != 0 && JoinMethod == EOffsetJoinMethod::Miter", EditConditionHides))
+	double MiterLimit = 1.0;
+};
+
+class FPCGSplineToMeshElement : public IPCGDynamicMeshBaseElement
+{
+protected:
+	virtual bool ExecuteInternal(FPCGContext* InContext) const override;
+};
+
