[PCG] GPU SM Spawner - fix percentage of missing instances when using weighted spawning.

Fix bug where ComputeFramework was calling PostGraphDispatch multiple times on each data interface proxy.

#rnx

diff --git a/Engine/Plugins/PCG/Source/PCG/Private/Compute/DataInterfaces/PCGInstanceDataInterface.cpp b/Engine/Plugins/PCG/Source/PCG/Private/Compute/DataInterfaces/PCGInstanceDataInterface.cpp
--- a/Engine/Plugins/PCG/Source/PCG/Private/Compute/DataInterfaces/PCGInstanceDataInterface.cpp
+++ b/Engine/Plugins/PCG/Source/PCG/Private/Compute/DataInterfaces/PCGInstanceDataInterface.cpp
@@ -35,6 +35,8 @@
 
 #include UE_INLINE_GENERATED_CPP_BY_NAME(PCGInstanceDataInterface)
 
+#define PCG_INSTANCE_DATA_LOGGING 0
+
 void UPCGInstanceDataInterface::GetSupportedOutputs(TArray<FShaderFunctionDefinition>& OutFunctions) const
 {
 	OutFunctions.AddDefaulted_GetRef()
@@ -120,9 +122,6 @@ UComputeDataProvider* UPCGInstanceDataInterface::CreateDataProvider(TObjectPtr<U
 	// If there were 0 input points (or too many input points) for this execution, we will not have created any primitives, so check for null.
 	if (FPCGSpawnerPrimitives* FoundPrimitives = Binding->MeshSpawnersToPrimitives.Find(ProducerSettings))
 	{
-		const FPCGDataCollectionDesc InputDataDesc = PCGDataForGPUHelpers::ComputeInputPinDataDesc(ProducerSettings, InputPinProvidingData, Binding);
-		DataProvider->NumInstancesAllPrimitives = InputDataDesc.ComputeDataElementCount(EPCGDataType::Point);
-
 		if (!ensure(FoundPrimitives->Primitives.Num() <= PCGComputeConstants::MAX_PRIMITIVE_COMPONENTS_PER_SPAWNER))
 		{
 			// Last resort - should be clamped earlier during setup and not come to this.
@@ -130,13 +129,8 @@ UComputeDataProvider* UPCGInstanceDataInterface::CreateDataProvider(TObjectPtr<U
 		}
 
 		DataProvider->Primitives = FoundPrimitives->Primitives;
+		DataProvider->NumInstancesAllPrimitives = FoundPrimitives->NumInstancesAllPrimitives;
 		DataProvider->NumCustomFloatsPerInstance = FoundPrimitives->NumCustomFloats;
-
-		if (FoundPrimitives->SelectorAttributeId != -1)
-		{
-			// When selecting primitives dynamically (by attribute), we don't know statically how many instances will end up in each primitive, so run worst case.
-			DataProvider->NumInstancesAllPrimitives *= FoundPrimitives->Primitives.Num();
-		}
 	}
 
 	return DataProvider;
@@ -250,6 +244,10 @@ void FPCGInstanceDataProviderProxy::PostGraphDispatch(FRDGBuilder& GraphBuilder)
 			return;
 		}
 
+#if PCG_INSTANCE_DATA_LOGGING
+		UE_LOG(LogPCG, Log, TEXT("\tProcess primitive [%d]:\tPrimitive ID %u\tProxy: %llu"), PrimitiveIndex, SceneInfo->PrimitiveComponentId.PrimIDValue, (uint64)this);
+#endif
+
 		const int32 NumInstancesAllocatedInGPUScene = PrimitiveSceneProxy->GetInstanceDataHeader().NumInstances;
 
 		// Post graph dispatch, the instance data will have been populated. Enqueue a scene update to write the instance data into the scene.
@@ -275,6 +273,15 @@ void FPCGInstanceDataProviderProxy::PostGraphDispatch(FRDGBuilder& GraphBuilder)
 				Parameters->InCustomDataCount = Params.NumCustomDataFloats;
 				Parameters->InPayloadDataFlags = Params.PackedInstanceSceneDataFlags;
 			
+#if PCG_INSTANCE_DATA_LOGGING
+				UE_LOG(LogPCG, Log, TEXT("\tScene writer delegate [%d]:\tPrimitive ID %u,\tsource instance offset %u,\tInstanceSceneDataOffset %u, num instances %u"),
+					Parameters->InPrimitiveIndex,
+					Parameters->InPrimitiveId,
+					Parameters->InInstanceOffset,
+					Params.InstanceSceneDataOffset,
+					Parameters->InNumInstancesAllocatedInGPUScene);
+#endif
+
 				// We're using a custom compute shader here rather than using a compute graph kernel, because ultimately the scene update will incorporate the
 				// new instance data. In the future we will not write directly to the scene here.
 				TShaderMapRef<FPCGSceneWriterCS> Shader(GetGlobalShaderMap(GMaxRHIFeatureLevel));
diff --git a/Engine/Plugins/PCG/Source/PCG/Private/Compute/Elements/PCGComputeGraphElement.cpp b/Engine/Plugins/PCG/Source/PCG/Private/Compute/Elements/PCGComputeGraphElement.cpp
--- a/Engine/Plugins/PCG/Source/PCG/Private/Compute/Elements/PCGComputeGraphElement.cpp
+++ b/Engine/Plugins/PCG/Source/PCG/Private/Compute/Elements/PCGComputeGraphElement.cpp
@@ -574,6 +574,8 @@ bool FPCGComputeGraphElement::SetupProceduralISMComponents(FPCGContext* InContex
 
 			Primitives.Primitives.Add(MISMC->GetComponent());
 
+			Primitives.NumInstancesAllPrimitives += Desc.NumInstances;
+
 			bAnyComponentsSetup = true;
 		}
 	}
diff --git a/Engine/Plugins/PCG/Source/PCG/Private/Compute/PCGDataBinding.h b/Engine/Plugins/PCG/Source/PCG/Private/Compute/PCGDataBinding.h
--- a/Engine/Plugins/PCG/Source/PCG/Private/Compute/PCGDataBinding.h
+++ b/Engine/Plugins/PCG/Source/PCG/Private/Compute/PCGDataBinding.h
@@ -38,6 +38,9 @@ struct FPCGSpawnerPrimitives
 
 	UPROPERTY()
 	TArray<int32> PrimitiveStringKeys;
+
+	UPROPERTY()
+	uint32 NumInstancesAllPrimitives = 0;
 };
 
 UCLASS(Transient, Category = PCG)
diff --git a/Engine/Plugins/Runtime/ComputeFramework/Source/ComputeFramework/Private/ComputeGraphWorker.cpp b/Engine/Plugins/Runtime/ComputeFramework/Source/ComputeFramework/Private/ComputeGraphWorker.cpp
--- a/Engine/Plugins/Runtime/ComputeFramework/Source/ComputeFramework/Private/ComputeGraphWorker.cpp
+++ b/Engine/Plugins/Runtime/ComputeFramework/Source/ComputeFramework/Private/ComputeGraphWorker.cpp
@@ -335,10 +335,8 @@ void FComputeGraphTaskWorker::SubmitWork(FRDGBuilder& GraphBuilder, FName InExec
 
 	// Let proxies perform any post-graph-dispatch actions.
 	// TODO: Given few proxies currently implement this method, it may be better to record which proxies require this callback rather than calling on all.
-	for (FSubmitDescription const& SubmitDesc : SubmitDescs)
+	for (FGraphInvocation const& GraphInvocation : GraphInvocations)
 	{
-		const int32 GraphIndex = SubmitDesc.GraphIndex;
-		FGraphInvocation const& GraphInvocation = GraphInvocations[GraphIndex];
 		for (const FComputeDataProviderRenderProxy* Proxy : GraphInvocation.DataProviderRenderProxies)
 		{
 			if (Proxy)
