Localization fix

#rnx

diff --git a/Engine/Plugins/PCG/Source/PCG/Private/Helpers/PCGDataLayerHelpers.cpp b/Engine/Plugins/PCG/Source/PCG/Private/Helpers/PCGDataLayerHelpers.cpp
--- a/Engine/Plugins/PCG/Source/PCG/Private/Helpers/PCGDataLayerHelpers.cpp
+++ b/Engine/Plugins/PCG/Source/PCG/Private/Helpers/PCGDataLayerHelpers.cpp
@@ -136,7 +136,7 @@ namespace PCGDataLayerHelpers
 		{
 			if (!DataLayerAssetsInput.Data)
 			{
-				PCGLog::LogWarningOnGraph(LOCTEXT("DataLayersWrongType", "Data layer assets input is not of type attribute set."), Context);
+				PCGLog::LogWarningOnGraph(LOCTEXT("DataLayerAssetsWrongType", "Data layer assets input is not of type attribute set."), Context);
 				continue;
 			}
 
