[PCG] GPU SM Spawner - fix mesh spawning in loop subgraph by not allowing re-use of PISMC resources, as they can't be shared across multiple executions of the spawner.

Expose errors on spawner node to bring behaviour inline with CPU SM Spawner.

#rnx

diff --git a/Engine/Plugins/PCG/Source/PCG/Private/Components/PCGProceduralISMComponent.cpp b/Engine/Plugins/PCG/Source/PCG/Private/Components/PCGProceduralISMComponent.cpp
--- a/Engine/Plugins/PCG/Source/PCG/Private/Components/PCGProceduralISMComponent.cpp
+++ b/Engine/Plugins/PCG/Source/PCG/Private/Components/PCGProceduralISMComponent.cpp
@@ -45,7 +45,8 @@ UPCGManagedProceduralISMComponent* PCGManagedProceduralISMComponent::GetOrCreate
 
 			if (UPCGManagedProceduralISMComponent* Resource = Cast<UPCGManagedProceduralISMComponent>(InResource))
 			{
-				if (Resource->GetSettingsUID() != InSettingsUID || !Resource->CanBeUsed())
+				// Note: Contrary to other managed resources, PISMCs can't be extended after being used so only allow each PISMC to be used once.
+				if (!Resource->CanBeUsed() || !Resource->IsMarkedUnused() || Resource->GetSettingsUID() != InSettingsUID)
 				{
 					return;
 				}
diff --git a/Engine/Plugins/PCG/Source/PCG/Private/Compute/Elements/PCGComputeGraphElement.cpp b/Engine/Plugins/PCG/Source/PCG/Private/Compute/Elements/PCGComputeGraphElement.cpp
--- a/Engine/Plugins/PCG/Source/PCG/Private/Compute/Elements/PCGComputeGraphElement.cpp
+++ b/Engine/Plugins/PCG/Source/PCG/Private/Compute/Elements/PCGComputeGraphElement.cpp
@@ -27,6 +27,12 @@
 
 #define LOCTEXT_NAMESPACE "PCGComputeGraphElement"
 
+namespace PCGComputeGraphElement
+{
+	static const FText CouldNotLoadStaticMeshFormat = LOCTEXT("CouldNotLoadStaticMesh", "Could not load static mesh from path '{0}'.");
+	static const FText TooManyPrimitivesFormat = LOCTEXT("TooManyPrimitives", "Attempted to emit too many primitive components, terminated after creating '{0}'.");
+}
+
 void FPCGComputeGraphContext::AddExtraStructReferencedObjects(FReferenceCollector& Collector)
 {
 	Collector.AddPropertyReferences(FComputeGraphInstance::StaticStruct(), &ComputeGraphInstance);
@@ -359,13 +365,17 @@ bool FPCGComputeGraphElement::SetupProceduralISMComponents(FPCGContext* InContex
 			{
 				if (AttributeName == NAME_None)
 				{
-					UE_LOG(LogPCG, Error, TEXT("Invalid instance data attribute specified '%s'."), *AttributeName.ToString());
+					PCG_KERNEL_VALIDATION_ERR(Context, SpawnerSettings, /*bQuiet=*/false,
+						FText::Format(LOCTEXT("InstanceDataAttributeInvalid", "Invalid instance data attribute specified '{0}'."), FText::FromName(AttributeName)));
+
 					continue;
 				}
 
 				if (!InBinding->GetAttributeLookupTable().Find(AttributeName) && !InBinding->GetInputAttributeTable().Find(AttributeName))
 				{
-					UE_LOG(LogPCG, Error, TEXT("Instance data attribute '%s' not found."), *AttributeName.ToString());
+					PCG_KERNEL_VALIDATION_ERR(Context, SpawnerSettings, /*bQuiet=*/false,
+						FText::Format(LOCTEXT("InstanceDataAttributeNotFound", "Instance data attribute '{0}' not found."), FText::FromName(AttributeName)));
+
 					continue;
 				}
 			}
@@ -392,13 +402,17 @@ bool FPCGComputeGraphElement::SetupProceduralISMComponents(FPCGContext* InContex
 
 			if (!FoundAttribute)
 			{
-				UE_LOG(LogPCG, Error, TEXT("Mesh selector attribute '%s' not found."), *SelectedName.ToString());
+				PCG_KERNEL_VALIDATION_ERR(Context, SpawnerSettings, /*bQuiet=*/false,
+					FText::Format(LOCTEXT("MeshSelectorAttributeNotFound", "Mesh selector attribute '{0}' not found."), FText::FromName(SelectedName)));
+
 				continue;
 			}
 
 			if (FoundAttribute->Type != EPCGKernelAttributeType::StringKey)
 			{
-				UE_LOG(LogPCG, Error, TEXT("Mesh selector attribute '%s' must be of type String Key."), *SelectedName.ToString());
+				PCG_KERNEL_VALIDATION_ERR(Context, SpawnerSettings, /*bQuiet=*/false,
+					FText::Format(LOCTEXT("MeshSelectorAttributeInvalidType", "Mesh selector attribute '{0}' must be of type String Key."), FText::FromName(SelectedName)));
+
 				continue;
 			}
 
@@ -418,7 +432,11 @@ bool FPCGComputeGraphElement::SetupProceduralISMComponents(FPCGContext* InContex
 						}
 						else
 						{
-							UE_LOG(LogPCG, Error, TEXT("Attribute '%s' not usable for mesh selection, only attributes of type String Key are supported."), *AttributeDesc.Name.ToString())
+							// Currently only a single type per attribute name is supported (the name/type in the attribute table). It's possible to wire a graph
+							// for which an attribute is present with multiple types. Warn if this is encountered.
+							PCG_KERNEL_VALIDATION_WARN(Context, SpawnerSettings, /*bQuiet=*/false, FText::Format(
+								LOCTEXT("MeshSelectorAttributeNotUsable", "Attribute '{0}' not usable for mesh selection, only attributes of type String Key are supported."),
+								FText::FromName(SelectedName)));
 						}
 					}
 				}
@@ -439,14 +457,18 @@ bool FPCGComputeGraphElement::SetupProceduralISMComponents(FPCGContext* InContex
 
 				if (ComponentsToCreate.Num() >= PCGComputeConstants::MAX_PRIMITIVE_COMPONENTS_PER_SPAWNER)
 				{
-					UE_LOG(LogPCG, Warning, TEXT("Attempted to emit too many primitive components, terminated after creating %d."), PCGComputeConstants::MAX_PRIMITIVE_COMPONENTS_PER_SPAWNER);
+					PCG_KERNEL_VALIDATION_WARN(Context, SpawnerSettings, /*bQuiet=*/false,
+						FText::Format(PCGComputeGraphElement::TooManyPrimitivesFormat, FText::FromString(FString::FromInt(PCGComputeConstants::MAX_PRIMITIVE_COMPONENTS_PER_SPAWNER))));
+
 					break;
 				}
 
 				UStaticMesh* StaticMesh = Cast<UStaticMesh>(FSoftObjectPath(MeshPathString).TryLoad());
 				if (!StaticMesh)
 				{
-					UE_LOG(LogPCG, Error, TEXT("Could not load static mesh from path '%s'."), *MeshPathString);
+					PCG_KERNEL_VALIDATION_ERR(Context, SpawnerSettings, /*bQuiet=*/false,
+						FText::Format(PCGComputeGraphElement::CouldNotLoadStaticMeshFormat, FText::FromString(MeshPathString)));
+
 					continue;
 				}
 
@@ -472,7 +494,8 @@ bool FPCGComputeGraphElement::SetupProceduralISMComponents(FPCGContext* InContex
 		{
 			if (SelectorWeighted->MeshEntries.IsEmpty())
 			{
-				UE_LOG(LogPCG, Error, TEXT("No mesh entries provided."));
+				PCG_KERNEL_VALIDATION_ERR(Context, SpawnerSettings, /*bQuiet=*/false, LOCTEXT("NoMeshEntries", "No mesh entries provided."));
+
 				continue;
 			}
 
@@ -500,7 +523,9 @@ bool FPCGComputeGraphElement::SetupProceduralISMComponents(FPCGContext* InContex
 
 				if (ComponentsToCreate.Num() >= PCGComputeConstants::MAX_PRIMITIVE_COMPONENTS_PER_SPAWNER)
 				{
-					UE_LOG(LogPCG, Warning, TEXT("Attempted to emit too many primitive components, terminated after creating %d."), PCGComputeConstants::MAX_PRIMITIVE_COMPONENTS_PER_SPAWNER);
+					PCG_KERNEL_VALIDATION_WARN(Context, SpawnerSettings, /*bQuiet=*/false,
+						FText::Format(PCGComputeGraphElement::TooManyPrimitivesFormat, FText::FromString(FString::FromInt(PCGComputeConstants::MAX_PRIMITIVE_COMPONENTS_PER_SPAWNER))));
+
 					break;
 				}
 
@@ -530,7 +555,8 @@ bool FPCGComputeGraphElement::SetupProceduralISMComponents(FPCGContext* InContex
 				}
 				else
 				{
-					UE_LOG(LogPCG, Error, TEXT("Could not load static mesh from path '%s'."), *Entry.Descriptor.StaticMesh.ToString())
+					PCG_KERNEL_VALIDATION_ERR(Context, SpawnerSettings, /*bQuiet=*/false,
+						FText::Format(PCGComputeGraphElement::CouldNotLoadStaticMeshFormat, FText::FromString(Entry.Descriptor.StaticMesh.ToString())));
 				}
 			}
 		}
@@ -545,7 +571,10 @@ bool FPCGComputeGraphElement::SetupProceduralISMComponents(FPCGContext* InContex
 
 			if (TotalInstanceCount >= MAX_INSTANCE_ID)
 			{
-				UE_LOG(LogPCG, Error, TEXT("Tried to spawn too many instances (%lu), procedural ISM component creation skipped and instances will not be rendered."), TotalInstanceCount);
+				PCG_KERNEL_VALIDATION_ERR(Context, SpawnerSettings, /*bQuiet=*/false, FText::Format(
+					LOCTEXT("TooManyInstances", "Tried to spawn too many instances ({0}), procedural ISM component creation skipped and instances will not be rendered."), 
+					FText::FromString(FString::FromInt(TotalInstanceCount))));
+
 				return false;
 			}
 		}
