Inform reflection system that PCGPoint does not have any object references, prevents expensive search when replacing references (package reload, BP compile, reinstancing generally)

#rnx

diff --git a/Engine/Plugins/PCG/Source/PCG/Public/PCGPoint.h b/Engine/Plugins/PCG/Source/PCG/Public/PCGPoint.h
--- a/Engine/Plugins/PCG/Source/PCG/Public/PCGPoint.h
+++ b/Engine/Plugins/PCG/Source/PCG/Public/PCGPoint.h
@@ -105,4 +105,5 @@ struct TStructOpsTypeTraits<FPCGPoint> : public TStructOpsTypeTraitsBase2<FPCGPo
 	{
 		WithStructuredSerializer = true,
 	};
+	static constexpr EPropertyObjectReferenceType WithSerializerObjectReferences = EPropertyObjectReferenceType::None;
 };
