Fixed potential check/crash when using the collision wrapper with meshes that don't have a good body setup.

#rnx

diff --git a/Engine/Plugins/PCG/Source/PCG/Private/Data/PCGCollisionWrapperData.cpp b/Engine/Plugins/PCG/Source/PCG/Private/Data/PCGCollisionWrapperData.cpp
--- a/Engine/Plugins/PCG/Source/PCG/Private/Data/PCGCollisionWrapperData.cpp
+++ b/Engine/Plugins/PCG/Source/PCG/Private/Data/PCGCollisionWrapperData.cpp
@@ -239,9 +239,11 @@ void UPCGCollisionWrapperData::FinalizeInitializationEx(const TArray<FSoftObject
 
 		for (FBodyInstance* BodyInstance : CollisionWrapper.BodyInstances)
 		{
-			check(BodyInstance);
 			PhysicsInterfaceTypes::FInlineShapeArray& Shapes = CachedShapes.Emplace_GetRef();
-			CollisionWrapper.GetShapeArray(BodyInstance, CollisionQueryFlag, Shapes);
+			if (BodyInstance)
+			{
+				CollisionWrapper.GetShapeArray(BodyInstance, CollisionQueryFlag, Shapes);
+			}
 		}
 	}
 }
