[PCG] Fix Get PCG Component data filtering and tracking

* If multiple PCG components were on the same actor, we had duplicated data
* Component selection was not set when getting PCG component data
* For dynamic tracking, add the extra dependency if it is getting PCG component data.

#rnx

diff --git a/Engine/Plugins/PCG/Source/PCG/Private/Elements/PCGDataFromActor.cpp b/Engine/Plugins/PCG/Source/PCG/Private/Elements/PCGDataFromActor.cpp
--- a/Engine/Plugins/PCG/Source/PCG/Private/Elements/PCGDataFromActor.cpp
+++ b/Engine/Plugins/PCG/Source/PCG/Private/Elements/PCGDataFromActor.cpp
@@ -328,6 +328,8 @@ bool FPCGDataFromActorElement::ExecuteInternal(FPCGContext* InContext) const
 
 	if (!Context->bPerformedQuery)
 	{
+		Context->ComponentSelector = Settings->ComponentSelector;
+
 		TFunction<bool(const AActor*)> BoundsCheck = [](const AActor*) -> bool { return true; };
 		const UPCGComponent* PCGComponent = Context->SourceComponent.IsValid() ? Context->SourceComponent.Get() : nullptr;
 		const AActor* Self = PCGComponent ? PCGComponent->GetOwner() : nullptr;
@@ -396,7 +398,7 @@ bool FPCGDataFromActorElement::ExecuteInternal(FPCGContext* InContext) const
 						{
 							if (Key.IsMatching(Actor, Component))
 							{
-								Context->FoundActors.Add(Actor);
+								Context->FoundActors.AddUnique(Actor);
 							}
 						}
 					}
@@ -409,7 +411,7 @@ bool FPCGDataFromActorElement::ExecuteInternal(FPCGContext* InContext) const
 						{
 							if (Key.IsMatching(Actor, Component))
 							{
-								Context->FoundActors.Add(Actor);
+								Context->FoundActors.AddUnique(Actor);
 							}
 						}
 					}
@@ -442,7 +444,6 @@ bool FPCGDataFromActorElement::ExecuteInternal(FPCGContext* InContext) const
 				}
 			}
 
-			Context->ComponentSelector = Settings->ComponentSelector;
 			Context->ComponentSelector.ComponentList = std::move(ComponentsFromInput);
 
 			Context->FoundActors = PCGActorSelector::FindActors(&Settings->ActorSelector, &Context->ComponentSelector, Context->SourceComponent.Get(), BoundsCheck, SelfIgnoreCheck, ActorsFromInput);
@@ -456,14 +457,26 @@ bool FPCGDataFromActorElement::ExecuteInternal(FPCGContext* InContext) const
 				DynamicTracking.EnableAndInitialize(Context, Context->PathsToObjectsAndDataIndex.Num());
 				for (const TTuple<FSoftObjectPath, int32, int32>& Path : Context->PathsToObjectsAndDataIndex)
 				{
-					DynamicTracking.AddToTracking(FPCGSelectionKey::CreateFromPath(Path.Get<0>()), Settings->ActorSelector.bMustOverlapSelf);
+					FPCGSelectionKey Key = FPCGSelectionKey::CreateFromPath(Path.Get<0>());
+					if (Settings->Mode == EPCGGetDataFromActorMode::GetDataFromPCGComponent || Settings->Mode == EPCGGetDataFromActorMode::GetDataFromPCGComponentOrParseComponents)
+					{
+						Key.SetExtraDependency(UPCGComponent::StaticClass());
+					}
+					
+					DynamicTracking.AddToTracking(std::move(Key), Settings->ActorSelector.bMustOverlapSelf);
 				}
 
 				DynamicTracking.Finalize(Context);
 			}
 			else if (Context->IsValueOverriden(GET_MEMBER_NAME_CHECKED(UPCGDataFromActorSettings, ActorSelector)))
 			{
-				FPCGDynamicTrackingHelper::AddSingleDynamicTrackingKey(Context, Settings->ActorSelector);
+				FPCGSelectionKey Key(Settings->ActorSelector);
+				if (Settings->Mode == EPCGGetDataFromActorMode::GetDataFromPCGComponent || Settings->Mode == EPCGGetDataFromActorMode::GetDataFromPCGComponentOrParseComponents)
+				{
+					Key.SetExtraDependency(UPCGComponent::StaticClass());
+				}
+				
+				FPCGDynamicTrackingHelper::AddSingleDynamicTrackingKey(Context, std::move(Key), Settings->ActorSelector.bMustOverlapSelf);
 			}
 #endif // WITH_EDITOR
 		}
