[PCG] Level to PCG : Fix Level Instance loading support for non-WP Level Instances (missing BlockTillLevelStreamingCompleted call)

#rnx

diff --git a/Engine/Plugins/PCG/Source/PCGEditor/Private/PCGLevelToAsset.cpp b/Engine/Plugins/PCG/Source/PCGEditor/Private/PCGLevelToAsset.cpp
--- a/Engine/Plugins/PCG/Source/PCGEditor/Private/PCGLevelToAsset.cpp
+++ b/Engine/Plugins/PCG/Source/PCGEditor/Private/PCGLevelToAsset.cpp
@@ -232,6 +232,9 @@ bool UPCGLevelToAsset::BP_ExportWorld_Implementation(UWorld* World, const FStrin
 	// Load all actors
 	bool bProcessedNewActors = LoadAllActors(World->GetWorldPartition());
 		
+	// Make sure to call this once as non World Partition levels do not load new actors
+	World->BlockTillLevelStreamingCompleted();
+
 	// Load all level instances and their actors recursively
 	TSet<ULevelStreaming*> ProcessedLevelStreamings;
 	
