[PCG] GPU SM Spawner - add missing newlines to shader code

#rnx

diff --git a/Engine/Plugins/PCG/Source/PCG/Private/Compute/DataInterfaces/Elements/PCGStaticMeshSpawnerDataInterface.cpp b/Engine/Plugins/PCG/Source/PCG/Private/Compute/DataInterfaces/Elements/PCGStaticMeshSpawnerDataInterface.cpp
--- a/Engine/Plugins/PCG/Source/PCG/Private/Compute/DataInterfaces/Elements/PCGStaticMeshSpawnerDataInterface.cpp
+++ b/Engine/Plugins/PCG/Source/PCG/Private/Compute/DataInterfaces/Elements/PCGStaticMeshSpawnerDataInterface.cpp
@@ -97,12 +97,12 @@ void UPCGStaticMeshSpawnerDataInterface::GetHLSL(FString& OutHLSL, FString const
 		"int {DataInterfaceName}_SelectorAttributeId;\n"
 		"uint {DataInterfaceName}_NumAttributes;\n"
 		"uint {DataInterfaceName}_NumPrimitives;\n"
-		"uint {DataInterfaceName}_ApplyBounds;"
+		"uint {DataInterfaceName}_ApplyBounds;\n"
 		"int {DataInterfaceName}_SelectedMeshAttributeId;\n"
 		"uint4 {DataInterfaceName}_AttributeIdOffsetStrides[{MaxAttributes}];\n"
-		"float4 {DataInterfaceName}_PrimitiveMeshBoundsMin[{MaxPrimitives}];"
-		"float4 {DataInterfaceName}_PrimitiveMeshBoundsMax[{MaxPrimitives}];"
-		"int4 {DataInterfaceName}_PrimitiveStringKeys[{MaxPrimitives}];"
+		"float4 {DataInterfaceName}_PrimitiveMeshBoundsMin[{MaxPrimitives}];\n"
+		"float4 {DataInterfaceName}_PrimitiveMeshBoundsMax[{MaxPrimitives}];\n"
+		"int4 {DataInterfaceName}_PrimitiveStringKeys[{MaxPrimitives}];\n"
 		"DECLARE_SCALAR_ARRAY(float, {DataInterfaceName}_SelectionCDF, {MaxPrimitives});\n"
 		"\n"
 		"int SMSpawner_GetSelectorAttributeId_{DataInterfaceName}()\n"
@@ -166,7 +166,7 @@ void UPCGStaticMeshSpawnerDataInterface::GetHLSL(FString& OutHLSL, FString const
 		"int SMSpawner_GetSelectedMeshAttributeId_{DataInterfaceName}()\n"
 		"{\n"
 		"	return {DataInterfaceName}_SelectedMeshAttributeId;\n"
-		"}\n"
+		"}\n\n"
 	), TemplateArgs);
 }
 
