[PCG] GPU - rework attribute tracking in compute graphs.

The broad ideas were correct but we took some wrong turns with how we use the attribute table and lost the ability to distinguish attributes across different branches of a compute graph.

* Reverted attribute table back to using (name, type) key.
* Improved naming around "static" attribute table that is built at compile time.
* Added function to augment the runtime attribute table with attributes from input data as part of initialization.
* Built-in attributes like point properties are not needed in the table and are removed. This removes some ropey code and fixes some bugs, but otherwise does not impact values of Ids.
* Fixed the attribute type resolution for the instance data packing. Here the attribute name is known but not type. Instead of pushing the attribute with an anonymous type into attribute table, resolve the attribute using node input data at execution time by looking for the attribute by name.
* This runtime resolution also means we can use whatever type is in the incoming data description, and this description is correctly tracks attributes across different branches.

#rnx

diff --git a/Engine/Plugins/PCG/Source/PCG/Private/Compute/DataInterfaces/Elements/PCGStaticMeshSpawnerDataInterface.cpp b/Engine/Plugins/PCG/Source/PCG/Private/Compute/DataInterfaces/Elements/PCGStaticMeshSpawnerDataInterface.cpp
--- a/Engine/Plugins/PCG/Source/PCG/Private/Compute/DataInterfaces/Elements/PCGStaticMeshSpawnerDataInterface.cpp
+++ b/Engine/Plugins/PCG/Source/PCG/Private/Compute/DataInterfaces/Elements/PCGStaticMeshSpawnerDataInterface.cpp
@@ -183,12 +183,16 @@ UComputeDataProvider* UPCGStaticMeshSpawnerDataInterface::CreateDataProvider(TOb
 	{
 		DataProvider->AttributeIdOffsetStrides = Primitives->AttributeIdOffsetStrides;
 		DataProvider->SelectionCDF = Primitives->SelectionCDF;
-		DataProvider->SelectorAttributeId = Primitives->SelectorAttributeId;
 		DataProvider->PrimitiveStringKeys = Primitives->PrimitiveStringKeys;
 		DataProvider->PrimitiveMeshBounds = Primitives->PrimitiveMeshBounds;
 
-		const FPCGKernelAttributeIDAndType* FoundAttribute = Binding->GetAttributeLookupTable().Find(DataProvider->Settings->OutAttributeName);
-		DataProvider->SelectedMeshAttributeId = ensure(FoundAttribute) ? FoundAttribute->Id : -1;
+		// We only know name at compile time (not type), so we resolved this Id at execution time by inspecting the data arriving on the input pin. 
+		DataProvider->SelectorAttributeId = Primitives->SelectorAttributeId;
+
+		// We know the name and type of the selected mesh attribute statically and declared the attribute in GetKernelAttributeKeys,
+		// so the attribute ID should be present in the attribute table.
+		DataProvider->SelectedMeshAttributeId = Binding->GetAttributeId(DataProvider->Settings->OutAttributeName, EPCGKernelAttributeType::StringKey);
+		ensure(DataProvider->SelectedMeshAttributeId != INDEX_NONE);
 	}
 
 	return DataProvider;
diff --git a/Engine/Plugins/PCG/Source/PCG/Private/Compute/DataInterfaces/Elements/PCGStaticMeshSpawnerDataInterface.h b/Engine/Plugins/PCG/Source/PCG/Private/Compute/DataInterfaces/Elements/PCGStaticMeshSpawnerDataInterface.h
--- a/Engine/Plugins/PCG/Source/PCG/Private/Compute/DataInterfaces/Elements/PCGStaticMeshSpawnerDataInterface.h
+++ b/Engine/Plugins/PCG/Source/PCG/Private/Compute/DataInterfaces/Elements/PCGStaticMeshSpawnerDataInterface.h
@@ -57,14 +57,16 @@ public:
 	UPROPERTY()
 	TArray<float> SelectionCDF;
 
+	/** Attribute Id for mesh selector. */
 	UPROPERTY()
-	int32 SelectorAttributeId = -1;
+	int32 SelectorAttributeId = INDEX_NONE;
 
 	UPROPERTY()
 	int32 NumInputPoints = 0;
 
+	/** Output attribute Id for the output selected mesh. */
 	UPROPERTY()
-	int32 SelectedMeshAttributeId = -1;
+	int32 SelectedMeshAttributeId = INDEX_NONE;
 };
 
 class FPCGStaticMeshSpawnerDataProviderProxy : public FComputeDataProviderRenderProxy
diff --git a/Engine/Plugins/PCG/Source/PCG/Private/Compute/DataInterfaces/PCGDataCollectionUploadDataInterface.cpp b/Engine/Plugins/PCG/Source/PCG/Private/Compute/DataInterfaces/PCGDataCollectionUploadDataInterface.cpp
--- a/Engine/Plugins/PCG/Source/PCG/Private/Compute/DataInterfaces/PCGDataCollectionUploadDataInterface.cpp
+++ b/Engine/Plugins/PCG/Source/PCG/Private/Compute/DataInterfaces/PCGDataCollectionUploadDataInterface.cpp
@@ -25,10 +25,9 @@ UComputeDataProvider* UPCGDataCollectionUploadDataInterface::CreateDataProvider(
 	// TODO: Ideally we could call ProducerSettings->ComputeOutputPinDataDesc() but some settings do not have associated nodes/pins.
 	check(!DownstreamInputPinLabelAliases.IsEmpty());
 	Provider->PinDesc = FPCGDataCollectionDesc::BuildFromInputDataCollectionAndInputPinLabel(
-		Binding->DataForGPU.InputDataCollection,
+		Binding->InputDataCollection,
 		DownstreamInputPinLabelAliases[0],
-		Binding->GetAttributeLookupTable(),
-		Binding->GetStringTable());
+		Binding);
 
 	Provider->DownstreamInputPinLabels = DownstreamInputPinLabelAliases;
 
@@ -42,7 +41,7 @@ FComputeDataProviderRenderProxy* UPCGDataProviderDataCollectionUpload::GetRender
 
 	// Use any downstream input pin label to grab data from the collection.
 	check(!DownstreamInputPinLabels.IsEmpty());
-	PinDesc.PackDataCollection(Binding->DataForGPU.InputDataCollection, DownstreamInputPinLabels[0], Binding->GetStringTable(), PackedDataCollection);
+	PinDesc.PackDataCollection(Binding->InputDataCollection, DownstreamInputPinLabels[0], Binding->GetStringTable(), PackedDataCollection);
 
 	if (PCGComputeHelpers::IsBufferSizeTooLarge(PackedDataCollection.Num() * PackedDataCollection.GetTypeSize()))
 	{
diff --git a/Engine/Plugins/PCG/Source/PCG/Private/Compute/DataInterfaces/PCGLandscapeDataInterface.cpp b/Engine/Plugins/PCG/Source/PCG/Private/Compute/DataInterfaces/PCGLandscapeDataInterface.cpp
--- a/Engine/Plugins/PCG/Source/PCG/Private/Compute/DataInterfaces/PCGLandscapeDataInterface.cpp
+++ b/Engine/Plugins/PCG/Source/PCG/Private/Compute/DataInterfaces/PCGLandscapeDataInterface.cpp
@@ -287,7 +287,7 @@ UComputeDataProvider* UPCGLandscapeDataInterface::CreateDataProvider(TObjectPtr<
 
 	// Take any input pin label alias to obtain the data from the input data collection.
 	check(!DownstreamInputPinLabelAliases.IsEmpty());
-	const TArray<FPCGTaggedData> TaggedDatas = Binding->DataForGPU.InputDataCollection.GetInputsByPin(DownstreamInputPinLabelAliases[0]);
+	const TArray<FPCGTaggedData> TaggedDatas = Binding->InputDataCollection.GetInputsByPin(DownstreamInputPinLabelAliases[0]);
 
 	if (!TaggedDatas.IsEmpty())
 	{
diff --git a/Engine/Plugins/PCG/Source/PCG/Private/Compute/DataInterfaces/PCGTextureDataInterface.cpp b/Engine/Plugins/PCG/Source/PCG/Private/Compute/DataInterfaces/PCGTextureDataInterface.cpp
--- a/Engine/Plugins/PCG/Source/PCG/Private/Compute/DataInterfaces/PCGTextureDataInterface.cpp
+++ b/Engine/Plugins/PCG/Source/PCG/Private/Compute/DataInterfaces/PCGTextureDataInterface.cpp
@@ -78,7 +78,7 @@ UComputeDataProvider* UPCGTextureDataInterface::CreateDataProvider(TObjectPtr<UO
 
 	// Take any input pin label alias to obtain the data from the input data collection.
 	check(!DownstreamInputPinLabelAliases.IsEmpty());
-	const TArray<FPCGTaggedData> TaggedDatas = Binding->DataForGPU.InputDataCollection.GetInputsByPin(DownstreamInputPinLabelAliases[0]);
+	const TArray<FPCGTaggedData> TaggedDatas = Binding->InputDataCollection.GetInputsByPin(DownstreamInputPinLabelAliases[0]);
 
 	if (!TaggedDatas.IsEmpty())
 	{
diff --git a/Engine/Plugins/PCG/Source/PCG/Private/Compute/Elements/PCGComputeGraphElement.cpp b/Engine/Plugins/PCG/Source/PCG/Private/Compute/Elements/PCGComputeGraphElement.cpp
--- a/Engine/Plugins/PCG/Source/PCG/Private/Compute/Elements/PCGComputeGraphElement.cpp
+++ b/Engine/Plugins/PCG/Source/PCG/Private/Compute/Elements/PCGComputeGraphElement.cpp
@@ -155,7 +155,7 @@ bool FPCGComputeGraphElement::ExecuteInternal(FPCGContext* InContext) const
 		UPCGDataBinding* DataBinding = FPCGContext::NewObject_AnyThread<UPCGDataBinding>(Context);
 		Context->DataBinding.Reset(DataBinding);
 
-		Context->DataBinding->Initialize(Context->ComputeGraph.Get(), Context->SourceComponent, Context->InputData, Context->ComputeGraph->GetAttributeLookupTable());
+		Context->DataBinding->Initialize(Context->ComputeGraph.Get(), Context->SourceComponent, Context->InputData, Context->ComputeGraph->GetAttributeStaticTable());
 
 		// Perform validation after input data is initialized.
 		for (TSoftObjectPtr<const UPCGNode> Node : Context->ComputeGraph->KernelToNode)
@@ -371,7 +371,7 @@ bool FPCGComputeGraphElement::SetupProceduralISMComponents(FPCGContext* InContex
 					continue;
 				}
 
-				if (!InBinding->GetAttributeLookupTable().Find(AttributeName) && !InBinding->GetInputAttributeTable().Find(AttributeName))
+				if (!InputDataDesc.ContainsAttributeOnAnyData(AttributeName))
 				{
 					PCG_KERNEL_VALIDATION_ERR(Context, SpawnerSettings, /*bQuiet=*/false,
 						FText::Format(LOCTEXT("InstanceDataAttributeNotFound", "Instance data attribute '{0}' not found."), FText::FromName(AttributeName)));
@@ -380,68 +380,54 @@ bool FPCGComputeGraphElement::SetupProceduralISMComponents(FPCGContext* InContex
 				}
 			}
 
-			PCGDataForGPUHelpers::ComputeCustomFloatPacking(AttributeNames, InBinding, InputDataDesc, CustomFloatCount, AttributeIdOffsetStrides);
+			PCGDataForGPUHelpers::ComputeCustomFloatPacking(InContext, SpawnerSettings, AttributeNames, InBinding, InputDataDesc, CustomFloatCount, AttributeIdOffsetStrides);
 		}
 
 		TArray<float> PrimitiveSelectionCDF;
-		int32 SelectorAttributeId = -1;
+		int32 SelectorAttributeId = INDEX_NONE;
 		TArray<int32> PrimitiveStringKeys;
 		TArray<FBox> PrimitiveMeshBounds;
 		TArray<FPCGProceduralISMComponentDescriptor> ComponentsToCreate;
 
 		if (const UPCGMeshSelectorByAttribute* SelectorByAttribute = Cast<UPCGMeshSelectorByAttribute>(SpawnerSettings->MeshSelectorParameters))
 		{
-			const FName SelectedName = SelectorByAttribute->AttributeName;
-
-			const FPCGKernelAttributeIDAndType* FoundAttribute = InBinding->GetAttributeLookupTable().Find(SelectedName);
-
-			if (!FoundAttribute)
-			{
-				FoundAttribute = InBinding->GetInputAttributeTable().Find(SelectedName);
-			}
-
-			if (!FoundAttribute)
-			{
-				PCG_KERNEL_VALIDATION_ERR(Context, SpawnerSettings, /*bQuiet=*/false,
-					FText::Format(LOCTEXT("MeshSelectorAttributeNotFound", "Mesh selector attribute '{0}' not found."), FText::FromName(SelectedName)));
-
-				continue;
-			}
-
-			if (FoundAttribute->Type != EPCGKernelAttributeType::StringKey)
-			{
-				PCG_KERNEL_VALIDATION_ERR(Context, SpawnerSettings, /*bQuiet=*/false,
-					FText::Format(LOCTEXT("MeshSelectorAttributeInvalidType", "Mesh selector attribute '{0}' must be of type String Key."), FText::FromName(SelectedName)));
-
-				continue;
-			}
-
-			SelectorAttributeId = FoundAttribute->Id;
+			const FName SelectorName = SelectorByAttribute->AttributeName;
 
 			// Compute how many unique incoming values we will receive.
 			TSet<int32> UniqueStringKeys;
+
 			for (const FPCGDataDesc& Desc : InputDataDesc.DataDescs)
 			{
 				for (const FPCGKernelAttributeDesc& AttributeDesc : Desc.AttributeDescs)
 				{
-					if (AttributeDesc.Name == SelectedName)
+					if (AttributeDesc.Name == SelectorName)
 					{
-						if (AttributeDesc.Type == EPCGKernelAttributeType::StringKey)
-						{
-							UniqueStringKeys.Append(AttributeDesc.UniqueStringKeys);
-						}
-						else
+						if (AttributeDesc.Type != EPCGKernelAttributeType::StringKey)
 						{
 							// Currently only a single type per attribute name is supported (the name/type in the attribute table). It's possible to wire a graph
 							// for which an attribute is present with multiple types. Warn if this is encountered.
 							PCG_KERNEL_VALIDATION_WARN(Context, SpawnerSettings, /*bQuiet=*/false, FText::Format(
 								LOCTEXT("MeshSelectorAttributeNotUsable", "Attribute '{0}' not usable for mesh selection, only attributes of type String Key are supported."),
-								FText::FromName(SelectedName)));
+								FText::FromName(SelectorName)));
+
+							continue;
 						}
+
+						SelectorAttributeId = AttributeDesc.Index;
+						UniqueStringKeys.Append(AttributeDesc.UniqueStringKeys);
 					}
 				}
 			}
 
+			if (SelectorAttributeId == INDEX_NONE)
+			{
+				PCG_KERNEL_VALIDATION_ERR(Context, SpawnerSettings, /*bQuiet=*/false, FText::Format(
+					LOCTEXT("MeshSelectorAttributeNotFound", "Mesh selector attribute '{0}' not found."),
+					FText::FromName(SelectorName)));
+
+				continue;
+			}
+
 			for (int32 StringKey : UniqueStringKeys)
 			{
 				if (!ensure(InBinding->GetStringTable().IsValidIndex(StringKey)))
diff --git a/Engine/Plugins/PCG/Source/PCG/Private/Compute/Elements/PCGCustomHLSL.cpp b/Engine/Plugins/PCG/Source/PCG/Private/Compute/Elements/PCGCustomHLSL.cpp
--- a/Engine/Plugins/PCG/Source/PCG/Private/Compute/Elements/PCGCustomHLSL.cpp
+++ b/Engine/Plugins/PCG/Source/PCG/Private/Compute/Elements/PCGCustomHLSL.cpp
@@ -190,21 +190,20 @@ FPCGElementPtr UPCGCustomHLSLSettings::CreateElement() const
 int UPCGCustomHLSLSettings::GetProcessingElemCountForInputPin(const UPCGPin* InputPin, const UPCGDataBinding* Binding) const
 {
 	check(InputPin && Binding);
-	const FPCGDataForGPU& DataForGPU = Binding->DataForGPU;
 
 	// Upper bound estimate of total number of data elements expected to arrive at this pin.
 	int ProcessingElemCount = 0;
 
-	if (DataForGPU.InputPins.Contains(InputPin))
+	if (Binding->InputPins.Contains(InputPin))
 	{
 		FName PinLabel = InputPin->Properties.Label;
-		if (const FName* PinLabelAlias = DataForGPU.InputPinLabelAliases.Find(InputPin))
+		if (const FName* PinLabelAlias = Binding->InputPinLabelAliases.Find(InputPin))
 		{
 			PinLabel = *PinLabelAlias;
 		}
 
 		// Data coming straight from CPU, know its size now so count it.
-		TArray<FPCGTaggedData> DataForPin = DataForGPU.InputDataCollection.GetInputsByPin(PinLabel);
+		TArray<FPCGTaggedData> DataForPin = Binding->InputDataCollection.GetInputsByPin(PinLabel);
 		for (const FPCGTaggedData& Data : DataForPin)
 		{
 			ProcessingElemCount += PCGComputeHelpers::GetElementCount(Data.Data);
@@ -329,8 +328,6 @@ bool UPCGCustomHLSLSettings::ComputeOutputPinDataDesc(const UPCGPin* OutputPin,
 		PCGComputeHelpers::ComputeOutputPinDataDesc(*Properties, this, Binding, PinDesc);
 	}
 
-	const TMap<FName, FPCGKernelAttributeIDAndType>& GlobalAttributeLookupTable = Binding->GetAttributeLookupTable();
-
 	for (const FPCGKernelAttributeKey& AttributeKey : KernelAttributeKeys)
 	{
 		// Add attributes that will be created for this pin on the GPU.
@@ -345,9 +342,11 @@ bool UPCGCustomHLSLSettings::ComputeOutputPinDataDesc(const UPCGPin* OutputPin,
 
 			for (FPCGDataDesc& DataDesc : PinDesc.DataDescs)
 			{
-				if (const FPCGKernelAttributeIDAndType* IDAndType = GlobalAttributeLookupTable.Find(AttributeKey.Name))
+				const int32 AttributeId = Binding->GetAttributeId(AttributeKey);
+
+				if (AttributeId != INDEX_NONE)
 				{
-					const FPCGKernelAttributeDesc AttributeDesc(IDAndType->Id, IDAndType->Type, AttributeKey.Name);
+					const FPCGKernelAttributeDesc AttributeDesc(AttributeId, AttributeKey.Type, AttributeKey.Name);
 					DataDesc.AttributeDescs.AddUnique(AttributeDesc);
 				}
 			}
@@ -1450,7 +1449,7 @@ bool UPCGCustomHLSLSettings::IsKernelValid(FPCGContext* InContext, bool bQuiet)
 	// Validate attributes
 	for (const FPCGKernelAttributeKey& AttributeKey : KernelAttributeKeys)
 	{
-		if (AttributeKey.Type == EPCGKernelAttributeType::None || AttributeKey.Type == EPCGKernelAttributeType::Invalid)
+		if (AttributeKey.Type == EPCGKernelAttributeType::Invalid)
 		{
 			const UEnum* AttributeTypeEnum = StaticEnum<EPCGKernelAttributeType>();
 			check(AttributeTypeEnum);
@@ -1845,7 +1844,7 @@ bool UPCGCustomHLSLSettings::GetGeneratedFunctions(FString* OutFunctions, FText*
 	return true;
 }
 
-FString UPCGCustomHLSLSettings::GetCookedKernelSource(const TMap<FName, FPCGKernelAttributeIDAndType>& GlobalAttributeLookupTable) const
+FString UPCGCustomHLSLSettings::GetCookedKernelSource(const FPCGKernelAttributeTable& InAttributeStaticTable) const
 {
 	const FIntVector GroupSize(PCGComputeConstants::THREAD_GROUP_SIZE, 1, 1);
 
@@ -1861,11 +1860,13 @@ FString UPCGCustomHLSLSettings::GetCookedKernelSource(const TMap<FName, FPCGKern
 	Functions.ReplaceInline(TEXT("\r"), TEXT(""));
 #endif
 
-	for (const TPair<FName, FPCGKernelAttributeIDAndType>& Pair : GlobalAttributeLookupTable)
+	// We inject attribute IDs directly into the source. This is most efficient and saves us passing them into the kernel. However
+	// the trade off is that a shader with a shared source will generate multiple variants if used in different compute graphs with
+	// different attribute IDs.
+	for (const FPCGKernelAttributeKey& AttributeKey : KernelAttributeKeys)
 	{
-		FPCGKernelAttributeKey AttributeKey(Pair.Key, Pair.Value.Type);
 		const FString SourceDefinition = PCGHLSLElement::GetKernelAttributeKeyAsString(AttributeKey);
-		const FString AttributeIndexAsString = FString::FromInt(Pair.Value.Id);
+		const FString AttributeIndexAsString = FString::FromInt(InAttributeStaticTable.GetAttributeId(AttributeKey));
 
 		Source.ReplaceInline(*SourceDefinition, *AttributeIndexAsString);
 		Functions.ReplaceInline(*SourceDefinition, *AttributeIndexAsString);
diff --git a/Engine/Plugins/PCG/Source/PCG/Private/Compute/Elements/PCGCustomHLSL.h b/Engine/Plugins/PCG/Source/PCG/Private/Compute/Elements/PCGCustomHLSL.h
--- a/Engine/Plugins/PCG/Source/PCG/Private/Compute/Elements/PCGCustomHLSL.h
+++ b/Engine/Plugins/PCG/Source/PCG/Private/Compute/Elements/PCGCustomHLSL.h
@@ -65,7 +65,7 @@ public:
 	virtual void ApplyPreconfiguredSettings(const FPCGPreConfiguredSettingsInfo& PreconfigureInfo) override;
 
 	virtual bool IsKernelValid(FPCGContext* InContext = nullptr, bool bQuiet = true) const override;
-	virtual FString GetCookedKernelSource(const TMap<FName, FPCGKernelAttributeIDAndType>& GlobalAttributeLookupTable) const override;
+	virtual FString GetCookedKernelSource(const FPCGKernelAttributeTable& InAttributeStaticTable) const override;
 	virtual void GetKernelAttributeKeys(TArray<FPCGKernelAttributeKey>& OutKeys) const override { OutKeys.Append(KernelAttributeKeys); }
 	virtual int ComputeKernelThreadCount(const UPCGDataBinding* Binding) const override;
 	virtual bool ComputeOutputPinDataDesc(const UPCGPin* OutputPin, const UPCGDataBinding* Binding, FPCGDataCollectionDesc& OutDesc) const override;
@@ -101,7 +101,9 @@ protected:
 	/** Enforce required pin settings and set display toggles to drive UI. */
 	void UpdatePinSettings();
 
-	/** Update the kernel keys read/written/created by this node. */
+	/** Update the kernel keys read/written/created by this node. In a pure sense this could register only the created attributes and instead let
+	* read/written attributes get registered by upstream nodes or initialized from input data, but this way we can statically know the attribute IDs
+	* and we can avoid resolving them at execution time. */
 	void UpdateAttributeKeys();
 
 #if WITH_EDITOR
diff --git a/Engine/Plugins/PCG/Source/PCG/Private/Compute/PCGComputeGraph.h b/Engine/Plugins/PCG/Source/PCG/Private/Compute/PCGComputeGraph.h
--- a/Engine/Plugins/PCG/Source/PCG/Private/Compute/PCGComputeGraph.h
+++ b/Engine/Plugins/PCG/Source/PCG/Private/Compute/PCGComputeGraph.h
@@ -24,8 +24,8 @@ public:
 	void OnKernelCompilationComplete(int32 InKernelIndex, FComputeKernelCompileResults const& InCompileResults) override;
 	//~End UComputeGraph interface
 
-	/** Get the global attribute information collated during compilation. */
-	const TMap<FName, FPCGKernelAttributeIDAndType>& GetAttributeLookupTable() const { return GlobalAttributeLookupTable; }
+	/** Get the attribute table collated statically during compilation. */
+	const FPCGKernelAttributeTable& GetAttributeStaticTable() const { return AttributeStaticTable; }
 
 	const TArray<FString>& GetStringTable() const { return StringTable; }
 
@@ -58,7 +58,7 @@ public:
 protected:
 	/** Global attribute information collated during compilation. */
 	UPROPERTY()
-	TMap<FName /* Attribute name */, FPCGKernelAttributeIDAndType> GlobalAttributeLookupTable;
+	FPCGKernelAttributeTable AttributeStaticTable;
 
 	/** String table collated during compilation. */
 	UPROPERTY()
diff --git a/Engine/Plugins/PCG/Source/PCG/Private/Compute/PCGDataBinding.cpp b/Engine/Plugins/PCG/Source/PCG/Private/Compute/PCGDataBinding.cpp
--- a/Engine/Plugins/PCG/Source/PCG/Private/Compute/PCGDataBinding.cpp
+++ b/Engine/Plugins/PCG/Source/PCG/Private/Compute/PCGDataBinding.cpp
@@ -14,18 +14,17 @@ void UPCGDataBinding::Initialize(
 	const UPCGComputeGraph* InComputeGraph,
 	TWeakObjectPtr<UPCGComponent> InSourceComponent,
 	const FPCGDataCollection& InComputeGraphElementInputData,
-	const TMap<FName, FPCGKernelAttributeIDAndType>& InStaticAttributeTable)
+	const FPCGKernelAttributeTable& InAttributeStaticTable)
 {
 	TRACE_CPUPROFILER_EVENT_SCOPE(UPCGDataBinding::Initialize);
 	check(InComputeGraph);
 
 	Graph = InComputeGraph;
 	SourceComponent = InSourceComponent;
-	GlobalAttributeLookupTable = InStaticAttributeTable;
 
-	InitializeInputData(InComputeGraphElementInputData);
+	AttributeTable = InAttributeStaticTable;
 
-	FillInMissingAttributeTableTypes(InComputeGraphElementInputData);
+	InitializeInputData(InComputeGraphElementInputData);
 
 #if PCG_KERNEL_LOGGING_ENABLED
 	if (Graph->bLogDataDescriptions)
@@ -49,11 +48,7 @@ bool UPCGDataBinding::ComputeCPUOutputPinDataDesc(const UPCGPin* OutputPin, FPCG
 	}
 
 	// Create description from compute graph element input data.
-	OutDesc = FPCGDataCollectionDesc::BuildFromInputDataCollectionAndInputPinLabel(
-		DataForGPU.InputDataCollection,
-		*FoundPinAlias,
-		GetAttributeLookupTable(),
-		GetStringTable());
+	OutDesc = FPCGDataCollectionDesc::BuildFromInputDataCollectionAndInputPinLabel(InputDataCollection, *FoundPinAlias, this);
 
 	return true;
 }
@@ -62,55 +57,45 @@ void UPCGDataBinding::InitializeInputData(const FPCGDataCollection& InComputeGra
 {
 	TRACE_CPUPROFILER_EVENT_SCOPE(UPCGDataBinding::InitializeInputData);
 
-	DataForGPU.InputDataCollection = InComputeGraphElementInputData;
+	InputDataCollection = InComputeGraphElementInputData;
 
 	// Link each input pin to the data collection, so that data providers can find the data.
 	for (const TSoftObjectPtr<const UPCGPin>& InputPinPtr : Graph->PinsReceivingDataFromCPU)
 	{
 		if (const UPCGPin* InputPin = InputPinPtr.Get())
 		{
-			DataForGPU.InputPins.Add(InputPin);
+			InputPins.Add(InputPin);
 		}
 	}
 
-	DataForGPU.InputPinLabelAliases = Graph->InputPinLabelAliases;
+	InputPinLabelAliases = Graph->InputPinLabelAliases;
 
-	BuildStringTable();
+	// Augment static string table with attributes incoming from CPU.
+	AddInputDataAttributesToTable(InComputeGraphElementInputData);
 
-	BuildInputAttributeTable(InComputeGraphElementInputData);
+	BuildStringTable();
 }
 
-void UPCGDataBinding::BuildInputAttributeTable(const FPCGDataCollection& InComputeGraphElementInputData)
+void UPCGDataBinding::AddInputDataAttributesToTable(const FPCGDataCollection& InComputeGraphElementInputData)
 {
-	const UEnum* PCGKernelAttributeTypeEnum = StaticEnum<EPCGKernelAttributeType>();
-	check(PCGKernelAttributeTypeEnum);
-
-	const FPCGDataCollectionDesc Desc = FPCGDataCollectionDesc::BuildFromDataCollection(InComputeGraphElementInputData, GetAttributeLookupTable(), GetStringTable());
-	for (const FPCGDataDesc& DataDesc : Desc.DataDescs)
+	for (const FPCGTaggedData& Data : InComputeGraphElementInputData.TaggedData)
 	{
-		for (const FPCGKernelAttributeDesc& AttrDesc : DataDesc.AttributeDescs)
+		const UPCGMetadata* Metadata = Data.Data ? Data.Data->ConstMetadata() : nullptr;
+		if (!Metadata)
 		{
-			// Skip non-attributes like point properties.
-			if (AttrDesc.Name == NAME_None)
-			{
-				continue;
-			}
+			continue;
+		}
 
-			if (FPCGKernelAttributeIDAndType* AttrIDAndType = InputAttributeTable.Find(AttrDesc.Name))
-			{
-				if (AttrIDAndType->Type != AttrDesc.Type)
-				{
-					UE_LOG(LogPCG, Warning, TEXT("Encountered attribute '%s' multiple times in input data with different types (%s vs %s). Attributes currently must have a single type in compute graphs."),
-						*AttrDesc.Name.ToString(),
-						*PCGKernelAttributeTypeEnum->GetNameStringByValue(static_cast<int64>(AttrDesc.Type)),
-						*PCGKernelAttributeTypeEnum->GetNameStringByValue(static_cast<int64>(AttrIDAndType->Type)));
-				}
-			}
-			else
-			{
-				const FPCGKernelAttributeIDAndType IDAndType(AttrDesc.Index, AttrDesc.Type);
-				InputAttributeTable.Add(AttrDesc.Name, IDAndType);
-			}
+		TArray<FName> AttributeNames;
+		TArray<EPCGMetadataTypes> AttributeTypes;
+		Metadata->GetAttributes(AttributeNames, AttributeTypes);
+
+		for (int Index = 0; Index < AttributeNames.Num(); ++Index)
+		{
+			const FName AttributeName = AttributeNames[Index];
+			const EPCGKernelAttributeType AttributeType = PCGDataForGPUHelpers::GetAttributeTypeFromMetadataType(AttributeTypes[Index]);
+
+			AttributeTable.AddAttribute(AttributeName, AttributeType);
 		}
 	}
 }
@@ -126,7 +111,7 @@ void UPCGDataBinding::BuildStringTable()
 	// Add any strings known statically at compilation time.
 	StringTable.Append(Graph->GetStringTable());
 
-	for (const FPCGTaggedData& Data : DataForGPU.InputDataCollection.TaggedData)
+	for (const FPCGTaggedData& Data : InputDataCollection.TaggedData)
 	{
 		const UPCGMetadata* Metadata = Data.Data ? Data.Data->ConstMetadata() : nullptr;
 		if (!Metadata)
@@ -198,56 +183,15 @@ void UPCGDataBinding::BuildStringTable()
 	}
 }
 
-void UPCGDataBinding::FillInMissingAttributeTableTypes(const FPCGDataCollection& InComputeGraphElementInputData)
-{
-	TRACE_CPUPROFILER_EVENT_SCOPE(UPCGDataBinding::FillInMissingAttributeTableTypes);
-
-	// Local to loop body but hoisted for perf.
-	TArray<FName> AttributeNames;
-	TArray<EPCGMetadataTypes> AttributeTypes;
-
-	for (TPair<FName, FPCGKernelAttributeIDAndType>& Entry : GlobalAttributeLookupTable)
-	{
-		if (Entry.Value.Type != EPCGKernelAttributeType::None)
-		{
-			continue;
-		}
-
-		for (const FPCGTaggedData& Data : InComputeGraphElementInputData.TaggedData)
-		{
-			AttributeNames.Reset();
-			AttributeTypes.Reset();
-
-			if (const UPCGMetadata* Metadata = Data.Data ? Data.Data->ConstMetadata() : nullptr)
-			{
-				Metadata->GetAttributes(AttributeNames, AttributeTypes);
-			}
-
-			for (int Index = 0; Index < AttributeNames.Num(); ++Index)
-			{
-				if (AttributeNames[Index] == Entry.Key)
-				{
-					Entry.Value.Type = PCGDataForGPUHelpers::GetAttributeTypeFromMetadataType(AttributeTypes[Index]);
-					break;
-				}
-			}
-
-			if (Entry.Value.Type != EPCGKernelAttributeType::None)
-			{
-				break;
-			}
-		}
-	}
-}
-
 void UPCGDataBinding::DebugLogDataDescriptions()
 {
 #if PCG_KERNEL_LOGGING_ENABLED
-	UE_LOG(LogPCG, Display, TEXT("\n### ATTRIBUTE TABLE ###"));
-	for (TPair<FName, FPCGKernelAttributeIDAndType>& AttributeEntry : GlobalAttributeLookupTable)
-	{
-		UE_LOG(LogPCG, Display, TEXT("\tName: %s\t\tID: %d\t\tType: %d"), *AttributeEntry.Get<0>().ToString(), AttributeEntry.Get<1>().Id, AttributeEntry.Get<1>().Type);
-	}
+	const UEnum* PCGDataTypeEnum = StaticEnum<EPCGDataType>();
+	const UEnum* PCGKernelAttributeTypeEnum = StaticEnum<EPCGKernelAttributeType>();
+	check(PCGDataTypeEnum && PCGKernelAttributeTypeEnum);
+
+	UE_LOG(LogPCG, Display, TEXT("\n### METADATA ATTRIBUTE TABLE ###"));
+	AttributeTable.DebugLog();
 
 	UE_LOG(LogPCG, Display, TEXT("\n### STRING TABLE ###"));
 	for (int32 i = 0; i < StringTable.Num(); ++i)
@@ -255,10 +199,6 @@ void UPCGDataBinding::DebugLogDataDescriptions()
 		UE_LOG(LogPCG, Display, TEXT("\t%d: %s"), i, *StringTable[i]);
 	}
 
-	const UEnum* PCGDataTypeEnum = StaticEnum<EPCGDataType>();
-	const UEnum* PCGKernelAttributeTypeEnum = StaticEnum<EPCGKernelAttributeType>();
-	check(PCGDataTypeEnum && PCGKernelAttributeTypeEnum);
-
 	auto LogDataDescription = [this, PCGDataTypeEnum, PCGKernelAttributeTypeEnum](const FPCGDataCollectionDesc& InDataDesc)
 	{
 		for (int32 DataIndex = 0; DataIndex < InDataDesc.DataDescs.Num(); ++DataIndex)
@@ -266,7 +206,7 @@ void UPCGDataBinding::DebugLogDataDescriptions()
 			const FPCGDataDesc& DataDesc = InDataDesc.DataDescs[DataIndex];
 
 			UE_LOG(LogPCG, Display, TEXT("\t\tData Index %d"), DataIndex);
-			UE_LOG(LogPCG, Display, TEXT("\t\t\tType: %s"), *PCGDataTypeEnum->GetNameStringByValue(static_cast<int64>(DataDesc.Type)));
+			UE_LOG(LogPCG, Display, TEXT("\t\t\tType: %s"), *PCGDataTypeEnum->GetDisplayNameTextByValue(static_cast<int64>(DataDesc.Type)).ToString());
 			UE_LOG(LogPCG, Display, TEXT("\t\t\tElementCount: %d"), DataDesc.ElementCount);
 			UE_LOG(LogPCG, Display, TEXT("\t\t\tAttributes (%d)"), DataDesc.AttributeDescs.Num());
 				
@@ -284,7 +224,7 @@ void UPCGDataBinding::DebugLogDataDescriptions()
 				UE_LOG(LogPCG, Display, TEXT("\t\t\t\tID: %d\t\tName: %s\t\tType: %s\t\tUniqueStringKeys%s"),
 					AttributeDesc.Index,
 					*AttributeDesc.Name.ToString(),
-					*(PCGKernelAttributeTypeEnum->GetNameStringByValue(static_cast<int64>(AttributeDesc.Type))),
+					*PCGKernelAttributeTypeEnum->GetDisplayNameTextByValue(static_cast<int64>(AttributeDesc.Type)).ToString(),
 					UniqueStringKeys.IsEmpty() ? TEXT(": ") : *UniqueStringKeys);
 			}
 		}
diff --git a/Engine/Plugins/PCG/Source/PCG/Private/Compute/PCGDataBinding.h b/Engine/Plugins/PCG/Source/PCG/Private/Compute/PCGDataBinding.h
--- a/Engine/Plugins/PCG/Source/PCG/Private/Compute/PCGDataBinding.h
+++ b/Engine/Plugins/PCG/Source/PCG/Private/Compute/PCGDataBinding.h
@@ -34,7 +34,7 @@ struct FPCGSpawnerPrimitives
 	TArray<FUintVector4> AttributeIdOffsetStrides;
 
 	UPROPERTY()
-	int32 SelectorAttributeId = -1;
+	int32 SelectorAttributeId = INDEX_NONE;
 
 	UPROPERTY()
 	TArray<int32> PrimitiveStringKeys;
@@ -49,13 +49,10 @@ class UPCGDataBinding : public UObject
 	GENERATED_BODY()
 
 public:
-	void Initialize(const UPCGComputeGraph* InComputeGraph, TWeakObjectPtr<UPCGComponent> InSourceComponent, const FPCGDataCollection& InComputeGraphElementInputData, const TMap<FName, FPCGKernelAttributeIDAndType>& InStaticAttributeTable);
+	void Initialize(const UPCGComputeGraph* InComputeGraph, TWeakObjectPtr<UPCGComponent> InSourceComponent, const FPCGDataCollection& InComputeGraphElementInputData, const FPCGKernelAttributeTable& InAttributeStaticTable);
 
-	/** Get the global attribute information for this graph instance. */
-	const TMap<FName, FPCGKernelAttributeIDAndType>& GetAttributeLookupTable() const { return GlobalAttributeLookupTable; }
-
-	/** Get the table of attributes present in the compute graph element input data. */
-	const TMap<FName, FPCGKernelAttributeIDAndType>& GetInputAttributeTable() const { return InputAttributeTable; }
+	int32 GetAttributeId(const FPCGKernelAttributeKey& InAttribute) const { return AttributeTable.GetAttributeId(InAttribute); }
+	int32 GetAttributeId(FName InName, EPCGKernelAttributeType InType) const { return AttributeTable.GetAttributeId(InName, InType); }
 
 	/** Get the set of all unique strings entering the compute graph. */
 	const TArray<FString>& GetStringTable() const { return StringTable; }
@@ -66,40 +63,46 @@ public:
 private:
 	void InitializeInputData(const FPCGDataCollection& InComputeGraphElementInputData);
 
-	void BuildInputAttributeTable(const FPCGDataCollection& InComputeGraphElementInputData);
+	/** Loop over all metadata attributes in all input data ensure all attributes are registered in attribute table. */
+	void AddInputDataAttributesToTable(const FPCGDataCollection& InComputeGraphElementInputData);
 
 	void BuildStringTable();
 
-	/** Some attributes in table may be missing type information that we could not infer statically. Fill in missing types from the execution-time data. */
-	void FillInMissingAttributeTableTypes(const FPCGDataCollection& InComputeGraphElementInputData);
-
 	void DebugLogDataDescriptions();
 
 public:
 	UPROPERTY()
 	TWeakObjectPtr<UPCGComponent> SourceComponent;
 
+	/** Compute graph element output data. Data items are labelled with unique virtual output pin labels so that the can be routed correctly by the graph
+	 * executor to downstream nodes. */
 	UPROPERTY()
 	FPCGDataCollection OutputDataCollection;
 
 	UPROPERTY()
 	TObjectPtr<const UPCGComputeGraph> Graph = nullptr;
 
-	// Data collection pointer + set of pins that get data from the collection (cross barrier)
+	/** All the node input pins that have edges that cross from CPU to GPU. */
 	UPROPERTY()
-	FPCGDataForGPU DataForGPU;
+	TSet<TSoftObjectPtr<const UPCGPin>> InputPins;
 
+	/** Input pin label aliases generated during compilation so that every GPU input pin receiving data from CPU has a unique label which can be used
+	 * to retrieve the data for this pin from the compute graph element data collection. */
 	UPROPERTY()
-	TMap<TObjectPtr<const UPCGSettings>, FPCGSpawnerPrimitives> MeshSpawnersToPrimitives;
+	TMap<TSoftObjectPtr<const UPCGPin>, FName> InputPinLabelAliases;
 
-protected:
-	/** Attributes read/written/created by kernels. */
+	/** Data arriving on compute graph element. Since the compute graph is collapsed to a single element, all data crossing from CPU to GPU is in a single collection. */
 	UPROPERTY()
-	TMap<FName /* Attribute name */, FPCGKernelAttributeIDAndType> GlobalAttributeLookupTable;
+	FPCGDataCollection InputDataCollection;
 
-	/** Attributes present across all input data. */
 	UPROPERTY()
-	TMap<FName /* Attribute name */, FPCGKernelAttributeIDAndType> InputAttributeTable;
+	TMap<TObjectPtr<const UPCGSettings>, FPCGSpawnerPrimitives> MeshSpawnersToPrimitives;
+
+protected:
+	/** All attributes present in graph at execution time. Coherent across all graph branches. Seeded using statically known created attributes at compile time, then
+	 * augmented with incoming attributes from input data collection at runtime. */
+	UPROPERTY()
+	FPCGKernelAttributeTable AttributeTable;
 
 	UPROPERTY()
 	TArray<FString> StringTable;
diff --git a/Engine/Plugins/PCG/Source/PCG/Private/Compute/PCGDataForGPU.cpp b/Engine/Plugins/PCG/Source/PCG/Private/Compute/PCGDataForGPU.cpp
--- a/Engine/Plugins/PCG/Source/PCG/Private/Compute/PCGDataForGPU.cpp
+++ b/Engine/Plugins/PCG/Source/PCG/Private/Compute/PCGDataForGPU.cpp
@@ -18,6 +18,8 @@
 
 #include "Async/ParallelFor.h"
 
+#define LOCTEXT_NAMESPACE "PCGDataForGPU"
+
 using namespace PCGComputeConstants;
 
 namespace PCGDataForGPUConstants
@@ -539,6 +541,8 @@ namespace PCGDataForGPUHelpers
 	}
 
 	void ComputeCustomFloatPacking(
+		FPCGContext* InContext,
+		const UPCGSettings* InSettings,
 		TArray<FName>& InAttributeNames,
 		const UPCGDataBinding* InBinding,
 		const FPCGDataCollectionDesc& InDataCollectionDescription,
@@ -547,28 +551,32 @@ namespace PCGDataForGPUHelpers
 	{
 		check(InBinding);
 		check(InBinding->Graph);
-		const TMap<FName, FPCGKernelAttributeIDAndType>& GlobalAttributeLookupTable = InBinding->GetAttributeLookupTable();
 
 		uint32 OffsetFloats = 0;
 
 		for (FName AttributeName : InAttributeNames)
 		{
-			const FPCGKernelAttributeIDAndType* FoundAttribute = GlobalAttributeLookupTable.Find(AttributeName);
-			if (!FoundAttribute)
-			{
-				continue;
-			}
+			// We need to do a lookup here as the user does not provide the full attribute description, only the attribute name. So we see what we can
+			// find in the input data similar to the CPU SM Spawner.
+			FPCGKernelAttributeDesc AttributeDesc;
+			bool bConflictingTypesInData = false;
+			InDataCollectionDescription.GetAttributeDesc(AttributeName, AttributeDesc, bConflictingTypesInData);
 
-			const EPCGKernelAttributeType AttributeType = FoundAttribute->Type;
-			if (AttributeType == EPCGKernelAttributeType::None)
+			if (bConflictingTypesInData)
 			{
-				continue;
+				// TODO: in future we could partition the execution - run it once per attribute type in input data.
+				PCG_KERNEL_VALIDATION_ERR(InContext, InSettings, /*bQuiet=*/false,
+					FText::Format(LOCTEXT("AttributePackingTypeConflict", "Attribute '{0}' encountered with multiple different types in input data, custom float packing failed."), FText::FromName(AttributeName)));
+
+				OutAttributeIdOffsetStrides.Empty();
+				OutCustomFloatCount = 0;
+
+				return;
 			}
 
-			const uint32 AttributeId = static_cast<uint32>(FoundAttribute->Id);
-			const uint32 StrideFloats = GetAttributeTypeStrideBytes(AttributeType) / sizeof(float);
+			const uint32 StrideFloats = GetAttributeTypeStrideBytes(AttributeDesc.Type) / sizeof(float);
 
-			OutAttributeIdOffsetStrides.Emplace(AttributeId, OffsetFloats, StrideFloats, /*Unused*/0);
+			OutAttributeIdOffsetStrides.Emplace(static_cast<uint32>(AttributeDesc.Index), OffsetFloats, StrideFloats, /*Unused*/0);
 
 			OffsetFloats += StrideFloats;
 		}
@@ -635,6 +643,36 @@ uint32 GetTypeHash(const FPCGKernelAttributeKey& In)
 	return HashCombine(GetTypeHash(In.Type), GetTypeHash(In.Name));
 }
 
+int32 FPCGKernelAttributeTable::GetAttributeId(const FPCGKernelAttributeKey& InAttribute) const
+{
+	return AttributeTable.IndexOfByKey(InAttribute) + PCGComputeConstants::NUM_RESERVED_ATTRS;
+}
+
+int32 FPCGKernelAttributeTable::GetAttributeId(FName InName, EPCGKernelAttributeType InType) const
+{
+	return AttributeTable.IndexOfByKey(FPCGKernelAttributeKey(InName, InType)) + PCGComputeConstants::NUM_RESERVED_ATTRS;
+}
+
+void FPCGKernelAttributeTable::AddAttribute(FName InName, EPCGKernelAttributeType InType)
+{
+	AttributeTable.AddUnique(FPCGKernelAttributeKey(InName, InType));
+}
+
+#if PCG_KERNEL_LOGGING_ENABLED
+void FPCGKernelAttributeTable::DebugLog()
+{
+	const UEnum* PCGKernelAttributeTypeEnum = StaticEnum<EPCGKernelAttributeType>();
+
+	for (const FPCGKernelAttributeKey& Attribute : AttributeTable)
+	{
+		UE_LOG(LogPCG, Display, TEXT("\tName: %s\t\tID: %d\t\tType: %s"),
+			*Attribute.Name.ToString(),
+			GetAttributeId(Attribute),
+			*PCGKernelAttributeTypeEnum->GetDisplayNameTextByValue(static_cast<int64>(Attribute.Type)).ToString());
+	}
+}
+#endif
+
 bool FPCGKernelAttributeDesc::operator==(const FPCGKernelAttributeDesc& Other) const
 {
 	return Index == Other.Index && Type == Other.Type && Name == Other.Name;
@@ -644,18 +682,17 @@ FPCGDataDesc::FPCGDataDesc(EPCGDataType InType, int InElementCount)
 	: Type(InType)
 	, ElementCount(InElementCount)
 {
-	TArray<FString> StringTableDummy;
-	InitializeAttributeDescs(nullptr, {}, StringTableDummy);
+	InitializeAttributeDescs(/*InData=*/nullptr, /*InBinding*/nullptr);
 }
 
-FPCGDataDesc::FPCGDataDesc(const UPCGData* InData, const TMap<FName, FPCGKernelAttributeIDAndType>& InGlobalAttributeLookupTable, const TArray<FString>& InStringTable)
+FPCGDataDesc::FPCGDataDesc(const UPCGData* InData, const UPCGDataBinding* InBinding)
 {
 	check(InData);
 
 	Type = InData->GetDataType();
 	ElementCount = PCGComputeHelpers::GetElementCount(InData);
 
-	InitializeAttributeDescs(InData, InGlobalAttributeLookupTable, InStringTable);
+	InitializeAttributeDescs(InData, InBinding);
 }
 
 uint64 FPCGDataDesc::ComputePackedSize() const
@@ -680,7 +717,15 @@ bool FPCGDataDesc::HasMetadataAttributes() const
 	});
 }
 
-void FPCGDataDesc::InitializeAttributeDescs(const UPCGData* InData, const TMap<FName, FPCGKernelAttributeIDAndType>& InGlobalAttributeLookupTable, const TArray<FString>& InStringTable)
+bool FPCGDataDesc::ContainsAttribute(FName InAttributeName) const
+{
+	return AttributeDescs.ContainsByPredicate([InAttributeName](const FPCGKernelAttributeDesc& AttributeDesc)
+	{
+		return AttributeDesc.Name == InAttributeName;
+	});
+}
+
+void FPCGDataDesc::InitializeAttributeDescs(const UPCGData* InData, const UPCGDataBinding* InBinding)
 {
 	TRACE_CPUPROFILER_EVENT_SCOPE(FPCGDataDesc::InitializeAttributeDescs);
 
@@ -690,15 +735,14 @@ void FPCGDataDesc::InitializeAttributeDescs(const UPCGData* InData, const TMap<F
 	}
 	else { /* TODO: More types! */ }
 
-	if (const UPCGMetadata* Metadata = InData ? InData->ConstMetadata() : nullptr)
+	const UPCGMetadata* Metadata = InData ? InData->ConstMetadata() : nullptr;
+
+	if (InBinding && Metadata)
 	{
 		TArray<FName> AttributeNames;
 		TArray<EPCGMetadataTypes> AttributeTypes;
 		Metadata->GetAttributes(AttributeNames, AttributeTypes);
 
-		TArray<TPair<FPCGKernelAttributeKey, TArray<int32>>> DelayedAttributeKeysAndStringKeys; // Attribute keys that don't exist in the global lookup table must be delayed so we can append them at the end.
-		int NumAttributesFromLUT = 0; // Keep track of how many attributes come from the LUT. This will help give us the starting index for our delayed attributes.
-
 		for (int CustomAttributeIndex = 0; CustomAttributeIndex < AttributeNames.Num(); ++CustomAttributeIndex)
 		{
 			const FName AttributeName = AttributeNames[CustomAttributeIndex];
@@ -731,6 +775,8 @@ void FPCGDataDesc::InitializeAttributeDescs(const UPCGData* InData, const TMap<F
 				const FPCGMetadataAttributeBase* AttributeBase = Metadata->GetConstAttribute(AttributeName);
 				check(AttributeBase);
 
+				const TArray<FString>& StringTable = InBinding->GetStringTable();
+
 				if (Type == EPCGDataType::Point && ensure(InData->IsA<UPCGPointData>()))
 				{
 					const UPCGPointData* PointData = CastChecked<UPCGPointData>(InData);
@@ -743,19 +789,19 @@ void FPCGDataDesc::InitializeAttributeDescs(const UPCGData* InData, const TMap<F
 
 						if (AttributeBase->GetTypeId() == PCG::Private::MetadataTypes<FSoftObjectPath>::Id)
 						{
-							StringTableIndex = InStringTable.IndexOfByKey(static_cast<const FPCGMetadataAttribute<FSoftObjectPath>*>(AttributeBase)->GetValue(ValueKey).ToString());
+							StringTableIndex = StringTable.IndexOfByKey(static_cast<const FPCGMetadataAttribute<FSoftObjectPath>*>(AttributeBase)->GetValue(ValueKey).ToString());
 						}
 						else if (AttributeBase->GetTypeId() == PCG::Private::MetadataTypes<FString>::Id)
 						{
-							StringTableIndex = InStringTable.IndexOfByKey(static_cast<const FPCGMetadataAttribute<FString>*>(AttributeBase)->GetValue(ValueKey));
+							StringTableIndex = StringTable.IndexOfByKey(static_cast<const FPCGMetadataAttribute<FString>*>(AttributeBase)->GetValue(ValueKey));
 						}
 						else if (AttributeBase->GetTypeId() == PCG::Private::MetadataTypes<FSoftClassPath>::Id)
 						{
-							StringTableIndex = InStringTable.IndexOfByKey(static_cast<const FPCGMetadataAttribute<FSoftClassPath>*>(AttributeBase)->GetValue(ValueKey).ToString());
+							StringTableIndex = StringTable.IndexOfByKey(static_cast<const FPCGMetadataAttribute<FSoftClassPath>*>(AttributeBase)->GetValue(ValueKey).ToString());
 						}
 						else if (AttributeBase->GetTypeId() == PCG::Private::MetadataTypes<FName>::Id)
 						{
-							StringTableIndex = InStringTable.IndexOfByKey(static_cast<const FPCGMetadataAttribute<FName>*>(AttributeBase)->GetValue(ValueKey).ToString());
+							StringTableIndex = StringTable.IndexOfByKey(static_cast<const FPCGMetadataAttribute<FName>*>(AttributeBase)->GetValue(ValueKey).ToString());
 						}
 						else
 						{
@@ -779,19 +825,19 @@ void FPCGDataDesc::InitializeAttributeDescs(const UPCGData* InData, const TMap<F
 
 						if (AttributeBase->GetTypeId() == PCG::Private::MetadataTypes<FSoftObjectPath>::Id)
 						{
-							StringTableIndex = InStringTable.IndexOfByKey(static_cast<const FPCGMetadataAttribute<FSoftObjectPath>*>(AttributeBase)->GetValue(MetadataKey).ToString());
+							StringTableIndex = StringTable.IndexOfByKey(static_cast<const FPCGMetadataAttribute<FSoftObjectPath>*>(AttributeBase)->GetValue(MetadataKey).ToString());
 						}
 						else if (AttributeBase->GetTypeId() == PCG::Private::MetadataTypes<FString>::Id)
 						{
-							StringTableIndex = InStringTable.IndexOfByKey(static_cast<const FPCGMetadataAttribute<FString>*>(AttributeBase)->GetValue(MetadataKey));
+							StringTableIndex = StringTable.IndexOfByKey(static_cast<const FPCGMetadataAttribute<FString>*>(AttributeBase)->GetValue(MetadataKey));
 						}
 						else if (AttributeBase->GetTypeId() == PCG::Private::MetadataTypes<FSoftClassPath>::Id)
 						{
-							StringTableIndex = InStringTable.IndexOfByKey(static_cast<const FPCGMetadataAttribute<FSoftClassPath>*>(AttributeBase)->GetValue(MetadataKey).ToString());
+							StringTableIndex = StringTable.IndexOfByKey(static_cast<const FPCGMetadataAttribute<FSoftClassPath>*>(AttributeBase)->GetValue(MetadataKey).ToString());
 						}
 						else if (AttributeBase->GetTypeId() == PCG::Private::MetadataTypes<FName>::Id)
 						{
-							StringTableIndex = InStringTable.IndexOfByKey(static_cast<const FPCGMetadataAttribute<FName>*>(AttributeBase)->GetValue(MetadataKey).ToString());
+							StringTableIndex = StringTable.IndexOfByKey(static_cast<const FPCGMetadataAttribute<FName>*>(AttributeBase)->GetValue(MetadataKey).ToString());
 						}
 						else
 						{
@@ -808,32 +854,20 @@ void FPCGDataDesc::InitializeAttributeDescs(const UPCGData* InData, const TMap<F
 				else { /* TODO: More types! */ }
 			}
 
-			if (const FPCGKernelAttributeIDAndType* AttributeIdAndType = InGlobalAttributeLookupTable.Find(AttributeName))
-			{
-				AttributeDescs.Emplace(AttributeIdAndType->Id, AttributeType, AttributeName, MoveTemp(UniqueStringKeys));
-				++NumAttributesFromLUT;
-			}
-			else
+			const int32 AttributeId = InBinding->GetAttributeId(AttributeName, AttributeType);
+			ensureMsgf(AttributeId != INDEX_NONE, TEXT("Attribute '%s' type %d was missing from attribute table."), *AttributeName.ToString(), (int32)AttributeType);
+
+			if (AttributeId != INDEX_NONE)
 			{
-				TPair<FPCGKernelAttributeKey, TArray<int32>> AttributeAndStringKeys;
-				AttributeAndStringKeys.Get<0>() = FPCGKernelAttributeKey(AttributeName, AttributeType);
-				AttributeAndStringKeys.Get<1>() = MoveTemp(UniqueStringKeys);
-				DelayedAttributeKeysAndStringKeys.Add(MoveTemp(AttributeAndStringKeys));
+				AttributeDescs.Emplace(AttributeId, AttributeType, AttributeName, MoveTemp(UniqueStringKeys));
 			}
 		}
-
-		for (int DelayedAttributeIndex = 0; DelayedAttributeIndex < DelayedAttributeKeysAndStringKeys.Num(); ++DelayedAttributeIndex)
-		{
-			const FPCGKernelAttributeKey& AttributeKey = DelayedAttributeKeysAndStringKeys[DelayedAttributeIndex].Get<0>();
-			AttributeDescs.Emplace(NUM_RESERVED_ATTRS + DelayedAttributeIndex + NumAttributesFromLUT + InGlobalAttributeLookupTable.Num(), AttributeKey.Type, AttributeKey.Name, MoveTemp(DelayedAttributeKeysAndStringKeys[DelayedAttributeIndex].Get<1>()));
-		}
 	}
 }
 
 FPCGDataCollectionDesc FPCGDataCollectionDesc::BuildFromDataCollection(
 	const FPCGDataCollection& InDataCollection,
-	const TMap<FName, FPCGKernelAttributeIDAndType>& InAttributeLookupTable,
-	const TArray<FString>& InStringTable)
+	const UPCGDataBinding* InBinding)
 {
 	FPCGDataCollectionDesc CollectionDesc;
 
@@ -844,7 +878,7 @@ FPCGDataCollectionDesc FPCGDataCollectionDesc::BuildFromDataCollection(
 			continue;
 		}
 
-		CollectionDesc.DataDescs.Emplace(Data.Data, InAttributeLookupTable, InStringTable);
+		CollectionDesc.DataDescs.Emplace(Data.Data, InBinding);
 	}
 
 	return CollectionDesc;
@@ -853,8 +887,7 @@ FPCGDataCollectionDesc FPCGDataCollectionDesc::BuildFromDataCollection(
 FPCGDataCollectionDesc FPCGDataCollectionDesc::BuildFromInputDataCollectionAndInputPinLabel(
 	const FPCGDataCollection& InDataCollection,
 	FName InputPinLabel,
-	const TMap<FName, FPCGKernelAttributeIDAndType>& InAttributeLookupTable,
-	const TArray<FString>& InStringTable)
+	const UPCGDataBinding* InBinding)
 {
 	FPCGDataCollectionDesc CollectionDesc;
 	TArray<FPCGTaggedData> DataForPin = InDataCollection.GetInputsByPin(InputPinLabel);
@@ -866,7 +899,7 @@ FPCGDataCollectionDesc FPCGDataCollectionDesc::BuildFromInputDataCollectionAndIn
 			continue;
 		}
 
-		CollectionDesc.DataDescs.Emplace(Data.Data, InAttributeLookupTable, InStringTable);
+		CollectionDesc.DataDescs.Emplace(Data.Data, InBinding);
 	}
 
 	return CollectionDesc;
@@ -1461,3 +1494,80 @@ void FPCGDataCollectionDesc::Combine(const FPCGDataCollectionDesc& Other)
 {
 	DataDescs.Append(Other.DataDescs);
 }
+
+bool FPCGDataCollectionDesc::GetAttributeDesc(FName InAttributeName, FPCGKernelAttributeDesc& OutAttributeDesc, bool& bOutConflictingTypesFound) const
+{
+	// Will be set to the type of the first attribute across all the data which has a matching name, and then is used to detect if subsequent
+	// attributes are found with conflicting type.
+	EPCGKernelAttributeType FoundType = EPCGKernelAttributeType::Invalid;
+
+	for (const FPCGDataDesc& DataDesc : DataDescs)
+	{
+		for (const FPCGKernelAttributeDesc& AttributeDesc : DataDesc.AttributeDescs)
+		{
+			if (AttributeDesc.Name == InAttributeName && ensure(AttributeDesc.Type != EPCGKernelAttributeType::Invalid))
+			{
+				if (FoundType == EPCGKernelAttributeType::Invalid)
+				{
+					// Take the first matching attribute.
+					OutAttributeDesc = AttributeDesc;
+					FoundType = AttributeDesc.Type;
+				}
+				else if (FoundType != AttributeDesc.Type)
+				{
+					// Signal conflict found.
+					bOutConflictingTypesFound = true;
+
+					// Can stop iterating once we find a conflict. Signal that a matching attribute was found (despite the conflict).
+					return true;
+				}
+			}
+		}
+	}
+
+	// Signal no conflict found.
+	bOutConflictingTypesFound = false;
+
+	// Signal attribute found or not.
+	return FoundType != EPCGKernelAttributeType::Invalid;
+}
+
+bool FPCGDataCollectionDesc::ContainsAttributeOnAnyData(FName InAttributeName) const
+{
+	return DataDescs.ContainsByPredicate([InAttributeName](const FPCGDataDesc& DataDesc)
+	{
+		return DataDesc.ContainsAttribute(InAttributeName);
+	});
+}
+
+void FPCGDataCollectionDesc::AddAttributeToAllData(FPCGKernelAttributeKey InAttribute, const UPCGDataBinding* InBinding, TArray<int32>* InOptionalUniqueStringKeys)
+{
+	const int32 AttributeId = InBinding->GetAttributeId(InAttribute);
+	if (AttributeId == INDEX_NONE)
+	{
+		return;
+	}
+
+	for (FPCGDataDesc& DataDesc : DataDescs)
+	{
+		bool bAttributeFound = false;
+
+		for (FPCGKernelAttributeDesc& AttributeDesc : DataDesc.AttributeDescs)
+		{
+			if (AttributeDesc.Name == InAttribute.Name)
+			{
+				ensureMsgf(!bAttributeFound, TEXT("Attribute '%s' encountered multiple times in a single data."), *AttributeDesc.Name.ToString());
+				bAttributeFound = true;
+
+				AttributeDesc.Type = InAttribute.Type;
+			}
+		}
+
+		if (!bAttributeFound)
+		{
+			DataDesc.AttributeDescs.Emplace(AttributeId, InAttribute.Type, InAttribute.Name, InOptionalUniqueStringKeys ? *InOptionalUniqueStringKeys : TArray<int32>());
+		}
+	}
+}
+
+#undef LOCTEXT_NAMESPACE
diff --git a/Engine/Plugins/PCG/Source/PCG/Private/Compute/PCGDataForGPU.h b/Engine/Plugins/PCG/Source/PCG/Private/Compute/PCGDataForGPU.h
--- a/Engine/Plugins/PCG/Source/PCG/Private/Compute/PCGDataForGPU.h
+++ b/Engine/Plugins/PCG/Source/PCG/Private/Compute/PCGDataForGPU.h
@@ -4,6 +4,7 @@
 
 #include "PCGCommon.h"
 #include "PCGData.h"
+#include "Compute/PCGComputeCommon.h"
 
 #include "PCGDataForGPU.generated.h"
 
@@ -24,8 +25,7 @@ enum class EPCGUnpackDataCollectionResult
 UENUM()
 enum class EPCGKernelAttributeType : uint8
 {
-	None = 0 UMETA(Hidden),
-	Bool,
+	Bool = 1,
 	Int,
 	Float,
 	Float2,
@@ -40,6 +40,7 @@ enum class EPCGKernelAttributeType : uint8
 	Invalid = std::numeric_limits<uint8>::max() UMETA(Hidden),
 };
 
+/** Attribute name and type which uniquely identify an attribute in a compute graph. */
 USTRUCT()
 struct FPCGKernelAttributeKey
 {
@@ -59,30 +60,40 @@ struct FPCGKernelAttributeKey
 	EPCGKernelAttributeType Type = EPCGKernelAttributeType::Float;
 
 	bool operator==(const FPCGKernelAttributeKey& Other) const;
+
 	friend uint32 GetTypeHash(const FPCGKernelAttributeKey& In);
 };
 
+/** Table of attributes used in compute graph with helpers to get unique attribute ID used to read/write attribute in data collection buffers. */
 USTRUCT()
-struct FPCGKernelAttributeIDAndType
+struct FPCGKernelAttributeTable
 {
 	GENERATED_BODY()
 
-	FPCGKernelAttributeIDAndType() = default;
+	FPCGKernelAttributeTable() = default;
 
-	explicit FPCGKernelAttributeIDAndType(int32 InId, EPCGKernelAttributeType InType)
-		: Id(InId)
-		, Type(InType)
-	{}
+	int32 GetAttributeId(const FPCGKernelAttributeKey& InAttribute) const;
 
-	UPROPERTY()
-	int32 Id = INDEX_NONE;
+	int32 GetAttributeId(FName InName, EPCGKernelAttributeType InType) const;
+
+	void AddAttribute(FName InName, EPCGKernelAttributeType InType);
+
+#if PCG_KERNEL_LOGGING_ENABLED
+	void DebugLog();
+#endif
 
+protected:
 	UPROPERTY()
-	EPCGKernelAttributeType Type = EPCGKernelAttributeType::Invalid;
+	TArray<FPCGKernelAttributeKey> AttributeTable;
+
+	friend class FPCGGraphCompilerGPU;
 };
 
+/** Data description for a metadata attribute. Stores identifying name and type as well as the unique attribute ID. */
 struct FPCGKernelAttributeDesc
 {
+	FPCGKernelAttributeDesc() = default;
+
 	explicit FPCGKernelAttributeDesc(int32 InIndex, EPCGKernelAttributeType InType, FName InName)
 		: Index(InIndex)
 		, Type(InType)
@@ -102,6 +113,7 @@ struct FPCGKernelAttributeDesc
 	int32 Index = INDEX_NONE;
 
 	EPCGKernelAttributeType Type = EPCGKernelAttributeType::Float;
+
 	FName Name = NAME_None;
 
 	/* All possible string keys arriving on this attribute (string keys are indices into the string table in the data binding). */
@@ -110,34 +122,37 @@ struct FPCGKernelAttributeDesc
 	bool operator==(const FPCGKernelAttributeDesc& Other) const;
 };
 
+/** Data description for a single data object (UPCGData). */
 struct FPCGDataDesc
 {
 	FPCGDataDesc(EPCGDataType InType, int32 InElementCount);
-	FPCGDataDesc(const UPCGData* InData, const TMap<FName, FPCGKernelAttributeIDAndType>& InGlobalAttributeLookupTable, const TArray<FString>& InStringTable);
+	FPCGDataDesc(const UPCGData* InData, const UPCGDataBinding* InBinding);
 
 	uint64 ComputePackedSize() const;
+
 	bool HasMetadataAttributes() const;
 
+	bool ContainsAttribute(FName InAttributeName) const;
+
 	EPCGDataType Type = EPCGDataType::Point;
 	TArray<FPCGKernelAttributeDesc> AttributeDescs;
 	int32 ElementCount = 0;
 
 private:
-	void InitializeAttributeDescs(const UPCGData* InData, const TMap<FName, FPCGKernelAttributeIDAndType>& InGlobalAttributeLookupTable, const TArray<FString>& InStringTable);
+	void InitializeAttributeDescs(const UPCGData* InData, const UPCGDataBinding* InBinding);
 };
 
+/** Data description for a data collection (FPCGDataCollection). */
 struct FPCGDataCollectionDesc
 {
 	static FPCGDataCollectionDesc BuildFromDataCollection(
 		const FPCGDataCollection& InDataCollection,
-		const TMap<FName, FPCGKernelAttributeIDAndType>& InAttributeLookupTable,
-		const TArray<FString>& InStringTable);
+		const UPCGDataBinding* InBinding);
 
 	static FPCGDataCollectionDesc BuildFromInputDataCollectionAndInputPinLabel(
 		const FPCGDataCollection& InDataCollection,
 		FName InputPinLabel,
-		const TMap<FName, FPCGKernelAttributeIDAndType>& InAttributeLookupTable,
-		const TArray<FString>& InStringTable);
+		const UPCGDataBinding* InBinding);
 
 	/** Computes the size (in bytes) of the header portion of the packed data collection buffer. */
 	uint32 ComputePackedHeaderSizeBytes() const;
@@ -159,24 +174,17 @@ struct FPCGDataCollectionDesc
 	/** Aggregate another data description. */
 	void Combine(const FPCGDataCollectionDesc& Other);
 
-	TArray<FPCGDataDesc> DataDescs;
-};
-
-USTRUCT()
-struct FPCGDataForGPU
-{
-	GENERATED_BODY()
+	/** Get description of first attribute with matching name in input data. Returns true if such an attribute found and also signals whether multiple matching attributes
+	* with conflicting names are present. */
+	bool GetAttributeDesc(FName InAttributeName, FPCGKernelAttributeDesc& OutAttributeDesc, bool &bOutConflictingTypesFound) const;
 
-	// All the node input pins that have edges that cross from CPU to GPU.
-	UPROPERTY()
-	TSet<TSoftObjectPtr<const UPCGPin>> InputPins;
+	bool ContainsAttributeOnAnyData(FName InAttributeName) const;
 
-	UPROPERTY()
-	TMap<TSoftObjectPtr<const UPCGPin>, FName> InputPinLabelAliases;
+	/** Makes attribute present on all data. If data has existing attribute with same name then the given type will be applied. */
+	void AddAttributeToAllData(FPCGKernelAttributeKey InAttribute, const UPCGDataBinding* InBinding, TArray<int32>* InOptionalUniqueStringKeys = nullptr);
 
-	// Since the compute graph is collapsed to a single element, all data crossing from CPU to GPU is in a single collection.
-	UPROPERTY()
-	FPCGDataCollection InputDataCollection;
+	/** Description of each data in this data collection. */
+	TArray<FPCGDataDesc> DataDescs;
 };
 
 namespace PCGDataForGPUHelpers
@@ -186,6 +194,8 @@ namespace PCGDataForGPUHelpers
 
 	/** Compute how attributes can be packed to custom floats. */
 	void ComputeCustomFloatPacking(
+		FPCGContext* InContext,
+		const UPCGSettings* InSettings,
 		TArray<FName>& InAttributeNames,
 		const UPCGDataBinding* InBinding,
 		const FPCGDataCollectionDesc& InDataCollectionDescription,
diff --git a/Engine/Plugins/PCG/Source/PCG/Private/Elements/PCGCopyPoints.cpp b/Engine/Plugins/PCG/Source/PCG/Private/Elements/PCGCopyPoints.cpp
--- a/Engine/Plugins/PCG/Source/PCG/Private/Elements/PCGCopyPoints.cpp
+++ b/Engine/Plugins/PCG/Source/PCG/Private/Elements/PCGCopyPoints.cpp
@@ -27,7 +27,7 @@ FText UPCGCopyPointsSettings::GetNodeTooltipText() const
 }
 #endif
 
-FString UPCGCopyPointsSettings::GetCookedKernelSource(const TMap<FName, FPCGKernelAttributeIDAndType>& GlobalAttributeLookupTable) const
+FString UPCGCopyPointsSettings::GetCookedKernelSource(const FPCGKernelAttributeTable& InAttributeStaticTable) const
 {
 	FString TemplateFile;
 	ensure(LoadShaderSourceFile(TEXT("/Plugin/PCG/Private/Elements/PCGCopyPoints.usf"), EShaderPlatform::SP_PCD3D_SM5, &TemplateFile, nullptr));
diff --git a/Engine/Plugins/PCG/Source/PCG/Private/Elements/PCGStaticMeshSpawner.cpp b/Engine/Plugins/PCG/Source/PCG/Private/Elements/PCGStaticMeshSpawner.cpp
--- a/Engine/Plugins/PCG/Source/PCG/Private/Elements/PCGStaticMeshSpawner.cpp
+++ b/Engine/Plugins/PCG/Source/PCG/Private/Elements/PCGStaticMeshSpawner.cpp
@@ -93,7 +93,7 @@ bool UPCGStaticMeshSpawnerSettings::IsKernelValid(FPCGContext* InContext, bool b
 	return true;
 }
 
-FString UPCGStaticMeshSpawnerSettings::GetCookedKernelSource(const TMap<FName, FPCGKernelAttributeIDAndType>& GlobalAttributeLookupTable) const
+FString UPCGStaticMeshSpawnerSettings::GetCookedKernelSource(const FPCGKernelAttributeTable& InAttributeStaticTable) const
 {
 	FString TemplateFile;
 	ensure(LoadShaderSourceFile(TemplateFilePath, EShaderPlatform::SP_PCD3D_SM5, &TemplateFile, nullptr));
@@ -106,6 +106,12 @@ bool UPCGStaticMeshSpawnerSettings::ComputeOutputPinDataDesc(const UPCGPin* Outp
 	check(OutputPin);
 	check(Binding);
 
+	// Code assumes single output pin.
+	if (!ensure(OutputPin->Properties.Label == PCGPinConstants::DefaultOutputLabel))
+	{
+		return false;
+	}
+
 	// First inspect data coming from CPU. Will return description if this node is a CPU node or in a different compute graph.
 	if (Binding->ComputeCPUOutputPinDataDesc(OutputPin, OutDesc))
 	{
@@ -113,70 +119,64 @@ bool UPCGStaticMeshSpawnerSettings::ComputeOutputPinDataDesc(const UPCGPin* Outp
 	}
 
 	const FPCGDataCollectionDesc InputPinDesc = PCGDataForGPUHelpers::ComputeInputPinDataDesc(this, PCGPinConstants::DefaultInputLabel, Binding);
+
+	// First forward input data to output.
 	FPCGDataCollectionDesc OutputPinDesc = InputPinDesc;
 
-	if (ensure(OutputPin->Properties.Label == PCGPinConstants::DefaultOutputLabel))
+	// Add output attribute (selected mesh).
 	{
-		const TMap<FName, FPCGKernelAttributeIDAndType>& GlobalAttributeLookupTable = Binding->GetAttributeLookupTable();
+		// Create unique value keys for the output string.
+		TArray<int32> UniqueStringKeys;
 
-		if (const FPCGKernelAttributeIDAndType* FoundAttributeIDAndType = GlobalAttributeLookupTable.Find(OutAttributeName))
+		if (UPCGMeshSelectorWeighted* SelectorWeighted = Cast<UPCGMeshSelectorWeighted>(MeshSelectorParameters))
 		{
-			// Create attribute description.
-			FPCGKernelAttributeDesc AttributeDesc(FoundAttributeIDAndType->Id, FoundAttributeIDAndType->Type, OutAttributeName);
-
-			// Create unique value keys for the output string.
-			if (UPCGMeshSelectorWeighted* SelectorWeighted = Cast<UPCGMeshSelectorWeighted>(MeshSelectorParameters))
+			// Weighted selection - add explicit strings from settings.
+			for (const FPCGMeshSelectorWeightedEntry& Entry : SelectorWeighted->MeshEntries)
 			{
-				// Weighted selection - add explicit strings from settings.
-				for (const FPCGMeshSelectorWeightedEntry& Entry : SelectorWeighted->MeshEntries)
+				const FString Value = Entry.Descriptor.StaticMesh.ToString();
+				const int32 StringIndex = Binding->GetStringTable().IndexOfByKey(Value);
+				if (ensureAlways(StringIndex != INDEX_NONE))
 				{
-					const FString Value = Entry.Descriptor.StaticMesh.ToString();
-					const int32 StringIndex = Binding->GetStringTable().IndexOfByKey(Value);
-					if (ensureAlways(StringIndex != INDEX_NONE))
-					{
-						AttributeDesc.UniqueStringKeys.Add(StringIndex);
-					}
+					UniqueStringKeys.AddUnique(StringIndex);
 				}
 			}
-			else if (UPCGMeshSelectorByAttribute* SelectorByAttribute = Cast<UPCGMeshSelectorByAttribute>(MeshSelectorParameters))
+		}
+		else if (UPCGMeshSelectorByAttribute* SelectorByAttribute = Cast<UPCGMeshSelectorByAttribute>(MeshSelectorParameters))
+		{
+			// By-attribute selection - pass on strings from input attribute.
+			if (SelectorByAttribute->AttributeName != NAME_None)
 			{
-				// By-attribute selection - pass on strings from input attribute.
-				if (SelectorByAttribute->AttributeName != NAME_None)
+				for (const FPCGDataDesc& InputDataDesc : InputPinDesc.DataDescs)
 				{
-					for (const FPCGDataDesc& InputDataDesc : InputPinDesc.DataDescs)
+					for (const FPCGKernelAttributeDesc& InputAttrDesc : InputDataDesc.AttributeDescs)
 					{
-						for (const FPCGKernelAttributeDesc& InputAttrDesc : InputDataDesc.AttributeDescs)
+						if (InputAttrDesc.Name == SelectorByAttribute->AttributeName)
 						{
-							if (InputAttrDesc.Name == SelectorByAttribute->AttributeName)
+							for (int32 InputStringKey : InputAttrDesc.UniqueStringKeys)
 							{
-								AttributeDesc.UniqueStringKeys.Append(InputAttrDesc.UniqueStringKeys);
+								UniqueStringKeys.AddUnique(InputStringKey);
 							}
 						}
 					}
 				}
 			}
-			else if (MeshSelectorParameters)
-			{
-				UE_LOG(LogPCG, Warning, TEXT("Mesh selector not supported by GPU Static Mesh Spawner: %s"), *MeshSelectorParameters->GetName());
-			}
-
-			for (FPCGDataDesc& DataDesc : OutputPinDesc.DataDescs)
-			{
-				if (const FPCGKernelAttributeIDAndType* IDAndType = GlobalAttributeLookupTable.Find(OutAttributeName))
-				{
-					DataDesc.AttributeDescs.AddUnique(AttributeDesc);
-				}
-			}
 		}
+		else if (MeshSelectorParameters)
+		{
+			UE_LOG(LogPCG, Warning, TEXT("Mesh selector not supported by GPU Static Mesh Spawner: %s"), *MeshSelectorParameters->GetName());
+		}
+
+		OutputPinDesc.AddAttributeToAllData(FPCGKernelAttributeKey(OutAttributeName, EPCGKernelAttributeType::StringKey), Binding, &UniqueStringKeys);
 	}
 
 	OutDesc = OutputPinDesc;
+
 	return true;
 }
 
 void UPCGStaticMeshSpawnerSettings::GetKernelAttributeKeys(TArray<FPCGKernelAttributeKey>& OutKeys) const
 {
-	// Write out attribute.
+	// Register the attribute this node creates.
 	OutKeys.AddUnique(FPCGKernelAttributeKey(OutAttributeName, EPCGKernelAttributeType::StringKey));
 }
 
diff --git a/Engine/Plugins/PCG/Source/PCG/Private/Graph/PCGGraphCompilerGPU.cpp b/Engine/Plugins/PCG/Source/PCG/Private/Graph/PCGGraphCompilerGPU.cpp
--- a/Engine/Plugins/PCG/Source/PCG/Private/Graph/PCGGraphCompilerGPU.cpp
+++ b/Engine/Plugins/PCG/Source/PCG/Private/Graph/PCGGraphCompilerGPU.cpp
@@ -687,29 +687,7 @@ void FPCGGraphCompilerGPU::CreateDataInterfaces(
 
 			for (const FPCGKernelAttributeKey& Key : AttributeKeys)
 			{
-				FPCGKernelAttributeIDAndType* FoundEntry = InOutComputeGraph->GlobalAttributeLookupTable.Find(Key.Name);
-				if (FoundEntry)
-				{
-					// If we already registered this attribute, populate the type if not populated already, or otherwise ensure type matches.
-					if (Key.Type != EPCGKernelAttributeType::None)
-					{
-						if (FoundEntry->Type == EPCGKernelAttributeType::None)
-						{
-							FoundEntry->Type = Key.Type;
-						}
-						else if (Key.Type != FoundEntry->Type)
-						{
-							// TODO: This is more a limitation of the current system rather than an issue with the user's data. Log for bug reports etc but remove this
-							// later when we properly flow attribute table and string table down graph branches.
-							UE_LOG(LogPCG, Warning, TEXT("Attribute '%s' multiply defined with different types (existing: %d, new: %d)."), *Key.Name.ToString(), FoundEntry->Type, Key.Type);
-						}
-					}
-				}
-				else
-				{
-					const FPCGKernelAttributeIDAndType IDAndType(InOutComputeGraph->GlobalAttributeLookupTable.Num() + PCGComputeConstants::NUM_RESERVED_ATTRS, Key.Type);
-					InOutComputeGraph->GlobalAttributeLookupTable.Add(Key.Name, IDAndType);
-				}
+				InOutComputeGraph->AttributeStaticTable.AttributeTable.AddUnique(Key);
 			}
 
 			// Record strings.
@@ -1122,25 +1100,20 @@ void FPCGGraphCompilerGPU::CompileComputeGraph(
 			KernelSource->EntryPoint = TEXT("Main");
 			KernelSource->GroupSize = FIntVector(PCGComputeConstants::THREAD_GROUP_SIZE, 1, 1);
 
-			KernelSource->SetSource(Settings->GetCookedKernelSource(InOutComputeGraph->GlobalAttributeLookupTable));
+			KernelSource->SetSource(Settings->GetCookedKernelSource(InOutComputeGraph->AttributeStaticTable));
 
 			if (Settings->bDumpCookedHLSL)
 			{
 				UE_LOG(LogPCG, Log, TEXT("Cooked HLSL:\n%s\n"), *KernelSource->GetSource());
 			}
 
+#if PCG_KERNEL_LOGGING_ENABLED
 			if (PCGGraphCompilerGPU::CVarEnableDebugging.GetValueOnAnyThread())
 			{
-				UE_LOG(LogPCG, Warning, TEXT("ATTRIBUTE LOOK-UP TABLE [%s]"), *Settings->GetDefaultNodeTitle().ToString());
-
-				for (const TPair<FName, FPCGKernelAttributeIDAndType>& Pair : InOutComputeGraph->GlobalAttributeLookupTable)
-				{
-					UE_LOG(LogPCG, Warning, TEXT("%s: %s (%s)"),
-						*FString::FromInt(Pair.Value.Id),
-						*Pair.Key.ToString(),
-						*UEnum::GetValueAsString(Pair.Value.Type));
-				}
+				UE_LOG(LogPCG, Display, TEXT("\n### STATIC METADATA ATTRIBUTE TABLE ###"));
+				InOutComputeGraph->AttributeStaticTable.DebugLog();
 			}
+#endif
 
 			// Add functions for external inputs/outputs which must be fulfilled by DIs
 			for (FInterfaceBinding& Binding : KernelWithBindings.InputDataBindings)
diff --git a/Engine/Plugins/PCG/Source/PCG/Public/Elements/PCGCopyPoints.h b/Engine/Plugins/PCG/Source/PCG/Public/Elements/PCGCopyPoints.h
--- a/Engine/Plugins/PCG/Source/PCG/Public/Elements/PCGCopyPoints.h
+++ b/Engine/Plugins/PCG/Source/PCG/Public/Elements/PCGCopyPoints.h
@@ -54,7 +54,7 @@ public:
 	virtual void CreateAdditionalInputDataInterfaces(FPCGGPUCompilationContext& InOutContext, UObject* InObjectOuter, TArray<TObjectPtr<UComputeDataInterface>>& OutDataInterfaces) const override;
 #endif
 
-	virtual FString GetCookedKernelSource(const TMap<FName, FPCGKernelAttributeIDAndType>& GlobalAttributeLookupTable) const override;
+	virtual FString GetCookedKernelSource(const FPCGKernelAttributeTable& InAttributeStaticTable) const override;
 	virtual int ComputeKernelThreadCount(const UPCGDataBinding* Binding) const override;
 	virtual bool ComputeOutputPinDataDesc(const UPCGPin* OutputPin, const UPCGDataBinding* Binding, FPCGDataCollectionDesc& OutDesc) const override;
 
diff --git a/Engine/Plugins/PCG/Source/PCG/Public/Elements/PCGStaticMeshSpawner.h b/Engine/Plugins/PCG/Source/PCG/Public/Elements/PCGStaticMeshSpawner.h
--- a/Engine/Plugins/PCG/Source/PCG/Public/Elements/PCGStaticMeshSpawner.h
+++ b/Engine/Plugins/PCG/Source/PCG/Public/Elements/PCGStaticMeshSpawner.h
@@ -30,7 +30,7 @@ public:
 
 	// ~Begin UPCGSettings interface
 	virtual bool IsKernelValid(FPCGContext* InContext = nullptr, bool bQuiet = true) const;
-	virtual FString GetCookedKernelSource(const TMap<FName, FPCGKernelAttributeIDAndType>& GlobalAttributeLookupTable) const override;
+	virtual FString GetCookedKernelSource(const FPCGKernelAttributeTable& InAttributeStaticTable) const override;
 	virtual void GetKernelAttributeKeys(TArray<FPCGKernelAttributeKey>& OutKeys) const override;
 	virtual void AddStaticCreatedStrings(TArray<FString>& InOutStringTable) const override;
 	virtual int ComputeKernelThreadCount(const UPCGDataBinding* Binding) const override;
diff --git a/Engine/Plugins/PCG/Source/PCG/Public/PCGSettings.h b/Engine/Plugins/PCG/Source/PCG/Public/PCGSettings.h
--- a/Engine/Plugins/PCG/Source/PCG/Public/PCGSettings.h
+++ b/Engine/Plugins/PCG/Source/PCG/Public/PCGSettings.h
@@ -23,8 +23,8 @@ class UPCGNode;
 class UPCGSettings;
 struct FPCGDataCollectionDesc;
 struct FPCGGPUCompilationContext;
-struct FPCGKernelAttributeIDAndType;
 struct FPCGKernelAttributeKey;
+struct FPCGKernelAttributeTable;
 struct FPCGPinProperties;
 struct FPropertyChangedEvent;
 
@@ -546,9 +546,11 @@ public:
 	/** [EXPERIMENTAL] Produces the node specific portion of kernel shader source text, including the main entry point.
 	* Note: This API function is subject to change without deprecation.
 	*/
-	virtual FString GetCookedKernelSource(const TMap<FName, FPCGKernelAttributeIDAndType>& GlobalAttributeLookupTable) const { return TEXT(""); }
+	virtual FString GetCookedKernelSource(const FPCGKernelAttributeTable& InAttributeStaticTable) const { return TEXT(""); }
 
-	/** [EXPERIMENTAL] Get a list of the attributes read or written by this node.
+	/** [EXPERIMENTAL] Get all the attributes read or written by this node for which we know the name and type statically. In cases that only the
+	* attribute name used by the node is known (example: user specifies SM Spawner instance data attributes by name only, not type), the attribute
+	* is omitted from this list and must be resolved at execution time in the data provider.
 	* Note: This API function is subject to change without deprecation.
 	*/
 	virtual void GetKernelAttributeKeys(TArray<FPCGKernelAttributeKey>& OutKeys) const {}
