[PCG] Add Get Render Target and support for sampling render targets in compute kernels.

Only supports pre-rendered/render once RTs. We need to deep dive on scheduling in order to get fresh RT contents each frame.

#rnx

diff --git a/Engine/Plugins/PCG/Source/PCG/Private/Compute/DataInterfaces/PCGTextureDataInterface.cpp b/Engine/Plugins/PCG/Source/PCG/Private/Compute/DataInterfaces/PCGTextureDataInterface.cpp
--- a/Engine/Plugins/PCG/Source/PCG/Private/Compute/DataInterfaces/PCGTextureDataInterface.cpp
+++ b/Engine/Plugins/PCG/Source/PCG/Private/Compute/DataInterfaces/PCGTextureDataInterface.cpp
@@ -20,6 +20,8 @@
 
 void UPCGTextureDataInterface::GetSupportedInputs(TArray<FShaderFunctionDefinition>& OutFunctions) const
 {
+	// TODO: Support texture2d array, texture3d. UNiagaraDataInterfaceRenderTarget2D and related are useful references.
+
 	OutFunctions.AddDefaulted_GetRef()
 		.SetName(TEXT("GetTexCoords"))
 		.AddReturnType(EShaderFundamentalType::Float, 2)
@@ -80,28 +82,16 @@ UComputeDataProvider* UPCGTextureDataInterface::CreateDataProvider(TObjectPtr<UO
 	check(!DownstreamInputPinLabelAliases.IsEmpty());
 	const TArray<FPCGTaggedData> TaggedDatas = Binding->InputDataCollection.GetInputsByPin(DownstreamInputPinLabelAliases[0]);
 
-	if (!TaggedDatas.IsEmpty())
+	if (!TaggedDatas.IsEmpty() && ensure(TaggedDatas[0].Data))
 	{
-		ensure(TaggedDatas.Num() == 1); // There should only be one texture data
-		ensure(TaggedDatas[0].Data->IsA<UPCGTextureData>());
-
-		if (const UPCGTextureData* TextureData = Cast<UPCGTextureData>(TaggedDatas[0].Data))
-		{
-			const UTexture* Texture = TextureData->Texture.IsValid() ? TextureData->Texture.Get() : nullptr;
-
-			// TODO: Support texture2d array, texture3d
-			// TODO: Bindless resources to handle multiple textures
-
-			if (Texture)
-			{
-				const FTextureResource* TextureResource = Texture->GetResource();
-
-				if (TextureResource && TextureResource->TextureRHI)
-				{
-					DataProvider->Texture = TextureResource->TextureRHI;
-				}
-			}
-		}
+		// There should only be one texture data
+		// TODO: Bindless resources to handle multiple textures
+		ensure(TaggedDatas.Num() == 1); 
+
+		const UPCGBaseTextureData* BaseTextureData = Cast<UPCGBaseTextureData>(TaggedDatas[0].Data);
+		ensure(BaseTextureData);
+
+		DataProvider->TextureRHI = BaseTextureData ? BaseTextureData->GetTextureRHI() : nullptr;
 	}
 
 	return DataProvider;
@@ -109,8 +99,7 @@ UComputeDataProvider* UPCGTextureDataInterface::CreateDataProvider(TObjectPtr<UO
 
 FComputeDataProviderRenderProxy* UPCGTextureDataProvider::GetRenderProxy()
 {
-	FPCGTextureDataProviderProxy* Proxy = new FPCGTextureDataProviderProxy(Texture);
-	return Proxy;
+	return new FPCGTextureDataProviderProxy(TextureRHI);
 }
 
 bool FPCGTextureDataProviderProxy::IsValid(FValidationData const& InValidationData) const
@@ -129,7 +118,7 @@ void FPCGTextureDataProviderProxy::GatherDispatchData(FDispatchData const& InDis
 	for (int32 InvocationIndex = 0; InvocationIndex < ParameterArray.Num(); ++InvocationIndex)
 	{
 		FParameters& Parameters = ParameterArray[InvocationIndex];
-		Parameters.Texture = Texture ? (FRHITexture*)Texture : (FRHITexture*)GBlackTexture->TextureRHI;
+		Parameters.Texture = TextureRHI ? (FRHITexture*)TextureRHI : (FRHITexture*)GBlackTexture->TextureRHI;
 		Parameters.Sampler = TStaticSamplerState<SF_Point, AM_Clamp, AM_Clamp, AM_Clamp>::GetRHI();
 	}
 }
diff --git a/Engine/Plugins/PCG/Source/PCG/Private/Compute/DataInterfaces/PCGTextureDataInterface.h b/Engine/Plugins/PCG/Source/PCG/Private/Compute/DataInterfaces/PCGTextureDataInterface.h
--- a/Engine/Plugins/PCG/Source/PCG/Private/Compute/DataInterfaces/PCGTextureDataInterface.h
+++ b/Engine/Plugins/PCG/Source/PCG/Private/Compute/DataInterfaces/PCGTextureDataInterface.h
@@ -43,14 +43,14 @@ public:
 	FComputeDataProviderRenderProxy* GetRenderProxy() override;
 	//~ End UComputeDataProvider Interface
 
-	FTextureRHIRef Texture;
+	FTextureRHIRef TextureRHI = nullptr;
 };
 
 class FPCGTextureDataProviderProxy : public FComputeDataProviderRenderProxy
 {
 public:
-	FPCGTextureDataProviderProxy(FTextureRHIRef InTexture)
-		: Texture(InTexture)
+	FPCGTextureDataProviderProxy(FTextureRHIRef InTextureRHI)
+		: TextureRHI(InTextureRHI)
 	{}
 
 	//~ Begin FComputeDataProviderRenderProxy Interface
@@ -61,5 +61,5 @@ public:
 protected:
 	using FParameters = FPCGTextureDataInterfaceParameters;
 
-	FTextureRHIRef Texture = nullptr;
+	FTextureRHIRef TextureRHI = nullptr;
 };
diff --git a/Engine/Plugins/PCG/Source/PCG/Private/Compute/PCGComputeCommon.h b/Engine/Plugins/PCG/Source/PCG/Private/Compute/PCGComputeCommon.h
--- a/Engine/Plugins/PCG/Source/PCG/Private/Compute/PCGComputeCommon.h
+++ b/Engine/Plugins/PCG/Source/PCG/Private/Compute/PCGComputeCommon.h
@@ -47,7 +47,7 @@ namespace PCGComputeConstants
 	constexpr int POINT_STEEPNESS_ATTRIBUTE_ID = 8;
 
 	/** PCG data types supported in GPU node inputs. */
-	constexpr EPCGDataType AllowedInputTypes = EPCGDataType::Point | EPCGDataType::Param | EPCGDataType::Landscape | EPCGDataType::Texture;
+	constexpr EPCGDataType AllowedInputTypes = EPCGDataType::Point | EPCGDataType::Param | EPCGDataType::Landscape | EPCGDataType::BaseTexture;
 
 	/** PCG data types supported in GPU node outputs. */
 	constexpr EPCGDataType AllowedOutputTypes = EPCGDataType::Point | EPCGDataType::Param;
diff --git a/Engine/Plugins/PCG/Source/PCG/Private/Data/PCGRenderTargetData.cpp b/Engine/Plugins/PCG/Source/PCG/Private/Data/PCGRenderTargetData.cpp
--- a/Engine/Plugins/PCG/Source/PCG/Private/Data/PCGRenderTargetData.cpp
+++ b/Engine/Plugins/PCG/Source/PCG/Private/Data/PCGRenderTargetData.cpp
@@ -47,6 +47,14 @@ void UPCGRenderTargetData::AddToCrc(FArchiveCrc32& Ar, bool bFullDataCrc) const
 	AddUIDToCrc(Ar);
 }
 
+FTextureRHIRef UPCGRenderTargetData::GetTextureRHI() const
+{
+	// TODO: This makes no attempt to actually acquire the resource after it has been written or, nor does it ensure
+	// resource transitions/barriers. Only works if RT is already populated and does not work if RT is rendered every frame.
+	FTextureResource* Resource = RenderTarget ? RenderTarget->GetResource() : nullptr;
+	return Resource ? Resource->GetTextureRHI() : nullptr;
+}
+
 UPCGSpatialData* UPCGRenderTargetData::CopyInternal(FPCGContext* Context) const
 {
 	UPCGRenderTargetData* NewRenderTargetData = FPCGContext::NewObject_AnyThread<UPCGRenderTargetData>(Context);
diff --git a/Engine/Plugins/PCG/Source/PCG/Private/Data/PCGTextureData.cpp b/Engine/Plugins/PCG/Source/PCG/Private/Data/PCGTextureData.cpp
--- a/Engine/Plugins/PCG/Source/PCG/Private/Data/PCGTextureData.cpp
+++ b/Engine/Plugins/PCG/Source/PCG/Private/Data/PCGTextureData.cpp
@@ -527,6 +527,12 @@ void UPCGTextureData::AddToCrc(FArchiveCrc32& Ar, bool bFullDataCrc) const
 	AddUIDToCrc(Ar);
 }
 
+FTextureRHIRef UPCGTextureData::GetTextureRHI() const
+{
+	FTextureResource* Resource = Texture.IsValid() ? Texture->GetResource() : nullptr;
+	return Resource ? Resource->GetTextureRHI() : nullptr;
+}
+
 UPCGSpatialData* UPCGTextureData::CopyInternal(FPCGContext* Context) const
 {
 	UPCGTextureData* NewTextureData = FPCGContext::NewObject_AnyThread<UPCGTextureData>(Context);
diff --git a/Engine/Plugins/PCG/Source/PCG/Private/Elements/PCGTextureSampler.cpp b/Engine/Plugins/PCG/Source/PCG/Private/Elements/PCGTextureSampler.cpp
--- a/Engine/Plugins/PCG/Source/PCG/Private/Elements/PCGTextureSampler.cpp
+++ b/Engine/Plugins/PCG/Source/PCG/Private/Elements/PCGTextureSampler.cpp
@@ -6,6 +6,7 @@
 #include "PCGContext.h"
 #include "PCGCrc.h"
 #include "PCGSubsystem.h"
+#include "Data/PCGRenderTargetData.h"
 #include "Helpers/PCGBlueprintHelpers.h"
 #include "Helpers/PCGDynamicTrackingHelpers.h"
 #include "Helpers/PCGHelpers.h"
@@ -14,6 +15,7 @@
 #include "Engine/Texture.h"
 #include "Engine/Texture2D.h"
 #include "Engine/Texture2DArray.h"
+#include "Engine/TextureRenderTarget2D.h"
 #include "GameFramework/Actor.h"
 
 #include UE_INLINE_GENERATED_CPP_BY_NAME(PCGTextureSampler)
@@ -78,7 +80,7 @@ FText UPCGTextureSamplerSettings::GetNodeTooltipText() const
 TArray<FPCGPinProperties> UPCGTextureSamplerSettings::OutputPinProperties() const
 {
 	TArray<FPCGPinProperties> Properties;
-	Properties.Emplace(PCGPinConstants::DefaultOutputLabel, EPCGDataType::Texture);
+	Properties.Emplace(PCGPinConstants::DefaultOutputLabel, EPCGDataType::BaseTexture);
 
 	return Properties;
 }
@@ -192,91 +194,111 @@ bool FPCGTextureSamplerElement::ExecuteInternal(FPCGContext* InContext) const
 			return true;
 		}
 	}
-	else if (!Texture->IsA<UTexture2D>())
+	else if (!Texture->IsA<UTexture2D>() && !Texture->IsA<UTextureRenderTarget2D>())
 	{
-		PCGE_LOG(Error, GraphAndLog, FText::Format(LOCTEXT("InvalidTextureType", "Texture at path '{0}' is not a valid type. Must be UTexture2D or UTexture2DArray."), FText::FromString(Settings->Texture.ToString())));
+		PCGE_LOG(Error, GraphAndLog, FText::Format(LOCTEXT("InvalidTextureType", "Texture at path '{0}' is not a valid type. Must be one of the following types: UTexture2D, UTexture2DArray, UTextureRenderTarget2D."), FText::FromString(Settings->Texture.ToString())));
 		return true;
 	}
 
-	const FTransform& Transform = Settings->Transform;
-#if WITH_EDITOR
-	const bool bForceEditorOnlyCPUSampling = Settings->bForceEditorOnlyCPUSampling;
-#else
-	const bool bForceEditorOnlyCPUSampling = false;
-#endif
-
-	// Texture data can take some frames to prepare, so we poll it once per frame until it is done.
-	// TODO - review other similar cases and consider adding helpers/abstractions to support this.
-	UPCGTextureData* TextureData = nullptr;
-	if (Context->TextureData)
+	if (!Context->bTransformInitialized)
 	{
-		TextureData = Context->TextureData.Get();
-	}
-	else
-	{
-		TextureData = FPCGContext::NewObject_AnyThread<UPCGTextureData>(Context);
-		Context->TextureData = TextureData;
-
-		AActor* OriginalActor = UPCGBlueprintHelpers::GetOriginalComponent(*Context)->GetOwner();
-
 		if (Settings->bUseAbsoluteTransform)
 		{
-			Context->Transform = Transform;
+			Context->Transform = Settings->Transform;
 		}
 		else
 		{
+			AActor* OriginalActor = UPCGBlueprintHelpers::GetOriginalComponent(*Context)->GetOwner();
+
 			FTransform OriginalActorTransform = OriginalActor->GetTransform();
-			Context->Transform = Transform * OriginalActorTransform;
+			Context->Transform = Settings->Transform * OriginalActorTransform;
 
 			FBox OriginalActorLocalBounds = PCGHelpers::GetActorLocalBounds(OriginalActor);
 			Context->Transform.SetScale3D(Context->Transform.GetScale3D() * 0.5 * (OriginalActorLocalBounds.Max - OriginalActorLocalBounds.Min));
 		}
+
+		Context->bTransformInitialized = true;
 	}
 
-	if (!ensure(TextureData))
+	// The new texture data to add. Will be assigned the correct texture data based on the referenced texture.
+	UPCGBaseTextureData* BaseTextureData = nullptr;
+
+	if (UTextureRenderTarget2D* RenderTarget = Cast<UTextureRenderTarget2D>(Texture))
 	{
-		PCGE_LOG_C(Error, LogOnly, Context, LOCTEXT("TextureDataInitFailed", "Failed to initialize texture data."));
-		return true;
+		UPCGRenderTargetData* RenderTargetData = FPCGContext::NewObject_AnyThread<UPCGRenderTargetData>(Context);
+		RenderTargetData->Initialize(RenderTarget, Context->Transform);
+		BaseTextureData = RenderTargetData;
 	}
-
-	if (!TextureData->IsSuccessfullyInitialized())
+	else
 	{
-		if (!TextureData->Initialize(Texture, TextureArrayIndex, Context->Transform, bForceEditorOnlyCPUSampling))
+		// Texture data can take some frames to prepare, so we poll it once per frame until it is done.
+		// TODO - review other similar cases and consider adding helpers/abstractions to support this.
+		UPCGTextureData* TextureData = nullptr;
+		if (Context->TextureData)
 		{
-			// Initialization not complete. Could be waiting on async texture processing or for GPU readback. Sleep until next frame.
-			Context->bIsPaused = true;
-			Context->SourceComponent->GetSubsystem()->RegisterBeginTickAction([ContextHandle = Context->GetOrCreateHandle()]()
-			{
-				if (FPCGTextureSamplerContext* ContextPtr = FPCGContext::GetContextFromHandle<FPCGTextureSamplerContext>(ContextHandle))
-				{
-					ContextPtr->bIsPaused = false;
-				}
-			});
+			TextureData = Context->TextureData.Get();
+		}
+		else
+		{
+			TextureData = FPCGContext::NewObject_AnyThread<UPCGTextureData>(Context);
+			Context->TextureData = TextureData;
+		}
 
-			return false;
+		if (!ensure(TextureData))
+		{
+			PCGE_LOG_C(Error, LogOnly, Context, LOCTEXT("TextureDataInitFailed", "Failed to initialize texture data."));
+			return true;
 		}
 
 		if (!TextureData->IsSuccessfullyInitialized())
 		{
-			PCGE_LOG(Warning, LogOnly, LOCTEXT("TextureInitFailed", "Data could not be retrieved for this texture, initialization failed."));
-			Context->OutputData.TaggedData.Empty();
-			return true;
+#if WITH_EDITOR
+			const bool bForceEditorOnlyCPUSampling = Settings->bForceEditorOnlyCPUSampling;
+#else
+			const bool bForceEditorOnlyCPUSampling = false;
+#endif
+
+			if (!TextureData->Initialize(Texture, TextureArrayIndex, Context->Transform, bForceEditorOnlyCPUSampling))
+			{
+				// Initialization not complete. Could be waiting on async texture processing or for GPU readback. Sleep until next frame.
+				Context->bIsPaused = true;
+				Context->SourceComponent->GetSubsystem()->RegisterBeginTickAction([ContextHandle = Context->GetOrCreateHandle()]()
+					{
+						if (FPCGTextureSamplerContext* ContextPtr = FPCGContext::GetContextFromHandle<FPCGTextureSamplerContext>(ContextHandle))
+						{
+							ContextPtr->bIsPaused = false;
+						}
+					});
+
+				return false;
+			}
+
+			if (!TextureData->IsSuccessfullyInitialized())
+			{
+				PCGE_LOG(Warning, LogOnly, LOCTEXT("TextureInitFailed", "Data could not be retrieved for this texture, initialization failed."));
+				Context->OutputData.TaggedData.Empty();
+				return true;
+			}
 		}
+
+		BaseTextureData = TextureData;
 	}
 
+	check(BaseTextureData);
+
 	// Commit to adding texture data.
-	Context->OutputData.TaggedData.Emplace_GetRef().Data = TextureData;
-
-	TextureData->bUseDensitySourceChannel = Settings->bUseDensitySourceChannel;
-	TextureData->ColorChannel = Settings->ColorChannel;
-	TextureData->Filter = Settings->Filter;
-	TextureData->TexelSize = Settings->TexelSize;
-	TextureData->bUseAdvancedTiling = Settings->bUseAdvancedTiling;
-	TextureData->Tiling = Settings->Tiling;
-	TextureData->CenterOffset = Settings->CenterOffset;
-	TextureData->Rotation = Settings->Rotation;
-	TextureData->bUseTileBounds = Settings->bUseTileBounds;
-	TextureData->TileBounds = FBox2D(Settings->TileBoundsMin, Settings->TileBoundsMax);
+	Context->OutputData.TaggedData.Emplace_GetRef().Data = BaseTextureData;
+
+	BaseTextureData->bUseDensitySourceChannel = Settings->bUseDensitySourceChannel;
+	BaseTextureData->ColorChannel = Settings->ColorChannel;
+	BaseTextureData->Filter = Settings->Filter;
+	BaseTextureData->TexelSize = Settings->TexelSize;
+	BaseTextureData->bUseAdvancedTiling = Settings->bUseAdvancedTiling;
+	BaseTextureData->Tiling = Settings->Tiling;
+	BaseTextureData->CenterOffset = Settings->CenterOffset;
+	BaseTextureData->Rotation = Settings->Rotation;
+	BaseTextureData->bUseTileBounds = Settings->bUseTileBounds;
+	BaseTextureData->TileBounds = FBox2D(Settings->TileBoundsMin, Settings->TileBoundsMax);
 
 #if WITH_EDITOR
 	// If we have an override, register for dynamic tracking.
diff --git a/Engine/Plugins/PCG/Source/PCG/Private/Graph/PCGGraphCompilerGPU.cpp b/Engine/Plugins/PCG/Source/PCG/Private/Graph/PCGGraphCompilerGPU.cpp
--- a/Engine/Plugins/PCG/Source/PCG/Private/Graph/PCGGraphCompilerGPU.cpp
+++ b/Engine/Plugins/PCG/Source/PCG/Private/Graph/PCGGraphCompilerGPU.cpp
@@ -646,7 +646,9 @@ void FPCGGraphCompilerGPU::CreateDataInterfaces(
 			DataInterface = DataInterfacePCGData;
 			break;
 		}
+		case EPCGDataType::BaseTexture:
 		case EPCGDataType::Texture:
+		case EPCGDataType::RenderTarget:
 		{
 			DataInterface = InOutContext.NewObject_AnyThread<UPCGTextureDataInterface>(InOutComputeGraph);
 			break;
diff --git a/Engine/Plugins/PCG/Source/PCG/Public/Data/PCGRenderTargetData.h b/Engine/Plugins/PCG/Source/PCG/Public/Data/PCGRenderTargetData.h
--- a/Engine/Plugins/PCG/Source/PCG/Public/Data/PCGRenderTargetData.h
+++ b/Engine/Plugins/PCG/Source/PCG/Public/Data/PCGRenderTargetData.h
@@ -4,6 +4,8 @@
 
 #include "Data/PCGTextureData.h"
 
+#include "RHI.h"
+
 #include "PCGRenderTargetData.generated.h"
 
 class UTextureRenderTarget2D;
@@ -21,6 +23,9 @@ public:
 	virtual void AddToCrc(FArchiveCrc32& Ar, bool bFullDataCrc) const override;
 	// ~End UPCGData interface
 
+	//~Begin UPCGBaseTextureData interface
+	virtual FTextureRHIRef GetTextureRHI() const override;
+	//~End UPCGBaseTextureData interface
 
 	//~Begin UPCGSpatialData interface
 protected:
diff --git a/Engine/Plugins/PCG/Source/PCG/Public/Data/PCGTextureData.h b/Engine/Plugins/PCG/Source/PCG/Public/Data/PCGTextureData.h
--- a/Engine/Plugins/PCG/Source/PCG/Public/Data/PCGTextureData.h
+++ b/Engine/Plugins/PCG/Source/PCG/Public/Data/PCGTextureData.h
@@ -4,6 +4,8 @@
 
 #include "Data/PCGSurfaceData.h"
 
+#include "RHI.h"
+
 #include "PCGTextureData.generated.h"
 
 class UPCGSpatialData;
@@ -79,6 +81,8 @@ public:
 
 	virtual bool IsValid() const;
 
+	virtual FTextureRHIRef GetTextureRHI() const PURE_VIRTUAL(UPCGBaseTextureData::GetTextureResource, return nullptr;);
+
 protected:
 	PRAGMA_DISABLE_DEPRECATION_WARNINGS
 	UFUNCTION(BlueprintGetter, meta = (BlueprintInternalUseOnly = "true"))
@@ -165,11 +169,15 @@ public:
 	/** Data is successfully initialized and is ready to use. */
 	PCG_API bool IsSuccessfullyInitialized() const { return bSuccessfullyInitialized; }
 
-	// ~Begin UPCGData interface
+	//~Begin UPCGData interface
 	virtual EPCGDataType GetDataType() const override { return EPCGDataType::Texture; }
 	virtual void AddToCrc(FArchiveCrc32& Ar, bool bFullDataCrc) const override;
-	// ~End UPCGData interface
+	//~End UPCGData interface
 
+	//~Begin UPCGBaseTextureData interface
+	virtual FTextureRHIRef GetTextureRHI() const override;
+	//~End UPCGBaseTextureData interface
+	
 	//~Begin UPCGSpatialData interface
 protected:
 	virtual UPCGSpatialData* CopyInternal(FPCGContext* Context) const override;
diff --git a/Engine/Plugins/PCG/Source/PCG/Public/Elements/PCGTextureSampler.h b/Engine/Plugins/PCG/Source/PCG/Public/Elements/PCGTextureSampler.h
--- a/Engine/Plugins/PCG/Source/PCG/Public/Elements/PCGTextureSampler.h
+++ b/Engine/Plugins/PCG/Source/PCG/Public/Elements/PCGTextureSampler.h
@@ -146,6 +146,8 @@ public:
 
 	TObjectPtr<UPCGTextureData> TextureData;
 
+	bool bTransformInitialized = false;
+
 	FTransform Transform;
 
 protected:
