[PCG] Fix crash with World Raycast on empty point data on Origins input

#rnx

diff --git a/Engine/Plugins/PCG/Source/PCG/Private/Elements/PCGWorldRaycast.cpp b/Engine/Plugins/PCG/Source/PCG/Private/Elements/PCGWorldRaycast.cpp
--- a/Engine/Plugins/PCG/Source/PCG/Private/Elements/PCGWorldRaycast.cpp
+++ b/Engine/Plugins/PCG/Source/PCG/Private/Elements/PCGWorldRaycast.cpp
@@ -170,7 +170,7 @@ bool FPCGWorldRaycastElement::PrepareDataInternal(FPCGContext* InContext) const
 		}
 
 		// --- Find the ray origin ---
-		if (!PCGAttributeAccessorHelpers::ExtractAllValues(OriginsData, Settings->OriginInputAttribute, OutState.CachedRayOrigins, Context))
+		if (!PCGAttributeAccessorHelpers::ExtractAllValues(OriginsData, Settings->OriginInputAttribute, OutState.CachedRayOrigins, Context) || OutState.CachedRayOrigins.IsEmpty())
 		{
 			return EPCGTimeSliceInitResult::NoOperation;
 		}
