[PCG] Add Blueprint helper method to add some packages to the allow list that are not actor packages
- Also add a waiting step in case PCGComponent generation triggered more generation

diff --git a/Engine/Plugins/PCG/Source/PCGEditor/Private/WorldPartitionBuilder/PCGWorldPartitionBuilder.cpp b/Engine/Plugins/PCG/Source/PCGEditor/Private/WorldPartitionBuilder/PCGWorldPartitionBuilder.cpp
--- a/Engine/Plugins/PCG/Source/PCGEditor/Private/WorldPartitionBuilder/PCGWorldPartitionBuilder.cpp
+++ b/Engine/Plugins/PCG/Source/PCGEditor/Private/WorldPartitionBuilder/PCGWorldPartitionBuilder.cpp
@@ -27,6 +27,7 @@
 #include "Misc/OutputDevice.h"
 #include "Modules/ModuleManager.h"
 #include "UObject/Linker.h"
+#include "UObject/UObjectIterator.h"
 #include "WorldPartition/IWorldPartitionEditorModule.h"
 #include "WorldPartition/WorldPartition.h"
 #include "WorldPartition/WorldPartitionHelpers.h"
@@ -77,7 +78,7 @@ struct FPCGWorldPartitionBuilder
 	static FSoftObjectPath GetActorPathFromComponentPtr(TSoftObjectPtr<UPCGComponent> InComponentPtr);
 
 	/** Get the list of Dirty packages */
-	static TArray<UPackage*> GetDirtyPackages();
+	static TArray<UPackage*> GetDirtyPackages(const TSet<FString>& AllowedAssetPackages);
 
 	/** Set if Partition Actor creation should be disabled for world */
 	static void SetDisablePartitionActorCreationForWorld(UWorld* InWorld, bool bDisable);
@@ -88,6 +89,9 @@ struct FPCGWorldPartitionBuilder
 	/** Get Grid Size information for a given graph and cache it for future calls */
 	static void GetGenerationGridSizes(const UPCGGraph* InGraph, const APCGWorldActor* InWorldActor, PCGHiGenGrid::FSizeArray& OutGridSizes, bool& bOutHasUnbounded);
 
+	/** Returns when GraphExecutor is done running */
+	static void WaitOnGraphExecutor(UWorld* InWorld);
+
 private:
 	static void CollectSubGridComponentsToGenerateInternal(UPCGSubsystem* InPCGSubsystem, int32 InGridSizeIndex, const PCGHiGenGrid::FSizeArray& InGridSizes, UPCGComponent* InOriginalComponent, APCGPartitionActor* InPartitionActor, TArray<UPCGComponent*>& OutAdditionalComponentsToGenerate, TMap<UPCGComponent*, UPCGComponent*>& OutDependencies);
 
@@ -466,6 +470,20 @@ UPCGWorldPartitionBuilder::UPCGWorldPartitionBuilder(const FObjectInitializer& O
 	}
 }
 
+void UPCGWorldPartitionBuilderHelper::AddAssetToSubmitAllowList(const TSoftObjectPtr<UObject>& InPath)
+{
+	for (TObjectIterator<UPCGWorldPartitionBuilder> It; It; ++It)
+	{
+		It->AddAssetToSubmitAllowList(InPath);
+	}
+}
+
+void UPCGWorldPartitionBuilder::AddAssetToSubmitAllowList(const TSoftObjectPtr<UObject>& InPath)
+{
+	UE_LOG(LogPCGWorldPartitionBuilder, Display, TEXT("Added Asset to Submit Allow List: '%s'"), *InPath.ToString());
+	AllowedAssetPackages.Add(InPath.GetLongPackageName());
+}
+
 bool UPCGWorldPartitionBuilder::PreWorldInitialization(UWorld* World, FPackageSourceControlHelper& PackageHelper)
 {
 	if (!Super::PreWorldInitialization(World, PackageHelper))
@@ -491,6 +509,9 @@ bool UPCGWorldPartitionBuilder::PreRun(UWorld* World, FPackageSourceControlHelpe
 		return false;
 	}
 
+	// Reset Allow list
+	AllowedAssetPackages.Empty();
+
 	// Set base class member
 	IterativeCellSize = Args.IterativeCellSize;
 
@@ -603,7 +624,10 @@ bool UPCGWorldPartitionBuilder::RunEntireWorld(UWorld* World, FPackageSourceCont
 
 	if (bGeneratedAnyComponent)
 	{
-		DirtyPackages = FPCGWorldPartitionBuilder::GetDirtyPackages();
+		// Wait for other executions that might have spawned from generating the components
+		FPCGWorldPartitionBuilder::WaitOnGraphExecutor(World);
+		
+		DirtyPackages = FPCGWorldPartitionBuilder::GetDirtyPackages(AllowedAssetPackages);
 	}
 
 	if (!bGeneratedAnyComponent)
@@ -856,6 +880,9 @@ bool UPCGWorldPartitionBuilder::RunPerCell(UWorld* InWorld, const FCellInfo& InC
 
 	if (bCellHasChanges)
 	{
+		// Wait for other executions that might have spawned from generating the components
+		FPCGWorldPartitionBuilder::WaitOnGraphExecutor(InWorld);
+		
 		return SavePerCellPackages(InWorld, PackageHelper, CellDeletedPackageNames);
 	}
 
@@ -946,7 +973,7 @@ bool UPCGWorldPartitionBuilder::PostRun(UWorld* World, FPackageSourceControlHelp
 
 bool UPCGWorldPartitionBuilder::SavePerCellPackages(UWorld* InWorld, FPackageSourceControlHelper& PackageHelper, TSet<FString> CellDeletedPackages)
 {
-	TArray<UPackage*> DirtyPackages = FPCGWorldPartitionBuilder::GetDirtyPackages();
+	TArray<UPackage*> DirtyPackages = FPCGWorldPartitionBuilder::GetDirtyPackages(AllowedAssetPackages);
 	TArray<UPackage*> PackagesToDelete;
 	UWorldPartition* WorldPartition = InWorld->GetWorldPartition();
 
@@ -1270,19 +1297,52 @@ void FPCGWorldPartitionBuilder::GetGenerationGridSizes(const UPCGGraph* InGraph,
 	NewGridSizeInfo.GridSizes = OutGridSizes;
 }
 
+void FPCGWorldPartitionBuilder::WaitOnGraphExecutor(UWorld* InWorld)
+{
+	UPCGSubsystem* PCGSubsystem = UWorld::GetSubsystem<UPCGSubsystem>(InWorld);
+	check(PCGSubsystem);
+
+	while (PCGSubsystem->IsAnyGraphCurrentlyExecuting())
+	{
+		FWorldPartitionHelpers::FakeEngineTick(InWorld);
+	}
+		
+	FAssetCompilingManager::Get().FinishAllCompilation();
+	FAssetCompilingManager::Get().ProcessAsyncTasks();
+}
+
 void FPCGWorldPartitionBuilder::ResetPackageDirtyFlags()
 {
-	TArray<UPackage*> DirtyPackages = FPCGWorldPartitionBuilder::GetDirtyPackages();
+	TArray<UPackage*> DirtyPackages;
+
+	FEditorFileUtils::GetDirtyWorldPackages(DirtyPackages);
+	FEditorFileUtils::GetDirtyContentPackages(DirtyPackages);
+
 	for (UPackage* DirtyPackage : DirtyPackages)
 	{
 		DirtyPackage->SetDirtyFlag(false);
 	}
 }
 
-TArray<UPackage*> FPCGWorldPartitionBuilder::GetDirtyPackages()
+TArray<UPackage*> FPCGWorldPartitionBuilder::GetDirtyPackages(const TSet<FString>& AllowedAssetPackages)
 {
 	TArray<UPackage*> DirtyPackages;
 	FEditorFileUtils::GetDirtyWorldPackages(DirtyPackages);
+
+	if (!AllowedAssetPackages.IsEmpty())
+	{
+		TArray<UPackage*> DirtyAssetPackages;
+		FEditorFileUtils::GetDirtyContentPackages(DirtyAssetPackages);
+
+		for (UPackage* DirtyAssetPackage : DirtyAssetPackages)
+		{
+			if (AllowedAssetPackages.Contains(DirtyAssetPackage->GetName()))
+			{
+				DirtyPackages.Add(DirtyAssetPackage);
+			}
+		}
+	}
+
 	DirtyPackages.Remove(nullptr);
 	return DirtyPackages;
 }
diff --git a/Engine/Plugins/PCG/Source/PCGEditor/Private/WorldPartitionBuilder/PCGWorldPartitionBuilder.h b/Engine/Plugins/PCG/Source/PCGEditor/Private/WorldPartitionBuilder/PCGWorldPartitionBuilder.h
--- a/Engine/Plugins/PCG/Source/PCGEditor/Private/WorldPartitionBuilder/PCGWorldPartitionBuilder.h
+++ b/Engine/Plugins/PCG/Source/PCGEditor/Private/WorldPartitionBuilder/PCGWorldPartitionBuilder.h
@@ -6,6 +6,7 @@
 #include "WorldPartition/WorldPartitionHandle.h"
 
 #include "EditorBuildUtils.h"
+#include "Kismet/BlueprintFunctionLibrary.h"
 
 #include "PCGWorldPartitionBuilder.generated.h"
 
@@ -94,6 +95,16 @@ struct FPCGWorldPartitionBuilderArgs
 	static bool InitializeFrom(const FPCGWorldPartitionCommandlineArgs& CommandlineArgs, FPCGWorldPartitionBuilderArgs& OutArgs);
 };
 
+UCLASS()
+class UPCGWorldPartitionBuilderHelper : public UBlueprintFunctionLibrary
+{
+	GENERATED_BODY()
+
+	/** Notify the running PCG Builder that this asset can be saved/submitted as part of the builder execution if it is dirty */
+	UFUNCTION(BlueprintCallable, Category = "PCG|Builder")
+	static void AddAssetToSubmitAllowList(const TSoftObjectPtr<UObject>& InPath);
+};
+
 /**
 * Builder that triggers generation on PCG components.
 * 
@@ -127,6 +138,8 @@ public:
 
 	// In the future we may load a world in stages, but for now load everything to make sure we have the complete level.
 	virtual ELoadingMode GetLoadingMode() const override { return Args.bIterativeCellLoading ? ELoadingMode::IterativeCells2D : ELoadingMode::EntireWorld; }
+		
+	void AddAssetToSubmitAllowList(const TSoftObjectPtr<UObject>& SoftObjectPath);
 
 protected:
 	virtual bool PreWorldInitialization(UWorld* World, FPackageSourceControlHelper& PackageHelper) override;
@@ -162,6 +175,9 @@ private:
 		
 	FPCGWorldPartitionCommandlineArgs CommandlineArgs;
 	FPCGWorldPartitionBuilderArgs Args;
+		
+	/** Non Actor assets that are allowed to be submitted by the builder */
+	TSet<FString> AllowedAssetPackages;
 
 	/** Flag to register if any error occurred during generation. */
 	bool bErrorOccurredWhileGenerating = false;
