Very early experimental support for anim bank mesh components in PCG. Expect frequent data breakage until improvements and refactors settle.

diff --git a/Engine/Plugins/PCG/Shaders/Private/Elements/PCGAnimBankMeshSpawner.usf b/Engine/Plugins/PCG/Shaders/Private/Elements/PCGAnimBankMeshSpawner.usf
new file mode 100644
--- /dev/null
+++ b/Engine/Plugins/PCG/Shaders/Private/Elements/PCGAnimBankMeshSpawner.usf
@@ -0,0 +1,188 @@
+// Copyright Epic Games, Inc. All Rights Reserved.
+
+// AnimBank mesh spawner kernel. Selects a primitive either using a primitive cumulative distribution function (CDF) or by an attribute, and packs attributes in custom floats.
+// Passes through points to output.
+
+// When spawning by attribute, we don't know in advance how many instances of each mesh will be spawned, so the worst case NumPoints * NumPrims is allocated in the GPU Scene.
+
+#include "/Engine/Private/ComputeShaderUtils.ush"
+#include "/Engine/Private/Quaternion.ush"
+#include "/Plugin/PCG/Private/PCGShaderUtils.ush"
+
+// TODO: should be code-gen'd from PCGComputeCommon.h
+#define MAX_NUM_ATTRS 128
+#define NUM_RESERVED_ATTRS 32
+
+void CopyAttributesToOutput(uint In_DataIndex, uint In_ElementIndex, uint Out_DataIndex, uint Out_ElementIndex);
+
+[numthreads(64, 1, 1)]
+void Main(uint3 GroupId : SV_GroupID, uint GroupIndex : SV_GroupIndex)
+{
+	// Mark the kernel as having executed. Must run before we early out via thread index, because the kernel is still 'executed' even if the number of
+	// threads to iterate on is zero. Even if GetNumThreads() returns 0, the kernel will still have been dispatched on a single thread to set this flag.
+	if (all(GroupId == 0) && GroupIndex == 0)
+	{
+		Out_SetAsExecutedInternal();
+	}
+
+	const uint ThreadIndex = GetUnWrappedDispatchThreadId(GroupId, GroupIndex, 64);
+
+	uint In_DataIndex, In_ElementIndex;
+	if (!In_GetThreadData(ThreadIndex, In_DataIndex, In_ElementIndex))
+	{
+		return;
+	}
+
+	uint Out_DataIndex, Out_ElementIndex;
+	if (!Out_GetThreadData(ThreadIndex, Out_DataIndex, Out_ElementIndex))
+	{
+		// Should never happen.
+		return;
+	}
+
+	if (In_IsPointRemoved(In_DataIndex, In_ElementIndex))
+	{
+		// Propagate the 'Removed' status to the output point if input point is removed. Otherwise, this point will not be culled.
+		Out_RemovePoint(Out_DataIndex, Out_ElementIndex);
+
+		return;
+	}
+
+	uint PrimitiveIndex = (uint)-1;
+
+	const int NumPrimitives = InstanceData_GetNumPrimitives();
+
+	if (NumPrimitives > 0)
+	{
+		const uint SelectorAttributeId = ABSpawner_GetSelectorAttributeId();
+		
+		if (SelectorAttributeId != (uint)-1)
+		{
+			// Select using string key
+			const int MeshStringKey = In_GetStringKey(In_DataIndex, In_ElementIndex, SelectorAttributeId);
+			PrimitiveIndex = ABSpawner_GetPrimitiveIndexFromStringKey(MeshStringKey);
+		}
+		else
+		{
+			// Randomly select a primitive by inverting the cumulative distribution function.
+			
+			// Exactly matches UPCGBlueprintHelpers::GetRandomStreamFromPoint - three seeds passed in to consecutive arguments of ComputeSeed
+			// (argument count & order matters).
+			const uint SeedPoint = In_GetSeed(In_DataIndex, In_ElementIndex);
+			const uint SeedSettings = GetSettingsSeed();
+			const uint SeedComponent = GetComponentSeed();
+			uint Seed = ComputeSeed(SeedPoint, SeedSettings, SeedComponent);
+
+			const float RandomDraw = FRand(Seed);
+
+			for (PrimitiveIndex = 0; PrimitiveIndex < uint(NumPrimitives) - 1; ++PrimitiveIndex)
+			{
+				if (RandomDraw < ABSpawner_GetPrimitiveSelectionCDF(PrimitiveIndex))
+				{
+					break;
+				}
+			}
+		}
+	}
+
+	// Get instance index to write to. Will return a negative value if we have saturated instances for this primitive, in which case we will not push the data.
+	const int InstanceIndexInt = (PrimitiveIndex != (uint)-1) ? InstanceData_GetIndexToWriteTo(PrimitiveIndex) : -1;
+
+	// Write instance transform, only if not culled by primitive instance index
+	if (InstanceIndexInt >= 0)
+	{
+		const uint InstanceIndex = InstanceIndexInt;
+		const float4x4 Transform = In_GetPointTransform(In_DataIndex, In_ElementIndex);
+		InstanceData_WriteInstanceLocalToWorld(InstanceIndex, transpose(Transform));
+
+		// Write instance custom float data.
+		const uint NumAttributes = ABSpawner_GetNumAttributes();
+		for (uint AttributeIndex = 0; AttributeIndex < NumAttributes; ++AttributeIndex)
+		{
+			const uint4 AttributeIdOffsetStride = ABSpawner_GetAttributeIdOffsetStride(AttributeIndex);
+			const uint AttributeId = AttributeIdOffsetStride[0];
+			const uint AttributeOffset = AttributeIdOffsetStride[1];
+			const uint AttributeStride = AttributeIdOffsetStride[2];
+
+			for (uint FloatIndex = 0; FloatIndex < AttributeStride; ++FloatIndex)
+			{
+				const uint AttributeValueAddress = In_GetElementAddressInternal(In_DataIndex, In_ElementIndex, AttributeId);
+				// Assumes floating point data. TODO add type information and deal with different types in different ways?
+				InstanceData_WriteCustomFloatRaw(InstanceIndex, AttributeOffset + FloatIndex, In_LoadBufferInternal(AttributeValueAddress + 4 * FloatIndex));
+			}
+		}
+	}
+
+	// Write output point data (pass through).
+	Out_SetPosition(Out_DataIndex, Out_ElementIndex, In_GetPosition(In_DataIndex, In_ElementIndex));
+	Out_SetRotation(Out_DataIndex, Out_ElementIndex, In_GetRotation(In_DataIndex, In_ElementIndex));
+	Out_SetScale(Out_DataIndex, Out_ElementIndex, In_GetScale(In_DataIndex, In_ElementIndex));
+	Out_SetColor(Out_DataIndex, Out_ElementIndex, In_GetColor(In_DataIndex, In_ElementIndex));
+	Out_SetSeed(Out_DataIndex, Out_ElementIndex, In_GetSeed(In_DataIndex, In_ElementIndex));
+	Out_SetDensity(Out_DataIndex, Out_ElementIndex, In_GetDensity(In_DataIndex, In_ElementIndex));
+	Out_SetSteepness(Out_DataIndex, Out_ElementIndex, In_GetSteepness(In_DataIndex, In_ElementIndex));
+
+	// Either apply AB bounds to points or pass bounds through.
+	if (PrimitiveIndex != (uint)-1 && ABSpawner_ShouldApplyBounds())
+	{
+		Out_SetBoundsMin(Out_DataIndex, Out_ElementIndex, ABSpawner_GetPrimitiveMeshBoundsMin(PrimitiveIndex));
+		Out_SetBoundsMax(Out_DataIndex, Out_ElementIndex, ABSpawner_GetPrimitiveMeshBoundsMax(PrimitiveIndex));
+	}
+	else
+	{
+		Out_SetBoundsMin(Out_DataIndex, Out_ElementIndex, In_GetBoundsMin(In_DataIndex, In_ElementIndex));
+		Out_SetBoundsMax(Out_DataIndex, Out_ElementIndex, In_GetBoundsMax(In_DataIndex, In_ElementIndex));
+	}
+
+	CopyAttributesToOutput(In_DataIndex, In_ElementIndex, Out_DataIndex, Out_ElementIndex);
+
+	if (ABSpawner_GetSelectedMeshAttributeId() != (uint)-1)
+	{
+		const int OutputMeshStringKey = (PrimitiveIndex != (uint)-1) ? ABSpawner_GetPrimitiveStringKey(PrimitiveIndex) : -1;
+		Out_SetStringKey(Out_DataIndex, Out_ElementIndex, ABSpawner_GetSelectedMeshAttributeId(), OutputMeshStringKey);
+	}
+}
+
+// TODO pass in attribute IDs to copy.
+void CopyAttributesToOutput(uint In_DataIndex, uint In_ElementIndex, uint Out_DataIndex, uint Out_ElementIndex)
+{
+	const uint NumAttributes = Out_GetDataNumAttributesInternal(Out_DataIndex);
+	uint NumAttributesProcessed = 0;
+	
+	for (int AttributeId = NUM_RESERVED_ATTRS; AttributeId < MAX_NUM_ATTRS; ++AttributeId)
+	{
+		const uint Out_Stride = Out_GetAttributeStrideInternal(Out_DataIndex, AttributeId);
+		
+		if (Out_Stride == 0)
+		{
+			// No output attribute to write to.
+			continue;
+		}
+		
+		const uint In_Stride = In_GetAttributeStrideInternal(In_DataIndex, AttributeId);
+		const uint Out_ElementAddress = Out_GetElementAddressInternal(Out_DataIndex, Out_ElementIndex, AttributeId);
+		
+		if (In_Stride == Out_Stride)
+		{
+			const uint In_ElementAddress = In_GetElementAddressInternal(In_DataIndex, In_ElementIndex, AttributeId);
+			
+			for (uint Offset = 0; Offset < Out_Stride; Offset += 4)
+			{
+				Out_StoreBufferInternal(Out_ElementAddress + Offset, In_LoadBufferInternal(In_ElementAddress + Offset));
+			}
+		}
+		else
+		{
+			for (uint Offset = 0; Offset < Out_Stride; Offset += 4)
+			{
+				// Initialize output data (only header part of buffer will be initialized).
+				Out_StoreBufferInternal(Out_ElementAddress + Offset, 0u);
+			}
+		}
+		
+		if (++NumAttributesProcessed >= NumAttributes)
+		{
+			break; // We can early-out when we've looked at all the possible attributes
+		}
+	}
+}
\ No newline at end of file
diff --git a/Engine/Plugins/PCG/Source/PCG/Private/Compute/DataInterfaces/Elements/PCGAnimBankMeshSpawnerDataInterface.cpp b/Engine/Plugins/PCG/Source/PCG/Private/Compute/DataInterfaces/Elements/PCGAnimBankMeshSpawnerDataInterface.cpp
new file mode 100644
--- /dev/null
+++ b/Engine/Plugins/PCG/Source/PCG/Private/Compute/DataInterfaces/Elements/PCGAnimBankMeshSpawnerDataInterface.cpp
@@ -0,0 +1,250 @@
+// Copyright Epic Games, Inc. All Rights Reserved.
+
+#include "PCGAnimBankMeshSpawnerDataInterface.h"
+
+#include "Compute/PCGComputeCommon.h"
+#include "Compute/PCGDataBinding.h"
+#include "Elements/PCGAnimBankMeshSpawner.h"
+
+#include "GlobalRenderResources.h"
+#include "RHIResources.h"
+#include "ShaderCompilerCore.h"
+#include "ShaderCore.h"
+#include "ShaderParameterMetadataBuilder.h"
+#include "ComputeFramework/ShaderParamTypeDefinition.h"
+
+void UPCGAnimBankMeshSpawnerDataInterface::GetSupportedInputs(TArray<FShaderFunctionDefinition>& OutFunctions) const
+{
+	OutFunctions.AddDefaulted_GetRef()
+		.SetName(TEXT("ABSpawner_GetSelectorAttributeId"))
+		.AddReturnType(EShaderFundamentalType::Uint); // Attribute id to get mesh path string key from, or invalid if we should use CDF instead.
+
+	OutFunctions.AddDefaulted_GetRef()
+		.SetName(TEXT("ABSpawner_GetNumAttributes"))
+		.AddReturnType(EShaderFundamentalType::Uint); // Num attributes
+
+	OutFunctions.AddDefaulted_GetRef()
+		.SetName(TEXT("ABSpawner_GetNumPrimitives"))
+		.AddReturnType(EShaderFundamentalType::Uint); // Num primitives
+
+	OutFunctions.AddDefaulted_GetRef()
+		.SetName(TEXT("ABSpawner_ShouldApplyBounds"))
+		.AddReturnType(EShaderFundamentalType::Bool);
+
+	OutFunctions.AddDefaulted_GetRef()
+		.SetName(TEXT("ABSpawner_GetPrimitiveMeshBoundsMin"))
+		.AddReturnType(EShaderFundamentalType::Float, 3) // Local bounds min
+		.AddParam(EShaderFundamentalType::Uint); // InPrimitiveIndex
+
+	OutFunctions.AddDefaulted_GetRef()
+		.SetName(TEXT("ABSpawner_GetPrimitiveMeshBoundsMax"))
+		.AddReturnType(EShaderFundamentalType::Float, 3) // Local bounds max
+		.AddParam(EShaderFundamentalType::Uint); // InPrimitiveIndex
+
+	OutFunctions.AddDefaulted_GetRef()
+		.SetName(TEXT("ABSpawner_GetAttributeIdOffsetStride"))
+		.AddReturnType(EShaderFundamentalType::Uint, 4)
+		.AddParam(EShaderFundamentalType::Uint); // InAttributeIndex
+
+	OutFunctions.AddDefaulted_GetRef()
+		.SetName(TEXT("ABSpawner_GetPrimitiveStringKey"))
+		.AddReturnType(EShaderFundamentalType::Int) // String key
+		.AddParam(EShaderFundamentalType::Uint); // InPrimitiveIndex
+
+	OutFunctions.AddDefaulted_GetRef()
+		.SetName(TEXT("ABSpawner_GetPrimitiveIndexFromStringKey"))
+		.AddReturnType(EShaderFundamentalType::Uint) // Primitive index
+		.AddParam(EShaderFundamentalType::Int); // InMeshPathStringKey
+
+	OutFunctions.AddDefaulted_GetRef()
+		.SetName(TEXT("ABSpawner_GetPrimitiveSelectionCDF"))
+		.AddReturnType(EShaderFundamentalType::Float) // CDF value
+		.AddParam(EShaderFundamentalType::Uint); // InPrimitiveIndex
+
+	OutFunctions.AddDefaulted_GetRef()
+		.SetName(TEXT("ABSpawner_GetSelectedMeshAttributeId"))
+		.AddReturnType(EShaderFundamentalType::Uint); // Attribute id to output mesh path string key to
+}
+
+BEGIN_SHADER_PARAMETER_STRUCT(FPCGAnimBankMeshSpawnerDataInterfaceParameters,)
+SHADER_PARAMETER_ARRAY(FUintVector4, AttributeIdOffsetStrides, [UPCGAnimBankMeshSpawnerDataInterface::MAX_ATTRIBUTES])
+	SHADER_PARAMETER_ARRAY(FIntVector4, PrimitiveStringKeys, [PCGComputeConstants::MAX_PRIMITIVE_COMPONENTS_PER_SPAWNER])
+	SHADER_PARAMETER_ARRAY(FVector4f, PrimitiveMeshBoundsMin, [PCGComputeConstants::MAX_PRIMITIVE_COMPONENTS_PER_SPAWNER])
+	SHADER_PARAMETER_ARRAY(FVector4f, PrimitiveMeshBoundsMax, [PCGComputeConstants::MAX_PRIMITIVE_COMPONENTS_PER_SPAWNER])
+	SHADER_PARAMETER_SCALAR_ARRAY(float, SelectionCDF, [PCGComputeConstants::MAX_PRIMITIVE_COMPONENTS_PER_SPAWNER])
+	SHADER_PARAMETER(uint32, NumAttributes)
+	SHADER_PARAMETER(uint32, NumPrimitives)
+	SHADER_PARAMETER(int32, SelectorAttributeId)
+	SHADER_PARAMETER(int32, SelectedMeshAttributeId)
+	SHADER_PARAMETER(uint32, ApplyBounds)
+END_SHADER_PARAMETER_STRUCT()
+
+void UPCGAnimBankMeshSpawnerDataInterface::GetShaderParameters(TCHAR const* UID, FShaderParametersMetadataBuilder& InOutBuilder, FShaderParametersMetadataAllocations& InOutAllocations) const
+{
+	InOutBuilder.AddNestedStruct<FPCGAnimBankMeshSpawnerDataInterfaceParameters>(UID);
+}
+
+void UPCGAnimBankMeshSpawnerDataInterface::GetHLSL(FString& OutHLSL, FString const& InDataInterfaceName) const
+{
+	TMap<FString, FStringFormatArg> TemplateArgs =
+	{
+		{ TEXT("DataInterfaceName"), InDataInterfaceName },
+		{ TEXT("MaxAttributes"), MAX_ATTRIBUTES },
+		{ TEXT("MaxPrimitives"), PCGComputeConstants::MAX_PRIMITIVE_COMPONENTS_PER_SPAWNER },
+	};
+
+	OutHLSL += FString::Format(TEXT(
+		"int {DataInterfaceName}_SelectorAttributeId;\n"
+		"uint {DataInterfaceName}_NumAttributes;\n"
+		"uint {DataInterfaceName}_NumPrimitives;\n"
+		"uint {DataInterfaceName}_ApplyBounds;"
+		"int {DataInterfaceName}_SelectedMeshAttributeId;\n"
+		"uint4 {DataInterfaceName}_AttributeIdOffsetStrides[{MaxAttributes}];\n"
+		"float4 {DataInterfaceName}_PrimitiveMeshBoundsMin[{MaxPrimitives}];"
+		"float4 {DataInterfaceName}_PrimitiveMeshBoundsMax[{MaxPrimitives}];"
+		"int4 {DataInterfaceName}_PrimitiveStringKeys[{MaxPrimitives}];"
+		"DECLARE_SCALAR_ARRAY(float, {DataInterfaceName}_SelectionCDF, {MaxPrimitives});\n"
+		"\n"
+		"int ABSpawner_GetSelectorAttributeId_{DataInterfaceName}()\n"
+		"{\n"
+		"	return {DataInterfaceName}_SelectorAttributeId;\n"
+		"}\n"
+		"\n"
+		"uint ABSpawner_GetNumAttributes_{DataInterfaceName}()\n"
+		"{\n"
+		"	return {DataInterfaceName}_NumAttributes;\n"
+		"}\n"
+		"\n"
+		"uint4 ABSpawner_GetAttributeIdOffsetStride_{DataInterfaceName}(uint InAttributeIndex)\n"
+		"{\n"
+		"	return {DataInterfaceName}_AttributeIdOffsetStrides[InAttributeIndex];\n"
+		"}\n"
+		"\n"
+		"uint ABSpawner_GetNumPrimitives_{DataInterfaceName}()\n"
+		"{\n"
+		"	return {DataInterfaceName}_NumPrimitives;\n"
+		"}\n"
+		"\n"
+		"bool ABSpawner_ShouldApplyBounds_{DataInterfaceName}()\n"
+		"{\n"
+		"	return {DataInterfaceName}_ApplyBounds > 0;\n"
+		"}\n"
+		"\n"
+		"float3 ABSpawner_GetPrimitiveMeshBoundsMin_{DataInterfaceName}(uint InPrimitiveIndex)\n"
+		"{\n"
+		"	return {DataInterfaceName}_PrimitiveMeshBoundsMin[InPrimitiveIndex].xyz;\n"
+		"}\n"
+		"\n"
+		"float3 ABSpawner_GetPrimitiveMeshBoundsMax_{DataInterfaceName}(uint InPrimitiveIndex)\n"
+		"{\n"
+		"	return {DataInterfaceName}_PrimitiveMeshBoundsMax[InPrimitiveIndex].xyz;\n"
+		"}\n"
+		"\n"
+		"int ABSpawner_GetPrimitiveStringKey_{DataInterfaceName}(uint InPrimitiveIndex)\n"
+		"{\n"
+		"	return {DataInterfaceName}_PrimitiveStringKeys[InPrimitiveIndex][0];\n"
+		"}\n"
+		"\n"
+		"uint ABSpawner_GetPrimitiveIndexFromStringKey_{DataInterfaceName}(int InMeshPathStringKey)\n"
+		"{\n"
+		"	for (uint Index = 0; Index < {DataInterfaceName}_NumPrimitives; ++Index)\n"
+		"	{\n"
+		"		if ({DataInterfaceName}_PrimitiveStringKeys[Index][0] == InMeshPathStringKey)\n"
+		"		{\n"
+		"			return Index;\n"
+		"		}\n"
+		"	}\n"
+		"	\n"
+		"	return (uint)-1;\n"
+		"}\n"
+		"\n"
+		"float ABSpawner_GetPrimitiveSelectionCDF_{DataInterfaceName}(uint InPrimitiveIndex)\n"
+		"{\n"
+		"	return GET_SCALAR_ARRAY_ELEMENT({DataInterfaceName}_SelectionCDF, InPrimitiveIndex);\n"
+		"}\n"
+		"\n"
+		"int ABSpawner_GetSelectedMeshAttributeId_{DataInterfaceName}()\n"
+		"{\n"
+		"	return {DataInterfaceName}_SelectedMeshAttributeId;\n"
+		"}\n"
+	), TemplateArgs);
+}
+
+UComputeDataProvider* UPCGAnimBankMeshSpawnerDataInterface::CreateDataProvider(TObjectPtr<UObject> InBinding, uint64 InInputMask, uint64 InOutputMask) const
+{
+	TRACE_CPUPROFILER_EVENT_SCOPE(UPCGAnimBankMeshSpawnerDataInterface::CreateDataProvider);
+	UPCGDataBinding* Binding = CastChecked<UPCGDataBinding>(InBinding);
+
+	TObjectPtr<UPCGAnimBankMeshSpawnerDataProvider> DataProvider = NewObject<UPCGAnimBankMeshSpawnerDataProvider>();
+	DataProvider->Settings = CastChecked<UPCGAnimBankMeshSpawnerSettings>(Settings);
+	
+	// If there were 0 input points for this execution, we will not have created any primitives, so check for null.
+	if (const FPCGSpawnerPrimitives* Primitives = Binding->MeshSpawnersToPrimitives.Find(Settings))
+	{
+		DataProvider->AttributeIdOffsetStrides = Primitives->AttributeIdOffsetStrides;
+		DataProvider->SelectionCDF = Primitives->SelectionCDF;
+		DataProvider->PrimitiveStringKeys = Primitives->PrimitiveStringKeys;
+		DataProvider->PrimitiveMeshBounds = Primitives->PrimitiveMeshBounds;
+
+		// We only know name at compile time (not type), so we resolved this Id at execution time by inspecting the data arriving on the input pin. 
+		DataProvider->SelectorAttributeId = Primitives->SelectorAttributeId;
+		// We know the name and type of the selected mesh attribute statically and declared the attribute in GetKernelAttributeKeys,
+		// so the attribute ID should be present in the attribute table.
+		DataProvider->SelectedMeshAttributeId = Binding->GetAttributeId(DataProvider->Settings->OutAttributeName, EPCGKernelAttributeType::StringKey);
+		ensure(DataProvider->SelectedMeshAttributeId != INDEX_NONE);
+	}
+
+	return DataProvider;
+}
+
+FComputeDataProviderRenderProxy* UPCGAnimBankMeshSpawnerDataProvider::GetRenderProxy()
+{
+	return new FPCGAnimBankMeshSpawnerDataProviderProxy(AttributeIdOffsetStrides, SelectorAttributeId, PrimitiveStringKeys, SelectionCDF, SelectedMeshAttributeId, PrimitiveMeshBounds);
+}
+
+bool FPCGAnimBankMeshSpawnerDataProviderProxy::IsValid(FValidationData const& InValidationData) const
+{
+	if (InValidationData.ParameterStructSize != sizeof(FParameters))
+	{
+		return false;
+	}
+
+	return true;
+}
+
+void FPCGAnimBankMeshSpawnerDataProviderProxy::GatherDispatchData(FDispatchData const& InDispatchData)
+{
+	const TStridedView<FParameters> ParameterArray = MakeStridedParameterView<FParameters>(InDispatchData);
+	for (int32 InvocationIndex = 0; InvocationIndex < ParameterArray.Num(); ++InvocationIndex)
+	{
+		FParameters& Parameters = ParameterArray[InvocationIndex];
+
+		Parameters.NumAttributes = AttributeIdOffsetStrides.Num();
+		Parameters.NumPrimitives = SelectionCDF.Num();
+		Parameters.SelectorAttributeId = SelectorAttributeId;
+		Parameters.SelectedMeshAttributeId = SelectedMeshAttributeId;
+
+		for (int32 Index = 0; Index < AttributeIdOffsetStrides.Num(); ++Index)
+		{
+			Parameters.AttributeIdOffsetStrides[Index] = AttributeIdOffsetStrides[Index];
+		}
+
+		for (int32 Index = 0; Index < PrimitiveStringKeys.Num(); ++Index)
+		{
+			Parameters.PrimitiveStringKeys[Index] = FIntVector4(PrimitiveStringKeys[Index], /*Unused*/0, /*Unused*/0, /*Unused*/0);
+		}
+
+		for (int32 Index = 0; Index < SelectionCDF.Num(); ++Index)
+		{
+			GET_SCALAR_ARRAY_ELEMENT(Parameters.SelectionCDF, Index) = SelectionCDF[Index];
+		}
+
+		Parameters.ApplyBounds = PrimitiveMeshBounds.IsEmpty() ? 0 : 1;
+
+		for (int32 Index = 0; Index < PrimitiveMeshBounds.Num(); ++Index)
+		{
+			Parameters.PrimitiveMeshBoundsMin[Index] = FVector4f(PrimitiveMeshBounds[Index].Min.X, PrimitiveMeshBounds[Index].Min.Y, PrimitiveMeshBounds[Index].Min.Z, /*Unused*/0.0f);
+			Parameters.PrimitiveMeshBoundsMax[Index] = FVector4f(PrimitiveMeshBounds[Index].Max.X, PrimitiveMeshBounds[Index].Max.Y, PrimitiveMeshBounds[Index].Max.Z, /*Unused*/0.0f);
+		}
+	}
+}
diff --git a/Engine/Plugins/PCG/Source/PCG/Private/Compute/DataInterfaces/Elements/PCGAnimBankMeshSpawnerDataInterface.h b/Engine/Plugins/PCG/Source/PCG/Private/Compute/DataInterfaces/Elements/PCGAnimBankMeshSpawnerDataInterface.h
new file mode 100644
--- /dev/null
+++ b/Engine/Plugins/PCG/Source/PCG/Private/Compute/DataInterfaces/Elements/PCGAnimBankMeshSpawnerDataInterface.h
@@ -0,0 +1,105 @@
+// Copyright Epic Games, Inc. All Rights Reserved.
+
+#pragma once
+
+#include "Compute/DataInterfaces/PCGComputeDataInterface.h"
+#include "ComputeFramework/ComputeDataProvider.h"
+
+#include "PCGAnimBankMeshSpawnerDataInterface.generated.h"
+
+class FPCGAnimBankMeshSpawnerDataInterfaceParameters;
+class UPCGAnimBankMeshSpawnerSettings;
+class UPCGSettings;
+
+/** Data Interface to marshal AnimBank Mesh Spawner settings to the GPU. */
+UCLASS(ClassGroup = (Procedural))
+class UPCGAnimBankMeshSpawnerDataInterface : public UPCGComputeDataInterface
+{
+	GENERATED_BODY()
+
+public:
+	//~ Begin UComputeDataInterface Interface
+	TCHAR const* GetClassName() const override { return TEXT("PCGAnimBankMeshSpawner"); }
+	void GetSupportedInputs(TArray<FShaderFunctionDefinition>& OutFunctions) const override;
+	void GetShaderParameters(TCHAR const* UID, FShaderParametersMetadataBuilder& InOutBuilder, FShaderParametersMetadataAllocations& InOutAllocations) const override;
+	void GetHLSL(FString& OutHLSL, FString const& InDataInterfaceName) const override;
+	UComputeDataProvider* CreateDataProvider(TObjectPtr<UObject> InBinding, uint64 InInputMask, uint64 InOutputMask) const override;
+	//~ End UComputeDataInterface Interface
+
+	UPROPERTY()
+	TObjectPtr<const UPCGSettings> Settings = nullptr;
+
+	static constexpr uint32 MAX_ATTRIBUTES = 64;
+};
+
+UCLASS()
+class UPCGAnimBankMeshSpawnerDataProvider : public UComputeDataProvider
+{
+	GENERATED_BODY()
+
+public:
+	//~ Begin UComputeDataProvider Interface
+	FComputeDataProviderRenderProxy* GetRenderProxy() override;
+	//~ End UComputeDataProvider Interface
+
+	UPROPERTY()
+	TObjectPtr<const UPCGAnimBankMeshSpawnerSettings> Settings = nullptr;
+
+	UPROPERTY()
+	TArray<FUintVector4> AttributeIdOffsetStrides;
+
+	UPROPERTY()
+	TArray<int32> PrimitiveStringKeys;
+
+	UPROPERTY()
+	TArray<FBox> PrimitiveMeshBounds;
+
+	UPROPERTY()
+	TArray<float> SelectionCDF;
+
+	UPROPERTY()
+	int32 SelectorAttributeId = INDEX_NONE;
+
+	UPROPERTY()
+	int32 NumInputPoints = 0;
+
+	UPROPERTY()
+	int32 SelectedMeshAttributeId = INDEX_NONE;
+};
+
+class FPCGAnimBankMeshSpawnerDataProviderProxy : public FComputeDataProviderRenderProxy
+{
+public:
+	FPCGAnimBankMeshSpawnerDataProviderProxy(
+		const TArray<FUintVector4>& InAttributeIdOffsetStrides,
+		int32 InSelectorAttributeId,
+		const TArray<int32>& InPrimitiveStringKeys,
+		TArray<float> InSelectionCDF,
+		int32 InSelectedMeshAttributeId,
+		const TArray<FBox>& InPrimitiveMeshBounds)
+		: AttributeIdOffsetStrides(InAttributeIdOffsetStrides)
+		, SelectionCDF(InSelectionCDF)
+		, SelectorAttributeId(InSelectorAttributeId)
+		, PrimitiveStringKeys(InPrimitiveStringKeys)
+		, PrimitiveMeshBounds(InPrimitiveMeshBounds)
+		, SelectedMeshAttributeId(InSelectedMeshAttributeId)
+	{}
+
+	//~ Begin FComputeDataProviderRenderProxy Interface
+	bool IsValid(FValidationData const& InValidationData) const override;
+	void GatherDispatchData(FDispatchData const& InDispatchData) override;
+	//~ End FComputeDataProviderRenderProxy Interface
+
+protected:
+	using FParameters = FPCGAnimBankMeshSpawnerDataInterfaceParameters;
+
+	TArray<FUintVector4> AttributeIdOffsetStrides;
+	TArray<float> SelectionCDF;
+
+	int32 SelectorAttributeId = INDEX_NONE;
+	TArray<int32> PrimitiveStringKeys;
+
+	TArray<FBox> PrimitiveMeshBounds;
+
+	int32 SelectedMeshAttributeId = INDEX_NONE;
+};
diff --git a/Engine/Plugins/PCG/Source/PCG/Private/Elements/PCGAnimBankMeshSpawner.cpp b/Engine/Plugins/PCG/Source/PCG/Private/Elements/PCGAnimBankMeshSpawner.cpp
new file mode 100644
--- /dev/null
+++ b/Engine/Plugins/PCG/Source/PCG/Private/Elements/PCGAnimBankMeshSpawner.cpp
@@ -0,0 +1,837 @@
+// Copyright Epic Games, Inc. All Rights Reserved.
+
+#include "Elements/PCGAnimBankMeshSpawner.h"
+
+#include "PCGComponent.h"
+#include "PCGCustomVersion.h"
+#include "PCGManagedResource.h"
+#include "Compute/PCGComputeCommon.h"
+#include "Compute/PCGDataBinding.h"
+#include "Compute/DataInterfaces/PCGInstanceDataInterface.h"
+#include "Compute/DataInterfaces/Elements/PCGAnimBankMeshSpawnerDataInterface.h"
+#include "Data/PCGPointData.h"
+#include "Data/PCGSpatialData.h"
+#include "Elements/PCGAnimBankMeshSpawnerContext.h"
+#include "Graph/PCGGPUGraphCompilationContext.h"
+#include "Helpers/PCGActorHelpers.h"
+#include "Helpers/PCGHelpers.h"
+#include "InstanceDataPackers/PCGAnimBankMeshInstanceDataPackerBase.h"
+#include "MeshSelectors/PCGAnimBankMeshSelectorBase.h"
+#include "MeshSelectors/PCGAnimBankMeshSelectorByAttribute.h"
+#include "MeshSelectors/PCGAnimBankMeshSelectorWeighted.h"
+
+#include "Components/AnimBankMeshComponent.h"
+#include "Engine/SkinnedAsset.h"
+#include "Materials/MaterialInterface.h"
+#include "ShaderCompilerCore.h"
+#include "UObject/Package.h"
+
+#include UE_INLINE_GENERATED_CPP_BY_NAME(PCGAnimBankMeshSpawner)
+
+#define LOCTEXT_NAMESPACE "PCGAnimBankMeshSpawnerElement"
+
+UPCGAnimBankMeshSpawnerSettings::UPCGAnimBankMeshSpawnerSettings(const FObjectInitializer &ObjectInitializer)
+{
+	MeshSelectorType = UPCGAnimBankMeshSelectorWeighted::StaticClass();
+	// Implementation note: this should not have been done here (it should have been null), as it causes issues with copy & paste
+	// when the thing to paste does not have that class for its instance.
+	// However, removing it makes it that any object actually using the instance created by default would be lost.
+	if (!this->HasAnyFlags(RF_ClassDefaultObject))
+	{
+		MeshSelectorParameters = ObjectInitializer.CreateDefaultSubobject<UPCGAnimBankMeshSelectorWeighted>(this, TEXT("DefaultSelectorInstance"));
+	}
+}
+
+TCHAR const* UPCGAnimBankMeshSpawnerSettings::TemplateFilePath = TEXT("/Plugin/PCG/Private/Elements/PCGAnimBankMeshSpawner.usf");
+
+bool UPCGAnimBankMeshSpawnerSettings::IsKernelValid(FPCGContext* InContext, bool bQuiet) const
+{
+	TRACE_CPUPROFILER_EVENT_SCOPE(UPCGAnimBankMeshSpawnerSettings::IsKernelValid);
+
+	if (!Super::IsKernelValid(InContext, bQuiet))
+	{
+		return false;
+	}
+
+	if (UPCGAnimBankMeshSelectorWeighted* SelectorWeighted = Cast<UPCGAnimBankMeshSelectorWeighted>(MeshSelectorParameters))
+	{
+		if (SelectorWeighted->MeshEntries.IsEmpty())
+		{
+			PCG_KERNEL_VALIDATION_ERR(InContext, this, bQuiet, LOCTEXT("NoMeshEntries", "No meshes specified in selector."));
+			return false;
+		}
+
+		for (const FPCGAnimBankMeshSelectorWeightedEntry& Entry : SelectorWeighted->MeshEntries)
+		{
+			if (Entry.Descriptor.SkinnedAsset.IsNull())
+			{
+				PCG_KERNEL_VALIDATION_ERR(InContext, this, bQuiet, LOCTEXT("UnassignedMesh", "Unassigned mesh."));
+				return false;
+			}
+
+			for (const FSoftAnimBankItem& BankItem : Entry.Descriptor.BankItems)
+			{
+				if (BankItem.BankAsset.IsNull())
+				{
+					PCG_KERNEL_VALIDATION_ERR(InContext, this, bQuiet, LOCTEXT("UnassignedAnimBank", "Unassigned anim bank."));
+					return false;
+				}
+			}
+		}
+	}
+	else if (UPCGAnimBankMeshSelectorByAttribute* SelectorByAttribute = Cast<UPCGAnimBankMeshSelectorByAttribute>(MeshSelectorParameters))
+	{
+		if (SelectorByAttribute->AttributeName == NAME_None)
+		{
+			PCG_KERNEL_VALIDATION_ERR(InContext, this, bQuiet, LOCTEXT("NoMeshAttribute", "Mesh selector attribute name is 'None'."));
+			return false;
+		}
+	}
+	else
+	{
+		PCG_KERNEL_VALIDATION_ERR(InContext, this, bQuiet, LOCTEXT("InvalidMeshSelector", "Currently GPU AnimBank Mesh Spawner nodes must use PCGAnimBankMeshSelectorWeighted or UPCGAnimBankMeshSelectorByAttribute as the mesh selector type."));
+		return false;
+	}
+
+	// Currently instance packers must be able to specify a full list of attribute names upfront, to build the attribute table at compile time.
+	// TODO: We should be able to augment a static attribute table with new attributes at execution time, which will allow other types like regex.
+	if (InstanceDataPackerParameters && !InstanceDataPackerParameters->GetAttributeNames(/*OutNames=*/nullptr))
+	{
+		PCG_KERNEL_VALIDATION_ERR(InContext, this, bQuiet, LOCTEXT("InvalidInstancePacker", "Selected instance packer does not support GPU execution."));
+		return false;
+	}
+
+	return true;
+}
+
+FString UPCGAnimBankMeshSpawnerSettings::GetCookedKernelSource(const FPCGKernelAttributeTable& InAttributeStaticTable) const
+{
+	FString TemplateFile;
+	ensure(LoadShaderSourceFile(TemplateFilePath, EShaderPlatform::SP_PCD3D_SM5, &TemplateFile, nullptr));
+	return TemplateFile;
+}
+
+bool UPCGAnimBankMeshSpawnerSettings::ComputeOutputPinDataDesc(const UPCGPin* OutputPin, const UPCGDataBinding* Binding, FPCGDataCollectionDesc& OutDesc) const
+{
+	check(OutputPin);
+	check(Binding);
+
+	// Code assumes single output pin.
+	if (!ensure(OutputPin->Properties.Label == PCGPinConstants::DefaultOutputLabel))
+	{
+		return false;
+	}
+
+	// First inspect data coming from CPU. Will return description if this node is a CPU node or in a different compute graph.
+	if (Binding->ComputeCPUOutputPinDataDesc(OutputPin, OutDesc))
+	{
+		return true;
+	}
+
+	const FPCGDataCollectionDesc InputPinDesc = PCGDataForGPUHelpers::ComputeInputPinDataDesc(this, PCGPinConstants::DefaultInputLabel, Binding);
+	
+	// First forward input data to output.
+	FPCGDataCollectionDesc OutputPinDesc = InputPinDesc;
+
+	// Add output attribute (selected mesh).
+	{
+		// Create unique value keys for the output string.
+		TArray<int32> UniqueStringKeys;
+		if (UPCGAnimBankMeshSelectorWeighted* SelectorWeighted = Cast<UPCGAnimBankMeshSelectorWeighted>(MeshSelectorParameters))
+		{
+			// Weighted selection - add explicit strings from settings.
+			for (const FPCGAnimBankMeshSelectorWeightedEntry& Entry : SelectorWeighted->MeshEntries)
+			{
+				const FString Value = Entry.Descriptor.SkinnedAsset.ToString();
+				const int32 StringIndex = Binding->GetStringTable().IndexOfByKey(Value);
+				if (ensureAlways(StringIndex != INDEX_NONE))
+				{
+					UniqueStringKeys.AddUnique(StringIndex);
+				}
+			}
+		}
+		else if (UPCGAnimBankMeshSelectorByAttribute* SelectorByAttribute = Cast<UPCGAnimBankMeshSelectorByAttribute>(MeshSelectorParameters))
+		{
+			// By-attribute selection - pass on strings from input attribute.
+			if (SelectorByAttribute->AttributeName != NAME_None)
+			{
+				for (const FPCGDataDesc& InputDataDesc : InputPinDesc.DataDescs)
+				{
+					for (const FPCGKernelAttributeDesc& InputAttrDesc : InputDataDesc.AttributeDescs)
+					{
+						if (InputAttrDesc.Name == SelectorByAttribute->AttributeName)
+						{
+							for (int32 InputStringKey : InputAttrDesc.UniqueStringKeys)
+							{
+								UniqueStringKeys.AddUnique(InputStringKey);
+							}
+						}
+					}
+				}
+			}
+		}
+		else if (MeshSelectorParameters)
+		{
+			UE_LOG(LogPCG, Warning, TEXT("Mesh selector not supported by GPU AnimBank Mesh Spawner: %s"), *MeshSelectorParameters->GetName());
+		}
+
+		OutputPinDesc.AddAttributeToAllData(FPCGKernelAttributeKey(OutAttributeName, EPCGKernelAttributeType::StringKey), Binding, &UniqueStringKeys);
+	}
+
+	OutDesc = OutputPinDesc;
+	return true;
+}
+
+void UPCGAnimBankMeshSpawnerSettings::GetKernelAttributeKeys(TArray<FPCGKernelAttributeKey>& OutKeys) const
+{
+	// Register the attribute this node creates.
+	OutKeys.AddUnique(FPCGKernelAttributeKey(OutAttributeName, EPCGKernelAttributeType::StringKey));
+}
+
+void UPCGAnimBankMeshSpawnerSettings::AddStaticCreatedStrings(TArray<FString>& InOutStringTable) const
+{
+	if (UPCGAnimBankMeshSelectorWeighted* SelectorWeighted = Cast<UPCGAnimBankMeshSelectorWeighted>(MeshSelectorParameters))
+	{
+		for (const FPCGAnimBankMeshSelectorWeightedEntry& Entry : SelectorWeighted->MeshEntries)
+		{
+			InOutStringTable.AddUnique(Entry.Descriptor.SkinnedAsset.ToString());
+		}
+	}
+}
+
+int UPCGAnimBankMeshSpawnerSettings::ComputeKernelThreadCount(const UPCGDataBinding* Binding) const
+{
+	const UPCGNode* Node = CastChecked<UPCGNode>(GetOuter());
+	const UPCGPin* Pin = Node->GetInputPin(PCGPinConstants::DefaultInputLabel);
+
+	return PCGDataForGPUHelpers::ComputeInputPinDataDesc(Pin, Binding).ComputeDataElementCount(EPCGDataType::Point);
+}
+
+#if WITH_EDITOR
+void UPCGAnimBankMeshSpawnerSettings::CreateAdditionalInputDataInterfaces(FPCGGPUCompilationContext& InOutContext, UObject* InObjectOuter, TArray<TObjectPtr<UComputeDataInterface>>& OutDataInterfaces) const
+{
+	Super::CreateAdditionalInputDataInterfaces(InOutContext, InObjectOuter, OutDataInterfaces);
+
+	TObjectPtr<UPCGAnimBankMeshSpawnerDataInterface> NodeDI = InOutContext.NewObject_AnyThread<UPCGAnimBankMeshSpawnerDataInterface>(InObjectOuter);
+	NodeDI->Settings = this;
+	OutDataInterfaces.Add(NodeDI);
+}
+
+void UPCGAnimBankMeshSpawnerSettings::CreateAdditionalOutputDataInterfaces(FPCGGPUCompilationContext& InOutContext, UObject* InObjectOuter, TArray<TObjectPtr<UComputeDataInterface>>& OutDataInterfaces) const
+{
+	Super::CreateAdditionalOutputDataInterfaces(InOutContext, InObjectOuter, OutDataInterfaces);
+
+	UPCGInstanceDataInterface* InstanceDI = InOutContext.NewObject_AnyThread<UPCGInstanceDataInterface>(InObjectOuter);
+	InstanceDI->SetProducerSettings(this);
+	InstanceDI->InputPinProvidingData = PCGPinConstants::DefaultInputLabel;
+	OutDataInterfaces.Add(InstanceDI);
+}
+
+FText UPCGAnimBankMeshSpawnerSettings::GetDefaultNodeTitle() const
+{
+	return LOCTEXT("NodeTitle", "AnimBank Mesh Spawner");
+}
+
+void UPCGAnimBankMeshSpawnerSettings::ApplyDeprecation(UPCGNode* InOutNode)
+{
+	check(InOutNode);
+	Super::ApplyDeprecation(InOutNode);
+}
+#endif
+
+TArray<FPCGPinProperties> UPCGAnimBankMeshSpawnerSettings::InputPinProperties() const
+{
+	// Note: If executing on the GPU, we need to prevent multiple connections on inputs, since it is not supported at this time.
+	// Also note: Since the ShouldExecuteOnGPU() is already tied to structural changes, we don't need to implement any logic for this in GetChangeTypeForProperty()
+	const bool bAllowMultipleConnections = !ShouldExecuteOnGPU();
+
+	TArray<FPCGPinProperties> Properties;
+	FPCGPinProperties& InputPinProperty = Properties.Emplace_GetRef(PCGPinConstants::DefaultInputLabel, EPCGDataType::Point, bAllowMultipleConnections);
+	InputPinProperty.SetRequiredPin();
+
+	return Properties;
+}
+
+FPCGElementPtr UPCGAnimBankMeshSpawnerSettings::CreateElement() const
+{
+	return MakeShared<FPCGAnimBankMeshSpawnerElement>();
+}
+
+FPCGContext* FPCGAnimBankMeshSpawnerElement::CreateContext()
+{
+	return new FPCGAnimBankMeshSpawnerContext();
+}
+
+bool FPCGAnimBankMeshSpawnerElement::PrepareDataInternal(FPCGContext* InContext) const
+{
+	TRACE_CPUPROFILER_EVENT_SCOPE(FPCGAnimBankMeshSpawnerElement::PrepareDataInternal);
+	// TODO : time-sliced implementation
+	FPCGAnimBankMeshSpawnerContext* Context = static_cast<FPCGAnimBankMeshSpawnerContext*>(InContext);
+	const UPCGAnimBankMeshSpawnerSettings* Settings = Context->GetInputSettings<UPCGAnimBankMeshSpawnerSettings>();
+	check(Settings);
+
+	if (!Settings->MeshSelectorParameters)
+	{
+		PCGE_LOG(Error, GraphAndLog, LOCTEXT("InvalidMeshSelectorInstance", "Invalid MeshSelector instance, try reselecting the MeshSelector type"));
+		return true;
+	}
+
+	if (!Context->SourceComponent.Get())
+	{
+		return true;
+	}
+
+#if WITH_EDITOR
+	// In editor, we always want to generate this data for inspection & to prevent caching issues
+	const bool bGenerateOutput = true;
+#else
+	const bool bGenerateOutput = Context->Node && Context->Node->IsOutputPinConnected(PCGPinConstants::DefaultOutputLabel);
+#endif
+
+	// Check if we can reuse existing resources
+	bool& bSkippedDueToReuse = Context->bSkippedDueToReuse;
+
+	if (!Context->bReuseCheckDone)
+	{
+		// Compute CRC if it has not been computed (it likely isn't, but this is to futureproof this)
+		if (!Context->DependenciesCrc.IsValid())
+		{
+			GetDependenciesCrc(Context->InputData, Settings, Context->SourceComponent.Get(), Context->DependenciesCrc);
+		}
+		
+		if (Context->DependenciesCrc.IsValid())
+		{
+			TArray<UPCGManagedAnimBankMeshComponent*> MABMCs;
+			Context->SourceComponent->ForEachManagedResource([&MABMCs, &Context, Settings](UPCGManagedResource* InResource)
+			{
+				if (UPCGManagedAnimBankMeshComponent* Resource = Cast<UPCGManagedAnimBankMeshComponent>(InResource))
+				{
+					if (Resource->GetCrc().IsValid() && Resource->GetCrc() == Context->DependenciesCrc)
+					{
+						MABMCs.Add(Resource);
+					}
+				}
+			});
+
+			for (UPCGManagedAnimBankMeshComponent* MABMC : MABMCs)
+			{
+				if (!MABMC->IsMarkedUnused() && Settings->bWarnOnIdenticalSpawn)
+				{
+					// TODO: Revisit if the stack is added to the managed components at creation
+					PCGLog::LogWarningOnGraph(LOCTEXT("IdenticalABMCSpawn", "Identical AnimBank Mesh Component spawn occurred. It may be beneficial to re-check graph logic for identical spawn conditions (same mesh descriptor at same location, etc) or repeated nodes."), Context);
+				}
+
+				MABMC->MarkAsReused();
+			}
+
+			if (!MABMCs.IsEmpty())
+			{
+				bSkippedDueToReuse = true;
+			}
+		}
+
+		Context->bReuseCheckDone = true;
+	}
+
+	// Early out - if we've established we could reuse resources and there is no need to generate an output, quit now
+	if (!bGenerateOutput && bSkippedDueToReuse)
+	{
+		return true;
+	}
+
+	// perform mesh selection
+	TArray<FPCGTaggedData> Inputs = Context->InputData.GetInputs();
+	TArray<FPCGTaggedData>& Outputs = Context->OutputData.TaggedData;
+
+	while(Context->CurrentInputIndex < Inputs.Num())
+	{
+		if (!Context->bCurrentInputSetup)
+		{
+			const FPCGTaggedData& Input = Inputs[Context->CurrentInputIndex];
+			const UPCGSpatialData* SpatialData = Cast<UPCGSpatialData>(Input.Data);
+
+			if (!SpatialData)
+			{
+				PCGE_LOG(Error, GraphAndLog, LOCTEXT("InvalidInputData", "Invalid input data"));
+				++Context->CurrentInputIndex;
+				continue;
+			}
+
+			const UPCGPointData* PointData = SpatialData->ToPointData(Context);
+			if (!PointData)
+			{
+				PCGE_LOG(Error, GraphAndLog, LOCTEXT("NoPointDataInInput", "Unable to get point data from input"));
+				++Context->CurrentInputIndex;
+				continue;
+			}
+
+			AActor* TargetActor = Settings->TargetActor.Get() ? Settings->TargetActor.Get() : Context->GetTargetActor(nullptr);
+			if (!TargetActor)
+			{
+				PCGE_LOG(Error, GraphAndLog, LOCTEXT("InvalidTargetActor", "Invalid target actor. Ensure TargetActor member is initialized when creating SpatialData."));
+				++Context->CurrentInputIndex;
+				continue;
+			}
+
+			if (bGenerateOutput)
+			{
+				FPCGTaggedData& Output = Outputs.Add_GetRef(Input);
+
+				UPCGPointData* OutputPointData = FPCGContext::NewObject_AnyThread<UPCGPointData>(Context);
+				OutputPointData->InitializeFromData(PointData);
+
+				if (OutputPointData->Metadata->HasAttribute(Settings->OutAttributeName))
+				{
+					OutputPointData->Metadata->DeleteAttribute(Settings->OutAttributeName);
+					PCGE_LOG(Verbose, LogOnly, FText::Format(LOCTEXT("AttributeOverwritten", "Metadata attribute '{0}' is being overwritten in the output data"), FText::FromName(Settings->OutAttributeName)));
+				}
+
+				OutputPointData->Metadata->CreateStringAttribute(Settings->OutAttributeName, FName(NAME_None).ToString(), /*bAllowsInterpolation=*/false);
+
+				Output.Data = OutputPointData;
+				check(!Context->CurrentOutputPointData);
+				Context->CurrentOutputPointData = OutputPointData;
+			}
+
+			FPCGAnimBankMeshSpawnerContext::FPackedInstanceListData& InstanceListData = Context->MeshInstancesData.Emplace_GetRef();
+			InstanceListData.TargetActor = TargetActor;
+			InstanceListData.SpatialData = PointData;
+
+			Context->CurrentPointData = PointData;
+			Context->bCurrentInputSetup = true;
+		}
+
+		// TODO: If we know we re-use the mesh components, we should not run the Selection, as it can be pretty costly.
+		// At the moment, the selection is filling the output point data, so it is necessary to run it. But we should just hit the cache in that case.
+		if (!Context->bSelectionDone)
+		{
+			check(Context->CurrentPointData);
+			Context->bSelectionDone = Settings->MeshSelectorParameters->SelectInstances(*Context, Settings, Context->CurrentPointData, Context->MeshInstancesData.Last().MeshInstances, Context->CurrentOutputPointData);
+		}
+
+		if (!Context->bSelectionDone)
+		{
+			return false;
+		}
+
+		// If we need the output but would otherwise skip the resource creation, we don't need to run the instance packing part of the processing
+		if (!bSkippedDueToReuse)
+		{
+			TArray<FPCGAnimBankMeshPackedCustomData>& PackedCustomData = Context->MeshInstancesData.Last().PackedCustomData;
+			const TArray<FPCGAnimBankMeshInstanceList>& MeshInstances = Context->MeshInstancesData.Last().MeshInstances;
+
+			if (PackedCustomData.Num() != MeshInstances.Num())
+			{
+				PackedCustomData.SetNum(MeshInstances.Num());
+			}
+
+			if (Settings->InstanceDataPackerParameters)
+			{
+				for (int32 InstanceListIndex = 0; InstanceListIndex < MeshInstances.Num(); ++InstanceListIndex)
+				{
+					Settings->InstanceDataPackerParameters->PackInstances(*Context, Context->CurrentPointData, MeshInstances[InstanceListIndex], PackedCustomData[InstanceListIndex]);
+				}
+			}
+		}
+
+		// We're done - cleanup for next iteration if we still have time
+		++Context->CurrentInputIndex;
+		Context->ResetInputIterationData();
+
+		// Continue on to next iteration if there is time left, otherwise, exit here
+		if (Context->AsyncState.ShouldStop() && Context->CurrentInputIndex < Inputs.Num())
+		{
+			return false;
+		}
+	}
+
+	IPCGAsyncLoadingContext* AsyncLoadingContext = static_cast<IPCGAsyncLoadingContext*>(Context);
+
+	if (Context->CurrentInputIndex == Inputs.Num() && !AsyncLoadingContext->WasLoadRequested() && !Context->MeshInstancesData.IsEmpty() && !Settings->bSynchronousLoad)
+	{
+		TArray<FSoftObjectPath> ObjectsToLoad;
+		for (const FPCGAnimBankMeshSpawnerContext::FPackedInstanceListData& InstanceData : Context->MeshInstancesData)
+		{
+			for (const FPCGAnimBankMeshInstanceList& MeshInstanceList : InstanceData.MeshInstances)
+			{
+				if (!MeshInstanceList.Descriptor.SkinnedAsset.IsNull())
+				{
+					ObjectsToLoad.AddUnique(MeshInstanceList.Descriptor.SkinnedAsset.ToSoftObjectPath());
+				}
+
+			#if 0 // AB-TODO
+				for (const TSoftObjectPtr<UMaterialInterface>& OverrideMaterial : MeshInstanceList.Descriptor.OverrideMaterials)
+				{
+					if (!OverrideMaterial.IsNull())
+					{
+						ObjectsToLoad.AddUnique(OverrideMaterial.ToSoftObjectPath());
+					}
+				}
+
+			#endif
+			}
+		}
+
+		return AsyncLoadingContext->RequestResourceLoad(Context, std::move(ObjectsToLoad), /*bAsynchronous=*/true);
+	}
+
+	return true;
+}
+
+bool FPCGAnimBankMeshSpawnerElement::ExecuteInternal(FPCGContext* InContext) const
+{
+	TRACE_CPUPROFILER_EVENT_SCOPE(FPCGAnimBankMeshSpawnerElement::Execute);
+	FPCGAnimBankMeshSpawnerContext* Context = static_cast<FPCGAnimBankMeshSpawnerContext*>(InContext);
+	const UPCGAnimBankMeshSpawnerSettings* Settings = Context->GetInputSettings<UPCGAnimBankMeshSpawnerSettings>();
+	check(Settings && !Settings->ShouldExecuteOnGPU());
+
+	while(!Context->MeshInstancesData.IsEmpty())
+	{
+		const FPCGAnimBankMeshSpawnerContext::FPackedInstanceListData& InstanceList = Context->MeshInstancesData.Last();
+		check(Context->bSkippedDueToReuse || InstanceList.MeshInstances.Num() == InstanceList.PackedCustomData.Num());
+
+		const bool bTargetActorValid = (InstanceList.TargetActor && IsValid(InstanceList.TargetActor));
+
+		if (bTargetActorValid)
+		{
+			while (Context->CurrentDataIndex < InstanceList.MeshInstances.Num())
+			{
+				const FPCGAnimBankMeshInstanceList& MeshInstance = InstanceList.MeshInstances[Context->CurrentDataIndex];
+				// We always have mesh instances, but if we are in re-use, we don't compute the packed custom data.
+				const FPCGAnimBankMeshPackedCustomData* PackedCustomData = InstanceList.PackedCustomData.IsValidIndex(Context->CurrentDataIndex) ? &InstanceList.PackedCustomData[Context->CurrentDataIndex] : nullptr;
+				SpawnAnimBankMeshInstances(Context, MeshInstance, InstanceList.TargetActor, PackedCustomData);
+
+				// Now that the mesh is loaded/spawned, set the bounds to out points if requested.
+				if (MeshInstance.Descriptor.SkinnedAsset && Settings->bApplyMeshBoundsToPoints)
+				{
+					if (TMap<UPCGPointData*, TArray<int32>>* OutPointDataToPointIndex = Context->MeshToOutPoints.Find(MeshInstance.Descriptor.SkinnedAsset))
+					{
+						const FBox Bounds = MeshInstance.Descriptor.SkinnedAsset->GetBounds().GetBox();
+						for (TPair<UPCGPointData*, TArray<int32>>& It : *OutPointDataToPointIndex)
+						{
+							check(It.Key);
+							TArray<FPCGPoint>& OutPoints = It.Key->GetMutablePoints();
+							for (int32 Index : It.Value)
+							{
+								FPCGPoint& Point = OutPoints[Index];
+								Point.BoundsMin = Bounds.Min;
+								Point.BoundsMax = Bounds.Max;
+							}
+						}
+					}
+				}
+
+				++Context->CurrentDataIndex;
+
+				if (Context->AsyncState.ShouldStop())
+				{
+					break;
+				}
+			}
+		}
+
+		if (!bTargetActorValid || Context->CurrentDataIndex == InstanceList.MeshInstances.Num())
+		{
+			Context->MeshInstancesData.RemoveAtSwap(Context->MeshInstancesData.Num() - 1);
+			Context->CurrentDataIndex = 0;
+		}
+
+		if (Context->AsyncState.ShouldStop())
+		{
+			break;
+		}
+	}
+
+	const bool bFinishedExecution = Context->MeshInstancesData.IsEmpty();
+	if (bFinishedExecution)
+	{
+		if (AActor* TargetActor = Settings->TargetActor.Get() ? Settings->TargetActor.Get() : Context->GetTargetActor(nullptr))
+		{
+			for (UFunction* Function : PCGHelpers::FindUserFunctions(TargetActor->GetClass(), Settings->PostProcessFunctionNames, { UPCGFunctionPrototypes::GetPrototypeWithNoParams() }, Context))
+			{
+				TargetActor->ProcessEvent(Function, nullptr);
+			}
+		}
+	}
+
+	return bFinishedExecution;
+}
+
+bool FPCGAnimBankMeshSpawnerElement::CanExecuteOnlyOnMainThread(FPCGContext* Context) const
+{
+	// PrepareData can call UPCGManagedComponent::MarkAsReused which registers the mesh component, which can go into Chaos code that asserts if not on main thread.
+	// TODO: We can likely re-enable multi-threading for PrepareData if we move the call to MarkAsReused to Execute. There should hopefully not be
+	// wider contention on resources resources are not shared across nodes and are also per-component.
+	return !Context || Context->CurrentPhase == EPCGExecutionPhase::Execute || Context->CurrentPhase == EPCGExecutionPhase::PrepareData;
+}
+
+void FPCGAnimBankMeshSpawnerElement::SpawnAnimBankMeshInstances(FPCGAnimBankMeshSpawnerContext* Context, const FPCGAnimBankMeshInstanceList& InstanceList, AActor* TargetActor, const FPCGAnimBankMeshPackedCustomData* InPackedCustomData) const
+{
+	// Populate the mesh component from the previously prepared entries
+	TRACE_CPUPROFILER_EVENT_SCOPE(FPCGAnimBankMeshSpawnerElement::Execute::PopulateAB);
+
+	if (InstanceList.Instances.Num() == 0)
+	{
+		return;
+	}
+
+	// Will be synchronously loaded if not loaded. But by default it should already have been loaded asynchronously in PrepareData, so this is free.
+	USkinnedAsset* LoadedMesh = InstanceList.Descriptor.SkinnedAsset.LoadSynchronous();
+
+	if (!LoadedMesh)
+	{
+		// Either we have no mesh (so nothing to do) or the mesh couldn't be loaded
+		if (InstanceList.Descriptor.SkinnedAsset.IsValid())
+		{
+			PCGE_LOG(Error, GraphAndLog, FText::Format(LOCTEXT("MeshLoadFailed", "Unable to load mesh '{0}'"), FText::FromString(InstanceList.Descriptor.SkinnedAsset.ToString())));
+		}
+
+		return;
+	}
+
+	// Don't spawn meshes if we reuse the ISMCs, but we still want to be sure that the mesh is loaded at least (for operations downstream).
+	if (Context->bSkippedDueToReuse)
+	{
+		return;
+	}
+
+#if 0 // AB-TODO
+	for (TSoftObjectPtr<UMaterialInterface> OverrideMaterial : InstanceList.Descriptor.OverrideMaterials)
+	{
+		// Will be synchronously loaded if not loaded. But by default it should already have been loaded asynchronously in PrepareData, so this is free.
+		if (OverrideMaterial.IsValid() && !OverrideMaterial.LoadSynchronous())
+		{
+			PCGE_LOG(Error, GraphAndLog, FText::Format(LOCTEXT("OverrideMaterialLoadFailed", "Unable to load override material '{0}'"), FText::FromString(OverrideMaterial.ToString())));
+			return;
+		}
+	}
+#endif
+
+	// If we spawn the meshes, we should have computed a packed custom data.
+	if (!ensure(InPackedCustomData))
+	{
+		return;
+	}
+
+	const FPCGAnimBankMeshPackedCustomData& PackedCustomData = *InPackedCustomData;
+
+	FPCGAnimBankMeshComponentBuilderParams Params;
+	Params.Descriptor = InstanceList.Descriptor;
+	Params.NumCustomDataFloats = PackedCustomData.NumCustomDataFloats;
+
+	// If the root actor we're binding to is movable, then the component should be movable by default
+	if (USceneComponent* SceneComponent = TargetActor->GetRootComponent())
+	{
+		Params.Descriptor.Mobility = SceneComponent->Mobility;
+	}
+
+	const UPCGAnimBankMeshSpawnerSettings* Settings = Context->GetInputSettings<UPCGAnimBankMeshSpawnerSettings>();
+	check(Settings);
+
+	Params.SettingsCrc = Settings->GetSettingsCrc();
+	ensure(Params.SettingsCrc.IsValid());
+
+	UPCGManagedAnimBankMeshComponent* MABMC = UPCGActorHelpers::GetOrCreateManagedABMC(TargetActor, Context->SourceComponent.Get(), Params);
+
+	check(MABMC);
+	MABMC->SetCrc(Context->DependenciesCrc);
+
+	// Keep track of all touched resources in the context, because if the execution is cancelled during the SMS execution
+	// we cannot easily guarantee that the state (esp. vs CRCs) is going to be entirely valid
+	Context->TouchedResources.Emplace(MABMC);
+
+	UAnimBankMeshComponent* ABMC = MABMC->GetComponent();
+	check(ABMC);
+
+	const int32 PreExistingInstanceCount = ABMC->GetInstanceCount();
+	const int32 NewInstanceCount = InstanceList.Instances.Num();
+	const int32 NumCustomDataFloats = PackedCustomData.NumCustomDataFloats;
+
+	check((ABMC->NumCustomDataFloats == 0 && PreExistingInstanceCount == 0) || ABMC->NumCustomDataFloats == NumCustomDataFloats);
+	ABMC->SetNumCustomDataFloats(NumCustomDataFloats);
+
+	// The index in ABMC PerInstanceSMCustomData where we should pick up to begin inserting new floats
+	const int32 PreviousCustomDataOffset = PreExistingInstanceCount * NumCustomDataFloats;
+
+	// Populate the instances
+	ABMC->AddInstances(InstanceList.Instances, /*bShouldReturnIndices=*/false, /*bWorldSpace=*/true);
+
+	// Copy new CustomData into the ABMC PerInstanceSMCustomData
+	if (NumCustomDataFloats > 0)
+	{
+		check(PreviousCustomDataOffset + PackedCustomData.CustomData.Num() == ABMC->InstanceCustomData.Num());
+		for (int32 NewIndex = 0; NewIndex < NewInstanceCount; ++NewIndex)
+		{
+			ABMC->SetCustomData(PreExistingInstanceCount + NewIndex, MakeArrayView(&PackedCustomData.CustomData[NewIndex * NumCustomDataFloats], NumCustomDataFloats));
+		}
+	}
+
+	ABMC->UpdateBounds();
+
+	{
+		PCGE_LOG(Verbose, LogOnly, FText::Format(LOCTEXT("GenerationInfo", "Added {0} instances of '{1}' on actor '{2}'"),
+			InstanceList.Instances.Num(), FText::FromString(InstanceList.Descriptor.SkinnedAsset->GetFName().ToString()), FText::FromString(TargetActor->GetFName().ToString())));
+	}
+}
+
+void FPCGAnimBankMeshSpawnerElement::AbortInternal(FPCGContext* InContext) const
+{
+	TRACE_CPUPROFILER_EVENT_SCOPE(FPCGAnimBankMeshSpawnerElement::AbortInternal);
+	// It is possible to Abort a ready task with no context yet
+	if (!InContext)
+	{
+		return;
+	}
+
+	FPCGAnimBankMeshSpawnerContext* Context = static_cast<FPCGAnimBankMeshSpawnerContext*>(InContext);
+
+	// Any resources we've touched during the execution of this node can potentially be in a "not-quite complete state" especially if we have multiple sources of data writing to the same ISMC.
+	// In this case, we're aiming to mark the resources as "Unused" so they are picked up to be removed during the component's OnProcessGraphAborted, which is why we call Release here.
+	for (TWeakObjectPtr<UPCGManagedAnimBankMeshComponent> ManagedResource : Context->TouchedResources)
+	{
+		if(ManagedResource.IsValid())
+		{
+			TSet<TSoftObjectPtr<AActor>> Dummy;
+			ManagedResource->Release(/*bHardRelease=*/false, Dummy);
+		}
+	}
+}
+
+void UPCGAnimBankMeshSpawnerSettings::PostLoad()
+{
+	Super::PostLoad();
+
+	const EObjectFlags Flags = GetMaskedFlags(RF_PropagateToSubObjects) | RF_Transactional;
+	
+	if (!MeshSelectorParameters)
+	{
+		RefreshMeshSelector();
+	}
+	else
+	{
+		MeshSelectorParameters->SetFlags(Flags);
+	}
+
+	if (!InstanceDataPackerParameters)
+	{
+		RefreshInstancePacker();
+	}
+	else
+	{
+		InstanceDataPackerParameters->SetFlags(Flags);
+	}
+}
+
+#if WITH_EDITOR
+void UPCGAnimBankMeshSpawnerSettings::PostEditChangeProperty(FPropertyChangedEvent& PropertyChangedEvent) 
+{
+	if (PropertyChangedEvent.Property)
+	{
+		const FName& PropertyName = PropertyChangedEvent.Property->GetFName();
+
+		if (PropertyName == GET_MEMBER_NAME_CHECKED(UPCGAnimBankMeshSpawnerSettings, MeshSelectorType))
+		{
+			RefreshMeshSelector();
+		} 
+		else if (PropertyName == GET_MEMBER_NAME_CHECKED(UPCGAnimBankMeshSpawnerSettings, InstanceDataPackerType))
+		{
+			RefreshInstancePacker();
+		}
+	}
+
+	Super::PostEditChangeProperty(PropertyChangedEvent);
+}
+
+bool UPCGAnimBankMeshSpawnerSettings::CanEditChange(const FProperty* InProperty) const
+{
+	// TODO: In place temporarily, until the other two modes are supported
+	if (InProperty->GetFName() == GET_MEMBER_NAME_CHECKED(UPCGAnimBankMeshSpawnerSettings, AnimBankMeshComponentPropertyOverrides))
+	{
+		if (!MeshSelectorType->IsChildOf(UPCGAnimBankMeshSelectorByAttribute::StaticClass()))
+		{
+			return false;
+		}
+	}
+
+	return Super::CanEditChange(InProperty);
+}
+#endif
+
+void UPCGAnimBankMeshSpawnerSettings::SetMeshSelectorType(TSubclassOf<UPCGAnimBankMeshSelectorBase> InMeshSelectorType) 
+{
+	if (!MeshSelectorParameters || InMeshSelectorType != MeshSelectorType)
+	{
+		if (InMeshSelectorType != MeshSelectorType)
+		{
+			MeshSelectorType = InMeshSelectorType;
+		}
+		
+		RefreshMeshSelector();
+	}
+}
+
+void UPCGAnimBankMeshSpawnerSettings::SetInstancePackerType(TSubclassOf<UPCGAnimBankMeshInstanceDataPackerBase> InInstancePackerType) 
+{
+	if (!InstanceDataPackerParameters || InInstancePackerType != InstanceDataPackerType)
+	{
+		if (InInstancePackerType != InstanceDataPackerType)
+		{
+			InstanceDataPackerType = InInstancePackerType;
+		}
+		
+		RefreshInstancePacker();
+	}
+}
+
+void UPCGAnimBankMeshSpawnerSettings::RefreshMeshSelector()
+{
+	if (MeshSelectorType)
+	{
+		ensure(IsInGameThread());
+
+		if (MeshSelectorParameters)
+		{
+		#if WITH_EDITOR
+			MeshSelectorParameters->Rename(nullptr, GetTransientPackage(), REN_DontCreateRedirectors);
+		#endif
+			MeshSelectorParameters->MarkAsGarbage();
+			MeshSelectorParameters = nullptr;
+		}
+
+		const EObjectFlags Flags = GetMaskedFlags(RF_PropagateToSubObjects);
+		MeshSelectorParameters = NewObject<UPCGAnimBankMeshSelectorBase>(this, MeshSelectorType, NAME_None, Flags);
+	}
+	else
+	{
+		MeshSelectorParameters = nullptr;
+	}
+}
+
+void UPCGAnimBankMeshSpawnerSettings::RefreshInstancePacker()
+{
+	if (InstanceDataPackerType)
+	{
+		ensure(IsInGameThread());
+
+		if (InstanceDataPackerParameters)
+		{
+		#if WITH_EDITOR
+			InstanceDataPackerParameters->Rename(nullptr, GetTransientPackage(), REN_DontCreateRedirectors);
+		#endif
+			InstanceDataPackerParameters->MarkAsGarbage();
+			InstanceDataPackerParameters = nullptr;
+		}
+
+		const EObjectFlags Flags = GetMaskedFlags(RF_PropagateToSubObjects);
+		InstanceDataPackerParameters = NewObject<UPCGAnimBankMeshInstanceDataPackerBase>(this, InstanceDataPackerType, NAME_None, Flags);
+	}
+	else
+	{
+		InstanceDataPackerParameters = nullptr;
+	}
+}
+
+FPCGAnimBankMeshSpawnerContext::FPackedInstanceListData::FPackedInstanceListData() = default;
+FPCGAnimBankMeshSpawnerContext::FPackedInstanceListData::~FPackedInstanceListData() = default;
+
+#undef LOCTEXT_NAMESPACE
diff --git a/Engine/Plugins/PCG/Source/PCG/Private/Elements/PCGAnimBankMeshSpawnerContext.cpp b/Engine/Plugins/PCG/Source/PCG/Private/Elements/PCGAnimBankMeshSpawnerContext.cpp
new file mode 100644
--- /dev/null
+++ b/Engine/Plugins/PCG/Source/PCG/Private/Elements/PCGAnimBankMeshSpawnerContext.cpp
@@ -0,0 +1,20 @@
+// Copyright Epic Games, Inc. All Rights Reserved.
+
+#include "Elements/PCGAnimBankMeshSpawnerContext.h"
+
+void FPCGAnimBankMeshSpawnerContext::ResetInputIterationData()
+{
+	bCurrentInputSetup = false;
+	bSelectionDone = false;
+	bPartitionDone = false;
+	CurrentPointData = nullptr;
+	CurrentOutputPointData = nullptr;
+	MaterialOverrideHelper.Reset();
+	CurrentPointIndex = 0;
+	WeightedMeshInstances.Reset();
+	MeshToValueKey.Reset();
+	CumulativeWeights.Reset();
+	CategoryEntryToInstancesAndWeights.Reset();
+	AttributeOverridePartition.Reset();
+	OverriddenDescriptors.Reset();
+}
\ No newline at end of file
diff --git a/Engine/Plugins/PCG/Source/PCG/Private/Helpers/PCGActorHelpers.cpp b/Engine/Plugins/PCG/Source/PCG/Private/Helpers/PCGActorHelpers.cpp
--- a/Engine/Plugins/PCG/Source/PCG/Private/Helpers/PCGActorHelpers.cpp
+++ b/Engine/Plugins/PCG/Source/PCG/Private/Helpers/PCGActorHelpers.cpp
@@ -11,7 +11,10 @@
 #include "EngineUtils.h"
 #include "Components/HierarchicalInstancedStaticMeshComponent.h"
 #include "Components/SplineMeshComponent.h"
+#include "Components/AnimBankMeshComponent.h"
 #include "Engine/StaticMesh.h"
+#include "Engine/SkinnedAsset.h"
+#include "Animation/AnimBank.h"
 #include "Materials/MaterialInterface.h"
 
 #include UE_INLINE_GENERATED_CPP_BY_NAME(PCGActorHelpers)
@@ -230,6 +233,126 @@ UPCGManagedISMComponent* UPCGActorHelpers::GetOrCreateManagedISMC(AActor* InTarg
 	return Resource;
 }
 
+UPCGManagedAnimBankMeshComponent* UPCGActorHelpers::GetOrCreateManagedABMC(AActor* InTargetActor, UPCGComponent* InSourceComponent, const FPCGAnimBankMeshComponentBuilderParams& InParams)
+{
+	check(InTargetActor && InSourceComponent);
+
+	TRACE_CPUPROFILER_EVENT_SCOPE(UPCGActorHelpers::GetOrCreateManagedABMC);
+
+	FAnimBankMeshComponentDescriptor Descriptor(InParams.Descriptor);
+
+	const USkinnedAsset* SkinnedAsset = Descriptor.SkinnedAsset;
+	if (!SkinnedAsset)
+	{
+		return nullptr;
+	}
+
+	auto AddTagsToComponent = [InSourceComponent, &InParams](UAnimBankMeshComponent* ABMC)
+	{
+		ABMC->ComponentTags.AddUnique(PCGHelpers::DefaultPCGTag);
+		ABMC->ComponentTags.AddUnique(InSourceComponent->GetFName());
+			
+		for (FName ComponentTag : InParams.Descriptor.ComponentTags)
+		{
+			ABMC->ComponentTags.AddUnique(ComponentTag);
+		}
+	};
+
+	if (!Descriptor.ComponentClass)
+	{
+		Descriptor.ComponentClass = UAnimBankMeshComponent::StaticClass();
+	}
+
+	{
+		TRACE_CPUPROFILER_EVENT_SCOPE(UPCGActorHelpers::GetOrCreateManagedABMC::FindMatchingMABMC);
+		UPCGManagedAnimBankMeshComponent* MatchingResource = nullptr;
+		if(InParams.SettingsCrc.IsValid())
+		{
+			InSourceComponent->ForEachManagedResource([&MatchingResource, &InParams, &InTargetActor, &Descriptor](UPCGManagedResource* InResource)
+			{
+				// Early out if already found a match
+				if (MatchingResource)
+				{
+					return;
+				}
+
+				if (UPCGManagedAnimBankMeshComponent* Resource = Cast<UPCGManagedAnimBankMeshComponent>(InResource))
+				{
+					if (!Resource->GetSettingsCrc().IsValid() || Resource->GetSettingsCrc() != InParams.SettingsCrc || !Resource->CanBeUsed())
+					{
+						return;
+					}
+
+					if (UAnimBankMeshComponent* ABMC = Resource->GetComponent())
+					{
+						if (IsValid(ABMC) &&
+							ABMC->GetOwner() == InTargetActor &&
+							ABMC->NumCustomDataFloats == InParams.NumCustomDataFloats &&
+							Resource->GetDescriptor() == Descriptor &&
+							ABMC->HasAnyFlags(RF_Transient) == InParams.bTransient)
+						{
+							MatchingResource = Resource;
+						}
+					}
+				}
+			});
+		}
+
+		if (MatchingResource)
+		{
+			TRACE_CPUPROFILER_EVENT_SCOPE(UPCGActorHelpers::GetOrCreateManagedABMC::MarkAsUsed);
+			MatchingResource->MarkAsUsed();
+
+			UAnimBankMeshComponent* ABMC = Cast<UAnimBankMeshComponent>(MatchingResource->GeneratedComponent.Get());
+			if (ensure(ABMC))
+			{
+				ABMC->Modify(!InSourceComponent->IsInPreviewMode());
+				AddTagsToComponent(ABMC);
+			}
+
+			return MatchingResource;
+		}
+	}
+
+	// No matching ABM component found, let's create a new one
+	InTargetActor->Modify(!InSourceComponent->IsInPreviewMode());
+
+	FString ComponentName = TEXT("ABM_") + SkinnedAsset->GetName();
+
+	EObjectFlags ObjectFlags = (InSourceComponent->IsInPreviewMode() ? RF_Transient | RF_NonPIEDuplicateTransient : RF_NoFlags);
+
+	// Used for debug data visualization
+	if (InParams.bTransient)
+	{
+		ObjectFlags |= RF_Transient;
+	}
+
+	UAnimBankMeshComponent* ABMC = NewObject<UAnimBankMeshComponent>(InTargetActor, Descriptor.ComponentClass, MakeUniqueObjectName(InTargetActor, Descriptor.ComponentClass, FName(ComponentName)), ObjectFlags);
+	Descriptor.InitComponent(ABMC);
+	ABMC->SetNumCustomDataFloats(InParams.NumCustomDataFloats);
+
+	ABMC->RegisterComponent();
+	InTargetActor->AddInstanceComponent(ABMC);
+
+	ABMC->AttachToComponent(InTargetActor->GetRootComponent(), FAttachmentTransformRules(EAttachmentRule::KeepRelative, EAttachmentRule::KeepWorld, EAttachmentRule::KeepWorld, false));
+	// Implementation note: Because we've used the FAnimBankMeshComponentDescriptor here (Descriptor vs InParams.Descriptor) which takes care of the loading, we still need to apply tags manually.
+	AddTagsToComponent(ABMC);
+
+	// Create managed resource on source component
+	UPCGManagedAnimBankMeshComponent* Resource = NewObject<UPCGManagedAnimBankMeshComponent>(InSourceComponent, NAME_None, InParams.bTransient ? RF_Transient : RF_NoFlags);
+	Resource->SetComponent(ABMC);
+	Resource->SetDescriptor(Descriptor);
+	if (InTargetActor->GetRootComponent())
+	{
+		Resource->SetRootLocation(InTargetActor->GetRootComponent()->GetComponentLocation());
+	}
+	
+	Resource->SetSettingsCrc(InParams.SettingsCrc);
+	InSourceComponent->AddToManagedResources(Resource);
+
+	return Resource;
+}
+
 // deprecated
 USplineMeshComponent* UPCGActorHelpers::GetOrCreateSplineMeshComponent(AActor* InTargetActor, UPCGComponent* InSourceComponent, uint64 SettingsUID, const FPCGSplineMeshComponentBuilderParameters& InParams)
 {
diff --git a/Engine/Plugins/PCG/Source/PCG/Private/InstanceDataPackers/PCGAnimBankMeshInstanceDataPackerBase.cpp b/Engine/Plugins/PCG/Source/PCG/Private/InstanceDataPackers/PCGAnimBankMeshInstanceDataPackerBase.cpp
new file mode 100644
--- /dev/null
+++ b/Engine/Plugins/PCG/Source/PCG/Private/InstanceDataPackers/PCGAnimBankMeshInstanceDataPackerBase.cpp
@@ -0,0 +1,325 @@
+// Copyright Epic Games, Inc. All Rights Reserved.
+
+#include "InstanceDataPackers/PCGAnimBankMeshInstanceDataPackerBase.h"
+
+#include "PCGElement.h"
+#include "Elements/Metadata/PCGMetadataElementCommon.h"
+#include "MeshSelectors/PCGAnimBankMeshSelectorBase.h"
+#include "Metadata/PCGMetadata.h"
+#include "Metadata/Accessors/IPCGAttributeAccessor.h"
+#include "Metadata/Accessors/PCGAttributeAccessorKeys.h"
+
+#include UE_INLINE_GENERATED_CPP_BY_NAME(PCGAnimBankMeshInstanceDataPackerBase)
+
+void UPCGAnimBankMeshInstanceDataPackerBase::PackInstances_Implementation(FPCGContext& Context, const UPCGSpatialData* InSpatialData, const FPCGAnimBankMeshInstanceList& InstanceList, FPCGAnimBankMeshPackedCustomData& OutPackedCustomData) const
+{
+	PCGE_LOG_C(Error, GraphAndLog, &Context, NSLOCTEXT("PCGAnimBankMeshInstanceDataPackerBase", "AnimBankMeshInstanceDataPackerBaseFailed", "Unable to execute AnimBankMeshInstanceDataPacker pure virtual base function, override the PackInstances function or use a default implementation."));
+}
+
+bool UPCGAnimBankMeshInstanceDataPackerBase::AddTypeToPacking(int TypeId, FPCGAnimBankMeshPackedCustomData& OutPackedCustomData) const
+{
+	switch (TypeId)
+	{
+	case PCG::Private::MetadataTypes<bool>::Id: // fall-through
+	case PCG::Private::MetadataTypes<float>::Id: // fall-through
+	case PCG::Private::MetadataTypes<double>::Id: // fall-through
+	case PCG::Private::MetadataTypes<int32>::Id: // fall-through
+	case PCG::Private::MetadataTypes<int64>::Id:
+		OutPackedCustomData.NumCustomDataFloats += 1;
+		break;
+	case PCG::Private::MetadataTypes<FVector2D>::Id:
+		OutPackedCustomData.NumCustomDataFloats += 2;
+		break;
+	case PCG::Private::MetadataTypes<FRotator>::Id: // fall-through
+	case PCG::Private::MetadataTypes<FVector>::Id:
+		OutPackedCustomData.NumCustomDataFloats += 3;
+		break;
+	case PCG::Private::MetadataTypes<FVector4>::Id: // fall-through
+	case PCG::Private::MetadataTypes<FQuat>::Id:
+		OutPackedCustomData.NumCustomDataFloats += 4;
+		break;
+	default:
+		return false;
+	}
+
+	return true;
+}
+
+void UPCGAnimBankMeshInstanceDataPackerBase::PackCustomDataFromAttributes(const FPCGAnimBankMeshInstanceList& InstanceList, const UPCGMetadata* Metadata, const TArray<FName>& AttributeNames, FPCGAnimBankMeshPackedCustomData& OutPackedCustomData) const
+{
+	if (!Metadata)
+	{
+		UE_LOG(LogPCG, Error, TEXT("Invalid metadata"));
+		return;
+	}
+
+	TArray<const FPCGMetadataAttributeBase*> Attributes;
+	for (const FName& AttributeName : AttributeNames)
+	{
+		Attributes.Add(Metadata->GetConstAttribute(AttributeName));
+	}
+
+	PackCustomDataFromAttributes(InstanceList, Attributes, OutPackedCustomData);
+}
+
+void UPCGAnimBankMeshInstanceDataPackerBase::PackCustomDataFromAttributes(const FPCGAnimBankMeshInstanceList& InstanceList, const TArray<const FPCGMetadataAttributeBase*>& Attributes, FPCGAnimBankMeshPackedCustomData& OutPackedCustomData) const
+{
+	const UPCGPointData* PointData = InstanceList.PointData.Get();
+	for (const int32 Index : InstanceList.InstancesIndices)
+	{
+		const PCGMetadataEntryKey EntryKey = PointData ? PointData->GetPoints()[Index].MetadataEntry : Index;
+		for (const FPCGMetadataAttributeBase* AttributeBase : Attributes)
+		{
+			check(AttributeBase);
+
+			switch (AttributeBase->GetTypeId())
+			{
+			case PCG::Private::MetadataTypes<bool>::Id:
+			{
+				const FPCGMetadataAttribute<bool>* Attribute = static_cast<const FPCGMetadataAttribute<bool>*>(AttributeBase);
+				check(Attribute);
+
+				const bool Value = Attribute->GetValueFromItemKey(EntryKey);
+				OutPackedCustomData.CustomData.Add(Value);
+				break;
+			}
+			case PCG::Private::MetadataTypes<float>::Id:
+			{
+				const FPCGMetadataAttribute<float>* Attribute = static_cast<const FPCGMetadataAttribute<float>*>(AttributeBase);
+				check(Attribute);
+
+				const float Value = Attribute->GetValueFromItemKey(EntryKey);
+				OutPackedCustomData.CustomData.Add(Value);
+				break;
+			}
+			case PCG::Private::MetadataTypes<double>::Id:
+			{
+				const FPCGMetadataAttribute<double>* Attribute = static_cast<const FPCGMetadataAttribute<double>*>(AttributeBase);
+				check(Attribute);
+
+				const double Value = Attribute->GetValueFromItemKey(EntryKey);
+				OutPackedCustomData.CustomData.Add(Value);
+				break;
+			}
+			case PCG::Private::MetadataTypes<int32>::Id:
+			{
+				const FPCGMetadataAttribute<int32>* Attribute = static_cast<const FPCGMetadataAttribute<int32>*>(AttributeBase);
+				check(Attribute);
+
+				const float Value = static_cast<float>(Attribute->GetValueFromItemKey(EntryKey));
+				OutPackedCustomData.CustomData.Add(Value);
+				break;
+			}
+			case PCG::Private::MetadataTypes<int64>::Id:
+			{
+				const FPCGMetadataAttribute<int64>* Attribute = static_cast<const FPCGMetadataAttribute<int64>*>(AttributeBase);
+				check(Attribute);
+
+				const float Value = static_cast<float>(Attribute->GetValueFromItemKey(EntryKey));
+				OutPackedCustomData.CustomData.Add(Value);
+				break;
+			}
+			case PCG::Private::MetadataTypes<FRotator>::Id:
+			{
+				const FPCGMetadataAttribute<FRotator>* Attribute = static_cast<const FPCGMetadataAttribute<FRotator>*>(AttributeBase);
+				check(Attribute);
+
+				const FRotator Value = Attribute->GetValueFromItemKey(EntryKey);
+				OutPackedCustomData.CustomData.Add(Value.Roll);
+				OutPackedCustomData.CustomData.Add(Value.Pitch);
+				OutPackedCustomData.CustomData.Add(Value.Yaw);
+				break;
+			}
+			case PCG::Private::MetadataTypes<FVector2D>::Id:
+			{
+				const FPCGMetadataAttribute<FVector2D>* Attribute = static_cast<const FPCGMetadataAttribute<FVector2D>*>(AttributeBase);
+				check(Attribute);
+
+				const FVector2D Value = Attribute->GetValueFromItemKey(EntryKey);
+				OutPackedCustomData.CustomData.Add(Value.X);
+				OutPackedCustomData.CustomData.Add(Value.Y);
+				break;
+			}
+			case PCG::Private::MetadataTypes<FVector>::Id:
+			{
+				const FPCGMetadataAttribute<FVector>* Attribute = static_cast<const FPCGMetadataAttribute<FVector>*>(AttributeBase);
+				check(Attribute);
+
+				const FVector Value = Attribute->GetValueFromItemKey(EntryKey);
+				OutPackedCustomData.CustomData.Add(Value.X);
+				OutPackedCustomData.CustomData.Add(Value.Y);
+				OutPackedCustomData.CustomData.Add(Value.Z);
+				break;
+			}
+			case PCG::Private::MetadataTypes<FVector4>::Id:
+			{
+				const FPCGMetadataAttribute<FVector4>* Attribute = static_cast<const FPCGMetadataAttribute<FVector4>*>(AttributeBase);
+				check(Attribute);
+
+				const FVector4 Value = Attribute->GetValueFromItemKey(EntryKey);
+				OutPackedCustomData.CustomData.Add(Value.X);
+				OutPackedCustomData.CustomData.Add(Value.Y);
+				OutPackedCustomData.CustomData.Add(Value.Z);
+				OutPackedCustomData.CustomData.Add(Value.W);
+				break;
+			}
+			case PCG::Private::MetadataTypes<FQuat>::Id:
+			{
+				const FPCGMetadataAttribute<FQuat>* Attribute = static_cast<const FPCGMetadataAttribute<FQuat>*>(AttributeBase);
+				check(Attribute);
+
+				const FQuat Value = Attribute->GetValueFromItemKey(EntryKey);
+				OutPackedCustomData.CustomData.Add(Value.X);
+				OutPackedCustomData.CustomData.Add(Value.Y);
+				OutPackedCustomData.CustomData.Add(Value.Z);
+				OutPackedCustomData.CustomData.Add(Value.W);
+				break;
+			}
+			default:
+				break;
+			}
+		}
+	}
+}
+
+void UPCGAnimBankMeshInstanceDataPackerBase::PackCustomDataFromAccessors(const FPCGAnimBankMeshInstanceList& InstanceList, TArray<TUniquePtr<const IPCGAttributeAccessor>> Accessors, TArray<TUniquePtr<const IPCGAttributeAccessorKeys>> AccessorKeys, FPCGAnimBankMeshPackedCustomData& OutPackedCustomData) const
+{
+	const int TotalNumCustomDataFloats = OutPackedCustomData.NumCustomDataFloats * InstanceList.Instances.Num();
+
+	if (!ensure(Accessors.Num() == AccessorKeys.Num()))
+	{
+		// Zero out the custom data to avoid sending garbage data to the GPU.
+		OutPackedCustomData.CustomData.SetNumZeroed(TotalNumCustomDataFloats);
+		return;
+	}
+
+	OutPackedCustomData.CustomData.SetNumUninitialized(TotalNumCustomDataFloats);
+
+	if (TotalNumCustomDataFloats == 0)
+	{
+		return;
+	}
+
+	// Index of the accessor's attribute/property for a single instance (e.g. if the
+	// accessors are [float, FVector, float], AccessorPackedIndex will be [0, 1, 4])
+	int AccessorPackedIndex = 0;
+	const EPCGAttributeAccessorFlags Flags = EPCGAttributeAccessorFlags::AllowBroadcast | EPCGAttributeAccessorFlags::AllowConstructible;
+
+	for (int AccessorIndex = 0; AccessorIndex < Accessors.Num(); ++AccessorIndex)
+	{
+		const IPCGAttributeAccessor* Accessor = Accessors[AccessorIndex].Get();
+		const IPCGAttributeAccessorKeys* Keys = AccessorKeys[AccessorIndex].Get();
+
+		check(Accessor && Keys);
+
+		switch (Accessor->GetUnderlyingType())
+		{
+		case PCG::Private::MetadataTypes<bool>::Id: // fall-through
+		case PCG::Private::MetadataTypes<float>::Id: // fall-through
+		case PCG::Private::MetadataTypes<double>::Id: // fall-through
+		case PCG::Private::MetadataTypes<int32>::Id: // fall-through
+		case PCG::Private::MetadataTypes<int64>::Id:
+		{
+			auto PackFloats = [&OutPackedCustomData, AccessorPackedIndex](float Value, int32 Index)
+			{
+				const int PackedIndex = Index * OutPackedCustomData.NumCustomDataFloats + AccessorPackedIndex;
+				OutPackedCustomData.CustomData[PackedIndex] = Value;
+			};
+
+			ensure(PCGMetadataElementCommon::ApplyOnAccessor<float>(*Keys, *Accessor, PackFloats, Flags));
+
+			AccessorPackedIndex += 1;
+
+			break;
+		}
+		case PCG::Private::MetadataTypes<FRotator>::Id:
+		{
+			auto PackFloats = [&OutPackedCustomData, AccessorPackedIndex](const FRotator& Value, int32 Index)
+			{
+				const int PackedIndex = Index * OutPackedCustomData.NumCustomDataFloats + AccessorPackedIndex;
+
+				OutPackedCustomData.CustomData[PackedIndex + 0] = static_cast<float>(Value.Roll);
+				OutPackedCustomData.CustomData[PackedIndex + 1] = static_cast<float>(Value.Pitch);
+				OutPackedCustomData.CustomData[PackedIndex + 2] = static_cast<float>(Value.Yaw);
+			};
+
+			ensure(PCGMetadataElementCommon::ApplyOnAccessor<FRotator>(*Keys, *Accessor, PackFloats, Flags));
+
+			AccessorPackedIndex += 3;
+
+			break;
+		}
+		case PCG::Private::MetadataTypes<FVector2D>::Id:
+		{
+			auto PackFloats = [&OutPackedCustomData, AccessorPackedIndex](const FVector2D& Value, int32 Index)
+			{
+				const int PackedIndex = Index * OutPackedCustomData.NumCustomDataFloats + AccessorPackedIndex;
+
+				OutPackedCustomData.CustomData[PackedIndex + 0] = static_cast<float>(Value.X);
+				OutPackedCustomData.CustomData[PackedIndex + 1] = static_cast<float>(Value.Y);
+			};
+
+			ensure(PCGMetadataElementCommon::ApplyOnAccessor<FVector2D>(*Keys, *Accessor, PackFloats, Flags));
+
+			AccessorPackedIndex += 2;
+
+			break;
+		}
+		case PCG::Private::MetadataTypes<FVector>::Id:
+		{
+			auto PackFloats = [&OutPackedCustomData, AccessorPackedIndex](const FVector& Value, int32 Index)
+			{
+				const int PackedIndex = Index * OutPackedCustomData.NumCustomDataFloats + AccessorPackedIndex;
+
+				OutPackedCustomData.CustomData[PackedIndex + 0] = static_cast<float>(Value.X);
+				OutPackedCustomData.CustomData[PackedIndex + 1] = static_cast<float>(Value.Y);
+				OutPackedCustomData.CustomData[PackedIndex + 2] = static_cast<float>(Value.Z);
+			};
+
+			ensure(PCGMetadataElementCommon::ApplyOnAccessor<FVector>(*Keys, *Accessor, PackFloats, Flags));
+
+			AccessorPackedIndex += 3;
+
+			break;
+		}
+		case PCG::Private::MetadataTypes<FVector4>::Id:
+		{
+			auto PackFloats = [&OutPackedCustomData, AccessorPackedIndex](const FVector4& Value, int32 Index)
+			{
+				const int PackedIndex = Index * OutPackedCustomData.NumCustomDataFloats + AccessorPackedIndex;
+
+				OutPackedCustomData.CustomData[PackedIndex + 0] = static_cast<float>(Value.X);
+				OutPackedCustomData.CustomData[PackedIndex + 1] = static_cast<float>(Value.Y);
+				OutPackedCustomData.CustomData[PackedIndex + 2] = static_cast<float>(Value.Z);
+				OutPackedCustomData.CustomData[PackedIndex + 3] = static_cast<float>(Value.W);
+			};
+
+			ensure(PCGMetadataElementCommon::ApplyOnAccessor<FVector4>(*Keys, *Accessor, PackFloats, Flags));
+
+			AccessorPackedIndex += 4;
+
+			break;
+		}
+		case PCG::Private::MetadataTypes<FQuat>::Id:
+		{
+			auto PackFloats = [&OutPackedCustomData, AccessorPackedIndex](const FQuat& Value, int32 Index)
+			{
+				const int PackedIndex = Index * OutPackedCustomData.NumCustomDataFloats + AccessorPackedIndex;
+
+				OutPackedCustomData.CustomData[PackedIndex + 0] = static_cast<float>(Value.X);
+				OutPackedCustomData.CustomData[PackedIndex + 1] = static_cast<float>(Value.Y);
+				OutPackedCustomData.CustomData[PackedIndex + 2] = static_cast<float>(Value.Z);
+				OutPackedCustomData.CustomData[PackedIndex + 3] = static_cast<float>(Value.W);
+			};
+
+			ensure(PCGMetadataElementCommon::ApplyOnAccessor<FQuat>(*Keys, *Accessor, PackFloats, Flags));
+
+			AccessorPackedIndex += 4;
+
+			break;
+		}
+		default:
+			break;
+		}
+	}
+}
diff --git a/Engine/Plugins/PCG/Source/PCG/Private/InstanceDataPackers/PCGAnimBankMeshInstanceDataPackerByAttribute.cpp b/Engine/Plugins/PCG/Source/PCG/Private/InstanceDataPackers/PCGAnimBankMeshInstanceDataPackerByAttribute.cpp
new file mode 100644
--- /dev/null
+++ b/Engine/Plugins/PCG/Source/PCG/Private/InstanceDataPackers/PCGAnimBankMeshInstanceDataPackerByAttribute.cpp
@@ -0,0 +1,89 @@
+// Copyright Epic Games, Inc. All Rights Reserved.
+
+#include "InstanceDataPackers/PCGAnimBankMeshInstanceDataPackerByAttribute.h"
+
+#include "PCGContext.h"
+#include "PCGElement.h"
+#include "Data/PCGSpatialData.h"
+#include "InstanceDataPackers/PCGAnimBankMeshInstanceDataPackerBase.h"
+#include "Metadata/Accessors/PCGAttributeAccessorHelpers.h"
+#include "MeshSelectors/PCGAnimBankMeshSelectorBase.h"
+
+#include UE_INLINE_GENERATED_CPP_BY_NAME(PCGAnimBankMeshInstanceDataPackerByAttribute)
+
+#define LOCTEXT_NAMESPACE "PCGAnimBankMeshInstanceDataPackerByAttribute"
+
+#if WITH_EDITOR
+void UPCGAnimBankMeshInstanceDataPackerByAttribute::PostLoad()
+{
+	Super::PostLoad();
+}
+#endif // WITH_EDITOR
+
+void UPCGAnimBankMeshInstanceDataPackerByAttribute::PackInstances_Implementation(FPCGContext& Context, const UPCGSpatialData* InSpatialData, const FPCGAnimBankMeshInstanceList& InstanceList, FPCGAnimBankMeshPackedCustomData& OutPackedCustomData) const
+{
+	if (!InSpatialData || !InSpatialData->Metadata)
+	{
+		PCGLog::InputOutput::LogInvalidInputDataError(&Context);
+		return;
+	}
+
+	TArray<TUniquePtr<const IPCGAttributeAccessor>> SelectedAccessors;
+	TArray<TUniquePtr<const IPCGAttributeAccessorKeys>> SelectedKeys;
+
+	SelectedAccessors.Reserve(AttributeSelectors.Num());
+	SelectedKeys.Reserve(AttributeSelectors.Num());
+
+	// Find attributes and calculate NumCustomDataFloats
+	for (const FPCGAttributePropertyInputSelector& Selector : AttributeSelectors)
+	{
+		TUniquePtr<const IPCGAttributeAccessor> Accessor = PCGAttributeAccessorHelpers::CreateConstAccessor(InSpatialData, Selector);
+		TUniquePtr<const IPCGAttributeAccessorKeys> Keys;
+		const UPCGPointData* PointData = InstanceList.PointData.Get();
+		TArray<const PCGMetadataEntryKey> ExtractedKeys;
+		
+		if (InSpatialData == PointData)
+		{
+			Keys = MakeUnique<const FPCGAttributeAccessorKeysPointsSubset>(PointData->GetPoints(), InstanceList.InstancesIndices);
+		}
+		else
+		{
+			// Convert indices to entry keys
+			ExtractedKeys.Reserve(InstanceList.InstancesIndices.Num());
+			Algo::Transform(InstanceList.InstancesIndices, ExtractedKeys, [](const int32 Index) -> PCGMetadataEntryKey{ return Index; });
+			Keys = MakeUnique<const FPCGAttributeAccessorKeysEntries>(ExtractedKeys);
+		}
+
+		if (!Accessor.IsValid() || !Keys.IsValid())
+		{
+			PCGLog::Metadata::LogFailToCreateAccessorError(Selector, &Context);
+			continue;
+		}
+
+		if (!AddTypeToPacking(Accessor->GetUnderlyingType(), OutPackedCustomData))
+		{
+			PCGLog::LogWarningOnGraph(FText::Format(LOCTEXT("AttributeInvalidType", "Attribute/property '{0}' is not a valid type - skipped."), Selector.GetDisplayText()), &Context);
+			continue;
+		}
+
+		SelectedAccessors.Add(std::move(Accessor));
+		SelectedKeys.Add(std::move(Keys));
+	}
+
+	PackCustomDataFromAccessors(InstanceList, std::move(SelectedAccessors), std::move(SelectedKeys), OutPackedCustomData);
+}
+
+bool UPCGAnimBankMeshInstanceDataPackerByAttribute::GetAttributeNames(TArray<FName>* OutNames)
+{
+	if (OutNames)
+	{
+		for (const FPCGAttributePropertyInputSelector& AttributeSelector : AttributeSelectors)
+		{
+			OutNames->Add(AttributeSelector.GetAttributeName());
+		}
+	}
+
+	return true;
+}
+
+#undef LOCTEXT_NAMESPACE
diff --git a/Engine/Plugins/PCG/Source/PCG/Private/InstanceDataPackers/PCGAnimBankMeshInstanceDataPackerByRegex.cpp b/Engine/Plugins/PCG/Source/PCG/Private/InstanceDataPackers/PCGAnimBankMeshInstanceDataPackerByRegex.cpp
new file mode 100644
--- /dev/null
+++ b/Engine/Plugins/PCG/Source/PCG/Private/InstanceDataPackers/PCGAnimBankMeshInstanceDataPackerByRegex.cpp
@@ -0,0 +1,70 @@
+// Copyright Epic Games, Inc. All Rights Reserved.
+
+#include "InstanceDataPackers/PCGAnimBankMeshInstanceDataPackerByRegex.h"
+
+#include "PCGContext.h"
+#include "PCGElement.h"
+#include "Data/PCGSpatialData.h"
+#include "InstanceDataPackers/PCGAnimBankMeshInstanceDataPackerBase.h"
+
+#include "Internationalization/Regex.h"
+
+#include UE_INLINE_GENERATED_CPP_BY_NAME(PCGAnimBankMeshInstanceDataPackerByRegex)
+
+#define LOCTEXT_NAMESPACE "PCGAnimBankMeshInstanceDataPackerByRegex"
+
+void UPCGAnimBankMeshInstanceDataPackerByRegex::PackInstances_Implementation(FPCGContext& Context, const UPCGSpatialData* InSpatialData, const FPCGAnimBankMeshInstanceList& InstanceList, FPCGAnimBankMeshPackedCustomData& OutPackedCustomData) const
+{
+	if (!InSpatialData || !InSpatialData->Metadata)
+	{
+		PCGE_LOG_C(Error, GraphAndLog, &Context, LOCTEXT("InvalidInputData", "Invalid input data"));
+		return;
+	}
+	
+	TArray<FName> AttributeNames;
+	TArray<EPCGMetadataTypes> AttributeTypes;
+	InSpatialData->Metadata->GetAttributes(AttributeNames, AttributeTypes);
+
+	TArray<const FPCGMetadataAttributeBase*> SelectedAttributes;
+	TSet<PCGMetadataAttributeKey> FoundAttributes;
+
+	for (const FString& Regex : RegexPatterns)
+	{
+		const FRegexPattern Pattern(Regex);
+	
+		// Match our Regex Pattern against every Attribute
+		for (int Index = 0; Index < AttributeNames.Num(); ++Index)
+		{
+			const FName& AttributeName = AttributeNames[Index];
+			FRegexMatcher RegexMatcher(Pattern, AttributeName.ToString());
+	
+			if (!RegexMatcher.FindNext()) 
+			{
+				continue;
+			}
+	
+			const FPCGMetadataAttributeBase* AttributeBase = InSpatialData->Metadata->GetConstAttribute(AttributeName);
+			check(AttributeBase);
+
+			// Avoid adding the same attribute multiple times if it is captured by different regex patterns
+			if (FoundAttributes.Contains(AttributeBase->AttributeId))
+			{
+				continue;
+			}
+
+			FoundAttributes.Add(AttributeBase->AttributeId);
+	
+			if (!AddTypeToPacking(AttributeBase->GetTypeId(), OutPackedCustomData))
+			{
+				PCGE_LOG_C(Warning, GraphAndLog, &Context, FText::Format(LOCTEXT("InvalidAttributeType", "Attribute name '{0}' is not of a valid type"), FText::FromName(AttributeName)));
+				continue;
+			}
+	
+			SelectedAttributes.Add(AttributeBase);
+		}
+	}
+	
+	PackCustomDataFromAttributes(InstanceList, SelectedAttributes, OutPackedCustomData);
+}
+
+#undef LOCTEXT_NAMESPACE
diff --git a/Engine/Plugins/PCG/Source/PCG/Private/MeshSelectors/PCGAnimBankMeshDescriptor.cpp b/Engine/Plugins/PCG/Source/PCG/Private/MeshSelectors/PCGAnimBankMeshDescriptor.cpp
new file mode 100644
--- /dev/null
+++ b/Engine/Plugins/PCG/Source/PCG/Private/MeshSelectors/PCGAnimBankMeshDescriptor.cpp
@@ -0,0 +1,32 @@
+// Copyright Epic Games, Inc. All Rights Reserved.
+
+#include "MeshSelectors/PCGAnimBankMeshDescriptor.h"
+#include "Components/AnimBankMeshComponent.h"
+
+#include UE_INLINE_GENERATED_CPP_BY_NAME(PCGAnimBankMeshDescriptor)
+
+FPCGSoftAnimBankMeshComponentDescriptor::FPCGSoftAnimBankMeshComponentDescriptor()
+: FSoftAnimBankMeshComponentDescriptor()
+{
+	// Override defaults from base class
+	ComponentClass = UAnimBankMeshComponent::StaticClass();
+	Mobility = EComponentMobility::Static;
+	// GW-TODO: bGenerateOverlapEvents = 0;
+	// GW-TODO: BodyInstance.SetCollisionProfileName(UCollisionProfile::NoCollision_ProfileName);
+}
+
+FPCGSoftAnimBankMeshComponentDescriptor::FPCGSoftAnimBankMeshComponentDescriptor(const FAnimBankMeshComponentDescriptor& Other)
+	: Super(Other)
+{}
+
+void FPCGSoftAnimBankMeshComponentDescriptor::InitFrom(const UAnimBankMeshComponent* Component, bool bInitBodyInstance)
+{
+	ComponentTags = Component->ComponentTags;
+	Super::InitFrom(Component, bInitBodyInstance);
+}
+
+void FPCGSoftAnimBankMeshComponentDescriptor::InitComponent(UAnimBankMeshComponent* ABMComponent) const
+{
+	ABMComponent->ComponentTags = ComponentTags;
+	Super::InitComponent(ABMComponent);
+}
diff --git a/Engine/Plugins/PCG/Source/PCG/Private/MeshSelectors/PCGAnimBankMeshSelectorBase.cpp b/Engine/Plugins/PCG/Source/PCG/Private/MeshSelectors/PCGAnimBankMeshSelectorBase.cpp
new file mode 100644
--- /dev/null
+++ b/Engine/Plugins/PCG/Source/PCG/Private/MeshSelectors/PCGAnimBankMeshSelectorBase.cpp
@@ -0,0 +1,14 @@
+// Copyright Epic Games, Inc. All Rights Reserved.
+
+#include "MeshSelectors/PCGAnimBankMeshSelectorBase.h"
+
+#include "PCGElement.h"
+#include "Elements/PCGAnimBankMeshSpawnerContext.h"
+
+#include UE_INLINE_GENERATED_CPP_BY_NAME(PCGAnimBankMeshSelectorBase)
+
+#define LOCTEXT_NAMESPACE "PCGAnimBankMeshSelectorBase"
+
+// 
+
+#undef LOCTEXT_NAMESPACE
diff --git a/Engine/Plugins/PCG/Source/PCG/Private/MeshSelectors/PCGAnimBankMeshSelectorByAttribute.cpp b/Engine/Plugins/PCG/Source/PCG/Private/MeshSelectors/PCGAnimBankMeshSelectorByAttribute.cpp
new file mode 100644
--- /dev/null
+++ b/Engine/Plugins/PCG/Source/PCG/Private/MeshSelectors/PCGAnimBankMeshSelectorByAttribute.cpp
@@ -0,0 +1,326 @@
+// Copyright Epic Games, Inc. All Rights Reserved.
+
+#include "MeshSelectors/PCGAnimBankMeshSelectorByAttribute.h"
+
+#include "Data/PCGPointData.h"
+#include "Data/PCGSpatialData.h"
+#include "Elements/PCGAnimBankMeshSpawner.h"
+#include "Elements/PCGAnimBankMeshSpawnerContext.h"
+#include "Elements/Metadata/PCGMetadataElementCommon.h"
+#include "MeshSelectors/PCGAnimBankMeshSelectorBase.h"
+#include "Metadata/PCGObjectPropertyOverride.h"
+#include "Metadata/PCGMetadataPartitionCommon.h"
+#include "Metadata/Accessors/PCGAttributeAccessorHelpers.h"
+
+#include "Algo/AnyOf.h"
+#include "Animation/AnimBank.h"
+
+#include UE_INLINE_GENERATED_CPP_BY_NAME(PCGAnimBankMeshSelectorByAttribute)
+
+#define LOCTEXT_NAMESPACE "PCGAnimBankMeshSelectorByAttribute"
+
+namespace PCGAnimBankMeshSelectorAttribute
+{
+	// Returns variation based on mesh, material overrides and reverse culling
+	FPCGAnimBankMeshInstanceList& GetInstanceList(
+		TArray<FPCGAnimBankMeshInstanceList>& InstanceLists,
+		const FPCGSoftAnimBankMeshComponentDescriptor& TemplateDescriptor,
+		TSoftObjectPtr<USkinnedAsset> Asset,
+		const TArray<FSoftAnimBankItem>& Items,
+		const TArray<TSoftObjectPtr<UMaterialInterface>>& MaterialOverrides,
+		bool bReverseCulling,
+		const UPCGPointData* InPointData,
+		const int AttributePartitionIndex = INDEX_NONE)
+	{
+		for (FPCGAnimBankMeshInstanceList& InstanceList : InstanceLists)
+		{
+			if (InstanceList.Descriptor.SkinnedAsset == Asset &&
+				InstanceList.Descriptor.BankItems == Items &&
+				// GW-TODO: InstanceList.Descriptor.bReverseCulling == bReverseCulling &&
+				// GW-TODO: InstanceList.Descriptor.OverrideMaterials == MaterialOverrides &&
+				InstanceList.AttributePartitionIndex == AttributePartitionIndex)
+			{
+				return InstanceList;
+			}
+		}
+
+		FPCGAnimBankMeshInstanceList& NewInstanceList = InstanceLists.Emplace_GetRef(TemplateDescriptor);
+		NewInstanceList.Descriptor.SkinnedAsset = Asset;
+		NewInstanceList.Descriptor.BankItems = Items;
+		// GW-TODO: NewInstanceList.Descriptor.OverrideMaterials = MaterialOverrides;
+		// GW-TODO: NewInstanceList.Descriptor.bReverseCulling = bReverseCulling;
+		NewInstanceList.AttributePartitionIndex = AttributePartitionIndex;
+		NewInstanceList.PointData = InPointData;
+
+		return NewInstanceList;
+	}
+}
+
+void UPCGAnimBankMeshSelectorByAttribute::PostLoad()
+{
+	Super::PostLoad();
+
+	// TODO: Remove if/when FBodyInstance is updated or replaced
+	// Necessary to update the collision Response Container from the Response Array
+	TemplateDescriptor.PostLoadFixup(this);
+}
+
+bool UPCGAnimBankMeshSelectorByAttribute::SelectInstances(
+	FPCGAnimBankMeshSpawnerContext& Context,
+	const UPCGAnimBankMeshSpawnerSettings* Settings,
+	const UPCGPointData* InPointData,
+	TArray<FPCGAnimBankMeshInstanceList>& OutMeshInstances,
+	UPCGPointData* OutPointData) const
+{
+	TRACE_CPUPROFILER_EVENT_SCOPE(UPCGMeshSelectorByAttribute::SelectInstances);
+
+	if (!InPointData)
+	{
+		PCGE_LOG_C(Error, GraphAndLog, &Context, LOCTEXT("InputMissingData", "Missing input data"));
+		return true;
+	}
+
+	if (!InPointData->Metadata)
+	{
+		PCGE_LOG_C(Error, GraphAndLog, &Context, LOCTEXT("InputMissingMetadata", "Unable to get metadata from input"));
+		return true;
+	}
+
+	if (!InPointData->Metadata->HasAttribute(AttributeName))
+	{
+		PCGE_LOG_C(Error, GraphAndLog, &Context, FText::Format(LOCTEXT("AttributeNotInMetadata", "Attribute '{0}' is not in the metadata"), FText::FromName(AttributeName)));
+		return true;
+	}
+
+	const FPCGMetadataAttributeBase* AttributeBase = InPointData->Metadata->GetConstAttribute(AttributeName);
+	check(AttributeBase);
+
+	// Validate that the "mesh" attribute is of the right type
+	if (!PCG::Private::IsOfTypes<FSoftObjectPath, FString>(AttributeBase->GetTypeId()))
+	{
+		PCGE_LOG_C(Error, GraphAndLog, &Context, FText::Format(LOCTEXT("AttributeInvalidType", "Attribute '{0}' is not of valid type (must be FString or FSoftObjectPath)"), FText::FromName(AttributeName)));
+		return true;
+	}
+
+	FPCGMeshMaterialOverrideHelper& MaterialOverrideHelper = Context.MaterialOverrideHelper;
+	if (!MaterialOverrideHelper.IsInitialized())
+	{
+		// GW-TODO: 		MaterialOverrideHelper.Initialize(Context, bUseAttributeMaterialOverrides, TemplateDescriptor.OverrideMaterials, MaterialOverrideAttributes, InPointData->Metadata);
+	}
+
+	if (!MaterialOverrideHelper.IsValid())
+	{
+		return true;
+	}
+
+	// ByAttribute takes in SoftObjectPaths per point in the metadata, so we can pass those directly into the outgoing pin if it exists
+	if (Context.CurrentPointIndex == 0 && OutPointData)
+	{
+		TRACE_CPUPROFILER_EVENT_SCOPE(UPCGMeshSelectorByAttribute::SetupOutPointData);
+		OutPointData->SetPoints(InPointData->GetPoints());
+		OutPointData->Metadata->DeleteAttribute(Settings->OutAttributeName);
+		OutPointData->Metadata->CopyAttribute(InPointData->Metadata, AttributeName, Settings->OutAttributeName);
+	}
+
+	TRACE_CPUPROFILER_EVENT_SCOPE(UPCGAnimBankMeshSelectorByAttribute::SelectEntries);
+
+	if (!Context.bPartitionDone)
+	{
+		TRACE_CPUPROFILER_EVENT_SCOPE(UPCGAnimBankMeshSelectorByAttribute::SelectEntries::BuildingPartition);
+
+		TArray<FPCGObjectPropertyOverrideDescription> PropertyOverrides = Settings->AnimBankMeshComponentPropertyOverrides;
+		const FString SkinnedAssetPropertyString = GET_MEMBER_NAME_CHECKED(FPCGSoftAnimBankMeshComponentDescriptor, SkinnedAsset).ToString();
+
+		// Add the skinned asset override to the list only if not already provided
+		if (!Algo::AnyOf(PropertyOverrides, [&SkinnedAssetPropertyString](const FPCGObjectPropertyOverrideDescription& PropertyOverride) { return PropertyOverride.PropertyTarget == SkinnedAssetPropertyString; }))
+		{
+			FPCGAttributePropertyInputSelector MeshSelector;
+			MeshSelector.SetAttributeName(AttributeName);
+			PropertyOverrides.Emplace(MeshSelector, SkinnedAssetPropertyString);
+		}
+
+		// Validate all the selectors are actual FPCGSoftAnimBankMeshComponentDescriptor properties
+		TArray<FPCGAttributePropertySelector> ValidSelectorOverrides;
+
+		for (const FPCGObjectPropertyOverrideDescription& PropertyOverride : PropertyOverrides)
+		{
+			if (FPCGSoftAnimBankMeshComponentDescriptor::StaticStruct()->FindPropertyByName(FName(PropertyOverride.PropertyTarget)))
+			{
+				ValidSelectorOverrides.Emplace_GetRef() = FPCGAttributePropertySelector::CreateFromOtherSelector<FPCGAttributePropertySelector>(PropertyOverride.InputSource);
+			}
+			else
+			{
+				PCGE_LOG_C(Error, GraphAndLog, &Context, FText::Format(LOCTEXT("OverriddenPropertyNotFound", "Property '{0}' not a valid property with an ISM Descriptor. It will be ignored."), FText::FromString(PropertyOverride.PropertyTarget)));
+			}
+		}
+
+		// If there are valid overrides, partition the points on those attributes so that an instance can be created for each
+		if (!ValidSelectorOverrides.IsEmpty())
+		{
+			Context.AttributeOverridePartition = PCGMetadataPartitionCommon::AttributeGenericPartition(InPointData, ValidSelectorOverrides, &Context, Settings->bSilenceOverrideAttributeNotFoundErrors);
+		}
+
+		// Set the descriptors to match the partition count. Uninitialized, because it will be copied from the template below
+		Context.OverriddenDescriptors.Reserve(Context.AttributeOverridePartition.Num());
+		for (int I = 0; I < Context.AttributeOverridePartition.Num(); ++I)
+		{
+			FPCGSoftAnimBankMeshComponentDescriptor& Descriptor = Context.OverriddenDescriptors.Add_GetRef(TemplateDescriptor);
+
+			// If partition is empty (which can happen, esp. for the default partition on the default value, we'll just skip it here.
+			if (Context.AttributeOverridePartition[I].IsEmpty())
+			{
+				continue;
+			}
+
+			// Use the Object Override to map the user's input selector and property to the descriptor
+			FPCGObjectOverrides Overrides(&Descriptor);
+			Overrides.Initialize(PropertyOverrides, &Descriptor, InPointData, &Context);
+
+			// Since they are already partitioned and identical, we can just use the value on the first point
+			check(!Context.AttributeOverridePartition[I].IsEmpty());
+			int32 AnyPointIndexOnThisPartition = Context.AttributeOverridePartition[I][0];
+			Overrides.Apply(AnyPointIndexOnThisPartition);
+		}
+
+		// Given partitioning is expensive, check if we're out of time for this frame
+		Context.bPartitionDone = true;
+		if (Context.ShouldStop())
+		{
+			return false;
+		}
+	}
+
+	// Assign points to entries
+	int32 CurrentPartitionIndex = Context.CurrentPointIndex; // misnomer but we're reusing another concept from the context
+	const TArray<FPCGPoint>& Points = InPointData->GetPoints();
+	{
+		TRACE_CPUPROFILER_EVENT_SCOPE(UPCGAnimBankMeshSelectorByAttribute::SelectEntries::PushingPointsToInstanceLists);
+
+		// TODO: Revisit this when attribute partitioning is returned in a more optimized form
+		// The partition index is used to assign the point to the correct partition's instance
+		while(CurrentPartitionIndex < Context.AttributeOverridePartition.Num())
+		{
+			const int32 PartitionIndex = CurrentPartitionIndex++;
+			TArray<int32>& Partition = Context.AttributeOverridePartition[PartitionIndex];
+			const FPCGSoftAnimBankMeshComponentDescriptor& CurrentPartitionDescriptor = Context.OverriddenDescriptors[PartitionIndex];
+			
+			if (Partition.IsEmpty() || CurrentPartitionDescriptor.SkinnedAsset.IsNull())
+			{
+				continue;
+			}
+
+			bool bSkip = false;
+
+			for (const FSoftAnimBankItem& BankItem : CurrentPartitionDescriptor.BankItems)
+			{
+				if (BankItem.BankAsset.IsNull())
+				{
+					bSkip = true;
+					break;
+				}
+			}
+
+			if (bSkip)
+			{
+				continue;
+			}
+
+			// Setup data for mesh bounds computation
+			if (OutPointData && Settings->bApplyMeshBoundsToPoints)
+			{
+				TRACE_CPUPROFILER_EVENT_SCOPE(UPCGAnimBankMeshSelectorByAttribute::SelectEntries::PushPointsToApplyMeshBounds);
+				TArray<int32>& PointIndices = Context.MeshToOutPoints.FindOrAdd(CurrentPartitionDescriptor.SkinnedAsset).FindOrAdd(OutPointData);
+				PointIndices.Append(Partition);
+			}
+
+			// Separate the inverse determinant instances so we can push them to a different ABM
+			TArray<int32> ReverseInstances;
+			{
+				TRACE_CPUPROFILER_EVENT_SCOPE(UPCGAnimBankMeshSelectorByAttribute::SelectEntries::BuildReverseInstances);
+				int32 WriteIndex = 0;
+				for (int32 ReadIndex = 0; ReadIndex < Partition.Num(); ++ReadIndex)
+				{
+					const FPCGPoint& Point = Points[Partition[ReadIndex]];
+					if (Point.Transform.GetDeterminant() < 0)
+					{
+						ReverseInstances.Add(Partition[ReadIndex]);
+					}
+					else
+					{
+						if (WriteIndex != ReadIndex)
+						{
+							Partition[WriteIndex] = Partition[ReadIndex];
+						}
+						++WriteIndex;
+					}
+				}
+
+				Partition.SetNum(WriteIndex);
+			}
+
+			auto AddPointsToInstanceList = [&OutMeshInstances, &CurrentPartitionDescriptor, &MaterialOverrideHelper, &Points, PartitionIndex, InPointData](const TArray<int32>& PointIndices, bool bReverseTransform)
+			{
+				if (MaterialOverrideHelper.OverridesMaterials())
+				{
+					for (int32 PointIndex = 0; PointIndex < PointIndices.Num(); ++PointIndex)
+					{
+						const FPCGPoint& Point = Points[PointIndices[PointIndex]];
+						FPCGAnimBankMeshInstanceList& InstanceList = PCGAnimBankMeshSelectorAttribute::GetInstanceList(
+							OutMeshInstances,
+							CurrentPartitionDescriptor,
+							CurrentPartitionDescriptor.SkinnedAsset,
+							CurrentPartitionDescriptor.BankItems,
+							MaterialOverrideHelper.GetMaterialOverrides(Point.MetadataEntry),
+							bReverseTransform,
+							InPointData,
+							PartitionIndex
+						);
+						InstanceList.Instances.Emplace(Point.Transform);
+						InstanceList.InstancesIndices.Emplace(PointIndex);
+					}
+				}
+				else
+				{
+					TArray<TSoftObjectPtr<UMaterialInterface>> DummyMaterialList;
+					FPCGAnimBankMeshInstanceList& InstanceList = PCGAnimBankMeshSelectorAttribute::GetInstanceList(
+						OutMeshInstances,
+						CurrentPartitionDescriptor,
+						CurrentPartitionDescriptor.SkinnedAsset,
+						CurrentPartitionDescriptor.BankItems,
+						DummyMaterialList,
+						bReverseTransform,
+						InPointData,
+						PartitionIndex
+					);
+
+					check(InstanceList.Instances.Num() == InstanceList.InstancesIndices.Num());
+					const int32 InstanceOffset = InstanceList.Instances.Num();
+					InstanceList.Instances.SetNum(InstanceOffset + PointIndices.Num());
+					InstanceList.InstancesIndices.Append(PointIndices);
+
+					for(int32 PointIndex = 0; PointIndex < PointIndices.Num(); ++PointIndex)
+					{
+						const FPCGPoint& Point = Points[PointIndices[PointIndex]];
+						InstanceList.Instances[InstanceOffset + PointIndex] = Point.Transform;
+					}
+				}
+			};
+
+			{
+				TRACE_CPUPROFILER_EVENT_SCOPE(UPCGAnimBankMeshSelectorByAttribute::SelectEntries::AddPointsToInstanceList);
+				AddPointsToInstanceList(Partition, false);
+				AddPointsToInstanceList(ReverseInstances, true);
+			}
+
+			if (Context.ShouldStop())
+			{
+				break;
+			}
+		}
+	}
+
+	Context.CurrentPointIndex = CurrentPartitionIndex; // misnomer, but we're using the same context
+	return CurrentPartitionIndex == Context.AttributeOverridePartition.Num();
+}
+
+#undef LOCTEXT_NAMESPACE
diff --git a/Engine/Plugins/PCG/Source/PCG/Private/MeshSelectors/PCGAnimBankMeshSelectorWeighted.cpp b/Engine/Plugins/PCG/Source/PCG/Private/MeshSelectors/PCGAnimBankMeshSelectorWeighted.cpp
new file mode 100644
--- /dev/null
+++ b/Engine/Plugins/PCG/Source/PCG/Private/MeshSelectors/PCGAnimBankMeshSelectorWeighted.cpp
@@ -0,0 +1,305 @@
+// Copyright Epic Games, Inc. All Rights Reserved.
+
+#include "MeshSelectors/PCGAnimBankMeshSelectorWeighted.h"
+
+#include "Data/PCGPointData.h"
+#include "Data/PCGSpatialData.h"
+#include "Elements/PCGAnimBankMeshSpawner.h"
+#include "Elements/PCGAnimBankMeshSpawnerContext.h"
+#include "Helpers/PCGBlueprintHelpers.h"
+#include "MeshSelectors/PCGAnimBankMeshSelectorBase.h"
+
+#include "Animation/AnimBank.h"
+#include "Math/RandomStream.h"
+
+#include UE_INLINE_GENERATED_CPP_BY_NAME(PCGAnimBankMeshSelectorWeighted)
+
+#define LOCTEXT_NAMESPACE "PCGAnimBankMeshSelectorWeighted"
+
+namespace PCGAnimBankMeshSelectorWeighted
+{
+	FPCGAnimBankMeshInstanceList& GetInstanceList(
+		TArray<FPCGAnimBankMeshInstanceList>& InstanceLists,
+		bool bUseMaterialOverrides,
+		const TArray<TSoftObjectPtr<UMaterialInterface>>& InMaterialOverrides,
+		bool bInReverseCulling,
+		const UPCGPointData* InPointData)
+	{
+		check(InstanceLists.Num() > 0);
+
+#if 0 // AB-TODO
+		// First look through previously existing values - note that we scope this to prevent issues with the 0 index access which might become invalid below
+		{
+			const TArray<TSoftObjectPtr<UMaterialInterface>>& MaterialOverrides = (bUseMaterialOverrides ? InMaterialOverrides : InstanceLists[0].Descriptor.OverrideMaterials);
+			for (FPCGAnimBankMeshInstanceList& InstanceList : InstanceLists)
+			{
+				if (InstanceList.Descriptor.bReverseCulling == bInReverseCulling &&
+					InstanceList.Descriptor.OverrideMaterials == MaterialOverrides)
+				{
+					return InstanceList;
+				}
+			}
+		}
+#endif
+
+		FPCGAnimBankMeshInstanceList& NewInstanceList = InstanceLists.Emplace_GetRef();
+		NewInstanceList.Descriptor = InstanceLists[0].Descriptor;
+		// AB-TODO: NewInstanceList.Descriptor.bReverseCulling = bInReverseCulling;
+		NewInstanceList.PointData = InPointData;
+
+		if (bUseMaterialOverrides)
+		{
+			// AB-TODO: NewInstanceList.Descriptor.OverrideMaterials = InMaterialOverrides;
+		}
+
+		return NewInstanceList;
+	}
+}
+
+FPCGAnimBankMeshSelectorWeightedEntry::FPCGAnimBankMeshSelectorWeightedEntry(TSoftObjectPtr<USkinnedAsset> InMesh, const TArray<FSoftAnimBankItem>& InBankItems, int32 InWeight)
+	: Weight(InWeight)
+{
+	Descriptor.SkinnedAsset = InMesh;
+#if WITH_EDITOR
+	DisplayName = InMesh.ToSoftObjectPath().GetAssetFName();
+#endif // WITH_EDITOR
+}
+
+bool UPCGAnimBankMeshSelectorWeighted::SelectInstances(
+	FPCGAnimBankMeshSpawnerContext& Context,
+	const UPCGAnimBankMeshSpawnerSettings* Settings, 
+	const UPCGPointData* InPointData,
+	TArray<FPCGAnimBankMeshInstanceList>& OutMeshInstances,
+	UPCGPointData* OutPointData) const
+{
+	if (!InPointData)
+	{
+		PCGE_LOG_C(Error, GraphAndLog, &Context, LOCTEXT("MissingInputData", "Missing input data"));
+		return true;
+	}
+
+	TArray<TArray<FPCGAnimBankMeshInstanceList>>& MeshInstances = Context.WeightedMeshInstances;
+	TArray<int>& CumulativeWeights = Context.CumulativeWeights;
+
+	// TODO: Remove this log once the other selection modes are available
+	if (!Settings->AnimBankMeshComponentPropertyOverrides.IsEmpty())
+	{
+		PCGE_LOG_C(Log, LogOnly, &Context, LOCTEXT("AttributeToPropertyOverrideUnavailable", "Attribute to Property Overrides are only currently available with the 'By Attribute' Selector"));
+	}
+
+	// Setup
+	if (Context.CurrentPointIndex == 0)
+	{
+		int TotalWeight = 0;
+
+		// Prepare common mesh setups which we will use as a kind of map
+		for (const FPCGAnimBankMeshSelectorWeightedEntry& Entry : MeshEntries)
+		{
+			if (Entry.Weight <= 0)
+			{
+				PCGE_LOG_C(Verbose, LogOnly, &Context, LOCTEXT("EntryWithNegativeWeight", "Entry found with weight <= 0"));
+				continue;
+			}
+
+			TArray<FPCGAnimBankMeshInstanceList>& PickEntry = MeshInstances.Emplace_GetRef();
+			FPCGAnimBankMeshInstanceList& MeshInstanceList = PickEntry.Emplace_GetRef(Entry.Descriptor);
+			MeshInstanceList.PointData = InPointData;
+
+			TotalWeight += Entry.Weight;
+			CumulativeWeights.Add(TotalWeight);
+		}
+
+		if (TotalWeight <= 0)
+		{
+			return true;
+		}
+	}
+
+	FPCGMeshMaterialOverrideHelper& MaterialOverrideHelper = Context.MaterialOverrideHelper;
+	if (!MaterialOverrideHelper.IsInitialized())
+	{
+		MaterialOverrideHelper.Initialize(Context, bUseAttributeMaterialOverrides, MaterialOverrideAttributes, InPointData->Metadata);
+	}
+
+	if (!MaterialOverrideHelper.IsValid())
+	{
+		return true;
+	}
+
+	TArray<FPCGPoint>* OutPoints = nullptr;
+	FPCGMetadataAttribute<FString>* OutAttribute = nullptr;
+
+	if (OutPointData)
+	{
+		check(OutPointData->Metadata);
+
+		if (!OutPointData->Metadata->HasAttribute(Settings->OutAttributeName)) 
+		{
+			PCGE_LOG_C(Error, GraphAndLog, &Context, FText::Format(LOCTEXT("AttributeNotInMetadata", "Out attribute '{0}' is not in the metadata"), FText::FromName(Settings->OutAttributeName)));
+		}
+
+		FPCGMetadataAttributeBase* OutAttributeBase = OutPointData->Metadata->GetMutableAttribute(Settings->OutAttributeName);
+
+		if (OutAttributeBase)
+		{
+			if (OutAttributeBase->GetTypeId() == PCG::Private::MetadataTypes<FString>::Id)
+			{
+				OutAttribute = static_cast<FPCGMetadataAttribute<FString>*>(OutAttributeBase);
+				OutPoints = &OutPointData->GetMutablePoints();
+			}
+			else
+			{
+				PCGE_LOG_C(Error, GraphAndLog, &Context, LOCTEXT("AttributeNotFString", "Out attribute is not of valid type FString"));
+			}
+		}
+	}
+
+	// Assign points to entries
+	{
+		TRACE_CPUPROFILER_EVENT_SCOPE(FPCGAnimBankMeshSpawnerElement::Execute::SelectEntries);
+
+		int32 CurrentPointIndex = Context.CurrentPointIndex;
+		int32 LastCheckpointIndex = CurrentPointIndex;
+		constexpr int32 TimeSlicingCheckFrequency = 1024;
+		TMap<TSoftObjectPtr<USkinnedAsset>, PCGMetadataValueKey>& MeshToValueKey = Context.MeshToValueKey;
+		const int32 TotalWeight = CumulativeWeights.Last();
+
+		const TArray<FPCGPoint>& Points = InPointData->GetPoints();
+		while (CurrentPointIndex < Points.Num())
+		{
+			const FPCGPoint& Point = Points[CurrentPointIndex++];
+
+			FRandomStream RandomSource = UPCGBlueprintHelpers::GetRandomStreamFromPoint(Point, Settings, Context.SourceComponent.Get());
+			const int RandomWeightedPick = RandomSource.RandRange(0, TotalWeight - 1);
+
+			int RandomPick = 0;
+			while(RandomPick < MeshInstances.Num() && CumulativeWeights[RandomPick] <= RandomWeightedPick)
+			{
+				++RandomPick;
+			}
+
+			if(RandomPick < MeshInstances.Num())
+			{
+				const bool bNeedsReverseCulling = (Point.Transform.GetDeterminant() < 0);
+				FPCGAnimBankMeshInstanceList& InstanceList = PCGAnimBankMeshSelectorWeighted::GetInstanceList(MeshInstances[RandomPick], bUseAttributeMaterialOverrides, MaterialOverrideHelper.GetMaterialOverrides(Point.MetadataEntry), bNeedsReverseCulling, InPointData);
+				InstanceList.Instances.Emplace(Point.Transform);
+				InstanceList.InstancesIndices.Emplace(CurrentPointIndex - 1); // - 1 because it is already incremented.
+
+
+				const TSoftObjectPtr<USkinnedAsset>& Mesh = InstanceList.Descriptor.SkinnedAsset;
+
+				if (OutPointData && OutAttribute)
+				{
+					FPCGPoint& OutPoint = OutPoints->Add_GetRef(Point);
+					PCGMetadataValueKey* OutValueKey = MeshToValueKey.Find(Mesh);
+					if (!OutValueKey)
+					{
+						PCGMetadataValueKey ValueKey = OutAttribute->AddValue(Mesh.ToSoftObjectPath().ToString());
+						OutValueKey = &MeshToValueKey.Add(Mesh, ValueKey);
+					}
+
+					check(OutValueKey);
+
+					OutPointData->Metadata->InitializeOnSet(OutPoint.MetadataEntry);
+					OutAttribute->SetValueFromValueKey(OutPoint.MetadataEntry, *OutValueKey);
+
+					if (Settings->bApplyMeshBoundsToPoints)
+					{
+						TArray<int32>& PointIndices = Context.MeshToOutPoints.FindOrAdd(Mesh).FindOrAdd(OutPointData);
+						// CurrentPointIndex - 1, because CurrentPointIndex is incremented at the beginning of the loop
+						PointIndices.Emplace(CurrentPointIndex - 1);
+					}
+				}
+			}
+
+			// Check if we should stop here and continue in a subsequent call
+			if (CurrentPointIndex - LastCheckpointIndex >= TimeSlicingCheckFrequency)
+			{
+				if (Context.ShouldStop())
+				{
+					break;
+				}
+				else
+				{
+					LastCheckpointIndex = CurrentPointIndex;
+				}
+			}
+		}
+
+		Context.CurrentPointIndex = CurrentPointIndex;
+	}
+
+	if (Context.CurrentPointIndex == InPointData->GetPoints().Num())
+	{
+		// Finally, collapse to OutMeshInstances
+		for (TArray<FPCGAnimBankMeshInstanceList>& PickedMeshInstances : MeshInstances)
+		{
+			for (FPCGAnimBankMeshInstanceList& PickedMeshInstanceEntry : PickedMeshInstances)
+			{
+				OutMeshInstances.Emplace(MoveTemp(PickedMeshInstanceEntry));
+			}
+		}
+
+		return true;
+	}
+	else
+	{
+		return false;
+	}
+}
+
+void UPCGAnimBankMeshSelectorWeighted::PostLoad()
+{
+	Super::PostLoad();
+
+	for (FPCGAnimBankMeshSelectorWeightedEntry& Entry : MeshEntries)
+	{
+		// TODO: Remove if/when FBodyInstance is updated or replaced
+		// Necessary to update the collision Response Container from the Response Array
+		Entry.Descriptor.PostLoadFixup(this);
+	}
+
+#if WITH_EDITOR
+	RefreshDisplayNames();
+#endif // WITH_EDITOR
+}
+
+void UPCGAnimBankMeshSelectorWeighted::PostDuplicate(bool bDuplicateForPIE)
+{
+	Super::PostDuplicate(bDuplicateForPIE);
+
+#if WITH_EDITOR
+	RefreshDisplayNames();
+#endif // WITH_EDITOR
+}
+
+void UPCGAnimBankMeshSelectorWeighted::PostEditImport()
+{
+	Super::PostEditImport();
+
+#if WITH_EDITOR
+	RefreshDisplayNames();
+#endif // WITH_EDITOR
+}
+
+#if WITH_EDITOR
+void UPCGAnimBankMeshSelectorWeighted::PostEditChangeProperty(FPropertyChangedEvent& PropertyChangedEvent)
+{
+	Super::PostEditChangeProperty(PropertyChangedEvent);
+
+	if (PropertyChangedEvent.GetPropertyName() == GET_MEMBER_NAME_CHECKED(FPCGSoftAnimBankMeshComponentDescriptor, SkinnedAsset))
+	{
+		RefreshDisplayNames();
+	}
+}
+
+void UPCGAnimBankMeshSelectorWeighted::RefreshDisplayNames()
+{
+	for (FPCGAnimBankMeshSelectorWeightedEntry& Entry : MeshEntries)
+	{
+		Entry.DisplayName = Entry.Descriptor.SkinnedAsset.ToSoftObjectPath().GetAssetFName();
+	}
+}
+#endif // WITH_EDITOR
+
+#undef LOCTEXT_NAMESPACE
diff --git a/Engine/Plugins/PCG/Source/PCG/Private/MeshSelectors/PCGAnimBankMeshSelectorWeightedByCategory.cpp b/Engine/Plugins/PCG/Source/PCG/Private/MeshSelectors/PCGAnimBankMeshSelectorWeightedByCategory.cpp
new file mode 100644
--- /dev/null
+++ b/Engine/Plugins/PCG/Source/PCG/Private/MeshSelectors/PCGAnimBankMeshSelectorWeightedByCategory.cpp
@@ -0,0 +1,351 @@
+// Copyright Epic Games, Inc. All Rights Reserved.
+
+#include "MeshSelectors/PCGAnimBankMeshSelectorWeightedByCategory.h"
+
+#include "Data/PCGPointData.h"
+#include "Data/PCGSpatialData.h"
+#include "Elements/PCGAnimBankMeshSpawner.h"
+#include "Elements/PCGAnimBankMeshSpawnerContext.h"
+#include "Helpers/PCGBlueprintHelpers.h"
+#include "MeshSelectors/PCGAnimBankMeshSelectorWeighted.h"
+
+#include "Animation/AnimBank.h"
+#include "Math/RandomStream.h"
+
+#include UE_INLINE_GENERATED_CPP_BY_NAME(PCGAnimBankMeshSelectorWeightedByCategory)
+
+#define LOCTEXT_NAMESPACE "PCGAnimBankMeshSelectorWeightedByCategory"
+
+
+bool UPCGAnimBankMeshSelectorWeightedByCategory::SelectInstances(
+	FPCGAnimBankMeshSpawnerContext& Context,
+	const UPCGAnimBankMeshSpawnerSettings* Settings,
+	const UPCGPointData* InPointData,
+	TArray<FPCGAnimBankMeshInstanceList>& OutMeshInstances,
+	UPCGPointData* OutPointData) const
+{
+	if (!InPointData)
+	{
+		PCGE_LOG_C(Error, GraphAndLog, &Context, LOCTEXT("InputMissingData", "Missing input data"));
+		return true;
+	}
+
+	if (!InPointData->Metadata)
+	{
+		PCGE_LOG_C(Error, GraphAndLog, &Context, LOCTEXT("InputMissingMetadata", "Unable to get metadata from input"));
+		return true;
+	}
+
+	if (!InPointData->Metadata->HasAttribute(CategoryAttribute))
+	{
+		PCGE_LOG_C(Error, GraphAndLog, &Context, FText::Format(LOCTEXT("InputMissingAttribute", "Attribute '{0}' is not in the metadata"), FText::FromName(CategoryAttribute)));
+		return true;
+	}
+
+	const FPCGMetadataAttributeBase* AttributeBase = InPointData->Metadata->GetConstAttribute(CategoryAttribute);
+	check(AttributeBase);
+
+	// TODO: support enum type as well
+	if (AttributeBase->GetTypeId() != PCG::Private::MetadataTypes<FString>::Id)
+	{
+		PCGE_LOG_C(Error, GraphAndLog, &Context, FText::Format(LOCTEXT("AttributeInvalidType", "Attribute '{0}' is not of valid type FString"), FText::FromName(CategoryAttribute)));
+		return true;
+	}
+
+	const FPCGMetadataAttribute<FString>* Attribute = static_cast<const FPCGMetadataAttribute<FString>*>(AttributeBase);
+
+	// maps a CategoryEntry ValueKey to the meshes and precomputed weight data 
+	TMap<PCGMetadataValueKey, FPCGAnimBankMeshInstancesAndWeights>& CategoryEntryToInstancesAndWeights = Context.CategoryEntryToInstancesAndWeights;
+
+	// unmarked points will fallback to the MeshEntries associated with the DefaultValueKey
+	PCGMetadataValueKey DefaultValueKey = PCGDefaultValueKey;
+
+	// TODO: Remove this log once the other selection modes are available
+	if (!Settings->AnimBankMeshComponentPropertyOverrides.IsEmpty())
+	{
+		PCGE_LOG_C(Log, LogOnly, &Context, LOCTEXT("AttributeToPropertyOverrideUnavailable", "Attribute to Property Overrides are only currently available with the 'By Attribute' Selector"));
+	}
+
+	// Setup
+	if (Context.CurrentPointIndex == 0)
+	{
+		const FString AttributeDefaultValue = Attribute->GetValue(PCGDefaultValueKey);
+
+		for (const FPCGAnimBankMeshWeightedByCategoryEntryList& Entry : Entries)
+		{
+			if (Entry.WeightedMeshEntries.Num() == 0)
+			{
+				PCGE_LOG_C(Verbose, LogOnly, &Context, FText::Format(LOCTEXT("EmptyEntryInCategory", "Empty entry found in category '{0}'"), FText::FromString(Entry.CategoryEntry)));
+				continue;
+			}
+
+			PCGMetadataValueKey ValueKey = Attribute->FindValue(Entry.CategoryEntry);
+
+			if (ValueKey == PCGNotFoundValueKey)
+			{
+				PCGE_LOG_C(Verbose, LogOnly, &Context, FText::Format(LOCTEXT("UnusedCategory", "Unused category '{0}'. Not a valid value for attribute '{1}'."), 
+					FText::FromString(Entry.CategoryEntry), FText::FromName(CategoryAttribute)));
+				continue;
+			}
+
+			FPCGAnimBankMeshInstancesAndWeights* InstancesAndWeights = CategoryEntryToInstancesAndWeights.Find(ValueKey);
+
+			if (InstancesAndWeights)
+			{
+				PCGE_LOG_C(Warning, GraphAndLog, &Context, FText::Format(LOCTEXT("DuplicateEntry", "Duplicate entry found in category '{0}'. Subsequent entries are ignored."), FText::FromString(Entry.CategoryEntry)));
+				continue;
+			}
+
+			if (Entry.IsDefault)
+			{
+				if (DefaultValueKey == PCGDefaultValueKey)
+				{
+					DefaultValueKey = ValueKey;
+				}
+				else
+				{
+					PCGE_LOG_C(Warning, GraphAndLog, &Context, LOCTEXT("DuplicateDefaultEntry", "Duplicate default entry found. Subsequent default entries are ignored."));
+				}
+			}
+
+			InstancesAndWeights = &CategoryEntryToInstancesAndWeights.Add(ValueKey, FPCGAnimBankMeshInstancesAndWeights());
+
+			int TotalWeight = 0;
+			for (const FPCGAnimBankMeshSelectorWeightedEntry& WeightedEntry : Entry.WeightedMeshEntries)
+			{
+				if (WeightedEntry.Weight <= 0)
+				{
+					PCGE_LOG_C(Verbose, LogOnly, &Context, FText::Format(LOCTEXT("ZeroOrNegativeWeight", "Entry found with weight <= 0 in category '{0}'"), FText::FromString(Entry.CategoryEntry)));
+					continue;
+				}
+
+				TArray<FPCGAnimBankMeshInstanceList>& PickEntry = InstancesAndWeights->MeshInstances.Emplace_GetRef();
+				FPCGAnimBankMeshInstanceList& MeshInstanceList = PickEntry.Emplace_GetRef(WeightedEntry.Descriptor);
+				MeshInstanceList.PointData = InPointData;
+
+				// precompute the weights
+				TotalWeight += WeightedEntry.Weight;
+				InstancesAndWeights->CumulativeWeights.Add(TotalWeight);
+			}
+
+			// if no weighted entries were collected, discard this InstancesAndWeights
+			if (InstancesAndWeights->CumulativeWeights.Num() == 0)
+			{
+				CategoryEntryToInstancesAndWeights.Remove(ValueKey);
+			}
+		}
+	}
+
+	FPCGMeshMaterialOverrideHelper& MaterialOverrideHelper = Context.MaterialOverrideHelper;
+	if (!MaterialOverrideHelper.IsInitialized())
+	{
+		MaterialOverrideHelper.Initialize(Context, bUseAttributeMaterialOverrides, MaterialOverrideAttributes, InPointData->Metadata);
+	}
+
+	if (!MaterialOverrideHelper.IsValid())
+	{
+		return true;
+	}
+
+	TArray<FPCGPoint>* OutPoints = nullptr;
+	FPCGMetadataAttribute<FString>* OutAttribute = nullptr;
+
+	if (OutPointData)
+	{
+		check(OutPointData->Metadata);
+
+		if (!OutPointData->Metadata->HasAttribute(Settings->OutAttributeName)) 
+		{
+			PCGE_LOG_C(Error, GraphAndLog, &Context, FText::Format(LOCTEXT("AttributeNotInMetadata", "Out attribute '{0}' is not in the metadata"), FText::FromName(Settings->OutAttributeName)));
+		}
+
+		FPCGMetadataAttributeBase* OutAttributeBase = OutPointData->Metadata->GetMutableAttribute(Settings->OutAttributeName);
+
+		if (OutAttributeBase)
+		{
+			if (OutAttributeBase->GetTypeId() == PCG::Private::MetadataTypes<FString>::Id)
+			{
+				OutAttribute = static_cast<FPCGMetadataAttribute<FString>*>(OutAttributeBase);
+				OutPoints = &OutPointData->GetMutablePoints();
+			}
+			else
+			{
+				PCGE_LOG_C(Error, GraphAndLog, &Context, LOCTEXT("TypeNotFString", "Out attribute is not of valid type FString"));
+			}
+		}
+	}
+
+	TRACE_CPUPROFILER_EVENT_SCOPE(FPCGAnimBankMeshSpawnerElement::Execute::SelectEntries);
+
+	int32 CurrentPointIndex = Context.CurrentPointIndex;
+	int32 LastCheckpointIndex = CurrentPointIndex;
+	constexpr int32 TimeSlicingCheckFrequency = 1024;
+	TMap<TSoftObjectPtr<USkinnedAsset>, PCGMetadataValueKey>& MeshToValueKey = Context.MeshToValueKey;
+
+	// Assign points to entries
+	const TArray<FPCGPoint>& Points = InPointData->GetPoints();
+	while (CurrentPointIndex < Points.Num())
+	{
+		const FPCGPoint& Point = Points[CurrentPointIndex++];
+
+		const PCGMetadataValueKey ValueKey = Attribute->GetValueKey(Point.MetadataEntry);
+
+		// if no mesh list was processed for this attribute value, fallback to the default mesh list
+		FPCGAnimBankMeshInstancesAndWeights* InstancesAndWeights = CategoryEntryToInstancesAndWeights.Find(ValueKey);
+		if (!InstancesAndWeights)
+		{
+			if (DefaultValueKey != PCGDefaultValueKey)
+			{
+				InstancesAndWeights = CategoryEntryToInstancesAndWeights.Find(DefaultValueKey);
+				check(InstancesAndWeights);
+			}
+			else
+			{
+				continue;
+			}
+		}
+
+		const int TotalWeight = InstancesAndWeights->CumulativeWeights.Last();
+
+		FRandomStream RandomSource = UPCGBlueprintHelpers::GetRandomStreamFromPoint(Point, Settings, Context.SourceComponent.Get());
+		const int RandomWeightedPick = RandomSource.RandRange(0, TotalWeight - 1);
+
+		int RandomPick = 0;
+		while(RandomPick < InstancesAndWeights->MeshInstances.Num() && InstancesAndWeights->CumulativeWeights[RandomPick] <= RandomWeightedPick)
+		{
+			++RandomPick;
+		}
+
+		if(RandomPick < InstancesAndWeights->MeshInstances.Num())
+		{
+			const bool bNeedsReverseCulling = (Point.Transform.GetDeterminant() < 0);
+			FPCGAnimBankMeshInstanceList& InstanceList = PCGAnimBankMeshSelectorWeighted::GetInstanceList(InstancesAndWeights->MeshInstances[RandomPick], bUseAttributeMaterialOverrides, MaterialOverrideHelper.GetMaterialOverrides(Point.MetadataEntry), bNeedsReverseCulling, InPointData);
+			InstanceList.Instances.Emplace(Point.Transform);
+			InstanceList.InstancesIndices.Emplace(CurrentPointIndex - 1); // - 1 because it is already incremented.
+
+			const TSoftObjectPtr<USkinnedAsset>& Mesh = InstanceList.Descriptor.SkinnedAsset;
+
+			if (OutPointData && OutAttribute)
+			{
+				FPCGPoint& OutPoint = OutPoints->Add_GetRef(Point);
+
+				PCGMetadataValueKey* OutValueKey = MeshToValueKey.Find(Mesh);
+				if (!OutValueKey)
+				{
+					PCGMetadataValueKey TempValueKey = OutAttribute->AddValue(Mesh.ToSoftObjectPath().ToString());
+					OutValueKey = &MeshToValueKey.Add(Mesh, TempValueKey);
+				}
+
+				check(OutValueKey);
+
+				OutPointData->Metadata->InitializeOnSet(OutPoint.MetadataEntry);
+				OutAttribute->SetValueFromValueKey(OutPoint.MetadataEntry, *OutValueKey);
+
+				if (Settings->bApplyMeshBoundsToPoints)
+				{
+					TArray<int32>& PointIndices = Context.MeshToOutPoints.FindOrAdd(Mesh).FindOrAdd(OutPointData);
+					// CurrentPointIndex - 1, because CurrentPointIndex is incremented at the beginning of the loop
+					PointIndices.Emplace(CurrentPointIndex - 1);
+				}
+			}
+		}
+
+		// Check if we should stop here and continue in a subsequent call
+		if (CurrentPointIndex - LastCheckpointIndex >= TimeSlicingCheckFrequency)
+		{
+			if (Context.ShouldStop())
+			{
+				break;
+			}
+			else
+			{
+				LastCheckpointIndex = CurrentPointIndex;
+			}
+		}
+	}
+
+	Context.CurrentPointIndex = CurrentPointIndex;
+
+	if (CurrentPointIndex == Points.Num())
+	{
+		// Collapse to OutMeshInstances
+		for (TPair<PCGMetadataValueKey, FPCGAnimBankMeshInstancesAndWeights>& Entry : CategoryEntryToInstancesAndWeights)
+		{
+			for (TArray<FPCGAnimBankMeshInstanceList>& PickedMeshInstances : Entry.Value.MeshInstances)
+			{
+				for (FPCGAnimBankMeshInstanceList& PickedMeshInstanceEntry : PickedMeshInstances)
+				{
+					OutMeshInstances.Emplace(MoveTemp(PickedMeshInstanceEntry));
+				}
+			}
+		}
+
+		return true;
+	}
+	else
+	{
+		return false;
+	}
+}
+
+void UPCGAnimBankMeshSelectorWeightedByCategory::PostLoad()
+{
+	Super::PostLoad();
+
+	for (FPCGAnimBankMeshWeightedByCategoryEntryList& Entry : Entries)
+	{
+		// TODO: Remove if/when FBodyInstance is updated or replaced
+		// Necessary to update the collision Response Container from the Response Array
+		for (FPCGAnimBankMeshSelectorWeightedEntry& WeightedEntry : Entry.WeightedMeshEntries)
+		{
+			WeightedEntry.Descriptor.PostLoadFixup(this);
+		}
+	}
+
+#if WITH_EDITOR
+	RefreshDisplayNames();
+#endif // WITH_EDITOR
+}
+
+void UPCGAnimBankMeshSelectorWeightedByCategory::PostDuplicate(bool bDuplicateForPIE)
+{
+	Super::PostDuplicate(bDuplicateForPIE);
+
+#if WITH_EDITOR
+	RefreshDisplayNames();
+#endif // WITH_EDITOR
+}
+
+void UPCGAnimBankMeshSelectorWeightedByCategory::PostEditImport()
+{
+	Super::PostEditImport();
+
+#if WITH_EDITOR
+	RefreshDisplayNames();
+#endif // WITH_EDITOR
+}
+
+#if WITH_EDITOR
+void UPCGAnimBankMeshSelectorWeightedByCategory::PostEditChangeProperty(FPropertyChangedEvent& PropertyChangedEvent)
+{
+	Super::PostEditChangeProperty(PropertyChangedEvent);
+
+	if (PropertyChangedEvent.GetPropertyName() == GET_MEMBER_NAME_CHECKED(FPCGSoftAnimBankMeshComponentDescriptor, SkinnedAsset))
+	{
+		RefreshDisplayNames();
+	}
+}
+
+void UPCGAnimBankMeshSelectorWeightedByCategory::RefreshDisplayNames()
+{
+	for (FPCGAnimBankMeshWeightedByCategoryEntryList& Entry : Entries)
+	{
+		for (FPCGAnimBankMeshSelectorWeightedEntry& WeightedEntry : Entry.WeightedMeshEntries)
+		{
+			WeightedEntry.DisplayName = WeightedEntry.Descriptor.SkinnedAsset.ToSoftObjectPath().GetAssetFName();
+		}
+	}
+}
+#endif // WITH_EDITOR
+
+#undef LOCTEXT_NAMESPACE
diff --git a/Engine/Plugins/PCG/Source/PCG/Private/PCGManagedResource.cpp b/Engine/Plugins/PCG/Source/PCG/Private/PCGManagedResource.cpp
--- a/Engine/Plugins/PCG/Source/PCG/Private/PCGManagedResource.cpp
+++ b/Engine/Plugins/PCG/Source/PCG/Private/PCGManagedResource.cpp
@@ -7,6 +7,7 @@
 #include "Helpers/PCGHelpers.h"
 
 #include "Components/InstancedStaticMeshComponent.h"
+#include "Components/AnimBankMeshComponent.h"
 #include "Components/SplineMeshComponent.h"
 #include "Engine/Level.h"
 #include "UObject/Package.h"
@@ -717,17 +718,6 @@ void UPCGManagedISMComponent::PostLoad()
 {
 	Super::PostLoad();
 
-	if (!bHasDescriptor)
-	{
-		if (UInstancedStaticMeshComponent* ISMC = GetComponent())
-		{
-			FISMComponentDescriptor NewDescriptor;
-			NewDescriptor.InitFrom(ISMC);
-
-			SetDescriptor(NewDescriptor);
-		}
-	}
-
 	// Cache raw ptr
 	GetComponent();
 }
@@ -866,6 +856,159 @@ void UPCGManagedISMComponent::SetComponent(UInstancedStaticMeshComponent* InComp
 	CachedRawComponentPtr = InComponent;
 }
 
+void UPCGManagedAnimBankMeshComponent::PostLoad()
+{
+	Super::PostLoad();
+
+	if (!bHasDescriptor)
+	{
+		if (UAnimBankMeshComponent* ABMC = GetComponent())
+		{
+			FAnimBankMeshComponentDescriptor NewDescriptor;
+			NewDescriptor.InitFrom(ABMC);
+
+			SetDescriptor(NewDescriptor);
+		}
+	}
+
+	// Cache raw ptr
+	GetComponent();
+}
+
+void UPCGManagedAnimBankMeshComponent::ForgetComponent()
+{
+	Super::ForgetComponent();
+	CachedRawComponentPtr = nullptr;
+}
+
+void UPCGManagedAnimBankMeshComponent::SetDescriptor(const FAnimBankMeshComponentDescriptor& InDescriptor)
+{
+	bHasDescriptor = true;
+	Descriptor = InDescriptor;
+}
+
+bool UPCGManagedAnimBankMeshComponent::ReleaseIfUnused(TSet<TSoftObjectPtr<AActor>>& OutActorsToDelete)
+{
+	if (Super::ReleaseIfUnused(OutActorsToDelete) || !GetComponent())
+	{
+		return true;
+	}
+	else if (GetComponent()->GetInstanceCount() == 0)
+	{
+		GeneratedComponent->DestroyComponent();
+		ForgetComponent();
+		return true;
+	}
+	else
+	{
+		return false;
+	}
+}
+
+void UPCGManagedAnimBankMeshComponent::ResetComponent()
+{
+	if (UAnimBankMeshComponent* ABMC = GetComponent())
+	{
+		ABMC->ClearInstances();
+		ABMC->UpdateBounds();
+	}
+}
+
+void UPCGManagedAnimBankMeshComponent::MarkAsUsed()
+{
+	const bool bWasMarkedUnused = bIsMarkedUnused;
+	Super::MarkAsUsed();
+
+	if (!bWasMarkedUnused)
+	{
+		return;
+	}
+
+	if (UAnimBankMeshComponent* ABMC = GetComponent())
+	{
+		const bool bHasPreviousRootLocation = bHasRootLocation;
+
+		// Keep track of the current root location so if we reuse this later we are able to update this appropriately
+		if (USceneComponent* RootComponent = ABMC->GetAttachmentRoot())
+		{
+			bHasRootLocation = true;
+			RootLocation = RootComponent->GetComponentLocation();
+		}
+		else
+		{
+			bHasRootLocation = false;
+			RootLocation = FVector::ZeroVector;
+		}
+
+		if (bHasPreviousRootLocation != bHasRootLocation || (ABMC->GetComponentLocation() - RootLocation).SquaredLength() > UE_DOUBLE_SMALL_NUMBER)
+		{
+			// Reset the rotation/scale to be identity otherwise if the root component transform has changed, the final transform will be wrong
+			// Since this is technically 'moving' the ABM, we need to unregister it before moving otherwise we could get a warning that we're moving a component with static mobility
+			ABMC->UnregisterComponent();
+			ABMC->SetWorldTransform(FTransform(FQuat::Identity, RootLocation, FVector::OneVector));
+			ABMC->RegisterComponent();
+		}
+	}
+}
+
+void UPCGManagedAnimBankMeshComponent::MarkAsReused()
+{
+	Super::MarkAsReused();
+
+	if (UAnimBankMeshComponent* ABMC = GetComponent())
+	{
+		// Reset the rotation/scale to be identity otherwise if the root component transform has changed, the final transform will be wrong
+		FVector TentativeRootLocation = RootLocation;
+
+		if (!bHasRootLocation)
+		{
+			if (USceneComponent* RootComponent = ABMC->GetAttachmentRoot())
+			{
+				TentativeRootLocation = RootComponent->GetComponentLocation();
+			}
+		}
+
+		if ((ABMC->GetComponentLocation() - TentativeRootLocation).SquaredLength() > UE_DOUBLE_SMALL_NUMBER)
+		{
+			// Since this is technically 'moving' the ABM, we need to unregister it before moving otherwise we could get a warning that we're moving a component with static mobility
+			ABMC->UnregisterComponent();
+			ABMC->SetWorldTransform(FTransform(FQuat::Identity, TentativeRootLocation, FVector::OneVector));
+			ABMC->RegisterComponent();
+		}
+	}
+}
+
+void UPCGManagedAnimBankMeshComponent::SetRootLocation(const FVector& InRootLocation)
+{
+	bHasRootLocation = true;
+	RootLocation = InRootLocation;
+}
+
+UAnimBankMeshComponent* UPCGManagedAnimBankMeshComponent::GetComponent() const
+{
+	if (!CachedRawComponentPtr)
+	{
+		UAnimBankMeshComponent* GeneratedComponentPtr = Cast<UAnimBankMeshComponent>(GeneratedComponent.Get());
+
+		// Implementation note:
+		// There is no surefire way to make sure that we can use the raw pointer UNLESS it is from the same owner
+		if (GeneratedComponentPtr && Cast<UPCGComponent>(GetOuter()) && GeneratedComponentPtr->GetOwner() == Cast<UPCGComponent>(GetOuter())->GetOwner())
+		{
+			CachedRawComponentPtr = GeneratedComponentPtr;
+		}
+
+		return GeneratedComponentPtr;
+	}
+
+	return CachedRawComponentPtr;
+}
+
+void UPCGManagedAnimBankMeshComponent::SetComponent(UAnimBankMeshComponent* InComponent)
+{
+	GeneratedComponent = InComponent;
+	CachedRawComponentPtr = InComponent;
+}
+
 void UPCGManagedSplineMeshComponent::ForgetComponent()
 {
 	Super::ForgetComponent();
diff --git a/Engine/Plugins/PCG/Source/PCG/Public/Elements/PCGAnimBankMeshSpawner.h b/Engine/Plugins/PCG/Source/PCG/Public/Elements/PCGAnimBankMeshSpawner.h
new file mode 100644
--- /dev/null
+++ b/Engine/Plugins/PCG/Source/PCG/Public/Elements/PCGAnimBankMeshSpawner.h
@@ -0,0 +1,143 @@
+// Copyright Epic Games, Inc. All Rights Reserved.
+
+#pragma once
+
+#include "PCGSettings.h"
+#include "Metadata/PCGObjectPropertyOverride.h"
+#include "Engine/CollisionProfile.h"
+#include "PCGAnimBankMeshSpawner.generated.h"
+
+class UPCGAnimBankMeshInstanceDataPackerBase;
+class UPCGAnimBankMeshSelectorBase;
+class UPCGSpatialData;
+struct FPCGContext;
+struct FPCGAnimBankMeshInstanceList;
+struct FPCGAnimBankMeshPackedCustomData;
+struct FPCGAnimBankMeshSpawnerContext;
+
+class UAnimBank;
+
+UCLASS(BlueprintType, ClassGroup = (Procedural), meta = (PrioritizeCategories = "Settings"))
+class PCG_API UPCGAnimBankMeshSpawnerSettings : public UPCGSettings
+{
+	GENERATED_BODY()
+
+public:
+	UPCGAnimBankMeshSpawnerSettings(const FObjectInitializer& ObjectInitializer);
+
+	// ~Begin UPCGSettings interface
+	virtual bool IsKernelValid(FPCGContext* InContext = nullptr, bool bQuiet = true) const;
+	virtual FString GetCookedKernelSource(const FPCGKernelAttributeTable& InAttributeStaticTable) const override;
+	virtual void GetKernelAttributeKeys(TArray<FPCGKernelAttributeKey>& OutKeys) const override;
+	virtual void AddStaticCreatedStrings(TArray<FString>& InOutStringTable) const override;
+	virtual int ComputeKernelThreadCount(const UPCGDataBinding* Binding) const override;
+	virtual bool ComputeOutputPinDataDesc(const UPCGPin* OutputPin, const UPCGDataBinding* Binding, FPCGDataCollectionDesc& OutDesc) const override;
+#if WITH_EDITOR
+	virtual void CreateAdditionalInputDataInterfaces(FPCGGPUCompilationContext& InOutContext, UObject* InObjectOuter, TArray<TObjectPtr<UComputeDataInterface>>& OutDataInterfaces) const override;
+	virtual void CreateAdditionalOutputDataInterfaces(FPCGGPUCompilationContext& InOutContext, UObject* InObjectOuter, TArray<TObjectPtr<UComputeDataInterface>>& OutDataInterfaces) const override;
+	virtual bool DisplayExecuteOnGPUSetting() const override { return true; }
+	virtual FName GetDefaultNodeName() const override { return FName(TEXT("AnimBankMeshSpawner")); }
+	virtual FText GetDefaultNodeTitle() const override;
+	virtual EPCGSettingsType GetType() const override { return EPCGSettingsType::Spawner; }
+	virtual void ApplyDeprecation(UPCGNode* InOutNode) override;
+#endif
+	virtual bool UseSeed() const override { return true; }
+
+protected:
+	virtual TArray<FPCGPinProperties> InputPinProperties() const override;
+	virtual TArray<FPCGPinProperties> OutputPinProperties() const override { return Super::DefaultPointOutputPinProperties(); }
+	virtual FPCGElementPtr CreateElement() const override;
+	// ~End UPCGSettings interface
+
+public:
+	// ~Begin UObject interface
+	virtual void PostLoad() override;
+#if WITH_EDITOR
+	virtual void PostEditChangeProperty(FPropertyChangedEvent& PropertyChangedEvent) override;
+	virtual bool CanEditChange(const FProperty* InProperty) const override;
+	// ~End UObject interface
+#endif
+
+	UFUNCTION(BlueprintCallable, Category = Settings)
+	void SetMeshSelectorType(TSubclassOf<UPCGAnimBankMeshSelectorBase> InMeshSelectorType);
+
+	UFUNCTION(BlueprintCallable, Category = Settings)
+	void SetInstancePackerType(TSubclassOf<UPCGAnimBankMeshInstanceDataPackerBase> InInstancePackerType);
+
+public:
+	/** Defines the method of mesh selection per input data */
+	UPROPERTY(BlueprintReadOnly, EditAnywhere, NoClear, Category = MeshSelector)
+	TSubclassOf<UPCGAnimBankMeshSelectorBase> MeshSelectorType;
+
+	UPROPERTY(BlueprintReadOnly, VisibleAnywhere, Instanced, Category = MeshSelector)
+	TObjectPtr<UPCGAnimBankMeshSelectorBase> MeshSelectorParameters;
+
+	/**
+	 * Defines the method of custom data packing for spawned anim bank meshes. Note, Rotators are treated as 3 floats, while Quaternions are
+	 * treated as 4 floats. You can see an attribute's type in the 'Attribute List View' window, and use an 'Attribute Cast' node to cast to the desired type.
+	 */
+	UPROPERTY(BlueprintReadOnly, EditAnywhere, Category = InstanceDataPacker)
+	TSubclassOf<UPCGAnimBankMeshInstanceDataPackerBase> InstanceDataPackerType;
+
+	UPROPERTY(BlueprintReadOnly, VisibleAnywhere, Instanced, Category = InstanceDataPacker)
+	TObjectPtr<UPCGAnimBankMeshInstanceDataPackerBase> InstanceDataPackerParameters = nullptr;
+
+	/** Map an attribute directly to an anim bank descriptor property, the value of which will be overriden when generated.
+	 * Note: Currently only enabled using SelectByAttribute mesh selection.
+	 */
+	UPROPERTY(BlueprintReadWrite, EditAnywhere, Category = Settings)
+	TArray<FPCGObjectPropertyOverrideDescription> AnimBankMeshComponentPropertyOverrides;
+
+	/** Attribute name to store mesh SoftObjectPaths inside if the output pin is connected. Note: Will overwrite existing data if the attribute name already exists. */
+	UPROPERTY(BlueprintReadWrite, EditAnywhere, Category = Settings)
+	FName OutAttributeName = NAME_None;
+
+	/** Sets the BoundsMin and BoundsMax attributes of each point to reflect the AnimBank mesh spawned at its location */
+	UPROPERTY(BlueprintReadWrite, EditAnywhere, Category = Settings)
+	bool bApplyMeshBoundsToPoints = true;
+
+	UPROPERTY(BlueprintReadWrite, Category = Settings, meta = (PCG_Overridable))
+	TSoftObjectPtr<AActor> TargetActor;
+
+	/** Specify a list of functions to be called on the target actor after instances are spawned. Functions need to be parameter-less and with "CallInEditor" flag enabled. */
+	UPROPERTY(BlueprintReadWrite, EditAnywhere, Category = Settings)
+	TArray<FName> PostProcessFunctionNames;
+
+	/** Meshes/Materials will be synchronously loaded before spawning instead of asynchronously. */
+	UPROPERTY(BlueprintReadWrite, EditAnywhere, Category = "Settings|Debug")
+	bool bSynchronousLoad = false;
+
+	/** Opt-in option to silence errors when the property override attributes are not found. */
+	UPROPERTY(BlueprintReadWrite, EditAnywhere, Category = "Settings|Advanced")
+	bool bSilenceOverrideAttributeNotFoundErrors = false;
+
+	/** Adds a warning to the node on repeated spawning with identical conditions (ie. same mesh descriptor at same spawn location, etc). */
+	UPROPERTY(BlueprintReadWrite, EditAnywhere, Category = "Settings|Advanced")
+	bool bWarnOnIdenticalSpawn = true;
+
+protected:
+	void RefreshMeshSelector();
+	void RefreshInstancePacker();
+
+private:
+	static TCHAR const* TemplateFilePath;
+};
+
+class PCG_API FPCGAnimBankMeshSpawnerElement : public IPCGElement
+{
+public:
+	virtual bool CanExecuteOnlyOnMainThread(FPCGContext* Context) const override;
+	virtual bool IsCacheable(const UPCGSettings* InSettings) const override { return false; }
+
+protected:
+	virtual FPCGContext* CreateContext() override;
+	virtual bool PrepareDataInternal(FPCGContext* Context) const override;
+	virtual bool ExecuteInternal(FPCGContext* Context) const override;
+	virtual void AbortInternal(FPCGContext* Context) const override;
+	
+	void SpawnAnimBankMeshInstances(
+		FPCGAnimBankMeshSpawnerContext* Context,
+		const FPCGAnimBankMeshInstanceList& InstanceList,
+		AActor* TargetActor,
+		const FPCGAnimBankMeshPackedCustomData* InPackedCustomData) const;
+};
diff --git a/Engine/Plugins/PCG/Source/PCG/Public/Elements/PCGAnimBankMeshSpawnerContext.h b/Engine/Plugins/PCG/Source/PCG/Public/Elements/PCGAnimBankMeshSpawnerContext.h
new file mode 100644
--- /dev/null
+++ b/Engine/Plugins/PCG/Source/PCG/Public/Elements/PCGAnimBankMeshSpawnerContext.h
@@ -0,0 +1,73 @@
+// Copyright Epic Games, Inc. All Rights Reserved.
+
+#pragma once
+
+#include "PCGContext.h"
+#include "PCGManagedResource.h"
+#include "Async/PCGAsyncLoadingContext.h"
+#include "MeshSelectors/PCGAnimBankMeshSelectorBase.h"
+#include "InstanceDataPackers/PCGAnimBankMeshInstanceDataPackerBase.h"
+
+class AActor;
+class UPCGSpatialData;
+
+struct FPCGAnimBankMeshInstancesAndWeights
+{
+	TArray<TArray<FPCGAnimBankMeshInstanceList>> MeshInstances;
+	TArray<int> CumulativeWeights;
+};
+
+struct FPCGAnimBankMeshSpawnerContext : public FPCGContext, public IPCGAsyncLoadingContext
+{
+	struct FPackedInstanceListData
+	{
+		FPackedInstanceListData();
+		~FPackedInstanceListData();
+		AActor* TargetActor = nullptr;
+		const UPCGSpatialData* SpatialData = nullptr;
+		TArray<FPCGAnimBankMeshInstanceList> MeshInstances;
+		TArray<FPCGAnimBankMeshPackedCustomData> PackedCustomData;
+	};
+
+	TArray<FPackedInstanceListData> MeshInstancesData;
+	// Index of input in the context (for selection)
+	int32 CurrentInputIndex = 0;
+	// Index of MeshInstances/PackedCustomData in last MeshInstancesData element
+	int32 CurrentDataIndex = 0;
+
+	// Whole-execution variables
+	bool bReuseCheckDone = false;
+	bool bSkippedDueToReuse = false;
+
+	// Per-input context variables
+	bool bCurrentInputSetup = false;
+	bool bSelectionDone = false;
+	bool bPartitionDone = false;
+
+	const UPCGPointData* CurrentPointData = nullptr;
+	UPCGPointData* CurrentOutputPointData = nullptr;
+	FPCGMeshMaterialOverrideHelper MaterialOverrideHelper;
+	int32 CurrentPointIndex = 0;
+
+	// Used in all selectors if we have to change the out points bounds by the mesh bounds. Will be empty otherwise.
+	// We need to keep all point indices that will spawn this mesh, in all output point data.
+	TMap<TSoftObjectPtr<USkinnedAsset>, TMap<UPCGPointData*, TArray<int32>>> MeshToOutPoints;
+
+	// Used in weighted selector
+	TArray<TArray<FPCGAnimBankMeshInstanceList>> WeightedMeshInstances;
+	//TMap<FSoftAnimBankItem, PCGMetadataValueKey> MeshToValueKey;
+	TMap<TSoftObjectPtr<USkinnedAsset>, PCGMetadataValueKey> MeshToValueKey;
+	TArray<int> CumulativeWeights;
+
+	// Used in the weighted by category selector
+	TMap<PCGMetadataValueKey, FPCGAnimBankMeshInstancesAndWeights> CategoryEntryToInstancesAndWeights;
+
+	// Used for mesh property overrides
+	TArray<TArray<int32>> AttributeOverridePartition;
+	TArray<FPCGSoftAnimBankMeshComponentDescriptor> OverriddenDescriptors;
+
+	// Keeping track of all touched resources to allow for correct cleanup on abort
+	TArray<TWeakObjectPtr<UPCGManagedAnimBankMeshComponent>> TouchedResources;
+
+	void ResetInputIterationData();
+};
\ No newline at end of file
diff --git a/Engine/Plugins/PCG/Source/PCG/Public/Helpers/PCGActorHelpers.h b/Engine/Plugins/PCG/Source/PCG/Public/Helpers/PCGActorHelpers.h
--- a/Engine/Plugins/PCG/Source/PCG/Public/Helpers/PCGActorHelpers.h
+++ b/Engine/Plugins/PCG/Source/PCG/Public/Helpers/PCGActorHelpers.h
@@ -5,10 +5,11 @@
 #include "PCGCrc.h"
 #include "Elements/PCGSplineMeshParams.h"
 #include "MeshSelectors/PCGISMDescriptor.h"
-
+#include "MeshSelectors/PCGAnimBankMeshDescriptor.h"
 #include "Engine/EngineTypes.h"
 #include "Engine/SplineMeshComponentDescriptor.h"
 #include "Engine/World.h"
+#include "Animation/AnimBank.h"
 #include "ISMPartition/ISMComponentDescriptor.h"
 #include "Kismet/BlueprintFunctionLibrary.h"
 #include "Templates/SubclassOf.h"
@@ -24,6 +25,7 @@ class UMaterialInterface;
 class USplineMeshComponent;
 class UPCGComponent;
 class UPCGManagedISMComponent;
+class UPCGManagedAnimBankMeshComponent;
 class UPCGManagedSplineMeshComponent;
 class UStaticMesh;
 class UWorld;
@@ -67,6 +69,23 @@ PRAGMA_ENABLE_DEPRECATION_WARNINGS
 	inline bool operator==(const FPCGISMComponentBuilderParams& Other) const { return Descriptor == Other.Descriptor && NumCustomDataFloats == Other.NumCustomDataFloats && bAllowDescriptorChanges == Other.bAllowDescriptorChanges; }
 };
 
+struct FPCGAnimBankMeshComponentBuilderParams
+{
+	FPCGAnimBankMeshComponentBuilderParams() = default;
+
+	FPCGSoftAnimBankMeshComponentDescriptor Descriptor;
+	int32 NumCustomDataFloats = 0;
+	FPCGCrc SettingsCrc;
+	bool bTransient = false;
+
+	friend inline uint32 GetTypeHash(const FPCGAnimBankMeshComponentBuilderParams& Key)
+	{
+		return HashCombine(GetTypeHash(Key.Descriptor), 1 + Key.NumCustomDataFloats);
+	}
+
+	inline bool operator==(const FPCGAnimBankMeshComponentBuilderParams& Other) const { return Descriptor == Other.Descriptor && NumCustomDataFloats == Other.NumCustomDataFloats; }
+};
+
 struct FPCGSplineMeshComponentBuilderParameters
 {
 	FSplineMeshComponentDescriptor Descriptor;
@@ -103,6 +122,7 @@ PRAGMA_ENABLE_DEPRECATION_WARNINGS
 	
 	static UInstancedStaticMeshComponent* GetOrCreateISMC(AActor* InTargetActor, UPCGComponent* SourceComponent, const FPCGISMComponentBuilderParams& Params);
 	static UPCGManagedISMComponent* GetOrCreateManagedISMC(AActor* InTargetActor, UPCGComponent* SourceComponent, const FPCGISMComponentBuilderParams& Params);
+	static UPCGManagedAnimBankMeshComponent* GetOrCreateManagedABMC(AActor* InTargetActor, UPCGComponent* SourceComponent, const FPCGAnimBankMeshComponentBuilderParams& Params);
 	static USplineMeshComponent* GetOrCreateSplineMeshComponent(AActor* InTargetActor, UPCGComponent* SourceComponent, const FPCGSplineMeshComponentBuilderParameters& Params);
 	static UPCGManagedSplineMeshComponent* GetOrCreateManagedSplineMeshComponent(AActor* InTargetActor, UPCGComponent* SourceComponent, const FPCGSplineMeshComponentBuilderParameters& Params);
 
diff --git a/Engine/Plugins/PCG/Source/PCG/Public/InstanceDataPackers/PCGAnimBankMeshInstanceDataPackerBase.h b/Engine/Plugins/PCG/Source/PCG/Public/InstanceDataPackers/PCGAnimBankMeshInstanceDataPackerBase.h
new file mode 100644
--- /dev/null
+++ b/Engine/Plugins/PCG/Source/PCG/Public/InstanceDataPackers/PCGAnimBankMeshInstanceDataPackerBase.h
@@ -0,0 +1,58 @@
+// Copyright Epic Games, Inc. All Rights Reserved.
+
+#pragma once
+
+#include "PCGAnimBankMeshInstanceDataPackerBase.generated.h"
+
+class FPCGMetadataAttributeBase;
+class IPCGAttributeAccessor;
+class IPCGAttributeAccessorKeys;
+class UPCGMetadata;
+class UPCGSpatialData;
+struct FPCGContext;
+struct FPCGAnimBankMeshInstanceList;
+
+USTRUCT(BlueprintType)
+struct FPCGAnimBankMeshPackedCustomData
+{
+	GENERATED_BODY()
+
+	UPROPERTY(BlueprintReadWrite, EditAnywhere, Category = Settings)
+	int SequenceIndex = 0;
+
+	UPROPERTY(BlueprintReadWrite, EditAnywhere, Category = Settings)
+	int NumCustomDataFloats = 0;
+
+	UPROPERTY(BlueprintReadWrite, EditAnywhere, Category = Settings)
+	TArray<float> CustomData;
+};
+
+UCLASS(Abstract, BlueprintType, Blueprintable, ClassGroup = (Procedural))
+class PCG_API UPCGAnimBankMeshInstanceDataPackerBase : public UObject 
+{
+	GENERATED_BODY()
+
+public:
+	/** Defines the strategy for custom float data packing */
+	UFUNCTION(BlueprintNativeEvent, Category = InstancePacking)
+	void PackInstances(FPCGContext& Context, const UPCGSpatialData* InSpatialData, const FPCGAnimBankMeshInstanceList& InstanceList, FPCGAnimBankMeshPackedCustomData& OutPackedCustomData) const;
+
+	virtual void PackInstances_Implementation(FPCGContext& Context, const UPCGSpatialData* InSpatialData, const FPCGAnimBankMeshInstanceList& InstanceList, FPCGAnimBankMeshPackedCustomData& OutPackedCustomData) const;
+
+	/** Interprets Metadata TypeId and increments OutPackedCustomData.NumCustomDataFloats appropriately. Returns false if the type could not be interpreted. */
+	UFUNCTION(BlueprintCallable, Category = InstancePacking)
+	bool AddTypeToPacking(int TypeId, FPCGAnimBankMeshPackedCustomData& OutPackedCustomData) const;
+
+	/** Build a PackedCustomData by processing each attribute in order for each point in the InstanceList */
+	UFUNCTION(BlueprintCallable, Category = InstancePacking) 
+	void PackCustomDataFromAttributes(const FPCGAnimBankMeshInstanceList& InstanceList, const UPCGMetadata* Metadata, const TArray<FName>& AttributeNames, FPCGAnimBankMeshPackedCustomData& OutPackedCustomData) const;
+
+	/** Build a PackedCustomData by processing each attribute in order for each point in the InstanceList */
+	void PackCustomDataFromAttributes(const FPCGAnimBankMeshInstanceList& InstanceList, const TArray<const FPCGMetadataAttributeBase*>& Attributes, FPCGAnimBankMeshPackedCustomData& OutPackedCustomData) const;
+
+	/** Build a PackedCustomData by processing each accessor in order for each point in the InstanceList */
+	void PackCustomDataFromAccessors(const FPCGAnimBankMeshInstanceList& InstanceList, TArray<TUniquePtr<const IPCGAttributeAccessor>> Accessors, TArray<TUniquePtr<const IPCGAttributeAccessorKeys>> AccessorKeys, FPCGAnimBankMeshPackedCustomData& OutPackedCustomData) const;
+
+	/** If OutNames is not null, returns a list of all attributes that will be packed. Returns true if this list can be statically determined (prior to execution). */
+	virtual bool GetAttributeNames(TArray<FName>* OutNames) { return false; }
+};
diff --git a/Engine/Plugins/PCG/Source/PCG/Public/InstanceDataPackers/PCGAnimBankMeshInstanceDataPackerByAttribute.h b/Engine/Plugins/PCG/Source/PCG/Public/InstanceDataPackers/PCGAnimBankMeshInstanceDataPackerByAttribute.h
new file mode 100644
--- /dev/null
+++ b/Engine/Plugins/PCG/Source/PCG/Public/InstanceDataPackers/PCGAnimBankMeshInstanceDataPackerByAttribute.h
@@ -0,0 +1,27 @@
+// Copyright Epic Games, Inc. All Rights Reserved.
+
+#pragma once
+
+#include "PCGAnimBankMeshInstanceDataPackerBase.h"
+#include "Metadata/PCGAttributePropertySelector.h"
+
+#include "PCGAnimBankMeshInstanceDataPackerByAttribute.generated.h"
+
+UCLASS(MinimalAPI, BlueprintType, ClassGroup = (Procedural))
+class UPCGAnimBankMeshInstanceDataPackerByAttribute : public UPCGAnimBankMeshInstanceDataPackerBase 
+{
+	GENERATED_BODY()
+
+public:
+#if WITH_EDITOR
+	// ~Begin UObject interface
+	virtual void PostLoad() override;
+	// ~End UObject interface
+#endif
+	virtual void PackInstances_Implementation(UPARAM(ref) FPCGContext& Context, const UPCGSpatialData* InSpatialData, UPARAM(ref) const FPCGAnimBankMeshInstanceList& InstanceList, FPCGAnimBankMeshPackedCustomData& OutPackedCustomData) const override;
+	virtual bool GetAttributeNames(TArray<FName>* OutNames) override;
+
+public:
+	UPROPERTY(BlueprintReadWrite, EditAnywhere, Category = InstanceDataPacker)
+	TArray<FPCGAttributePropertyInputSelector> AttributeSelectors;
+};
diff --git a/Engine/Plugins/PCG/Source/PCG/Public/InstanceDataPackers/PCGAnimBankMeshInstanceDataPackerByRegex.h b/Engine/Plugins/PCG/Source/PCG/Public/InstanceDataPackers/PCGAnimBankMeshInstanceDataPackerByRegex.h
new file mode 100644
--- /dev/null
+++ b/Engine/Plugins/PCG/Source/PCG/Public/InstanceDataPackers/PCGAnimBankMeshInstanceDataPackerByRegex.h
@@ -0,0 +1,20 @@
+// Copyright Epic Games, Inc. All Rights Reserved.
+
+#pragma once
+
+#include "PCGAnimBankMeshInstanceDataPackerBase.h"
+
+#include "PCGAnimBankMeshInstanceDataPackerByRegex.generated.h"
+
+UCLASS(MinimalAPI, BlueprintType, ClassGroup = (Procedural))
+class UPCGAnimBankMeshInstanceDataPackerByRegex : public UPCGAnimBankMeshInstanceDataPackerBase 
+{
+	GENERATED_BODY()
+
+public:
+	virtual void PackInstances_Implementation(UPARAM(ref) FPCGContext& Context, const UPCGSpatialData* InSpatialData, UPARAM(ref) const FPCGAnimBankMeshInstanceList& InstanceList, FPCGAnimBankMeshPackedCustomData& OutPackedCustomData) const override;
+
+public:
+	UPROPERTY(BlueprintReadWrite, EditAnywhere, Category = InstanceDataPacker)
+	TArray<FString> RegexPatterns; 
+};
diff --git a/Engine/Plugins/PCG/Source/PCG/Public/MeshSelectors/PCGAnimBankMeshDescriptor.h b/Engine/Plugins/PCG/Source/PCG/Public/MeshSelectors/PCGAnimBankMeshDescriptor.h
new file mode 100644
--- /dev/null
+++ b/Engine/Plugins/PCG/Source/PCG/Public/MeshSelectors/PCGAnimBankMeshDescriptor.h
@@ -0,0 +1,25 @@
+// Copyright Epic Games, Inc. All Rights Reserved.
+
+#pragma once
+
+#include "Engine/CollisionProfile.h"
+#include "Animation/AnimBank.h"
+
+#include "PCGAnimBankMeshDescriptor.generated.h"
+
+/** Convenience PCG-side component descriptor so we can adjust defaults to the most common use cases. */
+// Implementation note: the tags don't really need to contribute to the hash, so we will retain the base class !=, == and ComputeHash implementations
+USTRUCT()
+struct FPCGSoftAnimBankMeshComponentDescriptor : public FSoftAnimBankMeshComponentDescriptor
+{
+	GENERATED_BODY()
+
+	PCG_API FPCGSoftAnimBankMeshComponentDescriptor();
+	PCG_API explicit FPCGSoftAnimBankMeshComponentDescriptor(const FAnimBankMeshComponentDescriptor& Other);
+	PCG_API virtual void InitFrom(const UAnimBankMeshComponent* Component, bool bInitBodyInstance = true) override;
+	PCG_API virtual void InitComponent(UAnimBankMeshComponent* ISMComponent) const override;
+
+public:
+	UPROPERTY(EditAnywhere, Category="Component Settings|Tags", meta = (DisplayAfter = "ComponentClass"))
+	TArray<FName> ComponentTags;
+};
diff --git a/Engine/Plugins/PCG/Source/PCG/Public/MeshSelectors/PCGAnimBankMeshSelectorBase.h b/Engine/Plugins/PCG/Source/PCG/Public/MeshSelectors/PCGAnimBankMeshSelectorBase.h
new file mode 100644
--- /dev/null
+++ b/Engine/Plugins/PCG/Source/PCG/Public/MeshSelectors/PCGAnimBankMeshSelectorBase.h
@@ -0,0 +1,67 @@
+// Copyright Epic Games, Inc. All Rights Reserved.
+
+#pragma once
+
+#include "PCGPoint.h"
+#include "Data/PCGPointData.h"
+#include "MeshSelectors/PCGAnimBankMeshDescriptor.h"
+#include "MeshSelectors/PCGMeshMaterialOverrideHelper.h"
+#include "Metadata/PCGMetadata.h"
+
+#include "Engine/CollisionProfile.h"
+
+#include "PCGAnimBankMeshSelectorBase.generated.h"
+
+class UPCGPointData;
+class UPCGSpatialData;
+class UAnimBank;
+struct FPCGContext;
+struct FPCGAnimBankMeshSpawnerContext;
+class UPCGAnimBankMeshSpawnerSettings;
+
+USTRUCT(BlueprintType)
+struct FPCGAnimBankMeshInstanceList
+{
+	GENERATED_BODY()
+
+	explicit FPCGAnimBankMeshInstanceList(const FPCGSoftAnimBankMeshComponentDescriptor& InDescriptor)
+	: Descriptor(InDescriptor)
+	, AttributePartitionIndex(INDEX_NONE)
+	{}
+	
+	FPCGAnimBankMeshInstanceList() = default;
+	~FPCGAnimBankMeshInstanceList() = default;
+	FPCGAnimBankMeshInstanceList(const FPCGAnimBankMeshInstanceList&) = default;
+	FPCGAnimBankMeshInstanceList(FPCGAnimBankMeshInstanceList&&) = default;
+	FPCGAnimBankMeshInstanceList& operator=(const FPCGAnimBankMeshInstanceList&) = default;
+	FPCGAnimBankMeshInstanceList& operator=(FPCGAnimBankMeshInstanceList&&) = default;
+
+	UPROPERTY(EditAnywhere, Category = Settings)
+	FPCGSoftAnimBankMeshComponentDescriptor Descriptor;
+
+	UPROPERTY(BlueprintReadWrite, EditAnywhere, Category = Settings)
+	TArray<FTransform> Instances;
+
+	/** Tracks which partition the instance list belongs to. */
+	int64 AttributePartitionIndex;
+
+	UPROPERTY(BlueprintReadWrite, EditAnywhere, Category = Settings)
+	TWeakObjectPtr<const UPCGPointData> PointData;
+	
+	UPROPERTY(BlueprintReadWrite, EditAnywhere, Category = Settings)
+	TArray<int32> InstancesIndices;
+};
+
+UCLASS(Abstract, BlueprintType, ClassGroup = (Procedural))
+class PCG_API UPCGAnimBankMeshSelectorBase : public UObject 
+{
+	GENERATED_BODY()
+
+public:
+	virtual bool SelectInstances(
+		FPCGAnimBankMeshSpawnerContext& Context,
+		const UPCGAnimBankMeshSpawnerSettings* Settings,
+		const UPCGPointData* InPointData,
+		TArray<FPCGAnimBankMeshInstanceList>& OutMeshInstances,
+		UPCGPointData* OutPointData) const PURE_VIRTUAL(UPCGAnimBankMeshSelectorBase::SelectInstances, return true;);
+};
diff --git a/Engine/Plugins/PCG/Source/PCG/Public/MeshSelectors/PCGAnimBankMeshSelectorByAttribute.h b/Engine/Plugins/PCG/Source/PCG/Public/MeshSelectors/PCGAnimBankMeshSelectorByAttribute.h
new file mode 100644
--- /dev/null
+++ b/Engine/Plugins/PCG/Source/PCG/Public/MeshSelectors/PCGAnimBankMeshSelectorByAttribute.h
@@ -0,0 +1,39 @@
+// Copyright Epic Games, Inc. All Rights Reserved.
+
+#pragma once
+
+#include "MeshSelectors/PCGAnimBankMeshSelectorBase.h"
+#include "MeshSelectors/PCGAnimBankMeshDescriptor.h"
+
+#include "PCGAnimBankMeshSelectorByAttribute.generated.h"
+
+UCLASS(MinimalAPI, BlueprintType, ClassGroup = (Procedural))
+class UPCGAnimBankMeshSelectorByAttribute : public UPCGAnimBankMeshSelectorBase 
+{
+	GENERATED_BODY()
+
+public:
+	// ~Begin UObject interface
+	virtual void PostLoad() override;
+	// ~End UObject interface
+
+	virtual bool SelectInstances(
+		FPCGAnimBankMeshSpawnerContext& Context,
+		const UPCGAnimBankMeshSpawnerSettings* Settings,
+		const UPCGPointData* InPointData,
+		TArray<FPCGAnimBankMeshInstanceList>& OutMeshInstances,
+		UPCGPointData* OutPointData) const override;
+
+public:
+	UPROPERTY(BlueprintReadWrite, EditAnywhere, Category = MeshSelector)
+	FName AttributeName;
+
+	UPROPERTY(EditAnywhere, Category = MeshSelector)
+	FPCGSoftAnimBankMeshComponentDescriptor TemplateDescriptor;
+
+	UPROPERTY(BlueprintReadWrite, EditAnywhere, Category = MeshSelector, meta = (InlineEditConditionToggle))
+	bool bUseAttributeMaterialOverrides = false;
+
+	UPROPERTY(BlueprintReadWrite, EditAnywhere, DisplayName = "By Attribute Material Overrides", Category = MeshSelector, meta = (EditCondition = "bUseAttributeMaterialOverrides"))
+	TArray<FName> MaterialOverrideAttributes;
+};
diff --git a/Engine/Plugins/PCG/Source/PCG/Public/MeshSelectors/PCGAnimBankMeshSelectorWeighted.h b/Engine/Plugins/PCG/Source/PCG/Public/MeshSelectors/PCGAnimBankMeshSelectorWeighted.h
new file mode 100644
--- /dev/null
+++ b/Engine/Plugins/PCG/Source/PCG/Public/MeshSelectors/PCGAnimBankMeshSelectorWeighted.h
@@ -0,0 +1,75 @@
+// Copyright Epic Games, Inc. All Rights Reserved.
+
+#pragma once
+
+#include "MeshSelectors/PCGAnimBankMeshSelectorBase.h"
+#include "MeshSelectors/PCGAnimBankMeshDescriptor.h"
+
+#include "PCGAnimBankMeshSelectorWeighted.generated.h"
+
+namespace PCGAnimBankMeshSelectorWeighted
+{
+	// Returns variation matching the overrides & the reverse culling flag
+	FPCGAnimBankMeshInstanceList& GetInstanceList(
+		TArray<FPCGAnimBankMeshInstanceList>& InstanceLists,
+		bool bUseMaterialOverrides,
+		const TArray<TSoftObjectPtr<UMaterialInterface>>& InMaterialOverrides,
+		bool bInIsLocalToWorldDeterminantNegative,
+		const UPCGPointData* InPointData);
+}
+
+USTRUCT(BlueprintType)
+struct PCG_API FPCGAnimBankMeshSelectorWeightedEntry
+{
+	GENERATED_BODY()
+
+	FPCGAnimBankMeshSelectorWeightedEntry() = default;
+	FPCGAnimBankMeshSelectorWeightedEntry(TSoftObjectPtr<USkinnedAsset> InMesh, const TArray<FSoftAnimBankItem>& InBankItems, int32 InWeight);
+
+	UPROPERTY(EditAnywhere, Category = Settings)
+	FPCGSoftAnimBankMeshComponentDescriptor Descriptor;
+
+	UPROPERTY(BlueprintReadWrite, EditAnywhere, Category = Settings, meta = (ClampMin = "0"))
+	int Weight = 1;
+
+#if WITH_EDITORONLY_DATA
+	UPROPERTY(Transient, VisibleAnywhere, Category = Settings)
+	FName DisplayName = NAME_None;
+#endif
+};
+
+UCLASS(MinimalAPI, BlueprintType, ClassGroup = (Procedural))
+class UPCGAnimBankMeshSelectorWeighted : public UPCGAnimBankMeshSelectorBase 
+{
+	GENERATED_BODY()
+
+public:
+	virtual bool SelectInstances(
+		FPCGAnimBankMeshSpawnerContext& Context,
+		const UPCGAnimBankMeshSpawnerSettings* Settings, 
+		const UPCGPointData* InPointData,
+		TArray<FPCGAnimBankMeshInstanceList>& OutMeshInstances,
+		UPCGPointData* OutPointData) const override;
+
+	// ~Begin UObject interface
+	virtual void PostLoad() override;
+	virtual void PostDuplicate(bool bDuplicateForPIE) override;
+	virtual void PostEditImport() override;
+#if WITH_EDITOR
+	virtual void PostEditChangeProperty(FPropertyChangedEvent& PropertyChangedEvent) override;
+	// ~End UObject interface
+
+	/** Refresh MeshEntries display names */
+	PCG_API void RefreshDisplayNames();
+#endif // WITH_EDITOR
+
+public:
+	UPROPERTY(BlueprintReadWrite, EditAnywhere, Category = MeshSelector, meta = (TitleProperty = "DisplayName"))
+	TArray<FPCGAnimBankMeshSelectorWeightedEntry> MeshEntries;
+
+	UPROPERTY(BlueprintReadWrite, EditAnywhere, Category = MeshSelector, meta = (InlineEditConditionToggle))
+	bool bUseAttributeMaterialOverrides = false;
+
+	UPROPERTY(BlueprintReadWrite, EditAnywhere, DisplayName = "By Attribute Material Overrides", Category = MeshSelector, meta = (EditCondition = "bUseAttributeMaterialOverrides"))
+	TArray<FName> MaterialOverrideAttributes;
+};
diff --git a/Engine/Plugins/PCG/Source/PCG/Public/MeshSelectors/PCGAnimBankMeshSelectorWeightedByCategory.h b/Engine/Plugins/PCG/Source/PCG/Public/MeshSelectors/PCGAnimBankMeshSelectorWeightedByCategory.h
new file mode 100644
--- /dev/null
+++ b/Engine/Plugins/PCG/Source/PCG/Public/MeshSelectors/PCGAnimBankMeshSelectorWeightedByCategory.h
@@ -0,0 +1,70 @@
+// Copyright Epic Games, Inc. All Rights Reserved.
+
+#pragma once
+
+#include "MeshSelectors/PCGAnimBankMeshSelectorBase.h"
+#include "PCGAnimBankMeshSelectorWeighted.h"
+
+#include "PCGAnimBankMeshSelectorWeightedByCategory.generated.h"
+
+USTRUCT(BlueprintType)
+struct PCG_API FPCGAnimBankMeshWeightedByCategoryEntryList
+{
+	GENERATED_BODY()
+
+	FPCGAnimBankMeshWeightedByCategoryEntryList() = default;
+
+	FPCGAnimBankMeshWeightedByCategoryEntryList(const FString& InCategoryEntry, const TArray<FPCGAnimBankMeshSelectorWeightedEntry>& InWeightedMeshEntries)
+	: CategoryEntry(InCategoryEntry)
+	, WeightedMeshEntries(InWeightedMeshEntries)
+	{
+	}
+
+	UPROPERTY(BlueprintReadWrite, EditAnywhere, Category = Settings)
+	FString CategoryEntry;
+
+	UPROPERTY(BlueprintReadWrite, EditAnywhere, Category = Settings)
+	bool IsDefault = false;
+
+	UPROPERTY(BlueprintReadWrite, EditAnywhere, Category = Settings, meta = (TitleProperty = "DisplayName"))
+	TArray<FPCGAnimBankMeshSelectorWeightedEntry> WeightedMeshEntries;
+};
+
+UCLASS(MinimalAPI, BlueprintType, ClassGroup = (Procedural))
+class UPCGAnimBankMeshSelectorWeightedByCategory : public UPCGAnimBankMeshSelectorBase 
+{
+	GENERATED_BODY()
+
+public:
+	virtual bool SelectInstances(
+		FPCGAnimBankMeshSpawnerContext& Context,
+		const UPCGAnimBankMeshSpawnerSettings* Settings,
+		const UPCGPointData* InPointData,
+		TArray<FPCGAnimBankMeshInstanceList>& OutMeshInstances,
+		UPCGPointData* OutPointData) const override;
+
+	// ~Begin UObject interface
+	virtual void PostLoad() override;
+	virtual void PostDuplicate(bool bDuplicateForPIE) override;
+	virtual void PostEditImport() override;
+#if WITH_EDITOR
+	virtual void PostEditChangeProperty(FPropertyChangedEvent& PropertyChangedEvent) override;
+	// ~End UObject interface
+
+	// Refresh MeshEntries display names
+	PCG_API void RefreshDisplayNames();
+#endif // WITH_EDITOR
+
+public:
+	UPROPERTY(BlueprintReadWrite, EditAnywhere, Category = MeshSelector)
+	FName CategoryAttribute;
+
+	UPROPERTY(BlueprintReadWrite, EditAnywhere, Category = MeshSelector)
+	TArray<FPCGAnimBankMeshWeightedByCategoryEntryList> Entries;
+
+	UPROPERTY(BlueprintReadWrite, EditAnywhere, Category = MeshSelector, meta = (InlineEditConditionToggle))
+	bool bUseAttributeMaterialOverrides = false;
+
+	UPROPERTY(BlueprintReadWrite, EditAnywhere, DisplayName = "By Attribute Material Overrides", Category = MeshSelector, meta = (EditCondition = "bUseAttributeMaterialOverrides"))
+	TArray<FName> MaterialOverrideAttributes;
+};
diff --git a/Engine/Plugins/PCG/Source/PCG/Public/MeshSelectors/PCGMeshMaterialOverrideHelper.h b/Engine/Plugins/PCG/Source/PCG/Public/MeshSelectors/PCGMeshMaterialOverrideHelper.h
new file mode 100644
--- /dev/null
+++ b/Engine/Plugins/PCG/Source/PCG/Public/MeshSelectors/PCGMeshMaterialOverrideHelper.h
@@ -0,0 +1,62 @@
+// Copyright Epic Games, Inc. All Rights Reserved.
+
+#pragma once
+
+#include "Metadata/PCGMetadata.h"
+
+#include "PCGMeshMaterialOverrideHelper.generated.h"
+
+class UMaterialInterface;
+
+UENUM()
+enum class EPCGMeshSelectorMaterialOverrideMode : uint8
+{
+	NoOverride UMETA(Tooltip = "Does not apply any material overrides to the spawned mesh(es)"),
+	StaticOverride UMETA(Tooltip = "Applies the material overrides provided in the Static Material Overrides array"),
+	ByAttributeOverride UMETA(Tooltip = "Applies the materials overrides using the point data attribute(s) specified in the By Attribute Material Overrides array")
+};
+
+/** Struct used to efficiently gather overrides and cache them during instance packing */
+struct PCG_API FPCGMeshMaterialOverrideHelper
+{
+	FPCGMeshMaterialOverrideHelper() = default;
+
+	// Use this constructor when you have a 1:1 mapping between attributes or static overrides
+	void Initialize(
+		FPCGContext& InContext,
+		bool bUseMaterialOverrideAttributes,
+		const TArray<TSoftObjectPtr<UMaterialInterface>>& InStaticMaterialOverrides,
+		const TArray<FName>& InMaterialOverrideAttributeNames,
+		const UPCGMetadata* InMetadata);
+
+	// Use this constructor when you have common attribute usage or separate static overrides
+	void Initialize(
+		FPCGContext& InContext,
+		bool bInByAttributeOverride,
+		const TArray<FName>& InMaterialOverrideAttributeNames,
+		const UPCGMetadata* InMetadata);
+
+	void Reset();
+
+	bool IsInitialized() const { return bIsInitialized; }
+	bool IsValid() const { return bIsValid; }
+	bool OverridesMaterials() const { return bUseMaterialOverrideAttributes; }
+	const TArray<TSoftObjectPtr<UMaterialInterface>>& GetMaterialOverrides(PCGMetadataEntryKey EntryKey);
+
+private:
+	// Cached data
+	TArray<const FPCGMetadataAttributeBase*> MaterialAttributes;
+	TArray<TMap<PCGMetadataValueKey, TSoftObjectPtr<UMaterialInterface>>> ValueKeyToOverrideMaterials;
+	TArray<TSoftObjectPtr<UMaterialInterface>> WorkingMaterialOverrides;
+
+	// Data needed to perform operations
+	bool bIsInitialized = false;
+	bool bIsValid = false;
+	bool bUseMaterialOverrideAttributes = false;
+
+	TArray<TSoftObjectPtr<UMaterialInterface>> StaticMaterialOverrides;
+	TArray<FName> MaterialOverrideAttributeNames;
+	const UPCGMetadata* Metadata = nullptr;
+
+	void Initialize(FPCGContext& InContext);
+};
diff --git a/Engine/Plugins/PCG/Source/PCG/Public/MeshSelectors/PCGMeshSelectorBase.h b/Engine/Plugins/PCG/Source/PCG/Public/MeshSelectors/PCGMeshSelectorBase.h
--- a/Engine/Plugins/PCG/Source/PCG/Public/MeshSelectors/PCGMeshSelectorBase.h
+++ b/Engine/Plugins/PCG/Source/PCG/Public/MeshSelectors/PCGMeshSelectorBase.h
@@ -5,6 +5,7 @@
 #include "PCGPoint.h"
 #include "Data/PCGPointData.h"
 #include "MeshSelectors/PCGISMDescriptor.h"
+#include "MeshSelectors/PCGMeshMaterialOverrideHelper.h"
 #include "Metadata/PCGMetadata.h"
 
 #include "Engine/CollisionProfile.h"
@@ -16,9 +17,7 @@ class UPCGSpatialData;
 class UStaticMesh;
 struct FPCGContext;
 struct FPCGStaticMeshSpawnerContext;
-
 class UPCGStaticMeshSpawnerSettings;
-class UMaterialInterface;
 
 USTRUCT(BlueprintType)
 struct FPCGMeshInstanceList
@@ -62,59 +61,6 @@ struct FPCGMeshInstanceList
 #endif // WITH_EDITORONLY_DATA
 };
 
-UENUM()
-enum class EPCGMeshSelectorMaterialOverrideMode : uint8
-{
-	NoOverride UMETA(Tooltip = "Does not apply any material overrides to the spawned mesh(es)"),
-	StaticOverride UMETA(Tooltip = "Applies the material overrides provided in the Static Material Overrides array"),
-	ByAttributeOverride UMETA(Tooltip = "Applies the materials overrides using the point data attribute(s) specified in the By Attribute Material Overrides array")
-};
-
-/** Struct used to efficiently gather overrides and cache them during instance packing */
-struct PCG_API FPCGMeshMaterialOverrideHelper
-{
-	FPCGMeshMaterialOverrideHelper() = default;
-
-	// Use this constructor when you have a 1:1 mapping between attributes or static overrides
-	void Initialize(
-		FPCGContext& InContext,
-		bool bUseMaterialOverrideAttributes,
-		const TArray<TSoftObjectPtr<UMaterialInterface>>& InStaticMaterialOverrides,
-		const TArray<FName>& InMaterialOverrideAttributeNames,
-		const UPCGMetadata* InMetadata);
-
-	// Use this constructor when you have common attribute usage or separate static overrides
-	void Initialize(
-		FPCGContext& InContext,
-		bool bInByAttributeOverride,
-		const TArray<FName>& InMaterialOverrideAttributeNames,
-		const UPCGMetadata* InMetadata);
-
-	void Reset();
-
-	bool IsInitialized() const { return bIsInitialized; }
-	bool IsValid() const { return bIsValid; }
-	bool OverridesMaterials() const { return bUseMaterialOverrideAttributes; }
-	const TArray<TSoftObjectPtr<UMaterialInterface>>& GetMaterialOverrides(PCGMetadataEntryKey EntryKey);
-
-private:
-	// Cached data
-	TArray<const FPCGMetadataAttributeBase*> MaterialAttributes;
-	TArray<TMap<PCGMetadataValueKey, TSoftObjectPtr<UMaterialInterface>>> ValueKeyToOverrideMaterials;
-	TArray<TSoftObjectPtr<UMaterialInterface>> WorkingMaterialOverrides;
-
-	// Data needed to perform operations
-	bool bIsInitialized = false;
-	bool bIsValid = false;
-	bool bUseMaterialOverrideAttributes = false;
-
-	TArray<TSoftObjectPtr<UMaterialInterface>> StaticMaterialOverrides;
-	TArray<FName> MaterialOverrideAttributeNames;
-	const UPCGMetadata* Metadata = nullptr;
-
-	void Initialize(FPCGContext& InContext);
-};
-
 UCLASS(Abstract, BlueprintType, ClassGroup = (Procedural))
 class PCG_API UPCGMeshSelectorBase : public UObject 
 {
diff --git a/Engine/Plugins/PCG/Source/PCG/Public/PCGManagedResource.h b/Engine/Plugins/PCG/Source/PCG/Public/PCGManagedResource.h
--- a/Engine/Plugins/PCG/Source/PCG/Public/PCGManagedResource.h
+++ b/Engine/Plugins/PCG/Source/PCG/Public/PCGManagedResource.h
@@ -8,10 +8,12 @@
 
 #include "Engine/SplineMeshComponentDescriptor.h"
 #include "ISMPartition/ISMComponentDescriptor.h"
+#include "Animation/AnimBank.h"
 
 #include "PCGManagedResource.generated.h"
 
 class UActorComponent;
+class UAnimBankMeshComponent;
 class UInstancedStaticMeshComponent;
 class USplineMeshComponent;
 
@@ -290,6 +292,59 @@ protected:
 	mutable UInstancedStaticMeshComponent* CachedRawComponentPtr = nullptr;
 };
 
+UCLASS(BlueprintType)
+class PCG_API UPCGManagedAnimBankMeshComponent : public UPCGManagedComponent
+{
+	GENERATED_BODY()
+
+public:
+	//~Begin UObject interface
+	virtual void PostLoad() override;
+	//~End UObject interface
+
+	//~Begin UPCGManagedResource interface
+	virtual bool ReleaseIfUnused(TSet<TSoftObjectPtr<AActor>>& OutActorsToDelete) override;
+	//~End UPCGManagedResource interface
+
+	//~Begin UPCGManagedComponents interface
+	virtual void ResetComponent() override;
+	virtual bool SupportsComponentReset() const override{ return true; }
+	virtual void MarkAsUsed() override;
+	virtual void MarkAsReused() override;
+	virtual void ForgetComponent() override;
+	//~End UPCGManagedComponents interface
+
+	UAnimBankMeshComponent* GetComponent() const;
+	void SetComponent(UAnimBankMeshComponent* InComponent);
+
+	void SetDescriptor(const FAnimBankMeshComponentDescriptor& InDescriptor);
+	const FAnimBankMeshComponentDescriptor& GetDescriptor() const { return Descriptor; }
+
+	void SetRootLocation(const FVector& InRootLocation);
+
+	FPCGCrc GetSettingsCrc() const { return SettingsCrc; }
+	void SetSettingsCrc(const FPCGCrc& InSettingsCrc) { SettingsCrc = InSettingsCrc; }
+
+protected:
+	UPROPERTY()
+	bool bHasDescriptor = false;
+
+	UPROPERTY()
+	FAnimBankMeshComponentDescriptor Descriptor;
+
+	UPROPERTY()
+	bool bHasRootLocation = false;
+
+	UPROPERTY()
+	FVector RootLocation = FVector::ZeroVector;
+
+	UPROPERTY()
+	FPCGCrc SettingsCrc;
+
+	// Cached raw pointer to anim bank mesh component
+	mutable UAnimBankMeshComponent* CachedRawComponentPtr = nullptr;
+};
+
 UCLASS(BlueprintType)
 class PCG_API UPCGManagedSplineMeshComponent : public UPCGManagedComponent
 {
diff --git a/Engine/Source/Runtime/Engine/Classes/Animation/AnimBank.h b/Engine/Source/Runtime/Engine/Classes/Animation/AnimBank.h
--- a/Engine/Source/Runtime/Engine/Classes/Animation/AnimBank.h
+++ b/Engine/Source/Runtime/Engine/Classes/Animation/AnimBank.h
@@ -22,6 +22,7 @@ struct FPropertyChangedEvent;
 
 struct FReferenceSkeleton;
 class USkinnedAsset;
+class UAnimBankMeshComponent;
 class FAnimBankBuildAsyncCacheTask;
 
 struct FSkinnedAssetMapping
@@ -272,22 +273,180 @@ private:
 #endif
 };
 
+struct FSoftAnimBankItem;
+
 USTRUCT(BlueprintType)
 struct FAnimBankItem
 {
 	GENERATED_USTRUCT_BODY()
 
 	UPROPERTY(EditAnywhere, BlueprintReadWrite, Category = Animation)
-	TObjectPtr<class UAnimBank> BankAsset;
+	TObjectPtr<UAnimBank> BankAsset = nullptr;
 
 	UPROPERTY(EditAnywhere, BlueprintReadWrite, Category = Animation)
-	int32 SequenceIndex;
+	int32 SequenceIndex = 0;
+
+	ENGINE_API FAnimBankItem();
+	ENGINE_API FAnimBankItem(const FAnimBankItem& InBankItem);
+	ENGINE_API explicit FAnimBankItem(const FSoftAnimBankItem& InBankItem);
 
-	FAnimBankItem()
+	ENGINE_API bool operator!=(const FAnimBankItem& Other) const;
+	ENGINE_API bool operator==(const FAnimBankItem& Other) const;
+};
+
+inline uint32 GetTypeHash(const FAnimBankItem& Key)
+{
+	return HashCombine(GetTypeHash(Key.BankAsset.Get()), GetTypeHash(Key.SequenceIndex));
+}
+
+inline uint32 GetTypeHash(const TArray<FAnimBankItem>& InBankItems)
+{
+	uint32 Hash = 0;
+	for (const FAnimBankItem& BankItem : InBankItems)
 	{
-		BankAsset		= nullptr;
-		SequenceIndex	= 0;
+		Hash = HashCombine(Hash, GetTypeHash(BankItem));
 	}
+	return Hash;
+}
+
+USTRUCT(BlueprintType)
+struct FSoftAnimBankItem
+{
+	GENERATED_USTRUCT_BODY()
+
+	UPROPERTY(EditAnywhere, BlueprintReadWrite, Category = Animation)
+	TSoftObjectPtr<UAnimBank> BankAsset = nullptr;
+
+	UPROPERTY(EditAnywhere, BlueprintReadWrite, Category = Animation)
+	int32 SequenceIndex = 0;
+
+	ENGINE_API FSoftAnimBankItem();
+	ENGINE_API FSoftAnimBankItem(const FSoftAnimBankItem& InBankItem);
+	ENGINE_API explicit FSoftAnimBankItem(const FAnimBankItem& InBankItem);
+
+	ENGINE_API bool operator!=(const FSoftAnimBankItem& Other) const;
+	ENGINE_API bool operator==(const FSoftAnimBankItem& Other) const;
+};
+
+inline uint32 GetTypeHash(const FSoftAnimBankItem& Key)
+{
+	return HashCombine(GetTypeHash(Key.BankAsset.Get()), GetTypeHash(Key.SequenceIndex));
+}
+
+inline uint32 GetTypeHash(const TArray<FSoftAnimBankItem>& InBankItems)
+{
+	uint32 Hash = 0;
+	for (const FSoftAnimBankItem& BankItem : InBankItems)
+	{
+		Hash = HashCombine(Hash, GetTypeHash(BankItem));
+	}
+	return Hash;
+}
+
+USTRUCT()
+struct FAnimBankMeshComponentDescriptor
+{
+	GENERATED_BODY()
+
+	ENGINE_API FAnimBankMeshComponentDescriptor();
+	ENGINE_API explicit FAnimBankMeshComponentDescriptor(ENoInit);
+	ENGINE_API explicit FAnimBankMeshComponentDescriptor(const FAnimBankMeshComponentDescriptor&);
+	ENGINE_API explicit FAnimBankMeshComponentDescriptor(const FSoftAnimBankMeshComponentDescriptor&);
+	ENGINE_API virtual ~FAnimBankMeshComponentDescriptor();
+
+	ENGINE_API UAnimBankMeshComponent* CreateComponent(UObject* Outer, FName Name = NAME_None, EObjectFlags ObjectFlags = EObjectFlags::RF_NoFlags) const;
+
+	ENGINE_API virtual void InitFrom(const UAnimBankMeshComponent* Component, bool bInitBodyInstance = true);
+	ENGINE_API virtual uint32 ComputeHash() const;
+	ENGINE_API virtual void InitComponent(UAnimBankMeshComponent* ISMComponent) const;
+
+	ENGINE_API void PostLoadFixup(UObject* Loader);
+
+	ENGINE_API bool operator!=(const FAnimBankMeshComponentDescriptor& Other) const;
+	ENGINE_API bool operator==(const FAnimBankMeshComponentDescriptor& Other) const;
+
+	friend inline uint32 GetTypeHash(const FAnimBankMeshComponentDescriptor& Key)
+	{
+		return Key.GetTypeHash();
+	}
+
+	uint32 GetTypeHash() const
+	{
+		if (Hash == 0)
+		{
+			ComputeHash();
+		}
+		return Hash;
+	}
+
+public:
+	UPROPERTY()
+	mutable uint32 Hash = 0;
+
+	UPROPERTY(EditAnywhere, Category = "Component Settings")
+	TEnumAsByte<EComponentMobility::Type> Mobility;
+
+	UPROPERTY(EditAnywhere, Category = "Component Settings")
+	TSubclassOf<UAnimBankMeshComponent> ComponentClass;
+
+	UPROPERTY(EditAnywhere, Category = "Component Settings", meta = (DisplayAfter = "ComponentClass"))
+	TObjectPtr<class USkinnedAsset> SkinnedAsset;
+
+	UPROPERTY(EditAnywhere, Category = "Component Settings", meta = (DisplayAfter = "SkinnedAsset"))
+	TArray<FAnimBankItem> BankItems;
+};
+
+USTRUCT()
+struct FSoftAnimBankMeshComponentDescriptor
+{
+	GENERATED_BODY()
+
+	ENGINE_API FSoftAnimBankMeshComponentDescriptor();
+	ENGINE_API explicit FSoftAnimBankMeshComponentDescriptor(ENoInit);
+	ENGINE_API explicit FSoftAnimBankMeshComponentDescriptor(const FAnimBankMeshComponentDescriptor&);
+	ENGINE_API explicit FSoftAnimBankMeshComponentDescriptor(const FSoftAnimBankMeshComponentDescriptor&);
+	ENGINE_API virtual ~FSoftAnimBankMeshComponentDescriptor();
+
+	ENGINE_API UAnimBankMeshComponent* CreateComponent(UObject* Outer, FName Name = NAME_None, EObjectFlags ObjectFlags = EObjectFlags::RF_NoFlags) const;
+
+	ENGINE_API virtual void InitFrom(const UAnimBankMeshComponent* Component, bool bInitBodyInstance = true);
+	ENGINE_API virtual uint32 ComputeHash() const;
+	ENGINE_API virtual void InitComponent(UAnimBankMeshComponent* ISMComponent) const;
+
+	ENGINE_API void PostLoadFixup(UObject* Loader);
+
+	ENGINE_API bool operator!=(const FSoftAnimBankMeshComponentDescriptor& Other) const;
+	ENGINE_API bool operator==(const FSoftAnimBankMeshComponentDescriptor& Other) const;
+
+	friend inline uint32 GetTypeHash(const FSoftAnimBankMeshComponentDescriptor& Key)
+	{
+		return Key.GetTypeHash();
+	}
+
+	uint32 GetTypeHash() const
+	{
+		if (Hash == 0)
+		{
+			ComputeHash();
+		}
+		return Hash;
+	}
+
+public:
+	UPROPERTY()
+	mutable uint32 Hash = 0;
+
+	UPROPERTY(EditAnywhere, Category = "Component Settings")
+	TEnumAsByte<EComponentMobility::Type> Mobility;
+
+	UPROPERTY(EditAnywhere, Category = "Component Settings")
+	TSubclassOf<UAnimBankMeshComponent> ComponentClass;
+
+	UPROPERTY(EditAnywhere, Category = "Component Settings", meta = (DisplayAfter = "ComponentClass"))
+	TSoftObjectPtr<class USkinnedAsset> SkinnedAsset;
+
+	UPROPERTY(EditAnywhere, Category = "Component Settings", meta = (DisplayAfter = "SkinnedAsset"))
+	TArray<FSoftAnimBankItem> BankItems;
 };
 
 struct FAnimBankDesc
diff --git a/Engine/Source/Runtime/Engine/Classes/Components/AnimBankMeshComponent.h b/Engine/Source/Runtime/Engine/Classes/Components/AnimBankMeshComponent.h
--- a/Engine/Source/Runtime/Engine/Classes/Components/AnimBankMeshComponent.h
+++ b/Engine/Source/Runtime/Engine/Classes/Components/AnimBankMeshComponent.h
@@ -182,7 +182,7 @@ public:
 	virtual int32 AddAnimBankItem(const FAnimBankItem& BankItem);
 
 	bool IsEnabled() const;
-	int32 GetNumInstances() const;
+	int32 GetInstanceCount() const;
 
 #if WITH_EDITOR
 	/** One bit per instance if the instance is selected. */
diff --git a/Engine/Source/Runtime/Engine/Private/Animation/AnimBank.cpp b/Engine/Source/Runtime/Engine/Private/Animation/AnimBank.cpp
--- a/Engine/Source/Runtime/Engine/Private/Animation/AnimBank.cpp
+++ b/Engine/Source/Runtime/Engine/Private/Animation/AnimBank.cpp
@@ -13,6 +13,7 @@
 #include "EngineUtils.h"
 #include "Interfaces/ITargetPlatform.h"
 #include "Interfaces/ITargetPlatformManagerModule.h"
+#include "Serialization/ArchiveCrc32.h"
 #include "Logging/MessageLog.h"
 #include "Misc/App.h"
 #include "ReferenceSkeleton.h"
@@ -37,6 +38,8 @@
 #include "Serialization/MemoryWriter.h"
 #endif
 
+#include "Components/AnimBankMeshComponent.h"
+
 #define LOCTEXT_NAMESPACE "AnimBank"
 
 CSV_DECLARE_CATEGORY_MODULE_EXTERN(ENGINE_API, Animation);
@@ -1340,5 +1343,212 @@ void BuildSkinnedAssetMapping(const USkinnedAsset& Asset, FSkinnedAssetMapping&
 
 } // UE::AnimBank
 
+FAnimBankItem::FAnimBankItem() = default;
+
+FAnimBankItem::FAnimBankItem(const FAnimBankItem& InBankItem)
+{
+	BankAsset		= InBankItem.BankAsset;
+	SequenceIndex	= InBankItem.SequenceIndex;
+}
+
+FAnimBankItem::FAnimBankItem(const FSoftAnimBankItem& InBankItem)
+{
+	BankAsset		= InBankItem.BankAsset.LoadSynchronous();
+	SequenceIndex	= InBankItem.SequenceIndex;
+}
+
+bool FAnimBankItem::operator!=(const FAnimBankItem& Other) const
+{
+	return !(*this == Other);
+}
+
+bool FAnimBankItem::operator==(const FAnimBankItem& Other) const
+{
+	return BankAsset == Other.BankAsset && SequenceIndex == Other.SequenceIndex;
+}
+
+FSoftAnimBankItem::FSoftAnimBankItem() = default;
+
+FSoftAnimBankItem::FSoftAnimBankItem(const FSoftAnimBankItem& InBankItem)
+{
+	BankAsset		= InBankItem.BankAsset.Get();
+	SequenceIndex	= InBankItem.SequenceIndex;
+}
+
+FSoftAnimBankItem::FSoftAnimBankItem(const FAnimBankItem& InBankItem)
+{
+	BankAsset		= InBankItem.BankAsset.Get();
+	SequenceIndex	= InBankItem.SequenceIndex;
+}
+
+bool FSoftAnimBankItem::operator!=(const FSoftAnimBankItem& Other) const
+{
+	return !(*this == Other);
+}
+
+bool FSoftAnimBankItem::operator==(const FSoftAnimBankItem& Other) const
+{
+	return BankAsset == Other.BankAsset && SequenceIndex == Other.SequenceIndex;
+}
+
+FAnimBankMeshComponentDescriptor::FAnimBankMeshComponentDescriptor()
+{
+	// Make sure we have proper defaults
+	InitFrom(UAnimBankMeshComponent::StaticClass()->GetDefaultObject<UAnimBankMeshComponent>());
+}
+
+FAnimBankMeshComponentDescriptor::FAnimBankMeshComponentDescriptor(ENoInit) {}
+FAnimBankMeshComponentDescriptor::FAnimBankMeshComponentDescriptor(const FAnimBankMeshComponentDescriptor& InDescriptor) = default;
+FAnimBankMeshComponentDescriptor::~FAnimBankMeshComponentDescriptor() = default;
+
+FAnimBankMeshComponentDescriptor::FAnimBankMeshComponentDescriptor(const FSoftAnimBankMeshComponentDescriptor& InDescriptor)
+{
+	SkinnedAsset = InDescriptor.SkinnedAsset.LoadSynchronous();
+	Algo::Transform(InDescriptor.BankItems, BankItems, [](const FSoftAnimBankItem& BankItem) { return FAnimBankItem(BankItem); });
+	// AB-TODO: Algo::Transform(InDescriptor.OverrideMaterials, OverrideMaterials, [](TSoftObjectPtr<UMaterialInterface> Material) { return Material.LoadSynchronous(); });
+	// AB-TODO: OverlayMaterial = InDescriptor.OverlayMaterial.LoadSynchronous();
+	// AB-TODO: Algo::Transform(InDescriptor.RuntimeVirtualTextures, RuntimeVirtualTextures, [](TSoftObjectPtr<URuntimeVirtualTexture> RVT) { return RVT.LoadSynchronous(); });
+	Hash = InDescriptor.Hash;
+}
+
+UAnimBankMeshComponent* FAnimBankMeshComponentDescriptor::CreateComponent(UObject* Outer, FName Name, EObjectFlags ObjectFlags) const
+{
+	UAnimBankMeshComponent* Component = NewObject<UAnimBankMeshComponent>(Outer, ComponentClass, Name, ObjectFlags);
+	InitComponent(Component);
+	return Component;
+}
+
+void FAnimBankMeshComponentDescriptor::InitFrom(const UAnimBankMeshComponent* Template, bool bInitBodyInstance)
+{
+	check(Template);
+
+	ComponentClass = Template->GetClass();
+	Mobility = Template->Mobility;
+	SkinnedAsset = Template->GetSkinnedAsset();
+	BankItems = Template->AnimBankItems;
+}
+
+uint32 FAnimBankMeshComponentDescriptor::ComputeHash() const
+{
+	FAnimBankMeshComponentDescriptor& MutableSelf = *const_cast<FAnimBankMeshComponentDescriptor*>(this);
+	Hash = 0; // we don't want the hash to impact the calculation
+	
+	FArchiveCrc32 CrcArchive;
+	CrcArchive << *this;
+	Hash = CrcArchive.GetCrc();
+
+	return Hash;
+}
+
+void FAnimBankMeshComponentDescriptor::InitComponent(UAnimBankMeshComponent* Component) const
+{
+	Component->SetSkinnedAsset(SkinnedAsset);
+	Component->Mobility = Mobility;
+	Component->AnimBankItems = BankItems;
+}
+
+void FAnimBankMeshComponentDescriptor::PostLoadFixup(UObject* Loader)
+{
+	check(Loader);
+}
+
+bool FAnimBankMeshComponentDescriptor::operator!=(const FAnimBankMeshComponentDescriptor& Other) const
+{
+	return !(*this == Other);
+}
+
+bool FAnimBankMeshComponentDescriptor::operator==(const FAnimBankMeshComponentDescriptor& Other) const
+{
+	return
+		(Hash == 0 || Other.Hash == 0 || Hash == Other.Hash) && // Check hash first, other checks are in case of Hash collision
+		SkinnedAsset == Other.SkinnedAsset &&
+		BankItems == Other.BankItems &&
+		ComponentClass == Other.ComponentClass &&
+		Mobility == Other.Mobility;
+}
+
+FSoftAnimBankMeshComponentDescriptor::FSoftAnimBankMeshComponentDescriptor()
+{
+	// Make sure we have proper defaults
+	InitFrom(UAnimBankMeshComponent::StaticClass()->GetDefaultObject<UAnimBankMeshComponent>());
+}
+
+FSoftAnimBankMeshComponentDescriptor::FSoftAnimBankMeshComponentDescriptor(ENoInit) {}
+FSoftAnimBankMeshComponentDescriptor::FSoftAnimBankMeshComponentDescriptor(const FSoftAnimBankMeshComponentDescriptor& InDescriptor) = default;
+FSoftAnimBankMeshComponentDescriptor::~FSoftAnimBankMeshComponentDescriptor() = default;
+
+FSoftAnimBankMeshComponentDescriptor::FSoftAnimBankMeshComponentDescriptor(const FAnimBankMeshComponentDescriptor& InDescriptor)
+{
+	SkinnedAsset = InDescriptor.SkinnedAsset;
+	Algo::Transform(InDescriptor.BankItems, BankItems, [](const FAnimBankItem& BankItem) { return FSoftAnimBankItem(BankItem); });
+	// AB-TODO: Algo::Transform(InDescriptor.OverrideMaterials, OverrideMaterials, [](TObjectPtr<UMaterialInterface> Material) { return Material; });
+	// AB-TODO: OverlayMaterial = InDescriptor.OverlayMaterial;
+	// AB-TODO: Algo::Transform(InDescriptor.RuntimeVirtualTextures, RuntimeVirtualTextures, [](TObjectPtr<URuntimeVirtualTexture> RVT) { return RVT; });
+	Hash = InDescriptor.Hash;
+}
+
+UAnimBankMeshComponent* FSoftAnimBankMeshComponentDescriptor::CreateComponent(UObject* Outer, FName Name, EObjectFlags ObjectFlags) const
+{
+	UAnimBankMeshComponent* Component = NewObject<UAnimBankMeshComponent>(Outer, ComponentClass, Name, ObjectFlags);
+	InitComponent(Component);
+	return Component;
+}
+
+void FSoftAnimBankMeshComponentDescriptor::InitFrom(const UAnimBankMeshComponent* Template, bool bInitBodyInstance)
+{
+	check(Template);
+
+	ComponentClass = Template->GetClass();
+	Mobility = Template->Mobility;
+	SkinnedAsset = Template->GetSkinnedAsset();
+	BankItems.Reset(Template->AnimBankItems.Num());
+	for (const FAnimBankItem& BankItem : Template->AnimBankItems)
+	{
+		BankItems.Emplace(BankItem);
+	}
+}
+
+uint32 FSoftAnimBankMeshComponentDescriptor::ComputeHash() const
+{
+	Hash = 0; // we don't want the hash to impact the calculation
+	
+	FArchiveCrc32 CrcArchive;
+	CrcArchive << *this;
+	Hash = CrcArchive.GetCrc();
+
+	return Hash;
+}
+
+void FSoftAnimBankMeshComponentDescriptor::InitComponent(UAnimBankMeshComponent* Component) const
+{
+	Component->SetSkinnedAsset(SkinnedAsset.LoadSynchronous());
+	Component->Mobility = Mobility;
+	Component->AnimBankItems.Reset(BankItems.Num());
+	for (const FSoftAnimBankItem& BankItem : BankItems)
+	{
+		Component->AnimBankItems.Emplace(BankItem);
+	}
+}
+
+void FSoftAnimBankMeshComponentDescriptor::PostLoadFixup(UObject* Loader)
+{
+	check(Loader);
+}
+
+bool FSoftAnimBankMeshComponentDescriptor::operator!=(const FSoftAnimBankMeshComponentDescriptor& Other) const
+{
+	return !(*this == Other);
+}
+
+bool FSoftAnimBankMeshComponentDescriptor::operator==(const FSoftAnimBankMeshComponentDescriptor& Other) const
+{
+	return
+		(Hash == 0 || Other.Hash == 0 || Hash == Other.Hash) && // Check hash first, other checks are in case of Hash collision
+		SkinnedAsset == Other.SkinnedAsset &&
+		BankItems == Other.BankItems &&
+		ComponentClass == Other.ComponentClass &&
+		Mobility == Other.Mobility;
+}
+
 #undef LOCTEXT_NAMESPACE
 
diff --git a/Engine/Source/Runtime/Engine/Private/Components/AnimBankMeshComponent.cpp b/Engine/Source/Runtime/Engine/Private/Components/AnimBankMeshComponent.cpp
--- a/Engine/Source/Runtime/Engine/Private/Components/AnimBankMeshComponent.cpp
+++ b/Engine/Source/Runtime/Engine/Private/Components/AnimBankMeshComponent.cpp
@@ -79,7 +79,7 @@ public:
 		ProxyData.Flags.bHasPerInstanceEditorData = false;// GIsEditor != 0 && HasPerInstanceHitProxies();
 	#endif
 
-		const int32 InstanceCount = InComponent->GetNumInstances();
+		const int32 InstanceCount = InComponent->GetInstanceCount();
 
 		InstanceTransforms.Reset(InstanceCount);
 		if (ProxyData.Flags.bHasPerInstanceLocalBounds)
@@ -799,13 +799,13 @@ bool UAnimBankMeshComponent::IsEnabled() const
 	USkeletalMesh* SkeletalMeshPtr = Cast<USkeletalMesh>(GetSkinnedAsset());
 	if (SkeletalMeshPtr && SkeletalMeshPtr->GetResourceForRendering())
 	{
-		return GetNumInstances() > 0;
+		return GetInstanceCount() > 0;
 	}
 
 	return false;
 }
 
-int32 UAnimBankMeshComponent::GetNumInstances() const
+int32 UAnimBankMeshComponent::GetInstanceCount() const
 {
 	return InstanceData.Num();
 }
@@ -1341,7 +1341,7 @@ void UAnimBankMeshComponent::ClearInstances()
 	DeletionState = EInstanceDeletionReason::Clearing;
 #endif
 
-	const int32 PrevNumInstances = GetNumInstances();
+	const int32 PrevNumInstances = GetInstanceCount();
 
 	// Clear all the per-instance data
 	InstanceData.Empty();
