[PCG] Fix Cancel leak when pcg.GraphMultithreading is enabled
- Prevent scheduling while cancelling

#rnx

diff --git a/Engine/Plugins/PCG/Source/PCG/Private/Graph/PCGGraphExecutor.cpp b/Engine/Plugins/PCG/Source/PCG/Private/Graph/PCGGraphExecutor.cpp
--- a/Engine/Plugins/PCG/Source/PCG/Private/Graph/PCGGraphExecutor.cpp
+++ b/Engine/Plugins/PCG/Source/PCG/Private/Graph/PCGGraphExecutor.cpp
@@ -353,6 +353,13 @@ TSet<UPCGComponent*> FPCGGraphExecutor::Cancel(TFunctionRef<bool(TWeakObjectPtr<
 	check(IsInGameThread());
 	TSet<UPCGComponent*> CancelledComponents;
 	const FPCGGraphExecutor::EExecuteVersion CurrentExecuteVersion = GetExecuteVersion();
+	
+	// Prevent any scheduling while cancelling
+	SchedulingDisabled = true;
+	ON_SCOPE_EXIT
+	{
+		SchedulingDisabled = false;
+	};
 
 	// Visit scheduled tasks
 	{
@@ -1291,6 +1298,12 @@ FPCGGraphExecutor::EExecuteVersion FPCGGraphExecutor::GetExecuteVersion() const
 
 bool FPCGGraphExecutor::ExecuteScheduling(double EndTime, TSharedPtr<FPCGGraphActiveTask>* OutMainThreadTask, bool bForceCheckSleepingTasks)
 {
+	// Can currently be disabled by cancellation
+	if (SchedulingDisabled)
+	{
+		return false;
+	}
+
 	bool bStateChanged = false;
 	const bool bIsInGameThread = IsInGameThread();
 	TRACE_CPUPROFILER_EVENT_SCOPE(FPCGGraphExecutor::ExecuteScheduling);
diff --git a/Engine/Plugins/PCG/Source/PCG/Private/Graph/PCGGraphExecutor.h b/Engine/Plugins/PCG/Source/PCG/Private/Graph/PCGGraphExecutor.h
--- a/Engine/Plugins/PCG/Source/PCG/Private/Graph/PCGGraphExecutor.h
+++ b/Engine/Plugins/PCG/Source/PCG/Private/Graph/PCGGraphExecutor.h
@@ -262,6 +262,7 @@ private:
 	/** Lock level - 1 (top most lock) */
 	UE::FSpinLock ScheduleLock;
 	TArray<FPCGGraphScheduleTask> ScheduledTasks;
+	std::atomic<bool> SchedulingDisabled = false;
 
 	/** Lock level 2 */
 	mutable UE::FSpinLock TasksLock;
