[PCG] Fix two "not properly initialized" errors occurring in FPCGProceduralISMComponentDescriptor

#rnx

diff --git a/Engine/Plugins/PCG/Source/PCG/Private/Components/PCGProceduralISMComponent.cpp b/Engine/Plugins/PCG/Source/PCG/Private/Components/PCGProceduralISMComponent.cpp
--- a/Engine/Plugins/PCG/Source/PCG/Private/Components/PCGProceduralISMComponent.cpp
+++ b/Engine/Plugins/PCG/Source/PCG/Private/Components/PCGProceduralISMComponent.cpp
@@ -110,6 +110,8 @@ UPCGManagedProceduralISMComponent* PCGManagedProceduralISMComponent::GetOrCreate
 }
 
 FPCGProceduralISMComponentDescriptor::FPCGProceduralISMComponentDescriptor()
+	: NumInstances(0)
+	, NumCustomFloats(0)
 {
 	*this = FPCGSoftISMComponentDescriptor();
 }
