[PCG] Fix unique string keys being accumulated in a non-unique way, resulting in duplicate keys.

PISMC setup pushes the keys through a TSet so the duplicated keys were not generating duplicated PISMCs.

#rnx

diff --git a/Engine/Plugins/PCG/Source/PCG/Private/Compute/Elements/PCGComputeGraphElement.cpp b/Engine/Plugins/PCG/Source/PCG/Private/Compute/Elements/PCGComputeGraphElement.cpp
--- a/Engine/Plugins/PCG/Source/PCG/Private/Compute/Elements/PCGComputeGraphElement.cpp
+++ b/Engine/Plugins/PCG/Source/PCG/Private/Compute/Elements/PCGComputeGraphElement.cpp
@@ -414,7 +414,7 @@ bool FPCGComputeGraphElement::SetupProceduralISMComponents(FPCGContext* InContex
 						}
 
 						SelectorAttributeId = AttributeDesc.Index;
-						UniqueStringKeys.Append(AttributeDesc.UniqueStringKeys);
+						UniqueStringKeys.Append(AttributeDesc.GetUniqueStringKeys());
 					}
 				}
 			}
diff --git a/Engine/Plugins/PCG/Source/PCG/Private/Compute/Elements/PCGCustomHLSL.cpp b/Engine/Plugins/PCG/Source/PCG/Private/Compute/Elements/PCGCustomHLSL.cpp
--- a/Engine/Plugins/PCG/Source/PCG/Private/Compute/Elements/PCGCustomHLSL.cpp
+++ b/Engine/Plugins/PCG/Source/PCG/Private/Compute/Elements/PCGCustomHLSL.cpp
@@ -422,7 +422,7 @@ bool UPCGCustomHLSLSettings::ComputeOutputPinDataDesc(const UPCGPin* OutputPin,
 								const bool bIsString = InputAttributeDesc.Type == EPCGKernelAttributeType::StringKey || InputAttributeDesc.Type == EPCGKernelAttributeType::Name;
 								if (bIsString && InputAttributeDesc.Name == AttributeDesc.Name)
 								{
-									AttributeDesc.UniqueStringKeys.Append(InputAttributeDesc.UniqueStringKeys);
+									AttributeDesc.AddUniqueStringKeys(InputAttributeDesc.GetUniqueStringKeys());
 									bFoundMatchingAttribute = true;
 									break;
 								}
@@ -444,7 +444,7 @@ bool UPCGCustomHLSLSettings::ComputeOutputPinDataDesc(const UPCGPin* OutputPin,
 								{
 									if (InputAttributeDesc.Type == EPCGKernelAttributeType::StringKey || InputAttributeDesc.Type == EPCGKernelAttributeType::Name)
 									{
-										AttributeDesc.UniqueStringKeys.Append(InputAttributeDesc.UniqueStringKeys);
+										AttributeDesc.AddUniqueStringKeys(InputAttributeDesc.GetUniqueStringKeys());
 									}
 								}
 							}
diff --git a/Engine/Plugins/PCG/Source/PCG/Private/Compute/PCGDataBinding.cpp b/Engine/Plugins/PCG/Source/PCG/Private/Compute/PCGDataBinding.cpp
--- a/Engine/Plugins/PCG/Source/PCG/Private/Compute/PCGDataBinding.cpp
+++ b/Engine/Plugins/PCG/Source/PCG/Private/Compute/PCGDataBinding.cpp
@@ -215,7 +215,7 @@ void UPCGDataBinding::DebugLogDataDescriptions()
 				const FPCGKernelAttributeDesc& AttributeDesc = DataDesc.AttributeDescs[AttributeIndex];
 
 				FString UniqueStringKeys;
-				for (int32 StringKey : AttributeDesc.UniqueStringKeys)
+				for (int32 StringKey : AttributeDesc.GetUniqueStringKeys())
 				{
 					UniqueStringKeys += TEXT(", ");
 					UniqueStringKeys += FString::Format(TEXT("{0}"), { StringKey });
diff --git a/Engine/Plugins/PCG/Source/PCG/Private/Compute/PCGDataForGPU.cpp b/Engine/Plugins/PCG/Source/PCG/Private/Compute/PCGDataForGPU.cpp
--- a/Engine/Plugins/PCG/Source/PCG/Private/Compute/PCGDataForGPU.cpp
+++ b/Engine/Plugins/PCG/Source/PCG/Private/Compute/PCGDataForGPU.cpp
@@ -673,6 +673,14 @@ void FPCGKernelAttributeTable::DebugLog()
 }
 #endif
 
+void FPCGKernelAttributeDesc::AddUniqueStringKeys(const TArray<int32>& InOtherStringKeys)
+{
+	for (int32 OtherStringKey : InOtherStringKeys)
+	{
+		UniqueStringKeys.AddUnique(OtherStringKey);
+	}
+}
+
 bool FPCGKernelAttributeDesc::operator==(const FPCGKernelAttributeDesc& Other) const
 {
 	return Index == Other.Index && Type == Other.Type && Name == Other.Name;
diff --git a/Engine/Plugins/PCG/Source/PCG/Private/Compute/PCGDataForGPU.h b/Engine/Plugins/PCG/Source/PCG/Private/Compute/PCGDataForGPU.h
--- a/Engine/Plugins/PCG/Source/PCG/Private/Compute/PCGDataForGPU.h
+++ b/Engine/Plugins/PCG/Source/PCG/Private/Compute/PCGDataForGPU.h
@@ -92,6 +92,7 @@ protected:
 /** Data description for a metadata attribute. Stores identifying name and type as well as the unique attribute ID. */
 struct FPCGKernelAttributeDesc
 {
+public:
 	FPCGKernelAttributeDesc() = default;
 
 	explicit FPCGKernelAttributeDesc(int32 InIndex, EPCGKernelAttributeType InType, FName InName)
@@ -109,6 +110,13 @@ struct FPCGKernelAttributeDesc
 	{
 	}
 
+	const TArray<int32>& GetUniqueStringKeys() const { return UniqueStringKeys; }
+
+	void AddUniqueStringKeys(const TArray<int32>& InOtherStringKeys);
+
+	bool operator==(const FPCGKernelAttributeDesc& Other) const;
+
+public:
 	// Attribute Id
 	int32 Index = INDEX_NONE;
 
@@ -116,10 +124,9 @@ struct FPCGKernelAttributeDesc
 
 	FName Name = NAME_None;
 
+private:
 	/* All possible string keys arriving on this attribute (string keys are indices into the string table in the data binding). */
 	TArray<int32> UniqueStringKeys;
-
-	bool operator==(const FPCGKernelAttributeDesc& Other) const;
 };
 
 /** Data description for a single data object (UPCGData). */
diff --git a/Engine/Plugins/PCG/Source/PCG/Private/Elements/PCGAnimBankMeshSpawner.cpp b/Engine/Plugins/PCG/Source/PCG/Private/Elements/PCGAnimBankMeshSpawner.cpp
--- a/Engine/Plugins/PCG/Source/PCG/Private/Elements/PCGAnimBankMeshSpawner.cpp
+++ b/Engine/Plugins/PCG/Source/PCG/Private/Elements/PCGAnimBankMeshSpawner.cpp
@@ -161,7 +161,7 @@ bool UPCGAnimBankMeshSpawnerSettings::ComputeOutputPinDataDesc(const UPCGPin* Ou
 					{
 						if (InputAttrDesc.Name == SelectorByAttribute->AttributeName)
 						{
-							for (int32 InputStringKey : InputAttrDesc.UniqueStringKeys)
+							for (int32 InputStringKey : InputAttrDesc.GetUniqueStringKeys())
 							{
 								UniqueStringKeys.AddUnique(InputStringKey);
 							}
diff --git a/Engine/Plugins/PCG/Source/PCG/Private/Elements/PCGStaticMeshSpawner.cpp b/Engine/Plugins/PCG/Source/PCG/Private/Elements/PCGStaticMeshSpawner.cpp
--- a/Engine/Plugins/PCG/Source/PCG/Private/Elements/PCGStaticMeshSpawner.cpp
+++ b/Engine/Plugins/PCG/Source/PCG/Private/Elements/PCGStaticMeshSpawner.cpp
@@ -152,7 +152,7 @@ bool UPCGStaticMeshSpawnerSettings::ComputeOutputPinDataDesc(const UPCGPin* Outp
 					{
 						if (InputAttrDesc.Name == SelectorByAttribute->AttributeName)
 						{
-							for (int32 InputStringKey : InputAttrDesc.UniqueStringKeys)
+							for (int32 InputStringKey : InputAttrDesc.GetUniqueStringKeys())
 							{
 								UniqueStringKeys.AddUnique(InputStringKey);
 							}
