[PCG] GPU SM Spawner - fix remaining missing instances that could occur for some seed values and low instance counts by switching to a confidence interval calculation.

This has the additional benefit of massively reducing overallocation for large point sizes, as the law of large numbers dictates there will be much less relative variance.

Fix initialization of PISMC descriptor properties.

#rnx

diff --git a/Engine/Plugins/PCG/Source/PCG/Private/Compute/Elements/PCGComputeGraphElement.cpp b/Engine/Plugins/PCG/Source/PCG/Private/Compute/Elements/PCGComputeGraphElement.cpp
--- a/Engine/Plugins/PCG/Source/PCG/Private/Compute/Elements/PCGComputeGraphElement.cpp
+++ b/Engine/Plugins/PCG/Source/PCG/Private/Compute/Elements/PCGComputeGraphElement.cpp
@@ -528,9 +528,24 @@ bool FPCGComputeGraphElement::SetupProceduralISMComponents(FPCGContext* InContex
 					Descriptor.NumCustomFloats = CustomFloatCount;
 					Descriptor.StaticMesh = StaticMesh;
 
-					// Heuristic: The random draw can produce more or less values than expectation. 4% buffer was sufficient to cover a test case in the PCG test graph.
-					constexpr float VarianceFactor = 1.04f;
-					Descriptor.NumInstances = FMath::CeilToInt(InputPointCount * Weight * VarianceFactor);
+					uint32 InstanceCount = FMath::CeilToInt(InputPointCount * Weight);
+
+					if (SelectorWeighted->MeshEntries.Num() > 1)
+					{
+						// Since we'll be selecting meshes based on random draws using the point random seeds which we don't have on CPU,
+						// we may pick more or less than the expected number of instances for each mesh. Use binomial variance to calculate
+						// the overallocation that gives 99.7% confidence (3 sigma).
+						const double Variance = double(InputPointCount) * Weight * (1.0f - Weight);
+						const double Sigma = FMath::Sqrt(Variance);
+						const uint32 AdditionalAllocation = FMath::CeilToInt(3.0f * Sigma);
+
+						// Useful for debugging instance counts.
+						//UE_LOG(LogPCG, Warning, TEXT("Over allocation, N = %d, P = %.2f, NP = %d, 3sig = %f, Extra = %d (%.2f%%), Calculated = %d, Final = %d"), InputPointCount, Weight, InstanceCount, float(3.0f * Sigma), AdditionalAllocation, float(AdditionalAllocation) / InputPointCount, InstanceCount + AdditionalAllocation, FMath::Min(InstanceCount + AdditionalAllocation, InputPointCount));
+
+						InstanceCount += AdditionalAllocation;
+					}
+
+					Descriptor.NumInstances = FMath::Min(InstanceCount, InputPointCount);
 
 					ComponentsToCreate.Add(MoveTemp(Descriptor));
 
@@ -593,6 +608,9 @@ bool FPCGComputeGraphElement::SetupProceduralISMComponents(FPCGContext* InContex
 
 			bAnyComponentsSetup = true;
 		}
+
+		// Useful for debugging instance counts.
+		//UE_LOG(LogPCG, Log, TEXT("Input points %d, total instance count %d (+ %.2f%%)"), InputPointCount, Primitives.NumInstancesAllPrimitives, (float(Primitives.NumInstancesAllPrimitives - InputPointCount) / InputPointCount) * 100.0f);
 	}
 
 	return bAnyComponentsSetup;
