[PCG] Fix crash in compute graph data tracking when unsupported data types manage to bypass the pin types.

Prefilter data incident on compute graph to enforce type support.

#rnx

diff --git a/Engine/Plugins/PCG/Source/PCG/Private/Compute/PCGDataBinding.cpp b/Engine/Plugins/PCG/Source/PCG/Private/Compute/PCGDataBinding.cpp
--- a/Engine/Plugins/PCG/Source/PCG/Private/Compute/PCGDataBinding.cpp
+++ b/Engine/Plugins/PCG/Source/PCG/Private/Compute/PCGDataBinding.cpp
@@ -57,7 +57,19 @@ void UPCGDataBinding::InitializeInputData(const FPCGDataCollection& InComputeGra
 {
 	TRACE_CPUPROFILER_EVENT_SCOPE(UPCGDataBinding::InitializeInputData);
 
-	InputDataCollection = InComputeGraphElementInputData;
+	// Add compatible data objects from input collection.
+	InputDataCollection.TaggedData.Reserve(InComputeGraphElementInputData.TaggedData.Num());
+	for (const FPCGTaggedData& InputData : InComputeGraphElementInputData.TaggedData)
+	{
+		if (InputData.Data && PCGComputeHelpers::IsTypeAllowedAsInput(InputData.Data->GetDataType()))
+		{
+			InputDataCollection.TaggedData.Add(InputData);
+		}
+		else if (InputData.Data)
+		{
+			UE_LOG(LogPCG, Warning, TEXT("Stripped input data that is not currently supported by GPU execution: %s"), *InputData.Data->GetName());
+		}
+	}
 
 	// Link each input pin to the data collection, so that data providers can find the data.
 	for (const TSoftObjectPtr<const UPCGPin>& InputPinPtr : Graph->PinsReceivingDataFromCPU)
