[PCG] Fix invalid density literal unsupported for some platforms

#rnx

diff --git a/Engine/Plugins/PCG/Shaders/Private/PCGDataCollectionDataInterface.ush b/Engine/Plugins/PCG/Shaders/Private/PCGDataCollectionDataInterface.ush
--- a/Engine/Plugins/PCG/Shaders/Private/PCGDataCollectionDataInterface.ush
+++ b/Engine/Plugins/PCG/Shaders/Private/PCGDataCollectionDataInterface.ush
@@ -32,7 +32,9 @@ RWByteAddressBuffer {DataInterfaceName}_DataCollectionBuffer;
 #define POINT_STEEPNESS_ID 8
 #define POINT_METADATA_ENTRY_ID 9
 
-#define INVALID_DENSITY 1.#INF
+// Used to represent invalid/removed points. We use a value slightly less than max float,
+// as not all platforms support float infinity in shaders.
+#define INVALID_DENSITY 3.402823e+38f
 
 #define DATA_COLLECTION_HEADER_SIZE_BYTES 4
 #define DATA_HEADER_SIZE_BYTES 1036
diff --git a/Engine/Plugins/PCG/Source/PCG/Private/Compute/PCGComputeCommon.h b/Engine/Plugins/PCG/Source/PCG/Private/Compute/PCGComputeCommon.h
--- a/Engine/Plugins/PCG/Source/PCG/Private/Compute/PCGComputeCommon.h
+++ b/Engine/Plugins/PCG/Source/PCG/Private/Compute/PCGComputeCommon.h
@@ -60,6 +60,10 @@ namespace PCGComputeConstants
 	constexpr uint32 MAX_PRIMITIVE_COMPONENTS_PER_SPAWNER = 256;
 
 	constexpr uint32 THREAD_GROUP_SIZE = 64;
+
+	// Used to represent invalid/removed points. We use a value slightly less than max float,
+	// as not all platforms support float infinity in shaders.
+	constexpr float INVALID_DENSITY = 3.402823e+38f;
 }
 
 namespace PCGComputeHelpers
diff --git a/Engine/Plugins/PCG/Source/PCG/Private/Compute/PCGDataForGPU.cpp b/Engine/Plugins/PCG/Source/PCG/Private/Compute/PCGDataForGPU.cpp
--- a/Engine/Plugins/PCG/Source/PCG/Private/Compute/PCGDataForGPU.cpp
+++ b/Engine/Plugins/PCG/Source/PCG/Private/Compute/PCGDataForGPU.cpp
@@ -1426,7 +1426,7 @@ EPCGUnpackDataCollectionResult FPCGDataCollectionDesc::UnpackDataCollection(FPCG
 			// TODO: It may be more efficient to create a mapping from input point index to final output point index and do everything in one pass.
 			auto DiscardInvalidPoints = [&OutPoints](int32 Index, FPCGPoint& OutPoint) -> bool
 			{
-				if (!FMath::IsFinite(OutPoints[Index].Density))
+				if (OutPoints[Index].Density == PCGComputeConstants::INVALID_DENSITY)
 				{
 					return false;
 				}
