[PCG] Fix GetColor declaration in Custom HLSL node

#rnx

diff --git a/Engine/Plugins/PCG/Source/PCG/Private/Compute/Elements/PCGCustomHLSL.cpp b/Engine/Plugins/PCG/Source/PCG/Private/Compute/Elements/PCGCustomHLSL.cpp
--- a/Engine/Plugins/PCG/Source/PCG/Private/Compute/Elements/PCGCustomHLSL.cpp
+++ b/Engine/Plugins/PCG/Source/PCG/Private/Compute/Elements/PCGCustomHLSL.cpp
@@ -786,7 +786,7 @@ void UPCGCustomHLSLSettings::UpdateInputDeclarations()
 			"float3 {0}_GetScale(uint DataIndex, uint ElementIndex);\n"
 			"float3 {0}_GetBoundsMin(uint DataIndex, uint ElementIndex);\n"
 			"float3 {0}_GetBoundsMax(uint DataIndex, uint ElementIndex);\n"
-			"float3 {0}_GetColor(uint DataIndex, uint ElementIndex);\n"
+			"float4 {0}_GetColor(uint DataIndex, uint ElementIndex);\n"
 			"float {0}_GetDensity(uint DataIndex, uint ElementIndex);\n"
 			"int {0}_GetSeed(uint DataIndex, uint ElementIndex);\n"
 			"float {0}_GetSteepness(uint DataIndex, uint ElementIndex);\n"
diff --git a/Engine/Plugins/PCG/Source/PCG/Private/Compute/PCGDataForGPU.cpp b/Engine/Plugins/PCG/Source/PCG/Private/Compute/PCGDataForGPU.cpp
--- a/Engine/Plugins/PCG/Source/PCG/Private/Compute/PCGDataForGPU.cpp
+++ b/Engine/Plugins/PCG/Source/PCG/Private/Compute/PCGDataForGPU.cpp
@@ -1042,6 +1042,13 @@ void FPCGDataCollectionDesc::PackDataCollection(const FPCGDataCollection& InData
 	for (int32 DataIndex = 0; DataIndex < InputData.Num(); ++DataIndex)
 	{
 		const FPCGDataDesc& DataDesc = DataDescs[DataIndex];
+
+		// No work to do if there are no elements to process.
+		if (DataDesc.ElementCount <= 0)
+		{
+			continue;
+		}
+
 		const UPCGMetadata* Metadata = InputData[DataIndex].Data ? InputData[DataIndex].Data->ConstMetadata() : nullptr;
 		if (!ensure(Metadata))
 		{
@@ -1287,6 +1294,12 @@ EPCGUnpackDataCollectionResult FPCGDataCollectionDesc::UnpackDataCollection(FPCG
 			OutTaggedData.Data = OutPointData;
 			OutTaggedData.Pin = InPin;
 
+			// No work to do if there are no elements to process.
+			if (NumElements == 0)
+			{
+				continue;
+			}
+
 			// Loop over attributes.
 			for (const FPCGKernelAttributeDesc& AttributeDesc : AttributeDescs)
 			{
@@ -1465,6 +1478,12 @@ EPCGUnpackDataCollectionResult FPCGDataCollectionDesc::UnpackDataCollection(FPCG
 			OutTaggedData.Data = OutParamData;
 			OutTaggedData.Pin = InPin;
 
+			// No work to do if there are no elements to process.
+			if (NumElements == 0)
+			{
+				continue;
+			}
+
 			// Loop over attributes.
 			for (const FPCGKernelAttributeDesc& AttributeDesc : AttributeDescs)
 			{
