[PCG] Prevent creating Partition Actors inside a loading callstack

#rnx

diff --git a/Engine/Plugins/PCG/Source/PCG/Private/PCGSubsystem.cpp b/Engine/Plugins/PCG/Source/PCG/Private/PCGSubsystem.cpp
--- a/Engine/Plugins/PCG/Source/PCG/Private/PCGSubsystem.cpp
+++ b/Engine/Plugins/PCG/Source/PCG/Private/PCGSubsystem.cpp
@@ -1832,13 +1832,13 @@ void UPCGSubsystem::CreateMissingPartitionActors()
 void UPCGSubsystem::CreatePartitionActorsWithinBounds(UPCGComponent* InComponent, const FBox& InBounds, const PCGHiGenGrid::FSizeArray& InGridSizes)
 {
 	UWorld* World = GetWorld();
-	if (!PCGHelpers::IsRuntimeOrPIE() && !IsPartitionActorCreationDisabledForWorld(World))
+	if (!PCGHelpers::IsRuntimeOrPIE() && !UE::GetIsEditorLoadingPackage() && !IsPartitionActorCreationDisabledForWorld(World))
 	{
 		// We can't spawn actors if we are running constructions scripts, asserting when we try to get the actor with the WP API.
 		// We should never enter this if we are in a construction script. If the ensure is hit, we need to fix it.
 		if (ensure(World && !World->bIsRunningConstructionScript))
 		{
-			ForAllOverlappingCells(InComponent, InBounds, InGridSizes, true, {}, [](APCGPartitionActor*, const FBox&) { return InvalidPCGTaskId; });
+			ForAllOverlappingCells(InComponent, InBounds, InGridSizes, /*bCanCreateActor=*/true, {}, [](APCGPartitionActor*, const FBox&) { return InvalidPCGTaskId; });
 		}
 	}
 }
