[PCG] Replace NewObject with NewObject_AnyThread in GPU data unpacking

#rnx

diff --git a/Engine/Plugins/PCG/Source/PCG/Private/Compute/DataInterfaces/PCGDataCollectionDataInterface.cpp b/Engine/Plugins/PCG/Source/PCG/Private/Compute/DataInterfaces/PCGDataCollectionDataInterface.cpp
--- a/Engine/Plugins/PCG/Source/PCG/Private/Compute/DataInterfaces/PCGDataCollectionDataInterface.cpp
+++ b/Engine/Plugins/PCG/Source/PCG/Private/Compute/DataInterfaces/PCGDataCollectionDataInterface.cpp
@@ -536,6 +536,7 @@ bool UPCGDataCollectionDataProvider::ProcessReadBackData(FPCGComputeGraphContext
 	TRACE_CPUPROFILER_EVENT_SCOPE(UPCGDataCollectionDataProvider::ProcessReadBackData);
 	check(InContext);
 	check(ProducerSettings);
+	check(IsInGameThread());
 
 	if (!ensure(bReadbackComplete))
 	{
diff --git a/Engine/Plugins/PCG/Source/PCG/Private/Compute/PCGDataForGPU.cpp b/Engine/Plugins/PCG/Source/PCG/Private/Compute/PCGDataForGPU.cpp
--- a/Engine/Plugins/PCG/Source/PCG/Private/Compute/PCGDataForGPU.cpp
+++ b/Engine/Plugins/PCG/Source/PCG/Private/Compute/PCGDataForGPU.cpp
@@ -2,6 +2,7 @@
 
 #include "Compute/PCGDataForGPU.h"
 
+#include "PCGContext.h"
 #include "PCGData.h"
 #include "PCGEdge.h"
 #include "PCGNode.h"
@@ -1257,7 +1258,7 @@ EPCGUnpackDataCollectionResult FPCGDataCollectionDesc::UnpackDataCollection(FPCG
 		{
 			TRACE_CPUPROFILER_EVENT_SCOPE(UnpackPointDataItem);
 
-			UPCGPointData* OutPointData = NewObject<UPCGPointData>();
+			UPCGPointData* OutPointData = FPCGContext::NewObject_AnyThread<UPCGPointData>(InContext);
 			UPCGMetadata* Metadata = OutPointData->MutableMetadata();
 			TArray<FPCGPoint>& OutPoints = OutPointData->GetMutablePoints();
 
@@ -1457,7 +1458,7 @@ EPCGUnpackDataCollectionResult FPCGDataCollectionDesc::UnpackDataCollection(FPCG
 		{
 			TRACE_CPUPROFILER_EVENT_SCOPE(UnpackParamDataItem);
 
-			UPCGParamData* OutParamData = NewObject<UPCGParamData>();
+			UPCGParamData* OutParamData = FPCGContext::NewObject_AnyThread<UPCGParamData>(InContext);
 			UPCGMetadata* Metadata = OutParamData->MutableMetadata();
 
 			{
