enable autortfm instrumentation on the PCG module

diff --git a/Engine/Plugins/PCG/Source/PCG/PCG.Build.cs b/Engine/Plugins/PCG/Source/PCG/PCG.Build.cs
--- a/Engine/Plugins/PCG/Source/PCG/PCG.Build.cs
+++ b/Engine/Plugins/PCG/Source/PCG/PCG.Build.cs
@@ -44,7 +44,6 @@ namespace UnrealBuildTool.Rules
 					});
 			}
 
-			bDisableAutoRTFMInstrumentation = true;
 		}
 	}
 }
