[PCG] Update Attribute Filter's attribute code path to match the Point code path

#rnx

diff --git a/Engine/Plugins/PCG/Source/PCG/Private/Elements/PCGAttributeFilter.cpp b/Engine/Plugins/PCG/Source/PCG/Private/Elements/PCGAttributeFilter.cpp
--- a/Engine/Plugins/PCG/Source/PCG/Private/Elements/PCGAttributeFilter.cpp
+++ b/Engine/Plugins/PCG/Source/PCG/Private/Elements/PCGAttributeFilter.cpp
@@ -718,6 +718,9 @@ bool FPCGAttributeFilterElementBase::DoFiltering(FPCGContext* Context, EPCGAttri
 			OperationData.InFilterMetadata->AddAttributesFiltered(OriginalParamData->Metadata, TSet<FName>(), EPCGMetadataFilterMode::ExcludeAttributes);
 			OperationData.OutFilterMetadata->AddAttributesFiltered(OriginalParamData->Metadata, TSet<FName>(), EPCGMetadataFilterMode::ExcludeAttributes);
 
+			// Will be set individually in batches
+			OperationData.FilterBitArray.SetNumUninitialized(OriginalParamData->Metadata ? OriginalParamData->Metadata->GetItemCountForChild() : 0);
+
 			InFilterData = InFilterParamData;
 			OutFilterData = OutFilterParamData;
 		}
@@ -808,26 +811,12 @@ bool FPCGAttributeFilterElementBase::DoFiltering(FPCGContext* Context, EPCGAttri
 					return false;
 				}
 
-				auto AddElement = [&OperationData](bool bInFilter, int32 Index)
-				{
-					if (OperationData.bIsInputPointData)
-					{
-						OperationData.FilterBitArray[Index] = bInFilter;
-					}
-					else
-					{
-						UPCGMetadata* Metadata = bInFilter ? OperationData.InFilterMetadata : OperationData.OutFilterMetadata;
-						check(Metadata && OperationData.OriginalMetadata);
-						PCGMetadataEntryKey EntryKey = Metadata->AddEntry();
-						Metadata->SetAttributes(Index, OperationData.OriginalMetadata, EntryKey);
-					}
-				};
-
+				check(Range == 0 || OperationData.FilterBitArray.IsValidIndex(StartIndex + Range - 1));
 				for (int32 j = 0; j < Range; ++j)
 				{
 					if (bShouldSample && SkipTests[j])
 					{
-						AddElement(true, StartIndex + j);
+						OperationData.FilterBitArray[StartIndex + j] = true;
 						continue;
 					}
 
@@ -835,14 +824,7 @@ bool FPCGAttributeFilterElementBase::DoFiltering(FPCGContext* Context, EPCGAttri
 						PCGAttributeFilterHelpers::ApplyRange(TargetValues[j], FirstThresholdValues[j], SecondThresholdValues[j], FirstThreshold.bInclusive, SecondThreshold->bInclusive) : 
 						PCGAttributeFilterHelpers::ApplyCompare(TargetValues[j], FirstThresholdValues[j], Operator));
 
-					if (bShouldKeep)
-					{
-						AddElement(true, StartIndex + j);
-					}
-					else
-					{
-						AddElement(false, StartIndex + j);
-					}
+					OperationData.FilterBitArray[StartIndex + j] = bShouldKeep;
 				}
 			}
 
@@ -873,6 +855,39 @@ bool FPCGAttributeFilterElementBase::DoFiltering(FPCGContext* Context, EPCGAttri
 					}
 				}
 			}
+			else
+			{
+				check(OperationData.OriginalMetadata);
+
+				if (!OperationData.FilterBitArray.IsEmpty())
+				{
+					check(OperationData.FilterBitArray.Num() == OperationData.OriginalMetadata->GetItemCountForChild())
+
+					TArray<PCGMetadataEntryKey, TInlineAllocator<256>> InEntryKeys;
+					TArray<PCGMetadataEntryKey, TInlineAllocator<256>> OutEntryKeys;
+
+					const int32 NumInFilterMetadata = OperationData.FilterBitArray.CountSetBits();
+					const int32 NumOutFilterMetadata = OperationData.OriginalMetadata->GetItemCountForChild() - NumInFilterMetadata;
+
+					InEntryKeys.Reserve(NumInFilterMetadata);
+					OutEntryKeys.Reserve(NumOutFilterMetadata);
+
+					for (int32 Index = 0; Index < OperationData.FilterBitArray.Num(); ++Index)
+					{
+						if (OperationData.FilterBitArray[Index])
+						{
+							InEntryKeys.Add(Index);
+						}
+						else
+						{
+							OutEntryKeys.Add(Index);
+						}
+					}
+
+					OperationData.InFilterMetadata->SetAttributes(InEntryKeys, OperationData.OriginalMetadata, /*InOutOptionalKeys=*/nullptr, Context);
+					OperationData.OutFilterMetadata->SetAttributes(OutEntryKeys, OperationData.OriginalMetadata, /*InOutOptionalKeys=*/nullptr, Context);
+				}
+			}
 
 			FPCGTaggedData& InFilterOutput = Outputs.Add_GetRef(Input);
 			InFilterOutput.Pin = PCGPinConstants::DefaultInFilterLabel;
diff --git a/Engine/Plugins/PCG/Source/PCG/Private/Metadata/PCGMetadata.cpp b/Engine/Plugins/PCG/Source/PCG/Private/Metadata/PCGMetadata.cpp
--- a/Engine/Plugins/PCG/Source/PCG/Private/Metadata/PCGMetadata.cpp
+++ b/Engine/Plugins/PCG/Source/PCG/Private/Metadata/PCGMetadata.cpp
@@ -1374,7 +1374,7 @@ void UPCGMetadata::SetPointAttributes(const TArrayView<const FPCGPoint>& InPoint
 
 void UPCGMetadata::SetAttributes(const TArrayView<const PCGMetadataEntryKey>& InOriginalKeys, const UPCGMetadata* InMetadata, const TArrayView<PCGMetadataEntryKey>* InOutOptionalKeys, FPCGContext* OptionalContext)
 {
-	if (!InMetadata || InMetadata->GetAttributeCount() == 0 || GetAttributeCount() == 0)
+	if (!InMetadata || InMetadata->GetAttributeCount() == 0 || GetAttributeCount() == 0 || InOriginalKeys.IsEmpty())
 	{
 		return;
 	}
@@ -1414,6 +1414,7 @@ void UPCGMetadata::SetAttributes(const TArrayView<const PCGMetadataEntryKey>& In
 	}
 	else
 	{
+		InKeys.Append(InOriginalKeys);
 		OutKeys.Init(PCGInvalidEntryKey, InOriginalKeys.Num());
 	}
 
