[PCG] Fix RHI validation errors.

#rnx

diff --git a/Engine/Plugins/PCG/Source/PCG/Private/Compute/DataInterfaces/PCGDataCollectionDataInterface.cpp b/Engine/Plugins/PCG/Source/PCG/Private/Compute/DataInterfaces/PCGDataCollectionDataInterface.cpp
--- a/Engine/Plugins/PCG/Source/PCG/Private/Compute/DataInterfaces/PCGDataCollectionDataInterface.cpp
+++ b/Engine/Plugins/PCG/Source/PCG/Private/Compute/DataInterfaces/PCGDataCollectionDataInterface.cpp
@@ -659,7 +659,12 @@ void FPCGDataCollectionDataProviderProxy::AllocateResources(FRDGBuilder& GraphBu
 		TEXT("PCGDataCollectionBuffer");
 #endif
 
-	FRDGBufferDesc Desc = FRDGBufferDesc::CreateStructuredDesc(sizeof(uint32), SizeBytes >> 2);
+	FRDGBufferDesc Desc = FRDGBufferDesc::CreateByteAddressDesc(SizeBytes);
+	if (ReadbackMode != EPCGReadbackMode::None)
+	{
+		Desc.Usage |= BUF_SourceCopy;
+	}
+
 	Buffer = GraphBuilder.CreateBuffer(Desc, DebugName);
 	BufferUAV = GraphBuilder.CreateUAV(Buffer);
 
diff --git a/Engine/Plugins/PCG/Source/PCG/Private/Compute/DataInterfaces/PCGDataCollectionDataInterface.h b/Engine/Plugins/PCG/Source/PCG/Private/Compute/DataInterfaces/PCGDataCollectionDataInterface.h
--- a/Engine/Plugins/PCG/Source/PCG/Private/Compute/DataInterfaces/PCGDataCollectionDataInterface.h
+++ b/Engine/Plugins/PCG/Source/PCG/Private/Compute/DataInterfaces/PCGDataCollectionDataInterface.h
@@ -94,7 +94,7 @@ public:
 	TWeakObjectPtr<UPCGDataBinding> Binding;
 	FPCGDataCollectionDesc PinDesc;
 
-	EPCGReadbackMode ReadbackMode;
+	EPCGReadbackMode ReadbackMode = EPCGReadbackMode::None;
 
 	FName OutputPinLabel;
 	FName OutputPinLabelAlias;
@@ -128,7 +128,7 @@ public:
 protected:
 	uint32 SizeBytes = 0;
 
-	EPCGReadbackMode ReadbackMode;
+	EPCGReadbackMode ReadbackMode = EPCGReadbackMode::None;
 
 	using FParameters = FPCGDataCollectionDataInterfaceParameters;
 
diff --git a/Engine/Plugins/PCG/Source/PCG/Private/Compute/DataInterfaces/PCGDataCollectionUploadDataInterface.cpp b/Engine/Plugins/PCG/Source/PCG/Private/Compute/DataInterfaces/PCGDataCollectionUploadDataInterface.cpp
--- a/Engine/Plugins/PCG/Source/PCG/Private/Compute/DataInterfaces/PCGDataCollectionUploadDataInterface.cpp
+++ b/Engine/Plugins/PCG/Source/PCG/Private/Compute/DataInterfaces/PCGDataCollectionUploadDataInterface.cpp
@@ -82,7 +82,12 @@ void FPCGDataProviderDataCollectionUploadProxy::AllocateResources(FRDGBuilder& G
 		TEXT("PCGDataCollectionUploadBuffer");
 #endif
 
-	const FRDGBufferDesc Desc = FRDGBufferDesc::CreateStructuredDesc(sizeof(uint32), PackedDataCollection.Num());
+	FRDGBufferDesc Desc = FRDGBufferDesc::CreateByteAddressDesc(sizeof(uint32) * PackedDataCollection.Num());
+	if (ReadbackMode != EPCGReadbackMode::None)
+	{
+		Desc.Usage |= BUF_SourceCopy;
+	}
+
 	Buffer = GraphBuilder.CreateBuffer(Desc, DebugName);
 	BufferUAV = GraphBuilder.CreateUAV(Buffer);
 
diff --git a/Engine/Plugins/PCG/Source/PCG/Private/Compute/DataInterfaces/PCGDebugDataInterface.cpp b/Engine/Plugins/PCG/Source/PCG/Private/Compute/DataInterfaces/PCGDebugDataInterface.cpp
--- a/Engine/Plugins/PCG/Source/PCG/Private/Compute/DataInterfaces/PCGDebugDataInterface.cpp
+++ b/Engine/Plugins/PCG/Source/PCG/Private/Compute/DataInterfaces/PCGDebugDataInterface.cpp
@@ -39,7 +39,7 @@ void UPCGDebugDataInterface::GetShaderHash(FString& InOutKey) const
 
 void UPCGDebugDataInterface::GetHLSL(FString& OutHLSL, FString const& InDataInterfaceName) const
 {
-	const FString WriteDebugValueImpl = (DebugBufferSize > 0)
+	const FString WriteDebugValueImpl = (DebugBufferSizeFloats > 0)
 		? FString::Format(TEXT(
 			"\n{\n"
 			"	if (Index < {0})\n"
@@ -47,7 +47,7 @@ void UPCGDebugDataInterface::GetHLSL(FString& OutHLSL, FString const& InDataInte
 			"		return {1}_DebugBuffer.Store(Index * 4, Value);\n"
 			"	}\n"
 			"}\n"),
-			{ DebugBufferSize, *InDataInterfaceName })
+			{ DebugBufferSizeFloats, *InDataInterfaceName })
 		: TEXT("{ /* No-Op */ }\n");
 
 	OutHLSL += FString::Format(TEXT(
@@ -58,13 +58,13 @@ void UPCGDebugDataInterface::GetHLSL(FString& OutHLSL, FString const& InDataInte
 
 bool UPCGDebugDataInterface::GetRequiresReadback() const
 {
-	return DebugBufferSize > 0;
+	return DebugBufferSizeFloats > 0;
 }
 
 UComputeDataProvider* UPCGDebugDataInterface::CreateDataProvider(TObjectPtr<UObject> InBinding, uint64 InInputMask, uint64 InOutputMask) const
 {
 	UPCGDebugDataProvider* Provider = NewObject<UPCGDebugDataProvider>();
-	Provider->DebugBufferSize = DebugBufferSize;
+	Provider->DebugBufferSizeFloats = DebugBufferSizeFloats;
 	return Provider;
 }
 
@@ -79,7 +79,7 @@ FComputeDataProviderRenderProxy* UPCGDebugDataProvider::GetRenderProxy()
 			return;
 		}
 
-		check(ThisWeakPtr->DebugBufferSize * sizeof(float) == InNumBytes);
+		check(ThisWeakPtr->DebugBufferSizeFloats * sizeof(float) == InNumBytes);
 
 		UPCGDebugDataProvider* DataProvider = ThisWeakPtr.Get();
 		
@@ -95,7 +95,7 @@ FComputeDataProviderRenderProxy* UPCGDebugDataProvider::GetRenderProxy()
 		DataProvider->bOperationComplete = true;
 	};
 
-	return new FPCGDebugDataProviderProxy(DebugBufferSize, ProcessReadbackData_RenderThread);
+	return new FPCGDebugDataProviderProxy(DebugBufferSizeFloats, ProcessReadbackData_RenderThread);
 }
 
 bool FPCGDebugDataProviderProxy::IsValid(FValidationData const& InValidationData) const
@@ -121,15 +121,16 @@ void FPCGDebugDataProviderProxy::GatherDispatchData(FDispatchData const& InDispa
 
 void FPCGDebugDataProviderProxy::AllocateResources(FRDGBuilder& GraphBuilder, FAllocationData const& InAllocationData)
 {
-	if (DebugBufferSize > 0)
+	if (DebugBufferSizeFloats > 0)
 	{
-		FRDGBufferDesc Desc = FRDGBufferDesc::CreateStructuredDesc(sizeof(float), DebugBufferSize);
-		Desc.Usage = EBufferUsageFlags(Desc.Usage | BUF_SourceCopy);
+		FRDGBufferDesc Desc = FRDGBufferDesc::CreateByteAddressDesc(sizeof(float) * DebugBufferSizeFloats);
+		Desc.Usage |= BUF_SourceCopy;
+
 		DebugBuffer = GraphBuilder.CreateBuffer(Desc, TEXT("PCGDebugBuffer"));
 		DebugBufferUAV = GraphBuilder.CreateUAV(DebugBuffer);
 
 		TArray<float> ZeroInitializedBuffer;
-		ZeroInitializedBuffer.SetNumZeroed(DebugBufferSize);
+		ZeroInitializedBuffer.SetNumZeroed(DebugBufferSizeFloats);
 
 		GraphBuilder.QueueBufferUpload(DebugBuffer, ZeroInitializedBuffer.GetData(), ZeroInitializedBuffer.Num() * sizeof(float), ERDGInitialDataFlags::None);
 	}
@@ -139,7 +140,7 @@ void FPCGDebugDataProviderProxy::GetReadbackData(TArray<FReadbackData>& OutReadb
 {
 	FReadbackData Data;
 	Data.Buffer = DebugBuffer;
-	Data.NumBytes = DebugBufferSize * sizeof(float);
+	Data.NumBytes = DebugBufferSizeFloats * sizeof(float);
 	Data.ReadbackCallback_RenderThread = &AsyncReadbackCallback_RenderThread;
 
 	OutReadbackData.Add(MoveTemp(Data));
diff --git a/Engine/Plugins/PCG/Source/PCG/Private/Compute/DataInterfaces/PCGDebugDataInterface.h b/Engine/Plugins/PCG/Source/PCG/Private/Compute/DataInterfaces/PCGDebugDataInterface.h
--- a/Engine/Plugins/PCG/Source/PCG/Private/Compute/DataInterfaces/PCGDebugDataInterface.h
+++ b/Engine/Plugins/PCG/Source/PCG/Private/Compute/DataInterfaces/PCGDebugDataInterface.h
@@ -28,10 +28,10 @@ public:
 	UComputeDataProvider* CreateDataProvider(TObjectPtr<UObject> InBinding, uint64 InInputMask, uint64 InOutputMask) const override;
 	//~ End UComputeDataInterface Interface
 
-	void SetDebugBufferSize(uint32 InDebugBufferSize) { DebugBufferSize = InDebugBufferSize; }
+	void SetDebugBufferSize(uint32 InDebugBufferSizeFloats) { DebugBufferSizeFloats = InDebugBufferSizeFloats; }
 
 private:
-	uint32 DebugBufferSize = 0;
+	uint32 DebugBufferSizeFloats = 0;
 };
 
 /** Compute Framework Data Provider for each custom compute kernel. */
@@ -45,15 +45,15 @@ public:
 	FComputeDataProviderRenderProxy* GetRenderProxy() override;
 	//~ End UComputeDataProvider Interface
 
-	uint32 DebugBufferSize = 0;
+	uint32 DebugBufferSizeFloats = 0;
 	std::atomic<bool> bOperationComplete = false;
 };
 
 class FPCGDebugDataProviderProxy : public FComputeDataProviderRenderProxy
 {
 public:
-	FPCGDebugDataProviderProxy(uint32 InDebugBufferSize, FReadbackCallback InAsyncReadbackCallback_RenderThread)
-		: DebugBufferSize(InDebugBufferSize)
+	FPCGDebugDataProviderProxy(uint32 InDebugBufferSizeFloats, FReadbackCallback InAsyncReadbackCallback_RenderThread)
+		: DebugBufferSizeFloats(InDebugBufferSizeFloats)
 		, AsyncReadbackCallback_RenderThread(InAsyncReadbackCallback_RenderThread)
 	{}
 
@@ -69,6 +69,6 @@ protected:
 
 	FRDGBufferRef DebugBuffer = nullptr;
 	FRDGBufferUAVRef DebugBufferUAV = nullptr;
-	uint32 DebugBufferSize = 0;
+	uint32 DebugBufferSizeFloats = 0;
 	FReadbackCallback AsyncReadbackCallback_RenderThread;
 };
diff --git a/Engine/Plugins/PCG/Source/PCG/Private/Compute/DataInterfaces/PCGInstanceDataInterface.cpp b/Engine/Plugins/PCG/Source/PCG/Private/Compute/DataInterfaces/PCGInstanceDataInterface.cpp
--- a/Engine/Plugins/PCG/Source/PCG/Private/Compute/DataInterfaces/PCGInstanceDataInterface.cpp
+++ b/Engine/Plugins/PCG/Source/PCG/Private/Compute/DataInterfaces/PCGInstanceDataInterface.cpp
@@ -170,7 +170,7 @@ bool FPCGInstanceDataProviderProxy::IsValid(FValidationData const& InValidationD
 void FPCGInstanceDataProviderProxy::AllocateResources(FRDGBuilder& GraphBuilder, FAllocationData const& InAllocationData)
 {
 	const uint32 StrideUint4s = 3;
-	InstanceData = GraphBuilder.CreateBuffer(FRDGBufferDesc::CreateStructuredDesc(sizeof(FUintVector4) * StrideUint4s, FMath::Max(1u, NumInstancesAllPrimitives)), TEXT("PCGInstanceDataBuffer"));
+	InstanceData = GraphBuilder.CreateBuffer(FRDGBufferDesc::CreateStructuredDesc(sizeof(FUintVector4), FMath::Max(1u, NumInstancesAllPrimitives) * StrideUint4s), TEXT("PCGInstanceDataBuffer"));
 	InstanceDataSRV = GraphBuilder.CreateSRV(InstanceData);
 	InstanceDataUAV = GraphBuilder.CreateUAV(InstanceData);
 
