[PCG] GPU texture sampling - respect filter setting.

#rnx

diff --git a/Engine/Plugins/PCG/Source/PCG/Private/Compute/DataInterfaces/PCGTextureDataInterface.cpp b/Engine/Plugins/PCG/Source/PCG/Private/Compute/DataInterfaces/PCGTextureDataInterface.cpp
--- a/Engine/Plugins/PCG/Source/PCG/Private/Compute/DataInterfaces/PCGTextureDataInterface.cpp
+++ b/Engine/Plugins/PCG/Source/PCG/Private/Compute/DataInterfaces/PCGTextureDataInterface.cpp
@@ -88,18 +88,39 @@ UComputeDataProvider* UPCGTextureDataInterface::CreateDataProvider(TObjectPtr<UO
 		// TODO: Bindless resources to handle multiple textures
 		ensure(TaggedDatas.Num() == 1); 
 
-		const UPCGBaseTextureData* BaseTextureData = Cast<UPCGBaseTextureData>(TaggedDatas[0].Data);
-		ensure(BaseTextureData);
-
-		DataProvider->TextureRHI = BaseTextureData ? BaseTextureData->GetTextureRHI() : nullptr;
+		if (const UPCGBaseTextureData* BaseTextureData = Cast<UPCGBaseTextureData>(TaggedDatas[0].Data))
+		{
+			DataProvider->TextureRHI = BaseTextureData->GetTextureRHI();
+			DataProvider->Filter = BaseTextureData->Filter;
+		}
+		else
+		{
+			ensureMsgf(false, TEXT("Unsupported data type encountered by texture data interface: '%s'"), *TaggedDatas[0].Data->GetName());
+		}
 	}
 
 	return DataProvider;
 }
 
+FPCGTextureDataProviderProxy::FPCGTextureDataProviderProxy(FTextureRHIRef InTextureRHI, EPCGTextureFilter InFilter)
+	: TextureRHI(InTextureRHI)
+{
+	switch (InFilter)
+	{
+	case EPCGTextureFilter::Point:
+		SamplerRHI = TStaticSamplerState<SF_Point, AM_Clamp, AM_Clamp, AM_Clamp>::GetRHI();
+		break;
+	default:
+		ensureMsgf(false, TEXT("Unrecognized texture filter setting, falling back to bilinear."));
+	case EPCGTextureFilter::Bilinear: // Falls through
+		SamplerRHI = TStaticSamplerState<SF_Bilinear, AM_Clamp, AM_Clamp, AM_Clamp>::GetRHI();
+		break;
+	}
+}
+
 FComputeDataProviderRenderProxy* UPCGTextureDataProvider::GetRenderProxy()
 {
-	return new FPCGTextureDataProviderProxy(TextureRHI);
+	return new FPCGTextureDataProviderProxy(TextureRHI, Filter);
 }
 
 bool FPCGTextureDataProviderProxy::IsValid(FValidationData const& InValidationData) const
@@ -119,6 +140,6 @@ void FPCGTextureDataProviderProxy::GatherDispatchData(FDispatchData const& InDis
 	{
 		FParameters& Parameters = ParameterArray[InvocationIndex];
 		Parameters.Texture = TextureRHI ? (FRHITexture*)TextureRHI : (FRHITexture*)GBlackTexture->TextureRHI;
-		Parameters.Sampler = TStaticSamplerState<SF_Point, AM_Clamp, AM_Clamp, AM_Clamp>::GetRHI();
+		Parameters.Sampler = SamplerRHI;
 	}
 }
diff --git a/Engine/Plugins/PCG/Source/PCG/Private/Compute/DataInterfaces/PCGTextureDataInterface.h b/Engine/Plugins/PCG/Source/PCG/Private/Compute/DataInterfaces/PCGTextureDataInterface.h
--- a/Engine/Plugins/PCG/Source/PCG/Private/Compute/DataInterfaces/PCGTextureDataInterface.h
+++ b/Engine/Plugins/PCG/Source/PCG/Private/Compute/DataInterfaces/PCGTextureDataInterface.h
@@ -4,6 +4,8 @@
 
 #include "Compute/DataInterfaces/PCGComputeDataInterface.h"
 
+#include "Data/PCGTextureData.h"
+
 #include "RHIResources.h"
 #include "ComputeFramework/ComputeDataProvider.h"
 
@@ -44,14 +46,14 @@ public:
 	//~ End UComputeDataProvider Interface
 
 	FTextureRHIRef TextureRHI = nullptr;
+
+	EPCGTextureFilter Filter = EPCGTextureFilter::Bilinear;
 };
 
 class FPCGTextureDataProviderProxy : public FComputeDataProviderRenderProxy
 {
 public:
-	FPCGTextureDataProviderProxy(FTextureRHIRef InTextureRHI)
-		: TextureRHI(InTextureRHI)
-	{}
+	FPCGTextureDataProviderProxy(FTextureRHIRef InTextureRHI, EPCGTextureFilter InFilter);
 
 	//~ Begin FComputeDataProviderRenderProxy Interface
 	bool IsValid(FValidationData const& InValidationData) const override;
@@ -62,4 +64,6 @@ protected:
 	using FParameters = FPCGTextureDataInterfaceParameters;
 
 	FTextureRHIRef TextureRHI = nullptr;
+
+	FSamplerStateRHIRef SamplerRHI = nullptr;
 };
