[PCG] Fix Cancel leak when pcg.GraphMultithreading is enabled
- Prevent scheduling while cancelling

#rnx

diff --git a/Engine/Plugins/PCG/Source/PCG/Private/Graph/PCGGraphExecutor.cpp b/Engine/Plugins/PCG/Source/PCG/Private/Graph/PCGGraphExecutor.cpp
--- a/Engine/Plugins/PCG/Source/PCG/Private/Graph/PCGGraphExecutor.cpp
+++ b/Engine/Plugins/PCG/Source/PCG/Private/Graph/PCGGraphExecutor.cpp
@@ -1238,6 +1238,16 @@ double FPCGGraphExecutor::GetTickBudgetInSeconds() const
 	return PCGGraphExecutor::CVarTimePerFrame.GetValueOnAnyThread() / 1000.0;
 }
 
+{
+	if (ExecuteVersion == EExecuteVersion::None)
+	{
+		check(GetNonScheduledRemainingTaskCount() == 0);
+		return PCGGraphExecutor::CVarExecuteV2.GetValueOnAnyThread() ? EExecuteVersion::V2 : EExecuteVersion::V1;
+	}
+
+	return ExecuteVersion;
+}
+
 bool FPCGGraphExecutor::ExecuteScheduling(double EndTime, TSharedPtr<FPCGGraphActiveTask>* OutMainThreadTask, bool bForceCheckSleepingTasks)
 {
 	// Can currently be disabled by cancellation
