#UE Fix a compile error related to casting primitive types

diff --git a/Engine/Plugins/PCG/Source/PCG/Private/Compute/PCGComputeCommon.cpp b/Engine/Plugins/PCG/Source/PCG/Private/Compute/PCGComputeCommon.cpp
--- a/Engine/Plugins/PCG/Source/PCG/Private/Compute/PCGComputeCommon.cpp
+++ b/Engine/Plugins/PCG/Source/PCG/Private/Compute/PCGComputeCommon.cpp
@@ -304,7 +304,7 @@ namespace PCGComputeHelpers
 
 		// If buffer size exceeds a proportion of total graphics memory, then it is deemed too large. Using this as a heuristic as there
 		// is no RHI API to obtain available graphics memory outside of D3D12.
-		const uint64 GPUBudgetBytes = static_cast<uint64>(TextureMemStats.TotalGraphicsMemory * CVarMaxGPUBufferSizeProportion.GetValueOnAnyThread());
+		const uint64 GPUBudgetBytes = static_cast<uint64>(TextureMemStats.TotalGraphicsMemory) * static_cast<double>(CVarMaxGPUBufferSizeProportion.GetValueOnAnyThread());
 
 		// Buffer size also cannot exceed the max size of a uint32, otherwise the size will get truncated and other systems will break.
 		// TODO: This limits the maximum number of points to around 46 million. Support uint64 max instead of uint32 to get up to 2 billion points.
