[PCG] Add the ability to Display messages via the Print String node without logging

diff --git a/Engine/Plugins/PCG/Source/PCG/Private/Elements/PCGPrintElement.cpp b/Engine/Plugins/PCG/Source/PCG/Private/Elements/PCGPrintElement.cpp
--- a/Engine/Plugins/PCG/Source/PCG/Private/Elements/PCGPrintElement.cpp
+++ b/Engine/Plugins/PCG/Source/PCG/Private/Elements/PCGPrintElement.cpp
@@ -124,6 +124,9 @@ bool FPCGPrintElement::ExecuteInternal(FPCGContext* Context) const
 		case EPCGPrintVerbosity::Log:
 			PCGE_LOG_C(Log, LogOnly, Context, FText::FromString(FinalString));
 			break;
+		case EPCGPrintVerbosity::Display:
+			PCGE_LOG_C(Display, LogOnly, Context, FText::FromString(FinalString));
+			break;
 		default:
 			checkNoEntry();
 	}
diff --git a/Engine/Plugins/PCG/Source/PCG/Public/Elements/PCGPrintElement.h b/Engine/Plugins/PCG/Source/PCG/Public/Elements/PCGPrintElement.h
--- a/Engine/Plugins/PCG/Source/PCG/Public/Elements/PCGPrintElement.h
+++ b/Engine/Plugins/PCG/Source/PCG/Public/Elements/PCGPrintElement.h
@@ -16,7 +16,8 @@ enum class EPCGPrintVerbosity : uint8
 {
 	Log = ELogVerbosity::Log,
 	Warning = ELogVerbosity::Warning,
-	Error = ELogVerbosity::Error
+	Error = ELogVerbosity::Error,
+	Display = ELogVerbosity::Display
 };
 
 /** Used to track the debug message to properly remove it upon regen or clean up. */
