[PCG] Discard negative indices in WriteDebugValue()

#rnx

diff --git a/Engine/Plugins/PCG/Source/PCG/Private/Compute/DataInterfaces/PCGDebugDataInterface.cpp b/Engine/Plugins/PCG/Source/PCG/Private/Compute/DataInterfaces/PCGDebugDataInterface.cpp
--- a/Engine/Plugins/PCG/Source/PCG/Private/Compute/DataInterfaces/PCGDebugDataInterface.cpp
+++ b/Engine/Plugins/PCG/Source/PCG/Private/Compute/DataInterfaces/PCGDebugDataInterface.cpp
@@ -42,7 +42,7 @@ void UPCGDebugDataInterface::GetHLSL(FString& OutHLSL, FString const& InDataInte
 	const FString WriteDebugValueImpl = (DebugBufferSizeFloats > 0)
 		? FString::Format(TEXT(
 			"\n{\n"
-			"	if (Index < {0})\n"
+			"	if (Index >= 0 && Index < {0})\n"
 			"	{\n"
 			"		return {1}_DebugBuffer.Store(Index * 4, Value);\n"
 			"	}\n"
