[PCG] Fix 'Load as preview' causing GeneratedResouces to be lost causing duplicate assets on next generation
- LoadedPreviewResources needed to be populated on load so that any change to the actor that is persisted will still properly serialize the GeneratedResources

#rnx

diff --git a/Engine/Plugins/PCG/Source/PCG/Private/PCGComponent.cpp b/Engine/Plugins/PCG/Source/PCG/Private/PCGComponent.cpp
--- a/Engine/Plugins/PCG/Source/PCG/Private/PCGComponent.cpp
+++ b/Engine/Plugins/PCG/Source/PCG/Private/PCGComponent.cpp
@@ -1295,11 +1295,7 @@ void UPCGComponent::CleanupLocalImmediate(bool bRemoveComponents, bool bCleanupL
 				if (Resource)
 				{
 #if WITH_EDITOR
-					if (Resource->IsMarkedTransientOnLoad())
-					{
-						LoadedPreviewResources.Add(Resource);
-					}
-					else
+					if (!Resource->IsMarkedTransientOnLoad())
 					{
 						Resource->Rename(nullptr, nullptr, REN_DoNotDirty | REN_DontCreateRedirectors | REN_NonTransactional);
 					}
@@ -1403,11 +1399,7 @@ FPCGTaskId UPCGComponent::CreateCleanupTask(bool bRemoveComponents, const TArray
 					if (Resource)
 					{
 #if WITH_EDITOR
-						if (Resource->IsMarkedTransientOnLoad())
-						{
-							ThisComponent->LoadedPreviewResources.Add(Resource);
-						}
-						else
+						if (!Resource->IsMarkedTransientOnLoad())
 						{
 							Resource->Rename(nullptr, nullptr, REN_DoNotDirty | REN_DontCreateRedirectors | REN_NonTransactional);
 						}
@@ -1539,11 +1531,7 @@ void UPCGComponent::CleanupUnusedManagedResources()
 				if (Resource)
 				{
 #if WITH_EDITOR
-					if (Resource->IsMarkedTransientOnLoad())
-					{
-						LoadedPreviewResources.Add(Resource);
-					}
-					else
+					if (!Resource->IsMarkedTransientOnLoad())
 					{
 						Resource->Rename(nullptr, nullptr, REN_DoNotDirty | REN_DontCreateRedirectors | REN_NonTransactional);
 					}
@@ -1663,7 +1651,7 @@ void UPCGComponent::Serialize(FArchive& Ar)
 #if WITH_EDITOR
 	TArray<TObjectPtr<UPCGManagedResource>> GeneratedResourcesCopy;
 
-	if (Ar.IsSaving() && CurrentEditingMode == EPCGEditorDirtyMode::Preview)
+	if (!Ar.IsLoading() && CurrentEditingMode == EPCGEditorDirtyMode::Preview)
 	{
 		GeneratedResourcesCopy = GeneratedResources;
 		GeneratedResources = LoadedPreviewResources;
@@ -1703,7 +1691,7 @@ void UPCGComponent::Serialize(FArchive& Ar)
 
 
 #if WITH_EDITOR
-	if (Ar.IsSaving() && CurrentEditingMode == EPCGEditorDirtyMode::Preview)
+	if (!Ar.IsLoading() && CurrentEditingMode == EPCGEditorDirtyMode::Preview)
 	{
 		GeneratedResources = GeneratedResourcesCopy;
 	}
@@ -3527,10 +3515,16 @@ bool UPCGComponent::DeletePreviewResources()
 	// Make sure to release fully the resources that were loaded
 	for (TObjectPtr<UPCGManagedResource> ResourceToRelease : LoadedPreviewResources)
 	{
-		// Changing the transient state will clear the "marked transient on load" flag
-		ResourceToRelease->ChangeTransientState(EPCGEditorDirtyMode::Normal);
-		ResourceToRelease->Release(/*bHardRelease=*/true, ActorsToDelete);
-		bResourceWasReleased = true;
+		if (!GeneratedResources.Contains(ResourceToRelease))
+		{
+			// Changing the transient state will clear the "marked transient on load" flag
+			ResourceToRelease->ChangeTransientState(EPCGEditorDirtyMode::Normal);
+			ResourceToRelease->Release(/*bHardRelease=*/true, ActorsToDelete);
+			bResourceWasReleased = true;
+		}
+
+		// Either this resource was released in the previous code block or prior to calling DeletePreviewResources its transient state was changed back to normal
+		ensure(!ResourceToRelease->IsMarkedTransientOnLoad());
 	}
 
 	LoadedPreviewResources.Empty();
@@ -3552,6 +3546,8 @@ void UPCGComponent::MarkResourcesAsTransientOnLoad()
 			GeneratedResource->MarkTransientOnLoad();
 		}
 	}
+
+	LoadedPreviewResources = GeneratedResources;
 }
 
 void UPCGComponent::ChangeTransientState(EPCGEditorDirtyMode NewEditingMode)
