[PCG] Add warning messages when attaching a component spawned by PCG failed.

#rnx

diff --git a/Engine/Plugins/PCG/Source/PCG/Private/Components/PCGProceduralISMComponent.cpp b/Engine/Plugins/PCG/Source/PCG/Private/Components/PCGProceduralISMComponent.cpp
--- a/Engine/Plugins/PCG/Source/PCG/Private/Components/PCGProceduralISMComponent.cpp
+++ b/Engine/Plugins/PCG/Source/PCG/Private/Components/PCGProceduralISMComponent.cpp
@@ -18,7 +18,7 @@
 #include "Serialization/ArchiveCrc32.h"
 #include "VT/RuntimeVirtualTexture.h"
 
-UPCGManagedProceduralISMComponent* PCGManagedProceduralISMComponent::GetOrCreateManagedProceduralISMC(AActor* InTargetActor, UPCGComponent* InSourceComponent, uint64 InSettingsUID, const FPCGProceduralISMCBuilderParameters& InParams)
+UPCGManagedProceduralISMComponent* PCGManagedProceduralISMComponent::GetOrCreateManagedProceduralISMC(AActor* InTargetActor, UPCGComponent* InSourceComponent, uint64 InSettingsUID, const FPCGProceduralISMCBuilderParameters& InParams, FPCGContext* OptionalContext)
 {
 	check(InTargetActor && InSourceComponent);
 
@@ -89,7 +89,10 @@ UPCGManagedProceduralISMComponent* PCGManagedProceduralISMComponent::GetOrCreate
 	ISMC->RegisterComponent();
 	InTargetActor->AddInstanceComponent(ISMC);
 
-	ISMC->AttachToComponent(InTargetActor->GetRootComponent(), FAttachmentTransformRules(EAttachmentRule::KeepRelative, EAttachmentRule::KeepWorld, EAttachmentRule::KeepWorld, false));
+	if (!ISMC->AttachToComponent(InTargetActor->GetRootComponent(), FAttachmentTransformRules(EAttachmentRule::KeepRelative, EAttachmentRule::KeepWorld, EAttachmentRule::KeepWorld, false)))
+	{
+		PCGLog::Component::LogComponentAttachmentFailedWarning(OptionalContext);
+	}
 
 	// Create managed resource on source component
 	UPCGManagedProceduralISMComponent* Resource = NewObject<UPCGManagedProceduralISMComponent>(InSourceComponent);
diff --git a/Engine/Plugins/PCG/Source/PCG/Private/Components/PCGProceduralISMComponent.h b/Engine/Plugins/PCG/Source/PCG/Private/Components/PCGProceduralISMComponent.h
--- a/Engine/Plugins/PCG/Source/PCG/Private/Components/PCGProceduralISMComponent.h
+++ b/Engine/Plugins/PCG/Source/PCG/Private/Components/PCGProceduralISMComponent.h
@@ -19,6 +19,7 @@
 class FPrimitiveSceneProxy;
 class UPCGComponent;
 class UPCGManagedProceduralISMComponent;
+struct FPCGContext;
 struct FPCGSoftISMComponentDescriptor;
 
 /** Experimental instanced primitive that is not backed by CPU instance data. Instance data must be written direct to the GPU scene. */
@@ -342,5 +343,5 @@ protected:
 
 namespace PCGManagedProceduralISMComponent
 {
-	UPCGManagedProceduralISMComponent* GetOrCreateManagedProceduralISMC(AActor* InTargetActor, UPCGComponent* InSourceComponent, uint64 InSettingsUID, const FPCGProceduralISMCBuilderParameters& InParams);
+	UPCGManagedProceduralISMComponent* GetOrCreateManagedProceduralISMC(AActor* InTargetActor, UPCGComponent* InSourceComponent, uint64 InSettingsUID, const FPCGProceduralISMCBuilderParameters& InParams, FPCGContext* OptionalContext = nullptr);
 }
diff --git a/Engine/Plugins/PCG/Source/PCG/Private/Elements/PCGAddComponent.cpp b/Engine/Plugins/PCG/Source/PCG/Private/Elements/PCGAddComponent.cpp
--- a/Engine/Plugins/PCG/Source/PCG/Private/Elements/PCGAddComponent.cpp
+++ b/Engine/Plugins/PCG/Source/PCG/Private/Elements/PCGAddComponent.cpp
@@ -488,7 +488,11 @@ void FPCGAddComponentElement::AddComponents(FPCGAddComponentContext* Context,
 
 				if (USceneComponent* SceneComponent = Cast<USceneComponent>(Component))
 				{
-					SceneComponent->AttachToComponent(TargetActor->GetRootComponent(), FAttachmentTransformRules(EAttachmentRule::KeepRelative, EAttachmentRule::KeepWorld, EAttachmentRule::KeepWorld, false));
+					if (!SceneComponent->AttachToComponent(TargetActor->GetRootComponent(), FAttachmentTransformRules(EAttachmentRule::KeepRelative, EAttachmentRule::KeepWorld, EAttachmentRule::KeepWorld, false)))
+					{
+						PCGLog::Component::LogComponentAttachmentFailedWarning(Context);
+					}
+
 					SceneComponent->SetWorldTransform(Transform);
 				}
 
diff --git a/Engine/Plugins/PCG/Source/PCG/Private/Elements/PCGAnimBankMeshSpawner.cpp b/Engine/Plugins/PCG/Source/PCG/Private/Elements/PCGAnimBankMeshSpawner.cpp
--- a/Engine/Plugins/PCG/Source/PCG/Private/Elements/PCGAnimBankMeshSpawner.cpp
+++ b/Engine/Plugins/PCG/Source/PCG/Private/Elements/PCGAnimBankMeshSpawner.cpp
@@ -631,7 +631,7 @@ void FPCGAnimBankMeshSpawnerElement::SpawnAnimBankMeshInstances(FPCGAnimBankMesh
 	Params.SettingsCrc = Settings->GetSettingsCrc();
 	ensure(Params.SettingsCrc.IsValid());
 
-	UPCGManagedAnimBankMeshComponent* MABMC = UPCGActorHelpers::GetOrCreateManagedABMC(TargetActor, Context->SourceComponent.Get(), Params);
+	UPCGManagedAnimBankMeshComponent* MABMC = UPCGActorHelpers::GetOrCreateManagedABMC(TargetActor, Context->SourceComponent.Get(), Params, Context);
 
 	check(MABMC);
 	MABMC->SetCrc(Context->DependenciesCrc);
diff --git a/Engine/Plugins/PCG/Source/PCG/Private/Elements/PCGCreateSpline.cpp b/Engine/Plugins/PCG/Source/PCG/Private/Elements/PCGCreateSpline.cpp
--- a/Engine/Plugins/PCG/Source/PCG/Private/Elements/PCGCreateSpline.cpp
+++ b/Engine/Plugins/PCG/Source/PCG/Private/Elements/PCGCreateSpline.cpp
@@ -169,7 +169,11 @@ bool FPCGCreateSplineElement::ExecuteInternal(FPCGContext* Context) const
 
 			SplineComponent->RegisterComponent();
 			SplineActor->AddInstanceComponent(SplineComponent);
-			SplineComponent->AttachToComponent(SplineActor->GetRootComponent(), FAttachmentTransformRules(EAttachmentRule::KeepRelative, EAttachmentRule::KeepWorld, EAttachmentRule::KeepWorld, false));
+
+			if (!SplineComponent->AttachToComponent(SplineActor->GetRootComponent(), FAttachmentTransformRules(EAttachmentRule::KeepRelative, EAttachmentRule::KeepWorld, EAttachmentRule::KeepWorld, false)))
+			{
+				PCGLog::Component::LogComponentAttachmentFailedWarning(Context);
+			}
 
 			SplineData->ApplyTo(SplineComponent);
 
diff --git a/Engine/Plugins/PCG/Source/PCG/Private/Elements/PCGSpawnActor.cpp b/Engine/Plugins/PCG/Source/PCG/Private/Elements/PCGSpawnActor.cpp
--- a/Engine/Plugins/PCG/Source/PCG/Private/Elements/PCGSpawnActor.cpp
+++ b/Engine/Plugins/PCG/Source/PCG/Private/Elements/PCGSpawnActor.cpp
@@ -728,7 +728,7 @@ void FPCGSpawnActorElement::CollapseIntoTargetActor(FPCGSubgraphContext* Context
 	{
 		const FPCGISMComponentBuilderParams& ISMCParams = ISMCBuilderTransforms.Key;
 
-		UPCGManagedISMComponent* MISMC = UPCGActorHelpers::GetOrCreateManagedISMC(TargetActor, Context->SourceComponent.Get(), ISMCParams);
+		UPCGManagedISMComponent* MISMC = UPCGActorHelpers::GetOrCreateManagedISMC(TargetActor, Context->SourceComponent.Get(), ISMCParams, Context);
 		if (!MISMC)
 		{
 			continue;
diff --git a/Engine/Plugins/PCG/Source/PCG/Private/Elements/PCGSpawnSpline.cpp b/Engine/Plugins/PCG/Source/PCG/Private/Elements/PCGSpawnSpline.cpp
--- a/Engine/Plugins/PCG/Source/PCG/Private/Elements/PCGSpawnSpline.cpp
+++ b/Engine/Plugins/PCG/Source/PCG/Private/Elements/PCGSpawnSpline.cpp
@@ -112,7 +112,10 @@ bool FPCGSpawnSplineElement::ExecuteInternal(FPCGContext* InContext) const
 		SplineComponent->RegisterComponent();
 		TargetActor->AddInstanceComponent(SplineComponent);
 		// Spline Data apply to Spline Component is done in world coordinate, so keep world location.
-		SplineComponent->AttachToComponent(TargetActor->GetRootComponent(), FAttachmentTransformRules(EAttachmentRule::KeepRelative, EAttachmentRule::KeepWorld, EAttachmentRule::KeepWorld, false));
+		if (!SplineComponent->AttachToComponent(TargetActor->GetRootComponent(), FAttachmentTransformRules(EAttachmentRule::KeepRelative, EAttachmentRule::KeepWorld, EAttachmentRule::KeepWorld, false)))
+		{
+			PCGLog::Component::LogComponentAttachmentFailedWarning(InContext);
+		}
 
 		InSplineData->ApplyTo(SplineComponent);
 
diff --git a/Engine/Plugins/PCG/Source/PCG/Private/Elements/PCGSpawnSplineMesh.cpp b/Engine/Plugins/PCG/Source/PCG/Private/Elements/PCGSpawnSplineMesh.cpp
--- a/Engine/Plugins/PCG/Source/PCG/Private/Elements/PCGSpawnSplineMesh.cpp
+++ b/Engine/Plugins/PCG/Source/PCG/Private/Elements/PCGSpawnSplineMesh.cpp
@@ -215,7 +215,7 @@ bool FPCGSpawnSplineMeshElement::ExecuteInternal(FPCGContext* Context) const
 			SMCBuilderParams.SettingsCrc = Settings->GetSettingsCrc();
 			ensure(SMCBuilderParams.SettingsCrc.IsValid());
 
-			USplineMeshComponent* SplineMeshComponent = UPCGActorHelpers::GetOrCreateSplineMeshComponent(TargetActor, Context->SourceComponent.Get(), SMCBuilderParams);
+			USplineMeshComponent* SplineMeshComponent = UPCGActorHelpers::GetOrCreateSplineMeshComponent(TargetActor, Context->SourceComponent.Get(), SMCBuilderParams, Context);
 
 			// TODO: Write out the geometry to a dynamic mesh type.
 			//SplineMeshComponent->BodySetup->TriMeshGeometries
diff --git a/Engine/Plugins/PCG/Source/PCG/Private/Elements/PCGStaticMeshSpawner.cpp b/Engine/Plugins/PCG/Source/PCG/Private/Elements/PCGStaticMeshSpawner.cpp
--- a/Engine/Plugins/PCG/Source/PCG/Private/Elements/PCGStaticMeshSpawner.cpp
+++ b/Engine/Plugins/PCG/Source/PCG/Private/Elements/PCGStaticMeshSpawner.cpp
@@ -633,7 +633,7 @@ void FPCGStaticMeshSpawnerElement::SpawnStaticMeshInstances(FPCGStaticMeshSpawne
 
 	Params.bAllowDescriptorChanges = Settings->bAllowDescriptorChanges;
 
-	UPCGManagedISMComponent* MISMC = UPCGActorHelpers::GetOrCreateManagedISMC(TargetActor, Context->SourceComponent.Get(), Params);
+	UPCGManagedISMComponent* MISMC = UPCGActorHelpers::GetOrCreateManagedISMC(TargetActor, Context->SourceComponent.Get(), Params, Context);
 
 	check(MISMC);
 	MISMC->SetCrc(Context->DependenciesCrc);
diff --git a/Engine/Plugins/PCG/Source/PCG/Private/Helpers/PCGActorHelpers.cpp b/Engine/Plugins/PCG/Source/PCG/Private/Helpers/PCGActorHelpers.cpp
--- a/Engine/Plugins/PCG/Source/PCG/Private/Helpers/PCGActorHelpers.cpp
+++ b/Engine/Plugins/PCG/Source/PCG/Private/Helpers/PCGActorHelpers.cpp
@@ -67,9 +67,9 @@ UPCGManagedISMComponent* UPCGActorHelpers::GetOrCreateManagedISMC(AActor* InTarg
 	return GetOrCreateManagedISMC(InTargetActor, InSourceComponent, InParams);
 }
 
-UInstancedStaticMeshComponent* UPCGActorHelpers::GetOrCreateISMC(AActor* InTargetActor, UPCGComponent* InSourceComponent, const FPCGISMComponentBuilderParams& InParams)
+UInstancedStaticMeshComponent* UPCGActorHelpers::GetOrCreateISMC(AActor* InTargetActor, UPCGComponent* InSourceComponent, const FPCGISMComponentBuilderParams& InParams, FPCGContext* OptionalContext)
 {
-	UPCGManagedISMComponent* MISMC = GetOrCreateManagedISMC(InTargetActor, InSourceComponent, InParams);
+	UPCGManagedISMComponent* MISMC = GetOrCreateManagedISMC(InTargetActor, InSourceComponent, InParams, OptionalContext);
 	if (MISMC)
 	{
 		return MISMC->GetComponent();
@@ -80,7 +80,7 @@ UInstancedStaticMeshComponent* UPCGActorHelpers::GetOrCreateISMC(AActor* InTarge
 	}
 }
 
-UPCGManagedISMComponent* UPCGActorHelpers::GetOrCreateManagedISMC(AActor* InTargetActor, UPCGComponent* InSourceComponent, const FPCGISMComponentBuilderParams& InParams)
+UPCGManagedISMComponent* UPCGActorHelpers::GetOrCreateManagedISMC(AActor* InTargetActor, UPCGComponent* InSourceComponent, const FPCGISMComponentBuilderParams& InParams, FPCGContext* OptionalContext)
 {
 	check(InTargetActor && InSourceComponent);
 
@@ -214,7 +214,11 @@ UPCGManagedISMComponent* UPCGActorHelpers::GetOrCreateManagedISMC(AActor* InTarg
 	ISMC->RegisterComponent();
 	InTargetActor->AddInstanceComponent(ISMC);
 
-	ISMC->AttachToComponent(InTargetActor->GetRootComponent(), FAttachmentTransformRules(EAttachmentRule::KeepRelative, EAttachmentRule::KeepWorld, EAttachmentRule::KeepWorld, false));
+	if (!ISMC->AttachToComponent(InTargetActor->GetRootComponent(), FAttachmentTransformRules(EAttachmentRule::KeepRelative, EAttachmentRule::KeepWorld, EAttachmentRule::KeepWorld, false)))
+	{
+		PCGLog::Component::LogComponentAttachmentFailedWarning(OptionalContext);
+	}
+	
 	// Implementation note: Because we've used the FISMComponentDescriptor here (Descriptor vs InParams.Descriptor) which takes care of the loading, we still need to apply tags manually.
 	AddTagsToComponent(ISMC);
 
@@ -233,7 +237,7 @@ UPCGManagedISMComponent* UPCGActorHelpers::GetOrCreateManagedISMC(AActor* InTarg
 	return Resource;
 }
 
-UPCGManagedAnimBankMeshComponent* UPCGActorHelpers::GetOrCreateManagedABMC(AActor* InTargetActor, UPCGComponent* InSourceComponent, const FPCGAnimBankMeshComponentBuilderParams& InParams)
+UPCGManagedAnimBankMeshComponent* UPCGActorHelpers::GetOrCreateManagedABMC(AActor* InTargetActor, UPCGComponent* InSourceComponent, const FPCGAnimBankMeshComponentBuilderParams& InParams, FPCGContext* OptionalContext)
 {
 	check(InTargetActor && InSourceComponent);
 
@@ -334,7 +338,11 @@ UPCGManagedAnimBankMeshComponent* UPCGActorHelpers::GetOrCreateManagedABMC(AActo
 	ABMC->RegisterComponent();
 	InTargetActor->AddInstanceComponent(ABMC);
 
-	ABMC->AttachToComponent(InTargetActor->GetRootComponent(), FAttachmentTransformRules(EAttachmentRule::KeepRelative, EAttachmentRule::KeepWorld, EAttachmentRule::KeepWorld, false));
+	if (!ABMC->AttachToComponent(InTargetActor->GetRootComponent(), FAttachmentTransformRules(EAttachmentRule::KeepRelative, EAttachmentRule::KeepWorld, EAttachmentRule::KeepWorld, false)))
+	{
+		PCGLog::Component::LogComponentAttachmentFailedWarning(OptionalContext);
+	}
+	
 	// Implementation note: Because we've used the FAnimBankMeshComponentDescriptor here (Descriptor vs InParams.Descriptor) which takes care of the loading, we still need to apply tags manually.
 	AddTagsToComponent(ABMC);
 
@@ -366,9 +374,9 @@ UPCGManagedSplineMeshComponent* UPCGActorHelpers::GetOrCreateManagedSplineMeshCo
 	return GetOrCreateManagedSplineMeshComponent(InTargetActor, InSourceComponent, InParams);
 }
 
-USplineMeshComponent* UPCGActorHelpers::GetOrCreateSplineMeshComponent(AActor* InTargetActor, UPCGComponent* InSourceComponent, const FPCGSplineMeshComponentBuilderParameters& InParams)
+USplineMeshComponent* UPCGActorHelpers::GetOrCreateSplineMeshComponent(AActor* InTargetActor, UPCGComponent* InSourceComponent, const FPCGSplineMeshComponentBuilderParameters& InParams, FPCGContext* OptionalContext)
 {
-	UPCGManagedSplineMeshComponent* ManagedComponent = GetOrCreateManagedSplineMeshComponent(InTargetActor, InSourceComponent, InParams);
+	UPCGManagedSplineMeshComponent* ManagedComponent = GetOrCreateManagedSplineMeshComponent(InTargetActor, InSourceComponent, InParams, OptionalContext);
 	if (ManagedComponent)
 	{
 		return ManagedComponent->GetComponent();
@@ -379,7 +387,7 @@ USplineMeshComponent* UPCGActorHelpers::GetOrCreateSplineMeshComponent(AActor* I
 	}
 }
 
-UPCGManagedSplineMeshComponent* UPCGActorHelpers::GetOrCreateManagedSplineMeshComponent(AActor * InTargetActor, UPCGComponent * InSourceComponent, const FPCGSplineMeshComponentBuilderParameters & InParams)
+UPCGManagedSplineMeshComponent* UPCGActorHelpers::GetOrCreateManagedSplineMeshComponent(AActor* InTargetActor, UPCGComponent* InSourceComponent, const FPCGSplineMeshComponentBuilderParameters& InParams, FPCGContext* OptionalContext)
 {
 	check(InTargetActor && InSourceComponent);
 
@@ -473,7 +481,11 @@ UPCGManagedSplineMeshComponent* UPCGActorHelpers::GetOrCreateManagedSplineMeshCo
 	// Implementation note: since the data passed to the params here is in world space,
 	// We need the transform on the spline mesh component to be the identity - in world space, because unlike the ISMs
 	// where we can set the instances and specify the data is in world space, we can't do that here.
-	SplineMeshComponent->AttachToComponent(InTargetActor->GetRootComponent(), FAttachmentTransformRules(EAttachmentRule::KeepWorld, EAttachmentRule::KeepWorld, EAttachmentRule::KeepWorld, false));
+	if (!SplineMeshComponent->AttachToComponent(InTargetActor->GetRootComponent(), FAttachmentTransformRules(EAttachmentRule::KeepWorld, EAttachmentRule::KeepWorld, EAttachmentRule::KeepWorld, false)))
+	{
+		PCGLog::Component::LogComponentAttachmentFailedWarning(OptionalContext);
+	}
+	
 	SplineMeshComponent->ComponentTags.Add(InSourceComponent->GetFName());
 	SplineMeshComponent->ComponentTags.Add(PCGHelpers::DefaultPCGTag);
 
diff --git a/Engine/Plugins/PCG/Source/PCG/Private/Utils/PCGLogErrors.cpp b/Engine/Plugins/PCG/Source/PCG/Private/Utils/PCGLogErrors.cpp
--- a/Engine/Plugins/PCG/Source/PCG/Private/Utils/PCGLogErrors.cpp
+++ b/Engine/Plugins/PCG/Source/PCG/Private/Utils/PCGLogErrors.cpp
@@ -108,4 +108,17 @@ namespace PCGLog::Parsing
 	}
 }
 
+namespace PCGLog::Component
+{
+	namespace Format
+	{
+		const FText ComponentAttachmentFailed = LOCTEXT("ComponentAttachmentFailed", "Failed to attach the component, check the logs.");
+	}
+	
+	void LogComponentAttachmentFailedWarning(const FPCGContext* InContext)
+	{
+		LogWarningOnGraph(Format::ComponentAttachmentFailed, InContext);
+	}
+}
+
 #undef LOCTEXT_NAMESPACE
diff --git a/Engine/Plugins/PCG/Source/PCG/Public/Helpers/PCGActorHelpers.h b/Engine/Plugins/PCG/Source/PCG/Public/Helpers/PCGActorHelpers.h
--- a/Engine/Plugins/PCG/Source/PCG/Public/Helpers/PCGActorHelpers.h
+++ b/Engine/Plugins/PCG/Source/PCG/Public/Helpers/PCGActorHelpers.h
@@ -17,6 +17,7 @@
 #include "PCGActorHelpers.generated.h"
 
 class AActor;
+struct FPCGContext;
 class UActorComponent;
 class UDataLayerInstance;
 class UInstancedStaticMeshComponent;
@@ -120,11 +121,11 @@ PRAGMA_ENABLE_DEPRECATION_WARNINGS
 	UE_DEPRECATED(5.6, "Use signature with no SettingsUID. Make sure to fill Params.SettingsCRC to enable component reuse.")
 	static UPCGManagedSplineMeshComponent* GetOrCreateManagedSplineMeshComponent(AActor* InTargetActor, UPCGComponent* SourceComponent, uint64 SettingsUID, const FPCGSplineMeshComponentBuilderParameters& Params);
 	
-	static UInstancedStaticMeshComponent* GetOrCreateISMC(AActor* InTargetActor, UPCGComponent* SourceComponent, const FPCGISMComponentBuilderParams& Params);
-	static UPCGManagedISMComponent* GetOrCreateManagedISMC(AActor* InTargetActor, UPCGComponent* SourceComponent, const FPCGISMComponentBuilderParams& Params);
-	static UPCGManagedAnimBankMeshComponent* GetOrCreateManagedABMC(AActor* InTargetActor, UPCGComponent* SourceComponent, const FPCGAnimBankMeshComponentBuilderParams& Params);
-	static USplineMeshComponent* GetOrCreateSplineMeshComponent(AActor* InTargetActor, UPCGComponent* SourceComponent, const FPCGSplineMeshComponentBuilderParameters& Params);
-	static UPCGManagedSplineMeshComponent* GetOrCreateManagedSplineMeshComponent(AActor* InTargetActor, UPCGComponent* SourceComponent, const FPCGSplineMeshComponentBuilderParameters& Params);
+	static UInstancedStaticMeshComponent* GetOrCreateISMC(AActor* InTargetActor, UPCGComponent* SourceComponent, const FPCGISMComponentBuilderParams& Params, FPCGContext* OptionalContext = nullptr);
+	static UPCGManagedISMComponent* GetOrCreateManagedISMC(AActor* InTargetActor, UPCGComponent* SourceComponent, const FPCGISMComponentBuilderParams& Params, FPCGContext* OptionalContext = nullptr);
+	static UPCGManagedAnimBankMeshComponent* GetOrCreateManagedABMC(AActor* InTargetActor, UPCGComponent* SourceComponent, const FPCGAnimBankMeshComponentBuilderParams& Params, FPCGContext* OptionalContext = nullptr);
+	static USplineMeshComponent* GetOrCreateSplineMeshComponent(AActor* InTargetActor, UPCGComponent* SourceComponent, const FPCGSplineMeshComponentBuilderParameters& Params, FPCGContext* OptionalContext = nullptr);
+	static UPCGManagedSplineMeshComponent* GetOrCreateManagedSplineMeshComponent(AActor* InTargetActor, UPCGComponent* SourceComponent, const FPCGSplineMeshComponentBuilderParameters& Params, FPCGContext* OptionalContext = nullptr);
 
 	static bool DeleteActors(UWorld* World, const TArray<TSoftObjectPtr<AActor>>& ActorsToDelete);
 
diff --git a/Engine/Plugins/PCG/Source/PCG/Public/Utils/PCGLogErrors.h b/Engine/Plugins/PCG/Source/PCG/Public/Utils/PCGLogErrors.h
--- a/Engine/Plugins/PCG/Source/PCG/Public/Utils/PCGLogErrors.h
+++ b/Engine/Plugins/PCG/Source/PCG/Public/Utils/PCGLogErrors.h
@@ -124,4 +124,9 @@ namespace PCGLog
 		PCG_API void LogInvalidCharacterInParsedStringError(const FStringView& ParsedString, const FPCGContext* InContext = nullptr);
 		PCG_API void LogInvalidExpressionInParsedStringError(const FStringView& ParsedString, const FPCGContext* InContext = nullptr);
 	}
+
+	namespace Component
+	{
+		PCG_API void LogComponentAttachmentFailedWarning(const FPCGContext* InContext = nullptr);
+	}
 }
diff --git a/Engine/Plugins/PCG/Source/PCGEditor/Private/DataVisualizations/PCGSpatialDataVisualization.cpp b/Engine/Plugins/PCG/Source/PCGEditor/Private/DataVisualizations/PCGSpatialDataVisualization.cpp
--- a/Engine/Plugins/PCG/Source/PCGEditor/Private/DataVisualizations/PCGSpatialDataVisualization.cpp
+++ b/Engine/Plugins/PCG/Source/PCGEditor/Private/DataVisualizations/PCGSpatialDataVisualization.cpp
@@ -209,7 +209,7 @@ void IPCGSpatialDataVisualization::ExecuteDebugDisplay(FPCGContext* Context, con
 			continue;
 		}
 
-		UInstancedStaticMeshComponent* ISMC = UPCGActorHelpers::GetOrCreateISMC(TargetActor, Context->SourceComponent.Get(), Params[Direction]);
+		UInstancedStaticMeshComponent* ISMC = UPCGActorHelpers::GetOrCreateISMC(TargetActor, Context->SourceComponent.Get(), Params[Direction], Context);
 		check(ISMC && ISMC->NumCustomDataFloats == NumCustomData);
 
 		ISMC->ComponentTags.AddUnique(PCGHelpers::DefaultPCGDebugTag);
diff --git a/Engine/Plugins/PCGInterops/PCGGeometryScriptInterop/Source/PCGGeometryScriptInterop/Private/Resources/PCGDynamicMeshManagedComponent.cpp b/Engine/Plugins/PCGInterops/PCGGeometryScriptInterop/Source/PCGGeometryScriptInterop/Private/Resources/PCGDynamicMeshManagedComponent.cpp
--- a/Engine/Plugins/PCGInterops/PCGGeometryScriptInterop/Source/PCGGeometryScriptInterop/Private/Resources/PCGDynamicMeshManagedComponent.cpp
+++ b/Engine/Plugins/PCGInterops/PCGGeometryScriptInterop/Source/PCGGeometryScriptInterop/Private/Resources/PCGDynamicMeshManagedComponent.cpp
@@ -132,7 +132,10 @@ UPCGDynamicMeshManagedComponent* PCGDynamicMeshManagedComponent::GetOrCreateDyna
 		TargetActor->AddInstanceComponent(DynMeshComponent);
 
 		// Mimicking Static mesh managed resources
-		DynMeshComponent->AttachToComponent(TargetActor->GetRootComponent(), FAttachmentTransformRules(EAttachmentRule::KeepRelative, EAttachmentRule::KeepWorld, EAttachmentRule::KeepWorld, false));
+		if (!DynMeshComponent->AttachToComponent(TargetActor->GetRootComponent(), FAttachmentTransformRules(EAttachmentRule::KeepRelative, EAttachmentRule::KeepWorld, EAttachmentRule::KeepWorld, false)))
+		{
+			PCGLog::Component::LogComponentAttachmentFailedWarning(Context);
+		}
 		
 		ExistingResource->SetComponent(DynMeshComponent);
 	}
