[PCG] Add robustness to the poly line API, regarding alpha and segment length

TODO: Follow up shelf to address a related bug in the spline sampler

#rnx

diff --git a/Engine/Plugins/PCG/Source/PCG/Private/Data/PCGPolyLineData.cpp b/Engine/Plugins/PCG/Source/PCG/Private/Data/PCGPolyLineData.cpp
--- a/Engine/Plugins/PCG/Source/PCG/Private/Data/PCGPolyLineData.cpp
+++ b/Engine/Plugins/PCG/Source/PCG/Private/Data/PCGPolyLineData.cpp
@@ -47,6 +47,26 @@ void UPCGPolyLineData::GetTangentsAtSegmentStart(int SegmentIndex, FVector& OutA
 float UPCGPolyLineData::GetAlphaAtDistance(int SegmentIndex, FVector::FReal Distance) const
 {
 	const int32 NumSegments = GetNumSegments();
-	const FVector::FReal SegmentLength = GetSegmentLength(SegmentIndex);
-	return SegmentIndex / static_cast<float>(NumSegments) + (Distance / SegmentLength / NumSegments);
+	if (NumSegments < 1 || SegmentIndex < 0 || (SegmentIndex == 0 && Distance <= UE_DOUBLE_SMALL_NUMBER))
+	{
+		return 0.f;
+	}
+	else if (SegmentIndex >= NumSegments) // By definition, if the index is >= NumSegments, the alpha should be 1
+	{
+		return 1.f;
+	}
+
+	// Find the starting alpha of segments up to the current one.
+	const double SegmentStartAlpha = static_cast<double>(SegmentIndex) / NumSegments;
+
+	// Note: SegmentLength can be 0 for a start index, invalid index, or if control points are co-located.
+	const double SegmentLength = GetSegmentLength(SegmentIndex);
+
+	// Clamp the distance, as a distance longer than the current segment could artificially inflate the alpha.
+	const double LocalDistance = FMath::Clamp(Distance, 0.0, SegmentLength);
+
+	// Start the local alpha with the local distance past the segment length.
+	const double SegmentLocalAlpha = (SegmentLength > UE_DOUBLE_SMALL_NUMBER ? (LocalDistance / SegmentLength) : 0.f) / NumSegments;
+
+	return SegmentStartAlpha + SegmentLocalAlpha;
 }
diff --git a/Engine/Plugins/PCG/Source/PCG/Private/Data/PCGSplineData.cpp b/Engine/Plugins/PCG/Source/PCG/Private/Data/PCGSplineData.cpp
--- a/Engine/Plugins/PCG/Source/PCG/Private/Data/PCGSplineData.cpp
+++ b/Engine/Plugins/PCG/Source/PCG/Private/Data/PCGSplineData.cpp
@@ -104,6 +104,11 @@ int UPCGSplineData::GetNumSegments() const
 
 FVector::FReal UPCGSplineData::GetSegmentLength(int SegmentIndex) const
 {
+	if (SegmentIndex >= SplineStruct.GetNumberOfSplineSegments() || SegmentIndex < 0)
+	{
+		return 0.0;
+	}
+
 	return SplineStruct.GetDistanceAlongSplineAtSplinePoint(SegmentIndex + 1) - SplineStruct.GetDistanceAlongSplineAtSplinePoint(SegmentIndex);
 }
 
