[PCG] Correct Get PCG Component Data, Get Actor Property (PCGActorSelector) to display correct node subtitle when filters are overridden.

#rnx

diff --git a/Engine/Plugins/PCG/Source/PCG/Private/Elements/PCGActorSelector.cpp b/Engine/Plugins/PCG/Source/PCG/Private/Elements/PCGActorSelector.cpp
--- a/Engine/Plugins/PCG/Source/PCG/Private/Elements/PCGActorSelector.cpp
+++ b/Engine/Plugins/PCG/Source/PCG/Private/Elements/PCGActorSelector.cpp
@@ -501,17 +501,17 @@ uint32 GetTypeHash(const FPCGSelectionKey& In)
 }
 
 #if WITH_EDITOR
-FText FPCGActorSelectorSettings::GetTaskNameSuffix() const
+FText FPCGActorSelectorSettings::GetTaskName() const
 {
 	if (ActorFilter == EPCGActorFilter::AllWorldActors)
 	{
 		if (ActorSelection == EPCGActorSelection::ByClass)
 		{
-			return FText::Format(NSLOCTEXT("PCGActorSelectorSettings", "ClassLabel", "Class: {0}"), (ActorSelectionClass.Get() ? ActorSelectionClass->GetDisplayNameText() : FText::FromName(NAME_None)));
+			return NSLOCTEXT("PCGActorSelectorSettings", "ClassLabel", "Class");
 		}
 		else if (ActorSelection == EPCGActorSelection::ByTag)
 		{
-			return FText::Format(NSLOCTEXT("PCGActorSelectorSettings", "TagLabel", "Tag: {0}"), FText::FromName(ActorSelectionTag));
+			return NSLOCTEXT("PCGActorSelectorSettings", "TagLabel", "Tag");
 		}
 	}
 	else if(const UEnum* EnumPtr = StaticEnum<EPCGActorFilter>())
@@ -522,9 +522,21 @@ FText FPCGActorSelectorSettings::GetTaskNameSuffix() const
 	return FText();
 }
 
-FName FPCGActorSelectorSettings::GetTaskName(const FText& Prefix) const
+FText FPCGActorSelectorSettings::GetTaskNameSuffix() const
 {
-	return FName(FText::Format(NSLOCTEXT("PCGActorSelectorSettings", "NodeTitleFormat", "{0} ({1})"), Prefix, GetTaskNameSuffix()).ToString());
+	if (ActorFilter == EPCGActorFilter::AllWorldActors)
+	{
+		if (ActorSelection == EPCGActorSelection::ByClass)
+		{
+			return ActorSelectionClass.Get() ? ActorSelectionClass->GetDisplayNameText() : FText::FromName(NAME_None);
+		}
+		else if (ActorSelection == EPCGActorSelection::ByTag)
+		{
+			return FText::FromName(ActorSelectionTag);
+		}
+	}
+
+	return FText();
 }
 #endif // WITH_EDITOR
 
diff --git a/Engine/Plugins/PCG/Source/PCG/Private/Elements/PCGDataFromActor.cpp b/Engine/Plugins/PCG/Source/PCG/Private/Elements/PCGDataFromActor.cpp
--- a/Engine/Plugins/PCG/Source/PCG/Private/Elements/PCGDataFromActor.cpp
+++ b/Engine/Plugins/PCG/Source/PCG/Private/Elements/PCGDataFromActor.cpp
@@ -221,7 +221,24 @@ void UPCGDataFromActorSettings::PostLoad()
 FString UPCGDataFromActorSettings::GetAdditionalTitleInformation() const
 {
 #if WITH_EDITOR
-	return ActorSelector.GetTaskNameSuffix().ToString();
+	bool bTaskIsOverridden = false;
+	if (ActorSelector.ActorFilter == EPCGActorFilter::AllWorldActors)
+	{
+		if (ActorSelector.ActorSelection == EPCGActorSelection::ByClass
+			&& IsPropertyOverriddenByPin({GET_MEMBER_NAME_CHECKED(UPCGDataFromActorSettings, ActorSelector), GET_MEMBER_NAME_CHECKED(FPCGActorSelectorSettings, ActorSelectionClass)}))
+		{
+			bTaskIsOverridden = true;
+		}
+		else if (ActorSelector.ActorSelection == EPCGActorSelection::ByTag
+			&& IsPropertyOverriddenByPin({GET_MEMBER_NAME_CHECKED(UPCGDataFromActorSettings, ActorSelector), GET_MEMBER_NAME_CHECKED(FPCGActorSelectorSettings, ActorSelectionTag)}))
+		{
+			bTaskIsOverridden = true;
+		}
+	}
+
+	const FText Suffix = ActorSelector.GetTaskNameSuffix();
+	const FTextFormat Format = FText::FromString(Suffix.IsEmpty() ? "{0}" : "{0}: {1}");
+	return FText::Format(Format, ActorSelector.GetTaskName(), bTaskIsOverridden ? LOCTEXT("TaskOverridden", "Overridden") : Suffix).ToString();
 #else
 	return Super::GetAdditionalTitleInformation();
 #endif
diff --git a/Engine/Plugins/PCG/Source/PCG/Private/Elements/PCGGetActorProperty.cpp b/Engine/Plugins/PCG/Source/PCG/Private/Elements/PCGGetActorProperty.cpp
--- a/Engine/Plugins/PCG/Source/PCG/Private/Elements/PCGGetActorProperty.cpp
+++ b/Engine/Plugins/PCG/Source/PCG/Private/Elements/PCGGetActorProperty.cpp
@@ -107,7 +107,10 @@ void UPCGGetActorPropertySettings::PostLoad()
 FString UPCGGetActorPropertySettings::GetAdditionalTitleInformation() const
 {
 #if WITH_EDITOR
-	return FString::Printf(TEXT("%s, %s"), *ActorSelector.GetTaskNameSuffix().ToString(), *PropertyName.ToString());
+	const FText Suffix = ActorSelector.GetTaskNameSuffix();
+	return Suffix.IsEmpty()
+		? FString::Printf(TEXT("%s, %s"), *ActorSelector.GetTaskName().ToString(), *PropertyName.ToString())
+		: FString::Printf(TEXT("%s: %s, %s"), *ActorSelector.GetTaskName().ToString(), *Suffix.ToString(), *PropertyName.ToString());
 #else
 	return Super::GetAdditionalTitleInformation();
 #endif
diff --git a/Engine/Plugins/PCG/Source/PCG/Public/Elements/PCGActorSelector.h b/Engine/Plugins/PCG/Source/PCG/Public/Elements/PCGActorSelector.h
--- a/Engine/Plugins/PCG/Source/PCG/Public/Elements/PCGActorSelector.h
+++ b/Engine/Plugins/PCG/Source/PCG/Public/Elements/PCGActorSelector.h
@@ -171,9 +171,9 @@ struct PCG_API FPCGActorSelectorSettings
 	UPROPERTY(Transient, meta = (EditCondition = false, EditConditionHides))
 	bool bShowIgnoreSelfAndChildren = true;
 
-#if WITH_EDITOR	
+#if WITH_EDITOR
+	FText GetTaskName() const;
 	FText GetTaskNameSuffix() const;
-	FName GetTaskName(const FText& Prefix) const;
 #endif
 
 	FPCGSelectionKey GetAssociatedKey() const;
