[PCG] Do not execute while transaction is active (disabled by default, will go wider on testing before enabling)
- pcg.PauseExecutionWhileTransactionActive
- This will prevent from capturing PCG output inside transactions causing undo/redo issues

#rnx

diff --git a/Engine/Plugins/PCG/Source/PCG/Private/Graph/PCGGraphExecutor.cpp b/Engine/Plugins/PCG/Source/PCG/Private/Graph/PCGGraphExecutor.cpp
--- a/Engine/Plugins/PCG/Source/PCG/Private/Graph/PCGGraphExecutor.cpp
+++ b/Engine/Plugins/PCG/Source/PCG/Private/Graph/PCGGraphExecutor.cpp
@@ -73,6 +73,11 @@ namespace PCGGraphExecutor
 		TEXT("pcg.EditorNotificationDelayInSeconds"),
 		3.0f,
 		TEXT("Time in seconds to wait before showing a notification to avoid spamming for short tasks"));
+
+	TAutoConsoleVariable<bool> CVarPausePCGExecutionWhileTransactionActive(
+		TEXT("pcg.PauseExecutionWhileTransactionActive"),
+		false,
+		TEXT("Pauses all execution of PCG while user transactions are active"));
 #endif
 
 	TAutoConsoleVariable<bool> CVarDynamicTaskCulling(
@@ -1498,7 +1503,9 @@ void FPCGGraphExecutor::Execute()
 	TRACE_CPUPROFILER_EVENT_SCOPE(FPCGGraphExecutor::Execute);
 
 #if WITH_EDITOR
-	if (PCGSystemSwitches::CVarPausePCGExecution.GetValueOnAnyThread())
+	const bool bIsEditorWorld = !World || !World->IsGameWorld();
+
+	if (PCGSystemSwitches::CVarPausePCGExecution.GetValueOnAnyThread() || (bIsEditorWorld && GEditor && GEditor->IsTransactionActive() && PCGGraphExecutor::CVarPausePCGExecutionWhileTransactionActive.GetValueOnAnyThread()))
 	{
 		UpdateGenerationNotification();
 		return;
