[PCG] Cleanup now loads unloaded actors before cleaning up (disabled by default, will go wider on testing before enabling)
- Scale: this will eventually need to be replaced by proper cleanup of unloaded actors which doesn't require loading
- pcg.LoadManagedActorsForCleanup

#rnx

diff --git a/Engine/Plugins/PCG/Source/PCG/Private/PCGComponent.cpp b/Engine/Plugins/PCG/Source/PCG/Private/PCGComponent.cpp
--- a/Engine/Plugins/PCG/Source/PCG/Private/PCGComponent.cpp
+++ b/Engine/Plugins/PCG/Source/PCG/Private/PCGComponent.cpp
@@ -45,6 +45,10 @@
 #include "ScopedTransaction.h"
 #include "Editor/EditorEngine.h"
 #include "Editor/Transactor.h"
+#include "WorldPartition/WorldPartition.h"
+#include "WorldPartition/WorldPartitionActorDescInstance.h"
+#include "WorldPartition/WorldPartitionHandle.h"
+#include "WorldPartition/WorldPartitionHelpers.h"
 #endif
 
 #define LOCTEXT_NAMESPACE "UPCGComponent"
@@ -59,6 +63,53 @@ namespace PCGComponent
 #endif // WITH_EDITOR
 		Delegate.Broadcast(PCGComponent);
 	}
+
+#if WITH_EDITOR
+	static TAutoConsoleVariable<bool> CVarLoadManagedActorsForCleanup(
+		TEXT("pcg.LoadManagedActorsForCleanup"),
+		true,
+		TEXT("When enabled unloaded actors will get loaded before clean up"));
+#endif
+
+	struct FPCGManagedActorLoadingScope
+	{
+		void AddResources(UPCGComponent* InComponent, const TArray<UPCGManagedResource*>& InManagedResources)
+		{
+#if WITH_EDITOR
+			if (CVarLoadManagedActorsForCleanup.GetValueOnAnyThread())
+			{
+				if (!InComponent || !InComponent->GetWorld() || InComponent->GetWorld()->IsGameWorld())
+				{
+					return;
+				}
+
+				UWorldPartition* WorldPartition = FWorldPartitionHelpers::GetWorldPartition(InComponent);
+				if (!WorldPartition)
+				{
+					return;
+				}
+
+				for (UPCGManagedResource* ManagedResource : InManagedResources)
+				{
+					if (UPCGManagedActors* ManagedActors = Cast<UPCGManagedActors>(ManagedResource))
+					{
+						for (const TSoftObjectPtr<AActor>& ManagedActor : ManagedActors->GeneratedActors)
+						{
+							if (const FWorldPartitionActorDescInstance* ActorDescInstance = WorldPartition->GetActorDescInstanceByPath(ManagedActor.ToSoftObjectPath()))
+							{
+								LoadedActors.Add(FWorldPartitionReference(WorldPartition, ActorDescInstance->GetGuid()));
+							}
+						}
+					}
+				}
+			}
+#endif // WITH_EDITOR 
+		}
+
+#if WITH_EDITOR
+		TSet<FWorldPartitionReference> LoadedActors;
+#endif
+	};
 }
 
 UPCGComponent::UPCGComponent(const FObjectInitializer& InObjectInitializer)
@@ -1273,41 +1324,43 @@ void UPCGComponent::CleanupLocalImmediate(bool bRemoveComponents, bool bCleanupL
 	// Cancels generation of this component if there is an ongoing generation in progress.
 	CancelGeneration();
 
-	TSet<TSoftObjectPtr<AActor>> ActorsToDelete;
-
-	if (!bRemoveComponents && UPCGManagedResource::DebugForcePurgeAllResourcesOnGenerate())
 	{
-		bRemoveComponents = true;
-	}
+		TSet<TSoftObjectPtr<AActor>> ActorsToDelete;
+		PCGComponent::FPCGManagedActorLoadingScope Scope;
 
-	{
-		FScopeLock ResourcesLock(&GeneratedResourcesLock);
-		check(!GeneratedResourcesInaccessible);
-		for (int32 ResourceIndex = GeneratedResources.Num() - 1; ResourceIndex >= 0; --ResourceIndex)
+		if (!bRemoveComponents && UPCGManagedResource::DebugForcePurgeAllResourcesOnGenerate())
 		{
-			// Note: resources can be null here in some loading + bp object cases
-			UPCGManagedResource* Resource = GeneratedResources[ResourceIndex];
-
-			PCGGeneratedResourcesLogging::LogCleanupLocalImmediateResource(this, Resource);
+			bRemoveComponents = true;
+		}
 
-			if (!Resource || Resource->Release(bRemoveComponents, ActorsToDelete))
+		{
+			FScopeLock ResourcesLock(&GeneratedResourcesLock);
+			check(!GeneratedResourcesInaccessible);
+			Scope.AddResources(this, GeneratedResources);
+			for (int32 ResourceIndex = GeneratedResources.Num() - 1; ResourceIndex >= 0; --ResourceIndex)
 			{
-				if (Resource)
+				// Note: resources can be null here in some loading + bp object cases
+				UPCGManagedResource* Resource = GeneratedResources[ResourceIndex];
+
+				PCGGeneratedResourcesLogging::LogCleanupLocalImmediateResource(this, Resource);
+
+				if (!Resource || Resource->Release(bRemoveComponents, ActorsToDelete))
 				{
 #if WITH_EDITOR
-					if (!Resource->IsMarkedTransientOnLoad())
+					if (Resource && !Resource->IsMarkedTransientOnLoad())
 					{
 						Resource->Rename(nullptr, nullptr, REN_DoNotDirty | REN_DontCreateRedirectors | REN_NonTransactional);
 					}
 #endif
+
+					GeneratedResources.RemoveAt(ResourceIndex);
 				}
 
-				GeneratedResources.RemoveAtSwap(ResourceIndex);
 			}
 		}
-	}
 
-	UPCGActorHelpers::DeleteActors(GetWorld(), ActorsToDelete.Array());
+		UPCGActorHelpers::DeleteActors(GetWorld(), ActorsToDelete.Array());
+	}
 
 	// If bRemoveComponents is true, it means we are in a "real" cleanup, not a pre-cleanup before a generate.
 	// So call PostCleanup in this case.
@@ -1350,14 +1403,14 @@ FPCGTaskId UPCGComponent::CreateCleanupTask(bool bRemoveComponents, const TArray
 		AllDependencies = &AdditionalDependencies;
 	}
 
-	struct CleanupContext
+	struct FCleanupContext : public PCGComponent::FPCGManagedActorLoadingScope
 	{
 		bool bIsFirstIteration = true;
 		int32 ResourceIndex = -1;
 		TSet<TSoftObjectPtr<AActor>> ActorsToDelete;
 	};
 
-	TSharedPtr<CleanupContext> Context = MakeShared<CleanupContext>();
+	TSharedPtr<FCleanupContext> Context = MakeShared<FCleanupContext>();
 	TWeakObjectPtr<UPCGComponent> ThisComponentWeakPtr(this);
 	TWeakObjectPtr<UWorld> WorldPtr(GetWorld());
 
@@ -1380,6 +1433,7 @@ FPCGTaskId UPCGComponent::CreateCleanupTask(bool bRemoveComponents, const TArray
 				ThisComponent->GeneratedResourcesInaccessible = true;
 				Context->ResourceIndex = ThisComponent->GeneratedResources.Num() - 1;
 				Context->bIsFirstIteration = false;
+				Context->AddResources(ThisComponent, ThisComponent->GeneratedResources);
 			}
 
 			// Going backward
@@ -1396,15 +1450,12 @@ FPCGTaskId UPCGComponent::CreateCleanupTask(bool bRemoveComponents, const TArray
 
 				if (!Resource || Resource->Release(bRemoveComponents, Context->ActorsToDelete))
 				{
-					if (Resource)
-					{
 #if WITH_EDITOR
-						if (!Resource->IsMarkedTransientOnLoad())
-						{
-							Resource->Rename(nullptr, nullptr, REN_DoNotDirty | REN_DontCreateRedirectors | REN_NonTransactional);
-						}
-#endif
+					if (Resource && !Resource->IsMarkedTransientOnLoad())
+					{
+						Resource->Rename(nullptr, nullptr, REN_DoNotDirty | REN_DontCreateRedirectors | REN_NonTransactional);
 					}
+#endif
 
 					ThisComponent->GeneratedResources.RemoveAtSwap(Context->ResourceIndex);
 				}
@@ -1528,16 +1579,12 @@ void UPCGComponent::CleanupUnusedManagedResources()
 
 			if (!Resource || Resource->ReleaseIfUnused(ActorsToDelete))
 			{
-				if (Resource)
-				{
 #if WITH_EDITOR
-					if (!Resource->IsMarkedTransientOnLoad())
-					{
-						Resource->Rename(nullptr, nullptr, REN_DoNotDirty | REN_DontCreateRedirectors | REN_NonTransactional);
-					}
-#endif
+				if (Resource && !Resource->IsMarkedTransientOnLoad())
+				{
+					Resource->Rename(nullptr, nullptr, REN_DoNotDirty | REN_DontCreateRedirectors | REN_NonTransactional);
 				}
-
+#endif
 				GeneratedResources.RemoveAtSwap(ResourceIndex);
 			}
 		}
