[PCG] Refactor accessor creation in PCG Data in factory to allow other types to also have special accessors

#rnx

diff --git a/Engine/Plugins/PCG/Source/PCG/Private/Data/PCGPointData.cpp b/Engine/Plugins/PCG/Source/PCG/Private/Data/PCGPointData.cpp
--- a/Engine/Plugins/PCG/Source/PCG/Private/Data/PCGPointData.cpp
+++ b/Engine/Plugins/PCG/Source/PCG/Private/Data/PCGPointData.cpp
@@ -313,6 +313,62 @@ void UPCGPointData::AddToCrc(FArchiveCrc32& Ar, bool bFullDataCrc) const
 	}
 }
 
+TUniquePtr<IPCGAttributeAccessor> UPCGPointData::CreateStaticAccessor(const FPCGAttributePropertySelector& InSelector, bool bQuiet)
+{
+	TUniquePtr<IPCGAttributeAccessor> Accessor;
+	
+	if (InSelector.GetSelection() == EPCGAttributePropertySelection::PointProperty)
+	{
+		const FName PropertyName = InSelector.GetName();
+		if (const FProperty* Property = FPCGPoint::StaticStruct()->FindPropertyByName(PropertyName))
+		{
+			Accessor = PCGAttributeAccessorHelpers::CreatePropertyAccessor(Property);
+		}
+		else if (FPCGPoint::HasCustomPropertyGetterSetter(PropertyName))
+		{
+			Accessor = FPCGPoint::CreateCustomPropertyAccessor(PropertyName);
+		}
+	}
+	
+	return Accessor;
+}
+
+FPCGAttributeAccessorMethods UPCGPointData::GetPointAccessorMethods()
+{
+	auto CreateAccessorFunc = [](UPCGData* InData, const FPCGAttributePropertySelector& InSelector, bool bQuiet)
+	{
+		return CreateStaticAccessor(InSelector, bQuiet);
+	};
+
+	auto CreateConstAccessorFunc = [](const UPCGData* InData, const FPCGAttributePropertySelector& InSelector, bool bQuiet) -> TUniquePtr<const IPCGAttributeAccessor>
+	{
+		return CreateStaticAccessor(InSelector, bQuiet);
+	};
+
+	auto CreateAccessorKeysFunc = [](UPCGData* InData, const FPCGAttributePropertySelector& InSelector, bool bQuiet)
+	{
+		UPCGPointData* PointData = CastChecked<UPCGPointData>(InData);
+		TArrayView<FPCGPoint> View(PointData->GetMutablePoints());
+		return MakeUnique<FPCGAttributeAccessorKeysPoints>(View);
+	};
+	
+	auto CreateConstAccessorKeysFunc = [](const UPCGData* InData, const FPCGAttributePropertySelector& InSelector, bool bQuiet) -> TUniquePtr<const IPCGAttributeAccessorKeys>
+	{
+		const UPCGPointData* PointData = CastChecked<UPCGPointData>(InData);
+		return MakeUnique<const FPCGAttributeAccessorKeysPoints>(PointData->GetPoints());
+	};
+	
+	FPCGAttributeAccessorMethods Methods
+	{
+		.CreateAccessorFunc = CreateAccessorFunc,
+		.CreateConstAccessorFunc = CreateConstAccessorFunc,
+		.CreateAccessorKeysFunc = CreateAccessorKeysFunc,
+		.CreateConstAccessorKeysFunc = CreateConstAccessorKeysFunc
+	};
+
+	return Methods;
+}
+
 FBox UPCGPointData::GetBounds() const
 {
 	if (bBoundsAreDirty)
diff --git a/Engine/Plugins/PCG/Source/PCG/Private/Metadata/Accessors/PCGAttributeAccessorFactory.cpp b/Engine/Plugins/PCG/Source/PCG/Private/Metadata/Accessors/PCGAttributeAccessorFactory.cpp
new file mode 100644
--- /dev/null
+++ b/Engine/Plugins/PCG/Source/PCG/Private/Metadata/Accessors/PCGAttributeAccessorFactory.cpp
@@ -0,0 +1,244 @@
+// Copyright Epic Games, Inc. All Rights Reserved.
+
+#include "Metadata/Accessors/PCGAttributeAccessorFactory.h"
+
+#include "PCGModule.h"
+#include "Data/PCGSpatialData.h"
+#include "Metadata/PCGMetadata.h"
+#include "Metadata/PCGMetadataAttribute.h"
+#include "Metadata/Accessors/PCGAttributeAccessorHelpers.h"
+
+namespace PCGAttributeAccessorFactory
+{
+	TUniquePtr<IPCGAttributeAccessor> CreateDefaultAccessor(UPCGData* InData, const FPCGAttributePropertySelector& InSelector, const bool bQuiet)
+	{
+		check(InData);
+		TUniquePtr<IPCGAttributeAccessor> Accessor;
+	
+		if (InSelector.GetSelection() == EPCGAttributePropertySelection::ExtraProperty)
+		{
+			Accessor = PCGAttributeAccessorHelpers::CreateExtraAccessor(InSelector.GetExtraProperty());
+
+			if (!Accessor.IsValid())
+			{
+				if (!bQuiet)
+				{
+					UE_LOG(LogPCG, Error, TEXT("[FPCGAttributeAccessorFactory::DefaultAccessor] Expected to select an extra property but the data doesn't support this property."));
+				}
+				
+				return TUniquePtr<IPCGAttributeAccessor>();
+			}
+		}
+		else if (InSelector.GetSelection() == EPCGAttributePropertySelection::Attribute)
+		{
+			UPCGMetadata* Metadata = InData->MutableMetadata();
+			FPCGMetadataAttributeBase* Attribute = Metadata ? Metadata->GetMutableAttribute(InSelector.GetAttributeName()) : nullptr;
+		
+			Accessor = PCGAttributeAccessorHelpers::CreateAccessor(Attribute, Metadata, bQuiet);
+		}
+
+		return Accessor;
+	}
+	
+	TUniquePtr<const IPCGAttributeAccessor> CreateDefaultConstAccessor(const UPCGData* InData, const FPCGAttributePropertySelector& InSelector, const bool bQuiet)
+	{
+		check(InData);
+		TUniquePtr<const IPCGAttributeAccessor> Accessor;
+	
+		if (InSelector.GetSelection() == EPCGAttributePropertySelection::ExtraProperty)
+		{
+			Accessor = PCGAttributeAccessorHelpers::CreateExtraAccessor(InSelector.GetExtraProperty());
+
+			if (!Accessor.IsValid())
+			{
+				if (!bQuiet)
+				{
+					UE_LOG(LogPCG, Error, TEXT("[FPCGAttributeAccessorFactory::DefaultConstAccessor] Expected to select an extra property but the data doesn't support this property."));
+				}
+				
+				return TUniquePtr<IPCGAttributeAccessor>();
+			}
+		}
+		else if (InSelector.GetSelection() == EPCGAttributePropertySelection::Attribute)
+		{
+			const UPCGMetadata* Metadata = InData->ConstMetadata();
+			const FPCGMetadataAttributeBase* Attribute = Metadata ? Metadata->GetConstAttribute(InSelector.GetAttributeName()) : nullptr;
+		
+			Accessor = PCGAttributeAccessorHelpers::CreateConstAccessor(Attribute, Metadata, bQuiet);
+		}
+
+		return Accessor;
+	}
+	
+	TUniquePtr<IPCGAttributeAccessorKeys> CreateDefaultAccessorKeys(UPCGData* InData, const FPCGAttributePropertySelector& InSelector, const bool bQuiet)
+	{
+		UPCGMetadata* Metadata = InData->MutableMetadata();
+		return Metadata ? MakeUnique<FPCGAttributeAccessorKeysEntries>(Metadata) : TUniquePtr<IPCGAttributeAccessorKeys>();
+	}
+	
+	TUniquePtr<const IPCGAttributeAccessorKeys> CreateDefaultConstAccessorKeys(const UPCGData* InData, const FPCGAttributePropertySelector& InSelector, const bool bQuiet)
+	{
+		const UPCGMetadata* Metadata = InData->ConstMetadata();
+		return Metadata ? MakeUnique<const FPCGAttributeAccessorKeysEntries>(Metadata) : TUniquePtr<const IPCGAttributeAccessorKeys>();
+	}
+	
+	TUniquePtr<const IPCGAttributeAccessorKeys> CreateSpatialConstAccessorKeys(const UPCGData* InData, const FPCGAttributePropertySelector& InSelector, const bool bQuiet)
+	{
+		// For const metadata, that is on SpatialData (not points), we allow for Default Value to support "data-wide" metadata
+		// It's only for spatial data and not Param data because for ParamData, GetNum on the keys should return the number of entries.
+		// If we have a default key, GetNum will return 1, while there is actually 0 values.
+		const UPCGMetadata* Metadata = InData->ConstMetadata();
+		return Metadata ? MakeUnique<const FPCGAttributeAccessorKeysEntries>(Metadata, /*bAddDefaultValueIfEmpty=*/true) : TUniquePtr<const IPCGAttributeAccessorKeys>();
+	}
+}
+
+template <typename PCGData, typename Func>
+decltype(auto) FPCGAttributeAccessorFactory::CallOnMethod(PCGData* InData, Func Callback) const
+{
+	using AccessorType = decltype(Callback(FPCGAttributeAccessorMethods{}));
+
+	if (!InData)
+	{
+		return AccessorType();
+	}
+
+	TSubclassOf<UPCGData> CurrentClass = InData->GetClass();
+	check(CurrentClass);
+
+	// Current Class will dereference to child of UPCGData classes, until we have the super class of UPCGData, which will dereference to nullptr
+	while (*CurrentClass)
+	{
+		if (const FPCGAttributeAccessorMethods* Methods = AccessorMethods.Find(CurrentClass))
+		{
+			AccessorType Accessor = Callback(*Methods);
+			if (Accessor.IsValid())
+			{
+				return Accessor;
+			}
+		}
+
+		CurrentClass = CurrentClass->GetSuperClass();
+	}
+
+	return AccessorType();
+}
+
+FPCGAttributeAccessorFactory& FPCGAttributeAccessorFactory::GetMutableInstance()
+{
+	return FPCGModule::GetMutableAttributeAccessorFactory();
+}
+
+const FPCGAttributeAccessorFactory& FPCGAttributeAccessorFactory::GetInstance()
+{
+	return FPCGModule::GetConstAttributeAccessorFactory();
+}
+
+TUniquePtr<IPCGAttributeAccessor> FPCGAttributeAccessorFactory::CreateSimpleAccessor(UPCGData* InData,
+	const FPCGAttributePropertySelector& InSelector, const bool bQuiet) const
+{
+	return CallOnMethod(InData, [&InData, &InSelector, bQuiet](const FPCGAttributeAccessorMethods& Methods)
+	{
+		return Methods.CreateAccessorFunc ? Methods.CreateAccessorFunc(InData, InSelector, bQuiet) : TUniquePtr<IPCGAttributeAccessor>();
+	});
+}
+
+TUniquePtr<const IPCGAttributeAccessor> FPCGAttributeAccessorFactory::CreateSimpleConstAccessor(const UPCGData* InData,
+	const FPCGAttributePropertySelector& InSelector, const bool bQuiet) const
+{
+	return CallOnMethod(InData, [&InData, &InSelector, bQuiet](const FPCGAttributeAccessorMethods& Methods)
+	{
+		return Methods.CreateConstAccessorFunc ? Methods.CreateConstAccessorFunc(InData, InSelector, bQuiet) : TUniquePtr<const IPCGAttributeAccessor>();
+	});
+}
+
+TUniquePtr<IPCGAttributeAccessorKeys> FPCGAttributeAccessorFactory::CreateSimpleKeys(UPCGData* InData,
+	const FPCGAttributePropertySelector& InSelector, const bool bQuiet) const
+{
+	if (!InData)
+	{
+		if (!bQuiet)
+		{
+			UE_LOG(LogPCG, Error, TEXT("[FPCGAttributeAccessorFactory::CreateSimpleKeys] Can't create keys with no input data."));
+		}
+		
+		return {};
+	}
+	
+	return CallOnMethod(InData, [&InData, &InSelector, bQuiet](const FPCGAttributeAccessorMethods& Methods)
+	{
+		return Methods.CreateAccessorKeysFunc ? Methods.CreateAccessorKeysFunc(InData, InSelector, bQuiet) : TUniquePtr<IPCGAttributeAccessorKeys>();
+	});
+}
+
+TUniquePtr<const IPCGAttributeAccessorKeys> FPCGAttributeAccessorFactory::CreateSimpleConstKeys(const UPCGData* InData,
+	const FPCGAttributePropertySelector& InSelector, const bool bQuiet) const
+{
+	if (!InData)
+	{
+		if (!bQuiet)
+		{
+			UE_LOG(LogPCG, Error, TEXT("[FPCGAttributeAccessorFactory::CreateSimpleConstKeys] Can't create keys with no input data."));
+		}
+		
+		return {};
+	}
+	
+	return CallOnMethod(InData, [&InData, &InSelector, bQuiet](const FPCGAttributeAccessorMethods& Methods)
+	{
+		return Methods.CreateConstAccessorKeysFunc ? Methods.CreateConstAccessorKeysFunc(InData, InSelector, bQuiet) : TUniquePtr<const IPCGAttributeAccessorKeys>();
+	});
+}
+
+void FPCGAttributeAccessorFactory::ForEachSelectorMenuEntry(TFunctionRef<void(const FPCGAttributeSelectorMenuEntry&)> Callback) const
+{
+	for (const auto& [Class, Methods] : AccessorMethods)
+	{
+		for (const FPCGAttributeSelectorMenuEntry& SelectorMenuEntry : Methods.AttributeSelectorOptions)
+		{
+			Callback(SelectorMenuEntry);
+		}
+	}
+}
+
+void FPCGAttributeAccessorFactory::RegisterMethods(TSubclassOf<UPCGData> PCGDataClass, FPCGAttributeAccessorMethods&& InAccessorMethods)
+{
+	check(*PCGDataClass);
+	
+	if (!ensure(!AccessorMethods.Contains(PCGDataClass)))
+	{
+		UE_LOG(LogPCG, Error, TEXT("Trying to register %s accessor methods multiple times, will be ignored."), *PCGDataClass->GetName());
+		return;
+	}
+	
+	AccessorMethods.Add(PCGDataClass, std::move(InAccessorMethods));
+}
+
+void FPCGAttributeAccessorFactory::UnregisterMethods(TSubclassOf<UPCGData> PCGDataClass)
+{
+	AccessorMethods.Remove(PCGDataClass);
+}
+
+void FPCGAttributeAccessorFactory::RegisterDefaultMethods()
+{
+	FPCGAttributeAccessorMethods DefaultMethods
+	{
+		.CreateAccessorFunc = PCGAttributeAccessorFactory::CreateDefaultAccessor,
+		.CreateConstAccessorFunc = PCGAttributeAccessorFactory::CreateDefaultConstAccessor,
+		.CreateAccessorKeysFunc = PCGAttributeAccessorFactory::CreateDefaultAccessorKeys,
+		.CreateConstAccessorKeysFunc = PCGAttributeAccessorFactory::CreateDefaultConstAccessorKeys,
+	};
+
+	FPCGAttributeAccessorMethods SpatialMethods
+	{
+		.CreateConstAccessorKeysFunc = PCGAttributeAccessorFactory::CreateSpatialConstAccessorKeys,
+	};
+
+	RegisterMethods<UPCGData>(std::move(DefaultMethods));
+	RegisterMethods<UPCGSpatialData>(std::move(DefaultMethods));
+}
+
+void FPCGAttributeAccessorFactory::UnregisterDefaultMethods()
+{
+	UnregisterMethods<UPCGSpatialData>();
+	UnregisterMethods<UPCGData>();
+}
diff --git a/Engine/Plugins/PCG/Source/PCG/Private/Metadata/Accessors/PCGAttributeAccessorHelpers.cpp b/Engine/Plugins/PCG/Source/PCG/Private/Metadata/Accessors/PCGAttributeAccessorHelpers.cpp
--- a/Engine/Plugins/PCG/Source/PCG/Private/Metadata/Accessors/PCGAttributeAccessorHelpers.cpp
+++ b/Engine/Plugins/PCG/Source/PCG/Private/Metadata/Accessors/PCGAttributeAccessorHelpers.cpp
@@ -21,215 +21,6 @@
 
 namespace PCGAttributeAccessorHelpers
 {
-	TUniquePtr<IPCGAttributeAccessor> CreateAttributeAccessorImpl(FPCGMetadataAttributeBase* Attribute, UPCGMetadata* Metadata)
-	{
-		auto CreateTypedAccessor = [Attribute, Metadata]<typename T>(T) -> TUniquePtr<IPCGAttributeAccessor>
-		{
-			return MakeUnique<FPCGAttributeAccessor<T>>(static_cast<FPCGMetadataAttribute<T>*>(Attribute), Metadata);
-		};
-
-		if (Attribute && Metadata)
-		{
-			return PCGMetadataAttribute::CallbackWithRightType(Attribute->GetTypeId(), CreateTypedAccessor);
-		}
-		else
-		{
-			return {};
-		}
-	}
-
-	// Return a non-const version but it will be read only anyway. Done like that to factorize some non-const/const version of accessor creation.
-	TUniquePtr<IPCGAttributeAccessor> CreateConstAttributeAccessorImpl(const FPCGMetadataAttributeBase* Attribute, const UPCGMetadata* Metadata)
-	{
-		auto CreateTypedAccessor = [Attribute, Metadata]<typename T>(T) -> TUniquePtr<IPCGAttributeAccessor>
-		{
-			return MakeUnique<FPCGAttributeAccessor<T>>(static_cast<const FPCGMetadataAttribute<T>*>(Attribute), Metadata, /*bForceReadOnly=*/true);
-		};
-
-		if (Attribute && Metadata)
-		{
-			return PCGMetadataAttribute::CallbackWithRightType(Attribute->GetTypeId(), CreateTypedAccessor);
-		}
-		else
-		{
-			return {};
-		}
-	}
-
-	// Don't be afraid of this enormous template!
-	// We want to be able to create a TUniquePtr<const IPCGAttributeAccessor> for const accessors and TUniquePtr<IPCGAttributeAccessor> for mutable accessors.
-	// This template just says that:
-	// * AccessorType needs to be IPCGAttributeAccessor or const IPCGAttributeAccessor
-	// * DataType needs to be UPCGData or const UPCGData
-	// * And if one of them is const, they both need to be const.
-	// It allows us to factorise the same logic for const and non-const version (with a few const_cast in the process, but we still have something const at the end)
-	template <
-		typename AccessorType, 
-		typename DataType, 
-		typename = typename std::enable_if_t<
-			std::is_same_v<IPCGAttributeAccessor, std::remove_const_t<AccessorType>> &&
-			std::is_same_v<UPCGData, std::remove_const_t<DataType>> &&
-			std::is_const_v<AccessorType> == std::is_const_v<DataType>
-		>
-	>
-	TUniquePtr<AccessorType> CreateAccessorImpl(DataType* InData, const FPCGAttributePropertySelector& InSelector, bool bQuiet)
-	{
-		const FName Name = InSelector.GetName();
-		TUniquePtr<IPCGAttributeAccessor> Accessor;
-
-		if (InSelector.GetSelection() == EPCGAttributePropertySelection::PointProperty)
-		{
-			if (!InData || InData->template IsA<UPCGPointData>())
-			{
-				if (const FProperty* Property = FPCGPoint::StaticStruct()->FindPropertyByName(Name))
-				{
-					Accessor = CreatePropertyAccessor(Property);
-				}
-				else if (FPCGPoint::HasCustomPropertyGetterSetter(Name))
-				{
-					Accessor = FPCGPoint::CreateCustomPropertyAccessor(Name);
-				}
-			}
-		}
-		else if (InSelector.GetSelection() == EPCGAttributePropertySelection::ExtraProperty)
-		{
-			Accessor = CreateExtraAccessor(InSelector.GetExtraProperty());
-
-			if (!Accessor.IsValid())
-			{
-				if (!bQuiet)
-				{
-					UE_LOG(LogPCG, Error, TEXT("[PCGAttributeAccessorHelpers::ConstAccessor] Expected to select an extra property but the data doesn't support this property."));
-				}
-				
-				return TUniquePtr<IPCGAttributeAccessor>();
-			}
-		}
-
-		// At this point, it is not a point data or we didn't find a property.
-		// We can't continue if it is a property wanted.
-		if (InSelector.GetSelection() == EPCGAttributePropertySelection::PointProperty && !Accessor.IsValid())
-		{
-			if (!bQuiet)
-			{
-				UE_LOG(LogPCG, Error, TEXT("[PCGAttributeAccessorHelpers::CreateAccessor] Expected to select a property but the data doesn't support this property."));
-			}
-			
-			return TUniquePtr<IPCGAttributeAccessor>();
-		}
-
-		if (InSelector.GetSelection() == EPCGAttributePropertySelection::Attribute)
-		{
-			// We can't use const_cast here to only have a single path, as GetMutableAttribute will set the last selector
-			// and we don't want that on const data.
-			if constexpr (std::is_const_v<AccessorType>)
-			{
-				const UPCGMetadata* ConstMetadata = InData ? InData->ConstMetadata() : nullptr;
-				const FPCGMetadataAttributeBase* ConstAttribute = ConstMetadata ? ConstMetadata->GetConstAttribute(Name) : nullptr;
-				
-				Accessor = CreateConstAttributeAccessorImpl(ConstAttribute, ConstMetadata);
-			}
-			else
-			{
-				UPCGMetadata* MutableMetadata = InData ? InData->MutableMetadata() : nullptr;
-				FPCGMetadataAttributeBase* MutableAttribute = MutableMetadata ? MutableMetadata->GetMutableAttribute(Name) : nullptr;
-				
-				Accessor = CreateAttributeAccessorImpl(MutableAttribute, MutableMetadata);
-			}
-		}
-
-		if (!Accessor.IsValid())
-		{
-			return TUniquePtr<AccessorType>();
-		}
-
-		// At this point, check if we have chain accessors
-		for (const FString& ExtraName : InSelector.GetExtraNames())
-		{
-			bool bSuccess = false;
-			Accessor = CreateChainAccessor(std::move(Accessor), FName(ExtraName), bSuccess);
-			if (!bSuccess)
-			{
-				if (!bQuiet)
-				{
-					UE_LOG(LogPCG, Error, TEXT("[PCGAttributeAccessorHelpers::CreateAccessor] Extra selectors don't match existing properties."));
-				}
-				
-				return TUniquePtr<AccessorType>();
-			}
-		}
-
-		// Can't create mutable accessor that are read only.
-		if constexpr (!std::is_const_v<AccessorType>)
-		{
-			if (Accessor && Accessor->IsReadOnly())
-			{
-				if (!bQuiet)
-				{
-					UE_LOG(LogPCG, Error, TEXT("[PCGAttributeAccessorHelpers::CreateAccessor] Attribute can not be written into, since it is read-only."));
-				}
-				
-				return TUniquePtr<AccessorType>();
-			}
-
-			// Set the cached selector
-			if (InData)
-			{
-				InData->SetLastSelector(InSelector);
-			}
-		}
-
-		return Accessor;
-	}
-
-	// Same thing for the keys. Check CreateAccessorImpl for more info about this big template.
-	template <
-		typename KeysType,
-		typename DataType,
-		typename = typename std::enable_if_t<
-			std::is_same_v<IPCGAttributeAccessorKeys, std::remove_const_t<KeysType>> &&
-			std::is_same_v<UPCGData, std::remove_const_t<DataType>> &&
-			std::is_const_v<KeysType> == std::is_const_v<DataType>
-		>
-	>
-	TUniquePtr<KeysType> CreateKeysImpl(DataType* InData, const FPCGAttributePropertySelector& InSelector)
-	{
-		if (UPCGPointData* PointData = Cast<UPCGPointData>(const_cast<UPCGData*>(InData)))
-		{
-			if constexpr (std::is_const_v<KeysType>)
-			{
-				return MakeUnique<FPCGAttributeAccessorKeysPoints>(PointData->GetPoints());
-			}
-			else
-			{
-				TArrayView<FPCGPoint> View(PointData->GetMutablePoints());
-				return MakeUnique<FPCGAttributeAccessorKeysPoints>(View);
-			}
-		}
-		else if (InData) // e.g. is UPCGData and non-null
-		{
-			if constexpr (std::is_const_v<KeysType>)
-			{
-				if (const UPCGMetadata* Metadata = InData->ConstMetadata())
-				{
-					// For const metadata, that is on SpatialData (not points), we allow for Default Value to support "data-wide" metadata
-					// It's only for spatial data and not Param data because for ParamData, GetNum on the keys should return the number of entries.
-					// If we have a default key, GetNum will return 1, while there is actually 0 values.
-					return MakeUnique<FPCGAttributeAccessorKeysEntries>(Metadata, /*bAddDefaultValueIfEmpty=*/!!Cast<UPCGSpatialData>(InData));
-				}
-			}
-			else
-			{
-				if (UPCGMetadata* Metadata = InData->MutableMetadata())
-				{
-					return MakeUnique<FPCGAttributeAccessorKeysEntries>(Metadata);
-				}
-			}
-		}
-
-		return TUniquePtr<KeysType>();
-	}
-
 	TUniquePtr<IPCGAttributeAccessor> CreateChainAccessor(TUniquePtr<IPCGAttributeAccessor> InAccessor, FName Name, bool& bOutSuccess)
 	{
 		if (!InAccessor.IsValid())
@@ -238,6 +29,8 @@ namespace PCGAttributeAccessorHelpers
 			return TUniquePtr<IPCGAttributeAccessor>();
 		}
 
+		bOutSuccess = true;
+
 		auto Chain = [&Accessor = InAccessor, Name, &bOutSuccess](auto Dummy) -> TUniquePtr<IPCGAttributeAccessor>
 		{
 			using AccessorType = decltype(Dummy);
@@ -267,6 +60,32 @@ namespace PCGAttributeAccessorHelpers
 
 		return PCGMetadataAttribute::CallbackWithRightType(InAccessor->GetUnderlyingType(), Chain);
 	}
+	
+	TUniquePtr<IPCGAttributeAccessor> CreateChainAccessor(TUniquePtr<IPCGAttributeAccessor> InAccessor, const FPCGAttributePropertySelector& InSelector, bool& bOutSuccess, bool bQuiet)
+	{
+		if (!InAccessor.IsValid())
+		{
+			bOutSuccess = false;
+			return {};
+		}
+		
+		for (const FString& ExtraName : InSelector.GetExtraNames())
+		{
+			bool bSuccess = true;
+			InAccessor = CreateChainAccessor(std::move(InAccessor), FName(ExtraName), bSuccess);
+			if (!bSuccess)
+			{
+				if (!bQuiet)
+				{
+					UE_LOG(LogPCG, Error, TEXT("[PCGAttributeAccessorHelpers::CreateChainAccessor] Extra selectors don't match existing properties."));
+				}
+				
+				return {};
+			}
+		}
+
+		return InAccessor;
+	}
 
 	// Empty signature to passthrough types to functors.
 	template <typename T>
@@ -618,27 +437,119 @@ TUniquePtr<const IPCGAttributeAccessor> PCGAttributeAccessorHelpers::CreateConst
 
 TUniquePtr<const IPCGAttributeAccessor> PCGAttributeAccessorHelpers::CreateConstAccessor(const UPCGData* InData, const FPCGAttributePropertySelector& InSelector, bool bQuiet)
 {
-	return CreateAccessorImpl<const IPCGAttributeAccessor, const UPCGData>(InData, InSelector, bQuiet);
+	TUniquePtr<IPCGAttributeAccessor> Accessor;
+	
+	if (InData)
+	{
+		// Temporary const_cast for the chain accessor.
+		Accessor.Reset(const_cast<IPCGAttributeAccessor*>(FPCGAttributeAccessorFactory::GetInstance().CreateSimpleConstAccessor(InData, InSelector, bQuiet).Release()));
+	}
+	else
+	{
+		// For backward compatibility, no data is point data (or extra property)
+		if (InSelector.GetSelection() == EPCGAttributePropertySelection::ExtraProperty)
+		{
+			Accessor = CreateExtraAccessor(InSelector.GetExtraProperty());
+		}
+		else
+		{
+			Accessor = UPCGPointData::CreateStaticAccessor(InSelector, bQuiet);
+		}
+	}
+
+	if (!Accessor.IsValid())
+	{
+		return {};
+	}
+
+	return CreateChainAccessor(std::move(Accessor), InSelector, bQuiet);
 }
 
 TUniquePtr<const IPCGAttributeAccessor> PCGAttributeAccessorHelpers::CreateConstAccessor(const FPCGMetadataAttributeBase* InAttribute, const UPCGMetadata* InMetadata, bool bQuiet)
 {
-	return CreateConstAttributeAccessorImpl(const_cast<FPCGMetadataAttributeBase*>(InAttribute), const_cast<UPCGMetadata*>(InMetadata));
+	if (!InMetadata || !InAttribute)
+	{
+		return {};
+	}
+	else
+	{
+		auto MakeConstAccessor = [InAttribute, InMetadata]<typename T>(T) -> TUniquePtr<const IPCGAttributeAccessor>
+		{
+			return MakeUnique<const FPCGAttributeAccessor<T>>(static_cast<const FPCGMetadataAttribute<T>*>(InAttribute), InMetadata);
+		};
+		
+		return PCGMetadataAttribute::CallbackWithRightType(InAttribute->GetTypeId(), MakeConstAccessor);
+	}
 }
 
 TUniquePtr<IPCGAttributeAccessor> PCGAttributeAccessorHelpers::CreateAccessor(UPCGData* InData, const FPCGAttributePropertySelector& InSelector, bool bQuiet)
 {
-	return CreateAccessorImpl<IPCGAttributeAccessor, UPCGData>(InData, InSelector, bQuiet);
+	TUniquePtr<IPCGAttributeAccessor> Accessor;
+	
+	if (InData)
+	{
+		// Temporary const_cast for the chain accessor.
+		Accessor = FPCGAttributeAccessorFactory::GetInstance().CreateSimpleAccessor(InData, InSelector, bQuiet);
+	}
+	else
+	{
+		// For backward compatibility, no data is point data (or extra property)
+		if (InSelector.GetSelection() == EPCGAttributePropertySelection::ExtraProperty)
+		{
+			Accessor = CreateExtraAccessor(InSelector.GetExtraProperty());
+		}
+		else
+		{
+			Accessor = UPCGPointData::CreateStaticAccessor(InSelector, bQuiet);
+		}
+	}
+
+	if (!Accessor.IsValid())
+	{
+		return {};
+	}
+
+	Accessor = CreateChainAccessor(std::move(Accessor), InSelector, bQuiet);
+	
+	if (Accessor && Accessor->IsReadOnly())
+	{
+		if (!bQuiet)
+		{
+			UE_LOG(LogPCG, Error, TEXT("[PCGAttributeAccessorHelpers::CreateAccessor] Attribute can not be written into, since it is read-only."));
+		}
+				
+		return {};
+	}
+
+	// Set the cached selector
+	if (InData)
+	{
+		InData->SetLastSelector(InSelector);
+	}
+
+	return Accessor;
 }
 
 TUniquePtr<IPCGAttributeAccessor> PCGAttributeAccessorHelpers::CreateAccessor(FPCGMetadataAttributeBase* InAttribute, UPCGMetadata* InMetadata, bool bQuiet)
 {
-	return CreateAttributeAccessorImpl(InAttribute, InMetadata);
+	if (!InMetadata || !InAttribute)
+	{
+		return {};
+	}
+	else
+	{
+		auto MakeAccessor = [InAttribute, InMetadata]<typename T>(T) -> TUniquePtr<IPCGAttributeAccessor>
+		{
+			return MakeUnique<FPCGAttributeAccessor<T>>(static_cast<FPCGMetadataAttribute<T>*>(InAttribute), InMetadata);
+		};
+		
+		return PCGMetadataAttribute::CallbackWithRightType(InAttribute->GetTypeId(), MakeAccessor);
+	}
 }
 
 TUniquePtr<IPCGAttributeAccessor> PCGAttributeAccessorHelpers::CreateAccessorWithAttributeCreation(UPCGData* InData, const FPCGAttributePropertySelector& InSelector, const IPCGAttributeAccessor* InMatchingAccessor, EPCGAttributeAccessorFlags InTypeMatching, bool bQuiet)
 {
-	TUniquePtr<IPCGAttributeAccessor> Result = CreateAccessor(InData,InSelector, bQuiet);
+	TUniquePtr<IPCGAttributeAccessor> Result = CreateAccessor(InData, InSelector, bQuiet);
 
 	if (!InSelector.IsBasicAttribute() || !InMatchingAccessor)
 	{
@@ -697,10 +608,10 @@ TUniquePtr<IPCGAttributeAccessor> PCGAttributeAccessorHelpers::CreateAccessorWit
 
 TUniquePtr<const IPCGAttributeAccessorKeys> PCGAttributeAccessorHelpers::CreateConstKeys(const UPCGData* InData, const FPCGAttributePropertySelector& InSelector)
 {
-	return CreateKeysImpl<const IPCGAttributeAccessorKeys, const UPCGData>(InData, InSelector);
+	return FPCGAttributeAccessorFactory::GetInstance().CreateSimpleConstKeys(InData, InSelector, /*bQuiet=*/ false);
 }
 
 TUniquePtr<IPCGAttributeAccessorKeys> PCGAttributeAccessorHelpers::CreateKeys(UPCGData* InData, const FPCGAttributePropertySelector& InSelector)
 {
-	return CreateKeysImpl<IPCGAttributeAccessorKeys, UPCGData>(InData, InSelector);
+	return FPCGAttributeAccessorFactory::GetInstance().CreateSimpleKeys(InData, InSelector, /*bQuiet=*/ false);
 }
diff --git a/Engine/Plugins/PCG/Source/PCG/Private/PCGModule.cpp b/Engine/Plugins/PCG/Source/PCG/Private/PCGModule.cpp
--- a/Engine/Plugins/PCG/Source/PCG/Private/PCGModule.cpp
+++ b/Engine/Plugins/PCG/Source/PCG/Private/PCGModule.cpp
@@ -4,6 +4,7 @@
 
 #include "PCGContext.h"
 #include "PCGElement.h"
+#include "Data/PCGPointData.h"
 
 #include "Interfaces/IPluginManager.h"
 #include "Modules/ModuleManager.h"
@@ -42,10 +43,18 @@ void FPCGModule::StartupModule()
 
 	const FString PluginShaderDir = FPaths::Combine(IPluginManager::Get().FindPlugin(TEXT("PCG"))->GetBaseDir(), TEXT("Shaders"));
 	AddShaderSourceDirectoryMapping(TEXT("/Plugin/PCG"), PluginShaderDir);
+
+	// Registering accessor methods
+	AttributeAccessorFactory.RegisterDefaultMethods();
+	AttributeAccessorFactory.RegisterMethods<UPCGPointData>(UPCGPointData::GetPointAccessorMethods());
 }
 
 void FPCGModule::ShutdownModule()
 {
+	// Unregistering accessor methods
+	AttributeAccessorFactory.UnregisterDefaultMethods();
+	AttributeAccessorFactory.UnregisterMethods<UPCGPointData>();
+	
 	DeregisterNativeElementDeterminismTests();
 
 	PCGDeterminismTests::FNativeTestRegistry::Destroy();
diff --git a/Engine/Plugins/PCG/Source/PCG/Public/Data/PCGPointData.h b/Engine/Plugins/PCG/Source/PCG/Public/Data/PCGPointData.h
--- a/Engine/Plugins/PCG/Source/PCG/Public/Data/PCGPointData.h
+++ b/Engine/Plugins/PCG/Source/PCG/Public/Data/PCGPointData.h
@@ -81,6 +81,12 @@ public:
 	virtual void AddToCrc(FArchiveCrc32& Ar, bool bFullDataCrc) const override;
 	// ~End UPCGData interface
 
+	// Static helper to create an accessor on a data that doesn't yet exist, as accessors for point data don't rely on existing data.
+	static TUniquePtr<IPCGAttributeAccessor> CreateStaticAccessor(const FPCGAttributePropertySelector& InSelector, bool bQuiet = false);
+
+	// Get the functions to the accessor factory
+	static FPCGAttributeAccessorMethods GetPointAccessorMethods();
+
 	// ~Begin UPCGSpatialData interface
 	virtual int GetDimension() const override { return 0; }
 	virtual FBox GetBounds() const override;
diff --git a/Engine/Plugins/PCG/Source/PCG/Public/Metadata/Accessors/PCGAttributeAccessorFactory.h b/Engine/Plugins/PCG/Source/PCG/Public/Metadata/Accessors/PCGAttributeAccessorFactory.h
new file mode 100644
--- /dev/null
+++ b/Engine/Plugins/PCG/Source/PCG/Public/Metadata/Accessors/PCGAttributeAccessorFactory.h
@@ -0,0 +1,84 @@
+// Copyright Epic Games, Inc. All Rights Reserved.
+
+#pragma once
+
+#include "Metadata/PCGAttributePropertySelector.h"
+
+#include "Templates/Function.h"
+#include "Templates/SubclassOf.h"
+#include "Templates/UniquePtr.h"
+
+struct FPCGAttributePropertySelector;
+class FPCGMetadataAttributeBase;
+class FPCGModule;
+class IPCGAttributeAccessor;
+class IPCGAttributeAccessorKeys;
+class UPCGData;
+class UPCGMetadata;
+
+struct FPCGAttributeSelectorMenuEntry
+{
+	/** Display name in the menu. Can use "|" to indicate a hierarchy of menus */
+	FText DisplayName;
+
+	/** Tooltip to display for the entry. */
+	FText Tooltip;
+
+	/** Selector to copy if this entry is selected. */
+	FPCGAttributePropertySelector Selector;
+};
+
+struct FPCGAttributeAccessorMethods
+{
+	TFunction<TUniquePtr<IPCGAttributeAccessor>(UPCGData*, const FPCGAttributePropertySelector&, const bool)> CreateAccessorFunc;
+	TFunction<TUniquePtr<const IPCGAttributeAccessor>(const UPCGData*, const FPCGAttributePropertySelector&, const bool)> CreateConstAccessorFunc;
+	TFunction<TUniquePtr<IPCGAttributeAccessorKeys>(UPCGData*, const FPCGAttributePropertySelector&, const bool)> CreateAccessorKeysFunc;
+	TFunction<TUniquePtr<const IPCGAttributeAccessorKeys>(const UPCGData*, const FPCGAttributePropertySelector&, const bool)> CreateConstAccessorKeysFunc;
+
+	// Special option to gather all the possible options for a given data type in the Selector context menu.
+	TArray<FPCGAttributeSelectorMenuEntry> AttributeSelectorOptions;
+};
+
+class FPCGAttributeAccessorFactory
+{
+	friend FPCGModule;
+
+public:	
+	static PCG_API FPCGAttributeAccessorFactory& GetMutableInstance();
+	static PCG_API const FPCGAttributeAccessorFactory& GetInstance();
+
+	/** Create a simple accessor based on the data passed as input. No chain of extraction. Internal function deliberately not exposed. Use PCGAttributeAccessorHelpers::CreateAccessor. */
+	TUniquePtr<IPCGAttributeAccessor> CreateSimpleAccessor(UPCGData* InData, const FPCGAttributePropertySelector& InSelector, const bool bQuiet) const;
+	
+	/** Create a simple const accessor based on the data passed as input. No chain of extraction. Internal function deliberately not exposed. Use PCGAttributeAccessorHelpers::CreateConstAccessor. */
+	TUniquePtr<const IPCGAttributeAccessor> CreateSimpleConstAccessor(const UPCGData* InData, const FPCGAttributePropertySelector& InSelector, const bool bQuiet) const;
+	
+	/** Create a simple key based on the data passed as input. Internal function deliberately not exposed. Use PCGAttributeAccessorHelpers::CreateKeys. */
+	TUniquePtr<IPCGAttributeAccessorKeys> CreateSimpleKeys(UPCGData* InData, const FPCGAttributePropertySelector& InSelector, const bool bQuiet) const;
+	
+	/** Create a simple key based on the data passed as input. Internal function deliberately not exposed. Use PCGAttributeAccessorHelpers::CreateConstKeys. */
+	TUniquePtr<const IPCGAttributeAccessorKeys> CreateSimpleConstKeys(const UPCGData* InData, const FPCGAttributePropertySelector& InSelector, const bool bQuiet) const;
+
+	/** Call a Callback on all the FPCGAttributeSelectorMenuEntry currently registered. */
+	PCG_API void ForEachSelectorMenuEntry(TFunctionRef<void(const FPCGAttributeSelectorMenuEntry&)> Callback) const;
+	
+	template <typename T, typename std::enable_if_t<std::is_base_of_v<UPCGData, T>, bool> = false>
+	void RegisterMethods(FPCGAttributeAccessorMethods&& InAccessorMethods) { RegisterMethods(T::StaticClass(), std::move(InAccessorMethods)); }
+	
+	PCG_API void RegisterMethods(TSubclassOf<UPCGData> PCGDataClass, FPCGAttributeAccessorMethods&& InAccessorMethods);
+
+	template <typename T, typename std::enable_if_t<std::is_base_of_v<UPCGData, T>, bool> = false>
+	void UnregisterMethods() { UnregisterMethods(T::StaticClass()); }
+	
+	PCG_API void UnregisterMethods(TSubclassOf<UPCGData> PCGDataClass);
+
+private:
+	/** To be called by the FPCGModule on register. Will register the default UPCGData accessor (Metadata), keys and the spatial const keys (special case). */
+	void RegisterDefaultMethods();
+	void UnregisterDefaultMethods();
+	
+	template <typename PCGData, typename Func>
+	decltype(auto) CallOnMethod(PCGData* InData, Func Callback) const;
+	
+	TMap<const TSubclassOf<UPCGData>, FPCGAttributeAccessorMethods> AccessorMethods;
+};
diff --git a/Engine/Plugins/PCG/Source/PCG/Public/Metadata/Accessors/PCGAttributeAccessorHelpers.h b/Engine/Plugins/PCG/Source/PCG/Public/Metadata/Accessors/PCGAttributeAccessorHelpers.h
--- a/Engine/Plugins/PCG/Source/PCG/Public/Metadata/Accessors/PCGAttributeAccessorHelpers.h
+++ b/Engine/Plugins/PCG/Source/PCG/Public/Metadata/Accessors/PCGAttributeAccessorHelpers.h
@@ -39,6 +39,7 @@ namespace PCGAttributeAccessorHelpers
 	PCG_API TUniquePtr<IPCGAttributeAccessor> CreateExtraAccessor(EPCGExtraProperties InExtraProperties);
 
 	PCG_API TUniquePtr<IPCGAttributeAccessor> CreateChainAccessor(TUniquePtr<IPCGAttributeAccessor> InAccessor, FName Name, bool& bOutSuccess);
+	PCG_API TUniquePtr<IPCGAttributeAccessor> CreateChainAccessor(TUniquePtr<IPCGAttributeAccessor> InAccessor, const FPCGAttributePropertySelector& InSelector, bool& bOutSuccess, bool bQuiet = false);
 
 	struct AccessorParamResult
 	{
diff --git a/Engine/Plugins/PCG/Source/PCG/Public/PCGModule.h b/Engine/Plugins/PCG/Source/PCG/Public/PCGModule.h
--- a/Engine/Plugins/PCG/Source/PCG/Public/PCGModule.h
+++ b/Engine/Plugins/PCG/Source/PCG/Public/PCGModule.h
@@ -4,6 +4,7 @@
 
 #include "PCGDataVisualizationRegistry.h"
 #include "Data/PCGGetDataFunctionRegistry.h"
+#include "Metadata/Accessors/PCGAttributeAccessorFactory.h"
 #include "Utils/PCGLogErrors.h"
 
 #include "Modules/ModuleInterface.h"
@@ -38,11 +39,15 @@ public:
 	PCG_API static FPCGModule& GetPCGModuleChecked();
 	static const FPCGGetDataFunctionRegistry& ConstGetDataFunctionRegistry() { return GetPCGModuleChecked().GetDataFunctionRegistry; }
 	static FPCGGetDataFunctionRegistry& MutableGetDataFunctionRegistry() { return GetPCGModuleChecked().GetDataFunctionRegistry; }
+
+	static const FPCGAttributeAccessorFactory& GetConstAttributeAccessorFactory() { return GetPCGModuleChecked().AttributeAccessorFactory; }
+	static FPCGAttributeAccessorFactory& GetMutableAttributeAccessorFactory() { return GetPCGModuleChecked().AttributeAccessorFactory; }
 	
 	PCG_API static bool IsPCGModuleLoaded();
 	
 private:
 	FPCGGetDataFunctionRegistry GetDataFunctionRegistry;
+	FPCGAttributeAccessorFactory AttributeAccessorFactory;
 
 #if WITH_EDITOR
 private:
