[PCG] Make the PostEditChange property more robust for Graph Parameter changes in PCG Graph

* A change in the editor changed the "MemberProperty" and is not the property name that was expected to detect changes anymore.
* Changed to use the property chain instead, and scan the chain to find the exact property that correspond to the graph parameter that changed. It will be way more robust like that.

#rnx

diff --git a/Engine/Plugins/PCG/Source/PCG/Private/PCGGraph.cpp b/Engine/Plugins/PCG/Source/PCG/Private/PCGGraph.cpp
--- a/Engine/Plugins/PCG/Source/PCG/Private/PCGGraph.cpp
+++ b/Engine/Plugins/PCG/Source/PCG/Private/PCGGraph.cpp
@@ -237,6 +237,41 @@ bool UPCGGraphInterface::IsEquivalent(const UPCGGraphInterface* Other) const
 	return true;
 }
 
+#if WITH_EDITOR
+bool UPCGGraphInterface::VerifyAndUpdateIfGraphParameterValueChanged(const FPropertyChangedChainEvent& PropertyChangedEvent)
+{
+	const FInstancedPropertyBag* UserParameters = GetUserParametersStruct();
+	const UPropertyBag* PropertyBag = UserParameters ? UserParameters->GetPropertyBagStruct() : nullptr;
+
+	if (!PropertyBag)
+	{
+		return false;
+	}
+
+	const TDoubleLinkedList<FProperty*>::TDoubleLinkedListNode* PropertyNode = PropertyChangedEvent.PropertyChain.GetHead();
+	
+	while (PropertyNode)
+	{
+		if (PropertyNode->GetValue() && PropertyNode->GetValue()->GetOwnerStruct() == PropertyBag)
+		{
+			break;
+		}
+
+		PropertyNode = PropertyNode->GetNextNode();
+	}
+
+	if (PropertyNode && PropertyNode->GetValue())
+	{
+		OnGraphParametersChanged(EPCGGraphParameterEvent::ValueModifiedLocally, PropertyNode->GetValue()->GetFName());
+		return true;
+	}
+	else
+	{
+		return false;
+	}
+}
+#endif // WITH_EDITOR
+
 EPCGChangeType UPCGGraphInterface::GetChangeTypeForGraphParameterChange(EPCGGraphParameterEvent InChangeType, FName InChangedPropertyName)
 {
 	// If the parameter had its order changed in the struct, was just added or was removed but was not used in the graph, it is not a change that requires a refresh, so we go with Cosmetic change type.
@@ -1541,9 +1576,9 @@ void UPCGGraph::PreEditChange(FProperty* InProperty)
 	}
 }
 
-void UPCGGraph::PostEditChangeProperty(FPropertyChangedEvent& PropertyChangedEvent)
+void UPCGGraph::PostEditChangeChainProperty(FPropertyChangedChainEvent& PropertyChangedEvent)
 {
-	Super::PostEditChangeProperty(PropertyChangedEvent);
+	Super::PostEditChangeChainProperty(PropertyChangedEvent);
 
 	if (!PropertyChangedEvent.Property)
 	{
@@ -1551,7 +1586,6 @@ void UPCGGraph::PostEditChangeProperty(FPropertyChangedEvent& PropertyChangedEve
 	}
 
 	const FName PropertyName = PropertyChangedEvent.GetPropertyName();
-	const FName MemberPropertyName = PropertyChangedEvent.GetMemberPropertyName();
 
 	if (PropertyName == GET_MEMBER_NAME_CHECKED(UPCGGraph, bLandscapeUsesMetadata))
 	{
@@ -1608,10 +1642,6 @@ void UPCGGraph::PostEditChangeProperty(FPropertyChangedEvent& PropertyChangedEve
 			OnGraphParametersChanged(ChangeType, ChangedPropertyName);
 		}
 	}
-	else if (PropertyChangedEvent.MemberProperty && PropertyChangedEvent.MemberProperty->GetOwnerStruct() == UserParameters.GetPropertyBagStruct())
-	{
-		OnGraphParametersChanged(EPCGGraphParameterEvent::ValueModifiedLocally, MemberPropertyName);
-	}
 	else if (PropertyName == GET_MEMBER_NAME_CHECKED(UPCGGraph, HiGenGridSize)
 		|| PropertyName == GET_MEMBER_NAME_CHECKED(UPCGGraph, bUseHierarchicalGeneration)
 		|| PropertyName == GET_MEMBER_NAME_CHECKED(UPCGGraph, bUse2DGrid)
@@ -1627,6 +1657,10 @@ void UPCGGraph::PostEditChangeProperty(FPropertyChangedEvent& PropertyChangedEve
 	{
 		NotifyGraphChanged(EPCGChangeType::Cosmetic);
 	}
+	else if (VerifyAndUpdateIfGraphParameterValueChanged(PropertyChangedEvent))
+	{
+		// Handled in the function
+	}
 
 	PreviousPropertyBag = nullptr;
 }
@@ -1917,9 +1951,9 @@ void UPCGGraphInstance::PreEditChange(FProperty* InProperty)
 	}
 }
 
-void UPCGGraphInstance::PostEditChangeProperty(FPropertyChangedEvent& PropertyChangedEvent)
+void UPCGGraphInstance::PostEditChangeChainProperty(FPropertyChangedChainEvent& PropertyChangedEvent)
 {
-	Super::PostEditChangeProperty(PropertyChangedEvent);
+	Super::PostEditChangeChainProperty(PropertyChangedEvent);
 
 	if (!PropertyChangedEvent.Property)
 	{
@@ -1953,9 +1987,9 @@ void UPCGGraphInstance::PostEditChangeProperty(FPropertyChangedEvent& PropertyCh
 	{
 		// Changes to this property will be caught by the update on override status, and "OnGraphParametersChanged" will be called accordingly.
 	}
-	else if (PropertyChangedEvent.MemberProperty && PropertyChangedEvent.MemberProperty->GetOwnerStruct() == ParametersOverrides.Parameters.GetPropertyBagStruct())
+	else if (VerifyAndUpdateIfGraphParameterValueChanged(PropertyChangedEvent))
 	{
-		OnGraphParametersChanged(this, EPCGGraphParameterEvent::ValueModifiedLocally, PropertyChangedEvent.GetMemberPropertyName());
+		// Handled in the function
 	}
 	else
 	{
diff --git a/Engine/Plugins/PCG/Source/PCG/Public/PCGGraph.h b/Engine/Plugins/PCG/Source/PCG/Public/PCGGraph.h
--- a/Engine/Plugins/PCG/Source/PCG/Public/PCGGraph.h
+++ b/Engine/Plugins/PCG/Source/PCG/Public/PCGGraph.h
@@ -128,6 +128,8 @@ protected:
 
 	UFUNCTION()
 	virtual bool AreOverridesEnabled() const { return IsExportToLibraryEnabled() && bExposeToLibrary; }
+
+	bool VerifyAndUpdateIfGraphParameterValueChanged(const FPropertyChangedChainEvent& PropertyChangedEvent);
 #endif // WITH_EDITOR
 
 public:
@@ -246,7 +248,7 @@ public:
 
 #if WITH_EDITOR
 	virtual void PreEditChange(FProperty* InProperty) override;
-	virtual void PostEditChangeProperty(FPropertyChangedEvent& PropertyChangedEvent) override;
+	virtual void PostEditChangeChainProperty(FPropertyChangedChainEvent& PropertyChangedEvent);
 #endif
 	/** ~End UObject interface */
 
@@ -548,7 +550,7 @@ public:
 
 #if WITH_EDITOR
 	virtual void PreEditChange(FProperty* InProperty) override;
-	virtual void PostEditChangeProperty(FPropertyChangedEvent& PropertyChangedEvent) override;
+	virtual void PostEditChangeChainProperty(FPropertyChangedChainEvent& PropertyChangedEvent);
 	virtual bool CanEditChange(const FProperty* InProperty) const override;
 	virtual void PreEditUndo() override;
 	virtual void PostEditUndo() override;
