Find window update.
- Added find in debug object (finds in dynamic graphs)
- Added colors to nodes
- Changed icon for subgraph nodes
- Improved navigation so only double-click or enter opens up a new graph
- Added different visualization modes (minimum tree, full tree, flalt list)
- Added option to show pin results or not.

diff --git a/Engine/Plugins/PCG/Source/PCG/Private/PCGSubsystem.cpp b/Engine/Plugins/PCG/Source/PCG/Private/PCGSubsystem.cpp
--- a/Engine/Plugins/PCG/Source/PCG/Private/PCGSubsystem.cpp
+++ b/Engine/Plugins/PCG/Source/PCG/Private/PCGSubsystem.cpp
@@ -1778,6 +1778,22 @@ TArray<FPCGStack> UPCGSubsystem::GetExecutedStacks(const UPCGComponent* InCompon
 	return MatchingStacks;
 }
 
+TArray<FPCGStack> UPCGSubsystem::GetExecutedStacks(const FPCGStack& BeginningWithStack)
+{
+	FReadScopeLock Lock(ExecutedStacksLock);
+	TArray<FPCGStack> MatchingStacks;
+
+	for (const FPCGStack& Stack : ExecutedStacks)
+	{
+		if (Stack.BeginsWith(BeginningWithStack))
+		{
+			MatchingStacks.Add(Stack);
+		}
+	}
+
+	return MatchingStacks;
+}
+
 void UPCGSubsystem::ClearExecutedStacks(FPCGStack BeginningWithStack)
 {
 	FWriteScopeLock Lock(ExecutedStacksLock);
diff --git a/Engine/Plugins/PCG/Source/PCG/Public/PCGSubsystem.h b/Engine/Plugins/PCG/Source/PCG/Public/PCGSubsystem.h
--- a/Engine/Plugins/PCG/Source/PCG/Public/PCGSubsystem.h
+++ b/Engine/Plugins/PCG/Source/PCG/Public/PCGSubsystem.h
@@ -336,6 +336,7 @@ public:
 
 	/** Get a list of stacks that were executed during the last execution. */
 	TArray<FPCGStack> GetExecutedStacks(const UPCGComponent* InComponent, const UPCGGraph* InSubgraph);
+	TArray<FPCGStack> GetExecutedStacks(const FPCGStack& BeginningWithStack);
 	void ClearExecutedStacks(FPCGStack BeginningWithStack);
 	void ClearExecutedStacks(const UPCGGraph* InContainingGraph);
 
diff --git a/Engine/Plugins/PCG/Source/PCGEditor/Private/SPCGEditorGraphFind.cpp b/Engine/Plugins/PCG/Source/PCGEditor/Private/SPCGEditorGraphFind.cpp
--- a/Engine/Plugins/PCG/Source/PCGEditor/Private/SPCGEditorGraphFind.cpp
+++ b/Engine/Plugins/PCG/Source/PCGEditor/Private/SPCGEditorGraphFind.cpp
@@ -5,13 +5,17 @@
 #include "PCGEditor.h"
 #include "PCGEditorGraph.h"
 #include "PCGEditorGraphNodeBase.h"
+#include "PCGEditorStyle.h"
+#include "PCGGraph.h"
 #include "PCGNode.h"
 #include "PCGSubgraph.h"
+#include "PCGSubsystem.h"
 
 #include "EdGraph/EdGraphSchema.h"
 #include "Framework/Application/SlateApplication.h"
 #include "Framework/Views/TableViewMetadata.h"
 #include "Layout/WidgetPath.h"
+#include "Widgets/Input/SComboButton.h"
 #include "Widgets/Input/SSearchBox.h"
 #include "Widgets/Views/STreeView.h"
 
@@ -21,6 +25,7 @@ FPCGEditorGraphFindResult::FPCGEditorGraphFindResult(const FString& InValue)
 : Parent(nullptr)
 , Value(InValue)
 , GraphNode(nullptr)
+, RootGraphNode(nullptr)
 , Pin()
 {
 }
@@ -30,35 +35,54 @@ FPCGEditorGraphFindResult::FPCGEditorGraphFindResult(const FText& InValue)
 {
 }
 
-FPCGEditorGraphFindResult::FPCGEditorGraphFindResult(const FString& InValue, TSharedPtr<FPCGEditorGraphFindResult>& InParent, UEdGraphNode* InNode)
+FPCGEditorGraphFindResult::FPCGEditorGraphFindResult(const FString& InValue, const TSharedPtr<FPCGEditorGraphFindResult>& InParent, UEdGraphNode* InNode)
 : Parent(InParent)
 , Value(InValue)
 , GraphNode(InNode)
 , Pin()
 {
+	RootGraphNode = (Parent.IsValid() && Parent.Pin()->RootGraphNode.IsValid()) ? Parent.Pin()->RootGraphNode : InNode;
 }
 
-FPCGEditorGraphFindResult::FPCGEditorGraphFindResult(const FString& InValue, TSharedPtr<FPCGEditorGraphFindResult>& InParent, UEdGraphPin* InPin)
+FPCGEditorGraphFindResult::FPCGEditorGraphFindResult(const FString& InValue, const TSharedPtr<FPCGEditorGraphFindResult>& InParent, UEdGraphPin* InPin)
 : Parent(InParent)
 , Value(InValue)
 , GraphNode(nullptr)
 , Pin(InPin)
 {
+	RootGraphNode = (Parent.IsValid() && Parent.Pin()->RootGraphNode.IsValid()) ? Parent.Pin()->RootGraphNode : nullptr;
 }
 
 FReply FPCGEditorGraphFindResult::OnClick(TWeakPtr<FPCGEditor> InPCGEditorPtr)
 {
-	TSharedPtr<FPCGEditor> Editor;
+	TSharedPtr<FPCGEditor> Editor = InPCGEditorPtr.Pin();
 
-	// If result points to another graph, make that visible
-	if (ParentGraph)
+	if (!Editor)
 	{
-		GEditor->GetEditorSubsystem<UAssetEditorSubsystem>()->OpenEditorForAsset(ParentGraph->GetPCGGraph());
-		Editor = ParentGraph->GetEditor().Pin();
+		return FReply::Handled();
+	}
+
+	UEdGraphPin* ResolvedPin = Pin.Get();
+	if (ResolvedPin && ResolvedPin->GetOwningNode()->GetGraph() == Editor->GetPCGEditorGraph())
+	{
+		Editor->JumpToNode(ResolvedPin->GetOwningNode());
 	}
 	else
 	{
-		Editor = InPCGEditorPtr.Pin();
+		Editor->JumpToNode(RootGraphNode.Get());
+	}
+
+	return FReply::Handled();
+}
+
+FReply FPCGEditorGraphFindResult::OnDoubleClick(TWeakPtr<FPCGEditor> InPCGEditorPtr)
+{
+	TSharedPtr<FPCGEditor> Editor = InPCGEditorPtr.Pin();
+
+	if (ParentGraph)
+	{
+		GEditor->GetEditorSubsystem<UAssetEditorSubsystem>()->OpenEditorForAsset(ParentGraph->GetPCGGraph());
+		Editor = ParentGraph->GetEditor().Pin();
 	}
 
 	if (Editor)
@@ -129,6 +153,8 @@ TSharedRef<SWidget> FPCGEditorGraphFindResult::CreateIcon() const
 	FSlateColor IconColor = FSlateColor::UseForeground();
 	const FSlateBrush* Brush = nullptr;
 
+	// TODO: consider the pin connection state (connected or not + single + multi)
+	// TODO: consider node state (enabled or not)
 	if (UEdGraphPin* ResolvedPin = Pin.Get())
 	{
 		// TODO get pin icon from nodebase?
@@ -140,6 +166,16 @@ TSharedRef<SWidget> FPCGEditorGraphFindResult::CreateIcon() const
 	{
 		// TODO get icon and tint from nodebase?
 		Brush = FAppStyle::GetBrush(TEXT("GraphEditor.NodeGlyph"));
+
+		if (const UPCGEditorGraphNodeBase* PCGNode = Cast<UPCGEditorGraphNodeBase>(GraphNode))
+		{
+			if (Cast<UPCGBaseSubgraphNode>(PCGNode->GetPCGNode()))
+			{
+				Brush = FPCGEditorStyle::Get().GetBrush(TEXT("ClassIcon.PCGGraphInterface"));
+			}
+
+			IconColor = FSlateColor(PCGNode->GetNodeTitleColor());
+		}
 	}
 
 	return SNew(SImage)
@@ -148,20 +184,55 @@ TSharedRef<SWidget> FPCGEditorGraphFindResult::CreateIcon() const
 		.ToolTipText(GetCategory());
 }
 
+SPCGEditorGraphFind::~SPCGEditorGraphFind()
+{
+	if (PCGEditorPtr.IsValid())
+	{
+		PCGEditorPtr.Pin()->OnInspectedStackChangedDelegate.RemoveAll(this);
+	}
+}
+
 void SPCGEditorGraphFind::Construct(const FArguments& InArgs, TSharedPtr<FPCGEditor> InPCGEditor)
 {
 	PCGEditorPtr = InPCGEditor;
 
+	const TSharedPtr<FPCGEditor> PCGEditor = PCGEditorPtr.Pin();
+
+	if (PCGEditor)
+	{
+		PCGEditor->OnInspectedStackChangedDelegate.AddSP(this, &SPCGEditorGraphFind::OnInspectedStackChanged);
+	}
+
 	this->ChildSlot
 	[
 		SNew(SVerticalBox)
 		+SVerticalBox::Slot()
 		.AutoHeight()
 		[
-			SAssignNew(SearchTextField, SSearchBox)
-			.HintText(LOCTEXT("PCGGraphSearchHint", "Enter text to find nodes..."))
-			.OnTextChanged(this, &SPCGEditorGraphFind::OnSearchTextChanged)
-			.OnTextCommitted(this, &SPCGEditorGraphFind::OnSearchTextCommitted)
+			SNew(SHorizontalBox)
+			+SHorizontalBox::Slot()
+			.AutoWidth()
+			[
+				SNew(SComboButton)
+				.ForegroundColor(FSlateColor::UseStyle())
+				.HasDownArrow(false)
+				.OnGetMenuContent(this, &SPCGEditorGraphFind::OnFindFilterMenu)
+				.ContentPadding(1)
+				.ButtonContent()
+				[
+					SNew(SImage)
+					.Image(FAppStyle::Get().GetBrush("Icons.Filter"))
+					.ColorAndOpacity(FSlateColor::UseForeground())
+				]
+			]
+			+SHorizontalBox::Slot()
+			.AutoWidth()
+			[
+				SAssignNew(SearchTextField, SSearchBox)
+				.HintText(LOCTEXT("PCGGraphSearchHint", "Enter text to find nodes..."))
+				.OnTextChanged(this, &SPCGEditorGraphFind::OnSearchTextChanged)
+				.OnTextCommitted(this, &SPCGEditorGraphFind::OnSearchTextCommitted)
+			]
 		]
 		+ SVerticalBox::Slot()
 		.FillHeight(1.0f)
@@ -175,6 +246,8 @@ void SPCGEditorGraphFind::Construct(const FArguments& InArgs, TSharedPtr<FPCGEdi
 				.OnGenerateRow(this, &SPCGEditorGraphFind::OnGenerateRow)
 				.OnGetChildren(this, &SPCGEditorGraphFind::OnGetChildren)
 				.OnSelectionChanged(this, &SPCGEditorGraphFind::OnTreeSelectionChanged)
+				.OnMouseButtonDoubleClick(this, &SPCGEditorGraphFind::OnTreeDoubleClick)
+				.OnKeyDownHandler(this, &SPCGEditorGraphFind::OnTreeViewKeyDown)
 				.SelectionMode(ESelectionMode::Single)
 			]
 		]
@@ -197,9 +270,21 @@ void SPCGEditorGraphFind::OnSearchTextChanged(const FText& InText)
 	InitiateSearch();
 }
 
-void SPCGEditorGraphFind::OnSearchTextCommitted(const FText& /*InText*/, ETextCommit::Type /*InCommitType*/)
+void SPCGEditorGraphFind::OnSearchTextCommitted(const FText& /*InText*/, ETextCommit::Type InCommitType)
 {
-	InitiateSearch();
+	// Since we already initiate a search when the text changes, there's no real need to do anything on commit
+	if (InCommitType != ETextCommit::Type::OnUserMovedFocus)
+	{
+		InitiateSearch();
+	}
+}
+
+void SPCGEditorGraphFind::OnInspectedStackChanged(const FPCGStack& InPCGStack)
+{
+	if (!SearchValue.IsEmpty())
+	{
+		InitiateSearch();
+	}
 }
 
 void SPCGEditorGraphFind::OnGetChildren(FPCGEditorGraphFindResultPtr InItem, TArray<FPCGEditorGraphFindResultPtr>& OutChildren)
@@ -215,39 +300,139 @@ void SPCGEditorGraphFind::OnTreeSelectionChanged(FPCGEditorGraphFindResultPtr It
 	}
 }
 
+void SPCGEditorGraphFind::OnTreeDoubleClick(FPCGEditorGraphFindResultPtr Item)
+{
+	if (Item.IsValid())
+	{
+		Item->OnDoubleClick(PCGEditorPtr);
+	}
+}
+
+FReply SPCGEditorGraphFind::OnTreeViewKeyDown(const FGeometry& /*InGeometry*/, const FKeyEvent& InKeyEvent) const
+{
+	if (InKeyEvent.GetKey() == EKeys::Enter)
+	{
+		TArray<FPCGEditorGraphFindResultPtr> SelectedItems = TreeView->GetSelectedItems();
+		if (SelectedItems.Num() == 1 && SelectedItems[0].IsValid())
+		{
+			return SelectedItems[0]->OnDoubleClick(PCGEditorPtr);
+		}
+	}
+
+	return FReply::Unhandled();
+}
+
+void SPCGEditorGraphFind::SetFindMode(EPCGGraphFindMode InFindMode)
+{
+	if (InFindMode != FindMode)
+	{
+		FindMode = InFindMode;
+		InitiateSearch();
+	}
+}
+
+bool SPCGEditorGraphFind::IsCurrentFindMode(EPCGGraphFindMode InFindMode) const
+{
+	return FindMode == InFindMode;
+}
+
+TSharedRef<SWidget> SPCGEditorGraphFind::OnFindFilterMenu()
+{
+	FMenuBuilder MenuBuilder(/*bInShouldCloseWindowAfterMenuSelection=*/true, nullptr);
+
+	MenuBuilder.AddMenuEntry(
+		LOCTEXT("MinimalExpansion", "Show minimum tree"),
+		LOCTEXT("MinimalExpansionTooltip", "Shows minimum subset of visited tree to perform search."),
+		FSlateIcon(),
+		FUIAction(
+			FExecuteAction::CreateSP(this, &SPCGEditorGraphFind::SetFindMode, EPCGGraphFindMode::ShowMinimumTree),
+			FCanExecuteAction(),
+			FIsActionChecked::CreateSP(this, &SPCGEditorGraphFind::IsCurrentFindMode, EPCGGraphFindMode::ShowMinimumTree)
+		),
+		NAME_None,
+		EUserInterfaceActionType::RadioButton);
+
+	MenuBuilder.AddMenuEntry(
+		LOCTEXT("FullExpansion", "Show full tree"),
+		LOCTEXT("FullExpansionTooltip", "Shows the search results from the fully expanded tree, e.g. will have every occurrence."),
+		FSlateIcon(),
+		FUIAction(
+			FExecuteAction::CreateSP(this, &SPCGEditorGraphFind::SetFindMode, EPCGGraphFindMode::ShowFullTree),
+			FCanExecuteAction(),
+			FIsActionChecked::CreateSP(this, &SPCGEditorGraphFind::IsCurrentFindMode, EPCGGraphFindMode::ShowFullTree)
+		),
+		NAME_None,
+		EUserInterfaceActionType::RadioButton);
+
+	MenuBuilder.AddMenuEntry(
+		LOCTEXT("FlatList", "Show flat graph list"),
+		LOCTEXT("FlatListTooltip", "Shows the occurrences in all visited graphs but no hierarchy."),
+		FSlateIcon(),
+		FUIAction(
+			FExecuteAction::CreateSP(this, &SPCGEditorGraphFind::SetFindMode, EPCGGraphFindMode::ShowFlatList),
+			FCanExecuteAction(),
+			FIsActionChecked::CreateSP(this, &SPCGEditorGraphFind::IsCurrentFindMode, EPCGGraphFindMode::ShowFlatList)
+		),
+		NAME_None,
+		EUserInterfaceActionType::RadioButton);
+
+	MenuBuilder.AddSeparator();
+
+	MenuBuilder.AddMenuEntry(
+		LOCTEXT("SearchPins", "Include pin names"),
+		LOCTEXT("SearchPinsTooltip", "TODO"),
+		FSlateIcon(),
+		FUIAction(
+			FExecuteAction::CreateLambda([this]()
+			{
+				bShowPinResults = !bShowPinResults;
+				InitiateSearch();
+			}),
+			FCanExecuteAction(),
+			FIsActionChecked::CreateLambda([this]()
+			{
+				return bShowPinResults;
+			})
+		),
+		NAME_None,
+		EUserInterfaceActionType::ToggleButton);
+
+	return MenuBuilder.MakeWidget();
+}
+
 TSharedRef<ITableRow> SPCGEditorGraphFind::OnGenerateRow(FPCGEditorGraphFindResultPtr InItem, const TSharedRef<STableViewBase>& InOwnerTable)
 {
 	return SNew(STableRow<TSharedPtr<FPCGEditorGraphFindResult>>, InOwnerTable)
+	[
+		SNew(SHorizontalBox)
+		+ SHorizontalBox::Slot()
+		.VAlign(VAlign_Center)
+		.AutoWidth()
 		[
-			SNew(SHorizontalBox)
-			+ SHorizontalBox::Slot()
-			.VAlign(VAlign_Center)
-			.AutoWidth()
-			[
-				InItem->CreateIcon()
-			]
-			+SHorizontalBox::Slot()
-			.VAlign(VAlign_Center)
-			.AutoWidth()
-			.Padding(2, 0)
-			[
-				SNew(STextBlock)
-				.Text(FText::FromString(InItem->Value))
-				.HighlightText(HighlightText)
-				.ToolTipText(InItem->GetToolTip())
-			]
-			+ SHorizontalBox::Slot()
-			.FillWidth(1)
-			.HAlign(HAlign_Right)
-			.VAlign(VAlign_Center)
-			.Padding(2, 0)
-			[
-				SNew(STextBlock)
-				.Text(InItem->GetComment()) 
-				.ColorAndOpacity(FLinearColor::Yellow)
-				.HighlightText(HighlightText)
-			]
-		];
+			InItem->CreateIcon()
+		]
+		+SHorizontalBox::Slot()
+		.VAlign(VAlign_Center)
+		.AutoWidth()
+		.Padding(2, 0)
+		[
+			SNew(STextBlock)
+			.Text(FText::FromString(InItem->Value))
+			.HighlightText(HighlightText)
+			.ToolTipText(InItem->GetToolTip())
+		]
+		+ SHorizontalBox::Slot()
+		.FillWidth(1)
+		.HAlign(HAlign_Right)
+		.VAlign(VAlign_Center)
+		.Padding(2, 0)
+		[
+			SNew(STextBlock)
+			.Text(InItem->GetComment()) 
+			.ColorAndOpacity(FLinearColor::Yellow)
+			.HighlightText(HighlightText)
+		]
+	];
 }
 
 void SPCGEditorGraphFind::InitiateSearch()
@@ -270,30 +455,413 @@ void SPCGEditorGraphFind::InitiateSearch()
 
 	TreeView->RequestTreeRefresh();
 
+	// Expand so that all items that contain a matched token are shown
+	auto ExpandItem = [this](const FPCGEditorGraphFindResultPtr& Item, auto&& RecursiveCall) -> bool
+	{
+		check(Item);
+		bool bShouldExpandItem = Item->bIsMatch;
+
+		for (const FPCGEditorGraphFindResultPtr& Child : Item->Children)
+		{
+			bShouldExpandItem = RecursiveCall(Child, RecursiveCall);
+		}
+
+		if (bShouldExpandItem)
+		{
+			TreeView->SetItemExpansion(Item, true);
+		}
+
+		return bShouldExpandItem;
+	};
+
 	for (const FPCGEditorGraphFindResultPtr& Item : ItemsFound)
 	{
-		TreeView->SetItemExpansion(Item, true);
+		ExpandItem(Item, ExpandItem);
 	}
 }
 
+namespace PCGEditorGraphFind
+{
+	// Convenience struct to vist graphs & stacks and build a proper hierarchy
+	struct GraphTree
+	{
+		UPCGGraph* Graph = nullptr;
+		UEdGraphNode* Node = nullptr;
+
+		const TArray<GraphTree>& GetChildren() const { return Children; }
+		const GraphTree* GetParent() const { return Parent; }
+
+		bool HasVisitedGraph(UPCGGraph* InGraph) const
+		{
+			return Graph == InGraph || (Parent && Parent->HasVisitedGraph(InGraph));
+		}
+
+		GraphTree* FindOrCreateChild(UEdGraphNode* InNode, UPCGGraph* InGraph)
+		{
+			for (GraphTree& Child : Children)
+			{
+				if (Child.Node == InNode && Child.Graph == InGraph)
+				{
+					return &Child;
+				}
+			}
+
+			// Not found -> add new child
+			GraphTree& NewChild = AddChild();
+			NewChild.Node = InNode;
+			NewChild.Graph = InGraph;
+
+			return &NewChild;
+		}
+
+		// Collapse current subtree so that it is attached to the same parent graph, but closest to the root.
+		bool Collapse()
+		{
+			for (int ChildIndex = Children.Num() - 1; ChildIndex >= 0; --ChildIndex)
+			{
+				if (Children[ChildIndex].Collapse())
+				{
+					Children.RemoveAt(ChildIndex, EAllowShrinking::No);
+				}
+			}
+
+			GraphTree* CurrentParent = Parent;
+			GraphTree* TopmostOccurrenceOfThisGraph = nullptr;
+			while (CurrentParent)
+			{
+				if (CurrentParent->Graph == Graph)
+				{
+					TopmostOccurrenceOfThisGraph = CurrentParent;
+				}
+
+				CurrentParent = CurrentParent->Parent;
+			}
+
+			if (TopmostOccurrenceOfThisGraph)
+			{
+				TopmostOccurrenceOfThisGraph->AppendChildren(Children);
+				return true;
+			}
+			else
+			{
+				return false;
+			}
+		}
+
+		GraphTree* FindFirst(UPCGGraph* InGraph)
+		{
+			TArray<GraphTree*> BFSExpansion = GetBreadthFirstExpansion();
+			for (GraphTree* ExpandedNode : BFSExpansion)
+			{
+				if (ExpandedNode->Graph == InGraph)
+				{
+					return ExpandedNode;
+				}
+			}
+
+			return nullptr;
+		}
+
+	private:
+		TArray<GraphTree> Children;
+		GraphTree* Parent = nullptr;
+
+		TArray<GraphTree*> GetBreadthFirstExpansion()
+		{
+			TArray<GraphTree*> Expansion;
+			Expansion.Add(this);
+			int ExpansionIndex = 0;
+			while (ExpansionIndex < Expansion.Num())
+			{
+				GraphTree* Current = Expansion[ExpansionIndex];
+				for (GraphTree& Child : Current->Children)
+				{
+					Expansion.Add(&Child);
+				}
+
+				++ExpansionIndex;
+			}
+
+			return Expansion;
+		}
+
+		GraphTree& AddChild()
+		{
+			SIZE_T PreviouslyAllocatedSize = Children.GetAllocatedSize();
+			GraphTree& NewChild = Children.Emplace_GetRef();
+			NewChild.Parent = this;
+
+			if (Children.GetAllocatedSize() != PreviouslyAllocatedSize)
+			{
+				UpdateParentPointers();
+			}
+
+			return NewChild;
+		}
+
+		void AppendChildren(const TArray<GraphTree>& InChildren)
+		{
+			Children.Append(InChildren);
+			// Since we might have resized the array and because we're moving the children here, we must update the parent pointers.
+			UpdateParentPointers();
+		}
+
+		void UpdateParentPointers()
+		{
+			for (GraphTree& Child : Children)
+			{
+				Child.Parent = this;
+				Child.UpdateParentPointers();
+			}
+		}
+	};
+} // namespace PCGEditorGraphFind
+
 void SPCGEditorGraphFind::MatchTokens(const TArray<FString>& InTokens)
 {
-	RootFindResult.Reset();
+	TSharedPtr<FPCGEditor> PCGEditor = PCGEditorPtr.Pin();
+	UPCGEditorGraph* PCGEditorGraph = PCGEditor ? PCGEditor->GetPCGEditorGraph() : nullptr;
+	
+	if (!PCGEditor || !PCGEditorGraph)
+	{
+		return;
+	}
+
+	// Start by gathering all graphs to search through.
+	TSet<UPCGGraph*> AllGraphs;
+	UPCGGraph* ThisGraph = PCGEditorGraph->GetPCGGraph();
+	AllGraphs.Add(ThisGraph);
+
+	PCGEditorGraphFind::GraphTree GraphRoot;
+	GraphRoot.Graph = ThisGraph;
+	GraphRoot.Node = nullptr;
+
+	auto VisitAllNodes = [&AllGraphs](UPCGEditorGraph* EditorGraph, PCGEditorGraphFind::GraphTree& Parent, auto&& RecursiveCall)
+	{
+		if (!EditorGraph)
+		{
+			return;
+		}
+
+		for (UEdGraphNode* Node : EditorGraph->Nodes)
+		{
+			UPCGEditorGraphNodeBase* PCGEditorNode = Cast<UPCGEditorGraphNodeBase>(Node);
+			const UPCGBaseSubgraphNode* SubgraphNode = PCGEditorNode ? Cast<UPCGBaseSubgraphNode>(PCGEditorNode->GetPCGNode()) : nullptr;
+			if(UPCGGraph* Subgraph = SubgraphNode ? SubgraphNode->GetSubgraph() : nullptr)
+			{
+				// Check if this subgraph is already present in the hierarchy
+				if (!Parent.HasVisitedGraph(Subgraph))
+				{
+					AllGraphs.Add(Subgraph);
+					if (PCGEditorGraphFind::GraphTree* Child = Parent.FindOrCreateChild(Node, Subgraph))
+					{
+						RecursiveCall(FPCGEditor::GetPCGEditorGraph(Subgraph), *Child, RecursiveCall);
+					}
+				}
+			}
+		}
+	};
+
+	// Visit the static graph(s) starting from the root.
+	VisitAllNodes(PCGEditorGraph, GraphRoot, VisitAllNodes);
+
+	// Then gather all stacks starting from the stack being inspected, and link them to their matching subgraph node.
+	if (PCGEditor->GetStackBeingInspected() && PCGEditorPtr.Pin()->GetSubsystem())
+	{
+		const FPCGStack& CurrentStack = *PCGEditor->GetStackBeingInspected();
+		TArray<FPCGStack> Substacks = PCGEditorPtr.Pin()->GetSubsystem()->GetExecutedStacks(CurrentStack);
+		// For each stack, find the first subgraph node from the start, then get all graphs under it.
+		for (const FPCGStack& Substack : Substacks)
+		{
+			if (Substack.GetStackFrames().Num() <= CurrentStack.GetStackFrames().Num())
+			{
+				continue;
+			}
+
+			// We'll be trying to navigate down the stack and find the node + graph pairs, creating what we need in the graph tree.
+			PCGEditorGraphFind::GraphTree* Current = &GraphRoot;
+			int FrameIndex = CurrentStack.GetStackFrames().Num();
+
+			UEdGraphNode* SubgraphEditorNode = nullptr;
+			UPCGGraph* Subgraph = nullptr;
+			UPCGEditorGraph* CurrentEditorGraph = PCGEditorGraph;
+
+			// There are two kind of subgraph frame structures we are looking for:
+			// Subgraph node (UPCGNode) > Subgraph (UPCGGraph) -> this is when we have a static subgraphs
+			// Subgraph node (UPCGNode) > Loop Index / -1 -> this is when we have a dynamic subgraph or a loop
+			auto CheckFramePair = [&Substack, &SubgraphEditorNode, &Subgraph, &CurrentEditorGraph](int NodeFrameIndex, int SubgraphFrameIndex) -> bool
+			{
+				if(SubgraphFrameIndex >= Substack.GetStackFrames().Num())
+				{
+					return false;
+				}
+
+				const FPCGStackFrame& SubgraphNodeFrame = Substack.GetStackFrames()[NodeFrameIndex];
+				const UPCGNode* SubgraphNode = Cast<UPCGNode>(SubgraphNodeFrame.Object.Get());
+				SubgraphEditorNode = (SubgraphNode && CurrentEditorGraph) ? const_cast<UPCGEditorGraphNodeBase*>(CurrentEditorGraph->GetEditorNodeFromPCGNode(SubgraphNode)) : nullptr;
+				const FPCGStackFrame& SubgraphFrame = Substack.GetStackFrames()[SubgraphFrameIndex];
+				Subgraph = Cast<UPCGGraph>(const_cast<UObject*>(SubgraphFrame.Object.Get()));
+
+				return SubgraphEditorNode && Subgraph;
+			};
+
+			while (FrameIndex < Substack.GetStackFrames().Num() && Current)
+			{
+				bool bProcessPair = false;
+
+				// Test node + graph in a static subgraph configuration
+				if (CheckFramePair(FrameIndex, FrameIndex + 1))
+				{
+					bProcessPair = true;
+					FrameIndex += 2;
+				}
+				// Test node + graph in a dynamic subgraph/loop configuration
+				else if (CheckFramePair(FrameIndex, FrameIndex + 2))
+				{
+					bProcessPair = true;
+					FrameIndex += 3;
+				}
+				else
+				{
+					++FrameIndex;
+				}
+
+				if (bProcessPair)
+				{
+					AllGraphs.Add(Subgraph);
+					// Find/Create child on current
+					Current = Current->FindOrCreateChild(SubgraphEditorNode, Subgraph);
+					CurrentEditorGraph = Current && Subgraph ? FPCGEditor::GetPCGEditorGraph(Subgraph) : nullptr;
+
+					if (Current && CurrentEditorGraph)
+					{
+						VisitAllNodes(CurrentEditorGraph, *Current, VisitAllNodes);
+					}
+					else
+					{
+						// We can't progress here so leave the loop.
+						break;
+					}
+				}
+			}
+		}
+	}
 
-	UPCGEditorGraph* PCGEditorGraph = PCGEditorPtr.Pin()->GetPCGEditorGraph();
+	// Implementation note: at this point, we could visit each graph only once to search for the tokens
+	// but we'd need additional mechanisms to do deep copies, which is not super significant at this point in time.
+	// Perform search by graph based on the tree expansion we'll visit
+	if (FindMode == EPCGGraphFindMode::ShowMinimumTree)
+	{
+		// Moves subtrees to their earliest appearance in the breadth-first tree
+		GraphRoot.Collapse();
+	}
 
-	if (PCGEditorGraph)
+	if(FindMode == EPCGGraphFindMode::ShowMinimumTree || FindMode == EPCGGraphFindMode::ShowFullTree)
 	{
-		RootFindResult = MakeShared<FPCGEditorGraphFindResult>(FString("PCGTreeRoot"));
-		auto GetParentFunc = []() -> FPCGEditorGraphFindResultPtr { return nullptr; };
+		auto RecurseInGraphTree = [this, &InTokens](const PCGEditorGraphFind::GraphTree& Node, TFunctionRef<FPCGEditorGraphFindResultPtr()> GetParentFunc, auto&& RecurseCall) -> void
+		{
+			check(Node.Graph);
 
-		TArray<UPCGGraph*> VisitedGraphs;
-		VisitedGraphs.Add(PCGEditorGraph->GetPCGGraph());
-		MatchTokensInternal(InTokens, PCGEditorGraph, GetParentFunc, VisitedGraphs);
+			FString NodeString;
+			FPCGEditorGraphFindResultPtr ThisResult;
+
+			if (Node.Node)
+			{
+				NodeString = FString::Printf(TEXT("%s (%s)"), *FName::NameToDisplayString(Node.Graph->GetName(), /*bIsBool=*/false), *Node.Node->GetNodeTitle(ENodeTitleType::ListView).ToString());
+			}
+
+			auto GetOrCreateNodeResult = [this, &Node, &NodeString, &ThisResult, &GetParentFunc]()
+			{
+				FPCGEditorGraphFindResultPtr Parent = GetParentFunc();
+				check(Parent);
+
+				if(!ThisResult.IsValid() && Node.Node)
+				{
+					ThisResult = MakeShared<FPCGEditorGraphFindResult>(NodeString, Parent, Node.Node);
+					ThisResult->ParentGraph = FPCGEditor::GetPCGEditorGraph(Node.Graph);
+					Parent->Children.Add(ThisResult);
+				}
+
+				return ThisResult ? ThisResult : Parent;
+			};
+
+			// If current node (including dynamic graph name) matches tokens, add it
+			if (StringMatchesSearchTokens(InTokens, NodeString))
+			{
+				GetOrCreateNodeResult()->bIsMatch = true;
+			}
+
+			// Create local elements
+			MatchTokensInternal(InTokens, FPCGEditor::GetPCGEditorGraph(Node.Graph), GetOrCreateNodeResult/*, MatchedTokens*/);
+
+			// Continue through the tree
+			for (const PCGEditorGraphFind::GraphTree& Child : Node.GetChildren())
+			{
+				RecurseCall(Child, GetOrCreateNodeResult, RecurseCall);
+			}
+		};
+
+		FPCGEditorGraphFindResultPtr RootFindResult = MakeShared<FPCGEditorGraphFindResult>(FString("PCGTreeRoot"));
+		auto GetRootPtr = [&RootFindResult]() { return RootFindResult; };
+
+		RecurseInGraphTree(GraphRoot, GetRootPtr, RecurseInGraphTree);
+		ItemsFound.Append(RootFindResult->Children);
+	}
+	// - flat list (this graph + all downstream graphs after)
+	else if (FindMode == EPCGGraphFindMode::ShowFlatList)
+	{
+		for(UPCGGraph* Graph : AllGraphs)
+		{
+			FPCGEditorGraphFindResultPtr GraphNodePtr;
+			const FString GraphString = FName::NameToDisplayString(Graph->GetName(), /*bIsBool=*/false);
+
+			auto CreateGraphNode = [Graph, &GraphString, &GraphNodePtr, &GraphRoot]()
+			{
+				if (!GraphNodePtr.IsValid())
+				{
+					PCGEditorGraphFind::GraphTree* GraphTreeNode = GraphRoot.FindFirst(Graph);
+
+					UEdGraphNode* EditorNode = nullptr;
+					UEdGraphNode* RootEditorNode = nullptr;
+
+					if (GraphTreeNode)
+					{
+						EditorNode = GraphTreeNode->Node;
+
+						const PCGEditorGraphFind::GraphTree* GraphRootNode = GraphTreeNode;
+						while (GraphRootNode->GetParent() && GraphRootNode->GetParent() != &GraphRoot)
+						{
+							GraphRootNode = GraphRootNode->GetParent();
+						}
+
+						RootEditorNode = GraphRootNode->Node;
+					}
+
+					GraphNodePtr = MakeShared<FPCGEditorGraphFindResult>(GraphString, FPCGEditorGraphFindResultPtr(), EditorNode);
+					GraphNodePtr->ParentGraph = FPCGEditor::GetPCGEditorGraph(Graph);
+					GraphNodePtr->RootGraphNode = RootEditorNode;
+				}
+
+				return GraphNodePtr;
+			};
+
+			// If current node (including dynamic graph name) matches tokens, add it
+			if (StringMatchesSearchTokens(InTokens, GraphString))
+			{
+				CreateGraphNode()->bIsMatch = true;
+			}
+
+			MatchTokensInternal(InTokens, FPCGEditor::GetPCGEditorGraph(Graph), CreateGraphNode/*, Results*/);
+
+			if (GraphNodePtr)
+			{
+				ItemsFound.Add(GraphNodePtr);
+			}
+		}
 	}
 }
 
-void SPCGEditorGraphFind::MatchTokensInternal(const TArray<FString>& InTokens, UPCGEditorGraph* PCGEditorGraph, TFunctionRef<FPCGEditorGraphFindResultPtr()> GetParentFunc, TArray<UPCGGraph*>& VisitedSubgraphs)
+void SPCGEditorGraphFind::MatchTokensInternal(const TArray<FString>& InTokens, UPCGEditorGraph* PCGEditorGraph, TFunctionRef<FPCGEditorGraphFindResultPtr()> GetParentFunc) const
 {
 	if (!PCGEditorGraph)
 	{
@@ -319,61 +887,39 @@ void SPCGEditorGraphFind::MatchTokensInternal(const TArray<FString>& InTokens, U
 		NodeSearchString = NodeSearchString.Replace(TEXT(" "), TEXT(""));
 
 		FPCGEditorGraphFindResultPtr NodeResult;
-		auto GetOrCreateNodeResult = [this, &NodeResult, &NodeString, Node, PCGEditorGraph, &GetParentFunc]() -> FPCGEditorGraphFindResultPtr&
+		auto GetOrCreateNodeResult = [&NodeResult, &NodeString, Node, PCGEditorGraph, &GetParentFunc]() -> FPCGEditorGraphFindResultPtr&
 		{
 			if (!NodeResult.IsValid())
 			{
-				NodeResult = MakeShared<FPCGEditorGraphFindResult>(NodeString, RootFindResult, Node);
-
-				FPCGEditorGraphFindResultPtr Parent = GetParentFunc();
-				if(Parent)
-				{
-					NodeResult->ParentGraph = PCGEditorGraph;
-					Parent->Children.Add(NodeResult);
-				}
-				else
-				{
-					ItemsFound.Add(NodeResult);
-				}
+				NodeResult = MakeShared<FPCGEditorGraphFindResult>(NodeString, GetParentFunc(), Node);
+				NodeResult->ParentGraph = PCGEditorGraph;
+				GetParentFunc()->Children.Add(NodeResult);
 			}
+
 			return NodeResult;
 		};
 
 		if (StringMatchesSearchTokens(InTokens, NodeSearchString))
 		{
-			GetOrCreateNodeResult();
-		}
-	
-		for (UEdGraphPin* Pin : Node->Pins)
-		{
-			if (Pin && Pin->PinFriendlyName.CompareTo(FText::FromString(TEXT(" "))) != 0)
-			{
-				const FText PinName = Pin->GetSchema()->GetPinDisplayName(Pin);
-				FString PinSearchString = Pin->PinName.ToString() + Pin->PinFriendlyName.ToString() + Pin->DefaultValue + Pin->PinType.PinCategory.ToString() + Pin->PinType.PinSubCategory.ToString() + (Pin->PinType.PinSubCategoryObject.IsValid() ? Pin->PinType.PinSubCategoryObject.Get()->GetFullName() : TEXT(""));
-				PinSearchString = PinSearchString.Replace(TEXT(" "), TEXT(""));
-				if (StringMatchesSearchTokens(InTokens, PinSearchString))
-				{
-					FPCGEditorGraphFindResultPtr PinResult(MakeShared<FPCGEditorGraphFindResult>(PinName.ToString(), GetOrCreateNodeResult(), Pin));
-					PinResult->ParentGraph = PCGEditorGraph;
-					NodeResult->Children.Add(PinResult);
-				}
-			}
+			GetOrCreateNodeResult()->bIsMatch = true;
 		}
 
-		// Search recursively in subgraph nodes
-		const UPCGBaseSubgraphNode* SubgraphNode = PCGEditorGraphNodeBase ? Cast<UPCGBaseSubgraphNode>(PCGEditorGraphNodeBase->GetPCGNode()) : nullptr;
-		if(UPCGGraph* Subgraph = SubgraphNode ? SubgraphNode->GetSubgraph() : nullptr)
+		if (bShowPinResults)
 		{
-			// Skip already visited subgraph nodes to prevent recursion issues
-			if (!VisitedSubgraphs.Contains(Subgraph))
+			for (UEdGraphPin* Pin : Node->Pins)
 			{
-				VisitedSubgraphs.Add(Subgraph);
-
-				if (UPCGEditorGraph* EditorSubgraph = FPCGEditor::GetPCGEditorGraph(Subgraph))
+				if (Pin && Pin->PinFriendlyName.CompareTo(FText::FromString(TEXT(" "))) != 0)
 				{
-					const int CountBeforeRecursion = VisitedSubgraphs.Num();
-					MatchTokensInternal(InTokens, EditorSubgraph, GetOrCreateNodeResult, VisitedSubgraphs);
-					VisitedSubgraphs.SetNum(CountBeforeRecursion);
+					const FText PinName = Pin->GetSchema()->GetPinDisplayName(Pin);
+					FString PinSearchString = Pin->PinName.ToString() + Pin->PinFriendlyName.ToString() + Pin->DefaultValue + Pin->PinType.PinCategory.ToString() + Pin->PinType.PinSubCategory.ToString() + (Pin->PinType.PinSubCategoryObject.IsValid() ? Pin->PinType.PinSubCategoryObject.Get()->GetFullName() : TEXT(""));
+					PinSearchString = PinSearchString.Replace(TEXT(" "), TEXT(""));
+					if (StringMatchesSearchTokens(InTokens, PinSearchString))
+					{
+						FPCGEditorGraphFindResultPtr PinResult(MakeShared<FPCGEditorGraphFindResult>(PinName.ToString(), GetOrCreateNodeResult(), Pin));
+						PinResult->ParentGraph = PCGEditorGraph;
+						PinResult->bIsMatch = true;
+						NodeResult->Children.Add(PinResult);
+					}
 				}
 			}
 		}
diff --git a/Engine/Plugins/PCG/Source/PCGEditor/Private/SPCGEditorGraphFind.h b/Engine/Plugins/PCG/Source/PCGEditor/Private/SPCGEditorGraphFind.h
--- a/Engine/Plugins/PCG/Source/PCGEditor/Private/SPCGEditorGraphFind.h
+++ b/Engine/Plugins/PCG/Source/PCGEditor/Private/SPCGEditorGraphFind.h
@@ -13,6 +13,7 @@ namespace ETextCommit { enum Type : int; }
 template <typename ItemType> class STreeView;
 
 class FPCGEditor;
+struct FPCGStack;
 class UEdGraphNode;
 class UPCGEditorGraph;
 class UPCGGraph;
@@ -26,14 +27,17 @@ struct FPCGEditorGraphFindResult
 	FPCGEditorGraphFindResult(const FText& InValue);
 
 	/** Create a listing for a node result */
-	FPCGEditorGraphFindResult(const FString& InValue, TSharedPtr<FPCGEditorGraphFindResult>& InParent, UEdGraphNode* InNode);
+	FPCGEditorGraphFindResult(const FString& InValue, const TSharedPtr<FPCGEditorGraphFindResult>& InParent, UEdGraphNode* InNode);
 
 	/** Create a listing for a pin result */
-	FPCGEditorGraphFindResult(const FString& InValue, TSharedPtr<FPCGEditorGraphFindResult>& InParent, UEdGraphPin* InPin);
+	FPCGEditorGraphFindResult(const FString& InValue, const TSharedPtr<FPCGEditorGraphFindResult>& InParent, UEdGraphPin* InPin);
 
 	/** Called when user clicks on the search item */
 	FReply OnClick(TWeakPtr<FPCGEditor> InPCGEditorPtr);
 
+	/** Called when user double-clicks on the search item */
+	FReply OnDoubleClick(TWeakPtr<FPCGEditor> InPCGEditorPtr);
+
 	/** Get ToolTip for this search result */
 	FText GetToolTip() const;
 
@@ -61,10 +65,22 @@ struct FPCGEditorGraphFindResult
 	/** The graph node that this search result refers to */
 	TWeakObjectPtr<UEdGraphNode> GraphNode;
 
+	/** The graph node that`s in the current editor graph upstream. */
+	TWeakObjectPtr<UEdGraphNode> RootGraphNode;
+
 	/** The pin that this search result refers to */
 	FEdGraphPinReference Pin;
+
+	/** Whether this result contains a found token or not. */
+	bool bIsMatch = false;
 };
 
+enum class EPCGGraphFindMode : uint8
+{
+	ShowMinimumTree,
+	ShowFullTree,
+	ShowFlatList
+};
 
 class SPCGEditorGraphFind : public SCompoundWidget
 {
@@ -72,6 +88,7 @@ public:
 	SLATE_BEGIN_ARGS(SPCGEditorGraphFind) {}
 	SLATE_END_ARGS()
 
+	~SPCGEditorGraphFind();
 	void Construct(const FArguments& InArgs, TSharedPtr<FPCGEditor> InPCGEditor);
 
 	/** Focuses this widget's search box */
@@ -87,12 +104,24 @@ private:
 	/** Called when user changes commits text to the search box */
 	void OnSearchTextCommitted(const FText& InText, ETextCommit::Type InCommitType);
 
+	/** Called when the debug object selected changes, which should trigger a new search */
+	void OnInspectedStackChanged(const FPCGStack& InPCGStack);
+
 	/** Get the children of a row */
 	void OnGetChildren(FPCGEditorGraphFindResultPtr InItem, TArray<FPCGEditorGraphFindResultPtr>& OutChildren);
 
 	/** Called when user clicks on a new result */
 	void OnTreeSelectionChanged(FPCGEditorGraphFindResultPtr Item, ESelectInfo::Type);
 
+	/** Called when an element is double-clicked */
+	void OnTreeDoubleClick(FPCGEditorGraphFindResultPtr InItem);
+
+	/** Called when keys are entered when in the tree selection */
+	FReply OnTreeViewKeyDown(const FGeometry& InGeometry, const FKeyEvent& InKeyEvent) const;
+
+	/** Called when the user clicks on the filter icon */
+	TSharedRef<SWidget> OnFindFilterMenu();
+
 	/** Called when a new row is being generated */
 	TSharedRef<ITableRow> OnGenerateRow(FPCGEditorGraphFindResultPtr InItem, const TSharedRef<STableViewBase>& InOwnerTable);
 
@@ -103,11 +132,17 @@ private:
 	void MatchTokens(const TArray<FString>& InTokens);
 
 	/** Recursive internal implementation of the MatchTokens */
-	void MatchTokensInternal(const TArray<FString>& InTokens, UPCGEditorGraph* PCGEditorGraph, TFunctionRef<FPCGEditorGraphFindResultPtr()> GetParentFunc, TArray<UPCGGraph*>& VisitedSubgraphNodes);
+	void MatchTokensInternal(const TArray<FString>& InTokens, UPCGEditorGraph* PCGEditorGraph, TFunctionRef<FPCGEditorGraphFindResultPtr()> GetParentFunc) const;
 
 	/** Determines if a string matches the search tokens */
 	static bool StringMatchesSearchTokens(const TArray<FString>& InTokens, const FString& InComparisonString);
 
+	/** Called when the find mode is changed */
+	void SetFindMode(EPCGGraphFindMode InFindMode);
+
+	/** Returns true if the current find mode is the same as the one provided */
+	bool IsCurrentFindMode(EPCGGraphFindMode InFindMode) const;
+
 	/** Pointer back to the PCG editor that owns us */
 	TWeakPtr<FPCGEditor> PCGEditorPtr;
 
@@ -120,12 +155,15 @@ private:
 	/** This buffer stores the currently displayed results */
 	TArray<FPCGEditorGraphFindResultPtr> ItemsFound;
 
-	/** we need to keep a handle on the root result, because it won't show up in the tree */
-	FPCGEditorGraphFindResultPtr RootFindResult;
-
 	/** The string to highlight in the results */
 	FText HighlightText;
 
 	/** The string to search for */
 	FString SearchValue;
+
+	/** Controls the way find results are presented. */
+	EPCGGraphFindMode FindMode = EPCGGraphFindMode::ShowMinimumTree;
+
+	/** Controls whether pin names are show in the find results. */
+	bool bShowPinResults = false;
 };
