Removed unresolved keys from tracking in nodes that do loading to allow for test-and-load use cases.
Cleaned up tracking on cleanup.

diff --git a/Engine/Plugins/PCG/Source/PCG/Private/Elements/IO/PCGLoadAssetElement.cpp b/Engine/Plugins/PCG/Source/PCG/Private/Elements/IO/PCGLoadAssetElement.cpp
--- a/Engine/Plugins/PCG/Source/PCG/Private/Elements/IO/PCGLoadAssetElement.cpp
+++ b/Engine/Plugins/PCG/Source/PCG/Private/Elements/IO/PCGLoadAssetElement.cpp
@@ -184,16 +184,11 @@ bool FPCGLoadDataAssetElement::ExecuteInternal(FPCGContext* InContext) const
 	check(Settings);
 
 #if WITH_EDITOR
-	if (Context->IsValueOverriden(GET_MEMBER_NAME_CHECKED(UPCGLoadDataAssetSettings, Asset)) || (Settings->bLoadFromInput && !Context->PathsToObjectsAndDataIndex.IsEmpty()))
+	FPCGDynamicTrackingHelper DynamicTracking;
+	const bool bRequiresDynamicTracking = Context->IsValueOverriden(GET_MEMBER_NAME_CHECKED(UPCGLoadDataAssetSettings, Asset)) || (Settings->bLoadFromInput && !Context->PathsToObjectsAndDataIndex.IsEmpty());
+	if (bRequiresDynamicTracking)
 	{
-		FPCGDynamicTrackingHelper DynamicTracking;
 		DynamicTracking.EnableAndInitialize(Context, Context->PathsToObjectsAndDataIndex.Num());
-		for(const TTuple<FSoftObjectPath, int32, int32>& Path : Context->PathsToObjectsAndDataIndex)
-		{
-			DynamicTracking.AddToTracking(FPCGSelectionKey::CreateFromPath(Path.Get<0>()), /*bIsCulled=*/false);
-		}
-
-		DynamicTracking.Finalize(Context);
 	}
 #endif
 
@@ -229,8 +224,22 @@ bool FPCGLoadDataAssetElement::ExecuteInternal(FPCGContext* InContext) const
 					}
 				}
 			}
+
+#if WITH_EDITOR
+			if (bRequiresDynamicTracking)
+			{
+				DynamicTracking.AddToTracking(FPCGSelectionKey::CreateFromPath(AssetPath.Get<0>()), /*bIsCulled=*/false);
+			}
+#endif
 		}
 	}
 
+#if WITH_EDITOR
+	if (bRequiresDynamicTracking)
+	{
+		DynamicTracking.Finalize(Context);
+	}
+#endif
+
 	return true;
 }
\ No newline at end of file
diff --git a/Engine/Plugins/PCG/Source/PCG/Private/Elements/PCGDataFromActor.cpp b/Engine/Plugins/PCG/Source/PCG/Private/Elements/PCGDataFromActor.cpp
--- a/Engine/Plugins/PCG/Source/PCG/Private/Elements/PCGDataFromActor.cpp
+++ b/Engine/Plugins/PCG/Source/PCG/Private/Elements/PCGDataFromActor.cpp
@@ -474,6 +474,12 @@ bool FPCGDataFromActorElement::ExecuteInternal(FPCGContext* InContext) const
 				DynamicTracking.EnableAndInitialize(Context, Context->PathsToObjectsAndDataIndex.Num());
 				for (const TTuple<FSoftObjectPath, int32, int32>& Path : Context->PathsToObjectsAndDataIndex)
 				{
+					// Skip dependency tracking on actors that couldn't be loaded
+					if (!Path.Get<0>().ResolveObject())
+					{
+						continue;
+					}
+
 					FPCGSelectionKey Key = FPCGSelectionKey::CreateFromPath(Path.Get<0>());
 					if (Settings->Mode == EPCGGetDataFromActorMode::GetDataFromPCGComponent || Settings->Mode == EPCGGetDataFromActorMode::GetDataFromPCGComponentOrParseComponents)
 					{
diff --git a/Engine/Plugins/PCG/Source/PCG/Private/Elements/PCGGetActorProperty.cpp b/Engine/Plugins/PCG/Source/PCG/Private/Elements/PCGGetActorProperty.cpp
--- a/Engine/Plugins/PCG/Source/PCG/Private/Elements/PCGGetActorProperty.cpp
+++ b/Engine/Plugins/PCG/Source/PCG/Private/Elements/PCGGetActorProperty.cpp
@@ -247,6 +247,12 @@ bool FPCGGetActorPropertyElement::ExecuteInternal(FPCGContext* InContext) const
 		DynamicTracking.EnableAndInitialize(Context, Context->PathsToObjectsAndDataIndex.Num());
 		for (const TTuple<FSoftObjectPath, int32, int32>& Path : Context->PathsToObjectsAndDataIndex)
 		{
+			// Skip tracking for references that weren't successfully resolved
+			if (!Path.Get<0>().ResolveObject())
+			{
+				continue;
+			}
+
 			DynamicTracking.AddToTracking(FPCGSelectionKey::CreateFromPath(Path.Get<0>()), Settings->ActorSelector.bMustOverlapSelf);
 		}
 
diff --git a/Engine/Plugins/PCG/Source/PCG/Private/PCGComponent.cpp b/Engine/Plugins/PCG/Source/PCG/Private/PCGComponent.cpp
--- a/Engine/Plugins/PCG/Source/PCG/Private/PCGComponent.cpp
+++ b/Engine/Plugins/PCG/Source/PCG/Private/PCGComponent.cpp
@@ -811,6 +811,8 @@ void UPCGComponent::PostCleanupGraph()
 	GeneratedGraphOutput.Reset();
 
 #if WITH_EDITOR
+	DynamicallyTrackedKeysToSettings.Reset();
+
 	// Notify Subsystem first
 	if (UPCGSubsystem* Subsystem = GetSubsystem())
 	{
