Adjusted node palette so it never updates more than once per frame

diff --git a/Engine/Plugins/PCG/Source/PCGEditor/Private/SPCGEditorGraphNodePalette.cpp b/Engine/Plugins/PCG/Source/PCGEditor/Private/SPCGEditorGraphNodePalette.cpp
--- a/Engine/Plugins/PCG/Source/PCGEditor/Private/SPCGEditorGraphNodePalette.cpp
+++ b/Engine/Plugins/PCG/Source/PCGEditor/Private/SPCGEditorGraphNodePalette.cpp
@@ -134,7 +134,7 @@ void SPCGEditorGraphNodePalette::OnAssetChanged(const FAssetData& InAssetData)
 		InAssetData.IsInstanceOf<UPCGSettings>() ||
 		PCGEditorUtils::IsAssetPCGBlueprint(InAssetData))
 	{
-		RefreshActionsList(true);
+		bNeedsRefresh = true;
 	}
 }
 
@@ -146,7 +146,19 @@ void SPCGEditorGraphNodePalette::OnAssetRenamed(const FAssetData& InAssetData, c
 void SPCGEditorGraphNodePalette::OnTypeSelectionChanged(int32 InValue, ESelectInfo::Type /*SelectInfo*/)
 {
 	ElementType = static_cast<EPCGElementType>(InValue);
-	RefreshActionsList(true);
+	bNeedsRefresh = true;
+}
+
+void SPCGEditorGraphNodePalette::Tick(const FGeometry& AllottedGeometry, const double InCurrentTime, const float InDeltaTime)
+{
+	SGraphPalette::Tick(AllottedGeometry, InCurrentTime, InDeltaTime);
+
+	if (bNeedsRefresh)
+	{
+		TRACE_CPUPROFILER_EVENT_SCOPE(SPCGEditorGraphNodePalette::RefreshActionsList);
+		RefreshActionsList(true);
+		bNeedsRefresh = false;
+	}
 }
 
 int32 SPCGEditorGraphNodePalette::GetTypeValue() const
diff --git a/Engine/Plugins/PCG/Source/PCGEditor/Private/SPCGEditorGraphNodePalette.h b/Engine/Plugins/PCG/Source/PCGEditor/Private/SPCGEditorGraphNodePalette.h
--- a/Engine/Plugins/PCG/Source/PCGEditor/Private/SPCGEditorGraphNodePalette.h
+++ b/Engine/Plugins/PCG/Source/PCGEditor/Private/SPCGEditorGraphNodePalette.h
@@ -37,6 +37,10 @@ protected:
 	virtual void CollectAllActions(FGraphActionListBuilderBase& OutAllActions) override;
 	//~ End SGraphPalette Interface
 
+	//~ Begin SWidget interface
+	void Tick(const FGeometry& AllottedGeometry, const double InCurrentTime, const float InDeltaTime) override;
+	//~ End SWidget interface
+
 private:
 	void OnAssetChanged(const FAssetData& InAssetData);
 	void OnAssetRenamed(const FAssetData& InAssetData, const FString& InNewAssetName);
@@ -45,4 +49,5 @@ private:
 	int32 GetTypeValue() const;
 	
 	EPCGElementType ElementType = EPCGElementType::All;
+	bool bNeedsRefresh = false;
 };
