[PCG] Address bug with overstepping SegmentIndex in SplineSampler

#rnx

diff --git a/Engine/Plugins/PCG/Source/PCG/Private/Elements/PCGSplineSampler.cpp b/Engine/Plugins/PCG/Source/PCG/Private/Elements/PCGSplineSampler.cpp
--- a/Engine/Plugins/PCG/Source/PCG/Private/Elements/PCGSplineSampler.cpp
+++ b/Engine/Plugins/PCG/Source/PCG/Private/Elements/PCGSplineSampler.cpp
@@ -570,6 +570,16 @@ namespace PCGSplineSamplerHelpers
 
 		virtual void Step(FSamplerResult& OutResult) override
 		{
+			const int NumSegments = LineData->GetNumSegments();
+			// If the current segment exceeds the number of segments, assume the segment is the last valid segment.
+			if (CurrentSegmentIndex >= NumSegments)
+			{
+				CurrentSegmentIndex = NumSegments - 1;
+				// Adjust the current distances relative to the start of the final segment, so that the final control point is sampled.
+				CurrentDistance = LineData->GetSegmentLength(CurrentSegmentIndex);
+				DistanceToCurrentSegment = LineData->GetDistanceAtSegmentStart(CurrentSegmentIndex);
+			}
+
 			FVector::FReal OffsetDistance = CurrentDistance + (bUseRandomOffset ? RandomSource.FRandRange(-MaxRandomOffset, MaxRandomOffset) : 0.0);
 
 			// Prevent samples from wrapping around on open splines.
