[PCG] Fix accessor factory registered in editor only

#rnx

diff --git a/Engine/Plugins/PCG/Source/PCG/Private/PCGModule.cpp b/Engine/Plugins/PCG/Source/PCG/Private/PCGModule.cpp
--- a/Engine/Plugins/PCG/Source/PCG/Private/PCGModule.cpp
+++ b/Engine/Plugins/PCG/Source/PCG/Private/PCGModule.cpp
@@ -32,9 +32,9 @@ bool FPCGModule::IsPCGModuleLoaded()
 	return FModuleManager::Get().IsModuleLoaded(TEXT("PCG"));
 }
 
-#if WITH_EDITOR
 void FPCGModule::StartupModule()
 {
+#if WITH_EDITOR
 	PCGDeterminismTests::FNativeTestRegistry::Create();
 
 	RegisterNativeElementDeterminismTests();
@@ -43,6 +43,7 @@ void FPCGModule::StartupModule()
 
 	const FString PluginShaderDir = FPaths::Combine(IPluginManager::Get().FindPlugin(TEXT("PCG"))->GetBaseDir(), TEXT("Shaders"));
 	AddShaderSourceDirectoryMapping(TEXT("/Plugin/PCG"), PluginShaderDir);
+#endif // WITH_EDITOR
 
 	// Registering accessor methods
 	AttributeAccessorFactory.RegisterDefaultMethods();
@@ -54,12 +55,15 @@ void FPCGModule::ShutdownModule()
 	// Unregistering accessor methods
 	AttributeAccessorFactory.UnregisterDefaultMethods();
 	AttributeAccessorFactory.UnregisterMethods<UPCGPointData>();
-	
+
+#if WITH_EDITOR
 	DeregisterNativeElementDeterminismTests();
 
 	PCGDeterminismTests::FNativeTestRegistry::Destroy();
+#endif // WITH_EDITOR
 }
 
+#if WITH_EDITOR
 void FPCGModule::RegisterNativeElementDeterminismTests()
 {
 	PCGDeterminismTests::FNativeTestRegistry::RegisterTestFunction(UPCGDifferenceSettings::StaticClass(), PCGDeterminismTests::DifferenceElement::RunTestSuite);
@@ -70,7 +74,7 @@ void FPCGModule::DeregisterNativeElementDeterminismTests()
 {
 	PCGDeterminismTests::FNativeTestRegistry::DeregisterTestFunction(UPCGDifferenceSettings::StaticClass());
 }
-#endif
+#endif // WITH_EDITOR
 
 IMPLEMENT_MODULE(FPCGModule, PCG);
 
diff --git a/Engine/Plugins/PCG/Source/PCG/Public/PCGModule.h b/Engine/Plugins/PCG/Source/PCG/Public/PCGModule.h
--- a/Engine/Plugins/PCG/Source/PCG/Public/PCGModule.h
+++ b/Engine/Plugins/PCG/Source/PCG/Public/PCGModule.h
@@ -29,10 +29,8 @@ class FPCGModule final : public IModuleInterface
 {
 public:
 	//~ IModuleInterface implementation
-#if WITH_EDITOR
 	virtual void StartupModule() override;
 	virtual void ShutdownModule() override;
-#endif
 	virtual bool SupportsDynamicReloading() override { return true; }
 	//~ End IModuleInterface implementation
 
