[PCG] Fix CopyPoints overwriting the first inherited attributes
* Also fix linker error stemming from misuse of conversion rules for arithmetic operations

#rnx

diff --git a/Engine/Plugins/PCG/Shaders/Private/Elements/PCGCopyPoints.usf b/Engine/Plugins/PCG/Shaders/Private/Elements/PCGCopyPoints.usf
--- a/Engine/Plugins/PCG/Shaders/Private/Elements/PCGCopyPoints.usf
+++ b/Engine/Plugins/PCG/Shaders/Private/Elements/PCGCopyPoints.usf
@@ -21,7 +21,7 @@
 #define NUM_RESERVED_ATTRS 32
 
 // TODO pass in attribute IDs to copy.
-void CopyAttributesFromSource(uint Source_DataIndex, uint Source_ElementIndex, uint Out_DataIndex, uint Out_ElementIndex)
+void CopyAttributesFromSource(uint Source_DataIndex, uint Source_ElementIndex, uint Out_DataIndex, uint Out_ElementIndex, bool bInitializeNonCopiedAttributes)
 {
 	const uint NumAttributes = Out_GetDataNumAttributesInternal(Out_DataIndex);
 	uint NumAttributesProcessed = 0;
@@ -39,6 +39,7 @@ void CopyAttributesFromSource(uint Source_DataIndex, uint Source_ElementIndex, u
 		const uint Source_Stride = Source_GetAttributeStrideInternal(Source_DataIndex, AttributeId);
 		const uint Out_ElementAddress = Out_GetElementAddressInternal(Out_DataIndex, Out_ElementIndex, AttributeId);
 		
+		// Copy the attribute if it exists in the Source data.
 		if (Source_Stride == Out_Stride)
 		{
 			const uint Source_ElementAddress = Source_GetElementAddressInternal(Source_DataIndex, Source_ElementIndex, AttributeId);
@@ -48,7 +49,7 @@ void CopyAttributesFromSource(uint Source_DataIndex, uint Source_ElementIndex, u
 				Out_StoreBufferInternal(Out_ElementAddress + Offset, Source_LoadBufferInternal(Source_ElementAddress + Offset));
 			}
 		}
-		else
+		else if (bInitializeNonCopiedAttributes)
 		{
 			for (uint Offset = 0; Offset < Out_Stride; Offset += 4)
 			{
@@ -65,7 +66,7 @@ void CopyAttributesFromSource(uint Source_DataIndex, uint Source_ElementIndex, u
 }
 
 // TODO pass in attribute IDs to copy.
-void CopyAttributesFromTarget(uint Target_DataIndex, uint Target_ElementIndex, uint Out_DataIndex, uint Out_ElementIndex)
+void CopyAttributesFromTarget(uint Target_DataIndex, uint Target_ElementIndex, uint Out_DataIndex, uint Out_ElementIndex, bool bInitializeNonCopiedAttributes)
 {
 	const uint NumAttributes = Out_GetDataNumAttributesInternal(Out_DataIndex);
 	uint NumAttributesProcessed = 0;
@@ -83,6 +84,7 @@ void CopyAttributesFromTarget(uint Target_DataIndex, uint Target_ElementIndex, u
 		const uint Target_Stride = Target_GetAttributeStrideInternal(Target_DataIndex, AttributeId);
 		const uint Out_ElementAddress = Out_GetElementAddressInternal(Out_DataIndex, Out_ElementIndex, AttributeId);
 		
+		// Copy the attribute if it exists in the Target data.
 		if (Target_Stride == Out_Stride)
 		{
 			const uint Target_ElementAddress = Target_GetElementAddressInternal(Target_DataIndex, Target_ElementIndex, AttributeId);
@@ -92,7 +94,7 @@ void CopyAttributesFromTarget(uint Target_DataIndex, uint Target_ElementIndex, u
 				Out_StoreBufferInternal(Out_ElementAddress + Offset, Target_LoadBufferInternal(Target_ElementAddress + Offset));
 			}
 		}
-		else
+		else if (bInitializeNonCopiedAttributes)
 		{
 			for (uint Offset = 0; Offset < Out_Stride; Offset += 4)
 			{
@@ -302,20 +304,24 @@ void Main(uint3 GroupId : SV_GroupID, uint GroupIndex : SV_GroupIndex)
 	// Copy attributes
 	if (AttributeInheritance == COPY_POINTS_METADATA_INHERITANCE_MODE_SourceOnly)
 	{
-		CopyAttributesFromSource(SourceDataIndex, SourcePointIndexInData, OutDataIndex, OutElemIndex);
+		CopyAttributesFromSource(SourceDataIndex, SourcePointIndexInData, OutDataIndex, OutElemIndex, /*bInitializeNonCopiedAttributes=*/true);
 	}
 	else if (AttributeInheritance == COPY_POINTS_METADATA_INHERITANCE_MODE_TargetOnly)
 	{
-		CopyAttributesFromTarget(TargetDataIndex, TargetPointIndexInData, OutDataIndex, OutElemIndex);
+		CopyAttributesFromTarget(TargetDataIndex, TargetPointIndexInData, OutDataIndex, OutElemIndex, /*bInitializeNonCopiedAttributes=*/true);
 	}
 	else if (AttributeInheritance == COPY_POINTS_METADATA_INHERITANCE_MODE_SourceFirst)
 	{
-		CopyAttributesFromTarget(TargetDataIndex, TargetPointIndexInData, OutDataIndex, OutElemIndex);
-		CopyAttributesFromSource(SourceDataIndex, SourcePointIndexInData, OutDataIndex, OutElemIndex);
+		CopyAttributesFromTarget(TargetDataIndex, TargetPointIndexInData, OutDataIndex, OutElemIndex, /*bInitializeNonCopiedAttributes=*/true);
+
+		// Don't allow Source to initialize non-copied attributes in the output, otherwise it will overwrite attributes from Target.
+		CopyAttributesFromSource(SourceDataIndex, SourcePointIndexInData, OutDataIndex, OutElemIndex, /*bInitializeNonCopiedAttributes=*/false);
 	}
 	else if (AttributeInheritance == COPY_POINTS_METADATA_INHERITANCE_MODE_TargetFirst)
 	{
-		CopyAttributesFromSource(SourceDataIndex, SourcePointIndexInData, OutDataIndex, OutElemIndex);
-		CopyAttributesFromTarget(TargetDataIndex, TargetPointIndexInData, OutDataIndex, OutElemIndex);
+		CopyAttributesFromSource(SourceDataIndex, SourcePointIndexInData, OutDataIndex, OutElemIndex, /*bInitializeNonCopiedAttributes=*/true);
+
+		// Don't allow Target to initialize non-copied attributes in the output, otherwise it will overwrite attributes from Source.
+		CopyAttributesFromTarget(TargetDataIndex, TargetPointIndexInData, OutDataIndex, OutElemIndex, /*bInitializeNonCopiedAttributes=*/false);
 	}
 }
\ No newline at end of file
