[PCG] Factored out the metadata's populate by enum to be used by multiple systems

#rnx

diff --git a/Engine/Plugins/PCG/Source/PCG/Private/Compute/Elements/PCGCustomHLSL.cpp b/Engine/Plugins/PCG/Source/PCG/Private/Compute/Elements/PCGCustomHLSL.cpp
--- a/Engine/Plugins/PCG/Source/PCG/Private/Compute/Elements/PCGCustomHLSL.cpp
+++ b/Engine/Plugins/PCG/Source/PCG/Private/Compute/Elements/PCGCustomHLSL.cpp
@@ -13,7 +13,6 @@
 #include "Compute/PCGComputeGraph.h"
 #include "Compute/PCGDataBinding.h"
 #include "Data/PCGPointData.h"
-#include "Elements/Metadata/PCGMetadataElementCommon.h"
 
 #include "Internationalization/Regex.h"
 #include "Containers/StaticArray.h"
@@ -468,7 +467,7 @@ bool UPCGCustomHLSLSettings::ComputeOutputPinDataDesc(const UPCGPin* OutputPin,
 #if WITH_EDITOR
 TArray<FPCGPreConfiguredSettingsInfo> UPCGCustomHLSLSettings::GetPreconfiguredInfo() const
 {
-	return PCGMetadataElementCommon::FillPreconfiguredSettingsInfoFromEnum<EPCGKernelType>();
+	return FPCGPreConfiguredSettingsInfo::PopulateFromEnum<EPCGKernelType>();
 }
 
 EPCGChangeType UPCGCustomHLSLSettings::GetChangeTypeForProperty(const FName& InPropertyName) const
diff --git a/Engine/Plugins/PCG/Source/PCG/Private/Elements/Metadata/PCGMetadataBitwiseOpElement.cpp b/Engine/Plugins/PCG/Source/PCG/Private/Elements/Metadata/PCGMetadataBitwiseOpElement.cpp
--- a/Engine/Plugins/PCG/Source/PCG/Private/Elements/Metadata/PCGMetadataBitwiseOpElement.cpp
+++ b/Engine/Plugins/PCG/Source/PCG/Private/Elements/Metadata/PCGMetadataBitwiseOpElement.cpp
@@ -111,7 +111,7 @@ FText UPCGMetadataBitwiseSettings::GetDefaultNodeTitle() const
 
 TArray<FPCGPreConfiguredSettingsInfo> UPCGMetadataBitwiseSettings::GetPreconfiguredInfo() const
 {
-	return PCGMetadataElementCommon::FillPreconfiguredSettingsInfoFromEnum<EPCGMetadataBitwiseOperation>();
+	return FPCGPreConfiguredSettingsInfo::PopulateFromEnum<EPCGMetadataBitwiseOperation>();
 }
 #endif // WITH_EDITOR
 
diff --git a/Engine/Plugins/PCG/Source/PCG/Private/Elements/Metadata/PCGMetadataBooleanOpElement.cpp b/Engine/Plugins/PCG/Source/PCG/Private/Elements/Metadata/PCGMetadataBooleanOpElement.cpp
--- a/Engine/Plugins/PCG/Source/PCG/Private/Elements/Metadata/PCGMetadataBooleanOpElement.cpp
+++ b/Engine/Plugins/PCG/Source/PCG/Private/Elements/Metadata/PCGMetadataBooleanOpElement.cpp
@@ -116,7 +116,7 @@ FText UPCGMetadataBooleanSettings::GetDefaultNodeTitle() const
 
 TArray<FPCGPreConfiguredSettingsInfo> UPCGMetadataBooleanSettings::GetPreconfiguredInfo() const
 {
-	return PCGMetadataElementCommon::FillPreconfiguredSettingsInfoFromEnum<EPCGMetadataBooleanOperation>();
+	return FPCGPreConfiguredSettingsInfo::PopulateFromEnum<EPCGMetadataBooleanOperation>();
 }
 #endif // WITH_EDITOR
 
diff --git a/Engine/Plugins/PCG/Source/PCG/Private/Elements/Metadata/PCGMetadataCompareOpElement.cpp b/Engine/Plugins/PCG/Source/PCG/Private/Elements/Metadata/PCGMetadataCompareOpElement.cpp
--- a/Engine/Plugins/PCG/Source/PCG/Private/Elements/Metadata/PCGMetadataCompareOpElement.cpp
+++ b/Engine/Plugins/PCG/Source/PCG/Private/Elements/Metadata/PCGMetadataCompareOpElement.cpp
@@ -142,7 +142,7 @@ FText UPCGMetadataCompareSettings::GetDefaultNodeTitle() const
 
 TArray<FPCGPreConfiguredSettingsInfo> UPCGMetadataCompareSettings::GetPreconfiguredInfo() const
 {
-	return PCGMetadataElementCommon::FillPreconfiguredSettingsInfoFromEnum<EPCGMetadataCompareOperation>();
+	return FPCGPreConfiguredSettingsInfo::PopulateFromEnum<EPCGMetadataCompareOperation>();
 }
 #endif // WITH_EDITOR
 
diff --git a/Engine/Plugins/PCG/Source/PCG/Private/Elements/Metadata/PCGMetadataMakeRotator.cpp b/Engine/Plugins/PCG/Source/PCG/Private/Elements/Metadata/PCGMetadataMakeRotator.cpp
--- a/Engine/Plugins/PCG/Source/PCG/Private/Elements/Metadata/PCGMetadataMakeRotator.cpp
+++ b/Engine/Plugins/PCG/Source/PCG/Private/Elements/Metadata/PCGMetadataMakeRotator.cpp
@@ -107,7 +107,7 @@ FText UPCGMetadataMakeRotatorSettings::GetDefaultNodeTitle() const
 
 TArray<FPCGPreConfiguredSettingsInfo> UPCGMetadataMakeRotatorSettings::GetPreconfiguredInfo() const
 {
-	return PCGMetadataElementCommon::FillPreconfiguredSettingsInfoFromEnum<EPCGMetadataMakeRotatorOp>();
+	return FPCGPreConfiguredSettingsInfo::PopulateFromEnum<EPCGMetadataMakeRotatorOp>();
 }
 #endif // WITH_EDITOR
 
diff --git a/Engine/Plugins/PCG/Source/PCG/Private/Elements/Metadata/PCGMetadataMathsOpElement.cpp b/Engine/Plugins/PCG/Source/PCG/Private/Elements/Metadata/PCGMetadataMathsOpElement.cpp
--- a/Engine/Plugins/PCG/Source/PCG/Private/Elements/Metadata/PCGMetadataMathsOpElement.cpp
+++ b/Engine/Plugins/PCG/Source/PCG/Private/Elements/Metadata/PCGMetadataMathsOpElement.cpp
@@ -302,7 +302,7 @@ FText UPCGMetadataMathsSettings::GetDefaultNodeTitle() const
 
 TArray<FPCGPreConfiguredSettingsInfo> UPCGMetadataMathsSettings::GetPreconfiguredInfo() const
 {
-	return PCGMetadataElementCommon::FillPreconfiguredSettingsInfoFromEnum<EPCGMetadataMathsOperation>({ EPCGMetadataMathsOperation::UnaryOp, EPCGMetadataMathsOperation::BinaryOp, EPCGMetadataMathsOperation::TernaryOp });
+	return FPCGPreConfiguredSettingsInfo::PopulateFromEnum<EPCGMetadataMathsOperation>({ EPCGMetadataMathsOperation::UnaryOp, EPCGMetadataMathsOperation::BinaryOp, EPCGMetadataMathsOperation::TernaryOp });
 }
 #endif // WITH_EDITOR
 
diff --git a/Engine/Plugins/PCG/Source/PCG/Private/Elements/Metadata/PCGMetadataRotatorOpElement.cpp b/Engine/Plugins/PCG/Source/PCG/Private/Elements/Metadata/PCGMetadataRotatorOpElement.cpp
--- a/Engine/Plugins/PCG/Source/PCG/Private/Elements/Metadata/PCGMetadataRotatorOpElement.cpp
+++ b/Engine/Plugins/PCG/Source/PCG/Private/Elements/Metadata/PCGMetadataRotatorOpElement.cpp
@@ -219,7 +219,7 @@ FText UPCGMetadataRotatorSettings::GetDefaultNodeTitle() const
 
 TArray<FPCGPreConfiguredSettingsInfo> UPCGMetadataRotatorSettings::GetPreconfiguredInfo() const
 {
-	return PCGMetadataElementCommon::FillPreconfiguredSettingsInfoFromEnum<EPCGMetadataRotatorOperation>({ EPCGMetadataRotatorOperation::RotatorOp, EPCGMetadataRotatorOperation::TransformOp });
+	return FPCGPreConfiguredSettingsInfo::PopulateFromEnum<EPCGMetadataRotatorOperation>({ EPCGMetadataRotatorOperation::RotatorOp, EPCGMetadataRotatorOperation::TransformOp });
 }
 #endif // WITH_EDITOR
 
diff --git a/Engine/Plugins/PCG/Source/PCG/Private/Elements/Metadata/PCGMetadataStringOpElement.cpp b/Engine/Plugins/PCG/Source/PCG/Private/Elements/Metadata/PCGMetadataStringOpElement.cpp
--- a/Engine/Plugins/PCG/Source/PCG/Private/Elements/Metadata/PCGMetadataStringOpElement.cpp
+++ b/Engine/Plugins/PCG/Source/PCG/Private/Elements/Metadata/PCGMetadataStringOpElement.cpp
@@ -94,7 +94,7 @@ FText UPCGMetadataStringOpSettings::GetDefaultNodeTitle() const
 
 TArray<FPCGPreConfiguredSettingsInfo> UPCGMetadataStringOpSettings::GetPreconfiguredInfo() const
 {
-	return PCGMetadataElementCommon::FillPreconfiguredSettingsInfoFromEnum<EPCGMetadataStringOperation>();
+	return FPCGPreConfiguredSettingsInfo::PopulateFromEnum<EPCGMetadataStringOperation>();
 }
 #endif // WITH_EDITOR
 
diff --git a/Engine/Plugins/PCG/Source/PCG/Private/Elements/Metadata/PCGMetadataTransformOpElement.cpp b/Engine/Plugins/PCG/Source/PCG/Private/Elements/Metadata/PCGMetadataTransformOpElement.cpp
--- a/Engine/Plugins/PCG/Source/PCG/Private/Elements/Metadata/PCGMetadataTransformOpElement.cpp
+++ b/Engine/Plugins/PCG/Source/PCG/Private/Elements/Metadata/PCGMetadataTransformOpElement.cpp
@@ -182,7 +182,7 @@ FText UPCGMetadataTransformSettings::GetDefaultNodeTitle() const
 
 TArray<FPCGPreConfiguredSettingsInfo> UPCGMetadataTransformSettings::GetPreconfiguredInfo() const
 {
-	return PCGMetadataElementCommon::FillPreconfiguredSettingsInfoFromEnum<EPCGMetadataTransformOperation>();
+	return FPCGPreConfiguredSettingsInfo::PopulateFromEnum<EPCGMetadataTransformOperation>();
 }
 #endif // WITH_EDITOR
 
diff --git a/Engine/Plugins/PCG/Source/PCG/Private/Elements/Metadata/PCGMetadataTrigOpElement.cpp b/Engine/Plugins/PCG/Source/PCG/Private/Elements/Metadata/PCGMetadataTrigOpElement.cpp
--- a/Engine/Plugins/PCG/Source/PCG/Private/Elements/Metadata/PCGMetadataTrigOpElement.cpp
+++ b/Engine/Plugins/PCG/Source/PCG/Private/Elements/Metadata/PCGMetadataTrigOpElement.cpp
@@ -128,7 +128,7 @@ FText UPCGMetadataTrigSettings::GetDefaultNodeTitle() const
 
 TArray<FPCGPreConfiguredSettingsInfo> UPCGMetadataTrigSettings::GetPreconfiguredInfo() const
 {
-	return PCGMetadataElementCommon::FillPreconfiguredSettingsInfoFromEnum<EPCGMetadataTrigOperation>();
+	return FPCGPreConfiguredSettingsInfo::PopulateFromEnum<EPCGMetadataTrigOperation>();
 }
 #endif // WITH_EDITOR
 
diff --git a/Engine/Plugins/PCG/Source/PCG/Private/Elements/Metadata/PCGMetadataVectorOpElement.cpp b/Engine/Plugins/PCG/Source/PCG/Private/Elements/Metadata/PCGMetadataVectorOpElement.cpp
--- a/Engine/Plugins/PCG/Source/PCG/Private/Elements/Metadata/PCGMetadataVectorOpElement.cpp
+++ b/Engine/Plugins/PCG/Source/PCG/Private/Elements/Metadata/PCGMetadataVectorOpElement.cpp
@@ -344,7 +344,7 @@ FText UPCGMetadataVectorSettings::GetDefaultNodeTitle() const
 
 TArray<FPCGPreConfiguredSettingsInfo> UPCGMetadataVectorSettings::GetPreconfiguredInfo() const
 {
-	return PCGMetadataElementCommon::FillPreconfiguredSettingsInfoFromEnum<EPCGMetadataVectorOperation>({ EPCGMetadataVectorOperation::VectorOp, EPCGMetadataVectorOperation::TransformOp });
+	return FPCGPreConfiguredSettingsInfo::PopulateFromEnum<EPCGMetadataVectorOperation>({ EPCGMetadataVectorOperation::VectorOp, EPCGMetadataVectorOperation::TransformOp });
 }
 #endif // WITH_EDITOR
 
diff --git a/Engine/Plugins/PCG/Source/PCG/Private/Elements/PCGAttributeCast.cpp b/Engine/Plugins/PCG/Source/PCG/Private/Elements/PCGAttributeCast.cpp
--- a/Engine/Plugins/PCG/Source/PCG/Private/Elements/PCGAttributeCast.cpp
+++ b/Engine/Plugins/PCG/Source/PCG/Private/Elements/PCGAttributeCast.cpp
@@ -28,7 +28,7 @@ FPCGElementPtr UPCGAttributeCastSettings::CreateElement() const
 #if WITH_EDITOR
 TArray<FPCGPreConfiguredSettingsInfo> UPCGAttributeCastSettings::GetPreconfiguredInfo() const
 {
-	return PCGMetadataElementCommon::FillPreconfiguredSettingsInfoFromEnum<EPCGMetadataTypes>({ EPCGMetadataTypes::Count, EPCGMetadataTypes::Unknown }, LOCTEXT("PreconfigureInfo", "Cast to "));
+	return FPCGPreConfiguredSettingsInfo::PopulateFromEnum<EPCGMetadataTypes>({ EPCGMetadataTypes::Count, EPCGMetadataTypes::Unknown }, LOCTEXT("PreconfigureInfo", "Cast to "));
 }
 #endif
 
diff --git a/Engine/Plugins/PCG/Source/PCG/Private/Elements/PCGAttributeReduceElement.cpp b/Engine/Plugins/PCG/Source/PCG/Private/Elements/PCGAttributeReduceElement.cpp
--- a/Engine/Plugins/PCG/Source/PCG/Private/Elements/PCGAttributeReduceElement.cpp
+++ b/Engine/Plugins/PCG/Source/PCG/Private/Elements/PCGAttributeReduceElement.cpp
@@ -182,7 +182,7 @@ void UPCGAttributeReduceSettings::ApplyDeprecation(UPCGNode* InOutNode)
 
 TArray<FPCGPreConfiguredSettingsInfo> UPCGAttributeReduceSettings::GetPreconfiguredInfo() const
 {
-	return PCGMetadataElementCommon::FillPreconfiguredSettingsInfoFromEnum<EPCGAttributeReduceOperation>();
+	return FPCGPreConfiguredSettingsInfo::PopulateFromEnum<EPCGAttributeReduceOperation>();
 }
 #endif
 
diff --git a/Engine/Plugins/PCG/Source/PCG/Private/Utils/PCGLogErrors.cpp b/Engine/Plugins/PCG/Source/PCG/Private/Utils/PCGLogErrors.cpp
--- a/Engine/Plugins/PCG/Source/PCG/Private/Utils/PCGLogErrors.cpp
+++ b/Engine/Plugins/PCG/Source/PCG/Private/Utils/PCGLogErrors.cpp
@@ -121,4 +121,18 @@ namespace PCGLog::Component
 	}
 }
 
+namespace PCGLog::Settings
+{
+	namespace Format
+	{
+		const FTextFormat InvalidPreconfiguration = LOCTEXT("InvalidPreconfiguration", "Invalid preconfiguration index '{0}' for node settings '{1}'. Default settings will be used.");
+	}
+
+	// Warnings
+	void LogInvalidPreconfigurationWarning(const int32 PreconfigurationIndex, const FText& NodeTitle)
+	{
+		LogWarningOnGraph(FText::Format(Format::InvalidPreconfiguration, FText::AsNumber(PreconfigurationIndex), NodeTitle));
+	}
+}
+
 #undef LOCTEXT_NAMESPACE
diff --git a/Engine/Plugins/PCG/Source/PCG/Public/Elements/Metadata/PCGMetadataElementCommon.h b/Engine/Plugins/PCG/Source/PCG/Public/Elements/Metadata/PCGMetadataElementCommon.h
--- a/Engine/Plugins/PCG/Source/PCG/Public/Elements/Metadata/PCGMetadataElementCommon.h
+++ b/Engine/Plugins/PCG/Source/PCG/Public/Elements/Metadata/PCGMetadataElementCommon.h
@@ -7,6 +7,7 @@
 #include "Metadata/PCGMetadataAttributeTraits.h"
 #include "Metadata/Accessors/IPCGAttributeAccessor.h" // IWYU pragma: keep
 #include "Metadata/Accessors/PCGAttributeAccessorKeys.h"
+#include "Utils/PCGPreconfiguration.h"
 
 #include "Misc/ScopeExit.h" // For ON_SCOPE_EXIT
 #include "Templates/MemoryOps.h" // For DestructItem
@@ -262,29 +263,11 @@ namespace PCGMetadataElementCommon
 
 	PCG_API bool CopyFromAccessorToAccessor(FCopyFromAccessorToAccessorParams& Params);
 
-	/**
-	* Automatically fill all preconfigured settings depending on the enum operation.
-	* Can also specify explicitly values that should not be included, as metadata is not available in non-editor builds.
-	*/
+
 	template <typename EnumOperation>
+	UE_DEPRECATED(5.6, "Please use FPCGPreConfiguredSettingsInfo::PopulateFromEnum instead.")
 	TArray<FPCGPreConfiguredSettingsInfo> FillPreconfiguredSettingsInfoFromEnum(const TSet<EnumOperation>& InValuesToSkip = {}, const FText& InOptionalPrefix = FText())
 	{
-		TArray<FPCGPreConfiguredSettingsInfo> PreconfiguredInfo;
-
-		if (const UEnum* EnumPtr = StaticEnum<EnumOperation>())
-		{
-			PreconfiguredInfo.Reserve(EnumPtr->NumEnums());
-			for (int32 i = 0; i < EnumPtr->NumEnums(); ++i)
-			{
-				int64 Value = EnumPtr->GetValueByIndex(i);
-
-				if (Value != EnumPtr->GetMaxEnumValue() && !InValuesToSkip.Contains(EnumOperation(Value)))
-				{
-					PreconfiguredInfo.Emplace(Value, FText::Format(FText::FromString(TEXT("{0}{1}")), InOptionalPrefix, EnumPtr->GetDisplayNameTextByValue(Value)));
-				}
-			}
-		}
-
-		return PreconfiguredInfo;
+		return FPCGPreConfiguredSettingsInfo::PopulateFromEnum<EnumOperation>(InValuesToSkip, InOptionalPrefix);
 	}
 }
diff --git a/Engine/Plugins/PCG/Source/PCG/Public/PCGSettings.h b/Engine/Plugins/PCG/Source/PCG/Public/PCGSettings.h
--- a/Engine/Plugins/PCG/Source/PCG/Public/PCGSettings.h
+++ b/Engine/Plugins/PCG/Source/PCG/Public/PCGSettings.h
@@ -7,6 +7,7 @@
 #include "PCGElement.h"
 #include "PCGPin.h"
 #include "Elements/PCGActorSelector.h"
+#include "Utils/PCGPreconfiguration.h"
 #include "Tests/Determinism/PCGDeterminismSettings.h"
 
 #include "Algo/AllOf.h"
@@ -128,45 +129,6 @@ struct FPCGSettingsOverridableParam
 #endif // WITH_EDITOR
 };
 
-/**
-* Pre-configured settings info
-* Will be passed to the settings to pre-configure the settings on creation.
-* Example: Maths operations: Add, Mul, etc...
-*/
-USTRUCT(BlueprintType)
-struct FPCGPreConfiguredSettingsInfo
-{
-	GENERATED_BODY()
-
-	FPCGPreConfiguredSettingsInfo() = default;
-
-	explicit FPCGPreConfiguredSettingsInfo(int32 InIndex, FText InLabel = FText{})
-		: PreconfiguredIndex(InIndex)
-		, Label(std::move(InLabel))
-	{}
-
-#if WITH_EDITOR
-	FPCGPreConfiguredSettingsInfo(int32 InIndex, FText InLabel, FText InTooltip)
-		: PreconfiguredIndex(InIndex)
-		, Label(std::move(InLabel))
-		, Tooltip(std::move(InTooltip))
-	{}
-#endif // WITH_EDITOR
-
-	/* Index used by the settings to know which preconfigured settings it needs to set. */
-	UPROPERTY(BlueprintReadWrite, EditAnywhere, Category = "")
-	int32 PreconfiguredIndex = -1;
-
-	/* Label for the exposed asset. Can also be used instead of the index, if it is easier to deal with strings. */
-	UPROPERTY(BlueprintReadWrite, EditAnywhere, Category = "")
-	FText Label;
-
-#if WITH_EDITORONLY_DATA
-	UPROPERTY(EditAnywhere, Category = "")
-	FText Tooltip;
-#endif // WITH_EDITORONLY_DATA
-};
-
 UCLASS(Abstract)
 class PCG_API UPCGSettingsInterface : public UPCGData
 {
diff --git a/Engine/Plugins/PCG/Source/PCG/Public/Utils/PCGLogErrors.h b/Engine/Plugins/PCG/Source/PCG/Public/Utils/PCGLogErrors.h
--- a/Engine/Plugins/PCG/Source/PCG/Public/Utils/PCGLogErrors.h
+++ b/Engine/Plugins/PCG/Source/PCG/Public/Utils/PCGLogErrors.h
@@ -129,4 +129,9 @@ namespace PCGLog
 	{
 		PCG_API void LogComponentAttachmentFailedWarning(const FPCGContext* InContext = nullptr);
 	}
+
+	namespace Settings
+	{
+		PCG_API void LogInvalidPreconfigurationWarning(const int32 PreconfigurationIndex, const FText& NodeTitle);
+	}
 }
diff --git a/Engine/Plugins/PCG/Source/PCG/Public/Utils/PCGPreconfiguration.h b/Engine/Plugins/PCG/Source/PCG/Public/Utils/PCGPreconfiguration.h
new file mode 100644
--- /dev/null
+++ b/Engine/Plugins/PCG/Source/PCG/Public/Utils/PCGPreconfiguration.h
@@ -0,0 +1,74 @@
+// Copyright Epic Games, Inc. All Rights Reserved.
+
+#pragma once
+
+#include "PCGPreconfiguration.generated.h"
+
+/**
+* Pre-configured settings info
+* Will be passed to the settings to pre-configure the settings on creation.
+* Example: Maths operations: Add, Mul, Div etc...
+*/
+USTRUCT(BlueprintType)
+struct PCG_API FPCGPreConfiguredSettingsInfo
+{
+	GENERATED_BODY()
+
+	FPCGPreConfiguredSettingsInfo() = default;
+
+	explicit FPCGPreConfiguredSettingsInfo(int32 InIndex, FText InLabel = FText{})
+		: PreconfiguredIndex(InIndex)
+		, Label(std::move(InLabel))
+	{}
+
+#if WITH_EDITOR
+	FPCGPreConfiguredSettingsInfo(int32 InIndex, FText InLabel, FText InTooltip)
+		: PreconfiguredIndex(InIndex)
+		, Label(std::move(InLabel))
+		, Tooltip(std::move(InTooltip))
+	{}
+#endif // WITH_EDITOR
+
+	/**
+	* Automatically fill all preconfigured settings depending on the provided enum.
+	* Can also specify explicitly values that should not be included, i.e. cases that may not be available in non-editor builds.
+	*/
+	template <typename EnumType>
+	static TArray<FPCGPreConfiguredSettingsInfo> PopulateFromEnum(const TSet<EnumType>& InValuesToSkip = {}, const FText& InOptionalPrefix = FText())
+	{
+		TArray<FPCGPreConfiguredSettingsInfo> PreconfiguredInfo;
+
+		if (const UEnum* EnumPtr = StaticEnum<EnumType>())
+		{
+			PreconfiguredInfo.Reserve(EnumPtr->NumEnums());
+			for (int32 i = 0; i < EnumPtr->NumEnums(); ++i)
+			{
+				int64 Value = EnumPtr->GetValueByIndex(i);
+
+				if (Value != EnumPtr->GetMaxEnumValue() && !InValuesToSkip.Contains(EnumType(Value)))
+				{
+					FText DisplayName = EnumPtr->GetDisplayNameTextByValue(Value);
+					if (!DisplayName.IsEmpty())
+					{
+						PreconfiguredInfo.Emplace(Value, FText::Format(FText::FromString(TEXT("{0}{1}")), InOptionalPrefix, std::move(DisplayName)));
+					}
+				}
+			}
+		}
+
+		return PreconfiguredInfo;
+	}
+
+	/* Index used by the settings to know which preconfigured settings it needs to set. */
+	UPROPERTY(BlueprintReadWrite, EditAnywhere, Category = "")
+	int32 PreconfiguredIndex = -1;
+
+	/* Label for the exposed asset. Can also be used instead of the index, if it is easier to deal with strings. */
+	UPROPERTY(BlueprintReadWrite, EditAnywhere, Category = "")
+	FText Label;
+
+#if WITH_EDITORONLY_DATA
+	UPROPERTY(EditAnywhere, Category = "")
+	FText Tooltip;
+#endif // WITH_EDITORONLY_DATA
+};
\ No newline at end of file
