[PCG] Added shape sweep option to World Raycast node

TODO: World Ray Hit Query

diff --git a/Engine/Plugins/PCG/Source/PCG/Private/Data/PCGCollisionShape.cpp b/Engine/Plugins/PCG/Source/PCG/Private/Data/PCGCollisionShape.cpp
new file mode 100644
--- /dev/null
+++ b/Engine/Plugins/PCG/Source/PCG/Private/Data/PCGCollisionShape.cpp
@@ -0,0 +1,66 @@
+// Copyright Epic Games, Inc. All Rights Reserved.
+
+#include "Data/PCGCollisionShape.h"
+#include "Utils/PCGLogErrors.h"
+
+#include UE_INLINE_GENERATED_CPP_BY_NAME(PCGCollisionShape)
+
+#define LOCTEXT_NAMESPACE "PCGCollisionShape"
+
+FPCGCollisionShape::FPCGCollisionShape(const FCollisionShape InShape, const FPCGContext* InOptionalContext)
+{
+	switch (InShape.ShapeType)
+	{
+		case ECollisionShape::Line:
+			break;
+		case ECollisionShape::Box:
+			ShapeType = EPCGCollisionShapeType::Box;
+			BoxHalfExtent = InShape.GetBox();
+			break;
+		case ECollisionShape::Sphere:
+			ShapeType = EPCGCollisionShapeType::Sphere;
+			SphereRadius = InShape.GetSphereRadius();
+			break;
+		case ECollisionShape::Capsule:
+			ShapeType = EPCGCollisionShapeType::Capsule;
+			CapsuleRadius = InShape.GetCapsuleRadius();
+			CapsuleHalfHeight = InShape.GetCapsuleHalfHeight();
+			break;
+		default:
+			PCGLog::LogErrorOnGraph(LOCTEXT("InvalidPCGCollisionShape", "Invalid FCollisionShape shape type."), InOptionalContext);
+	}
+}
+
+FCollisionShape FPCGCollisionShape::ToCollisionShape(const FPCGContext* InOptionalContext) const
+{
+	FCollisionShape Shape;
+
+	// Recreated the SetShape function to use delineated extents from the settings.
+	switch (ShapeType)
+	{
+		case EPCGCollisionShapeType::Line:
+			break;
+		case EPCGCollisionShapeType::Box:
+			Shape.SetBox(FVector3f(BoxHalfExtent));
+			break;
+		case EPCGCollisionShapeType::Sphere:
+			Shape.SetSphere(SphereRadius);
+			break;
+		case EPCGCollisionShapeType::Capsule:
+			Shape.SetCapsule(CapsuleRadius, CapsuleHalfHeight);
+			break;
+		default:
+			PCGLog::LogErrorOnGraph(
+					FText::Format(
+							LOCTEXT("InvalidPCGCollisionShape", "Invalid EPCGCollisionShapeType '{0}'."),
+							StaticEnum<EPCGCollisionShapeType>()
+								? StaticEnum<EPCGCollisionShapeType>()->GetDisplayNameTextByIndex(static_cast<int64>(ShapeType))
+								: LOCTEXT("UnknownEnum", "Unknown")),
+					InOptionalContext);
+			break;
+	}
+
+	return Shape;
+}
+
+#undef LOCTEXT_NAMESPACE
diff --git a/Engine/Plugins/PCG/Source/PCG/Private/Elements/PCGWorldRaycast.cpp b/Engine/Plugins/PCG/Source/PCG/Private/Elements/PCGWorldRaycast.cpp
--- a/Engine/Plugins/PCG/Source/PCG/Private/Elements/PCGWorldRaycast.cpp
+++ b/Engine/Plugins/PCG/Source/PCG/Private/Elements/PCGWorldRaycast.cpp
@@ -19,11 +19,37 @@
 
 #define LOCTEXT_NAMESPACE "PCGWorldRaycastElement"
 
-namespace PCGWorldRaycastElementConstants
+namespace PCGWorldRaycastElement::Constants
 {
-	const FName OriginsInputLabel = TEXT("Origins");
-	const FName EndPointsInputLabel = TEXT("End Points");
-	const FName BoundingShapeLabel = TEXT("Bounding Shape");
+	const FText DefaultNodeTitle = LOCTEXT("NodeTitle", "World Raycast");
+
+	namespace Labels
+	{
+		const FName OriginsInput = TEXT("Origins");
+		const FName EndPointsInput = TEXT("End Points");
+		const FName BoundingShape = TEXT("Bounding Shape");
+
+		const FText LineTrace = LOCTEXT("LineTraceLabel", "Line Trace");
+		const FText BoxSweep = LOCTEXT("BoxSweepLabel", "Box Sweep");
+		const FText SphereSweep = LOCTEXT("SphereSweepLabel", "Sphere Sweep");
+		const FText CapsuleSweep = LOCTEXT("CapsuleSweepLabel", "Capsule Sweep");
+	}
+
+	namespace Tooltips
+	{
+		const FText LineTrace = LOCTEXT("LineTraceTooltip", "Conduct a trace along a given ray.");
+		const FText BoxSweep = LOCTEXT("BoxSweepTooltip", "Casts a box sweep along a given ray.");
+		const FText SphereSweep = LOCTEXT("SphereSweepTooltip", "Casts a sphere sweep along a given ray.");
+		const FText CapsuleSweep = LOCTEXT("CapsuleSweepTooltip", "Casts a capsule sweep along a given ray.");
+	}
+
+	namespace Preconfiguration
+	{
+		constexpr int32 LineTraceIndex = 0;
+		constexpr int32 BoxSweepIndex = 1;
+		constexpr int32 SphereSweepIndex = 2;
+		constexpr int32 CapsuleSweepIndex = 3;
+	}
 }
 
 UPCGWorldRaycastElementSettings::UPCGWorldRaycastElementSettings()
@@ -34,11 +60,127 @@ UPCGWorldRaycastElementSettings::UPCGWorldRaycastElementSettings()
 	RayDirectionAttribute.Update("$Rotation.Up");
 }
 
+#if WITH_EDITOR
+FText UPCGWorldRaycastElementSettings::GetDefaultNodeTitle() const
+{
+	return PCGWorldRaycastElement::Constants::DefaultNodeTitle;
+}
+
+TArray<FText> UPCGWorldRaycastElementSettings::GetNodeTitleAliases() const
+{
+	return {LOCTEXT("WorldTraceAlias", "World Trace"), LOCTEXT("WorldSweepAlias", "World Sweep")};
+}
+
+TArray<FPCGPreConfiguredSettingsInfo> UPCGWorldRaycastElementSettings::GetPreconfiguredInfo() const
+{
+	using namespace PCGWorldRaycastElement::Constants;
+
+	return {
+		{0, Labels::LineTrace, Tooltips::LineTrace},
+		{1, Labels::BoxSweep, Tooltips::BoxSweep},
+		{2, Labels::SphereSweep, Tooltips::SphereSweep},
+		{3, Labels::CapsuleSweep, Tooltips::CapsuleSweep}
+	};
+}
+
+EPCGChangeType UPCGWorldRaycastElementSettings::GetChangeTypeForProperty(FPropertyChangedEvent& PropertyChangedEvent) const
+{
+	EPCGChangeType Result = Super::GetChangeTypeForProperty(PropertyChangedEvent);
+
+	const FName PropertyName = PropertyChangedEvent.GetPropertyName();
+	const FName MemberPropertyName = PropertyChangedEvent.GetMemberPropertyName();
+
+	if (MemberPropertyName == GET_MEMBER_NAME_CHECKED(UPCGWorldRaycastElementSettings, WorldQueryParams)
+		&& PropertyName == GET_MEMBER_NAME_CHECKED(FPCGWorldCommonQueryParams, ActorFilterFromInput))
+	{
+		// This can add/remove a pin, so we need a structural change
+		Result |= EPCGChangeType::Structural;
+	}
+
+	if ((MemberPropertyName == GET_MEMBER_NAME_CHECKED(UPCGWorldRaycastElementSettings, CollisionShape)
+		&& PropertyName == GET_MEMBER_NAME_CHECKED(FPCGCollisionShape, ShapeType))
+		|| PropertyName == GET_MEMBER_NAME_CHECKED(UPCGWorldRaycastElementSettings, RaycastMode))
+	{
+		// Changes the additional title info
+		Result |= EPCGChangeType::Cosmetic;
+	}
+
+	return Result;
+}
+#endif // WITH_EDITOR
+
+FString UPCGWorldRaycastElementSettings::GetAdditionalTitleInformation() const
+{
+	using namespace PCGWorldRaycastElement::Constants;
+
+	FText FinalLabel;
+
+	if (IsPropertyOverriddenByPin({GET_MEMBER_NAME_CHECKED(UPCGWorldRaycastElementSettings, CollisionShape), GET_MEMBER_NAME_CHECKED(FPCGCollisionShape, ShapeType)}))
+	{
+		FinalLabel = LOCTEXT("ShapeTypeOverridden", "Overridden");
+	}
+	else
+	{
+		switch (CollisionShape.ShapeType)
+		{
+			case EPCGCollisionShapeType::Line:
+				FinalLabel = Labels::LineTrace;
+				break;
+			case EPCGCollisionShapeType::Box:
+				FinalLabel = Labels::BoxSweep;
+				break;
+			case EPCGCollisionShapeType::Sphere:
+				FinalLabel = Labels::SphereSweep;
+				break;
+			case EPCGCollisionShapeType::Capsule:
+				FinalLabel = Labels::CapsuleSweep;
+				break;
+			default:
+				PCGLog::LogErrorOnGraph(LOCTEXT("InvalidCollisionShapeType", "World Raycast - Invalid Collision Shape Type"));
+		}
+	}
+
+	if (!IsPropertyOverriddenByPin(GET_MEMBER_NAME_CHECKED(UPCGWorldRaycastElementSettings, RaycastMode)))
+	{
+		if (const UEnum* EnumPtr = StaticEnum<EPCGWorldRaycastMode>())
+		{
+			FinalLabel = FText::Format(INVTEXT("{0} ({1})"), FinalLabel, EnumPtr->GetDisplayNameTextByValue(static_cast<int64>(RaycastMode)));
+		}
+	}
+
+	return FinalLabel.ToString();
+}
+
+void UPCGWorldRaycastElementSettings::ApplyPreconfiguredSettings(const FPCGPreConfiguredSettingsInfo& PreconfigureInfo)
+{
+	Super::ApplyPreconfiguredSettings(PreconfigureInfo);
+
+	using namespace PCGWorldRaycastElement::Constants;
+	switch (PreconfigureInfo.PreconfiguredIndex)
+	{
+		case Preconfiguration::LineTraceIndex:
+			CollisionShape.ShapeType = EPCGCollisionShapeType::Line;
+			break;
+		case Preconfiguration::BoxSweepIndex:
+			CollisionShape.ShapeType = EPCGCollisionShapeType::Box;
+			break;
+		case Preconfiguration::SphereSweepIndex:
+			CollisionShape.ShapeType = EPCGCollisionShapeType::Sphere;
+			break;
+		case Preconfiguration::CapsuleSweepIndex:
+			CollisionShape.ShapeType = EPCGCollisionShapeType::Capsule;
+			break;
+		default:
+			PCGLog::Settings::LogInvalidPreconfigurationWarning(PreconfigureInfo.PreconfiguredIndex, DefaultNodeTitle);
+	}
+}
+
 TArray<FPCGPinProperties> UPCGWorldRaycastElementSettings::InputPinProperties() const
 {
 	TArray<FPCGPinProperties> PinProperties;
 
-	FPCGPinProperties& OriginPointsInputPin = PinProperties.Emplace_GetRef(PCGWorldRaycastElementConstants::OriginsInputLabel, EPCGDataType::PointOrParam);
+	using namespace PCGWorldRaycastElement::Constants;
+	FPCGPinProperties& OriginPointsInputPin = PinProperties.Emplace_GetRef(Labels::OriginsInput, EPCGDataType::PointOrParam);
 	OriginPointsInputPin.SetRequiredPin();
 #if WITH_EDITOR
 	OriginPointsInputPin.Tooltip = LOCTEXT("OriginPointsInputPinTooltip", "Points acting as the point origin of each individual ray. Only point data from this pin may be forwarded.");
@@ -46,13 +188,13 @@ TArray<FPCGPinProperties> UPCGWorldRaycastElementSettings::InputPinProperties()
 
 	if (RaycastMode == EPCGWorldRaycastMode::Segments)
 	{
-		FPCGPinProperties& EndPointsInputPin = PinProperties.Emplace_GetRef(PCGWorldRaycastElementConstants::EndPointsInputLabel, EPCGDataType::PointOrParam);
+		FPCGPinProperties& EndPointsInputPin = PinProperties.Emplace_GetRef(Labels::EndPointsInput, EPCGDataType::PointOrParam);
 #if WITH_EDITOR
 		EndPointsInputPin.Tooltip = LOCTEXT("EndPointsInputPinTooltip", "Points acting as the end point of each individual ray.");
 #endif // WITH_EDITOR
 	}
 
-	FPCGPinProperties& BoundingShapeInputPin = PinProperties.Emplace_GetRef(PCGWorldRaycastElementConstants::BoundingShapeLabel, EPCGDataType::Spatial);
+	FPCGPinProperties& BoundingShapeInputPin = PinProperties.Emplace_GetRef(Labels::BoundingShape, EPCGDataType::Spatial);
 #if WITH_EDITOR
 	BoundingShapeInputPin.Tooltip = LOCTEXT("BoundingShapeInputPinTooltip", "All projected points must be contained within this shape. If this input is omitted then bounds will be taken from the actor so that points are contained within actor bounds.");
 #endif // WITH_EDITOR
@@ -72,22 +214,6 @@ FPCGElementPtr UPCGWorldRaycastElementSettings::CreateElement() const
 	return MakeShared<FPCGWorldRaycastElement>();
 }
 
-#if WITH_EDITOR
-EPCGChangeType UPCGWorldRaycastElementSettings::GetChangeTypeForProperty(FPropertyChangedEvent& PropertyChangedEvent) const
-{
-	EPCGChangeType Result = Super::GetChangeTypeForProperty(PropertyChangedEvent);
-
-	if (PropertyChangedEvent.GetPropertyName() == GET_MEMBER_NAME_CHECKED(UPCGWorldRaycastElementSettings, WorldQueryParams) &&
-		PropertyChangedEvent.GetMemberPropertyName() == GET_MEMBER_NAME_CHECKED(FPCGWorldCommonQueryParams, ActorFilterFromInput))
-	{
-		// This can add/remove a pin, so we need a structural change
-		Result |= EPCGChangeType::Structural;
-	}
-
-	return Result;
-}
-#endif // WITH_EDITOR
-
 bool FPCGWorldRaycastElement::PrepareDataInternal(FPCGContext* InContext) const
 {
 	ContextType* Context = static_cast<ContextType*>(InContext);
@@ -95,8 +221,9 @@ bool FPCGWorldRaycastElement::PrepareDataInternal(FPCGContext* InContext) const
 	const UPCGWorldRaycastElementSettings* Settings = Context->GetInputSettings<UPCGWorldRaycastElementSettings>();
 	check(Settings);
 
-	const TArray<FPCGTaggedData> OriginsInputData = Context->InputData.GetInputsByPin(PCGWorldRaycastElementConstants::OriginsInputLabel);
-	const TArray<FPCGTaggedData> EndPointsInputData = Context->InputData.GetInputsByPin(PCGWorldRaycastElementConstants::EndPointsInputLabel);
+	using namespace PCGWorldRaycastElement::Constants;
+	const TArray<FPCGTaggedData> OriginsInputData = Context->InputData.GetInputsByPin(Labels::OriginsInput);
+	const TArray<FPCGTaggedData> EndPointsInputData = Context->InputData.GetInputsByPin(Labels::EndPointsInput);
 	const TArray<FPCGTaggedData> FilterActorInputData = Context->InputData.GetInputsByPin(PCGWorldRayHitConstants::FilterActorPinLabel);
 
 	const EPCGTimeSliceInitResult ExecutionInitResult = Context->InitializePerExecutionState([&OriginsInputData, &EndPointsInputData, &FilterActorInputData, Settings](ContextType* Context, ExecStateType& OutState)
@@ -110,7 +237,7 @@ bool FPCGWorldRaycastElement::PrepareDataInternal(FPCGContext* InContext) const
 		if (!Settings->bUnbounded)
 		{
 			bool bUnionWasCreated;
-			const UPCGSpatialData* BoundingShape = PCGSettingsHelpers::ComputeBoundingShape(Context, PCGWorldRaycastElementConstants::BoundingShapeLabel, bUnionWasCreated);
+			const UPCGSpatialData* BoundingShape = PCGSettingsHelpers::ComputeBoundingShape(Context, Labels::BoundingShape, bUnionWasCreated);
 			OutState.Bounds = BoundingShape ? BoundingShape->GetBounds() : FBox(EForceInit::ForceInit);
 			if (!OutState.Bounds.IsValid)
 			{
@@ -124,7 +251,7 @@ bool FPCGWorldRaycastElement::PrepareDataInternal(FPCGContext* InContext) const
 			PCGLog::LogErrorOnGraph(LOCTEXT("InvalidEndPointDataCount", "End point input data count must match the Origins input data count."), Context);
 			return EPCGTimeSliceInitResult::NoOperation;
 		}
-		
+
 		if (FilterActorInputData.Num() > 1 && FilterActorInputData.Num() != OriginsInputData.Num())
 		{
 			PCGLog::LogErrorOnGraph(LOCTEXT("InvalidFilterActorInputCount", "Filter Actor input data count must be 1 or match the Origins input data count."), Context);
@@ -152,7 +279,7 @@ bool FPCGWorldRaycastElement::PrepareDataInternal(FPCGContext* InContext) const
 		// Accept only point or param data
 		if (!OriginsData || (!OriginsData->IsA<UPCGPointData>() && !OriginsData->IsA<UPCGParamData>()))
 		{
-			PCGLog::InputOutput::LogTypedDataNotFoundWarning(EPCGDataType::PointOrParam, PCGWorldRaycastElementConstants::OriginsInputLabel, Context);
+			PCGLog::InputOutput::LogTypedDataNotFoundWarning(EPCGDataType::PointOrParam, Labels::OriginsInput, Context);
 			return EPCGTimeSliceInitResult::NoOperation;
 		}
 
@@ -288,7 +415,7 @@ bool FPCGWorldRaycastElement::ExecuteInternal(FPCGContext* InContext) const
 		UWorld* World = Context->SourceComponent->GetWorld();
 		check(World);
 
-		const TArray<FPCGTaggedData> OriginInputs = Context->InputData.GetInputsByPin(PCGWorldRaycastElementConstants::OriginsInputLabel);
+		const TArray<FPCGTaggedData> OriginInputs = Context->InputData.GetInputsByPin(PCGWorldRaycastElement::Constants::Labels::OriginsInput);
 		const FPCGTaggedData& OriginInputData = OriginInputs[IterIndex];
 		const UPCGPointData* OriginInputPointData = Cast<UPCGPointData>(OriginInputData.Data);
 		const TArray<FPCGPoint>* OriginInputPoints = OriginInputPointData ? &OriginInputPointData->GetPoints() : nullptr;
@@ -343,8 +470,16 @@ bool FPCGWorldRaycastElement::ExecuteInternal(FPCGContext* InContext) const
 				continue;
 			}
 
+			// The physics system will automatically follow the line trace path, if the collision shape is line or has no volume.
 			TArray<FHitResult> Hits;
-			World->LineTraceMultiByObjectType(Hits, Origin, Origin + RayVector, ExecState.CollisionObjectQueryParams, ExecState.CollisionQueryParams);
+			World->SweepMultiByObjectType(
+				Hits,
+				Origin,
+				Origin + RayVector,
+				FQuat(Settings->CollisionShape.ShapeRotation),
+				ExecState.CollisionObjectQueryParams,
+				Settings->CollisionShape.ToCollisionShape(),
+				ExecState.CollisionQueryParams);
 			const TOptional<FHitResult>& HitResult = PCGWorldQueryHelpers::FilterRayHitResults(&WorldQueryParams, Context->SourceComponent, Hits, IterState.CachedFilterActors);
 
 			// If there's a hit, and it's within bounds, keep the point
diff --git a/Engine/Plugins/PCG/Source/PCG/Public/Data/PCGCollisionShape.h b/Engine/Plugins/PCG/Source/PCG/Public/Data/PCGCollisionShape.h
new file mode 100644
--- /dev/null
+++ b/Engine/Plugins/PCG/Source/PCG/Public/Data/PCGCollisionShape.h
@@ -0,0 +1,54 @@
+// Copyright Epic Games, Inc. All Rights Reserved.
+
+#pragma once
+
+#include "CollisionShape.h"
+
+#include "PCGCollisionShape.generated.h"
+
+struct FPCGContext;
+
+UENUM(BlueprintType)
+enum class EPCGCollisionShapeType : uint8
+{
+	Line,
+	Box,
+	Sphere,
+	Capsule
+};
+
+/** Parameters for conducting a sweep with a specified shape against the physical world. */
+USTRUCT(BlueprintType)
+struct PCG_API FPCGCollisionShape
+{
+	GENERATED_BODY()
+
+	FPCGCollisionShape() = default;
+	explicit FPCGCollisionShape(FCollisionShape InShape, const FPCGContext* InOptionalContext = nullptr);
+
+	/** Shape that will be used in the collision detection. */
+	UPROPERTY(BlueprintType, EditAnywhere, Category = Settings, meta = (PCG_Overridable))
+	EPCGCollisionShapeType ShapeType = EPCGCollisionShapeType::Line;
+
+	/** Half the size of the sweep box's X, Y, Z components in world space. */
+	UPROPERTY(BlueprintType, EditAnywhere, Category = Settings, DisplayName = "Half Extent (Box)", meta = (EditCondition = "ShapeType == EPCGCollisionShapeType::Box", EditConditionHides, PCG_Overridable))
+	FVector BoxHalfExtent = FVector::OneVector;
+
+	/** Radius of the sweep's sphere in world space. */
+	UPROPERTY(BlueprintType, EditAnywhere, Category = Settings, DisplayName = "Radius (Sphere)", meta = (EditCondition = "ShapeType == EPCGCollisionShapeType::Sphere", EditConditionHides, PCG_Overridable))
+	float SphereRadius = 100.f;
+
+	/** Radius of the spherical shape of the sweep's capsule in world space. */
+	UPROPERTY(BlueprintType, EditAnywhere, Category = Settings, DisplayName = "Radius (Capsule)", meta = (EditCondition = "ShapeType == EPCGCollisionShapeType::Capsule", EditConditionHides, PCG_Overridable))
+	float CapsuleRadius = 100.f;
+
+	/** Half the length of the sweep's capsule in world space. */
+	UPROPERTY(BlueprintType, EditAnywhere, Category = Settings, DisplayName = "Half Height (Capsule)", meta = (EditCondition = "ShapeType == EPCGCollisionShapeType::Capsule", EditConditionHides, PCG_Overridable))
+	float CapsuleHalfHeight = 100.f;
+
+	/** World space rotation applied to the collision shape prior to the sweep. */
+	UPROPERTY(BlueprintType, EditAnywhere, Category = Settings, meta = (EditCondition = "ShapeType != EPCGCollisionShapeType::Line", EditConditionHides, PCG_Overridable))
+	FRotator ShapeRotation = FRotator::ZeroRotator;
+
+	FCollisionShape ToCollisionShape(const FPCGContext* InOptionalContext = nullptr) const;
+};
\ No newline at end of file
diff --git a/Engine/Plugins/PCG/Source/PCG/Public/Elements/PCGWorldRaycast.h b/Engine/Plugins/PCG/Source/PCG/Public/Elements/PCGWorldRaycast.h
--- a/Engine/Plugins/PCG/Source/PCG/Public/Elements/PCGWorldRaycast.h
+++ b/Engine/Plugins/PCG/Source/PCG/Public/Elements/PCGWorldRaycast.h
@@ -3,6 +3,7 @@
 
 #include "PCGSettings.h"
 #include "PCGTimeSlicedElementBase.h"
+#include "Data/PCGCollisionShape.h"
 #include "Data/PCGWorldData.h"
 #include "Helpers/PCGWorldQueryHelpers.h"
 
@@ -24,21 +25,30 @@ public:
 	//~Begin UPCGSettings interface
 #if WITH_EDITOR
 	virtual FName GetDefaultNodeName() const override { return FName(TEXT("WorldRaycast")); }
-	virtual FText GetDefaultNodeTitle() const override { return NSLOCTEXT("PCGWorldRaycastElement", "NodeTitle", "World Raycast"); }
-	virtual FText GetNodeTooltipText() const override { return NSLOCTEXT("PCGWorldRaycastElement", "NodeTooltip", "Casts rays from provided points along a given direction and transform points to the impact point."); }
+	virtual FText GetDefaultNodeTitle() const override;
+	virtual FText GetNodeTooltipText() const override { return NSLOCTEXT("PCGWorldRaycastElement", "NodeTooltip", "Casts a line trace or collision shape sweep from provided points along a given direction returning the location of the impact."); }
 	virtual EPCGSettingsType GetType() const override { return EPCGSettingsType::Spatial; }
+	virtual TArray<FText> GetNodeTitleAliases() const override;
+	virtual TArray<FPCGPreConfiguredSettingsInfo> GetPreconfiguredInfo() const override;
+	virtual bool OnlyExposePreconfiguredSettings() const override { return false; }
+	virtual bool GroupPreconfiguredSettings() const override { return true; }
+	virtual EPCGChangeType GetChangeTypeForProperty(FPropertyChangedEvent& PropertyChangedEvent) const override;
 #endif // WITH_EDITOR
 
+	virtual FString GetAdditionalTitleInformation() const override;
+	virtual void ApplyPreconfiguredSettings(const FPCGPreConfiguredSettingsInfo& PreconfigureInfo) override;
+
 protected:
 	virtual TArray<FPCGPinProperties> InputPinProperties() const override;
 	virtual TArray<FPCGPinProperties> OutputPinProperties() const override;
 	virtual FPCGElementPtr CreateElement() const override;
-#if WITH_EDITOR
-	virtual EPCGChangeType GetChangeTypeForProperty(FPropertyChangedEvent& PropertyChangedEvent) const override;
-#endif // WITH_EDITOR
 	//~End UPCGSettings interface
 
 public:
+	/** Parameters for either using a line trace or specifying a collision shape. */
+	UPROPERTY(BlueprintReadWrite, EditAnywhere, Category = Settings, meta = (ShowOnlyInnerProperties, PCG_Overridable))
+	FPCGCollisionShape CollisionShape;
+
 	/** Determines how the ray's direction and distance will be calculated. */
 	UPROPERTY(BlueprintReadWrite, EditAnywhere, Category = Settings)
 	EPCGWorldRaycastMode RaycastMode = EPCGWorldRaycastMode::Infinite;
