[PCG] Fix assert of loading objects on other threads in OverrideSettings

#rnx

diff --git a/Engine/Plugins/PCG/Source/PCG/Private/PCGContext.cpp b/Engine/Plugins/PCG/Source/PCG/Private/PCGContext.cpp
--- a/Engine/Plugins/PCG/Source/PCG/Private/PCGContext.cpp
+++ b/Engine/Plugins/PCG/Source/PCG/Private/PCGContext.cpp
@@ -144,10 +144,9 @@ void FPCGContext::InitializeSettings(bool bSkipPostLoad)
 
 			if (bHasParamConnected)
 			{
-				// If we have an override, we are not on the main thread and there is a hard ref override, we need to make sure that all objects
+				// If we have an override and there is a hard ref override, we need to make sure that all objects
 				// are loaded. If not, we need to schedule the task on the main thread.
-				const bool bIsInGameThread = IsInGameThread();
-				if (!bIsInGameThread && NodeSettings->HasAnyOverridableHardReferences())
+				if (NodeSettings->HasAnyOverridableHardReferences())
 				{
 					for (const FPCGSettingsOverridableParam& Param : OverridableParams)
 					{
@@ -172,7 +171,7 @@ void FPCGContext::InitializeSettings(bool bSkipPostLoad)
 							FGCScopeGuard GCScope;
 							if (!ObjectPath.ResolveObject())
 							{
-								// We have an override value that is not loaded, and we are not on the main thread. We need to schedule the task on the main thread.
+								// We have an override value that is not loaded. We need to schedule the task on the main thread.
 								bOverrideSettingsOnMainThread = true;
 								break;
 							}
@@ -191,7 +190,8 @@ void FPCGContext::InitializeSettings(bool bSkipPostLoad)
 					FGCScopeGuard Scope;
 					SettingsWithOverride = Cast<UPCGSettings>(StaticDuplicateObjectEx(DuplicateParams));
 				}
-			
+
+				const bool bIsInGameThread = IsInGameThread();
 				for (auto& KeyValuePair : CreatedObjects)
 				{
 					if (KeyValuePair.Value)
