[PCG] Fix LOCTEXT conflict in CollisionShape

#rnx

diff --git a/Engine/Plugins/PCG/Source/PCG/Private/Data/PCGCollisionShape.cpp b/Engine/Plugins/PCG/Source/PCG/Private/Data/PCGCollisionShape.cpp
--- a/Engine/Plugins/PCG/Source/PCG/Private/Data/PCGCollisionShape.cpp
+++ b/Engine/Plugins/PCG/Source/PCG/Private/Data/PCGCollisionShape.cpp
@@ -27,7 +27,7 @@ FPCGCollisionShape::FPCGCollisionShape(const FCollisionShape InShape, const FPCG
 			CapsuleHalfHeight = InShape.GetCapsuleHalfHeight();
 			break;
 		default:
-			PCGLog::LogErrorOnGraph(LOCTEXT("InvalidPCGCollisionShape", "Invalid FCollisionShape shape type."), InOptionalContext);
+			PCGLog::LogErrorOnGraph(LOCTEXT("InvalidFCollisionShape", "Invalid FCollisionShape shape type."), InOptionalContext);
 	}
 }
 
@@ -52,7 +52,7 @@ FCollisionShape FPCGCollisionShape::ToCollisionShape(const FPCGContext* InOption
 		default:
 			PCGLog::LogErrorOnGraph(
 					FText::Format(
-							LOCTEXT("InvalidPCGCollisionShape", "Invalid EPCGCollisionShapeType '{0}'."),
+							LOCTEXT("InvalidEPCGCollisionShape", "Invalid EPCGCollisionShapeType '{0}'."),
 							StaticEnum<EPCGCollisionShapeType>()
 								? StaticEnum<EPCGCollisionShapeType>()->GetDisplayNameTextByIndex(static_cast<int64>(ShapeType))
 								: LOCTEXT("UnknownEnum", "Unknown")),
