[PCG] Fix crash on pcg module shutdown

* We were unregistering the accessor factory in the shutdown module,which is too late because classes are already dead.
* Do it in PreExit instead

#rnx

diff --git a/Engine/Plugins/PCG/Source/PCG/Private/PCGModule.cpp b/Engine/Plugins/PCG/Source/PCG/Private/PCGModule.cpp
--- a/Engine/Plugins/PCG/Source/PCG/Private/PCGModule.cpp
+++ b/Engine/Plugins/PCG/Source/PCG/Private/PCGModule.cpp
@@ -7,6 +7,7 @@
 #include "Data/PCGPointData.h"
 
 #include "Interfaces/IPluginManager.h"
+#include "Misc/CoreDelegates.h"
 #include "Modules/ModuleManager.h"
 
 #if WITH_EDITOR
@@ -32,9 +33,9 @@ bool FPCGModule::IsPCGModuleLoaded()
 	return FModuleManager::Get().IsModuleLoaded(TEXT("PCG"));
 }
 
-#if WITH_EDITOR
 void FPCGModule::StartupModule()
 {
+#if WITH_EDITOR
 	PCGDeterminismTests::FNativeTestRegistry::Create();
 
 	RegisterNativeElementDeterminismTests();
@@ -43,23 +44,35 @@ void FPCGModule::StartupModule()
 
 	const FString PluginShaderDir = FPaths::Combine(IPluginManager::Get().FindPlugin(TEXT("PCG"))->GetBaseDir(), TEXT("Shaders"));
 	AddShaderSourceDirectoryMapping(TEXT("/Plugin/PCG"), PluginShaderDir);
+#endif // WITH_EDITOR
 
 	// Registering accessor methods
 	AttributeAccessorFactory.RegisterDefaultMethods();
 	AttributeAccessorFactory.RegisterMethods<UPCGPointData>(UPCGPointData::GetPointAccessorMethods());
+
+	// Register onto the PreExit, because we need the class to be still valid to remove them from the mapping
+	FCoreDelegates::OnPreExit.AddRaw(this, &FPCGModule::PreExit);
 }
 
 void FPCGModule::ShutdownModule()
 {
-	// Unregistering accessor methods
-	AttributeAccessorFactory.UnregisterDefaultMethods();
-	AttributeAccessorFactory.UnregisterMethods<UPCGPointData>();
+	FCoreDelegates::OnPreExit.RemoveAll(this);
 	
+#if WITH_EDITOR
 	DeregisterNativeElementDeterminismTests();
 
 	PCGDeterminismTests::FNativeTestRegistry::Destroy();
+#endif // WITH_EDITOR
 }
 
+void FPCGModule::PreExit()
+{
+	// Unregistering accessor methods
+	AttributeAccessorFactory.UnregisterDefaultMethods();
+	AttributeAccessorFactory.UnregisterMethods<UPCGPointData>();
+}
+
+#if WITH_EDITOR
 void FPCGModule::RegisterNativeElementDeterminismTests()
 {
 	PCGDeterminismTests::FNativeTestRegistry::RegisterTestFunction(UPCGDifferenceSettings::StaticClass(), PCGDeterminismTests::DifferenceElement::RunTestSuite);
@@ -70,7 +83,7 @@ void FPCGModule::DeregisterNativeElementDeterminismTests()
 {
 	PCGDeterminismTests::FNativeTestRegistry::DeregisterTestFunction(UPCGDifferenceSettings::StaticClass());
 }
-#endif
+#endif // WITH_EDITOR
 
 IMPLEMENT_MODULE(FPCGModule, PCG);
 
diff --git a/Engine/Plugins/PCG/Source/PCG/Public/PCGModule.h b/Engine/Plugins/PCG/Source/PCG/Public/PCGModule.h
--- a/Engine/Plugins/PCG/Source/PCG/Public/PCGModule.h
+++ b/Engine/Plugins/PCG/Source/PCG/Public/PCGModule.h
@@ -29,13 +29,13 @@ class FPCGModule final : public IModuleInterface
 {
 public:
 	//~ IModuleInterface implementation
-#if WITH_EDITOR
 	virtual void StartupModule() override;
 	virtual void ShutdownModule() override;
-#endif
 	virtual bool SupportsDynamicReloading() override { return true; }
 	//~ End IModuleInterface implementation
 
+	void PreExit();
+
 	PCG_API static FPCGModule& GetPCGModuleChecked();
 	static const FPCGGetDataFunctionRegistry& ConstGetDataFunctionRegistry() { return GetPCGModuleChecked().GetDataFunctionRegistry; }
 	static FPCGGetDataFunctionRegistry& MutableGetDataFunctionRegistry() { return GetPCGModuleChecked().GetDataFunctionRegistry; }
