[PCG] Fix double duplication now that 'LoadedPreviewResources' is initialized on load, it is possible that it points to the same ManagedResources as 'GeneratedResource' until a new generation is triggered

#rnx

diff --git a/Engine/Plugins/PCG/Source/PCG/Private/PCGComponent.cpp b/Engine/Plugins/PCG/Source/PCG/Private/PCGComponent.cpp
--- a/Engine/Plugins/PCG/Source/PCG/Private/PCGComponent.cpp
+++ b/Engine/Plugins/PCG/Source/PCG/Private/PCGComponent.cpp
@@ -3784,13 +3784,16 @@ void FPCGComponentInstanceData::ApplyToComponent(UActorComponent* Component, con
 			// Non-critical but should be done: transient data, tracked actors cache, landscape tracking
 			// TODO Validate usefulness + move accordingly
 		}
-		
+
 		if (ConstructionSourceComponent)
 		{
 			// Critical: LastGeneratedBounds
 			PCGComponent->LastGeneratedBounds = ConstructionSourceComponent->LastGeneratedBounds;
 
 			// Duplicate generated resources + retarget them
+#if WITH_EDITOR
+			TMap<TObjectPtr<UPCGManagedResource>, TObjectPtr<UPCGManagedResource>> GeneratedResourceMapping;
+#endif
 			TArray<TObjectPtr<UPCGManagedResource>> DuplicatedResources;
 			for (const TObjectPtr<UPCGManagedResource>& Resource : ConstructionSourceComponent->GeneratedResources)
 			{
@@ -3799,6 +3802,10 @@ void FPCGComponentInstanceData::ApplyToComponent(UActorComponent* Component, con
 					UPCGManagedResource* DuplicatedResource = CastChecked<UPCGManagedResource>(StaticDuplicateObject(Resource, PCGComponent, FName()));
 					DuplicatedResource->PostApplyToComponent();
 					DuplicatedResources.Add(DuplicatedResource);
+
+#if WITH_EDITOR
+					GeneratedResourceMapping.Add(Resource, DuplicatedResource);
+#endif
 				}
 			}
 
@@ -3816,8 +3823,16 @@ void FPCGComponentInstanceData::ApplyToComponent(UActorComponent* Component, con
 			{
 				if (Resource)
 				{
-					UPCGManagedResource* DuplicatedResource = CastChecked<UPCGManagedResource>(StaticDuplicateObject(Resource, PCGComponent, FName()));
-					DuplicatedResource->PostApplyToComponent();
+					UPCGManagedResource* DuplicatedResource = nullptr;
+					if (TObjectPtr<UPCGManagedResource>* FoundDuplicate = GeneratedResourceMapping.Find(Resource))
+					{
+						DuplicatedResource = *FoundDuplicate;
+					}
+					else
+					{
+						DuplicatedResource = CastChecked<UPCGManagedResource>(StaticDuplicateObject(Resource, PCGComponent, FName()));
+						DuplicatedResource->PostApplyToComponent();
+					}
 					DuplicateLoadedPreviewResources.Add(DuplicatedResource);
 				}
 			}
