Added log error when spawn actor fails.
Added recovery path for BP assets that aren't cooked but their class is.
Changed behavior of value tags when they are empty so that they behave as strings and not booleans.

diff --git a/Engine/Plugins/PCG/Source/PCG/Private/Elements/PCGSpawnActor.cpp b/Engine/Plugins/PCG/Source/PCG/Private/Elements/PCGSpawnActor.cpp
--- a/Engine/Plugins/PCG/Source/PCG/Private/Elements/PCGSpawnActor.cpp
+++ b/Engine/Plugins/PCG/Source/PCG/Private/Elements/PCGSpawnActor.cpp
@@ -36,9 +36,10 @@ static TAutoConsoleVariable<bool> CVarAllowActorReuse(
 class FPCGSpawnActorPartitionByAttribute : public FPCGDataPartitionBase<FPCGSpawnActorPartitionByAttribute, TSubclassOf<AActor>>
 {
 public:
-	FPCGSpawnActorPartitionByAttribute(FName InSpawnAttribute)
+	FPCGSpawnActorPartitionByAttribute(FName InSpawnAttribute, FPCGContext* InContext)
 		: FPCGDataPartitionBase<FPCGSpawnActorPartitionByAttribute, TSubclassOf<AActor>>()
 		, SpawnAttribute(InSpawnAttribute)
+		, Context(InContext)
 	{
 	}
 
@@ -109,12 +110,21 @@ public:
 				ActorClass = Blueprint->GeneratedClass.Get();
 			}
 		}
+
+		// Finally, we'll try to take the path and check if it matches a class path instead, because the (BP) asset will not be available in packaged builds.
+		if (!ActorClass)
+		{
+			ActorPath = FSoftClassPath(ActorPath.ToString() + TEXT("_C"));
+			ActorClassSoftPtr = TSoftClassPtr<AActor>(ActorPath);
+			ActorClass = ActorClassSoftPtr.LoadSynchronous();
+		}
 		
 		if (ActorClass && ActorClass->IsChildOf<AActor>())
 		{
 			return &ElementMap.FindOrAdd(ActorClass);
 		}
 
+		PCGLog::LogErrorOnGraph(LOCTEXT("InvalidActorClass", "Invalid actor class in Spawn Actor by Attribute"), Context);
 		return nullptr;
 	}
 
@@ -125,6 +135,7 @@ public:
 	TUniquePtr<const IPCGAttributeAccessor> SpawnAttributeAccessor;
 	TUniquePtr<const IPCGAttributeAccessorKeys> SpawnAttributeKeys;
 	FName SpawnAttribute = NAME_None;
+	FPCGContext* Context = nullptr;
 };
 
 UPCGSpawnActorSettings::UPCGSpawnActorSettings(const FObjectInitializer& ObjectInitializer)
@@ -570,7 +581,7 @@ bool FPCGSpawnActorElement::SpawnAndPrepareSubgraphs(FPCGSubgraphContext* Contex
 
 		if (Settings->bSpawnByAttribute && (!bFullySkippedDueToReuse || bGenerateOutputsWithActorReference))
 		{
-			FPCGSpawnActorPartitionByAttribute Selector(Settings->SpawnAttribute);
+			FPCGSpawnActorPartitionByAttribute Selector(Settings->SpawnAttribute, Context);
 			int32 CurrentPointIndex = 0;
 
 			// Selection is still needed if are fully skipped in order to write to the OutPointData.
diff --git a/Engine/Plugins/PCG/Source/PCG/Private/Helpers/PCGTagHelpers.cpp b/Engine/Plugins/PCG/Source/PCG/Private/Helpers/PCGTagHelpers.cpp
--- a/Engine/Plugins/PCG/Source/PCG/Private/Helpers/PCGTagHelpers.cpp
+++ b/Engine/Plugins/PCG/Source/PCG/Private/Helpers/PCGTagHelpers.cpp
@@ -31,11 +31,6 @@ namespace Private
 			{
 				// Tag doesn't have an attribute - ignore
 			}
-			else if (RightSide.IsEmpty())
-			{
-				// Tag doesn't have a value, ignore the right hand side
-				Result.Attribute = MoveTemp(LeftSide);
-			}
 			else
 			{
 				if (RightSide.IsNumeric())
