[PCG] Fix outstanding crashes

* Fix DuplicateData on PCG Volume data, we can't copy the VolumeBodyInstance if it was initialized with a body
* Factorized the body setup into a protected function and call it in the Initialize and the DuplicateData.
* Also moved one line in the attribute list view causing a crash for some data like the Collision Wrapper Data

#rnx

diff --git a/Engine/Plugins/PCG/Source/PCG/Private/Data/PCGVolumeData.cpp b/Engine/Plugins/PCG/Source/PCG/Private/Data/PCGVolumeData.cpp
--- a/Engine/Plugins/PCG/Source/PCG/Private/Data/PCGVolumeData.cpp
+++ b/Engine/Plugins/PCG/Source/PCG/Private/Data/PCGVolumeData.cpp
@@ -53,23 +53,31 @@ void UPCGVolumeData::Initialize(AVolume* InVolume)
 	// TODO: Compute the strict bounds, we must find a FBox inscribed into the oriented box.
 	// Currently, we'll leave the strict bounds empty and fall back to checking against the local box
 	Bounds = FBox::BuildAABB(BoxSphereBounds.Origin, BoxSphereBounds.BoxExtent);
-	
-	// Keep a "sceneless" equivalent body so we can do queries against it without locking constraints
-	if (UBrushComponent* BrushComponent = Volume->GetBrushComponent())
-	{
-		FBodyInstance* BodyInstance = BrushComponent->GetBodyInstance();
-		UBodySetup* BodySetup = BrushComponent->GetBodySetup();
 
-		// In some instances, non-collidable bodies will not be initialized, but it's not an issue for PCG so we can continue regardless.
-		// Otherwise, require that the body is not dynamic.
-		if (BodyInstance && BodySetup && (!FPhysicsInterface::IsValid(BodyInstance->ActorHandle) || !BodyInstance->IsDynamic()))
+	SetupVolumeBodyInstance();
+}
+
+void UPCGVolumeData::SetupVolumeBodyInstance()
+{
+	if (AVolume* CurrentVolume = Volume.Get())
+	{
+		// Keep a "sceneless" equivalent body so we can do queries against it without locking constraints
+		if (UBrushComponent* BrushComponent = CurrentVolume->GetBrushComponent())
 		{
-			ReleaseInternalBodyInstance();
+			FBodyInstance* BodyInstance = BrushComponent->GetBodyInstance();
+			UBodySetup* BodySetup = BrushComponent->GetBodySetup();
+
+			// In some instances, non-collidable bodies will not be initialized, but it's not an issue for PCG so we can continue regardless.
+			// Otherwise, require that the body is not dynamic.
+			if (BodyInstance && BodySetup && (!FPhysicsInterface::IsValid(BodyInstance->ActorHandle) || !BodyInstance->IsDynamic()))
+			{
+				ReleaseInternalBodyInstance();
 
-			VolumeBodyInstance = new FBodyInstance();
-			VolumeBodyInstance->bAutoWeld = false;
-			VolumeBodyInstance->bSimulatePhysics = false;
-			VolumeBodyInstance->InitBody(BodySetup, BrushComponent->GetComponentTransform(), nullptr, nullptr);
+				VolumeBodyInstance = new FBodyInstance();
+				VolumeBodyInstance->bAutoWeld = false;
+				VolumeBodyInstance->bSimulatePhysics = false;
+				VolumeBodyInstance->InitBody(BodySetup, BrushComponent->GetComponentTransform(), nullptr, nullptr);
+			}
 		}
 	}
 }
@@ -197,11 +205,10 @@ void UPCGVolumeData::CopyBaseVolumeData(UPCGVolumeData* NewVolumeData) const
 	NewVolumeData->Bounds = Bounds;
 	NewVolumeData->StrictBounds = StrictBounds;
 
+	// We can't copy the properties if the VolumeBodyInstance is setup for a body. Re-create it using the same method as in Initialize.
 	if (VolumeBodyInstance)
 	{
-		NewVolumeData->ReleaseInternalBodyInstance();
-		NewVolumeData->VolumeBodyInstance = new FBodyInstance();
-		NewVolumeData->VolumeBodyInstance->CopyBodyInstancePropertiesFrom(VolumeBodyInstance);
+		NewVolumeData->SetupVolumeBodyInstance();
 	}
 }
 
diff --git a/Engine/Plugins/PCG/Source/PCG/Public/Data/PCGVolumeData.h b/Engine/Plugins/PCG/Source/PCG/Public/Data/PCGVolumeData.h
--- a/Engine/Plugins/PCG/Source/PCG/Public/Data/PCGVolumeData.h
+++ b/Engine/Plugins/PCG/Source/PCG/Public/Data/PCGVolumeData.h
@@ -46,6 +46,7 @@ public:
 protected:
 	void CopyBaseVolumeData(UPCGVolumeData* NewVolumeData) const;
 	void ReleaseInternalBodyInstance();
+	void SetupVolumeBodyInstance();
 
 	UPROPERTY(BlueprintReadWrite, EditAnywhere, Category = SourceData)
 	TWeakObjectPtr<AVolume> Volume = nullptr;
