[PCG] GPU - lazy readback of GPU data first pass

* Proxy data added to represent data on GPU
* Nodes/elements can opt in to receiving proxies, otherwise the data is read back to CPU in PrepareData. Current opt-ins: Trivial, Gather, Branch, Switch, Subgraph, Input, and Output (if not top graph output)
* Currently grid size node forces readback to CPU but we may want to allow passing proxies across grids (will hold onto graphics memory though, so requires caution). Added cvar to enable passing gpu proxies, seems to work well functionally, but not sure if i'm comfortable to switch it on by default immediately:  pcg.Graph.GPU.PassGPUDataThroughGridLinks

Current limitations / issues / next steps:

* GPU buffer reuse only happens if entire buffer is used. If you Gather data items from multiple GPU buffers, or Filter Data By Index, or anything else which means the downstream node is not getting a complete, single, set of data items from 1 gpu buffer, then it will trigger a readback and then upload. This is an obvious next step, but bindless would help here and would also help grass maps so i will look at bindless next.
* Debug flag does not debug visualize GPU data passing through CPU nodes (it works on GPU nodes though). This will require a refactor to how debug visualization is executed so it can be polled.
* No visual communication of uploads/readbacks which is a big missing piece. Added a cvar which will put warnings in log/on nodes when readbacks happen which is useful, disabled by default: pcg.Graph.GPU.WarnOnGPUReadbacks

#rnx

diff --git a/Engine/Plugins/PCG/Source/PCG/Internal/Compute/Data/PCGProxyForGPUData.h b/Engine/Plugins/PCG/Source/PCG/Internal/Compute/Data/PCGProxyForGPUData.h
new file mode 100644
--- /dev/null
+++ b/Engine/Plugins/PCG/Source/PCG/Internal/Compute/Data/PCGProxyForGPUData.h
@@ -0,0 +1,61 @@
+// Copyright Epic Games, Inc. All Rights Reserved.
+
+#pragma once
+
+#include "PCGData.h"
+
+#include "Compute/PCGDataForGPU.h"
+
+#include "PCGProxyForGPUData.generated.h"
+
+struct FPCGContext;
+
+/** A proxy for data residing on the GPU with functionality to read the data back to the CPU. */
+UCLASS(ClassGroup = (Procedural), DisplayName = "GPU Proxy")
+class PCG_API UPCGProxyForGPUData : public UPCGData
+{
+	GENERATED_BODY()
+
+public:
+	void Initialize(TSharedPtr<FPCGProxyForGPUDataCollection> InDataCollection, int InDataIndexInCollection);
+
+	//~Begin UPCGData interface
+	virtual EPCGDataType GetDataType() const override { return EPCGDataType::ProxyForGPU; }
+	virtual void AddToCrc(FArchiveCrc32& Ar, bool bFullDataCrc) const override;
+	virtual bool CanBeSerialized() const override { return false; }
+	virtual bool IsCacheable() const override { return false; }
+	virtual EPCGDataType GetUnderlyingDataType() const override;
+	//~End UPCGData interface
+
+	/** Returns the GPU info. Returns null if the buffer has been discarded. */
+	TSharedPtr<const FPCGProxyForGPUDataCollection> GetInputDataCollectionInfo() const;
+
+	int GetDataIndexInCollection() const { return DataIndexInCollection; }
+
+	struct FReadbackResult
+	{
+		/** Set false until read back has been performed. */
+		bool bComplete = false;
+
+		/** The data created from readback, null if read back failed. */
+		const UPCGData* Data = nullptr;
+	};
+
+	/** Populates a CPU data object representing the GPU data, performing a readback from GPU->CPU if required. */
+	FReadbackResult GetCPUData(FPCGContext* InContext) const;
+
+	/** Returns the element count for this data. Does not trigger a GPU->CPU readback. */
+	int GetElementCount() const;
+
+	/** Returns a description of this data. Does not trigger a GPU->CPU readback. */
+	bool GetDescription(FPCGDataDesc& OutDescription) const;
+
+protected:
+	TSharedPtr<FPCGProxyForGPUDataCollection> GetGPUInfoValidated() const;
+
+protected:
+	UPROPERTY()
+	int32 DataIndexInCollection = INDEX_NONE;
+
+	TSharedPtr<FPCGProxyForGPUDataCollection> DataCollectionOnGPU;
+};
diff --git a/Engine/Plugins/PCG/Source/PCG/Internal/Compute/PCGComputeCommon.h b/Engine/Plugins/PCG/Source/PCG/Internal/Compute/PCGComputeCommon.h
new file mode 100644
--- /dev/null
+++ b/Engine/Plugins/PCG/Source/PCG/Internal/Compute/PCGComputeCommon.h
@@ -0,0 +1,116 @@
+// Copyright Epic Games, Inc. All Rights Reserved.
+
+#pragma once
+
+#include "PCGCommon.h"
+
+#include "PCGComputeCommon.generated.h"
+
+class UPCGData;
+class UPCGDataBinding;
+class UPCGSettings;
+struct FPCGContext;
+struct FPCGDataCollectionDesc;
+struct FPCGDataDesc;
+struct FPCGPinPropertiesGPU;
+
+#define PCG_KERNEL_LOGGING_ENABLED (!(UE_BUILD_SHIPPING || UE_BUILD_TEST) || USE_LOGGING_IN_SHIPPING)
+
+#if PCG_KERNEL_LOGGING_ENABLED
+#define PCG_KERNEL_VALIDATION_WARN(Context, Settings, bQuiet, ValidationMessage) if (!bQuiet) PCGComputeHelpers::LogKernelWarning(Context, Settings, ValidationMessage);
+#define PCG_KERNEL_VALIDATION_ERR(Context, Settings, bQuiet, ValidationMessage) if (!bQuiet) PCGComputeHelpers::LogKernelError(Context, Settings, ValidationMessage);
+#else
+#define PCG_KERNEL_VALIDATION_WARN(Context, Settings, bQuiet, ValidationMessage) // Log removed
+#define PCG_KERNEL_VALIDATION_ERR(Context, Settings, bQuiet, ValidationMessage) // Log removed
+#endif
+
+/** Modes for exporting the buffer from transient to persistent for downstream consumption. */
+UENUM(meta = (Bitflags))
+enum class EPCGExportMode : uint8
+{
+	/** Buffer is transient and freed after usage. */
+	NoExport = 0,
+	/** Buffer will be exported and a proxy will be output from the compute graph and passed to downstream nodes. */
+	ComputeGraphOutput = 1 << 0,
+	/** Producer node is being inspected, read back data and store in inspection data. */
+	Inspection = 1 << 1,
+	/** Producer node is being debugged, read back data and execute debug visualization. */
+	DebugVisualization = 1 << 2,
+};
+ENUM_CLASS_FLAGS(EPCGExportMode);
+
+namespace PCGComputeConstants
+{
+	constexpr int MAX_NUM_ATTRS = 128;
+	constexpr int NUM_RESERVED_ATTRS = 32; // Reserved for point properties, spline accessors, etc.
+	constexpr int MAX_NUM_CUSTOM_ATTRS = MAX_NUM_ATTRS - NUM_RESERVED_ATTRS; // Reserved for custom attributes
+
+	constexpr int DATA_COLLECTION_HEADER_SIZE_BYTES = 4; // 4 bytes for NumData
+	constexpr int DATA_HEADER_PREAMBLE_SIZE_BYTES = 12; // 4 bytes for DataType, 4 bytes for NumAttrs, 4 bytes for NumElements
+	constexpr int ATTRIBUTE_HEADER_SIZE_BYTES = 8; // 4 bytes for PackedIdAndStride, 4 bytes for data start address
+	constexpr int DATA_HEADER_SIZE_BYTES = DATA_HEADER_PREAMBLE_SIZE_BYTES + MAX_NUM_ATTRS * ATTRIBUTE_HEADER_SIZE_BYTES;
+
+	constexpr int POINT_DATA_TYPE_ID = 0;
+	constexpr int PARAM_DATA_TYPE_ID = 1;
+
+	constexpr int NUM_POINT_PROPERTIES = 9;
+	constexpr int POINT_POSITION_ATTRIBUTE_ID = 0;
+	constexpr int POINT_ROTATION_ATTRIBUTE_ID = 1;
+	constexpr int POINT_SCALE_ATTRIBUTE_ID = 2;
+	constexpr int POINT_BOUNDS_MIN_ATTRIBUTE_ID = 3;
+	constexpr int POINT_BOUNDS_MAX_ATTRIBUTE_ID = 4;
+	constexpr int POINT_COLOR_ATTRIBUTE_ID = 5;
+	constexpr int POINT_DENSITY_ATTRIBUTE_ID = 6;
+	constexpr int POINT_SEED_ATTRIBUTE_ID = 7;
+	constexpr int POINT_STEEPNESS_ATTRIBUTE_ID = 8;
+
+	/** PCG data types supported in GPU node inputs. */
+	constexpr EPCGDataType AllowedInputTypes = EPCGDataType::Point | EPCGDataType::Param | EPCGDataType::Landscape | EPCGDataType::BaseTexture | EPCGDataType::ProxyForGPU;
+
+	/** PCG data types supported in GPU node outputs. */
+	constexpr EPCGDataType AllowedOutputTypes = EPCGDataType::Point | EPCGDataType::Param | EPCGDataType::ProxyForGPU;
+
+	/** PCG data types supported in GPU data collections. */
+	constexpr EPCGDataType AllowedDataCollectionTypes = EPCGDataType::Point | EPCGDataType::Param | EPCGDataType::ProxyForGPU;
+
+	constexpr uint32 KernelExecutedFlag = 1 << 31;
+
+	constexpr uint32 MAX_PRIMITIVE_COMPONENTS_PER_SPAWNER = 256;
+
+	constexpr uint32 THREAD_GROUP_SIZE = 64;
+
+	// Used to represent invalid/removed points. We use a value slightly less than max float,
+	// as not all platforms support float infinity in shaders.
+	constexpr float INVALID_DENSITY = 3.402823e+38f;
+}
+
+namespace PCGComputeHelpers
+{
+	/** Gets the element count for a given data. E.g. number of points in a PointData, number of metadata entries in a ParamData, etc. */
+	int GetElementCount(const UPCGData* InData);
+
+	/** True if 'Type' is valid on a GPU input pin. */
+	bool IsTypeAllowedAsInput(EPCGDataType Type);
+
+	/** True if 'Type' is valid on a GPU output pin. */
+	bool IsTypeAllowedAsOutput(EPCGDataType Type);
+
+	/** True if 'Type' is valid in a GPU data collection. Some types are only supported as DataInterfaces, and cannot be uploaded in data collections. */
+	bool IsTypeAllowedInDataCollection(EPCGDataType Type);
+
+#if PCG_KERNEL_LOGGING_ENABLED
+	/** Logs a warning on a GPU node in the graph and console. */
+	void LogKernelWarning(const FPCGContext* Context, const UPCGSettings* Settings, const FText& InText);
+
+	/** Logs an error on a GPU node in the graph and console. */
+	void LogKernelError(const FPCGContext* Context, const UPCGSettings* Settings, const FText& InText);
+#endif
+
+	/** Computes a pin data description based on the given GPU pin properties. */
+	void ComputeOutputPinDataDesc(const FPCGPinPropertiesGPU& PinProperties, const UPCGSettings* Settings, const UPCGDataBinding* Binding, FPCGDataCollectionDesc& OutPinDesc);
+
+	/** Returns true if the given buffer size is dangerously large. Optionally emits error log. */
+	bool IsBufferSizeTooLarge(uint64 InBufferSizeBytes, bool bInLogError = true);
+
+	bool ComputeDataDescription(const UPCGData* InData, const UPCGDataBinding* InBinding, FPCGDataDesc& OutDescription);
+}
diff --git a/Engine/Plugins/PCG/Source/PCG/Internal/Compute/PCGDataForGPU.h b/Engine/Plugins/PCG/Source/PCG/Internal/Compute/PCGDataForGPU.h
new file mode 100644
--- /dev/null
+++ b/Engine/Plugins/PCG/Source/PCG/Internal/Compute/PCGDataForGPU.h
@@ -0,0 +1,263 @@
+// Copyright Epic Games, Inc. All Rights Reserved.
+
+#pragma once
+
+#include "PCGCommon.h"
+#include "PCGData.h"
+#include "Compute/PCGComputeCommon.h"
+
+#include "RHIGPUReadback.h"
+#include "Misc/SpinLock.h"
+
+#include "PCGDataForGPU.generated.h"
+
+class FRDGPooledBuffer;
+class FReferenceCollector;
+class UPCGDataBinding;
+class UPCGMetadata;
+class UPCGPin;
+class UPCGSettings;
+struct FPCGContext;
+enum class EPCGMetadataTypes : uint8;
+
+enum class EPCGUnpackDataCollectionResult
+{
+	Success,
+	DataMismatch,
+	NoData
+};
+
+UENUM()
+enum class EPCGKernelAttributeType : uint8
+{
+	Bool = 1,
+	Int,
+	Float,
+	Float2,
+	Float3,
+	Float4,
+	Rotator,
+	Quat,
+	Transform,
+	StringKey,
+	Name,
+
+	Invalid = std::numeric_limits<uint8>::max() UMETA(Hidden),
+};
+
+/** Attribute name and type which uniquely identify an attribute in a compute graph. */
+USTRUCT()
+struct FPCGKernelAttributeKey
+{
+	GENERATED_BODY()
+
+	FPCGKernelAttributeKey() = default;
+
+	explicit FPCGKernelAttributeKey(FName InName, EPCGKernelAttributeType InType)
+		: Name(InName)
+		, Type(InType)
+	{}
+
+	UPROPERTY(EditAnywhere, Category = "Settings")
+	FName Name = NAME_None;
+
+	UPROPERTY(EditAnywhere, Category = "Settings")
+	EPCGKernelAttributeType Type = EPCGKernelAttributeType::Float;
+
+	bool operator==(const FPCGKernelAttributeKey& Other) const;
+
+	friend uint32 GetTypeHash(const FPCGKernelAttributeKey& In);
+};
+
+/** Table of attributes used in compute graph with helpers to get unique attribute ID used to read/write attribute in data collection buffers. */
+USTRUCT()
+struct FPCGKernelAttributeTable
+{
+	GENERATED_BODY()
+
+	FPCGKernelAttributeTable() = default;
+
+	int32 GetAttributeId(const FPCGKernelAttributeKey& InAttribute) const;
+
+	int32 GetAttributeId(FName InName, EPCGKernelAttributeType InType) const;
+
+	void AddAttribute(FName InName, EPCGKernelAttributeType InType);
+
+#if PCG_KERNEL_LOGGING_ENABLED
+	void DebugLog();
+#endif
+
+protected:
+	UPROPERTY()
+	TArray<FPCGKernelAttributeKey> AttributeTable;
+
+	friend class FPCGGraphCompilerGPU;
+};
+
+/** Data description for a metadata attribute. Stores identifying name and type as well as the unique attribute ID. */
+struct FPCGKernelAttributeDesc
+{
+public:
+	FPCGKernelAttributeDesc() = default;
+
+	explicit FPCGKernelAttributeDesc(int32 InIndex, EPCGKernelAttributeType InType, FName InName)
+		: Index(InIndex)
+		, Type(InType)
+		, Name(InName)
+	{
+	}
+
+	explicit FPCGKernelAttributeDesc(int32 InIndex, EPCGKernelAttributeType InType, FName InName, TArray<int32>&& InUniqueStringKeys)
+		: Index(InIndex)
+		, Type(InType)
+		, Name(InName)
+		, UniqueStringKeys(MoveTemp(InUniqueStringKeys))
+	{
+	}
+
+	const TArray<int32>& GetUniqueStringKeys() const { return UniqueStringKeys; }
+
+	void AddUniqueStringKeys(const TArray<int32>& InOtherStringKeys);
+
+	bool operator==(const FPCGKernelAttributeDesc& Other) const;
+
+public:
+	// Attribute Id
+	int32 Index = INDEX_NONE;
+
+	EPCGKernelAttributeType Type = EPCGKernelAttributeType::Float;
+
+	FName Name = NAME_None;
+
+private:
+	/* All possible string keys arriving on this attribute (string keys are indices into the string table in the data binding). */
+	TArray<int32> UniqueStringKeys;
+};
+
+/** Data description for a single data object (UPCGData). */
+struct FPCGDataDesc
+{
+public:
+	FPCGDataDesc() = default;
+	FPCGDataDesc(EPCGDataType InType, int32 InElementCount);
+	FPCGDataDesc(const UPCGData* InData, const UPCGDataBinding* InBinding);
+
+	uint64 ComputePackedSize() const;
+
+	bool HasMetadataAttributes() const;
+
+	bool ContainsAttribute(FName InAttributeName) const;
+
+public:
+	EPCGDataType Type = EPCGDataType::Point;
+	TArray<FPCGKernelAttributeDesc> AttributeDescs;
+	int32 ElementCount = 0;
+
+private:
+	void InitializeAttributeDescs(const UPCGData* InData, const UPCGDataBinding* InBinding);
+};
+
+/** Data description for a data collection (FPCGDataCollection). */
+struct FPCGDataCollectionDesc
+{
+	static FPCGDataCollectionDesc BuildFromDataCollection(
+		const FPCGDataCollection& InDataCollection,
+		const UPCGDataBinding* InBinding);
+
+	static FPCGDataCollectionDesc BuildFromInputDataCollectionAndInputPinLabel(
+		const FPCGDataCollection& InDataCollection,
+		FName InputPinLabel,
+		const UPCGDataBinding* InBinding);
+
+	/** Computes the size (in bytes) of the header portion of the packed data collection buffer. */
+	uint32 ComputePackedHeaderSizeBytes() const;
+
+	/** Computes the size (in bytes) of the data collection after packing. */
+	uint64 ComputePackedSizeBytes() const;
+
+	void WriteHeader(TArray<uint32>& OutPackedDataCollectionHeader) const;
+
+	/** Pack a data collection into the GPU data format. DataDescs defines which attributes are packed. */
+	void PackDataCollection(const FPCGDataCollection& InDataCollection, FName InPin, const TArray<FString>& InStringTable, TArray<uint32>& OutPackedDataCollection) const;
+
+	/** Unpack a buffer of 8-bit uints to a data collection. */
+	EPCGUnpackDataCollectionResult UnpackDataCollection(FPCGContext* InContext, const TArray<uint8>& InPackedData, FName InPin, const TArray<FString>& InStringTable, FPCGDataCollection& OutDataCollection) const;
+
+	/** Compute total number of processing elements of the given type. */
+	uint32 ComputeDataElementCount(EPCGDataType InDataType) const;
+
+	/** Aggregate another data description. */
+	void Combine(const FPCGDataCollectionDesc& Other);
+
+	/** Get description of first attribute with matching name in input data. Returns true if such an attribute found and also signals whether multiple matching attributes
+	* with conflicting names are present. */
+	bool GetAttributeDesc(FName InAttributeName, FPCGKernelAttributeDesc& OutAttributeDesc, bool &bOutConflictingTypesFound) const;
+
+	bool ContainsAttributeOnAnyData(FName InAttributeName) const;
+
+	/** Makes attribute present on all data. If data has existing attribute with same name then the given type will be applied. */
+	void AddAttributeToAllData(FPCGKernelAttributeKey InAttribute, const UPCGDataBinding* InBinding, TArray<int32>* InOptionalUniqueStringKeys = nullptr);
+
+	/** Description of each data in this data collection. */
+	TArray<FPCGDataDesc> DataDescs;
+};
+
+/** A proxy for a data collection residing in a GPU buffer along with functionality to retrieve the data on the CPU. Holds onto GPU memory. */
+struct FPCGProxyForGPUDataCollection : public TSharedFromThis<FPCGProxyForGPUDataCollection>
+{
+public:
+	FPCGProxyForGPUDataCollection(TRefCountPtr<FRDGPooledBuffer> InBuffer, uint32 InBufferSizeBytes, const FPCGDataCollectionDesc& InDescription, const TArray<FString>& InStringTable);
+
+	/** Populates a CPU data object representing the GPU data for the given index, performing a readback from GPU->CPU if required. */
+	bool GetCPUData(FPCGContext* InContext, int32 InDataIndex, const UPCGData*& OutData);
+
+	void AddReferences(FReferenceCollector& Collector);
+
+	const TRefCountPtr<FRDGPooledBuffer>& GetBuffer() const { return Buffer; }
+
+	const FPCGDataCollectionDesc& GetDescription() const { return Description; }
+
+protected:
+	/** Persistent GPU buffer that can be read back. Buffer will be freed when this ref count is 0. */
+	TRefCountPtr<FRDGPooledBuffer> Buffer;
+
+	uint32 BufferSizeBytes = 0;
+
+	FPCGDataCollectionDesc Description;
+
+	/** Used to comprehend string IDs in buffer. */
+	TArray<FString> StringTable;
+
+	/** Read back data. Populated once upon first readback request. */
+	FPCGDataCollection ReadbackData;
+
+	TSharedPtr<FRHIGPUBufferReadback> ReadbackRequest;
+
+	TArray<uint8> RawReadbackData;
+
+	std::atomic<bool> bReadbackDataProcessed = false;
+
+	bool bReadbackDataArrived = false;
+
+	UE::FSpinLock ReadbackLock;
+};
+
+namespace PCGDataForGPUHelpers
+{
+	/** Returns GPU type that will be used to represent the given metadata type. */
+	EPCGKernelAttributeType GetAttributeTypeFromMetadataType(EPCGMetadataTypes MetadataType);
+
+	/** Compute how attributes can be packed to custom floats. */
+	void ComputeCustomFloatPacking(
+		FPCGContext* InContext,
+		const UPCGSettings* InSettings,
+		TArray<FName>& InAttributeNames,
+		const UPCGDataBinding* InBinding,
+		const FPCGDataCollectionDesc& InDataCollectionDescription,
+		uint32& OutCustomFloatCount,
+		TArray<FUint32Vector4>& OutAttributeIdOffsetStrides);
+
+	/** Compute a description of all data arriving on InputPin. */
+	FPCGDataCollectionDesc ComputeInputPinDataDesc(const UPCGSettings* Settings, const FName& InputPinLabel, const UPCGDataBinding* Binding);
+	FPCGDataCollectionDesc ComputeInputPinDataDesc(const UPCGPin* InputPin, const UPCGDataBinding* Binding);
+}
diff --git a/Engine/Plugins/PCG/Source/PCG/Private/Compute/Data/PCGProxyForGPUData.cpp b/Engine/Plugins/PCG/Source/PCG/Private/Compute/Data/PCGProxyForGPUData.cpp
new file mode 100644
--- /dev/null
+++ b/Engine/Plugins/PCG/Source/PCG/Private/Compute/Data/PCGProxyForGPUData.cpp
@@ -0,0 +1,90 @@
+// Copyright Epic Games, Inc. All Rights Reserved.
+
+#include "Compute/Data/PCGProxyForGPUData.h"
+
+#include "PCGModule.h"
+
+void UPCGProxyForGPUData::Initialize(TSharedPtr<FPCGProxyForGPUDataCollection> InDataCollection, int InDataIndexInCollection)
+{
+	DataCollectionOnGPU = InDataCollection;
+	DataIndexInCollection = InDataIndexInCollection;
+}
+
+void UPCGProxyForGPUData::AddToCrc(FArchiveCrc32& Ar, bool bFullDataCrc) const
+{
+	Super::AddToCrc(Ar, bFullDataCrc);
+
+	// This data does not have a bespoke CRC implementation so just use the unique object instance UID.
+	AddUIDToCrc(Ar);
+}
+
+EPCGDataType UPCGProxyForGPUData::GetUnderlyingDataType() const
+{
+	if (TSharedPtr<FPCGProxyForGPUDataCollection> DataCollectionGPU = GetGPUInfoValidated())
+	{
+		return DataCollectionGPU->GetDescription().DataDescs[DataIndexInCollection].Type;
+	}
+
+	return EPCGDataType::None;
+}
+
+TSharedPtr<const FPCGProxyForGPUDataCollection> UPCGProxyForGPUData::GetInputDataCollectionInfo() const
+{
+	return DataCollectionOnGPU;
+}
+
+int UPCGProxyForGPUData::GetElementCount() const
+{
+	if (TSharedPtr<FPCGProxyForGPUDataCollection> DataCollectionGPU = GetGPUInfoValidated())
+	{
+		return DataCollectionGPU->GetDescription().DataDescs[DataIndexInCollection].ElementCount;
+	}
+
+	return 0;
+}
+
+bool UPCGProxyForGPUData::GetDescription(FPCGDataDesc& OutDescription) const
+{
+	if (TSharedPtr<FPCGProxyForGPUDataCollection> DataCollectionGPU = GetGPUInfoValidated())
+	{
+		OutDescription = DataCollectionGPU->GetDescription().DataDescs[DataIndexInCollection];
+		return true;
+	}
+
+	return false;
+}
+
+UPCGProxyForGPUData::FReadbackResult UPCGProxyForGPUData::GetCPUData(FPCGContext* InContext) const
+{
+	TSharedPtr<FPCGProxyForGPUDataCollection> DataOnGPU = GetGPUInfoValidated();
+	if (!DataOnGPU)
+	{
+		UE_LOG(LogPCG, Error, TEXT("Data collection lost!"));
+
+		return FReadbackResult{ .bComplete = true, .Data = nullptr };
+	}
+
+	const UPCGData* ResultData = nullptr;
+	if (!DataOnGPU->GetCPUData(InContext, DataIndexInCollection, ResultData))
+	{
+		return FReadbackResult{ .bComplete = false, .Data = nullptr };
+	}
+
+	return FReadbackResult{ .bComplete = true, .Data = ResultData };
+}
+
+TSharedPtr<FPCGProxyForGPUDataCollection> UPCGProxyForGPUData::GetGPUInfoValidated() const
+{
+	if (!DataCollectionOnGPU)
+	{
+		ensureMsgf(false, TEXT("Data %s: GPU data collection lost"), *GetName());
+		return nullptr;
+	}
+	else if (!DataCollectionOnGPU->GetDescription().DataDescs.IsValidIndex(DataIndexInCollection))
+	{
+		ensureMsgf(false, TEXT("Data %s: DataIndexInCollection (%d) was out of range [0, %d)."), *GetName(), DataIndexInCollection, DataCollectionOnGPU->GetDescription().DataDescs.Num());
+		return nullptr;
+	}
+
+	return DataCollectionOnGPU;
+}
diff --git a/Engine/Plugins/PCG/Source/PCG/Private/Compute/DataInterfaces/PCGComputeDataInterface.h b/Engine/Plugins/PCG/Source/PCG/Private/Compute/DataInterfaces/PCGComputeDataInterface.h
--- a/Engine/Plugins/PCG/Source/PCG/Private/Compute/DataInterfaces/PCGComputeDataInterface.h
+++ b/Engine/Plugins/PCG/Source/PCG/Private/Compute/DataInterfaces/PCGComputeDataInterface.h
@@ -3,6 +3,7 @@
 #pragma once
 
 #include "ComputeFramework/ComputeDataInterface.h"
+#include "ComputeFramework/ComputeDataProvider.h"
 
 #include "PCGComputeDataInterface.generated.h"
 
@@ -37,3 +38,12 @@ protected:
 	UPROPERTY()
 	TArray<FName> DownstreamInputPinLabelAliases;
 };
+
+UCLASS(Abstract, ClassGroup (Procedural))
+class UPCGComputeDataProvider : public UComputeDataProvider
+{
+	GENERATED_BODY()
+
+public:
+	virtual bool PrepareForExecute_GameThread() { return true; }
+};
diff --git a/Engine/Plugins/PCG/Source/PCG/Private/Compute/DataInterfaces/PCGDataCollectionDataInterface.cpp b/Engine/Plugins/PCG/Source/PCG/Private/Compute/DataInterfaces/PCGDataCollectionDataInterface.cpp
--- a/Engine/Plugins/PCG/Source/PCG/Private/Compute/DataInterfaces/PCGDataCollectionDataInterface.cpp
+++ b/Engine/Plugins/PCG/Source/PCG/Private/Compute/DataInterfaces/PCGDataCollectionDataInterface.cpp
@@ -9,7 +9,9 @@
 #include "PCGSubsystem.h"
 #include "Compute/PCGComputeCommon.h"
 #include "Compute/PCGDataBinding.h"
+#include "Compute/Data/PCGProxyForGPUData.h"
 #include "Compute/Elements/PCGComputeGraphElement.h"
+#include "Helpers/PCGHelpers.h"
 
 #include "ComputeFramework/ShaderParamTypeDefinition.h"
 #include "RenderGraphBuilder.h"
@@ -441,17 +443,17 @@ UComputeDataProvider* UPCGDataCollectionDataInterface::CreateDataProvider(TObjec
 	Provider->ProducerSettings = ProducerSettings;
 	Provider->PinDesc = MoveTemp(PinDataDesc);
 
-	Provider->ReadbackMode = bRequiresReadback ? EPCGReadbackMode::GraphOutput : EPCGReadbackMode::None;
+	Provider->ExportMode = bRequiresExport ? EPCGExportMode::ComputeGraphOutput : EPCGExportMode::NoExport;
 
 #if WITH_EDITOR
 	if (Binding->SourceComponent->IsInspecting())
 	{
-		Provider->ReadbackMode |= EPCGReadbackMode::Inspection;
+		Provider->ExportMode |= EPCGExportMode::Inspection;
 	}
 
 	if (ProducerSettings->bDebug)
 	{
-		Provider->ReadbackMode |= EPCGReadbackMode::DebugVisualization;
+		Provider->ExportMode |= EPCGExportMode::DebugVisualization;
 	}
 #endif
 
@@ -464,21 +466,11 @@ UComputeDataProvider* UPCGDataCollectionDataInterface::CreateDataProvider(TObjec
 	return Provider;
 }
 
-bool UPCGDataCollectionDataInterface::GetRequiresReadback() const
-{
-#if WITH_EDITOR
-	// In editor DIs are always flagged as requiring readback, and decision whether to actually readback data or not
-	// is made at execution time (in GetReadbackData) based on inspection state.
-	return true;
-#else
-	return bRequiresReadback;
-#endif
-}
-
 FComputeDataProviderRenderProxy* UPCGDataCollectionDataProvider::GetRenderProxy()
 {
 	TRACE_CPUPROFILER_EVENT_SCOPE(UPCGDataCollectionDataProvider::GetRenderProxy);
-	FPCGDataCollectionDataProviderProxy* Proxy = new FPCGDataCollectionDataProviderProxy(PinDesc, ReadbackMode);
+
+	FPCGDataCollectionDataProviderProxy* Proxy = new FPCGDataCollectionDataProviderProxy(this, PinDesc, ExportMode, OutputPinLabel, OutputPinLabelAlias);
 
 #if WITH_EDITOR
 	if (const UPCGNode* Node = Cast<UPCGNode>(ProducerSettings ? ProducerSettings->GetOuter() : nullptr))
@@ -491,130 +483,15 @@ FComputeDataProviderRenderProxy* UPCGDataCollectionDataProvider::GetRenderProxy(
 	}
 #endif
 
-	if (ReadbackMode != EPCGReadbackMode::None)
-	{
-		TWeakObjectPtr<UPCGDataCollectionDataProvider> ThisWeakPtr(this);
-
-		Proxy->AsyncReadbackCallback_RenderThread = [ThisWeakPtr](const void* InData, int InNumBytes)
-		{
-			TRACE_CPUPROFILER_EVENT_SCOPE(UPCGDataCollectionDataProvider::ProcessReadbackData_RenderThread);
-
-			if (!ThisWeakPtr.IsValid() || !ThisWeakPtr.Get())
-			{
-				return;
-			}
-
-			UPCGDataCollectionDataProvider* ThisDataProvider = ThisWeakPtr.Get();
-			check(ThisDataProvider);
-
-			// We should never find ourselves stomping existing data.
-			check(ThisDataProvider->RawReadbackData.IsEmpty());
-
-			if (InData && InNumBytes > 0)
-			{
-				// Copy the data to temp storage for the game thread to pick up.
-				ThisDataProvider->RawReadbackData.SetNumUninitialized(InNumBytes);
-				FMemory::Memcpy(ThisDataProvider->RawReadbackData.GetData(), InData, InNumBytes);
-			}
-			else
-			{
-				// Can happen currently if no threads dispatched.
-				ThisDataProvider->RawReadbackData.Reset();
-			}
-
-			ThisDataProvider->bReadbackComplete = true;
-
-			ThisDataProvider->OnReadbackComplete.Broadcast();
-		};
-	}
-
 	return Proxy;
 }
 
-bool UPCGDataCollectionDataProvider::ProcessReadBackData(FPCGComputeGraphContext* InContext)
-{
-	TRACE_CPUPROFILER_EVENT_SCOPE(UPCGDataCollectionDataProvider::ProcessReadBackData);
-	check(InContext);
-	check(ProducerSettings);
-	check(IsInGameThread());
-
-	if (!ensure(bReadbackComplete))
-	{
-		// This should not be called until readback has been done.
-		return false;
-	}
-
-	if (RawReadbackData.IsEmpty())
-	{
-		// If there's no data, we leave an empty data collection and we're done.
-		return true;
-	}
-
-	FPCGDataCollection DataFromGPU;
-	const EPCGUnpackDataCollectionResult Result = PinDesc.UnpackDataCollection(InContext, RawReadbackData, OutputPinLabelAlias, InContext->DataBinding->GetStringTable(), DataFromGPU);
-
-	if (Result == EPCGUnpackDataCollectionResult::DataMismatch)
-	{
-		PCG_KERNEL_VALIDATION_WARN(InContext, ProducerSettings, /*bQuiet=*/false, FText::Format(
-			LOCTEXT("UnpackDataCollectionDataMismatch", "Mismatch in expected data while unpacking GPU data collection on pin '{0}'. Static analysis does not match the received data. Data collection will be ignored."),
-			FText::FromName(OutputPinLabel)));
-	}
-
-	RawReadbackData.Reset();
-
-	// Store data in output collection.
-	if (!!(ReadbackMode & EPCGReadbackMode::GraphOutput) && ensure(Binding.IsValid()))
-	{
-		Binding->OutputDataCollection.TaggedData.Append(DataFromGPU.TaggedData);
-	}
-
-	// Store data for Inspection.
-#if WITH_EDITOR
-	if (!!(ReadbackMode & EPCGReadbackMode::Inspection))
-	{
-		const UPCGNode* Node = Cast<UPCGNode>(ProducerSettings->GetOuter());
-		UPCGComponent* Component = InContext->SourceComponent.Get();
-		if (Component && InContext->Stack && Node)
-		{
-			// Virtual pin labels confuse inspection. Apply the original output label before storing.
-			for (FPCGTaggedData& Data : DataFromGPU.TaggedData)
-			{
-				Data.Pin = OutputPinLabel;
-			}
-
-			// Required by inspection code.
-			DataFromGPU.ComputeCrcs(/*bFullDataCrc=*/false);
-
-			// Input data not yet supported.
-			Component->StoreInspectionData(InContext->Stack, Node, /*InTimer=*/nullptr, /*InInputData=*/{}, DataFromGPU, /*bUsedCache*/false);
-		}
-	}
-
-	if (!!(ReadbackMode & EPCGReadbackMode::DebugVisualization))
-	{
-		const FPCGDataVisualizationRegistry& DataVisRegistry = FPCGModule::GetConstPCGDataVisualizationRegistry();
-
-		for (const FPCGTaggedData& Output : DataFromGPU.TaggedData)
-		{
-			if (!Output.Data)
-			{
-				continue;
-			}
-
-			if (const IPCGDataVisualization* DataVis = DataVisRegistry.GetDataVisualization(Output.Data->GetClass()))
-			{
-				DataVis->ExecuteDebugDisplay(InContext, ProducerSettings, Output.Data, InContext->GetTargetActor(nullptr));
-			}
-		}
-	}
-#endif // WITH_EDITOR
-
-	return true;
-}
-
-FPCGDataCollectionDataProviderProxy::FPCGDataCollectionDataProviderProxy(const FPCGDataCollectionDesc& InPinDesc, EPCGReadbackMode InReadbackMode)
-	: ReadbackMode(InReadbackMode)
+FPCGDataCollectionDataProviderProxy::FPCGDataCollectionDataProviderProxy(TWeakObjectPtr<UPCGDataCollectionDataProvider> InDataProvider, const FPCGDataCollectionDesc& InPinDesc, EPCGExportMode InExportMode, FName InOutputPinLabel, FName InOutputPinLabelAlias)
+	: ExportMode(InExportMode)
 	, PinDesc(InPinDesc)
+	, DataProviderWeakPtr(InDataProvider)
+	, OutputPinLabel(InOutputPinLabel)
+	, OutputPinLabelAlias(InOutputPinLabelAlias)
 {
 	SizeBytes = PinDesc.ComputePackedSizeBytes();
 }
@@ -660,32 +537,58 @@ void FPCGDataCollectionDataProviderProxy::AllocateResources(FRDGBuilder& GraphBu
 #endif
 
 	FRDGBufferDesc Desc = FRDGBufferDesc::CreateByteAddressDesc(SizeBytes);
-	if (ReadbackMode != EPCGReadbackMode::None)
+	if (ExportMode != EPCGExportMode::NoExport)
 	{
+		// We don't know for sure whether buffer will be read back or not, so need to flag the possibility if the buffer will be passed downstream.
 		Desc.Usage |= BUF_SourceCopy;
 	}
 
 	Buffer = GraphBuilder.CreateBuffer(Desc, DebugName);
 	BufferUAV = GraphBuilder.CreateUAV(Buffer);
 
-	// Initialize with an empty data collection. The kernel may not run, for example if indirect dispatch args end up being 0. Ensure
+	// Initialize with an empty data collection. The kernel may not run, for example if indirect dispatch args end up being 0. This ensures
 	// there is something meaningful to readback.
 	TArray<uint32> PackedDataCollection;
 	PinDesc.WriteHeader(PackedDataCollection);
 
 	GraphBuilder.QueueBufferUpload(Buffer, PackedDataCollection.GetData(), PackedDataCollection.Num() * PackedDataCollection.GetTypeSize(), ERDGInitialDataFlags::None);
-}
 
-void FPCGDataCollectionDataProviderProxy::GetReadbackData(TArray<FReadbackData>& OutReadbackData) const
-{
-	if (ReadbackMode != EPCGReadbackMode::None)
+	// If buffer needs to be exported, can do it now. Pass back a reference to the main thread where it can be picked up by the compute graph element.
+	if (ExportMode != EPCGExportMode::NoExport)
 	{
-		FReadbackData Data;
-		Data.Buffer = Buffer;
-		Data.NumBytes = SizeBytes;
-		Data.ReadbackCallback_RenderThread = &AsyncReadbackCallback_RenderThread;
+		const TRefCountPtr<FRDGPooledBuffer> ExportedPooledBuffer = GraphBuilder.ConvertToExternalBuffer(Buffer);
+		GraphBuilder.SetBufferAccessFinal(Buffer, ERHIAccess::CopySrc);
+
+		PCGHelpers::ExecuteOnGameThread(UE_SOURCE_LOCATION, [ExportedPooledBuffer, DataProviderWeakPtr=DataProviderWeakPtr, PinDesc=PinDesc, OutputPinLabel=OutputPinLabel, SizeBytes=SizeBytes, ExportMode=ExportMode, OutputPinLabelAlias=OutputPinLabelAlias]()
+		{
+			// Obtain objects. No ensures added because a graph cancellation could feasibly destroy some or all of these.
+			UPCGDataCollectionDataProvider* DataProvider = DataProviderWeakPtr.Get();
+			UPCGDataBinding* Binding = DataProvider ? DataProvider->Binding.Get() : nullptr;
+			UPCGSubsystem* Subsystem = UPCGSubsystem::GetSubsystemForCurrentWorld();
 
-		OutReadbackData.Add(MoveTemp(Data));
+			if (DataProvider && Binding && Subsystem)
+			{
+				TSharedPtr<FPCGProxyForGPUDataCollection> DataCollectionOnGPU = MakeShared<FPCGProxyForGPUDataCollection>(ExportedPooledBuffer, SizeBytes, PinDesc, Binding->GetStringTable());
+
+				for (int DataIndex = 0; DataIndex < PinDesc.DataDescs.Num(); ++DataIndex)
+				{
+					UPCGProxyForGPUData* Proxy = NewObject<UPCGProxyForGPUData>();
+					Proxy->Initialize(DataCollectionOnGPU, DataIndex);
+
+					// TODO - binding is doing a lot of work. Could store a context handle in the data provider instead?
+					Binding->ReceiveDataFromGPU_GameThread(Proxy, DataProvider->ProducerSettings, ExportMode, OutputPinLabel, OutputPinLabelAlias);
+				}
+			}
+			else
+			{
+				UE_LOG(LogPCG, Error, TEXT("Could not resolve UPCGDataCollectionDataProvider object to pass back buffer handle."));
+			}
+
+			if (DataProvider)
+			{
+				DataProvider->OnDataExported_GameThread().Broadcast();
+			}
+		});
 	}
 }
 
diff --git a/Engine/Plugins/PCG/Source/PCG/Private/Compute/DataInterfaces/PCGDataCollectionDataInterface.h b/Engine/Plugins/PCG/Source/PCG/Private/Compute/DataInterfaces/PCGDataCollectionDataInterface.h
--- a/Engine/Plugins/PCG/Source/PCG/Private/Compute/DataInterfaces/PCGDataCollectionDataInterface.h
+++ b/Engine/Plugins/PCG/Source/PCG/Private/Compute/DataInterfaces/PCGDataCollectionDataInterface.h
@@ -5,29 +5,15 @@
 #include "Compute/PCGDataForGPU.h"
 #include "Compute/DataInterfaces/PCGComputeDataInterface.h"
 
+#include "RenderGraphResources.h"
 #include "ComputeFramework/ComputeDataProvider.h"
 
 #include "PCGDataCollectionDataInterface.generated.h"
 
 struct FPCGComputeGraphContext;
 class FPCGDataCollectionDataInterfaceParameters;
-class FRDGBuffer;
-class FRDGBufferSRV;
 class UPCGDataBinding;
 
-UENUM(meta = (Bitflags))
-enum class EPCGReadbackMode : uint8
-{
-	None = 0,
-	/** Read back and store in compute graph output data collection. */
-	GraphOutput = 1 << 1,
-	/** Read back and store inspection data. */
-	Inspection = 1 << 2,
-	/** Read back and execute debug visualization. */
-	DebugVisualization = 1 << 3,
-};
-ENUM_CLASS_FLAGS(EPCGReadbackMode);
-
 /** Compute Framework Data Interface for reading PCG data. */
 UCLASS(ClassGroup = (Procedural))
 class UPCGDataCollectionDataInterface : public UPCGComputeDataInterface
@@ -52,24 +38,23 @@ public:
 
 	void SetProducerSettings(const UPCGSettings* InProducerSettings) { ProducerSettings = InProducerSettings; }
 
-	bool GetRequiresReadback() const override;
-	void SetRequiresReadback(bool bInRequiresReadback) { bRequiresReadback = bInRequiresReadback; }
+	void SetRequiresExport(bool bInRequiresExport) { bRequiresExport = bInRequiresExport; }
 
 protected:
 	/** Settings of node that produces this data, normally the upstream node. */
 	UPROPERTY()
 	TObjectPtr<const UPCGSettings> ProducerSettings;
 
-	/** Whether this data should be copied from GPU to CPU. */
+	/** Whether this data is passed to downstream tasks outside of this compute graph. */
 	UPROPERTY()
-	bool bRequiresReadback = false;
+	bool bRequiresExport = false;
 
 	static TCHAR const* TemplateFilePath;
 };
 
 /** Compute Framework Data Provider for reading a PCG Data Collection. */
 UCLASS()
-class UPCGDataCollectionDataProvider : public UComputeDataProvider
+class UPCGDataCollectionDataProvider : public UPCGComputeDataProvider
 {
 	GENERATED_BODY()
 
@@ -78,13 +63,11 @@ public:
 	FComputeDataProviderRenderProxy* GetRenderProxy() override;
 	//~ End UComputeDataProvider Interface
 
-	bool RequiresReadback() const { return ReadbackMode != EPCGReadbackMode::None; }
+	/** Returns true if the generated buffer should be exported to a persistent buffer to be passed to downstream nodes. */
+	bool IsExportRequired() const { return ExportMode != EPCGExportMode::NoExport; }
 
-	/** Processed read back data. Should be called after readback complete. Called from game thread. Returns true if data processed successfully. */
-	bool ProcessReadBackData(FPCGComputeGraphContext* InContext);
-
-	DECLARE_EVENT(UPCGDataProviderDataCollectionReadback, FOnReadbackComplete);
-	FOnReadbackComplete& OnReadbackComplete_RenderThread() { return OnReadbackComplete; }
+	DECLARE_EVENT(UPCGDataCollectionDataProvider, FOnDataExported);
+	FOnDataExported& OnDataExported_GameThread() { return OnDataExported; }
 
 public:
 	/** Settings of node that produces this data, normally the upstream node. */
@@ -92,43 +75,41 @@ public:
 	TObjectPtr<const UPCGSettings> ProducerSettings;
 
 	TWeakObjectPtr<UPCGDataBinding> Binding;
+
 	FPCGDataCollectionDesc PinDesc;
 
-	EPCGReadbackMode ReadbackMode = EPCGReadbackMode::None;
+	UPROPERTY()
+	EPCGExportMode ExportMode = EPCGExportMode::NoExport;
 
+	UPROPERTY()
 	FName OutputPinLabel;
-	FName OutputPinLabelAlias;
 
-protected:
-	std::atomic<bool> bReadbackComplete = false;
-	TArray<uint8> RawReadbackData;
+	UPROPERTY()
+	FName OutputPinLabelAlias;
 
-	FOnReadbackComplete OnReadbackComplete;
+private:
+	FOnDataExported OnDataExported;
 };
 
 class FPCGDataCollectionDataProviderProxy : public FComputeDataProviderRenderProxy
 {
 public:
-	FPCGDataCollectionDataProviderProxy(const FPCGDataCollectionDesc& InPinDesc, EPCGReadbackMode InReadbackMode);
+	FPCGDataCollectionDataProviderProxy(TWeakObjectPtr<UPCGDataCollectionDataProvider> InDataProvider, const FPCGDataCollectionDesc& InPinDesc, EPCGExportMode InExportMode, FName InOutputPinLabel, FName InOutputPinLabelAlias);
 
 	//~ Begin FComputeDataProviderRenderProxy Interface
 	bool IsValid(FValidationData const& InValidationData) const override;
 	void GatherDispatchData(FDispatchData const& InDispatchData) override;
 	void AllocateResources(FRDGBuilder& GraphBuilder, FAllocationData const& InAllocationData) override;
-	void GetReadbackData(TArray<FReadbackData>& OutReadbackData) const override;
 	//~ End FComputeDataProviderRenderProxy Interface
 
 #if WITH_EDITOR
 	void SetBufferName(const FString& InBufferName);
 #endif
 
-public:
-	FReadbackCallback AsyncReadbackCallback_RenderThread;
-
 protected:
 	uint32 SizeBytes = 0;
 
-	EPCGReadbackMode ReadbackMode = EPCGReadbackMode::None;
+	EPCGExportMode ExportMode = EPCGExportMode::NoExport;
 
 	using FParameters = FPCGDataCollectionDataInterfaceParameters;
 
@@ -137,6 +118,12 @@ protected:
 	FRDGBufferRef Buffer = nullptr;
 	FRDGBufferUAVRef BufferUAV = nullptr;
 
+	// Weak pointer useful for passing back buffer handles. Do not access this directly from the render thread.
+	TWeakObjectPtr<UPCGDataCollectionDataProvider> DataProviderWeakPtr;
+
+	FName OutputPinLabel;
+	FName OutputPinLabelAlias;
+
 #if WITH_EDITOR
 	FString BufferName;
 #endif
diff --git a/Engine/Plugins/PCG/Source/PCG/Private/Compute/DataInterfaces/PCGDataCollectionUploadDataInterface.cpp b/Engine/Plugins/PCG/Source/PCG/Private/Compute/DataInterfaces/PCGDataCollectionUploadDataInterface.cpp
--- a/Engine/Plugins/PCG/Source/PCG/Private/Compute/DataInterfaces/PCGDataCollectionUploadDataInterface.cpp
+++ b/Engine/Plugins/PCG/Source/PCG/Private/Compute/DataInterfaces/PCGDataCollectionUploadDataInterface.cpp
@@ -6,13 +6,175 @@
 #include "PCGModule.h"
 #include "PCGNode.h"
 #include "PCGSettings.h"
+#include "PCGSubsystem.h"
 #include "Compute/PCGComputeCommon.h"
 #include "Compute/PCGComputeGraph.h"
 #include "Compute/PCGDataBinding.h"
+#include "Compute/Data/PCGProxyForGPUData.h"
 
 #include "RenderGraphBuilder.h"
 #include "RenderGraphResources.h"
 
+/** A helper to ingest incoming data objects and pass the required data to the data interface. */
+class FPCGDataUploadAdaptor
+{
+public:
+	FPCGDataUploadAdaptor(UPCGDataBinding* InDataBinding, const FPCGDataCollectionDesc& InPinDesc, FName InPinLabel)
+	{
+		PinDesc = InPinDesc;
+		PinLabel = InPinLabel;
+		DataBinding = InDataBinding;
+
+		// Are all our data items coming direct from GPU - and importantly are a complete, single data collection?
+		bool bMultipleBuffersEncountered = false;
+		TSharedPtr<const FPCGProxyForGPUDataCollection> UniqueDataCollection;
+		for (int InputDataIndex = 0; InputDataIndex < DataBinding->InputDataCollection.TaggedData.Num(); ++InputDataIndex)
+		{
+			const UPCGProxyForGPUData* DataGPU = Cast<UPCGProxyForGPUData>(DataBinding->InputDataCollection.TaggedData[InputDataIndex].Data);
+			if (!DataGPU)
+			{
+				// Encountered CPU data, no reuse for now.
+				// TODO: Support multiple input buffers and only upload required data.
+				UniqueDataCollection.Reset();
+				break;
+			}
+
+			TSharedPtr<const FPCGProxyForGPUDataCollection> InputDataCollection = DataGPU->GetInputDataCollectionInfo();
+
+			if (!InputDataCollection)
+			{
+				UE_LOG(LogPCG, Error, TEXT("No reuse: Missing data collection buffer!"));
+
+				UniqueDataCollection.Reset();
+				break;
+			}
+
+			if (DataGPU->GetDataIndexInCollection() != InputDataIndex)
+			{
+				// Out of order data items, no reuse for now.
+				// TODO: support indirection of data items.
+				UniqueDataCollection.Reset();
+				break;
+			}
+
+			if (UniqueDataCollection)
+			{
+				if (UniqueDataCollection != InputDataCollection)
+				{
+					// Multiple buffer items, no reuse for now.
+					// TODO: Support multiple input buffers.
+					UniqueDataCollection.Reset();
+					break;
+				}
+			}
+			else
+			{
+				UniqueDataCollection = InputDataCollection;
+			}
+		}
+
+		if (UniqueDataCollection && UniqueDataCollection->GetDescription().DataDescs.Num() != DataBinding->InputDataCollection.TaggedData.Num())
+		{
+			// We're not using all the data items from the buffer, no reuse for now.
+			// TODO: Support indirection of data items.
+			UniqueDataCollection.Reset();
+		}
+
+		if (UniqueDataCollection)
+		{
+			ExternalBufferForReuse = UniqueDataCollection->GetBuffer();
+		}
+	}
+
+	/** Do any preparation work such as data readbacks. Returns true when preparation is complete. */
+	bool PrepareData_GameThread()
+	{
+		check(IsInGameThread());
+
+		if (!ExternalBufferForReuse)
+		{
+			UPCGDataBinding* Binding = DataBinding.Get();
+			if (!Binding)
+			{
+				return true;
+			}
+
+			const uint64 BufferSize = PinDesc.ComputePackedSizeBytes();
+			if (PCGComputeHelpers::IsBufferSizeTooLarge(BufferSize))
+			{
+				return true;
+			}
+
+			bool bAllDataReady = true;
+
+			for (FPCGTaggedData& Data : Binding->InputDataCollection.TaggedData)
+			{
+				if (const UPCGProxyForGPUData* Proxy = Cast<UPCGProxyForGPUData>(Data.Data))
+				{
+					UPCGProxyForGPUData::FReadbackResult Result = Proxy->GetCPUData(nullptr);
+					if (!Result.bComplete)
+					{
+						bAllDataReady = false;
+					}
+					else
+					{
+						Data.Data = Result.Data;
+					}
+				}
+			}
+
+			if (!bAllDataReady)
+			{
+				return false;
+			}
+
+			PinDesc.PackDataCollection(Binding->InputDataCollection, PinLabel, Binding->GetStringTable(), PackedDataCollection);
+		}
+
+		return true;
+	}
+
+	/** Gets the buffer that can then be used to read the data in kernels. */
+	FRDGBufferRef GetBuffer(FRDGBuilder& InGraphBuilder, EPCGExportMode InExportMode, const TCHAR* InBufferDebugName)
+	{
+		FRDGBufferRef Buffer = nullptr;
+
+		if (ExternalBufferForReuse)
+		{
+			ensure(PackedDataCollection.IsEmpty());
+
+			Buffer = InGraphBuilder.RegisterExternalBuffer(ExternalBufferForReuse);
+		}
+		else
+		{
+			FRDGBufferDesc Desc = FRDGBufferDesc::CreateByteAddressDesc(sizeof(uint32) * PackedDataCollection.Num());
+			if (InExportMode != EPCGExportMode::NoExport)
+			{
+				// We don't know for sure whether buffer will be read back or not, so need to flag the possibility if the buffer will be passed downstream.
+				Desc.Usage |= BUF_SourceCopy;
+			}
+
+			Buffer = InGraphBuilder.CreateBuffer(Desc, InBufferDebugName);
+
+			InGraphBuilder.QueueBufferUpload(Buffer, PackedDataCollection.GetData(), PackedDataCollection.Num() * PackedDataCollection.GetTypeSize(), ERDGInitialDataFlags::None);
+		}
+
+		return Buffer;
+	}
+
+private:
+	TRefCountPtr<FRDGPooledBuffer> ExternalBufferForReuse;
+
+	TWeakObjectPtr<UPCGDataBinding> DataBinding;
+
+	FPCGDataCollectionDesc PinDesc;
+
+	FName PinLabel;
+
+	/** If data cannot be reused from GPU, will be packed into this buffer prior to upload. */
+	TArray<uint32> PackedDataCollection;
+};
+
 UComputeDataProvider* UPCGDataCollectionUploadDataInterface::CreateDataProvider(TObjectPtr<UObject> InBinding, uint64 InInputMask, uint64 InOutputMask) const
 {
 	TRACE_CPUPROFILER_EVENT_SCOPE(UPCGDataCollectionUploadDataInterface::CreateDataProvider);
@@ -31,24 +193,22 @@ UComputeDataProvider* UPCGDataCollectionUploadDataInterface::CreateDataProvider(
 
 	Provider->DownstreamInputPinLabels = DownstreamInputPinLabelAliases;
 
+	check(!DownstreamInputPinLabelAliases.IsEmpty());
+	Provider->DataAdaptor = MakeShared<FPCGDataUploadAdaptor>(Binding, Provider->PinDesc, DownstreamInputPinLabelAliases[0]);
+
 	return Provider;
 }
 
+bool UPCGDataProviderDataCollectionUpload::PrepareForExecute_GameThread()
+{
+	return !DataAdaptor || DataAdaptor->PrepareData_GameThread();
+}
+
 FComputeDataProviderRenderProxy* UPCGDataProviderDataCollectionUpload::GetRenderProxy()
 {
 	TRACE_CPUPROFILER_EVENT_SCOPE(UPCGDataProviderDataCollectionUpload::GetRenderProxy);
-	TArray<uint32> PackedDataCollection;
 
-	// Use any downstream input pin label to grab data from the collection.
-	check(!DownstreamInputPinLabels.IsEmpty());
-	PinDesc.PackDataCollection(Binding->InputDataCollection, DownstreamInputPinLabels[0], Binding->GetStringTable(), PackedDataCollection);
-
-	if (PCGComputeHelpers::IsBufferSizeTooLarge(PackedDataCollection.Num() * PackedDataCollection.GetTypeSize()))
-	{
-		return nullptr;
-	}
-
-	FPCGDataProviderDataCollectionUploadProxy* Proxy = new FPCGDataProviderDataCollectionUploadProxy(PinDesc, MoveTemp(PackedDataCollection), DownstreamInputPinLabels);
+	FPCGDataProviderDataCollectionUploadProxy* Proxy = new FPCGDataProviderDataCollectionUploadProxy(DataAdaptor, this, PinDesc, DownstreamInputPinLabels, OutputPinLabel, OutputPinLabelAlias, ExportMode);
 
 #if WITH_EDITOR
 	if (const UPCGNode* Node = Cast<UPCGNode>(ProducerSettings ? ProducerSettings->GetOuter() : nullptr))
@@ -64,9 +224,16 @@ FComputeDataProviderRenderProxy* UPCGDataProviderDataCollectionUpload::GetRender
 	return Proxy;
 }
 
-FPCGDataProviderDataCollectionUploadProxy::FPCGDataProviderDataCollectionUploadProxy(const FPCGDataCollectionDesc& InPinDesc, TArray<uint32>&& InPackedDataCollection, const TArray<FName>& InDownstreamInputPinLabels)
-	: FPCGDataCollectionDataProviderProxy(InPinDesc, EPCGReadbackMode::None)
-	, PackedDataCollection(MoveTemp(InPackedDataCollection))
+FPCGDataProviderDataCollectionUploadProxy::FPCGDataProviderDataCollectionUploadProxy(
+	TSharedPtr<FPCGDataUploadAdaptor> InDataAdaptor,
+	TWeakObjectPtr<UPCGDataCollectionDataProvider> InDataProvider,
+	const FPCGDataCollectionDesc& InPinDesc,
+	const TArray<FName>& InDownstreamInputPinLabels,
+	FName InOutputPinLabel,
+	FName InOutputPinLabelAlias,
+	EPCGExportMode InExportMode)
+	: FPCGDataCollectionDataProviderProxy(InDataProvider, InPinDesc, InExportMode, InOutputPinLabel, InOutputPinLabelAlias)
+	, DataAdaptor(InDataAdaptor)
 	, DownstreamInputPinLabels(InDownstreamInputPinLabels)
 {
 }
@@ -82,14 +249,6 @@ void FPCGDataProviderDataCollectionUploadProxy::AllocateResources(FRDGBuilder& G
 		TEXT("PCGDataCollectionUploadBuffer");
 #endif
 
-	FRDGBufferDesc Desc = FRDGBufferDesc::CreateByteAddressDesc(sizeof(uint32) * PackedDataCollection.Num());
-	if (ReadbackMode != EPCGReadbackMode::None)
-	{
-		Desc.Usage |= BUF_SourceCopy;
-	}
-
-	Buffer = GraphBuilder.CreateBuffer(Desc, DebugName);
+	Buffer = DataAdaptor->GetBuffer(GraphBuilder, ExportMode, DebugName);
 	BufferUAV = GraphBuilder.CreateUAV(Buffer);
-
-	GraphBuilder.QueueBufferUpload(Buffer, PackedDataCollection.GetData(), PackedDataCollection.Num() * PackedDataCollection.GetTypeSize(), ERDGInitialDataFlags::None);
 }
diff --git a/Engine/Plugins/PCG/Source/PCG/Private/Compute/DataInterfaces/PCGDataCollectionUploadDataInterface.h b/Engine/Plugins/PCG/Source/PCG/Private/Compute/DataInterfaces/PCGDataCollectionUploadDataInterface.h
--- a/Engine/Plugins/PCG/Source/PCG/Private/Compute/DataInterfaces/PCGDataCollectionUploadDataInterface.h
+++ b/Engine/Plugins/PCG/Source/PCG/Private/Compute/DataInterfaces/PCGDataCollectionUploadDataInterface.h
@@ -1,7 +1,5 @@
 // Copyright Epic Games, Inc. All Rights Reserved.
 
-// Based on OptimusDataInterfaceSkinnedMesh.h
-
 #pragma once
 
 #include "Compute/DataInterfaces/PCGDataCollectionDataInterface.h"
@@ -10,6 +8,8 @@
 
 #include "PCGDataCollectionUploadDataInterface.generated.h"
 
+class FPCGDataUploadAdaptor;
+
 /** Compute Framework Data Interface for reading PCG data. */
 UCLASS(ClassGroup = (Procedural))
 class UPCGDataCollectionUploadDataInterface : public UPCGDataCollectionDataInterface
@@ -33,19 +33,35 @@ public:
 	FComputeDataProviderRenderProxy* GetRenderProxy() override;
 	//~ End UComputeDataProvider Interface
 
+	//~ Begin UPCGComputeDataProvider Interface
+	virtual bool PrepareForExecute_GameThread() override;
+	//~ End UPCGComputeDataProvider Interface
+
 	TArray<FName> DownstreamInputPinLabels;
+
+	TSharedPtr<FPCGDataUploadAdaptor> DataAdaptor;
 };
 
 class FPCGDataProviderDataCollectionUploadProxy : public FPCGDataCollectionDataProviderProxy
 {
 public:
-	FPCGDataProviderDataCollectionUploadProxy(const FPCGDataCollectionDesc& InPinDesc, TArray<uint32>&& PackedDataCollection, const TArray<FName>& InDownstreamInputPinLabels);
+	FPCGDataProviderDataCollectionUploadProxy(
+		TSharedPtr<FPCGDataUploadAdaptor> InDataAdaptor,
+		TWeakObjectPtr<UPCGDataCollectionDataProvider> InDataProvider,
+		const FPCGDataCollectionDesc& InPinDesc,
+		const TArray<FName>& InDownstreamInputPinLabels,
+		FName InOutputPinLabel,
+		FName InOutputPinLabelAlias,
+		EPCGExportMode InExportMode);
 
 	//~ Begin FComputeDataProviderRenderProxy Interface
 	void AllocateResources(FRDGBuilder& GraphBuilder, FAllocationData const& InAllocationData) override;
 	//~ End FComputeDataProviderRenderProxy Interface
 
-	TArray<uint32> PackedDataCollection;
+protected:
+	TSharedPtr<FPCGDataUploadAdaptor> DataAdaptor;
 
 	TArray<FName> DownstreamInputPinLabels;
+
+	TRefCountPtr<FRDGPooledBuffer> ExternalBufferForReuse;
 };
diff --git a/Engine/Plugins/PCG/Source/PCG/Private/Compute/Elements/PCGComputeGraphElement.cpp b/Engine/Plugins/PCG/Source/PCG/Private/Compute/Elements/PCGComputeGraphElement.cpp
--- a/Engine/Plugins/PCG/Source/PCG/Private/Compute/Elements/PCGComputeGraphElement.cpp
+++ b/Engine/Plugins/PCG/Source/PCG/Private/Compute/Elements/PCGComputeGraphElement.cpp
@@ -9,6 +9,7 @@
 #include "PCGSubsystem.h"
 #include "Components/PCGProceduralISMComponent.h"
 #include "Compute/PCGComputeCommon.h"
+#include "Compute/Data/PCGProxyForGPUData.h"
 #include "Compute/DataInterfaces/PCGDataCollectionDataInterface.h"
 #include "Elements/PCGStaticMeshSpawner.h"
 #include "Engine/StaticMesh.h"
@@ -36,6 +37,11 @@ namespace PCGComputeGraphElement
 void FPCGComputeGraphContext::AddExtraStructReferencedObjects(FReferenceCollector& Collector)
 {
 	Collector.AddPropertyReferences(FComputeGraphInstance::StaticStruct(), &ComputeGraphInstance);
+
+	for (TObjectPtr<UComputeDataProvider> Provider : ProvidersWithBufferExports)
+	{
+		Collector.AddReferencedObject(Provider);
+	}
 }
 
 #if WITH_EDITOR
@@ -50,6 +56,7 @@ bool FPCGComputeGraphElement::ExecuteInternal(FPCGContext* InContext) const
 	TRACE_CPUPROFILER_EVENT_SCOPE(FPCGComputeGraphElement::ExecuteInternal);
 	check(InContext);
 	check(InContext->SourceComponent.IsValid());
+	check(IsInGameThread());
 
 	FPCGComputeGraphContext* Context = static_cast<FPCGComputeGraphContext*>(InContext);
 
@@ -66,7 +73,7 @@ bool FPCGComputeGraphElement::ExecuteInternal(FPCGContext* InContext) const
 
 		// Higen is always disabled within dynamic subgraphs - will retrieve Uninitialized tasks (which are always cooked).
 		const UPCGGraph* DynamicSubgraph = InContext->Stack ? InContext->Stack->GetNearestDynamicSubgraphForCurrentFrame() : nullptr;
-			
+
 		if (TopGraph->IsHierarchicalGenerationEnabled() && !DynamicSubgraph)
 		{
 			if (Component->IsLocalComponent() || Component->IsPartitioned())
@@ -90,33 +97,6 @@ bool FPCGComputeGraphElement::ExecuteInternal(FPCGContext* InContext) const
 	// The sequence appears out of order as calls so that paths that are executed multiple times (like checking for completion) are as
 	// short as possible.
 
-	// 7. Execution is complete when any async readbacks are complete.
-	if (Context->bAllAsyncOperationsDone)
-	{
-		TRACE_CPUPROFILER_EVENT_SCOPE(Context->bAllAsyncOperationsDone);
-
-		Context->bExecutionSuccess = true;
-
-		for (UComputeDataProvider* DataProvider : Context->ComputeGraphInstance.GetDataProviders())
-		{
-			UPCGDataCollectionDataProvider* PCGDataProvider = Cast<UPCGDataCollectionDataProvider>(DataProvider);
-			if (PCGDataProvider && PCGDataProvider->RequiresReadback())
-			{
-				// Process data for all readbacks, and track whether all succeeded.
-				const bool bProcessResult = PCGDataProvider->ProcessReadBackData(Context);
-				Context->bExecutionSuccess &= bProcessResult;
-			}
-		}
-
-		// Currently we don't output anything if processing any readback data processing failed.
-		if (ensure(Context->bExecutionSuccess) && ensure(Context->DataBinding))
-		{
-			Context->OutputData = Context->DataBinding->OutputDataCollection;
-		}
-
-		return true;
-	}
-
 	auto SleepUntilNextFrame = [](FPCGComputeGraphContext* InContext)
 	{
 		check(InContext);
@@ -133,6 +113,32 @@ bool FPCGComputeGraphElement::ExecuteInternal(FPCGContext* InContext) const
 		});
 	};
 
+	// 7. Execution is complete when graph was enqueued/dispatched, and any exports are prepared.
+	if (Context->bGraphEnqueued && Context->ProvidersWithBufferExports.IsEmpty())
+	{
+#if WITH_EDITOR
+		// Call both to start all required readbacks.
+		const bool bDebugDrawDone = ExecuteDebugDraw(Context);
+		const bool bInspectionDone = StoreDataForInspection(Context);
+
+		if (!bDebugDrawDone || !bInspectionDone)
+		{
+			SleepUntilNextFrame(Context);
+			return false;
+		}
+#endif
+
+		Context->bExecutionSuccess = true;
+
+		// Currently we don't output anything if processing any readback data processing failed.
+		if (ensure(Context->bExecutionSuccess) && ensure(Context->DataBinding))
+		{
+			Context->OutputData = Context->DataBinding->OutputDataCollection;
+		}
+
+		return true;
+	}
+
 	// 3. If still compiling, try again next frame.
 	if (Context->ComputeGraph->HasKernelResourcesPendingShaderCompilation())
 	{
@@ -180,28 +186,39 @@ bool FPCGComputeGraphElement::ExecuteInternal(FPCGContext* InContext) const
 			Context->ComputeGraphInstance.CreateDataProviders(Context->ComputeGraph.Get(), 0, Context->DataBinding.Get());
 		}
 
-		// Register all providers running async operations. TODO review if we should have a general API like "RunsAsyncOperations()"?
-		for (UComputeDataProvider* DataProvider : Context->ComputeGraphInstance.GetDataProviders())
+		// Register all providers running async operations.
+		for (UComputeDataProvider* ComputeDataProvider : Context->ComputeGraphInstance.GetDataProviders())
 		{
-			UPCGDataCollectionDataProvider* PCGDataProvider = Cast<UPCGDataCollectionDataProvider>(DataProvider);
-			if (PCGDataProvider && PCGDataProvider->RequiresReadback())
+			UPCGDataCollectionDataProvider* DataProvider = Cast<UPCGDataCollectionDataProvider>(ComputeDataProvider);
+
+			const bool bDataProviderDebugged = DataProvider && DataProvider->ProducerSettings && DataProvider->ProducerSettings->bDebug;
+
+			if (DataProvider && (DataProvider->IsExportRequired() || bDataProviderDebugged))
 			{
-				Context->ProvidersRunningAsyncOperations.Add(PCGDataProvider);
+				Context->ProvidersWithBufferExports.Add(DataProvider);
 
-				PCGDataProvider->OnReadbackComplete_RenderThread().AddLambda([Context, PCGDataProvider]()
-				{
-					FWriteScopeLock Lock(Context->ProvidersRunningAsyncOperationsLock);
+				TWeakObjectPtr<UPCGDataCollectionDataProvider> DataProviderWeak = DataProvider;
 
-					const bool bEmptyBefore = Context->ProvidersRunningAsyncOperations.IsEmpty();
+				DataProvider->OnDataExported_GameThread().AddLambda([ContextHandle = Context->GetOrCreateHandle(), DataProviderWeak]()
+				{
+					check(IsInGameThread());
 
-					ensure(Context->ProvidersRunningAsyncOperations.Contains(PCGDataProvider));
-					Context->ProvidersRunningAsyncOperations.Remove(PCGDataProvider);
+					TSharedPtr<FPCGContextHandle> SharedHandle = ContextHandle.Pin();
+					FPCGComputeGraphContext* ContextPtr = static_cast<FPCGComputeGraphContext*>(SharedHandle ? SharedHandle->GetContext() : nullptr);
+					UPCGDataCollectionDataProvider* DataProvider = DataProviderWeak.Get();
 
-					if (!bEmptyBefore && Context->ProvidersRunningAsyncOperations.IsEmpty())
+					if (!DataProvider || !ContextPtr)
 					{
-						Context->bAllAsyncOperationsDone = true;
+						// Safe to just jump out. The GPU buffer is ref counted.
+						return;
+					}
+
+					ensure(ContextPtr->ProvidersWithBufferExports.Contains(DataProvider));
+					ContextPtr->ProvidersWithBufferExports.Remove(DataProvider);
 
-						Context->bIsPaused = false;
+					if (ContextPtr->ProvidersWithBufferExports.IsEmpty())
+					{
+						ContextPtr->bIsPaused = false;
 					}
 				});
 			}
@@ -248,6 +265,28 @@ bool FPCGComputeGraphElement::ExecuteInternal(FPCGContext* InContext) const
 		Context->bPrimitiveProxiesValidated = true;
 	}
 
+	if (!Context->bDataProvidersReady)
+	{
+		bool bAllReady = true;
+
+		for (UComputeDataProvider* ComputeDataProvider : Context->ComputeGraphInstance.GetDataProviders())
+		{
+			if (UPCGComputeDataProvider* DataProvider = Cast<UPCGComputeDataProvider>(ComputeDataProvider))
+			{
+				if (!DataProvider->PrepareForExecute_GameThread())
+				{
+					bAllReady = false;
+				}
+			}
+		}
+
+		if (!bAllReady)
+		{
+			SleepUntilNextFrame(Context);
+			return false;
+		}
+	}
+
 	// 1. Prepare render resources. In editor, this will trigger shader compile if not compiled already.
 	if (!Context->ComputeGraph->GetRenderProxy())
 	{
@@ -296,11 +335,9 @@ bool FPCGComputeGraphElement::ExecuteInternal(FPCGContext* InContext) const
 
 	if (Context->bGraphEnqueued)
 	{
-		FReadScopeLock Lock(Context->ProvidersRunningAsyncOperationsLock);
-
-		if (!Context->ProvidersRunningAsyncOperations.IsEmpty())
+		if (!Context->ProvidersWithBufferExports.IsEmpty())
 		{
-			// If we're running async operations (like readbacks), go to sleep and let them wake us up later.
+			// If we're running async operations (like preparing exports), go to sleep and let them wake us up later.
 			Context->bIsPaused = true;
 			return false;
 		}
@@ -649,25 +686,138 @@ void FPCGComputeGraphElement::AbortInternal(FPCGContext* InContext) const
 
 void FPCGComputeGraphElement::ResetAsyncOperations(FPCGContext* InContext) const
 {
+	check(IsInGameThread());
+
 	if (InContext)
 	{
-		FPCGComputeGraphContext* Context = static_cast<FPCGComputeGraphContext*>(InContext);
+		static_cast<FPCGComputeGraphContext*>(InContext)->ProvidersWithBufferExports.Reset();
+	}
+}
+
+#if WITH_EDITOR
+bool FPCGComputeGraphElement::ExecuteDebugDraw(FPCGComputeGraphContext* InContext) const
+{
+	if (!InContext || !InContext->DataBinding || InContext->DataBinding->DataToDebug.IsEmpty())
+	{
+		return true;
+	}
+
+	const FPCGDataVisualizationRegistry& DataVisRegistry = FPCGModule::GetConstPCGDataVisualizationRegistry();
+
+	bool bAllDone = true;
 
-		for (UComputeDataProvider* DataProvider : Context->ComputeGraphInstance.GetDataProviders())
+	for (int Index = InContext->DataBinding->DataToDebug.Num() - 1; Index >= 0; --Index)
+	{
+		FPCGDataToDebug& DataToDebug = InContext->DataBinding->DataToDebug[Index];
+		const UPCGProxyForGPUData* Proxy = DataToDebug.Proxy.Get();
+		const UPCGSettings* ProducerSettings = DataToDebug.ProducerSettings.Get();
+
+		if (Proxy && ProducerSettings)
 		{
-			UPCGDataCollectionDataProvider* PCGDataProvider = Cast<UPCGDataCollectionDataProvider>(DataProvider);
-			if (PCGDataProvider && PCGDataProvider->RequiresReadback())
+			UPCGProxyForGPUData::FReadbackResult Readback = Proxy->GetCPUData(InContext);
+			if (Readback.bComplete)
+			{
+				if (Readback.Data)
+				{
+					if (const IPCGDataVisualization* DataVis = DataVisRegistry.GetDataVisualization(Readback.Data->GetClass()))
+					{
+						DataVis->ExecuteDebugDisplay(InContext, ProducerSettings, Readback.Data, InContext->GetTargetActor(nullptr));
+					}
+				}
+
+				InContext->DataBinding->DataToDebug.RemoveAtSwap(Index);
+			}
+			else
 			{
-				PCGDataProvider->OnReadbackComplete_RenderThread().Clear();
+				// Readback needs another frame to complete.
+				bAllDone = false;
 			}
 		}
+		else
+		{
+			// Could not obtain proxy, remove this data.
+			InContext->DataBinding->DataToDebug.RemoveAtSwap(Index);
+		}
+	}
 
-		FWriteScopeLock Lock(Context->ProvidersRunningAsyncOperationsLock);
-		Context->ProvidersRunningAsyncOperations.Reset();
+	if (bAllDone)
+	{
+		InContext->DataBinding->DataToDebug.Empty();
 	}
+
+	return bAllDone;
+}
+
+bool FPCGComputeGraphElement::StoreDataForInspection(FPCGComputeGraphContext* InContext) const
+{
+	if (!InContext || !InContext->DataBinding || InContext->DataBinding->DataToInspect.IsEmpty() || !InContext->SourceComponent.IsValid())
+	{
+		return true;
+	}
+
+	const FPCGDataVisualizationRegistry& DataVisRegistry = FPCGModule::GetConstPCGDataVisualizationRegistry();
+
+	bool bAllReady = true;
+
+	// First ensure all are read back.
+	for (FPCGDataToDebug& DataToInspect : InContext->DataBinding->DataToInspect)
+	{
+		const UPCGProxyForGPUData* Proxy = DataToInspect.Proxy.Get();
+		const UPCGSettings* ProducerSettings = DataToInspect.ProducerSettings.Get();
+
+		if (Proxy && ProducerSettings)
+		{
+			UPCGProxyForGPUData::FReadbackResult Readback = Proxy->GetCPUData(InContext);
+			if (!Readback.bComplete)
+			{
+				bAllReady = false;
+			}
+		}
+	}
+
+	if (!bAllReady)
+	{
+		return false;
+	}
+
+	// Now collect all data into a collection and store it.
+	TMap<const UPCGSettings*, FPCGDataCollection> SettingsToDataCollection;
+
+	for (FPCGDataToDebug& DataToInspect : InContext->DataBinding->DataToInspect)
+	{
+		const UPCGProxyForGPUData* Proxy = DataToInspect.Proxy.Get();
+		const UPCGSettings* ProducerSettings = DataToInspect.ProducerSettings.Get();
+
+		if (Proxy && ProducerSettings)
+		{
+			UPCGProxyForGPUData::FReadbackResult Readback = Proxy->GetCPUData(InContext);
+			if (ensure(Readback.bComplete && Readback.Data))
+			{
+				FPCGDataCollection& DataCollection = SettingsToDataCollection.FindOrAdd(ProducerSettings);
+				FPCGTaggedData Data;
+				Data.Data = Readback.Data;
+				Data.Pin = DataToInspect.PinLabel;
+				DataCollection.TaggedData.Add(Data);
+			}
+		}
+	}
+
+	for (TPair<const UPCGSettings*, FPCGDataCollection>& SettingsAndData : SettingsToDataCollection)
+	{
+		// Required by inspection code.
+		SettingsAndData.Get<1>().ComputeCrcs(/*bFullDataCrc=*/false);
+
+		const UPCGNode* Node = Cast<UPCGNode>(SettingsAndData.Get<0>()->GetOuter());
+
+		// TODO: Input data not yet supported.
+		InContext->SourceComponent->StoreInspectionData(InContext->Stack, Node, /*InTimer=*/nullptr, /*InInputData=*/{}, SettingsAndData.Get<1>(), /*bUsedCache*/false);
+	}
+
+	InContext->DataBinding->DataToInspect.Empty();
+
+	return bAllReady;
 }
 
-#if WITH_EDITOR
 void FPCGComputeGraphElement::LogCompilationMessages(FPCGComputeGraphContext* InContext) const
 {
 	if (InContext->SourceComponent.IsValid() && InContext->Stack)
diff --git a/Engine/Plugins/PCG/Source/PCG/Private/Compute/Elements/PCGComputeGraphElement.h b/Engine/Plugins/PCG/Source/PCG/Private/Compute/Elements/PCGComputeGraphElement.h
--- a/Engine/Plugins/PCG/Source/PCG/Private/Compute/Elements/PCGComputeGraphElement.h
+++ b/Engine/Plugins/PCG/Source/PCG/Private/Compute/Elements/PCGComputeGraphElement.h
@@ -28,21 +28,21 @@ public:
 	/** Data providers created from data interfaces and data bindings. */
 	FComputeGraphInstance ComputeGraphInstance;
 
-	/** Keep track of data providers that perform async operations and require multiple frames to complete. */
-	TSet<UComputeDataProvider*> ProvidersRunningAsyncOperations;
-	mutable FRWLock ProvidersRunningAsyncOperationsLock;
-
 	/** Graph enqueued (scheduled for execution by GPU). */
 	bool bGraphEnqueued = false;
 
-	/** All async operations complete and results processed. */
-	bool bAllAsyncOperationsDone = false;
-
 	/** Graph executed successfully. */
 	bool bExecutionSuccess = false;
 
 	/** Whether all required primitive proxies have been created. */
 	bool bPrimitiveProxiesValidated = false;
+
+	bool bDataProvidersReady = false;
+
+	/** Data providers with buffers that are passed to downstream tasks. The buffer will be created on the render thread, and then passed
+	* back to main thread, upon which a reference is taken to the buffer and the provider is removed from this set to signal completion.
+	*/
+	TSet<TObjectPtr<UComputeDataProvider>> ProvidersWithBufferExports;
 };
 
 /** Executes a CF graph. Created by the compiler when collapsing GPU nodes rather than by a settings/node. */
@@ -72,6 +72,7 @@ protected:
 	virtual bool ExecuteInternal(FPCGContext* InContext) const override;
 	virtual void PostExecuteInternal(FPCGContext* InContext) const override;
 	virtual void AbortInternal(FPCGContext* InContext) const override;;
+	virtual bool SupportsGPUResidentData(FPCGContext* InContext) const override { return true; }
 
 	// The calls to initialize the compute graph are not thread safe.
 	virtual bool CanExecuteOnlyOnMainThread(FPCGContext* Context) const override { return true; }
@@ -84,6 +85,12 @@ protected:
 	void ResetAsyncOperations(FPCGContext* InContext) const;
 
 #if WITH_EDITOR
+	// TODO: Debug draw has a special path here because the compute graph element represents a set of original elements which do not themselves
+	// execute. Review if this can be reconciled or unified with the normal debug draw path.
+	bool ExecuteDebugDraw(FPCGComputeGraphContext* InContext) const;
+
+	bool StoreDataForInspection(FPCGComputeGraphContext* InContext) const;
+
 	void LogCompilationMessages(FPCGComputeGraphContext* InContext) const;
 #endif
 };
diff --git a/Engine/Plugins/PCG/Source/PCG/Private/Compute/PCGComputeCommon.cpp b/Engine/Plugins/PCG/Source/PCG/Private/Compute/PCGComputeCommon.cpp
--- a/Engine/Plugins/PCG/Source/PCG/Private/Compute/PCGComputeCommon.cpp
+++ b/Engine/Plugins/PCG/Source/PCG/Private/Compute/PCGComputeCommon.cpp
@@ -1,11 +1,12 @@
 // Copyright Epic Games, Inc. All Rights Reserved.
 
-#include "PCGComputeCommon.h"
+#include "Compute/PCGComputeCommon.h"
 
 #include "PCGComponent.h"
 #include "PCGContext.h"
 #include "PCGParamData.h"
 #include "PCGSubsystem.h"
+#include "Compute/Data/PCGProxyForGPUData.h"
 #include "Compute/PCGPinPropertiesGPU.h"
 #include "Data/PCGPointData.h"
 
@@ -32,6 +33,10 @@ namespace PCGComputeHelpers
 				return Metadata->GetItemCountForChild();
 			}
 		}
+		else if (const UPCGProxyForGPUData* Proxy = Cast<UPCGProxyForGPUData>(InData))
+		{
+			return Proxy->GetElementCount();
+		}
 
 		return 0;
 	}
@@ -322,4 +327,17 @@ namespace PCGComputeHelpers
 
 		return bBufferTooLarge;
 	}
+
+	bool ComputeDataDescription(const UPCGData* InData, const UPCGDataBinding* InBinding, FPCGDataDesc& OutDescription)
+	{
+		if (const UPCGProxyForGPUData* Proxy = Cast<UPCGProxyForGPUData>(InData))
+		{
+			return Proxy->GetDescription(OutDescription);
+		}
+		else
+		{
+			OutDescription = FPCGDataDesc(InData, InBinding);
+			return true;
+		}
+	}
 }
diff --git a/Engine/Plugins/PCG/Source/PCG/Private/Compute/PCGComputeCommon.h b/Engine/Plugins/PCG/Source/PCG/Private/Compute/PCGComputeCommon.h
deleted file mode 100644
diff --git a/Engine/Plugins/PCG/Source/PCG/Private/Compute/PCGDataBinding.cpp b/Engine/Plugins/PCG/Source/PCG/Private/Compute/PCGDataBinding.cpp
--- a/Engine/Plugins/PCG/Source/PCG/Private/Compute/PCGDataBinding.cpp
+++ b/Engine/Plugins/PCG/Source/PCG/Private/Compute/PCGDataBinding.cpp
@@ -7,6 +7,7 @@
 #include "PCGSettings.h"
 #include "Compute/PCGComputeCommon.h"
 #include "Compute/PCGComputeGraph.h"
+#include "Compute/Data/PCGProxyForGPUData.h"
 #include "Data/PCGPointData.h"
 #include "Metadata/PCGMetadata.h"
 
@@ -53,6 +54,37 @@ bool UPCGDataBinding::ComputeCPUOutputPinDataDesc(const UPCGPin* OutputPin, FPCG
 	return true;
 }
 
+void UPCGDataBinding::ReceiveDataFromGPU_GameThread(
+	const UPCGProxyForGPUData* InData,
+	const UPCGSettings* InProducerSettings,
+	EPCGExportMode InExportMode,
+	FName InPinLabel,
+	FName InPinLabelAlias)
+{
+	check(IsInGameThread());
+
+	if (!!(InExportMode & EPCGExportMode::ComputeGraphOutput))
+	{
+		FPCGTaggedData& TaggedData = OutputDataCollection.TaggedData.Emplace_GetRef();
+		TaggedData.Pin = InPinLabelAlias;
+		TaggedData.Data = InData;
+	}
+
+	if (!!(InExportMode & EPCGExportMode::DebugVisualization))
+	{
+		TStrongObjectPtr<const UPCGProxyForGPUData> Proxy;
+		Proxy.Reset(InData);
+		DataToDebug.Add(FPCGDataToDebug{ .Proxy = Proxy, .ProducerSettings = InProducerSettings, .PinLabel = InPinLabel });
+	}
+
+	if (!!(InExportMode & EPCGExportMode::Inspection))
+	{
+		TStrongObjectPtr<const UPCGProxyForGPUData> Proxy;
+		Proxy.Reset(InData);
+		DataToInspect.Add(FPCGDataToDebug{ .Proxy = Proxy, .ProducerSettings = InProducerSettings, .PinLabel = InPinLabel });
+	}
+}
+
 void UPCGDataBinding::InitializeInputData(const FPCGDataCollection& InComputeGraphElementInputData)
 {
 	TRACE_CPUPROFILER_EVENT_SCOPE(UPCGDataBinding::InitializeInputData);
diff --git a/Engine/Plugins/PCG/Source/PCG/Private/Compute/PCGDataBinding.h b/Engine/Plugins/PCG/Source/PCG/Private/Compute/PCGDataBinding.h
--- a/Engine/Plugins/PCG/Source/PCG/Private/Compute/PCGDataBinding.h
+++ b/Engine/Plugins/PCG/Source/PCG/Private/Compute/PCGDataBinding.h
@@ -9,6 +9,7 @@
 
 class UPCGComponent;
 class UPCGComputeGraph;
+class UPCGProxyForGPUData;
 
 USTRUCT()
 struct FPCGSpawnerPrimitives
@@ -43,6 +44,13 @@ struct FPCGSpawnerPrimitives
 	uint32 NumInstancesAllPrimitives = 0;
 };
 
+struct FPCGDataToDebug
+{
+	TStrongObjectPtr<const UPCGProxyForGPUData> Proxy;
+	TWeakObjectPtr<const UPCGSettings> ProducerSettings;
+	FName PinLabel;
+};
+
 UCLASS(Transient, Category = PCG)
 class UPCGDataBinding : public UObject
 {
@@ -60,6 +68,8 @@ public:
 	/** If the output pin is a connected upstream pin of the compute graph (on a CPU node or a GPU node from a different compute graph), returns a description of this data. */
 	bool ComputeCPUOutputPinDataDesc(const UPCGPin* OutputPin, FPCGDataCollectionDesc& OutDesc) const;
 
+	void ReceiveDataFromGPU_GameThread(const UPCGProxyForGPUData* InData, const UPCGSettings* InProducerSettings, EPCGExportMode InExportMode, FName InPinLabel, FName InPinLabelAlias);
+
 private:
 	void InitializeInputData(const FPCGDataCollection& InComputeGraphElementInputData);
 
@@ -98,6 +108,9 @@ public:
 	UPROPERTY()
 	TMap<TObjectPtr<const UPCGSettings>, FPCGSpawnerPrimitives> MeshSpawnersToPrimitives;
 
+	TArray<FPCGDataToDebug> DataToDebug;
+	TArray<FPCGDataToDebug> DataToInspect;
+
 protected:
 	/** All attributes present in graph at execution time. Coherent across all graph branches. Seeded using statically known created attributes at compile time, then
 	 * augmented with incoming attributes from input data collection at runtime. */
diff --git a/Engine/Plugins/PCG/Source/PCG/Private/Compute/PCGDataForGPU.cpp b/Engine/Plugins/PCG/Source/PCG/Private/Compute/PCGDataForGPU.cpp
--- a/Engine/Plugins/PCG/Source/PCG/Private/Compute/PCGDataForGPU.cpp
+++ b/Engine/Plugins/PCG/Source/PCG/Private/Compute/PCGDataForGPU.cpp
@@ -12,17 +12,27 @@
 #include "Compute/PCGComputeCommon.h"
 #include "Compute/PCGComputeGraph.h"
 #include "Compute/PCGDataBinding.h"
+#include "Compute/Data/PCGProxyForGPUData.h"
 #include "Data/PCGPointData.h"
 #include "Helpers/PCGAsync.h"
 #include "Metadata/PCGMetadata.h"
 #include "Metadata/Accessors/PCGAttributeAccessorHelpers.h"
+#include "Utils/PCGLogErrors.h"
 
+#include "RenderGraphResources.h"
 #include "Async/ParallelFor.h"
 
 #define LOCTEXT_NAMESPACE "PCGDataForGPU"
 
 using namespace PCGComputeConstants;
 
+#if WITH_EDITOR
+static TAutoConsoleVariable<bool> CVarWarnOnGPUReadbacks(
+	TEXT("pcg.Graph.GPU.WarnOnGPUReadbacks"),
+	false,
+	TEXT("Emits warnings on nodes which trigger a readback of data from GPU to CPU."));
+#endif
+
 namespace PCGDataForGPUConstants
 {
 	const static FPCGKernelAttributeDesc PointPropertyDescs[NUM_POINT_PROPERTIES] =
@@ -908,7 +918,11 @@ FPCGDataCollectionDesc FPCGDataCollectionDesc::BuildFromInputDataCollectionAndIn
 			continue;
 		}
 
-		CollectionDesc.DataDescs.Emplace(Data.Data, InBinding);
+		FPCGDataDesc DataDesc;
+		if (ensure(PCGComputeHelpers::ComputeDataDescription(Data.Data, InBinding, DataDesc)))
+		{
+			CollectionDesc.DataDescs.Add(MoveTemp(DataDesc));
+		}
 	}
 
 	return CollectionDesc;
@@ -1044,6 +1058,12 @@ void FPCGDataCollectionDesc::PackDataCollection(const FPCGDataCollection& InData
 	{
 		const FPCGDataDesc& DataDesc = DataDescs[DataIndex];
 
+		if (Cast<UPCGProxyForGPUData>(InputData[DataIndex].Data))
+		{
+			UE_LOG(LogPCG, Error, TEXT("Attempted to pack a data that is not resident on the CPU. Uploaded data will be uninitialized."));
+			continue;
+		}
+
 		// No work to do if there are no elements to process.
 		if (DataDesc.ElementCount <= 0)
 		{
@@ -1598,4 +1618,110 @@ void FPCGDataCollectionDesc::AddAttributeToAllData(FPCGKernelAttributeKey InAttr
 	}
 }
 
+FPCGProxyForGPUDataCollection::FPCGProxyForGPUDataCollection(TRefCountPtr<FRDGPooledBuffer> InBuffer, uint32 InBufferSizeBytes, const FPCGDataCollectionDesc& InDescription, const TArray<FString>& InStringTable)
+	: Buffer(InBuffer)
+	, BufferSizeBytes(InBufferSizeBytes)
+	, Description(InDescription)
+	, StringTable(InStringTable)
+{
+}
+
+bool FPCGProxyForGPUDataCollection::GetCPUData(FPCGContext* InContext, int32 InDataIndex, const UPCGData*& OutData)
+{
+	UE::TScopeLock Lock(ReadbackLock);
+
+	if (bReadbackDataProcessed)
+	{
+		OutData = ReadbackData.TaggedData.IsValidIndex(InDataIndex) ? ReadbackData.TaggedData[InDataIndex].Data : nullptr;
+		return true;
+	}
+
+	if (!ReadbackRequest && !bReadbackDataArrived)
+	{
+#if WITH_EDITOR
+		if (CVarWarnOnGPUReadbacks.GetValueOnAnyThread())
+		{
+			PCGLog::LogWarningOnGraph(LOCTEXT("PerformingReadback", "Performing readback"), InContext);
+		}
+#endif
+
+		ReadbackRequest = MakeShared<FRHIGPUBufferReadback>(TEXT("FPCGProxyForGPUDataCollectionReadback"));
+
+		ENQUEUE_RENDER_COMMAND(ReadbackDataCollection)([WeakThis=AsWeak(), ReadbackRequest=ReadbackRequest, Buffer=Buffer, BufferSizeBytes=BufferSizeBytes](FRHICommandListImmediate& RHICmdList)
+		{
+			FRDGPooledBuffer* RDGBuffer = Buffer.GetReference();
+
+			ReadbackRequest->EnqueueCopy(RHICmdList, Buffer->GetRHI());
+
+			auto RunnerFunc = [WeakThis, ReadbackRequest, BufferSizeBytes](auto&& RunnerFunc) -> void
+			{
+				if (TSharedPtr<FPCGProxyForGPUDataCollection> This = WeakThis.Pin())
+				{
+					if (ReadbackRequest->IsReady())
+					{
+						if (void* RawData = ReadbackRequest->Lock(BufferSizeBytes))
+						{
+							uint32* DataAsUint = static_cast<uint32*>(RawData);
+
+							This->RawReadbackData.SetNumUninitialized(BufferSizeBytes);
+							FMemory::Memcpy(This->RawReadbackData.GetData(), RawData, BufferSizeBytes);
+						}
+
+						ReadbackRequest->Unlock();
+
+						This->bReadbackDataArrived = true;
+					}
+					else
+					{
+						AsyncTask(ENamedThreads::ActualRenderingThread, [RunnerFunc]()
+						{
+							RunnerFunc(RunnerFunc);
+						});
+					}
+				}
+			};
+
+			AsyncTask(ENamedThreads::ActualRenderingThread, [RunnerFunc]()
+			{
+				RunnerFunc(RunnerFunc);
+			});
+
+		});
+	}
+
+	if (!bReadbackDataArrived)
+	{
+		OutData = nullptr;
+		return false;
+	}
+
+	ReadbackRequest.Reset();
+
+	bReadbackDataProcessed = true;
+
+	if (!RawReadbackData.IsEmpty())
+	{
+		FPCGDataCollection DataFromGPU;
+		const EPCGUnpackDataCollectionResult Result = Description.UnpackDataCollection(InContext, RawReadbackData, /*OutputPinLabelAlias*/NAME_None, StringTable, DataFromGPU);
+
+		RawReadbackData.Empty();
+
+		if (Result != EPCGUnpackDataCollectionResult::Success)
+		{
+			OutData = nullptr;
+			return true;
+		}
+
+		ReadbackData = MoveTemp(DataFromGPU);
+	}
+
+	OutData = ReadbackData.TaggedData.IsValidIndex(InDataIndex) ? ReadbackData.TaggedData[InDataIndex].Data : nullptr;
+	return true;
+}
+
+void FPCGProxyForGPUDataCollection::AddReferences(FReferenceCollector& Collector)
+{
+	ReadbackData.AddReferences(Collector);
+}
+
 #undef LOCTEXT_NAMESPACE
diff --git a/Engine/Plugins/PCG/Source/PCG/Private/Compute/PCGDataForGPU.h b/Engine/Plugins/PCG/Source/PCG/Private/Compute/PCGDataForGPU.h
deleted file mode 100644
diff --git a/Engine/Plugins/PCG/Source/PCG/Private/Graph/PCGGraphCache.cpp b/Engine/Plugins/PCG/Source/PCG/Private/Graph/PCGGraphCache.cpp
--- a/Engine/Plugins/PCG/Source/PCG/Private/Graph/PCGGraphCache.cpp
+++ b/Engine/Plugins/PCG/Source/PCG/Private/Graph/PCGGraphCache.cpp
@@ -111,6 +111,12 @@ void FPCGGraphCache::StoreInCache(const FPCGStoreInCacheParams& Params, const FP
 		return;
 	}
 
+	// Proxies should never go into the graph cache. These can hold onto large chunks of video memory.
+	for (const FPCGTaggedData& Data : InOutput.TaggedData)
+	{
+		ensure(!Data.Data || Data.Data->IsCacheable());
+	}
+
 	{
 		TRACE_CPUPROFILER_EVENT_SCOPE(FPCGGraphCache::StoreInCache);
 		UE::TScopeLock ScopedLock(CacheLock);
diff --git a/Engine/Plugins/PCG/Source/PCG/Private/Graph/PCGGraphCompilerGPU.cpp b/Engine/Plugins/PCG/Source/PCG/Private/Graph/PCGGraphCompilerGPU.cpp
--- a/Engine/Plugins/PCG/Source/PCG/Private/Graph/PCGGraphCompilerGPU.cpp
+++ b/Engine/Plugins/PCG/Source/PCG/Private/Graph/PCGGraphCompilerGPU.cpp
@@ -603,7 +603,7 @@ void FPCGGraphCompilerGPU::CreateDataInterfaces(
 
 	// Create data interfaces which allow kernels to read or write data. Each data interface is associated with a node output pin.
 	// For CPU->GPU edges, an upload data interface is created. For GPU->CPU edges, a readback data interface is created.
-	auto CreateDataInterface = [&InOutContext, &InCollapsedTasks, &InCompiledTasks, InOutComputeGraph](FPCGTaskId InTaskId, bool bRequiresReadback, const FPCGPinProperties& InOutputPinProperties) -> UPCGComputeDataInterface*
+	auto CreateDataInterface = [&InOutContext, &InCollapsedTasks, &InCompiledTasks, InOutComputeGraph](FPCGTaskId InTaskId, bool bRequiresExport, const FPCGPinProperties& InOutputPinProperties) -> UPCGComputeDataInterface*
 	{
 		TRACE_CPUPROFILER_EVENT_SCOPE(CreateDataInterface);
 
@@ -641,7 +641,7 @@ void FPCGGraphCompilerGPU::CreateDataInterfaces(
 
 			check(DataInterfacePCGData);
 			DataInterfacePCGData->SetProducerSettings(InCompiledTasks[InTaskId].Node ? InCompiledTasks[InTaskId].Node->GetSettings() : nullptr);
-			DataInterfacePCGData->SetRequiresReadback(bRequiresReadback);
+			DataInterfacePCGData->SetRequiresExport(bRequiresExport);
 
 			DataInterface = DataInterfacePCGData;
 			break;
@@ -709,8 +709,8 @@ void FPCGGraphCompilerGPU::CreateDataInterfaces(
 					continue;
 				}
 
-				// Request readback if any downstream task is outside of this compute graph.
-				bool bRequiresReadback = false;
+				// Request buffer export from compute graph if any downstream task is outside of this compute graph.
+				bool bRequiresExport = false;
 				if (const TArray<FPCGTaskId>* Successors = InTaskSuccessors.Find(TaskId))
 				{
 					for (FPCGTaskId Successor : *Successors)
@@ -719,14 +719,14 @@ void FPCGGraphCompilerGPU::CreateDataInterfaces(
 						{
 							if (Input.TaskId == TaskId && Input.UpstreamPin && !InCollapsedTasks.Contains(Successor))
 							{
-								bRequiresReadback = true;
+								bRequiresExport = true;
 								break;
 							}
 						}
 					}
 				}
 
-				if (UPCGComputeDataInterface* OutputDI = CreateDataInterface(TaskId, bRequiresReadback, OutputPinProperties))
+				if (UPCGComputeDataInterface* OutputDI = CreateDataInterface(TaskId, bRequiresExport, OutputPinProperties))
 				{
 					// Get the compute graph virtual pin corresponding to this pin, if one was created.
 					const FName* VirtualPinLabel = InOriginalToVirtualPin.Find({ TaskId, OutputPinProperties.Label, /*bIsInputPin*/false });
@@ -789,7 +789,7 @@ void FPCGGraphCompilerGPU::CreateDataInterfaces(
 			{
 				OutputDI = *FoundOutputDI;
 			}
-			else if (UPCGComputeDataInterface* NewOutputDI = CreateDataInterface(Input.TaskId, /*bRequiresReadback=*/false, *Input.UpstreamPin))
+			else if (UPCGComputeDataInterface* NewOutputDI = CreateDataInterface(Input.TaskId, /*bRequiresExport=*/false, *Input.UpstreamPin))
 			{
 				OutputDI = NewOutputDI;
 
diff --git a/Engine/Plugins/PCG/Source/PCG/Private/Graph/PCGGraphExecutor.cpp b/Engine/Plugins/PCG/Source/PCG/Private/Graph/PCGGraphExecutor.cpp
--- a/Engine/Plugins/PCG/Source/PCG/Private/Graph/PCGGraphExecutor.cpp
+++ b/Engine/Plugins/PCG/Source/PCG/Private/Graph/PCGGraphExecutor.cpp
@@ -90,6 +90,11 @@ namespace PCGGraphExecutor
 		32,
 		TEXT("Max in flight PCG Element tasks allowed at the same time. Note that PCG Element execution can spawn async tasks themselves"));
 
+	TAutoConsoleVariable<bool> CVarPassGPUDataThroughGridLinks(
+		TEXT("pcg.Graph.GPU.PassGPUDataThroughGridLinks"),
+		false,
+		TEXT("Whether proxies for GPU data are cached in per pin output data and passed through grid links. If false data is read back to CPU."));
+
 	void ClearAsyncFlags(TSet<TObjectPtr<UObject>>& AsyncObjects)
 	{
 		// Remove Async Flags from Objects created in Async PCG Tasks so that they can get tracked by the GarbageCollector properly
@@ -898,6 +903,22 @@ void FPCGGraphExecutor::PostTaskExecute(TSharedPtr<FPCGGraphActiveTask> ActiveTa
 	check(ActiveTask.Context);
 	check(ActiveTask.bIsExecutingTask || ActiveTask.bWasCancelled);
 
+	// If input data collection was modified (for example replacing proxies with full data), add any new references.
+	if (ActiveTask.Context->bInputDataModified)
+	{
+		PCGGraphExecutor::TScopeLock ScopeLock(LiveTasksLock);
+
+		for (const FPCGTaggedData& Data : ActiveTask.Context->InputData.TaggedData)
+		{
+			if (Data.Data)
+			{
+				ActiveTask.ExecutingReferences.AddUnique(Data.Data);
+			}
+		}
+
+		ActiveTask.Context->bInputDataModified = false;
+	}
+
 	if (!bIsDone)
 	{
 		TRACE_CPUPROFILER_EVENT_SCOPE(FPCGGraphExecutor::PostTaskExecute::NotDone);
@@ -2578,6 +2599,11 @@ namespace PCGGraphExecutor
 	}
 #endif
 
+	bool FPCGGridLinkageElement::SupportsGPUResidentData(FPCGContext* InContext) const
+	{
+		return CVarPassGPUDataThroughGridLinks.GetValueOnAnyThread();
+	}
+
 	bool ExecuteGridLinkage(
 		EPCGHiGenGrid InGenerationGrid,
 		EPCGHiGenGrid InFromGrid,
diff --git a/Engine/Plugins/PCG/Source/PCG/Private/Graph/PCGGraphExecutor.h b/Engine/Plugins/PCG/Source/PCG/Private/Graph/PCGGraphExecutor.h
--- a/Engine/Plugins/PCG/Source/PCG/Private/Graph/PCGGraphExecutor.h
+++ b/Engine/Plugins/PCG/Source/PCG/Private/Graph/PCGGraphExecutor.h
@@ -438,6 +438,7 @@ namespace PCGGraphExecutor
 
 		//~Begin IPCGElement interface
 		virtual bool IsGridLinkage() const override { return true; }
+		virtual bool SupportsGPUResidentData(FPCGContext* InContext) const override;
 		//~End IPCGElement interface
 
 		const UPCGPin* GetUpstreamPin() const { return UpstreamPin.Get(); }
diff --git a/Engine/Plugins/PCG/Source/PCG/Private/PCGComponent.cpp b/Engine/Plugins/PCG/Source/PCG/Private/PCGComponent.cpp
--- a/Engine/Plugins/PCG/Source/PCG/Private/PCGComponent.cpp
+++ b/Engine/Plugins/PCG/Source/PCG/Private/PCGComponent.cpp
@@ -623,6 +623,13 @@ void UPCGComponent::PostProcessGraph(const FBox& InNewBounds, bool bInGenerated,
 			{
 				if (ensure(TaggedData.Data))
 				{
+					// Proxies should never get cached on the component output. These can hold onto large chunks of video memory and are not
+					// directly serializable.
+					for (const FPCGTaggedData& Data : Context->InputData.TaggedData)
+					{
+						ensure(!Data.Data || Data.Data->CanBeSerialized());
+					}
+
 					// Visit the generated data; if the outer is the transient package or this component for all data, then
 					// we don't need to duplicate the data and can change the outer & flatten the data without any additional copies.
 					// Caveat: since that data could be in the cache currently, we shouldn't reouter it to a component that's not in the persistent level, otherwise when the
diff --git a/Engine/Plugins/PCG/Source/PCG/Private/PCGElement.cpp b/Engine/Plugins/PCG/Source/PCG/Private/PCGElement.cpp
--- a/Engine/Plugins/PCG/Source/PCG/Private/PCGElement.cpp
+++ b/Engine/Plugins/PCG/Source/PCG/Private/PCGElement.cpp
@@ -6,8 +6,10 @@
 #include "PCGContext.h"
 #include "PCGGraph.h"
 #include "PCGSubsystem.h"
+#include "Compute/Data/PCGProxyForGPUData.h"
 #include "Data/PCGPointData.h"
 #include "Graph/PCGGraphCache.h"
+#include "Helpers/PCGHelpers.h"
 
 #if WITH_EDITOR
 #include "PCGDataVisualization.h"
@@ -172,7 +174,7 @@ bool IPCGElement::Execute(FPCGContext* Context) const
 				TRACE_CPUPROFILER_EVENT_SCOPE(EPCGExecutionPhase::PrepareData);
 				PCG_ELEMENT_EXECUTION_BREAKPOINT();
 
-				if (PrepareDataInternal(Context))
+				if (PrepareData(Context))
 				{
 					Context->CurrentPhase = EPCGExecutionPhase::Execute;
 				}
@@ -356,6 +358,59 @@ void IPCGElement::PreExecutePrimaryLoopElement(FPCGContext* Context, const UPCGS
 	// 
 }
 
+bool IPCGElement::PrepareData(FPCGContext* Context) const
+{
+	if (!SupportsGPUResidentData(Context))
+	{
+		bool bHasPendingReadbacks = false;
+
+		// If there are any proxies in the input data, request readback to CPU.
+		for (FPCGTaggedData& Data : Context->InputData.TaggedData)
+		{
+			if (const UPCGProxyForGPUData* DataWithGPUSupport = Cast<UPCGProxyForGPUData>(Data.Data))
+			{
+				// Poll until readback is done.
+				UPCGProxyForGPUData::FReadbackResult Result = DataWithGPUSupport->GetCPUData(Context);
+
+				if (Result.bComplete)
+				{
+					// Readback done, replace the proxy with the result data.
+					ensure(Result.Data);
+					Data.Data = Result.Data;
+
+					// Will update referenced data objects.
+					Context->bInputDataModified = true;
+				}
+				else
+				{
+					bHasPendingReadbacks = true;
+				}
+			}
+		}
+
+		if (bHasPendingReadbacks)
+		{
+			Context->bIsPaused = true;
+
+			// Not ready to execute and unlikely to be in the very short term, sleep until next frame.
+			PCGHelpers::ExecuteOnGameThread(UE_SOURCE_LOCATION, [ContextHandle = Context->GetOrCreateHandle()]()
+			{
+				if (TSharedPtr<FPCGContextHandle> SharedHandle = ContextHandle.Pin())
+				{
+					if (FPCGContext* ContextPtr = SharedHandle->GetContext())
+					{
+						ContextPtr->bIsPaused = false;
+					}
+				}
+			});
+
+			return false;
+		}
+	}
+
+	return PrepareDataInternal(Context);
+}
+
 bool IPCGElement::PrepareDataInternal(FPCGContext* Context) const
 {
 	return true;
@@ -403,11 +458,26 @@ void IPCGElement::PostExecute(FPCGContext* Context) const
 	// TODO - The dependencies CRC here should always be valid except in the indirection case, which we should normalize to allow caching (tested here, otherwise it can ensure in the graph cache)
 	if (!Context->OutputData.bCancelExecution && !bHasErrorsOrWarnings && Context->DependenciesCrc.IsValid() && IsCacheableInstance(Settings))
 	{
-		UPCGSubsystem* Subsystem = Context->SourceComponent.Get() ? Context->SourceComponent->GetSubsystem() : nullptr;
-		if (Subsystem)
+		bool bCacheable = true;
+
+		// GPU proxies are never cached. Caching them would hold on to precious graphics memory.
+		for (const FPCGTaggedData& TaggedData : Context->OutputData.TaggedData)
 		{
-			FPCGStoreInCacheParams Params = { .Element = this, .Crc = Context->DependenciesCrc };
-			Subsystem->GetCache()->StoreInCache(Params, Context->OutputData);
+			if (TaggedData.Data && !TaggedData.Data->IsCacheable())
+			{
+				bCacheable = false;
+				break;
+			}
+		}
+
+		if (bCacheable)
+		{
+			UPCGSubsystem* Subsystem = Context->SourceComponent.Get() ? Context->SourceComponent->GetSubsystem() : nullptr;
+			if (Subsystem)
+			{
+				FPCGStoreInCacheParams Params = { .Element = this, .Crc = Context->DependenciesCrc };
+				Subsystem->GetCache()->StoreInCache(Params, Context->OutputData);
+			}
 		}
 	}
 
@@ -791,12 +861,13 @@ void IPCGElement::CleanupAndValidateOutput(FPCGContext* Context) const
 					const UPCGPin* OutputPin = Context->Node ? Context->Node->GetOutputPin(OutputPinProperties[MatchIndex].Label) : nullptr;
 					const EPCGDataType PinTypes = OutputPin ? OutputPin->GetCurrentTypes() : OutputPinProperties[MatchIndex].AllowedTypes;
 
-					const bool bTypesOverlap = !!(PinTypes & TaggedData.Data->GetDataType());
-					const bool bTypeIsSubset = !(~PinTypes & TaggedData.Data->GetDataType());
+					EPCGDataType DataType = TaggedData.Data->GetUnderlyingDataType();
+					const bool bTypesOverlap = !!(PinTypes & DataType);
+					const bool bTypeIsSubset = !(~PinTypes & DataType);
 					// TODO: Temporary fix for Settings directly from InputData (ie. from elements with code and not PCG nodes)
-					if ((!bTypesOverlap || !bTypeIsSubset) && TaggedData.Data->GetDataType() != EPCGDataType::Settings)
+					if ((!bTypesOverlap || !bTypeIsSubset) && DataType != EPCGDataType::Settings)
 					{
-						PCGE_LOG(Warning, GraphAndLog, FText::Format(LOCTEXT("OutputIncompatibleType", "Output data generated for pin '{0}' does not have a compatible type: '{1}'. Consider using more specific/narrower input pin types, or more general/wider output pin types."), FText::FromName(TaggedData.Pin), FText::FromString(UEnum::GetValueAsString(TaggedData.Data->GetDataType()))));
+						PCGE_LOG(Warning, GraphAndLog, FText::Format(LOCTEXT("OutputIncompatibleType", "Output data generated for pin '{0}' does not have a compatible type: '{1}'. Consider using more specific/narrower input pin types, or more general/wider output pin types."), FText::FromName(TaggedData.Pin), FText::FromString(UEnum::GetValueAsString(DataType))));
 					}
 				}
 
diff --git a/Engine/Plugins/PCG/Source/PCG/Private/PCGInputOutputSettings.cpp b/Engine/Plugins/PCG/Source/PCG/Private/PCGInputOutputSettings.cpp
--- a/Engine/Plugins/PCG/Source/PCG/Private/PCGInputOutputSettings.cpp
+++ b/Engine/Plugins/PCG/Source/PCG/Private/PCGInputOutputSettings.cpp
@@ -5,6 +5,7 @@
 #include "PCGContext.h"
 #include "PCGCustomVersion.h"
 #include "PCGPin.h"
+#include "Graph/PCGStackContext.h"
 
 #include UE_INLINE_GENERATED_CPP_BY_NAME(PCGInputOutputSettings)
 
@@ -81,6 +82,27 @@ bool FPCGInputOutputElement::ExecuteInternal(FPCGContext* Context) const
 	return true;
 }
 
+bool FPCGInputOutputElement::SupportsGPUResidentData(FPCGContext* InContext) const
+{
+	check(InContext);
+
+	const UPCGGraphInputOutputSettings* Settings = InContext->GetOriginalSettings<UPCGGraphInputOutputSettings>();
+	if (!Settings || !ensure(InContext->Stack))
+	{
+		return false;
+	}
+
+	if (Settings->IsInput())
+	{
+		return true;
+	}
+	else
+	{
+		// The output node of subgraphs can return GPU data proxies, but the top graph output must read back.
+		return !InContext->Stack->IsCurrentFrameInRootGraph();
+	}
+}
+
 void UPCGGraphInputOutputSettings::PostLoad()
 {
 	Super::PostLoad();
diff --git a/Engine/Plugins/PCG/Source/PCG/Public/Elements/ControlFlow/PCGBranch.h b/Engine/Plugins/PCG/Source/PCG/Public/Elements/ControlFlow/PCGBranch.h
--- a/Engine/Plugins/PCG/Source/PCG/Public/Elements/ControlFlow/PCGBranch.h
+++ b/Engine/Plugins/PCG/Source/PCG/Public/Elements/ControlFlow/PCGBranch.h
@@ -44,4 +44,5 @@ class FPCGBranchElement : public IPCGElement
 {
 protected:
 	virtual bool ExecuteInternal(FPCGContext* Context) const override;
+	virtual bool SupportsGPUResidentData(FPCGContext* InContext) const override { return true; }
 };
diff --git a/Engine/Plugins/PCG/Source/PCG/Public/Elements/ControlFlow/PCGSwitch.h b/Engine/Plugins/PCG/Source/PCG/Public/Elements/ControlFlow/PCGSwitch.h
--- a/Engine/Plugins/PCG/Source/PCG/Public/Elements/ControlFlow/PCGSwitch.h
+++ b/Engine/Plugins/PCG/Source/PCG/Public/Elements/ControlFlow/PCGSwitch.h
@@ -100,4 +100,5 @@ class FPCGSwitchElement : public IPCGElement
 {
 protected:
 	virtual bool ExecuteInternal(FPCGContext* Context) const override;
+	virtual bool SupportsGPUResidentData(FPCGContext* InContext) const override { return true; }
 };
diff --git a/Engine/Plugins/PCG/Source/PCG/Public/Elements/PCGGather.h b/Engine/Plugins/PCG/Source/PCG/Public/Elements/PCGGather.h
--- a/Engine/Plugins/PCG/Source/PCG/Public/Elements/PCGGather.h
+++ b/Engine/Plugins/PCG/Source/PCG/Public/Elements/PCGGather.h
@@ -35,6 +35,7 @@ class FPCGGatherElement : public IPCGElement
 {
 protected:
 	virtual bool ExecuteInternal(FPCGContext* Context) const override;
+	virtual bool SupportsGPUResidentData(FPCGContext* InContext) const override { return true; }
 };
 
 namespace PCGGather
diff --git a/Engine/Plugins/PCG/Source/PCG/Public/PCGCommon.h b/Engine/Plugins/PCG/Source/PCG/Public/PCGCommon.h
--- a/Engine/Plugins/PCG/Source/PCG/Public/PCGCommon.h
+++ b/Engine/Plugins/PCG/Source/PCG/Public/PCGCommon.h
@@ -94,6 +94,9 @@ enum class EPCGDataType : uint32
 	/** Combinations of concrete data and/or boolean operations. */
 	Spatial = Composite | Concrete,
 
+	/** Proxy for data that was created on the GPU and not necessarily read back to CPU. */
+	ProxyForGPU = 1 << 26 UMETA(Hidden),
+
 	Param = 1 << 27 UMETA(DisplayName = "Attribute Set"),
 
 	// Combination of Param and Point, necessary for named-based serialization of enums.
diff --git a/Engine/Plugins/PCG/Source/PCG/Public/PCGContext.h b/Engine/Plugins/PCG/Source/PCG/Public/PCGContext.h
--- a/Engine/Plugins/PCG/Source/PCG/Public/PCGContext.h
+++ b/Engine/Plugins/PCG/Source/PCG/Public/PCGContext.h
@@ -75,6 +75,9 @@ struct PCG_API FPCGContext
 	FPCGDataCollection OutputData;
 	TWeakObjectPtr<UPCGComponent> SourceComponent = nullptr;
 
+	/** Used to track when data transformations (such as CPU readback) occurred on the input data collection as part of the element prepare data. */
+	bool bInputDataModified = false;
+
 	FPCGAsyncState AsyncState;
 	FPCGCrc DependenciesCrc;
 
diff --git a/Engine/Plugins/PCG/Source/PCG/Public/PCGData.h b/Engine/Plugins/PCG/Source/PCG/Public/PCGData.h
--- a/Engine/Plugins/PCG/Source/PCG/Public/PCGData.h
+++ b/Engine/Plugins/PCG/Source/PCG/Public/PCGData.h
@@ -40,6 +40,15 @@ public:
 	/** Executes a lambda over all connected data objects. */
 	virtual void VisitDataNetwork(TFunctionRef<void(const UPCGData*)> Action) const;
 
+	/** Whether this data can be serialized. */
+	virtual bool CanBeSerialized() const { return true; }
+
+	/** Whether the data can be placed in the graph cache. */
+	virtual bool IsCacheable() const { return true; }
+
+	/** If this data is a proxy, returns the underlying data type, otherwise returns this data type. */
+	virtual EPCGDataType GetUnderlyingDataType() const { return GetDataType(); }
+
 	/** Unique ID for this object instance. */
 	UPROPERTY(Transient)
 	uint64 UID = 0;
@@ -71,7 +80,7 @@ protected:
 	/** Computes Crc for this and any connected data. */
 	virtual FPCGCrc ComputeCrc(bool bFullDataCrc) const;
 
-	/** Adds this data to Crc. Fallback implementation writes object instance UID. */
+	/** Adds this data to Crc. Derived classes should override this and add all required data to CRC, or fall back to adding the unique object instance UID. */
 	virtual void AddToCrc(FArchiveCrc32& Ar, bool bFullDataCrc) const;
 
 	/** Whether the data distinguishes between regular crc and full crc */
diff --git a/Engine/Plugins/PCG/Source/PCG/Public/PCGElement.h b/Engine/Plugins/PCG/Source/PCG/Public/PCGElement.h
--- a/Engine/Plugins/PCG/Source/PCG/Public/PCGElement.h
+++ b/Engine/Plugins/PCG/Source/PCG/Public/PCGElement.h
@@ -149,6 +149,7 @@ protected:
 	/** This function will be called once and once only, at the beginning of an execution */
 	void PreExecute(FPCGContext* Context) const;
 	/** The prepare data phase is one where it is more likely to be able to multithread */
+	bool PrepareData(FPCGContext* Context) const;
 	virtual bool PrepareDataInternal(FPCGContext* Context) const;
 	/** Core execution method for the given element. Will be called until it returns true. */
 	virtual bool ExecuteInternal(FPCGContext* Context) const = 0;
@@ -179,6 +180,9 @@ protected:
 	/** Let each element optionally act as a concrete factory for its own context */
 	virtual FPCGContext* CreateContext();
 
+	/** Returns true if node supports data objects that are only resident on GPU (not CPU). */
+	virtual bool SupportsGPUResidentData(FPCGContext* InContext) const { return false; }
+
 #if WITH_EDITOR
 	virtual bool ShouldLog() const { return true; }
 #endif
diff --git a/Engine/Plugins/PCG/Source/PCG/Public/PCGInputOutputSettings.h b/Engine/Plugins/PCG/Source/PCG/Public/PCGInputOutputSettings.h
--- a/Engine/Plugins/PCG/Source/PCG/Public/PCGInputOutputSettings.h
+++ b/Engine/Plugins/PCG/Source/PCG/Public/PCGInputOutputSettings.h
@@ -21,6 +21,7 @@ class FPCGInputOutputElement : public IPCGElement
 protected:
 	virtual bool ExecuteInternal(FPCGContext* Context) const override;
 	virtual bool IsCacheable(const UPCGSettings* InSettings) const override { return false; }
+	virtual bool SupportsGPUResidentData(FPCGContext* InContext) const override;
 };
 
 UCLASS(MinimalAPI, NotBlueprintable, Hidden, ClassGroup = (Procedural))
diff --git a/Engine/Plugins/PCG/Source/PCG/Public/PCGSettings.h b/Engine/Plugins/PCG/Source/PCG/Public/PCGSettings.h
--- a/Engine/Plugins/PCG/Source/PCG/Public/PCGSettings.h
+++ b/Engine/Plugins/PCG/Source/PCG/Public/PCGSettings.h
@@ -649,4 +649,5 @@ class FPCGTrivialElement : public IPCGElement
 protected:
 	virtual bool ExecuteInternal(FPCGContext* Context) const override;
 	virtual bool IsPassthrough(const UPCGSettings* InSettings) const override { return true; }
+	virtual bool SupportsGPUResidentData(FPCGContext* InContext) const override { return true; }
 };
diff --git a/Engine/Plugins/PCG/Source/PCG/Public/PCGSubgraph.h b/Engine/Plugins/PCG/Source/PCG/Public/PCGSubgraph.h
--- a/Engine/Plugins/PCG/Source/PCG/Public/PCGSubgraph.h
+++ b/Engine/Plugins/PCG/Source/PCG/Public/PCGSubgraph.h
@@ -192,6 +192,7 @@ protected:
 	virtual bool IsPassthrough(const UPCGSettings* InSettings) const override;
 	void PrepareSubgraphData(const UPCGSubgraphSettings* Settings, FPCGSubgraphContext* Context, const FPCGDataCollection& InputData, FPCGDataCollection& OutputData) const;
 	void PrepareSubgraphUserParameters(const UPCGSubgraphSettings* Settings, FPCGSubgraphContext* Context, FPCGDataCollection& OutputData) const;
+	virtual bool SupportsGPUResidentData(FPCGContext* InContext) const override { return true; }
 };
 
 // Implementation note: this node forwards data, but does not keep that data alive. This is the responsibility of the corresponding FPCGSubgraphContext
@@ -206,5 +207,7 @@ public:
 protected:
 	virtual bool ExecuteInternal(FPCGContext* Context) const override;
 	virtual bool IsPassthrough(const UPCGSettings* InSettings) const override { return true; }
+	virtual bool SupportsGPUResidentData(FPCGContext* InContext) const override { return true; }
+
 	FPCGDataCollection Input;
 };
diff --git a/Engine/Plugins/PCG/Source/PCGEditor/Private/SPCGEditorGraphAttributeListView.cpp b/Engine/Plugins/PCG/Source/PCGEditor/Private/SPCGEditorGraphAttributeListView.cpp
--- a/Engine/Plugins/PCG/Source/PCGEditor/Private/SPCGEditorGraphAttributeListView.cpp
+++ b/Engine/Plugins/PCG/Source/PCGEditor/Private/SPCGEditorGraphAttributeListView.cpp
@@ -11,6 +11,7 @@
 #include "PCGParamData.h"
 #include "PCGPin.h"
 #include "PCGSubsystem.h"
+#include "Compute/Data/PCGProxyForGPUData.h"
 #include "Data/PCGSpatialData.h"
 #include "Elements/IO/PCGSaveAssetElement.h"
 #include "Metadata/PCGAttributePropertySelector.h"
@@ -716,6 +717,27 @@ void SPCGEditorGraphAttributeListView::RefreshAttributeList()
 		return;
 	}
 
+	// If we have a proxy for GPU data, read back CPU data for inspection.
+	const UPCGProxyForGPUData* Proxy = Cast<UPCGProxyForGPUData>(PCGData);
+	if (Proxy)
+	{
+		UPCGProxyForGPUData::FReadbackResult Result = Proxy->GetCPUData(/*InContext=*/nullptr);
+
+		if (!Result.bComplete)
+		{
+			// Poll next tick.
+			bNeedsRefresh = true;
+			return;
+		}
+
+		if (!Result.Data)
+		{
+			return;
+		}
+
+		PCGData = Result.Data;
+	}
+
 	const FPCGDataVisualizationRegistry& DataVisRegistry = FPCGModule::GetConstPCGDataVisualizationRegistry();
 	const UPCGData* DataToVisualize = PCGData;
 
