[PCG] Add new node 'Get Console Variable'

diff --git a/Engine/Plugins/PCG/Source/PCG/Private/Elements/PCGGetConsoleVariable.cpp b/Engine/Plugins/PCG/Source/PCG/Private/Elements/PCGGetConsoleVariable.cpp
new file mode 100644
--- /dev/null
+++ b/Engine/Plugins/PCG/Source/PCG/Private/Elements/PCGGetConsoleVariable.cpp
@@ -0,0 +1,111 @@
+// Copyright Epic Games, Inc. All Rights Reserved.
+
+#include "Elements/PCGGetConsoleVariable.h"
+
+#include "PCGContext.h"
+#include "PCGModule.h"
+#include "PCGParamData.h"
+#include "Metadata/PCGMetadataAttributeTpl.h"
+
+#include UE_INLINE_GENERATED_CPP_BY_NAME(PCGGetConsoleVariable)
+
+#define LOCTEXT_NAMESPACE "PCGGetConsoleVariableElement"
+
+TArray<FPCGPinProperties> UPCGGetConsoleVariableSettings::OutputPinProperties() const
+{
+	TArray<FPCGPinProperties> PinProperties;
+	PinProperties.Emplace(PCGPinConstants::DefaultOutputLabel, EPCGDataType::Param);
+
+	return PinProperties;
+}
+
+#if WITH_EDITOR
+FText UPCGGetConsoleVariableSettings::GetNodeTooltipText() const
+{
+	return LOCTEXT("NodeTooltip", "Reads the given console variable and writes the value to an attribute set.\n"
+		"Note: Setting the console variable will not trigger a regeneration.");
+}
+#endif // WITH_EDITOR
+
+FPCGElementPtr UPCGGetConsoleVariableSettings::CreateElement() const
+{
+	return MakeShared<FPCGGetConsoleVariableElement>();
+}
+
+bool FPCGGetConsoleVariableElement::ExecuteInternal(FPCGContext* Context) const
+{
+	TRACE_CPUPROFILER_EVENT_SCOPE(FPCGGetConsoleVariableElement::ExecuteInternal);
+
+	check(Context);
+
+	const UPCGGetConsoleVariableSettings* Settings = Context->GetInputSettings<UPCGGetConsoleVariableSettings>();
+	check(Settings);
+
+	IConsoleVariable* ConsoleVariable = IConsoleManager::Get().FindConsoleVariable(*Settings->ConsoleVariableName.ToString());
+
+	if (!ConsoleVariable)
+	{
+		PCGE_LOG(Error, GraphAndLog, FText::Format(LOCTEXT("FailedToFindCVar", "Failed to find console variable '{0}'."), FText::FromName(Settings->ConsoleVariableName)));
+		return true;
+	}
+
+	UPCGParamData* OutParamData = FPCGContext::NewObject_AnyThread<UPCGParamData>(Context);
+	UPCGMetadata* Metadata = OutParamData->Metadata;
+	check(Metadata);
+
+	const PCGMetadataEntryKey EntryKey = OutParamData->Metadata->AddEntry();
+
+	auto CreateAttribute = [Metadata, AttributeName = Settings->OutputAttributeName, EntryKey]<typename T>(const T& Value) -> bool
+	{
+		if (FPCGMetadataAttribute<T>* Attribute = Metadata->CreateAttribute<T>(AttributeName, Value, /*bAllowsInterpolation=*/true, /*bOverrideParent=*/false))
+		{
+			Attribute->SetValue(EntryKey, Value);
+			return true;
+		}
+		else
+		{
+			return false;
+		}
+	};
+	
+	bool bCreatedAttribute = false;
+
+	if (ConsoleVariable->IsVariableBool())
+	{
+		const bool Value = ConsoleVariable->GetBool();
+		bCreatedAttribute |= CreateAttribute(Value);
+	}
+	else if (ConsoleVariable->IsVariableInt())
+	{
+		const int Value = ConsoleVariable->GetInt();
+		bCreatedAttribute |= CreateAttribute(Value);
+	}
+	else if (ConsoleVariable->IsVariableFloat())
+	{
+		const float Value = ConsoleVariable->GetFloat();
+		bCreatedAttribute |= CreateAttribute(Value);
+	}
+	else if (ConsoleVariable->IsVariableString())
+	{
+		const FString Value = ConsoleVariable->GetString();
+		bCreatedAttribute |= CreateAttribute(Value);
+	}
+	else
+	{
+		PCGE_LOG(Error, GraphAndLog, FText::Format(LOCTEXT("UnsupportedCVarType", "Console variable '{0}' is not a supported type."), FText::FromName(Settings->ConsoleVariableName)));
+		return true;
+	}
+
+	if (!bCreatedAttribute)
+	{
+		PCGE_LOG(Error, GraphAndLog, FText::Format(LOCTEXT("AttributeCreationFailed", "Failed to create attribute {0}."), FText::FromName(Settings->OutputAttributeName)));
+		return true;
+	}
+
+	FPCGTaggedData& OutData = Context->OutputData.TaggedData.Emplace_GetRef();
+	OutData.Data = OutParamData;
+
+	return true;
+}
+
+#undef LOCTEXT_NAMESPACE
\ No newline at end of file
diff --git a/Engine/Plugins/PCG/Source/PCG/Public/Elements/PCGGetConsoleVariable.h b/Engine/Plugins/PCG/Source/PCG/Public/Elements/PCGGetConsoleVariable.h
new file mode 100644
--- /dev/null
+++ b/Engine/Plugins/PCG/Source/PCG/Public/Elements/PCGGetConsoleVariable.h
@@ -0,0 +1,46 @@
+// Copyright Epic Games, Inc. All Rights Reserved.
+
+#pragma once
+
+#include "PCGSettings.h"
+
+#include "PCGGetConsoleVariable.generated.h"
+
+/**
+ * Reads the given console variable and writes the value to an attribute set.
+ * Note: Setting the console variable will not trigger a regeneration.
+ */
+UCLASS(MinimalAPI, BlueprintType, ClassGroup = (Procedural))
+class UPCGGetConsoleVariableSettings : public UPCGSettings
+{
+	GENERATED_BODY()
+
+public:
+	//~Begin UPCGSettings interface
+	virtual TArray<FPCGPinProperties> InputPinProperties() const override { return TArray<FPCGPinProperties>(); }
+	virtual TArray<FPCGPinProperties> OutputPinProperties() const override;
+#if WITH_EDITOR
+	virtual FName GetDefaultNodeName() const override { return FName(TEXT("GetConsoleVariable")); }
+	virtual FText GetDefaultNodeTitle() const override { return NSLOCTEXT("PCGGetConsoleVariableElement", "NodeTitle", "Get Console Variable"); }
+	virtual FText GetNodeTooltipText() const override;
+	virtual EPCGSettingsType GetType() const override { return EPCGSettingsType::Generic; }
+#endif
+	virtual FString GetAdditionalTitleInformation() const override { return ConsoleVariableName.ToString(); }
+
+protected:
+	virtual FPCGElementPtr CreateElement() const override;
+	//~End UPCGSettings interface
+
+public:
+	UPROPERTY(BlueprintReadWrite, EditAnywhere, Category = "Settings", meta = (PCG_Overridable))
+	FName ConsoleVariableName = NAME_None;
+
+	UPROPERTY(BlueprintReadWrite, EditAnywhere, Category = "Settings", meta = (PCG_Overridable))
+	FName OutputAttributeName = NAME_None;
+};
+
+class FPCGGetConsoleVariableElement : public IPCGElement
+{
+protected:
+	virtual bool ExecuteInternal(FPCGContext* Context) const override;
+};
