[PCG] Various improvements around PCGPartitionActor registration and mapping and Component cleanup
- PCGPartitionActors & PCG Partitioned Volumes will now remap properly when out of sync (different bGenerated value)
- Avoid modification of PCGPartitionActors when editor is loading
- Change some CleanLocalImmediate calls to Scheduled cleanups where it made sense
- Removed LocalToOriginalMap_DEPRECATED and avoid calling CleanupDeadGraphInstances from PostLoad

#rnx

diff --git a/Engine/Plugins/PCG/Source/PCG/Private/Grid/PCGPartitionActor.cpp b/Engine/Plugins/PCG/Source/PCG/Private/Grid/PCGPartitionActor.cpp
--- a/Engine/Plugins/PCG/Source/PCG/Private/Grid/PCGPartitionActor.cpp
+++ b/Engine/Plugins/PCG/Source/PCG/Private/Grid/PCGPartitionActor.cpp
@@ -56,16 +56,6 @@ void APCGPartitionActor::PostLoad()
 	}
 
 #if WITH_EDITOR
-	if (LocalToOriginalMap_DEPRECATED.Num() > 0)
-	{
-		for (const TPair<TObjectPtr<UPCGComponent>, TWeakObjectPtr<UPCGComponent>>& LocalToOriginalPair : LocalToOriginalMap_DEPRECATED)
-		{
-			LocalToOriginal.Emplace(LocalToOriginalPair.Key, LocalToOriginalPair.Value.Get());
-		}
-
-		LocalToOriginalMap_DEPRECATED.Reset();
-	}
-
 	if (GetGridSize() != PCGGridSize)
 	{
 		SetGridSize(PCGGridSize);
@@ -87,9 +77,6 @@ void APCGPartitionActor::PostLoad()
 		LocalToOriginal.Remove(nullptr);
 	}
 
-	// Make sure that we don't track objects that do not exist anymore
-	CleanupDeadGraphInstances(/*bRemoveNullOnly=*/true);
-
 	// Mark all our local components as local
 	for (UPCGComponent* LocalComponent : GetAllLocalPCGComponents())
 	{
@@ -202,7 +189,10 @@ void APCGPartitionActor::RebuildOriginalToLocal()
 	OriginalToLocal.Reserve(LocalToOriginal.Num());
 	for (const auto& It : LocalToOriginal)
 	{
-		OriginalToLocal.Add(It.Value.Get(), It.Key);
+		if (UPCGComponent* OriginalComponent = It.Value.Get())
+		{
+			OriginalToLocal.Add(OriginalComponent, It.Key);
+		}
 	}
 }
 
@@ -221,18 +211,20 @@ void APCGPartitionActor::BeginPlay()
 		// Pass through all the pcg components, to verify if we need to generate them.
 		for (auto& It : OriginalToLocal)
 		{
-			if (It.Key && It.Value)
+			UPCGComponent* OriginalComponent = It.Key.Get();
+			UPCGComponent* LocalComponent = It.Value;
+			if (OriginalComponent && LocalComponent)
 			{
 				// If we have an original component that is generated (or generating), this one is automatically generated => GenerateOnLoad.
 				// But if its runtime generated then it's handled by the runtime generation scheduler.
-				if ((It.Key->bGenerated || It.Key->IsGenerating()) && ensure(!It.Key->IsManagedByRuntimeGenSystem()))
+				if ((OriginalComponent->bGenerated || OriginalComponent->IsGenerating()) && ensure(!OriginalComponent->IsManagedByRuntimeGenSystem()))
 				{
-					It.Value->GenerationTrigger = EPCGComponentGenerationTrigger::GenerateOnLoad;
+					LocalComponent->GenerationTrigger = EPCGComponentGenerationTrigger::GenerateOnLoad;
 				}
 				// Otherwise, make them match
 				else
 				{
-					It.Value->GenerationTrigger = It.Key->GenerationTrigger;
+					LocalComponent->GenerationTrigger = OriginalComponent->GenerationTrigger;
 				}
 			}
 		}
@@ -246,10 +238,10 @@ void APCGPartitionActor::EndPlay(const EEndPlayReason::Type EndPlayReason)
 	// It is possible for a PA to get called EndPlay before its original volume actor so make sure to cleanup properly here
 	if (IsRuntimeGenerated())
 	{
-		TMap<TObjectPtr<UPCGComponent>, TObjectPtr<UPCGComponent>> OriginalToLocalCopy(OriginalToLocal);
-		for (const TTuple<TObjectPtr<UPCGComponent>, TObjectPtr<UPCGComponent>>& OriginalToLocalItem : OriginalToLocalCopy)
+		TMap<TWeakObjectPtr<UPCGComponent>, TObjectPtr<UPCGComponent>> OriginalToLocalCopy(OriginalToLocal);
+		for (const TTuple<TWeakObjectPtr<UPCGComponent>, TObjectPtr<UPCGComponent>>& OriginalToLocalItem : OriginalToLocalCopy)
 		{
-			RemoveGraphInstance(OriginalToLocalItem.Key.Get());
+			RemoveGraphInstance(OriginalToLocalItem.Key.GetEvenIfUnreachable());
 		}
 	}
 
@@ -374,7 +366,7 @@ void APCGPartitionActor::RegisterPCG()
 	{
 		if (!bIsRegistered)
 		{
-			Subsystem->RegisterPartitionActor(this, /*bDoComponentMapping=*/false);
+			Subsystem->RegisterPartitionActor(this);
 			bIsRegistered = true;
 		}
 	}
@@ -413,7 +405,7 @@ UPCGComponent* APCGPartitionActor::GetLocalComponent(const UPCGComponent* Origin
 	{
 		return *LocalComponent;
 	}
-	else if (bRebuildMappingOnNullEntries && OriginalToLocal.Contains(nullptr))
+	else if (bRebuildMappingOnNullEntries && (OriginalToLocal.Contains(nullptr) || (OriginalToLocal.Num() != LocalToOriginal.Num())))
 	{
 		const_cast<APCGPartitionActor*>(this)->RebuildOriginalToLocal();
 		return GetLocalComponent(OriginalComponent, /*bRebuildMappingOnNullEntries=*/false);
@@ -428,10 +420,16 @@ UPCGComponent* APCGPartitionActor::GetOriginalComponent(const UPCGComponent* Loc
 	return OriginalComponent ? (*OriginalComponent).Get() : nullptr;
 }
 
+// deprecated
 void APCGPartitionActor::CleanupDeadGraphInstances(bool bRemoveNonNullOnly)
+{
+	CleanupDeadGraphInstancesInternal();
+}
+
+void APCGPartitionActor::CleanupDeadGraphInstancesInternal()
 {
 	// First find if we have any local dead instance (= nullptr) hooked to an original component.
-	TSet<TObjectPtr<UPCGComponent>> DeadOriginalInstances;
+	TSet<TWeakObjectPtr<UPCGComponent>> DeadOriginalInstances;
 	for (const auto& OriginalToLocalItem : OriginalToLocal)
 	{
 		if (!OriginalToLocalItem.Value)
@@ -444,7 +442,7 @@ void APCGPartitionActor::CleanupDeadGraphInstances(bool bRemoveNonNullOnly)
 	{
 		Modify();
 
-		for (const TObjectPtr<UPCGComponent>& DeadInstance : DeadOriginalInstances)
+		for (const TWeakObjectPtr<UPCGComponent>& DeadInstance : DeadOriginalInstances)
 		{
 			OriginalToLocal.Remove(DeadInstance);
 		}
@@ -456,8 +454,7 @@ void APCGPartitionActor::CleanupDeadGraphInstances(bool bRemoveNonNullOnly)
 	TSet<TObjectPtr<UPCGComponent>> DeadLocalInstances;
 	for (const auto& LocalToOriginalItem : LocalToOriginal)
 	{
-		if((bRemoveNonNullOnly && LocalToOriginalItem.Value.IsNull()) ||
-			(!bRemoveNonNullOnly && !LocalToOriginalItem.Value.IsValid()))
+		if(!LocalToOriginalItem.Value.IsValid())
 		{
 			DeadLocalInstances.Add(LocalToOriginalItem.Key);
 		}
@@ -570,7 +567,7 @@ bool APCGPartitionActor::RemoveGraphInstance(UPCGComponent* OriginalComponent)
 	{
 		// If we don't have a local component, perhaps the original component is already dead,
 		// so do some clean up
-		CleanupDeadGraphInstances();
+		CleanupDeadGraphInstancesInternal();
 		return false;
 	}
 
@@ -676,9 +673,12 @@ TSet<TObjectPtr<UPCGComponent>> APCGPartitionActor::GetAllLocalPCGComponents() c
 TSet<TObjectPtr<UPCGComponent>> APCGPartitionActor::GetAllOriginalPCGComponents() const
 {
 	TSet<TObjectPtr<UPCGComponent>> ResultComponents;
-	for(const TPair<TObjectPtr<UPCGComponent>, TObjectPtr<UPCGComponent>>& OriginalPair : OriginalToLocal)
+	for(const TPair<TWeakObjectPtr<UPCGComponent>, TObjectPtr<UPCGComponent>>& OriginalPair : OriginalToLocal)
 	{
-		ResultComponents.Add(OriginalPair.Key.Get());
+		if (UPCGComponent* OriginalComponent = OriginalPair.Key.Get())
+		{
+			ResultComponents.Add(OriginalComponent);
+		}
 	}
 
 	return ResultComponents;
diff --git a/Engine/Plugins/PCG/Source/PCG/Private/PCGActorAndComponentMapping.cpp b/Engine/Plugins/PCG/Source/PCG/Private/PCGActorAndComponentMapping.cpp
--- a/Engine/Plugins/PCG/Source/PCG/Private/PCGActorAndComponentMapping.cpp
+++ b/Engine/Plugins/PCG/Source/PCG/Private/PCGActorAndComponentMapping.cpp
@@ -395,6 +395,11 @@ bool FPCGActorAndComponentMapping::RegisterOrUpdatePartitionedPCGComponent(UPCGC
 			// If we do not want a mapping, delete the existing one
 			DeleteMappingPCGComponentPartitionActor(InComponent);
 		}
+		else
+		{
+			// Here bDoActorMapping is false so we just want to update the ComponentToPartitionActorsMap without adding/removing graph instances
+			UpdateMappingPCGComponentPartitionActor(InComponent, /*bChangeGraphInstances=*/false);
+		}
 	}
 
 	return bComponentHasChanged;
@@ -592,7 +597,7 @@ TArray<UPCGComponent*> FPCGActorAndComponentMapping::GetAllIntersectingComponent
 	return Result;
 }
 
-void FPCGActorAndComponentMapping::RegisterPartitionActor(APCGPartitionActor* InActor, bool bDoComponentMapping)
+void FPCGActorAndComponentMapping::RegisterPartitionActor(APCGPartitionActor* InActor)
 {
 	check(InActor);
 		
@@ -651,15 +656,17 @@ PRAGMA_ENABLE_DEPRECATION_WARNINGS
 	if (!PartitionActorDescriptor.IsRuntime())
 	{
 		FWriteScopeLock WriteLock(ComponentToPartitionActorsMapLock);
-		ForAllIntersectingPartitionedComponents(FBoxCenterAndExtent(InActor->GetFixedBounds()), [this, InActor, bDoComponentMapping, WorldActor, PartitionActorDescriptor](UPCGComponent* Component)
+		ForAllIntersectingPartitionedComponents(FBoxCenterAndExtent(InActor->GetFixedBounds()), [this, InActor, WorldActor, PartitionActorDescriptor](UPCGComponent* Component)
 		{
-			// For each component, do the mapping if we ask it explicitly, or if the component is generated
-			if (Component && (bDoComponentMapping || Component->bGenerated))
+			// For each component, do the mapping if the component is generated or if the Partition Actor already has a local component
+			// Partition Actors could have a local component while the original doesn't have a mapping if for instance a cleanup/save on the original component happened
+			// while not saving the Partition Actors
+			if (Component)
 			{
 				FPCGGridDescriptor ComponentGridDescriptor = Component->GetGridDescriptor(InActor->GetPCGGridSize());
 				if (ComponentGridDescriptor == PartitionActorDescriptor)
 				{
-					if (TSet<TObjectPtr<APCGPartitionActor>>* PartitionActorsPtr = ComponentToPartitionActorsMap.Find(Component))
+					if (Component->bGenerated || (InActor->GetLocalComponent(Component) != nullptr))
 					{
 						bool bHasUnbounded = false;
 						PCGHiGenGrid::FSizeArray GridSizes;
@@ -669,8 +676,17 @@ PRAGMA_ENABLE_DEPRECATION_WARNINGS
 						{
 							// In editor we might load/create partition actors while the component is registering. Because of that,
 							// the mapping might not already exists, even if the component is marked generated.
-							InActor->AddGraphInstance(Component);
-							PartitionActorsPtr->Add(InActor);
+							if (!UE::GetIsEditorLoadingPackage())
+							{
+								InActor->AddGraphInstance(Component);
+							}
+							
+							// Either we already had a local component or one was just added
+							if (InActor->GetLocalComponent(Component))
+							{
+								// Add mapping
+								ComponentToPartitionActorsMap.FindOrAdd(Component).Add(InActor);
+							}
 						}
 					}
 				}
@@ -776,7 +792,7 @@ void FPCGActorAndComponentMapping::ForAllIntersectingPartitionActors(const FBox&
 	}
 }
 
-void FPCGActorAndComponentMapping::UpdateMappingPCGComponentPartitionActor(UPCGComponent* InComponent)
+void FPCGActorAndComponentMapping::UpdateMappingPCGComponentPartitionActor(UPCGComponent* InComponent, bool bInChangeGraphInstances)
 {
 	if (!PCGSubsystem->IsInitialized())
 	{
@@ -785,6 +801,9 @@ void FPCGActorAndComponentMapping::UpdateMappingPCGComponentPartitionActor(UPCGC
 
 	check(InComponent);
 
+	// Do not modify Partition Actors when loading
+	const bool bChangeGraphInstances = bInChangeGraphInstances && !UE::GetIsEditorLoadingPackage();
+
 	if (!InComponent->GetGraph())
 	{
 		return;
@@ -810,23 +829,16 @@ void FPCGActorAndComponentMapping::UpdateMappingPCGComponentPartitionActor(UPCGC
 			return InComponent->GetGridDescriptor(GridSize);
 		});
 
-		auto UpdateMapping = [this, InComponent, &Bounds, WorldActor, &GridDescriptors]()
+		auto UpdateMapping = [this, InComponent, &Bounds, WorldActor, &GridDescriptors, bChangeGraphInstances]()
 		{
 			TSet<TObjectPtr<APCGPartitionActor>> RemovedActors;
 
 			{
 				FWriteScopeLock WriteLock(ComponentToPartitionActorsMapLock);
-				TSet<TObjectPtr<APCGPartitionActor>>* PartitionActorsPtr = ComponentToPartitionActorsMap.Find(InComponent);
-
-				if (!PartitionActorsPtr)
-				{
-					// Does not yet exists, add it
-					PartitionActorsPtr = &ComponentToPartitionActorsMap.Emplace(InComponent);
-					check(PartitionActorsPtr);
-				}
+				TSet<TObjectPtr<APCGPartitionActor>>& PartitionActors = ComponentToPartitionActorsMap.FindOrAdd(InComponent);
 
 				TSet<TObjectPtr<APCGPartitionActor>> NewMapping;
-				ForAllIntersectingPartitionActors(Bounds, [&NewMapping, InComponent, WorldActor, &GridDescriptors](APCGPartitionActor* Actor)
+				ForAllIntersectingPartitionActors(Bounds, [&NewMapping, InComponent, WorldActor, &GridDescriptors, bChangeGraphInstances](APCGPartitionActor* Actor)
 				{
 					if (!Actor)
 					{
@@ -836,23 +848,37 @@ void FPCGActorAndComponentMapping::UpdateMappingPCGComponentPartitionActor(UPCGC
 					// Only add a graph instance to partition actors with compatible Grid Descriptors
 					if (GridDescriptors.Contains(Actor->GetGridDescriptor()))
 					{
-						Actor->AddGraphInstance(InComponent);
-						NewMapping.Add(Actor);
+						if (bChangeGraphInstances)
+						{
+							Actor->AddGraphInstance(InComponent);
+						}
+
+						// Either we already had a local component or one was just added
+						if (Actor->GetLocalComponent(InComponent))
+						{
+							NewMapping.Add(Actor);
+						}
 					}
 				});
 
 				// Find the ones that were removed
-				RemovedActors = PartitionActorsPtr->Difference(NewMapping);
+				RemovedActors = PartitionActors.Difference(NewMapping);
 
-				*PartitionActorsPtr = MoveTemp(NewMapping);
+				PartitionActors = MoveTemp(NewMapping);
 			}
 
-			// No need to be locked to do this.
-			for (APCGPartitionActor* RemovedActor : RemovedActors)
+			// Here we can only remove graph instances if bChangeGraphInstances, which means in some edge cases
+			// we might have Partition Actors that have a PCG Component which is no longer linked to the original component.
+			// Those need to be manually cleaned up
+			if (bChangeGraphInstances)
 			{
-				if (RemovedActor)
+				// No need to be locked to do this.
+				for (APCGPartitionActor* RemovedActor : RemovedActors)
 				{
-					RemovedActor->RemoveGraphInstance(InComponent);
+					if (RemovedActor)
+					{
+						RemovedActor->RemoveGraphInstance(InComponent);
+					}
 				}
 			}
 		};
diff --git a/Engine/Plugins/PCG/Source/PCG/Private/PCGComponent.cpp b/Engine/Plugins/PCG/Source/PCG/Private/PCGComponent.cpp
--- a/Engine/Plugins/PCG/Source/PCG/Private/PCGComponent.cpp
+++ b/Engine/Plugins/PCG/Source/PCG/Private/PCGComponent.cpp
@@ -383,9 +383,12 @@ bool UPCGComponent::ShouldGenerate(bool bForce, EPCGComponentGenerationTrigger R
 
 	// A request is invalid only if it was requested "GenerateOnLoad", but it is "GenerateOnDemand"
 	// Meaning that all "GenerateOnDemand" requests are always valid, and "GenerateOnLoad" request is only valid if we want a "GenerateOnLoad" trigger.
-	bool bValidRequest = !(RequestedGenerationTrigger == EPCGComponentGenerationTrigger::GenerateOnLoad && GenerationTrigger == EPCGComponentGenerationTrigger::GenerateOnDemand);
+	const bool bValidRequest = !(RequestedGenerationTrigger == EPCGComponentGenerationTrigger::GenerateOnLoad && GenerationTrigger == EPCGComponentGenerationTrigger::GenerateOnDemand);
+	
+	// Consider the component Generated only if it is not currently cleaning up
+	const bool bConsiderGenerated = bGenerated && !IsCleaningUp();
 
-	return ((!bGenerated && bValidRequest) ||
+	return ((!bConsiderGenerated && bValidRequest) ||
 #if WITH_EDITOR
 			bDirtyGenerated || 
 #endif
@@ -551,26 +554,18 @@ FPCGTaskId UPCGComponent::CreateGenerateTask(bool bForce, const TArray<FPCGTaskI
 	TArray<FPCGTaskId> AdditionalDependencies;
 	const TArray<FPCGTaskId>* AllDependencies = &Dependencies;
 
-#if WITH_EDITOR
-	// If PCG is currently paused, we don't want to do immediate cleanup here otherwise it looks wrong and can't be cancelled.
-	// In this case, we'll "pre-queue" a cleanup before so it behaves as-if we had done this exactly.
-	if (bGenerated && PCGSystemSwitches::CVarPausePCGExecution.GetValueOnAnyThread())
+	if (bGenerated)
 	{
-		CleanupInternal(/*bRemoveComponents=*/false, Dependencies);
+		CleanupLocal(/*bRemoveComponents=*/false, Dependencies);
 	}
-#endif
 
 	if (IsCleaningUp())
 	{
-		AdditionalDependencies = Dependencies;
+		AdditionalDependencies.Reserve(Dependencies.Num() + 1);
+		AdditionalDependencies.Append(Dependencies);
 		AdditionalDependencies.Add(CurrentCleanupTask);
 		AllDependencies = &AdditionalDependencies;
 	}
-	else if (bGenerated)
-	{
-		// Immediate pass to mark all resources unused (and remove the ones that cannot be re-used)
-		CleanupLocalImmediate(/*bRemoveComponents=*/false);
-	}
 
 	const FBox NewBounds = GetGridBounds();
 	if (!NewBounds.IsValid)
@@ -884,16 +879,6 @@ void UPCGComponent::OnProcessGraphAborted(bool bQuiet, bool bCleanupUnusedResour
 
 void UPCGComponent::Cleanup()
 {
-	if (!GetSubsystem() || IsCleaningUp())
-	{
-		return;
-	}
-
-	if (!bGenerated && !IsGenerating() && GeneratedResources.IsEmpty())
-	{
-		return;
-	}
-
 	if (IsManagedByRuntimeGenSystem())
 	{
 		UE_LOG(LogPCG, Warning, TEXT("Cleanup request denied as this component is managed by the runtime generation scheduler."));
@@ -907,7 +892,7 @@ void UPCGComponent::Cleanup()
 	CleanupLocal(/*bRemoveComponents=*/true);
 }
 
-void UPCGComponent::Cleanup_Implementation(bool bRemoveComponents, bool bSave)
+void UPCGComponent::Cleanup_Implementation(bool bRemoveComponents)
 {
 	CleanupLocal(bRemoveComponents);
 }
@@ -1015,34 +1000,30 @@ void UPCGComponent::CleanupLocalDeleteAllGeneratedObjects(const TArray<FPCGTaskI
 		Subsystem->ForAllRegisteredLocalComponents(this, SchedulePurge);
 	}
 
-	CleanupLocal(/*bRemoveComponents=*/true, /*bSave=*/true, TaskIds);
+	CleanupLocal(/*bRemoveComponents=*/true, TaskIds);
 }
 
-void UPCGComponent::CleanupLocal(bool bRemoveComponents, bool bSave)
+void UPCGComponent::CleanupLocal(bool bRemoveComponents)
 {
-	CleanupLocal(bRemoveComponents, bSave, {});
+	CleanupLocal(bRemoveComponents, TArray<FPCGTaskId>());
 }
 
-void UPCGComponent::CleanupLocal(bool bRemoveComponents, bool bSave, const TArray<FPCGTaskId>& Dependencies)
-{
-	CleanupInternal(bRemoveComponents, Dependencies);
-}
-
-FPCGTaskId UPCGComponent::CleanupInternal(bool bRemoveComponents, const TArray<FPCGTaskId>& Dependencies)
+FPCGTaskId UPCGComponent::CleanupLocal(bool bRemoveComponents, const TArray<FPCGTaskId>& Dependencies)
 {
 	if (!GetSubsystem() || IsCleaningUp())
 	{
 		return InvalidPCGTaskId;
 	}
 
-	if (!bGenerated && !IsGenerating() && GeneratedResources.IsEmpty())
+	const bool bNeedsLocalCleanup = bGenerated || IsGenerating() || !GeneratedResources.IsEmpty();
+	if (!bNeedsLocalCleanup && GetSubsystem()->GetPCGComponentPartitionActorMappings(this).IsEmpty())
 	{
 		return InvalidPCGTaskId;
 	}
 
-	PCGGeneratedResourcesLogging::LogCleanupInternal(this, bRemoveComponents);
+	PCGGeneratedResourcesLogging::LogCleanupLocal(this, bRemoveComponents);
 
-	Modify(!IsInPreviewMode());
+	Modify(!IsInPreviewMode() && bNeedsLocalCleanup);
 
 #if WITH_EDITOR
 	ExtraCapture.ResetCapturedMessages();
@@ -1407,7 +1388,8 @@ FPCGTaskId UPCGComponent::CreateCleanupTask(bool bRemoveComponents, const TArray
 
 	if (IsGenerating())
 	{
-		AdditionalDependencies = Dependencies;
+		AdditionalDependencies.Reserve(Dependencies.Num() + 1);
+		AdditionalDependencies.Append(Dependencies);
 		AdditionalDependencies.Add(CurrentGenerationTask);
 		AllDependencies = &AdditionalDependencies;
 	}
diff --git a/Engine/Plugins/PCG/Source/PCG/Private/PCGSubsystem.cpp b/Engine/Plugins/PCG/Source/PCG/Private/PCGSubsystem.cpp
--- a/Engine/Plugins/PCG/Source/PCG/Private/PCGSubsystem.cpp
+++ b/Engine/Plugins/PCG/Source/PCG/Private/PCGSubsystem.cpp
@@ -632,7 +632,7 @@ FPCGTaskId UPCGSubsystem::ScheduleComponent(UPCGComponent* PCGComponent, EPCGHiG
 			if (!GridSizes.Contains(LocalComponent->GetGenerationGridSize()))
 			{
 				// Local component with invalid grid size. Grid sizes may have changed in graph.
-				return LocalComponent->CleanupInternal(/*bRemoveComponents=*/true, Dependencies);
+				return LocalComponent->CleanupLocal(/*bRemoveComponents=*/true, Dependencies);
 			}
 			else if (Grid != EPCGHiGenGrid::Uninitialized && !(Grid & LocalComponent->GetGenerationGrid()))
 			{
@@ -650,14 +650,25 @@ FPCGTaskId UPCGSubsystem::ScheduleComponent(UPCGComponent* PCGComponent, EPCGHiG
 			// Ensure that the PCG actor match our original
 			LocalComponent->SetPropertiesFromOriginal(OriginalComponent);
 
+			FPCGTaskId LocalCleanupTaskId = InvalidPCGTaskId;
 			if (LocalComponent->bGenerated && !OriginalComponent->bGenerated)
 			{
 				// Detected a mismatch between the original component and the local component.
 				// Request a cleanup first
-				LocalComponent->CleanupLocalImmediate(true);
+				LocalCleanupTaskId = LocalComponent->CleanupLocal(/*bRemoveComponents=*/true, Dependencies);
 			}
 
-			return LocalComponent->GenerateInternal(bForce, LocalComponent->GetGenerationGrid(), EPCGComponentGenerationTrigger::GenerateOnDemand, Dependencies);
+			TArray<FPCGTaskId> AdditionalDependencies;
+			const TArray<FPCGTaskId>* AllDependencies = &Dependencies;
+			if (LocalCleanupTaskId != InvalidPCGTaskId)
+			{
+				AdditionalDependencies.Reserve(Dependencies.Num() + 1);
+				AdditionalDependencies.Append(Dependencies);
+				AdditionalDependencies.Add(LocalCleanupTaskId);
+				AllDependencies = &AdditionalDependencies;
+			}
+
+			return LocalComponent->GenerateInternal(bForce, LocalComponent->GetGenerationGrid(), EPCGComponentGenerationTrigger::GenerateOnDemand, *AllDependencies);
 		};
 
 		ExecutionDependencyTasks.Append(ActorAndComponentMapping.DispatchToRegisteredLocalComponents(PCGComponent, LocalGenerateTask));
@@ -731,7 +742,7 @@ FPCGTaskId UPCGSubsystem::ScheduleCleanup(UPCGComponent* PCGComponent, bool bRem
 			}
 
 			// Always executes regardless of local component grid size - clean up as much as possible.
-			return LocalComponent->CleanupInternal(bRemoveComponents, Dependencies);
+			return LocalComponent->CleanupLocal(bRemoveComponents, Dependencies);
 		};
 
 		AllTasks.Append(ActorAndComponentMapping.DispatchToRegisteredLocalComponents(PCGComponent, LocalCleanupTask));
@@ -1303,6 +1314,11 @@ APCGPartitionActor* UPCGSubsystem::FindOrCreatePCGPartitionActor(const FPCGGridD
 	return NewActor;
 }
 
+TSet<TObjectPtr<APCGPartitionActor>> UPCGSubsystem::GetPCGComponentPartitionActorMappings(UPCGComponent* InComponent) const
+{
+	return ActorAndComponentMapping.GetPCGComponentPartitionActorMappings(InComponent);
+}
+
 bool UPCGSubsystem::IsGraphCacheDebuggingEnabled() const
 {
 	return GraphExecutor && GraphExecutor->IsGraphCacheDebuggingEnabled();
@@ -1864,11 +1880,6 @@ void UPCGSubsystem::UpdateMappingPCGComponentPartitionActor(UPCGComponent* InCom
 	ActorAndComponentMapping.UpdateMappingPCGComponentPartitionActor(InComponent); 
 }
 
-TSet<TObjectPtr<APCGPartitionActor>> UPCGSubsystem::GetPCGComponentPartitionActorMappings(UPCGComponent* InComponent) const 
-{ 
-	return ActorAndComponentMapping.GetPCGComponentPartitionActorMappings(InComponent); 
-}
-
 void UPCGSubsystem::OnPCGGraphCancelled(UPCGComponent* InComponent)
 {
 	ActorAndComponentMapping.OnPCGGraphCancelled(InComponent);
diff --git a/Engine/Plugins/PCG/Source/PCG/Private/Utils/PCGGeneratedResourcesLogging.cpp b/Engine/Plugins/PCG/Source/PCG/Private/Utils/PCGGeneratedResourcesLogging.cpp
--- a/Engine/Plugins/PCG/Source/PCG/Private/Utils/PCGGeneratedResourcesLogging.cpp
+++ b/Engine/Plugins/PCG/Source/PCG/Private/Utils/PCGGeneratedResourcesLogging.cpp
@@ -127,7 +127,7 @@ namespace PCGGeneratedResourcesLogging
 #endif
 	}
 
-	void LogCleanupInternal(const UPCGComponent* InComponent, bool bRemoveComponents)
+	void LogCleanupLocal(const UPCGComponent* InComponent, bool bRemoveComponents)
 	{
 #if WITH_EDITOR
 		if (!LogEnabled())
@@ -136,7 +136,7 @@ namespace PCGGeneratedResourcesLogging
 		}
 
 		UE_LOG(LogPCG, Log, TEXT("[PCGMANAGEDRESOURCES]")); // Blank line
-		UE_LOG(LogPCG, Log, TEXT("[PCGMANAGEDRESOURCES] [%s/%s] UPCGComponent::CleanupInternal, bRemoveComponents: %d"),
+		UE_LOG(LogPCG, Log, TEXT("[PCGMANAGEDRESOURCES] [%s/%s] UPCGComponent::CleanupLocal, bRemoveComponents: %d"),
 			(InComponent && InComponent->GetOwner()) ? *InComponent->GetOwner()->GetName() : TEXT("MISSINGCOMPONENT"),
 			(InComponent && InComponent->GetGraph()) ? *InComponent->GetGraph()->GetName() : TEXT("MISSINGGRAPH"),
 			bRemoveComponents);
diff --git a/Engine/Plugins/PCG/Source/PCG/Private/Utils/PCGGeneratedResourcesLogging.h b/Engine/Plugins/PCG/Source/PCG/Private/Utils/PCGGeneratedResourcesLogging.h
--- a/Engine/Plugins/PCG/Source/PCG/Private/Utils/PCGGeneratedResourcesLogging.h
+++ b/Engine/Plugins/PCG/Source/PCG/Private/Utils/PCGGeneratedResourcesLogging.h
@@ -15,7 +15,7 @@ namespace PCGGeneratedResourcesLogging
 {
 	void LogAddToManagedResources(const UPCGComponent* InComponent, UPCGManagedResource* InResource);
 	
-	void LogCleanupInternal(const UPCGComponent* InComponent, bool bRemoveComponents);
+	void LogCleanupLocal(const UPCGComponent* InComponent, bool bRemoveComponents);
 	
 	void LogCleanupLocalImmediate(const UPCGComponent* InComponent, bool bHardRelease, const TArray<UPCGManagedResource*>& GeneratedResources);
 	void LogCleanupLocalImmediateResource(const UPCGComponent* InComponent, UPCGManagedResource* InResource);
diff --git a/Engine/Plugins/PCG/Source/PCG/Public/Grid/PCGPartitionActor.h b/Engine/Plugins/PCG/Source/PCG/Public/Grid/PCGPartitionActor.h
--- a/Engine/Plugins/PCG/Source/PCG/Public/Grid/PCGPartitionActor.h
+++ b/Engine/Plugins/PCG/Source/PCG/Public/Grid/PCGPartitionActor.h
@@ -85,6 +85,8 @@ public:
 	void AddGraphInstance(UPCGComponent* OriginalComponent);
 	void RemapGraphInstance(const UPCGComponent* OldOriginalComponent, UPCGComponent* NewOriginalComponent);
 	bool RemoveGraphInstance(UPCGComponent* OriginalComponent);
+
+	UE_DEPRECATED(5.6, "This function will be removed")
 	void CleanupDeadGraphInstances(bool bRemoveNonNullOnly = false);
 
 	// When a local component is destroyed. It calls this function. We make sure we don't keep mappings that are dead.
@@ -124,6 +126,8 @@ public:
 	FGuid PCGGuid;
 
 private:
+	void CleanupDeadGraphInstancesInternal();
+
 	UPCGSubsystem* GetSubsystem() const;
 
 	friend class FPCGActorAndComponentMapping;
@@ -139,7 +143,7 @@ private:
 	UPCGComponent* GetLocalComponent(const UPCGComponent* OriginalComponent, bool bRebuildMappingOnNullEntries) const;
 
 	// Note: this map is not a property and not serialized since we will rebuild it from the LocalToOriginal
-	mutable TMap<TObjectPtr<UPCGComponent>, TObjectPtr<UPCGComponent>> OriginalToLocal;
+	mutable TMap<TWeakObjectPtr<UPCGComponent>, TObjectPtr<UPCGComponent>> OriginalToLocal;
 
 	UPROPERTY(NonTransactional)
 	TMap<TObjectPtr<UPCGComponent>, TSoftObjectPtr<UPCGComponent>> LocalToOriginal;
@@ -148,11 +152,6 @@ private:
 	UPROPERTY(Transient)
 	TMap<TObjectPtr<UPCGComponent>, TSoftObjectPtr<UPCGComponent>> LoadedPreviewComponents;
 
-#if WITH_EDITORONLY_DATA
-	UPROPERTY()
-	TMap<TObjectPtr<UPCGComponent>, TWeakObjectPtr<UPCGComponent>> LocalToOriginalMap_DEPRECATED;
-#endif
-
 	UPROPERTY(VisibleAnywhere, Category = WorldPartition)
 	uint32 PCGGridSize;
 
diff --git a/Engine/Plugins/PCG/Source/PCG/Public/PCGActorAndComponentMapping.h b/Engine/Plugins/PCG/Source/PCG/Public/PCGActorAndComponentMapping.h
--- a/Engine/Plugins/PCG/Source/PCG/Public/PCGActorAndComponentMapping.h
+++ b/Engine/Plugins/PCG/Source/PCG/Public/PCGActorAndComponentMapping.h
@@ -79,7 +79,7 @@ public:
 	void UnregisterPCGComponent(UPCGComponent* InComponent, bool bForce = false);
 
 	/** Register a new Partition actor, will be added to a map and will query all intersecting volume to bind to them if asked. Thread safe */
-	void RegisterPartitionActor(APCGPartitionActor* InActor, bool bDoComponentMapping = true);
+	void RegisterPartitionActor(APCGPartitionActor* InActor);
 
 	/** Unregister a Partition actor, will be removed from the map and remove itself to all intersecting volumes. Thread safe */
 	void UnregisterPartitionActor(APCGPartitionActor* InActor);
@@ -149,7 +149,7 @@ private:
 	void ForAllIntersectingPartitionActors(const FBox& InBounds, TFunctionRef<void(APCGPartitionActor*)> InFunc) const;
 
 	/** Update the current mapping between a PCG component and its PCG Partition actors */
-	void UpdateMappingPCGComponentPartitionActor(UPCGComponent* InComponent);
+	void UpdateMappingPCGComponentPartitionActor(UPCGComponent* InComponent, bool bInChangeGraphInstances = true);
 
 	/** Returns the current mapping between a PCG component and its PCG Partition actors */
 	TSet<TObjectPtr<APCGPartitionActor>> GetPCGComponentPartitionActorMappings(UPCGComponent* InComponent) const;
diff --git a/Engine/Plugins/PCG/Source/PCG/Public/PCGComponent.h b/Engine/Plugins/PCG/Source/PCG/Public/PCGComponent.h
--- a/Engine/Plugins/PCG/Source/PCG/Public/PCGComponent.h
+++ b/Engine/Plugins/PCG/Source/PCG/Public/PCGComponent.h
@@ -178,9 +178,15 @@ public:
 
 	/** Cleans up the generation from a local (vs. remote) standpoint. Will not be replicated. Will be delayed. */
 	UFUNCTION(BlueprintCallable, Category = PCG)
-	void CleanupLocal(bool bRemoveComponents, bool bSave = false);
+	void CleanupLocal(bool bRemoveComponents);
 
-	void CleanupLocal(bool bRemoveComponents, bool bSave, const TArray<FPCGTaskId>& Dependencies);
+	FPCGTaskId CleanupLocal(bool bRemoveComponents, const TArray<FPCGTaskId>& Dependencies);
+
+	UE_DEPRECATED(5.6, "Use version with no bSave parameter.")
+	void CleanupLocal(bool bRemoveComponents, bool bSave) { CleanupLocal(bRemoveComponents); }
+
+	UE_DEPRECATED(5.6, "Use version with no bSave parameter.")
+	void CleanupLocal(bool bRemoveComponents, bool bSave, const TArray<FPCGTaskId>& Dependencies) { CleanupLocal(bRemoveComponents, Dependencies); }
 
 	/** Cleanup the generation while purging Actors and Components tagged as generated by PCG but are no longer managed by this or any other actor*/
 	void CleanupLocalDeleteAllGeneratedObjects(const TArray<FPCGTaskId>& Dependencies);
@@ -197,7 +203,7 @@ public:
 
 	/** Networked cleanup call */
 	UFUNCTION(BlueprintCallable, NetMulticast, Reliable, Category = PCG)
-	void Cleanup(bool bRemoveComponents, bool bSave = false);
+	void Cleanup(bool bRemoveComponents);
 
 	/** Cancels in-progress generation */
 	void CancelGeneration();
@@ -526,7 +532,6 @@ private:
 
 	/* Internal call that allows to delay a Generate/Cleanup call, chain with dependencies and keep track of the task id created. This task id is also returned. */
 	FPCGTaskId GenerateInternal(bool bForce, EPCGHiGenGrid Grid, EPCGComponentGenerationTrigger RequestedGenerationTrigger, const TArray<FPCGTaskId>& Dependencies);
-	FPCGTaskId CleanupInternal(bool bRemoveComponents, const TArray<FPCGTaskId>& Dependencies);
 
 	/* Internal call to create tasks to generate the component. If there is nothing to do, an invalid task id will be returned. Should only be used by the subsystem. */
 	FPCGTaskId CreateGenerateTask(bool bForce, const TArray<FPCGTaskId>& Dependencies);
diff --git a/Engine/Plugins/PCG/Source/PCG/Public/PCGSubsystem.h b/Engine/Plugins/PCG/Source/PCG/Public/PCGSubsystem.h
--- a/Engine/Plugins/PCG/Source/PCG/Public/PCGSubsystem.h
+++ b/Engine/Plugins/PCG/Source/PCG/Public/PCGSubsystem.h
@@ -216,7 +216,7 @@ public:
 	void UnregisterPCGComponent(UPCGComponent* InComponent, bool bForce = false) { ActorAndComponentMapping.UnregisterPCGComponent(InComponent, bForce); }
 
 	/** Register a new Partition actor, will be added to a map and will query all intersecting volume to bind to them if asked. Thread safe */
-	void RegisterPartitionActor(APCGPartitionActor* InActor, bool bDoComponentMapping = true) { ActorAndComponentMapping.RegisterPartitionActor(InActor, bDoComponentMapping); }
+	void RegisterPartitionActor(APCGPartitionActor* InActor) { ActorAndComponentMapping.RegisterPartitionActor(InActor); }
 
 	/** Unregister a Partition actor, will be removed from the map and remove itself to all intersecting volumes. Thread safe */
 	void UnregisterPartitionActor(APCGPartitionActor* InActor) { ActorAndComponentMapping.UnregisterPartitionActor(InActor); }
@@ -275,6 +275,9 @@ public:
 	/** Creates a new partition actor if one does not already exist with the same grid size, coords, and generation mode. */
 	APCGPartitionActor* FindOrCreatePCGPartitionActor(const FPCGGridDescriptor& GridDescriptor, const FIntVector& GridCoords, bool bCanCreateActor = true, bool bHideFromOutliner = false) const;
 	
+	/** Retrieves partition actors for this original component */
+	TSet<TObjectPtr<APCGPartitionActor>> GetPCGComponentPartitionActorMappings(UPCGComponent* InComponent) const;
+
 	/** True if graph cache debugging is enabled. */
 	bool IsGraphCacheDebuggingEnabled() const;
 
@@ -365,7 +368,6 @@ private:
 
 	void CreatePartitionActorsWithinBounds(UPCGComponent* InComponent, const FBox& InBounds, const PCGHiGenGrid::FSizeArray& InGridSizes);
 	void UpdateMappingPCGComponentPartitionActor(UPCGComponent* InComponent);
-	TSet<TObjectPtr<APCGPartitionActor>> GetPCGComponentPartitionActorMappings(UPCGComponent* InComponent) const;
 
 	FPCGNodeVisualLogs NodeVisualLogs;
 #endif // WITH_EDITOR
diff --git a/Engine/Plugins/PCG/Source/PCGEditor/Private/WorldPartitionBuilder/PCGWorldPartitionBuilder.cpp b/Engine/Plugins/PCG/Source/PCGEditor/Private/WorldPartitionBuilder/PCGWorldPartitionBuilder.cpp
--- a/Engine/Plugins/PCG/Source/PCGEditor/Private/WorldPartitionBuilder/PCGWorldPartitionBuilder.cpp
+++ b/Engine/Plugins/PCG/Source/PCGEditor/Private/WorldPartitionBuilder/PCGWorldPartitionBuilder.cpp
@@ -1601,7 +1601,7 @@ bool FPCGWorldPartitionBuilder::UpdatePartitionedActors(UWorld* InWorld, const T
 		if (PartitionActor->GetLevel() == InWorld->PersistentLevel)
 		{
 			// Here we cleanup dead graphs since we tried to load the Original components before, if they can't resolve it means they don't exist anymore.
-			PartitionActor->CleanupDeadGraphInstances();
+			PartitionActor->CleanupDeadGraphInstancesInternal();
 
 			// Remove Graph Instances from Partition Actors that are no longer part of the component mappings (bounds change)
 			for (const TKeyValuePair<UPCGComponent*, TSet<TObjectPtr<APCGPartitionActor>>>& Mapping : UpdatedMappings)
