[PCG] GPU: Fix object referencing that I broke when refactoring the lazy readback changes.

#rnx

diff --git a/Engine/Plugins/PCG/Source/PCG/Internal/Compute/PCGDataForGPU.h b/Engine/Plugins/PCG/Source/PCG/Internal/Compute/PCGDataForGPU.h
--- a/Engine/Plugins/PCG/Source/PCG/Internal/Compute/PCGDataForGPU.h
+++ b/Engine/Plugins/PCG/Source/PCG/Internal/Compute/PCGDataForGPU.h
@@ -211,8 +211,6 @@ public:
 	/** Populates a CPU data object representing the GPU data for the given index, performing a readback from GPU->CPU if required. */
 	bool GetCPUData(FPCGContext* InContext, int32 InDataIndex, const UPCGData*& OutData);
 
-	void AddReferences(FReferenceCollector& Collector);
-
 	const TRefCountPtr<FRDGPooledBuffer>& GetBuffer() const { return Buffer; }
 
 	const FPCGDataCollectionDesc& GetDescription() const { return Description; }
@@ -229,7 +227,7 @@ protected:
 	TArray<FString> StringTable;
 
 	/** Read back data. Populated once upon first readback request. */
-	FPCGDataCollection ReadbackData;
+	TArray<TStrongObjectPtr<const UPCGData>> ReadbackData;
 
 	TSharedPtr<FRHIGPUBufferReadback> ReadbackRequest;
 
diff --git a/Engine/Plugins/PCG/Source/PCG/Private/Compute/PCGDataForGPU.cpp b/Engine/Plugins/PCG/Source/PCG/Private/Compute/PCGDataForGPU.cpp
--- a/Engine/Plugins/PCG/Source/PCG/Private/Compute/PCGDataForGPU.cpp
+++ b/Engine/Plugins/PCG/Source/PCG/Private/Compute/PCGDataForGPU.cpp
@@ -1632,7 +1632,7 @@ bool FPCGProxyForGPUDataCollection::GetCPUData(FPCGContext* InContext, int32 InD
 
 	if (bReadbackDataProcessed)
 	{
-		OutData = ReadbackData.TaggedData.IsValidIndex(InDataIndex) ? ReadbackData.TaggedData[InDataIndex].Data : nullptr;
+		OutData = ReadbackData.IsValidIndex(InDataIndex) ? ReadbackData[InDataIndex].Get() : nullptr;
 		return true;
 	}
 
@@ -1712,16 +1712,15 @@ bool FPCGProxyForGPUDataCollection::GetCPUData(FPCGContext* InContext, int32 InD
 			return true;
 		}
 
-		ReadbackData = MoveTemp(DataFromGPU);
+		ReadbackData.Reset();
+		for (const FPCGTaggedData& TaggedData : DataFromGPU.TaggedData)
+		{
+			ReadbackData.Add(TStrongObjectPtr<const UPCGData>(TaggedData.Data));
+		}
 	}
 
-	OutData = ReadbackData.TaggedData.IsValidIndex(InDataIndex) ? ReadbackData.TaggedData[InDataIndex].Data : nullptr;
+	OutData = ReadbackData.IsValidIndex(InDataIndex) ? ReadbackData[InDataIndex].Get() : nullptr;
 	return true;
 }
 
-void FPCGProxyForGPUDataCollection::AddReferences(FReferenceCollector& Collector)
-{
-	ReadbackData.AddReferences(Collector);
-}
-
 #undef LOCTEXT_NAMESPACE
