[PCG] Adjust Print Element to maintain consisting logging style between Warning/Error and Display/Log

#rnx

diff --git a/Engine/Plugins/PCG/Source/PCG/Private/Elements/PCGPrintElement.cpp b/Engine/Plugins/PCG/Source/PCG/Private/Elements/PCGPrintElement.cpp
--- a/Engine/Plugins/PCG/Source/PCG/Private/Elements/PCGPrintElement.cpp
+++ b/Engine/Plugins/PCG/Source/PCG/Private/Elements/PCGPrintElement.cpp
@@ -116,10 +116,24 @@ bool FPCGPrintElement::ExecuteInternal(FPCGContext* Context) const
 	switch (Settings->Verbosity)
 	{
 		case EPCGPrintVerbosity::Error:
-			PCGLog::LogErrorOnGraph(FText::FromString(FinalString), Settings->bDisplayOnNode ? Context : nullptr);
+			if (Settings->bDisplayOnNode)
+			{
+				PCGE_LOG_C(Error, GraphAndLog, Context, FText::FromString(FinalString));
+			}
+			else
+			{
+				PCGE_LOG_C(Error, LogOnly, Context, FText::FromString(FinalString));
+			}
 			break;
 		case EPCGPrintVerbosity::Warning:
-			PCGLog::LogWarningOnGraph(FText::FromString(FinalString), Settings->bDisplayOnNode ? Context : nullptr);
+			if (Settings->bDisplayOnNode)
+			{
+				PCGE_LOG_C(Warning, GraphAndLog, Context, FText::FromString(FinalString));
+			}
+			else
+			{
+				PCGE_LOG_C(Warning, LogOnly, Context, FText::FromString(FinalString));
+			}
 			break;
 		case EPCGPrintVerbosity::Log:
 			PCGE_LOG_C(Log, LogOnly, Context, FText::FromString(FinalString));
