[PCG] Add support for dynamic pins in Merge Nodes

diff --git a/Engine/Plugins/PCG/Source/PCG/Private/Elements/PCGMergeAttributes.cpp b/Engine/Plugins/PCG/Source/PCG/Private/Elements/PCGMergeAttributes.cpp
--- a/Engine/Plugins/PCG/Source/PCG/Private/Elements/PCGMergeAttributes.cpp
+++ b/Engine/Plugins/PCG/Source/PCG/Private/Elements/PCGMergeAttributes.cpp
@@ -7,6 +7,8 @@
 #include "PCGPin.h"
 #include "Metadata/PCGMetadata.h"
 
+#include UE_INLINE_GENERATED_CPP_BY_NAME(PCGMergeAttributes)
+
 #define LOCTEXT_NAMESPACE "PCGMergeAttributesSettings"
 
 #if WITH_EDITOR
@@ -21,27 +23,40 @@ FText UPCGMergeAttributesSettings::GetNodeTooltipText() const
 }
 #endif // WITH_EDITOR
 
-TArray<FPCGPinProperties> UPCGMergeAttributesSettings::InputPinProperties() const
+TArray<FPCGPinProperties> UPCGMergeAttributesSettings::OutputPinProperties() const
 {
 	TArray<FPCGPinProperties> PinProperties;
-	FPCGPinProperties& InputPinProperty = PinProperties.Emplace_GetRef(PCGPinConstants::DefaultInputLabel, EPCGDataType::Param, /*bAllowMultipleConnections=*/true);
-	InputPinProperty.SetRequiredPin();
+	PinProperties.Emplace(PCGPinConstants::DefaultOutputLabel, EPCGDataType::Param, /*bAllowMultipleConnections=*/true, /*bAllowMultipleData=*/false);
 
 	return PinProperties;
 }
 
-TArray<FPCGPinProperties> UPCGMergeAttributesSettings::OutputPinProperties() const
+FPCGElementPtr UPCGMergeAttributesSettings::CreateElement() const
+{
+	return MakeShared<FPCGMergeAttributesElement>();
+}
+
+FName UPCGMergeAttributesSettings::GetDynamicInputPinsBaseLabel() const
+{
+	return PCGPinConstants::DefaultInputLabel;
+}
+
+TArray<FPCGPinProperties> UPCGMergeAttributesSettings::StaticInputPinProperties() const
 {
 	TArray<FPCGPinProperties> PinProperties;
-	PinProperties.Emplace(PCGPinConstants::DefaultOutputLabel, EPCGDataType::Param, /*bAllowMultipleConnections=*/true, /*bAllowMultipleData=*/false);
+	FPCGPinProperties& InputPinProperty = PinProperties.Emplace_GetRef(PCGPinConstants::DefaultInputLabel, EPCGDataType::Param, /*bAllowMultipleConnections=*/true);
+	InputPinProperty.SetRequiredPin();
 
 	return PinProperties;
 }
 
-FPCGElementPtr UPCGMergeAttributesSettings::CreateElement() const
+#if WITH_EDITOR
+void UPCGMergeAttributesSettings::AddDefaultDynamicInputPin()
 {
-	return MakeShared<FPCGMergeAttributesElement>();
+	FPCGPinProperties SecondaryPinProperties(FName(GetDynamicInputPinsBaseLabel().ToString() + FString::FromInt(DynamicInputPinProperties.Num() + 2)), EPCGDataType::Param);
+	AddDynamicInputPin(std::move(SecondaryPinProperties));
 }
+#endif // WITH_EDITOR
 
 bool FPCGMergeAttributesElement::ExecuteInternal(FPCGContext* Context) const
 {
@@ -53,7 +68,15 @@ bool FPCGMergeAttributesElement::ExecuteInternal(FPCGContext* Context) const
 	FPCGTaggedData* MergedOutput = nullptr;
 	UPCGParamData* MergedAttributeSet = nullptr;
 
-	for (const FPCGTaggedData& Source : Context->InputData.TaggedData)
+	TArray<FPCGTaggedData> Sources;
+	for (const UPCGPin* InputPin : Context->Node->GetInputPins())
+	{
+		check(InputPin);
+
+		Sources.Append(Context->InputData.GetInputsByPin(InputPin->Properties.Label));
+	}
+
+	for (const FPCGTaggedData& Source : Sources)
 	{
 		const UPCGParamData* SourceData = Cast<const UPCGParamData>(Source.Data);
 
diff --git a/Engine/Plugins/PCG/Source/PCG/Private/Elements/PCGMergeElement.cpp b/Engine/Plugins/PCG/Source/PCG/Private/Elements/PCGMergeElement.cpp
--- a/Engine/Plugins/PCG/Source/PCG/Private/Elements/PCGMergeElement.cpp
+++ b/Engine/Plugins/PCG/Source/PCG/Private/Elements/PCGMergeElement.cpp
@@ -17,11 +17,6 @@ FText UPCGMergeSettings::GetNodeTooltipText() const
 }
 #endif
 
-TArray<FPCGPinProperties> UPCGMergeSettings::InputPinProperties() const
-{
-	return Super::DefaultPointInputPinProperties();
-}
-
 TArray<FPCGPinProperties> UPCGMergeSettings::OutputPinProperties() const
 {
 	TArray<FPCGPinProperties> PinProperties;
@@ -35,6 +30,24 @@ FPCGElementPtr UPCGMergeSettings::CreateElement() const
 	return MakeShared<FPCGMergeElement>();
 }
 
+FName UPCGMergeSettings::GetDynamicInputPinsBaseLabel() const
+{
+	return PCGPinConstants::DefaultInputLabel;
+}
+
+TArray<FPCGPinProperties> UPCGMergeSettings::StaticInputPinProperties() const
+{
+	return DefaultPointInputPinProperties();
+}
+
+#if WITH_EDITOR
+void UPCGMergeSettings::AddDefaultDynamicInputPin()
+{
+	FPCGPinProperties SecondaryPinProperties(FName(GetDynamicInputPinsBaseLabel().ToString() + FString::FromInt(DynamicInputPinProperties.Num() + 2)), EPCGDataType::Point);
+	AddDynamicInputPin(std::move(SecondaryPinProperties));
+}
+#endif // WITH_EDITOR
+
 bool FPCGMergeElement::ExecuteInternal(FPCGContext* Context) const
 {
 	TRACE_CPUPROFILER_EVENT_SCOPE(FPCGMergeElement::Execute);
@@ -43,7 +56,14 @@ bool FPCGMergeElement::ExecuteInternal(FPCGContext* Context) const
 	const UPCGMergeSettings* Settings = Context->GetInputSettings<UPCGMergeSettings>();
 	const bool bMergeMetadata = !Settings || Settings->bMergeMetadata;
 
-	TArray<FPCGTaggedData> Sources = Context->InputData.GetInputsByPin(PCGPinConstants::DefaultInputLabel);
+	TArray<FPCGTaggedData> Sources;
+	for (const UPCGPin* InputPin : Context->Node->GetInputPins())
+	{
+		check(InputPin);
+
+		Sources.Append(Context->InputData.GetInputsByPin(InputPin->Properties.Label));
+	}
+
 	TArray<FPCGTaggedData>& Outputs = Context->OutputData.TaggedData;
 
 	UPCGPointData* TargetPointData = nullptr;
diff --git a/Engine/Plugins/PCG/Source/PCG/Public/Elements/PCGMergeAttributes.h b/Engine/Plugins/PCG/Source/PCG/Public/Elements/PCGMergeAttributes.h
--- a/Engine/Plugins/PCG/Source/PCG/Public/Elements/PCGMergeAttributes.h
+++ b/Engine/Plugins/PCG/Source/PCG/Public/Elements/PCGMergeAttributes.h
@@ -2,13 +2,13 @@
 
 #pragma once
 
-#include "PCGSettings.h"
+#include "PCGSettingsWithDynamicInputs.h"
 
 #include "PCGMergeAttributes.generated.h"
 
 /** Merges multiple attribute sets together into a single attribute set */
 UCLASS(MinimalAPI, BlueprintType, Classgroup = (Procedural))
-class UPCGMergeAttributesSettings : public UPCGSettings
+class UPCGMergeAttributesSettings : public UPCGSettingsWithDynamicInputs
 {
 	GENERATED_BODY()
 
@@ -23,10 +23,20 @@ public:
 	
 
 protected:
-	virtual TArray<FPCGPinProperties> InputPinProperties() const override;
 	virtual TArray<FPCGPinProperties> OutputPinProperties() const override;
 	virtual FPCGElementPtr CreateElement() const override;
 	//~End UPCGSettings interface
+
+	//~Begin UPCGDynamicSettings interface
+	virtual FName GetDynamicInputPinsBaseLabel() const override;
+	/** The input pin properties that are statically defined */
+	virtual TArray<FPCGPinProperties> StaticInputPinProperties() const override;
+
+public:
+#if WITH_EDITOR
+	virtual void AddDefaultDynamicInputPin() override;
+#endif // WITH_EDITOR
+	//~End UPCGDynamicSettings interface
 };
 
 class FPCGMergeAttributesElement : public IPCGElement
diff --git a/Engine/Plugins/PCG/Source/PCG/Public/Elements/PCGMergeElement.h b/Engine/Plugins/PCG/Source/PCG/Public/Elements/PCGMergeElement.h
--- a/Engine/Plugins/PCG/Source/PCG/Public/Elements/PCGMergeElement.h
+++ b/Engine/Plugins/PCG/Source/PCG/Public/Elements/PCGMergeElement.h
@@ -2,13 +2,13 @@
 
 #pragma once
 
-#include "PCGSettings.h"
+#include "PCGSettingsWithDynamicInputs.h"
 
 #include "PCGMergeElement.generated.h"
 
 /** Merges multiple data sources (currently only points supported) into a single output. */
 UCLASS(MinimalAPI, BlueprintType, Classgroup = (Procedural))
-class UPCGMergeSettings : public UPCGSettings
+class UPCGMergeSettings : public UPCGSettingsWithDynamicInputs
 {
 	GENERATED_BODY()
 
@@ -20,16 +20,24 @@ public:
 	virtual FText GetNodeTooltipText() const override;
 	virtual EPCGSettingsType GetType() const override { return EPCGSettingsType::Spatial; }
 #endif
-	
 
-	virtual TArray<FPCGPinProperties> InputPinProperties() const override;
 	virtual TArray<FPCGPinProperties> OutputPinProperties() const override;
 
 protected:
 	virtual FPCGElementPtr CreateElement() const override;
 	//~End UPCGSettings interface
 
+	//~Begin UPCGDynamicSettings interface
+	virtual FName GetDynamicInputPinsBaseLabel() const override;
+	/** The input pin properties that are statically defined */
+	virtual TArray<FPCGPinProperties> StaticInputPinProperties() const override;
+
 public:
+#if WITH_EDITOR
+	virtual void AddDefaultDynamicInputPin() override;
+#endif // WITH_EDITOR
+	//~End UPCGDynamicSettings interface
+
 	/** Controls whether the resulting merge data will have any metadata */
 	UPROPERTY(BlueprintReadWrite, EditAnywhere, Category = Settings, meta = (PCG_Overridable))
 	bool bMergeMetadata = true;
