[PCG] Fix Level Instance related issues (still more work needed)
- Fix PCGWorldActor outliner parenting for sub-level PCG Partitioned Actors
- Prevent some editing in specific cases where PCG Component is inside a non-edit Level Instance

#rnx

diff --git a/Engine/Plugins/PCG/Source/PCG/Private/Grid/PCGPartitionActor.cpp b/Engine/Plugins/PCG/Source/PCG/Private/Grid/PCGPartitionActor.cpp
--- a/Engine/Plugins/PCG/Source/PCG/Private/Grid/PCGPartitionActor.cpp
+++ b/Engine/Plugins/PCG/Source/PCG/Private/Grid/PCGPartitionActor.cpp
@@ -609,7 +609,7 @@ void APCGPartitionActor::RemoveLocalComponent(UPCGComponent* LocalComponent)
 #if WITH_EDITOR
 AActor* APCGPartitionActor::GetSceneOutlinerParent() const
 {
-	if (APCGWorldActor* PCGActor = PCGHelpers::GetPCGWorldActor(GetWorld()))
+	if (APCGWorldActor* PCGActor = PCGHelpers::FindPCGWorldActor(GetLevel()))
 	{
 		return PCGActor;
 	}
diff --git a/Engine/Plugins/PCG/Source/PCG/Private/Helpers/PCGHelpers.cpp b/Engine/Plugins/PCG/Source/PCG/Private/Helpers/PCGHelpers.cpp
--- a/Engine/Plugins/PCG/Source/PCG/Private/Helpers/PCGHelpers.cpp
+++ b/Engine/Plugins/PCG/Source/PCG/Private/Helpers/PCGHelpers.cpp
@@ -13,6 +13,7 @@
 #include "Components/BillboardComponent.h"
 #include "Landscape.h"
 #include "Algo/AnyOf.h"
+#include "Engine/Level.h"
 #include "UObject/UObjectIterator.h"
 
 #if WITH_EDITOR
@@ -317,6 +318,12 @@ namespace PCGHelpers
 		return (InWorld && InWorld->GetSubsystem<UPCGSubsystem>()) ? InWorld->GetSubsystem<UPCGSubsystem>()->GetPCGWorldActor() : nullptr;
 	}
 
+	APCGWorldActor* FindPCGWorldActor(ULevel* InLevel)
+	{
+		UWorld* World = InLevel ? InLevel->GetWorld() : nullptr;
+		return (World && World->GetSubsystem<UPCGSubsystem>()) ? World->GetSubsystem<UPCGSubsystem>()->FindPCGWorldActor(InLevel) : nullptr;
+	}
+
 	// TODO: Temporary validation during transition of allowing spaces in tags/attributes. Deprecate in 5.6.
 	TArray<FString> GetStringArrayFromCommaSeparatedString(const FString& InCommaSeparatedString, const FPCGContext* InOptionalContext)
 	{
diff --git a/Engine/Plugins/PCG/Source/PCG/Private/PCGComponent.cpp b/Engine/Plugins/PCG/Source/PCG/Private/PCGComponent.cpp
--- a/Engine/Plugins/PCG/Source/PCG/Private/PCGComponent.cpp
+++ b/Engine/Plugins/PCG/Source/PCG/Private/PCGComponent.cpp
@@ -1861,7 +1861,7 @@ void UPCGComponent::SetupCallbacksOnCreation()
 bool UPCGComponent::CanEditChange(const FProperty* InProperty) const
 {
 	// Can't change anything if the component is local
-	return !bIsComponentLocal;
+	return !bIsComponentLocal && Super::CanEditChange(InProperty);
 }
 #endif
 
diff --git a/Engine/Plugins/PCG/Source/PCG/Private/PCGGraph.cpp b/Engine/Plugins/PCG/Source/PCG/Private/PCGGraph.cpp
--- a/Engine/Plugins/PCG/Source/PCG/Private/PCGGraph.cpp
+++ b/Engine/Plugins/PCG/Source/PCG/Private/PCGGraph.cpp
@@ -24,6 +24,7 @@
 #include "Editor.h"
 #include "Dialogs/Dialogs.h"
 #include "EdGraph/EdGraphPin.h"
+#include "GameFramework/Actor.h"
 #else
 #include "UObject/Package.h"
 #endif
@@ -2039,10 +2040,18 @@ void UPCGGraphInstance::OnGraphChanged(UPCGGraphInterface* InGraph, EPCGChangeTy
 
 bool UPCGGraphInstance::CanEditChange(const FProperty* InProperty) const
 {
+	UPCGComponent* Component = Cast<UPCGComponent>(GetOuter());
+	AActor* ComponentOwner = Component ? Component->GetOwner() : nullptr;
+		
+	if(ComponentOwner && ComponentOwner->IsInLevelInstance() && !ComponentOwner->IsInEditLevelInstance())
+	{
+		return false;
+	}
+
 	// Graph can only be changed if it is not in a local PCGComponent
 	if (InProperty && InProperty->GetFName() == GET_MEMBER_NAME_CHECKED(UPCGGraphInstance, Graph))
 	{
-		if (UPCGComponent* Component = Cast<UPCGComponent>(GetOuter()))
+		if (Component)
 		{
 			return !Component->IsLocalComponent();
 		}
diff --git a/Engine/Plugins/PCG/Source/PCG/Private/PCGSubsystem.cpp b/Engine/Plugins/PCG/Source/PCG/Private/PCGSubsystem.cpp
--- a/Engine/Plugins/PCG/Source/PCG/Private/PCGSubsystem.cpp
+++ b/Engine/Plugins/PCG/Source/PCG/Private/PCGSubsystem.cpp
@@ -327,11 +327,29 @@ APCGWorldActor* UPCGSubsystem::GetPCGWorldActor()
 	return PCGWorldActor;
 }
 
-APCGWorldActor* UPCGSubsystem::FindPCGWorldActor()
+APCGWorldActor* UPCGSubsystem::FindPCGWorldActor() const
 {
 	return PCGWorldActor;
 }
 
+APCGWorldActor* UPCGSubsystem::FindPCGWorldActor(ULevel* InLevel) const
+{
+	ensure(IsInGameThread());
+	if (ensure(InLevel))
+	{
+		if (InLevel->IsPersistentLevel())
+		{
+			return PCGWorldActor;
+		}
+		else if(TWeakObjectPtr<APCGWorldActor>const* FoundPCGWorldActor = PCGWorldActors.Find(InLevel))
+		{
+			return FoundPCGWorldActor->Get();
+		}
+	}
+
+	return nullptr;
+}
+
 int32 UPCGSubsystem::GetPCGQualityLevel()
 {
 	return PCGSubsystemConsole::CVarPCGQuality.GetValueOnAnyThread();
@@ -483,23 +501,46 @@ void UPCGSubsystem::LogAbnormalComponentStates(bool bGroupByState) const
 
 void UPCGSubsystem::RegisterPCGWorldActor(APCGWorldActor* InActor)
 {
-	// TODO: we should support merging or multi world actor support when relevant
-	if (!PCGWorldActor)
-	{
-		PCGWorldActor = InActor;
-	}
-	else if (InActor != PCGWorldActor)
+	// @todo_pcg: Revisit the need for merging, for now do not merge multiple PCGWorldActors of sub-levels
+	check(InActor);
+	if (ULevel* Level = InActor->GetLevel())
 	{
-		PCGWorldActor->MergeFrom(InActor);
+		if (Level->IsPersistentLevel())
+		{
+			if (!PCGWorldActor)
+			{
+				PCGWorldActor = InActor;
+			}
+			else if (InActor != PCGWorldActor)
+			{
+				PCGWorldActor->MergeFrom(InActor);
+			}
+		}
+		else 
+		{
+			ensure(IsInGameThread());
+			TWeakObjectPtr<APCGWorldActor>& FoundActor = PCGWorldActors.FindOrAdd(Level, InActor);
+			ensure(FoundActor.Get() == InActor);
+		}
 	}
 }
 
 void UPCGSubsystem::UnregisterPCGWorldActor(APCGWorldActor* InActor)
 {
+	check(InActor);
+
 	if (PCGWorldActor == InActor)
 	{
 		PCGWorldActor = nullptr;
 	}
+	else if(ULevel* Level = InActor->GetLevel())
+	{
+		TWeakObjectPtr<APCGWorldActor> FoundActor;
+		if (PCGWorldActors.RemoveAndCopyValue(Level, FoundActor))
+		{
+			ensure(FoundActor.GetEvenIfUnreachable() == InActor);
+		}
+	}
 }
 
 void UPCGSubsystem::OnOriginalComponentRegistered(UPCGComponent* InComponent)
diff --git a/Engine/Plugins/PCG/Source/PCG/Private/PCGWorldActor.cpp b/Engine/Plugins/PCG/Source/PCG/Private/PCGWorldActor.cpp
--- a/Engine/Plugins/PCG/Source/PCG/Private/PCGWorldActor.cpp
+++ b/Engine/Plugins/PCG/Source/PCG/Private/PCGWorldActor.cpp
@@ -159,6 +159,56 @@ void APCGWorldActor::MergeFrom(APCGWorldActor* OtherWorldActor)
 }
 
 #if WITH_EDITOR
+void APCGWorldActor::PostRegisterAllComponents()
+{
+	Super::PostRegisterAllComponents();
+
+	UWorld* World = GetWorld();
+	ULevel* Level = GetLevel();
+	// Register sub level PCGWorldActors
+	if (World && Level && !World->IsGameWorld() && !Level->IsPersistentLevel())
+	{
+		RegisterToSubsystem();
+	}
+}
+
+void APCGWorldActor::PostUnregisterAllComponents()
+{
+	Super::PostUnregisterAllComponents();
+		
+	ULevel* Level = GetLevel();
+	if (!Level)
+	{
+		return;
+	}
+
+	// If we are not unloading, do not unregister to prevent temporary unregistering on property changes
+	if (!UE::GetIsEditorLoadingPackage() && !Level->bIsBeingRemoved)
+	{
+		return;
+	}
+
+	UWorld* World = GetWorld();
+	if (World && !World->IsGameWorld() && !Level->IsPersistentLevel())
+	{
+		UnregisterFromSubsystem();
+	}
+}
+
+void APCGWorldActor::Destroyed()
+{
+	Super::Destroyed();
+
+	UWorld* World = GetWorld();
+	ULevel* Level = GetLevel();
+
+	// Handle case where actor is deleted
+	if (World && Level && !World->IsGameWorld() && !Level->IsPersistentLevel())
+	{
+		UnregisterFromSubsystem();
+	}
+}
+
 APCGWorldActor* APCGWorldActor::CreatePCGWorldActor(UWorld* InWorld)
 {
 	APCGWorldActor* PCGActor = nullptr;
@@ -268,6 +318,7 @@ void APCGWorldActor::PostEditChangeProperty(FPropertyChangedEvent& PropertyChang
 
 void APCGWorldActor::BeginDestroy()
 {
+	// @todo_pcg: revisit because at this point it is possible that we can't properly unregister
 	UnregisterFromSubsystem();
 	Super::BeginDestroy();
 }
diff --git a/Engine/Plugins/PCG/Source/PCG/Public/Helpers/PCGHelpers.h b/Engine/Plugins/PCG/Source/PCG/Public/Helpers/PCGHelpers.h
--- a/Engine/Plugins/PCG/Source/PCG/Public/Helpers/PCGHelpers.h
+++ b/Engine/Plugins/PCG/Source/PCG/Public/Helpers/PCGHelpers.h
@@ -54,6 +54,7 @@ namespace PCGHelpers
 	PCG_API bool IsRuntimeOrPIE();
 
 	PCG_API APCGWorldActor* GetPCGWorldActor(UWorld* InWorld);
+	PCG_API APCGWorldActor* FindPCGWorldActor(ULevel* InLevel);
 
 	UE_DEPRECATED(5.5, "This function was incorrectly parsing by whitespace. Please use 'GetStringArrayFromCommaSeparatedList' instead.")
 	PCG_API TArray<FString> GetStringArrayFromCommaSeparatedString(const FString& InCommaSeparatedString, const FPCGContext* InOptionalContext = nullptr);
diff --git a/Engine/Plugins/PCG/Source/PCG/Public/PCGSubsystem.h b/Engine/Plugins/PCG/Source/PCG/Public/PCGSubsystem.h
--- a/Engine/Plugins/PCG/Source/PCG/Public/PCGSubsystem.h
+++ b/Engine/Plugins/PCG/Source/PCG/Public/PCGSubsystem.h
@@ -99,7 +99,8 @@ public:
 	bool IsInitialized() const { return GraphExecutor != nullptr; }
 
 	APCGWorldActor* GetPCGWorldActor();
-	APCGWorldActor* FindPCGWorldActor();
+	APCGWorldActor* FindPCGWorldActor() const;
+	APCGWorldActor* FindPCGWorldActor(ULevel* InLevel) const;
 
 	/** Returns current quality level between Low (0) and Cinematic (4). */
 	static int32 GetPCGQualityLevel();
@@ -391,7 +392,12 @@ private:
 
 	void ExecuteBeginTickActions();
 
+	// Main PCG World Actor
 	APCGWorldActor* PCGWorldActor = nullptr;
+
+	// Other Level PCG World Actors
+	TMap<ULevel*, TWeakObjectPtr<APCGWorldActor>> PCGWorldActors;
+
 	FPCGGraphExecutor* GraphExecutor = nullptr;
 	FPCGRuntimeGenScheduler* RuntimeGenScheduler = nullptr;
 	bool bHasTickedOnce = false;
diff --git a/Engine/Plugins/PCG/Source/PCG/Public/PCGWorldActor.h b/Engine/Plugins/PCG/Source/PCG/Public/PCGWorldActor.h
--- a/Engine/Plugins/PCG/Source/PCG/Public/PCGWorldActor.h
+++ b/Engine/Plugins/PCG/Source/PCG/Public/PCGWorldActor.h
@@ -26,6 +26,9 @@ public:
 	virtual void BeginPlay() override;
 	virtual void EndPlay(const EEndPlayReason::Type EndPlayReason) override;
 #if WITH_EDITOR
+	virtual void PostRegisterAllComponents() override;
+	virtual void PostUnregisterAllComponents() override;
+	virtual void Destroyed() override;
 	virtual bool ActorTypeSupportsDataLayer() const override { return false; }
 	virtual bool ActorTypeSupportsExternalDataLayer() const override { return false; }
 	virtual bool CanChangeIsSpatiallyLoadedFlag() const override { return false; }
diff --git a/Engine/Plugins/PCG/Source/PCGEditor/Private/Details/PCGComponentDetails.cpp b/Engine/Plugins/PCG/Source/PCGEditor/Private/Details/PCGComponentDetails.cpp
--- a/Engine/Plugins/PCG/Source/PCGEditor/Private/Details/PCGComponentDetails.cpp
+++ b/Engine/Plugins/PCG/Source/PCGEditor/Private/Details/PCGComponentDetails.cpp
@@ -61,6 +61,19 @@ void FPCGComponentDetails::CustomizeDetails(IDetailLayoutBuilder& DetailBuilder)
 		}
 	}
 
+	for (TWeakObjectPtr<UPCGComponent>& ComponentPtr : SelectedComponents)
+	{
+		if (UPCGComponent* Component = ComponentPtr.Get())
+		{
+			AActor* ComponentOwner = Component->GetOwner();
+			if (!ComponentOwner || (ComponentOwner->IsInLevelInstance() && !ComponentOwner->IsInEditLevelInstance()))
+			{
+				// Do not customize for non editing Level Instance
+				return;
+			}
+		}
+	}
+
 	FDetailWidgetRow& NewRow = PCGCategory.AddCustomRow(FText::GetEmpty());
 
 	NewRow.ValueContent()
