Stop gap fix to prevent resource loss when user parameters change on BP in setups with components on remote actors.

#rnx

diff --git a/Engine/Plugins/PCG/Source/PCG/Private/PCGManagedResource.cpp b/Engine/Plugins/PCG/Source/PCG/Private/PCGManagedResource.cpp
--- a/Engine/Plugins/PCG/Source/PCG/Private/PCGManagedResource.cpp
+++ b/Engine/Plugins/PCG/Source/PCG/Private/PCGManagedResource.cpp
@@ -22,7 +22,9 @@ static TAutoConsoleVariable<bool> CVarForceReleaseResourcesOnGenerate(
 
 void UPCGManagedResource::PostApplyToComponent()
 {
-	PostEditImport();
+	// Nothing - apply to component should already properly remap most of everything we need to do
+	// In the case of actors, this means we'll keep the references to the actors as-is
+	// In the cas of components, the remapping will not be needed either since they won't be affected this way.
 }
 
 // By default, if it is not a hard release, we mark the resource unused.
@@ -68,16 +70,11 @@ void UPCGManagedResource::ChangeTransientState(EPCGEditorDirtyMode /*NewEditingM
 void UPCGManagedActors::PostEditImport()
 {
 	// In this case, the managed actors won't be copied along the actor/component,
-	// So we just have to "forget" the actors.
+	// So we just have to "forget" the actors, leaving the ownership to the original actor only.
 	Super::PostEditImport();
 	GeneratedActors.Reset();
 }
 
-void UPCGManagedActors::PostApplyToComponent()
-{
-	// In this case, we want to preserve the data, so we need to do nothing
-}
-
 bool UPCGManagedActors::Release(bool bHardRelease, TSet<TSoftObjectPtr<AActor>>& OutActorsToDelete)
 {
 	TRACE_CPUPROFILER_EVENT_SCOPE(UPCGManagedActors::Release);
diff --git a/Engine/Plugins/PCG/Source/PCG/Public/PCGManagedResource.h b/Engine/Plugins/PCG/Source/PCG/Public/PCGManagedResource.h
--- a/Engine/Plugins/PCG/Source/PCG/Public/PCGManagedResource.h
+++ b/Engine/Plugins/PCG/Source/PCG/Public/PCGManagedResource.h
@@ -90,7 +90,6 @@ public:
 	//~End UObject interface
 
 	//~Begin UPCGManagedResource interface
-	virtual void PostApplyToComponent() override;
 	virtual bool Release(bool bHardRelease, TSet<TSoftObjectPtr<AActor>>& OutActorsToDelete) override;
 	virtual bool ReleaseIfUnused(TSet<TSoftObjectPtr<AActor>>& OutActorsToDelete) override;
 	virtual bool MoveResourceToNewActor(AActor* NewActor) override;
