[PCG] Fix Level Instance related issues (still more work needed)
- Fix Level Instance PCGPartitionActors being parented to main world PCGWorldActor in outliner
- Prevent some editing in specific cases where PCG Component is inside a non-edit Level Instance

#rnx

diff --git a/Engine/Plugins/PCG/Source/PCG/Private/Grid/PCGPartitionActor.cpp b/Engine/Plugins/PCG/Source/PCG/Private/Grid/PCGPartitionActor.cpp
--- a/Engine/Plugins/PCG/Source/PCG/Private/Grid/PCGPartitionActor.cpp
+++ b/Engine/Plugins/PCG/Source/PCG/Private/Grid/PCGPartitionActor.cpp
@@ -14,6 +14,7 @@
 #include "UObject/FortniteMainBranchObjectVersion.h"
 
 #if WITH_EDITOR
+#include "Engine/Level.h"
 #include "Grid/PCGPartitionActorDesc.h"
 #endif
 
@@ -609,9 +610,12 @@ void APCGPartitionActor::RemoveLocalComponent(UPCGComponent* LocalComponent)
 #if WITH_EDITOR
 AActor* APCGPartitionActor::GetSceneOutlinerParent() const
 {
-	if (APCGWorldActor* PCGActor = PCGHelpers::GetPCGWorldActor(GetWorld()))
+	APCGWorldActor* PCGWorldActor = PCGHelpers::FindPCGWorldActor(GetWorld());
+	ULevel* Level = GetLevel();
+	
+	if (Level && Level->IsPersistentLevel() && PCGWorldActor)
 	{
-		return PCGActor;
+		return PCGWorldActor;
 	}
 	else
 	{
diff --git a/Engine/Plugins/PCG/Source/PCG/Private/Helpers/PCGHelpers.cpp b/Engine/Plugins/PCG/Source/PCG/Private/Helpers/PCGHelpers.cpp
--- a/Engine/Plugins/PCG/Source/PCG/Private/Helpers/PCGHelpers.cpp
+++ b/Engine/Plugins/PCG/Source/PCG/Private/Helpers/PCGHelpers.cpp
@@ -317,6 +317,11 @@ namespace PCGHelpers
 		return (InWorld && InWorld->GetSubsystem<UPCGSubsystem>()) ? InWorld->GetSubsystem<UPCGSubsystem>()->GetPCGWorldActor() : nullptr;
 	}
 
+	APCGWorldActor* FindPCGWorldActor(UWorld* InWorld)
+	{
+		return (InWorld && InWorld->GetSubsystem<UPCGSubsystem>()) ? InWorld->GetSubsystem<UPCGSubsystem>()->FindPCGWorldActor() : nullptr;
+	}
+
 	// TODO: Temporary validation during transition of allowing spaces in tags/attributes. Deprecate in 5.6.
 	TArray<FString> GetStringArrayFromCommaSeparatedString(const FString& InCommaSeparatedString, const FPCGContext* InOptionalContext)
 	{
diff --git a/Engine/Plugins/PCG/Source/PCG/Private/PCGComponent.cpp b/Engine/Plugins/PCG/Source/PCG/Private/PCGComponent.cpp
--- a/Engine/Plugins/PCG/Source/PCG/Private/PCGComponent.cpp
+++ b/Engine/Plugins/PCG/Source/PCG/Private/PCGComponent.cpp
@@ -1861,7 +1861,7 @@ void UPCGComponent::SetupCallbacksOnCreation()
 bool UPCGComponent::CanEditChange(const FProperty* InProperty) const
 {
 	// Can't change anything if the component is local
-	return !bIsComponentLocal;
+	return !bIsComponentLocal && Super::CanEditChange(InProperty);
 }
 #endif
 
diff --git a/Engine/Plugins/PCG/Source/PCG/Private/PCGGraph.cpp b/Engine/Plugins/PCG/Source/PCG/Private/PCGGraph.cpp
--- a/Engine/Plugins/PCG/Source/PCG/Private/PCGGraph.cpp
+++ b/Engine/Plugins/PCG/Source/PCG/Private/PCGGraph.cpp
@@ -24,6 +24,7 @@
 #include "Editor.h"
 #include "Dialogs/Dialogs.h"
 #include "EdGraph/EdGraphPin.h"
+#include "GameFramework/Actor.h"
 #else
 #include "UObject/Package.h"
 #endif
@@ -2039,10 +2040,18 @@ void UPCGGraphInstance::OnGraphChanged(UPCGGraphInterface* InGraph, EPCGChangeTy
 
 bool UPCGGraphInstance::CanEditChange(const FProperty* InProperty) const
 {
+	UPCGComponent* Component = Cast<UPCGComponent>(GetOuter());
+	AActor* ComponentOwner = Component ? Component->GetOwner() : nullptr;
+		
+	if(ComponentOwner && ComponentOwner->IsInLevelInstance() && !ComponentOwner->IsInEditLevelInstance())
+	{
+		return false;
+	}
+
 	// Graph can only be changed if it is not in a local PCGComponent
 	if (InProperty && InProperty->GetFName() == GET_MEMBER_NAME_CHECKED(UPCGGraphInstance, Graph))
 	{
-		if (UPCGComponent* Component = Cast<UPCGComponent>(GetOuter()))
+		if (Component)
 		{
 			return !Component->IsLocalComponent();
 		}
diff --git a/Engine/Plugins/PCG/Source/PCG/Public/Helpers/PCGHelpers.h b/Engine/Plugins/PCG/Source/PCG/Public/Helpers/PCGHelpers.h
--- a/Engine/Plugins/PCG/Source/PCG/Public/Helpers/PCGHelpers.h
+++ b/Engine/Plugins/PCG/Source/PCG/Public/Helpers/PCGHelpers.h
@@ -54,6 +54,7 @@ namespace PCGHelpers
 	PCG_API bool IsRuntimeOrPIE();
 
 	PCG_API APCGWorldActor* GetPCGWorldActor(UWorld* InWorld);
+	PCG_API APCGWorldActor* FindPCGWorldActor(UWorld* InWorld);
 
 	UE_DEPRECATED(5.5, "This function was incorrectly parsing by whitespace. Please use 'GetStringArrayFromCommaSeparatedList' instead.")
 	PCG_API TArray<FString> GetStringArrayFromCommaSeparatedString(const FString& InCommaSeparatedString, const FPCGContext* InOptionalContext = nullptr);
diff --git a/Engine/Plugins/PCG/Source/PCGEditor/Private/Details/PCGComponentDetails.cpp b/Engine/Plugins/PCG/Source/PCGEditor/Private/Details/PCGComponentDetails.cpp
--- a/Engine/Plugins/PCG/Source/PCGEditor/Private/Details/PCGComponentDetails.cpp
+++ b/Engine/Plugins/PCG/Source/PCGEditor/Private/Details/PCGComponentDetails.cpp
@@ -61,6 +61,19 @@ void FPCGComponentDetails::CustomizeDetails(IDetailLayoutBuilder& DetailBuilder)
 		}
 	}
 
+	for (TWeakObjectPtr<UPCGComponent>& ComponentPtr : SelectedComponents)
+	{
+		if (UPCGComponent* Component = ComponentPtr.Get())
+		{
+			AActor* ComponentOwner = Component->GetOwner();
+			if (!ComponentOwner || (ComponentOwner->IsInLevelInstance() && !ComponentOwner->IsInEditLevelInstance()))
+			{
+				// Do not customize for non editing Level Instance
+				return;
+			}
+		}
+	}
+
 	FDetailWidgetRow& NewRow = PCGCategory.AddCustomRow(FText::GetEmpty());
 
 	NewRow.ValueContent()
