[PCG] Add backend for entity conversion within a PCG Graph

+ Node settings to convert to a different settings class
+ Added support for `SoftClassPath` in `PCGGraphParameterExtension` library

diff --git a/Engine/Plugins/PCG/Source/PCG/Private/Helpers/PCGConversion.cpp b/Engine/Plugins/PCG/Source/PCG/Private/Helpers/PCGConversion.cpp
new file mode 100644
--- /dev/null
+++ b/Engine/Plugins/PCG/Source/PCG/Private/Helpers/PCGConversion.cpp
@@ -0,0 +1,168 @@
+// Copyright Epic Games, Inc. All Rights Reserved.
+
+#include "Helpers/PCGConversion.h"
+
+#include "PCGEdge.h"
+#include "PCGGraph.h"
+#include "PCGNode.h"
+#include "PCGSettings.h"
+
+namespace PCGConversion
+{
+#if WITH_EDITOR
+	namespace Helpers
+	{
+		FText GetDefaultNodeTitle(TSubclassOf<UPCGSettings> Class)
+		{
+			if (IsValid(Class))
+			{
+				return CastChecked<UPCGSettings>(Class->ClassDefaultObject)->GetDefaultNodeTitle();
+			}
+
+			return {};
+		}
+	}
+#endif // WITH_EDITOR
+
+	namespace Node::Helpers
+	{
+		bool RewireEdges(UPCGPin* OldPin, UPCGPin* NewPin)
+		{
+			if (!OldPin || !NewPin || OldPin == NewPin)
+			{
+				return false;
+			}
+
+			while (!OldPin->Edges.IsEmpty())
+			{
+				UPCGEdge* Edge = OldPin->Edges[0];
+				UPCGPin* OtherPin = Edge->InputPin == OldPin ? Edge->OutputPin : Edge->InputPin;
+				check(OtherPin);
+				OtherPin->BreakEdgeTo(OldPin);
+				OtherPin->AddEdgeTo(NewPin);
+			}
+
+			return true;
+		}
+	}
+}
+
+FPCGConverterBase::FPCGConverterBase(UPCGGraph* InOutGraph)
+	: SourceGraph(InOutGraph)
+{
+	if (SourceGraph)
+	{
+		// Mark the graph as initialized.
+		CurrentStatus |= EPCGConversionStatus::InitializedGraph;
+
+#if WITH_EDITOR
+		SourceGraph->DisableNotificationsForEditor();
+		SourceGraph->Modify(/*bAlwaysMarkDirty=*/false);
+#endif // WITH_EDITOR
+	}
+}
+
+FPCGConverterBase::~FPCGConverterBase()
+{
+#if WITH_EDITOR
+	if (bGraphIsDirty)
+	{
+		SourceGraph->MarkPackageDirty();
+	}
+	SourceGraph->EnableNotificationsForEditor();
+#endif // WITH_EDITOR
+}
+
+void FPCGConverterBase::PrepareData()
+{
+	bool bMarkGraphAsDirty = false;
+	if (ExecutePrepareData(bMarkGraphAsDirty))
+	{
+		bGraphIsDirty |= bMarkGraphAsDirty;
+		// Mark that the data was prepared.
+		CurrentStatus |= EPCGConversionStatus::DataPrepared;
+	}
+}
+
+void FPCGConverterBase::ApplyStructural()
+{
+	bool bMarkGraphAsDirty = false;
+	if (ExecuteApplyStructural(bMarkGraphAsDirty))
+	{
+		bGraphIsDirty |= bMarkGraphAsDirty;
+		CurrentStatus |= EPCGConversionStatus::StructuralChangesApplied;
+	}
+}
+
+FPCGSingleNodeConverter::FPCGSingleNodeConverter(UPCGNode* InNode, TSubclassOf<UPCGSettings> InTargetSettingsClass)
+	: FPCGConverterBase(InNode ? InNode->GetGraph() : nullptr)
+	, SourceNode(InNode)
+	, TargetSettingsClass(InTargetSettingsClass)
+{
+	if (SourceNode && SourceNode->GetSettings())
+	{
+		SetSourceInitialized();
+	}
+}
+
+bool FPCGSingleNodeConverter::ExecutePrepareData(bool& bOutMarkGraphAsDirty)
+{
+	if (ensure(IsValid()))
+	{
+		GeneratedNode = SourceGraph->AddNodeOfType(TargetSettingsClass, GeneratedSettings);
+		bOutMarkGraphAsDirty = true;
+
+		if (GeneratedNode && GeneratedSettings)
+		{
+			if (GeneratedNode->HasAuthoredTitle())
+			{
+				GeneratedNode->NodeTitle = SourceNode->NodeTitle;
+			}
+#if WITH_EDITOR
+			SourceNode->TransferEditorProperties(GeneratedNode);
+#endif // WITH_EDITOR
+			GeneratedSettings->bEnabled = SourceNode->GetSettings()->bEnabled;
+
+			return true;
+		}
+	}
+
+	return false;
+}
+
+bool FPCGSingleNodeConverter::ExecuteApplyStructural(bool& bOutMarkGraphAsDirty)
+{
+	if (!ensure(IsValid()))
+	{
+		return false;
+	}
+
+	bool bResult = true;
+	// TODO: Could have better ways to match pins: based on the same name, etc.
+	// For now, the assumption is that the pins should just match ordered 1:1. Do the best to fit them, otherwise discard the rest.
+	const TArray<TObjectPtr<UPCGPin>>& SourceInputPins = SourceNode->GetInputPins();
+	const TArray<TObjectPtr<UPCGPin>>& GeneratedInputPins = GeneratedNode->GetInputPins();
+	for (int32 InputIndex = 0; InputIndex < GeneratedInputPins.Num() && InputIndex < SourceInputPins.Num(); ++InputIndex)
+	{
+		bResult &= PCGConversion::Node::Helpers::RewireEdges(SourceInputPins[InputIndex], GeneratedInputPins[InputIndex]);
+	}
+
+	const TArray<TObjectPtr<UPCGPin>>& SourceOutputPins = SourceNode->GetOutputPins();
+	const TArray<TObjectPtr<UPCGPin>>& GeneratedOutputPins = GeneratedNode->GetOutputPins();
+	for (int32 OutputIndex = 0; OutputIndex < GeneratedOutputPins.Num() && OutputIndex < SourceOutputPins.Num(); ++OutputIndex)
+	{
+		bResult &= PCGConversion::Node::Helpers::RewireEdges(SourceOutputPins[OutputIndex], GeneratedOutputPins[OutputIndex]);
+	}
+
+	SourceGraph->RemoveNode(SourceNode);
+
+	return bResult;
+}
+
+void FPCGSingleNodeConverter::Finalize()
+{
+	if (ensure(IsValid()))
+	{
+		GeneratedNode->UpdateAfterSettingsChangeDuringCreation();
+	}
+}
\ No newline at end of file
diff --git a/Engine/Plugins/PCG/Source/PCG/Private/Helpers/PCGConversion.h b/Engine/Plugins/PCG/Source/PCG/Private/Helpers/PCGConversion.h
new file mode 100644
--- /dev/null
+++ b/Engine/Plugins/PCG/Source/PCG/Private/Helpers/PCGConversion.h
@@ -0,0 +1,133 @@
+// Copyright Epic Games, Inc. All Rights Reserved.
+
+#pragma once
+
+#include "Templates/SubclassOf.h"
+
+#include "PCGConversion.generated.h"
+
+// TODO:
+// - Structural helpers may transition into a Graph Builder system.
+// - Factor out the N-NodeConverter into a more generic version that accepts an array of pin to pin mappings.
+//   Could additionally have an "alias" class for single node conversion on top of that.
+
+class UPCGPin;
+class UPCGSettings;
+class UPCGGraph;
+class UPCGNode;
+
+namespace PCGConversion
+{
+#if WITH_EDITOR
+	namespace Helpers
+	{
+		/** Helper to get the default node title from the settings class. */
+		FText GetDefaultNodeTitle(TSubclassOf<UPCGSettings> Class);
+	}
+#endif // WITH_EDITOR
+}
+
+UENUM(meta = (Bitflags))
+enum class EPCGConversionStatus : uint8
+{
+	Uninitialized = 0,
+	InitializedGraph = 1 << 0,
+	InitializedSource = 1 << 1,
+	DataPrepared = 1 << 2,
+	StructuralChangesApplied = 1 << 3,
+	Complete = InitializedSource | InitializedGraph | DataPrepared | StructuralChangesApplied
+};
+
+ENUM_CLASS_FLAGS(EPCGConversionStatus);
+
+class FPCGConverterBase
+{
+public:
+	explicit FPCGConverterBase(UPCGGraph* InOutGraph);
+	virtual ~FPCGConverterBase();
+
+	// Intentionally bar copy/assignment.
+	FPCGConverterBase(const FPCGConverterBase&) = delete;
+	FPCGConverterBase(FPCGConverterBase&&) = delete;
+	FPCGConverterBase& operator=(const FPCGConverterBase&) = delete;
+	FPCGConverterBase& operator=(FPCGConverterBase&&) = delete;
+
+	/** Create whatever new data must exist in the conversion, such as a new node. */
+	void PrepareData();
+	/** Apply structural changes to the graph. */
+	void ApplyStructural();
+
+	/** Overridable validation query. */
+	virtual bool IsValid() = 0;
+
+	/** Get the current status of the conversion. */
+	EPCGConversionStatus GetCurrentStatus() const { return CurrentStatus; }
+
+	/** Get the conversion's corresponding graph. */
+	UPCGGraph* GetGraph() const { return SourceGraph; }
+
+	/** Should be called once the conversion's source has been initialized. */
+	void SetSourceInitialized() { CurrentStatus |= EPCGConversionStatus::InitializedSource; }
+
+	bool IsSourceInitialized() const { return !!(CurrentStatus & EPCGConversionStatus::InitializedSource); }
+	bool IsGraphInitialized() const { return !!(CurrentStatus & EPCGConversionStatus::InitializedGraph); }
+	bool IsDataPrepared() const { return !!(CurrentStatus & EPCGConversionStatus::DataPrepared); }
+	bool AreStructuralChangesApplied() const { return !!(CurrentStatus & EPCGConversionStatus::StructuralChangesApplied); }
+
+	/** The entire conversion process is complete. */
+	bool IsComplete() const { return CurrentStatus == EPCGConversionStatus::Complete; }
+
+protected:
+	/** Overridable and intended to do the work needed to prepare the target conversion data. */
+	virtual bool ExecutePrepareData(bool& bOutMarkGraphAsDirty) { return true; }
+	/** Overridable and intended to do the work needed to apply structural changes during the conversion. */
+	virtual bool ExecuteApplyStructural(bool& bOutMarkGraphAsDirty) { return true; }
+
+	/** The source graph owning the conversion. This should not change during the process. */
+	UPCGGraph* const SourceGraph;
+
+	/** The graph will need to be dirtied upon completion. */
+	bool bGraphIsDirty = false;
+
+private:
+	/** The completed steps of the conversion in bitflag form. */
+	EPCGConversionStatus CurrentStatus = EPCGConversionStatus::Uninitialized;
+};
+
+/** Converts a single source node into a single compatible target node and rewires edges. */
+class FPCGSingleNodeConverter final : public FPCGConverterBase // TODO: final can be removed when the class becomes public
+{
+public:
+	explicit FPCGSingleNodeConverter(UPCGNode* InNode, TSubclassOf<UPCGSettings> InTargetSettingsClass);
+
+	//~Begin FPCGConverterBase interface
+	/** The node conversion still has valid references. */
+	virtual bool IsValid() override { return SourceGraph && (GeneratedNode || !IsDataPrepared()); }
+	//~Begin FPCGConverterBase interface
+
+	/** Finalizes the node conversion. */
+	void Finalize();
+
+	UPCGNode* GetSourceNode() const { return SourceNode; }
+	UPCGNode* GetGeneratedNode() const { return GeneratedNode; }
+	UPCGSettings* GetGeneratedSettings() const { return GeneratedSettings; }
+
+protected:
+	//~Begin FPCGConverterBase interface
+	/** Create the new node and settings. */
+	virtual bool ExecutePrepareData(bool& bOutMarkGraphAsDirty) override;
+	/** Rewire the input and output pins. */
+	virtual bool ExecuteApplyStructural(bool& bOutMarkGraphAsDirty) override;
+	//~Begin FPCGConverterBase interface
+
+	/** The node being converted. */
+	UPCGNode* const SourceNode;
+	/** The class of the target node. */
+	const TSubclassOf<UPCGSettings> TargetSettingsClass;
+
+private:
+	/** The node created during the conversion. */
+	UPCGNode* GeneratedNode = nullptr;
+	/** The settings created during the conversion. */
+	UPCGSettings* GeneratedSettings = nullptr;
+};
\ No newline at end of file
diff --git a/Engine/Plugins/PCG/Source/PCG/Private/Helpers/PCGGraphParameterExtension.cpp b/Engine/Plugins/PCG/Source/PCG/Private/Helpers/PCGGraphParameterExtension.cpp
--- a/Engine/Plugins/PCG/Source/PCG/Private/Helpers/PCGGraphParameterExtension.cpp
+++ b/Engine/Plugins/PCG/Source/PCG/Private/Helpers/PCGGraphParameterExtension.cpp
@@ -43,6 +43,10 @@ template<> TValueOrError<FSoftObjectPath*, EPropertyBagResult> FPCGGraphParamete
 {
 	return PropertyBag.GetValueStruct<FSoftObjectPath>(PropertyName);
 }
+template<> TValueOrError<FSoftClassPath*, EPropertyBagResult> FPCGGraphParameterExtension::GetGraphParameter(const FInstancedPropertyBag& PropertyBag, const FName PropertyName)
+{
+	return PropertyBag.GetValueStruct<FSoftClassPath>(PropertyName);
+}
 template<> TValueOrError<TSoftObjectPtr<UObject>, EPropertyBagResult> FPCGGraphParameterExtension::GetGraphParameter(const FInstancedPropertyBag& PropertyBag, const FName PropertyName)
 {
 	TValueOrError<FSoftObjectPath, EPropertyBagResult> InitialResult = PropertyBag.GetValueSoftPath(PropertyName);
@@ -147,6 +151,10 @@ template<> EPropertyBagResult FPCGGraphParameterExtension::SetGraphParameter(FIn
 {
 	return PropertyBag.SetValueStruct(PropertyName, Value);
 }
+template<> EPropertyBagResult FPCGGraphParameterExtension::SetGraphParameter(FInstancedPropertyBag& PropertyBag, const FName PropertyName, const FSoftClassPath& Value)
+{
+	return PropertyBag.SetValueStruct(PropertyName, Value);
+}
 template<> EPropertyBagResult FPCGGraphParameterExtension::SetGraphParameter(FInstancedPropertyBag& PropertyBag, const FName PropertyName, const TSoftObjectPtr<UObject>& Value)
 {
 	return PropertyBag.SetValueSoftPath(PropertyName, Value.ToSoftObjectPath());
diff --git a/Engine/Plugins/PCG/Source/PCG/Private/Helpers/PCGPropertyHelpers.cpp b/Engine/Plugins/PCG/Source/PCG/Private/Helpers/PCGPropertyHelpers.cpp
--- a/Engine/Plugins/PCG/Source/PCG/Private/Helpers/PCGPropertyHelpers.cpp
+++ b/Engine/Plugins/PCG/Source/PCG/Private/Helpers/PCGPropertyHelpers.cpp
@@ -238,6 +238,69 @@ EPCGMetadataTypes PCGPropertyHelpers::GetMetadataTypeFromProperty(const FPropert
 	return PropertyAccessor.IsValid() ? EPCGMetadataTypes(PropertyAccessor->GetUnderlyingType()) : EPCGMetadataTypes::Unknown;
 }
 
+FPropertyBagPropertyDesc PCGPropertyHelpers::GetPropertyBagDescWithMetadataType(const FName InPropertyName, const EPCGMetadataTypes Type)
+{
+	EPropertyBagPropertyType ValueType = EPropertyBagPropertyType::Struct;
+	const UObject* ValueTypeObject = nullptr;
+
+	switch (Type)
+	{
+		// Simple Types
+		case EPCGMetadataTypes::Float:
+			ValueType = EPropertyBagPropertyType::Float;
+			break;
+		case EPCGMetadataTypes::Double:
+			ValueType = EPropertyBagPropertyType::Double;
+			break;
+		case EPCGMetadataTypes::Integer32:
+			ValueType = EPropertyBagPropertyType::Int32;
+			break;
+		case EPCGMetadataTypes::Integer64:
+			ValueType = EPropertyBagPropertyType::Int64;
+			break;
+		case EPCGMetadataTypes::String:
+			ValueType = EPropertyBagPropertyType::String;
+			break;
+		case EPCGMetadataTypes::Boolean:
+			ValueType = EPropertyBagPropertyType::Bool;
+			break;
+		case EPCGMetadataTypes::Name:
+			ValueType = EPropertyBagPropertyType::Name;
+			break;
+
+		// Value Type Objects - PropertyType is Struct by default.
+		case EPCGMetadataTypes::SoftObjectPath:
+			ValueTypeObject = TBaseStructure<FSoftObjectPath>::Get();
+			break;
+		case EPCGMetadataTypes::SoftClassPath:
+			ValueTypeObject = TBaseStructure<FSoftClassPath>::Get();
+			break;
+		case EPCGMetadataTypes::Vector:
+			ValueTypeObject = TBaseStructure<FVector>::Get();
+			break;
+		case EPCGMetadataTypes::Vector2:
+			ValueTypeObject = TBaseStructure<FVector2d>::Get();
+			break;
+		case EPCGMetadataTypes::Vector4:
+			ValueTypeObject = TBaseStructure<FVector4>::Get();
+			break;
+		case EPCGMetadataTypes::Transform:
+			ValueTypeObject = TBaseStructure<FTransform>::Get();
+			break;
+		case EPCGMetadataTypes::Quaternion:
+			ValueTypeObject = TBaseStructure<FQuat>::Get();
+			break;
+		case EPCGMetadataTypes::Rotator:
+			ValueTypeObject = TBaseStructure<FRotator>::Get();
+			break;
+
+		default:
+			return {};
+	}
+
+	return FPropertyBagPropertyDesc(InPropertyName, ValueType, ValueTypeObject);
+}
+
 UPCGParamData* PCGPropertyHelpers::ExtractPropertyAsAttributeSet(const PCGPropertyHelpers::FExtractorParameters& Parameters, FPCGContext* OptionalContext, TSet<FSoftObjectPath>* OptionalObjectTraversed, bool bQuiet)
 {
 	check(Parameters.Container && Parameters.Class);
diff --git a/Engine/Plugins/PCG/Source/PCG/Public/Helpers/PCGGraphParameterExtension.h b/Engine/Plugins/PCG/Source/PCG/Public/Helpers/PCGGraphParameterExtension.h
--- a/Engine/Plugins/PCG/Source/PCG/Public/Helpers/PCGGraphParameterExtension.h
+++ b/Engine/Plugins/PCG/Source/PCG/Public/Helpers/PCGGraphParameterExtension.h
@@ -31,6 +31,7 @@ template<> PCG_API TValueOrError<int64, EPropertyBagResult> FPCGGraphParameterEx
 template<> PCG_API TValueOrError<FName, EPropertyBagResult> FPCGGraphParameterExtension::GetGraphParameter(const FInstancedPropertyBag& PropertyBag, const FName PropertyName);
 template<> PCG_API TValueOrError<FString, EPropertyBagResult> FPCGGraphParameterExtension::GetGraphParameter(const FInstancedPropertyBag& PropertyBag, const FName PropertyName);
 template<> PCG_API TValueOrError<FSoftObjectPath*, EPropertyBagResult> FPCGGraphParameterExtension::GetGraphParameter(const FInstancedPropertyBag& PropertyBag, const FName PropertyName);
+template<> PCG_API TValueOrError<FSoftClassPath*, EPropertyBagResult> FPCGGraphParameterExtension::GetGraphParameter(const FInstancedPropertyBag& PropertyBag, const FName PropertyName);
 template<> PCG_API TValueOrError<TSoftObjectPtr<UObject>, EPropertyBagResult> FPCGGraphParameterExtension::GetGraphParameter(const FInstancedPropertyBag& PropertyBag, const FName PropertyName);
 template<> PCG_API TValueOrError<TSoftClassPtr<UObject>, EPropertyBagResult> FPCGGraphParameterExtension::GetGraphParameter(const FInstancedPropertyBag& PropertyBag, const FName PropertyName);
 template<> PCG_API TValueOrError<UObject*, EPropertyBagResult> FPCGGraphParameterExtension::GetGraphParameter(const FInstancedPropertyBag& PropertyBag, const FName PropertyName);
@@ -51,6 +52,7 @@ template<> PCG_API EPropertyBagResult FPCGGraphParameterExtension::SetGraphParam
 template<> PCG_API EPropertyBagResult FPCGGraphParameterExtension::SetGraphParameter(FInstancedPropertyBag& PropertyBag, const FName PropertyName, const FName& Value);
 template<> PCG_API EPropertyBagResult FPCGGraphParameterExtension::SetGraphParameter(FInstancedPropertyBag& PropertyBag, const FName PropertyName, const FString& Value);
 template<> PCG_API EPropertyBagResult FPCGGraphParameterExtension::SetGraphParameter(FInstancedPropertyBag& PropertyBag, const FName PropertyName, const FSoftObjectPath& Value);
+template<> PCG_API EPropertyBagResult FPCGGraphParameterExtension::SetGraphParameter(FInstancedPropertyBag& PropertyBag, const FName PropertyName, const FSoftClassPath& Value);
 template<> PCG_API EPropertyBagResult FPCGGraphParameterExtension::SetGraphParameter(FInstancedPropertyBag& PropertyBag, const FName PropertyName, const TSoftObjectPtr<UObject>& Value);
 template<> PCG_API EPropertyBagResult FPCGGraphParameterExtension::SetGraphParameter(FInstancedPropertyBag& PropertyBag, const FName PropertyName, const TSoftClassPtr<UObject>& Value);
 template<> PCG_API EPropertyBagResult FPCGGraphParameterExtension::SetGraphParameter(FInstancedPropertyBag& PropertyBag, const FName PropertyName, UObject* const& Value); // UObject* const& to respect template specialization
diff --git a/Engine/Plugins/PCG/Source/PCG/Public/Helpers/PCGPropertyHelpers.h b/Engine/Plugins/PCG/Source/PCG/Public/Helpers/PCGPropertyHelpers.h
--- a/Engine/Plugins/PCG/Source/PCG/Public/Helpers/PCGPropertyHelpers.h
+++ b/Engine/Plugins/PCG/Source/PCG/Public/Helpers/PCGPropertyHelpers.h
@@ -8,6 +8,7 @@
 #include "Metadata/Accessors/PCGAttributeAccessorHelpers.h"
 #include "Metadata/Accessors/PCGAttributeAccessorKeys.h"
 
+#include "StructUtils/PropertyBag.h"
 #include "UObject/UnrealType.h" // IWYU pragma: keep
 
 enum class EPCGMetadataTypes : uint8;
@@ -45,6 +46,14 @@ namespace PCGPropertyHelpers
 	*/
 	PCG_API EPCGMetadataTypes GetMetadataTypeFromProperty(const FProperty* InProperty);
 
+	/**
+	 * Conversion between StructUtils' PropertyBag types and PCG type.
+	 * @param InPropertyName - The name of the property.
+	 * @param Type - The PCG metadata type.
+	 * @returns The property bag property type.
+	 */
+	PCG_API FPropertyBagPropertyDesc GetPropertyBagDescWithMetadataType(FName InPropertyName, EPCGMetadataTypes Type);
+
 	struct FExtractorParameters
 	{
 		// Disable deprecation warning on the rule of 5 because of the PropertySelector member.
diff --git a/Engine/Plugins/PCG/Source/PCG/Public/PCGSettings.h b/Engine/Plugins/PCG/Source/PCG/Public/PCGSettings.h
--- a/Engine/Plugins/PCG/Source/PCG/Public/PCGSettings.h
+++ b/Engine/Plugins/PCG/Source/PCG/Public/PCGSettings.h
@@ -300,6 +300,9 @@ public:
 	/** If there are preconfigured info, decide if they are grouped in the palette in a folder with the node name, or if they are not grouped. */
 	virtual bool GroupPreconfiguredSettings() const { return true; }
 
+	/** The predefined parameters defining the conversion. */
+	virtual TArray<FPCGPreconfiguredInfo> GetConversionInfo() const { return {}; }
+
 	/** Perform post-operations when an editor node is copied */
 	virtual void PostPaste();
 #endif // WITH_EDITOR
@@ -351,6 +354,9 @@ public:
 	*/
 	virtual EPCGDataType GetCurrentPinTypes(const UPCGPin* InPin) const;
 
+	/** Convert this owning node into the corresponding target. Returns true upon success. */
+	virtual bool ConvertNode(const FPCGPreconfiguredInfo& ConversionInfo) { return false; }
+
 	UPROPERTY(EditAnywhere, BlueprintReadWrite, Category = Settings, meta=(EditCondition="UseSeed()", EditConditionHides, PCG_Overridable))
 	int Seed = 0xC35A9631; // Default seed is a random prime number, but will be overriden for new settings based on the class type name hash, making each settings class have a different default seed.
 
diff --git a/Engine/Plugins/PCG/Source/PCG/Public/Utils/PCGPreconfiguration.h b/Engine/Plugins/PCG/Source/PCG/Public/Utils/PCGPreconfiguration.h
--- a/Engine/Plugins/PCG/Source/PCG/Public/Utils/PCGPreconfiguration.h
+++ b/Engine/Plugins/PCG/Source/PCG/Public/Utils/PCGPreconfiguration.h
@@ -5,24 +5,24 @@
 #include "PCGPreconfiguration.generated.h"
 
 /**
-* Pre-configured settings info
-* Will be passed to the settings to pre-configure the settings on creation.
-* Example: Maths operations: Add, Mul, Div etc...
+* Pre-configured info
+* Utility struct for use with any type of PCG preconfiguration.
+* Example: A set of user actions available on a specific node or predefined settings configuration for a node element.
 */
 USTRUCT(BlueprintType)
-struct PCG_API FPCGPreConfiguredSettingsInfo
+struct PCG_API FPCGPreconfiguredInfo
 {
 	GENERATED_BODY()
 
-	FPCGPreConfiguredSettingsInfo() = default;
+	FPCGPreconfiguredInfo() = default;
 
-	explicit FPCGPreConfiguredSettingsInfo(int32 InIndex, FText InLabel = FText{})
+	explicit FPCGPreconfiguredInfo(int32 InIndex, FText InLabel = FText{})
 		: PreconfiguredIndex(InIndex)
 		, Label(std::move(InLabel))
 	{}
 
 #if WITH_EDITOR
-	FPCGPreConfiguredSettingsInfo(int32 InIndex, FText InLabel, FText InTooltip)
+	explicit FPCGPreconfiguredInfo(int32 InIndex, FText InLabel, FText InTooltip)
 		: PreconfiguredIndex(InIndex)
 		, Label(std::move(InLabel))
 		, Tooltip(std::move(InTooltip))
@@ -33,10 +33,12 @@ struct PCG_API FPCGPreConfiguredSettingsInfo
 	* Automatically fill all preconfigured settings depending on the provided enum.
 	* Can also specify explicitly values that should not be included, i.e. cases that may not be available in non-editor builds.
 	*/
-	template <typename EnumType>
-	static TArray<FPCGPreConfiguredSettingsInfo> PopulateFromEnum(const TSet<EnumType>& InValuesToSkip = {}, const FText& InOptionalPrefix = FText())
+	template <typename EnumType, typename SubclassType = FPCGPreconfiguredInfo>
+	static TArray<SubclassType> PopulateFromEnum(const TSet<EnumType>& InValuesToSkip = {}, const FText& InOptionalPrefix = FText())
 	{
-		TArray<FPCGPreConfiguredSettingsInfo> PreconfiguredInfo;
+		static_assert(std::is_base_of_v<FPCGPreconfiguredInfo, SubclassType>, "SubclassType must be a subclass of FPCGPreconfiguredInfo");
+
+		TArray<SubclassType> PreconfiguredInfo;
 
 		if (const UEnum* EnumPtr = StaticEnum<EnumType>())
 		{
@@ -71,4 +73,33 @@ struct PCG_API FPCGPreConfiguredSettingsInfo
 	UPROPERTY(EditAnywhere, Category = "")
 	FText Tooltip;
 #endif // WITH_EDITORONLY_DATA
+};
+
+/**
+* Pre-configured settings info
+* Will be passed to the settings to pre-configure the settings on creation. Also used for pre-configured node elements.
+* Example: Maths operations: Add, Mul, Div etc...
+*/
+USTRUCT(BlueprintType)
+struct PCG_API FPCGPreConfiguredSettingsInfo : public FPCGPreconfiguredInfo
+{
+	GENERATED_BODY()
+
+	FPCGPreConfiguredSettingsInfo() = default;
+
+	explicit FPCGPreConfiguredSettingsInfo(int32 InIndex, FText InLabel = FText{})
+		: FPCGPreconfiguredInfo(InIndex, InLabel)
+	{}
+
+#if WITH_EDITOR
+	FPCGPreConfiguredSettingsInfo(int32 InIndex, FText InLabel, FText InTooltip)
+		: FPCGPreconfiguredInfo(InIndex, InLabel, InTooltip)
+	{}
+#endif // WITH_EDITOR
+
+	template <typename EnumType>
+	static TArray<FPCGPreConfiguredSettingsInfo> PopulateFromEnum(const TSet<EnumType>& InValuesToSkip = {}, const FText& InOptionalPrefix = FText())
+	{
+		return FPCGPreconfiguredInfo::PopulateFromEnum<EnumType, FPCGPreConfiguredSettingsInfo>(InValuesToSkip, InOptionalPrefix);
+	}
 };
\ No newline at end of file
diff --git a/Engine/Plugins/PCG/Source/PCGEditor/Private/PCGEditorGraphNodeBase.cpp b/Engine/Plugins/PCG/Source/PCGEditor/Private/PCGEditorGraphNodeBase.cpp
--- a/Engine/Plugins/PCG/Source/PCGEditor/Private/PCGEditorGraphNodeBase.cpp
+++ b/Engine/Plugins/PCG/Source/PCGEditor/Private/PCGEditorGraphNodeBase.cpp
@@ -213,22 +213,53 @@ void UPCGEditorGraphNodeBase::GetNodeContextMenuActions(UToolMenu* Menu, class U
 		Section.AddMenuEntry(FPCGEditorCommands::Get().ConvertToStandaloneNodes);
 		Section.AddMenuEntry(FPCGEditorCommands::Get().RenameNode, LOCTEXT("RenameNode", "Rename"));
 
+		const UPCGSettings* Settings = PCGNode->GetSettings();
+
 		// Special nodes operations
-		if (PCGNode && PCGNode->GetSettings())
+		if (PCGNode && Settings)
 		{
-			if (PCGNode->GetSettings()->IsA<UPCGNamedRerouteDeclarationSettings>())
+			if (Settings->IsA<UPCGNamedRerouteDeclarationSettings>())
 			{
 				Section.AddMenuEntry(FPCGEditorCommands::Get().ConvertNamedRerouteToReroute);
 				Section.AddMenuEntry(FPCGEditorCommands::Get().SelectNamedRerouteUsages);
 			}
-			else if (PCGNode->GetSettings()->IsA<UPCGNamedRerouteUsageSettings>())
+			else if (Settings->IsA<UPCGNamedRerouteUsageSettings>())
 			{
 				Section.AddMenuEntry(FPCGEditorCommands::Get().SelectNamedRerouteDeclaration);
 			}
-			else if (PCGNode->GetSettings()->IsA<UPCGRerouteSettings>())
+			else if (Settings->IsA<UPCGRerouteSettings>())
 			{
 				Section.AddMenuEntry(FPCGEditorCommands::Get().ConvertRerouteToNamedReroute);
 			}
+
+			// Operation to convert a node in place into something else, such as a different node.
+			const TArray<FPCGPreconfiguredInfo> ConversionInfo = Settings->GetConversionInfo();
+			if (!ConversionInfo.IsEmpty())
+			{
+				Section.AddSubMenu(
+					TEXT("NodeConversion"),
+					LOCTEXT("NodeConversion", "Convert Node"),
+					LOCTEXT("NodeConversionTooltip", "Convert a single node into a different node, or otherwise compatible output."),
+					FNewToolMenuDelegate::CreateLambda([this, ConversionInfo = std::move(ConversionInfo)](UToolMenu* AlignmentMenu)
+					{
+						FToolMenuSection& SubSection = AlignmentMenu->AddSection("EdGraphSchemaConversion", LOCTEXT("ConvertHeader", "Convert To"));
+
+						for (const FPCGPreconfiguredInfo& Conversion : ConversionInfo)
+						{
+							// TODO: UX feedback (read-only or greyed out font, etc) if the conversion can't happen.
+							SubSection.AddMenuEntry(
+									FName(Conversion.Label.ToString()),
+									Conversion.Label,
+									Conversion.Tooltip,
+									FSlateIcon(),
+									// Capture Conversion by value for delayed use.
+									FUIAction(FExecuteAction::CreateLambda([this, Conversion]
+									{
+										const_cast<UPCGEditorGraphNodeBase*>(this)->OnConvertNode(Conversion);
+									})));
+						}
+					}));
+			}
 		}
 	}
 
@@ -1144,6 +1175,29 @@ void UPCGEditorGraphNodeBase::OnUserRemoveDynamicInputPin(UEdGraphPin* InRemoved
 	}
 }
 
+void UPCGEditorGraphNodeBase::OnConvertNode(const FPCGPreconfiguredInfo& ConversionInfo)
+{
+	check(PCGNode && PCGNode->GetSettings());
+
+	if (UPCGEditorGraph* EditorGraph = Cast<UPCGEditorGraph>(GetGraph()))
+	{
+		if (UPCGGraph* PCGGraph = EditorGraph->GetPCGGraph())
+		{
+			FText ConversionTransactionName = FText::Format(INVTEXT("{0}: {1}"), LOCTEXT("PCGEditorConvertNode", "Convert Node"), PCGNode->GetDefaultTitle());
+			FScopedTransaction Transaction(*FPCGEditorCommon::ContextIdentifier, std::move(ConversionTransactionName), PCGGraph);
+
+			if (!PCGNode->GetSettings()->ConvertNode(ConversionInfo))
+			{
+				// TODO: It would be useful to have an error/feedback process for it can not be converted.
+				Transaction.Cancel();
+				return;
+			}
+
+			EditorGraph->ReconstructGraph();
+		}
+	}
+}
+
 void UPCGEditorGraphNodeBase::UpdateCommentBubblePinned()
 {
 	if (PCGNode)
diff --git a/Engine/Plugins/PCG/Source/PCGEditor/Private/PCGEditorGraphNodeBase.h b/Engine/Plugins/PCG/Source/PCGEditor/Private/PCGEditorGraphNodeBase.h
--- a/Engine/Plugins/PCG/Source/PCGEditor/Private/PCGEditorGraphNodeBase.h
+++ b/Engine/Plugins/PCG/Source/PCGEditor/Private/PCGEditorGraphNodeBase.h
@@ -6,9 +6,11 @@
 #include "EdGraph/EdGraphPin.h"
 
 #include "PCGCommon.h"
+#include "Utils/PCGPreconfiguration.h"
 
 #include "PCGEditorGraphNodeBase.generated.h"
 
+struct FPCGPreconfiguredInfo;
 enum class EPCGChangeType : uint8;
 
 class UPCGComponent;
@@ -51,7 +53,9 @@ public:
 	void OnUserAddDynamicInputPin();
 	bool CanUserRemoveDynamicInputPin(UEdGraphPin* InPinToRemove);
 	void OnUserRemoveDynamicInputPin(UEdGraphPin* InRemovedPin);
-	
+
+	void OnConvertNode(const FPCGPreconfiguredInfo& ConversionInfo);
+
 	UPCGNode* GetPCGNode() { return PCGNode; }
 	const UPCGNode* GetPCGNode() const { return PCGNode; }
 	void PostCopy();
