[PCG] Node Conversion - Get Graph Parameter <-> Get Graph Parameter Generic

diff --git a/Engine/Plugins/PCG/Source/PCG/Private/Elements/PCGUserParameterGet.cpp b/Engine/Plugins/PCG/Source/PCG/Private/Elements/PCGUserParameterGet.cpp
--- a/Engine/Plugins/PCG/Source/PCG/Private/Elements/PCGUserParameterGet.cpp
+++ b/Engine/Plugins/PCG/Source/PCG/Private/Elements/PCGUserParameterGet.cpp
@@ -9,6 +9,7 @@
 #include "PCGPin.h"
 #include "PCGSubgraph.h"
 #include "Data/PCGUserParametersData.h"
+#include "Helpers/PCGConversion.h"
 #include "Helpers/PCGDynamicTrackingHelpers.h"
 #include "Helpers/PCGPropertyHelpers.h"
 #include "Metadata/PCGAttributePropertySelector.h"
@@ -18,79 +19,91 @@
 
 #define LOCTEXT_NAMESPACE "PCGUserParameterGetElement"
 
-namespace PCGUserParameterGetSettings
+namespace PCGUserParameterGet
 {
-	/**
-	* Utility function to get the first valid instanced property bag
-	* We define valid as if the ParameterOverrides from a GraphInstance and UserParameters from the graph owner of the node have the same property bag.
-	* By construction, it should be always the case, but we want to prevent cases where graph instances depend on other graph instances, that has changed their graph
-	* but didn't propagate the changes.
-	* If the property bags aren't the same, we are traversing the graph instance hierarchy to find the first graph/graph instance that matches.
-	*/
-	TArray<FConstStructView, TInlineAllocator<16>> GetValidLayouts(FPCGContext& InContext)
+	namespace Settings
 	{
-		TArray<FConstStructView, TInlineAllocator<16>> Layouts;
-		const UPCGGraphInterface* GraphInterface = nullptr;
-
-		// First we will read from the input if we find an input for this graph, set by the subgraph element, then use it.
-		TArray<FPCGTaggedData> UserParameterData = InContext.InputData.GetTaggedTypedInputs<UPCGUserParametersData>(PCGBaseSubgraphConstants::UserParameterTagData);
-		if (!UserParameterData.IsEmpty())
+		/**
+		* Utility function to get the first valid instanced property bag
+		* We define valid as if the ParameterOverrides from a GraphInstance and UserParameters from the graph owner of the node have the same property bag.
+		* By construction, it should be always the case, but we want to prevent cases where graph instances depend on other graph instances, that has changed their graph
+		* but didn't propagate the changes.
+		* If the property bags aren't the same, we are traversing the graph instance hierarchy to find the first graph/graph instance that matches.
+		*/
+		TArray<FConstStructView, TInlineAllocator<16>> GetValidLayouts(FPCGContext& InContext)
 		{
-#if WITH_EDITOR
-			// Safe guard to make sure we always have one and only one data of this type.
-			ensure(UserParameterData.Num() == 1);
-#endif // WITH_EDITOR
-			if (const UPCGUserParametersData* OverrideParametersData = CastChecked<UPCGUserParametersData>(UserParameterData[0].Data, ECastCheckedType::NullAllowed))
-			{
-				Layouts.Emplace(OverrideParametersData->UserParameters);
+			TArray<FConstStructView, TInlineAllocator<16>> Layouts;
+			const UPCGGraphInterface* GraphInterface = nullptr;
 
-				const UPCGUserParametersData* UpstreamData = OverrideParametersData->UpstreamData.Get();
-				while (UpstreamData)
+			// First we will read from the input if we find an input for this graph, set by the subgraph element, then use it.
+			TArray<FPCGTaggedData> UserParameterData = InContext.InputData.GetTaggedTypedInputs<UPCGUserParametersData>(PCGBaseSubgraphConstants::UserParameterTagData);
+			if (!UserParameterData.IsEmpty())
+			{
+	#if WITH_EDITOR
+				// Safe guard to make sure we always have one and only one data of this type.
+				ensure(UserParameterData.Num() == 1);
+	#endif // WITH_EDITOR
+				if (const UPCGUserParametersData* OverrideParametersData = CastChecked<UPCGUserParametersData>(UserParameterData[0].Data, ECastCheckedType::NullAllowed))
 				{
-					Layouts.Emplace(UpstreamData->UserParameters);
-					UpstreamData = UpstreamData->UpstreamData.Get();
+					Layouts.Emplace(OverrideParametersData->UserParameters);
+
+					const UPCGUserParametersData* UpstreamData = OverrideParametersData->UpstreamData.Get();
+					while (UpstreamData)
+					{
+						Layouts.Emplace(UpstreamData->UserParameters);
+						UpstreamData = UpstreamData->UpstreamData.Get();
+					}
+				}
+				else
+				{
+					PCGE_LOG_C(Error, LogOnly, &InContext, LOCTEXT("InvalidUserParameterData", "Internal error, PCG User Parameters Data is null"));
+					return {};
 				}
 			}
-			else
-			{
-				PCGE_LOG_C(Error, LogOnly, &InContext, LOCTEXT("InvalidUserParameterData", "Internal error, PCG User Parameters Data is null"));
-				return {};
-			}
-		}
 
-		// We gather the outer graph of this node, to make sure it matches our interface.
-		const UPCGGraph* GraphFromNode = Cast<UPCGGraph>(InContext.Node->GetOuter());
-		check(GraphFromNode);
-		FConstStructView GraphFromNodeParameters = GraphFromNode->GetUserParametersStruct() ? GraphFromNode->GetUserParametersStruct()->GetValue() : FConstStructView{};
+			// We gather the outer graph of this node, to make sure it matches our interface.
+			const UPCGGraph* GraphFromNode = Cast<UPCGGraph>(InContext.Node->GetOuter());
+			check(GraphFromNode);
+			FConstStructView GraphFromNodeParameters = GraphFromNode->GetUserParametersStruct() ? GraphFromNode->GetUserParametersStruct()->GetValue() : FConstStructView{};
 
-		// If we don't have a graph instance, we just use the user parameters from the node graph owner.
-		const UPCGComponent* SourceComponent = InContext.SourceComponent.Get();
-		const UPCGGraphInstance* GraphInstance = SourceComponent ? SourceComponent->GetGraphInstance() : nullptr;
+			// If we don't have a graph instance, we just use the user parameters from the node graph owner.
+			const UPCGComponent* SourceComponent = InContext.SourceComponent.Get();
+			const UPCGGraphInstance* GraphInstance = SourceComponent ? SourceComponent->GetGraphInstance() : nullptr;
 
-		// If there is no user parameters, just use the one from the graph.
-		// Go down the graph instance chain and take all layouts. Revert back to original graph node params if none were found
-		bool bAddedNonTrivialLayout = false;
+			// If there is no user parameters, just use the one from the graph.
+			// Go down the graph instance chain and take all layouts. Revert back to original graph node params if none were found
+			bool bAddedNonTrivialLayout = false;
 
-		while (GraphInstance)
-		{
-			FConstStructView GraphParameters = GraphInstance->GetUserParametersStruct() ? GraphInstance->GetUserParametersStruct()->GetValue() : FConstStructView{};
-			if (GraphParameters.IsValid())
+			while (GraphInstance)
 			{
-				bAddedNonTrivialLayout = true;
-				Layouts.Add(GraphParameters);
+				FConstStructView GraphParameters = GraphInstance->GetUserParametersStruct() ? GraphInstance->GetUserParametersStruct()->GetValue() : FConstStructView{};
+				if (GraphParameters.IsValid())
+				{
+					bAddedNonTrivialLayout = true;
+					Layouts.Add(GraphParameters);
+				}
+
+				GraphInstance = Cast<UPCGGraphInstance>(GraphInstance->Graph);
 			}
 
-			GraphInstance = Cast<UPCGGraphInstance>(GraphInstance->Graph);
-		}
+			if (!bAddedNonTrivialLayout && GraphFromNodeParameters.IsValid())
+			{
+				Layouts.Add(GraphFromNodeParameters);
+			}
 
-		if (!bAddedNonTrivialLayout && GraphFromNodeParameters.IsValid())
-		{
-			Layouts.Add(GraphFromNodeParameters);
+			return Layouts;
 		}
+	}
+
+	namespace Constants::Conversion
+	{
+		constexpr int32 SpecificToGenericIndex = 0;
+		constexpr int32 GenericToSpecificIndex = 0;
 
-		return Layouts;
+		const FText SpecificToGenericTooltip = LOCTEXT("SpecificToGenericTooltip", "Convert to a Get User Parameter (Generic) node.");
+		const FText GenericToSpecificTooltip = LOCTEXT("GenericToSpecificTooltip", "Convert to a specified Get User Parameter node.");
 	}
-}
+} // namespace PCGUserParameterGet
 
 UPCGUserParameterGetSettings::UPCGUserParameterGetSettings(const FObjectInitializer& ObjectInitializer)
 {
@@ -121,6 +134,49 @@ FPCGElementPtr UPCGUserParameterGetSettings::CreateElement() const
 	return MakeShared<FPCGUserParameterGetElement>();
 }
 
+#if WITH_EDITOR
+TArray<FPCGPreconfiguredInfo> UPCGUserParameterGetSettings::GetConversionInfo() const
+{
+	using namespace PCGUserParameterGet::Constants::Conversion;
+	TArray<FPCGPreconfiguredInfo> ConversionInfo;
+
+	ConversionInfo.Emplace(
+		SpecificToGenericIndex,
+		FText::Format(INVTEXT("{0} (Generic)"), PCGConversion::Helpers::GetDefaultNodeTitle(UPCGGenericUserParameterGetSettings::StaticClass())),
+		SpecificToGenericTooltip);
+
+	return ConversionInfo;
+}
+#endif // WITH_EDITOR
+
+bool UPCGUserParameterGetSettings::ConvertNode(const FPCGPreconfiguredInfo& ConversionInfo)
+{
+	UPCGNode* Node = CastChecked<UPCGNode>(GetOuter());
+
+	if (ConversionInfo.PreconfiguredIndex == PCGUserParameterGet::Constants::Conversion::SpecificToGenericIndex)
+	{
+		FPCGSingleNodeConverter NodeConverter(Node, UPCGGenericUserParameterGetSettings::StaticClass());
+		if (!NodeConverter.IsGraphInitialized() || !NodeConverter.IsSourceInitialized())
+		{
+			PCGLog::Settings::LogInvalidPreconfigurationWarning(ConversionInfo.PreconfiguredIndex, LOCTEXT("GetUserParameterNodeTitle", "Get User Parameter"));
+			return false;
+		}
+
+		NodeConverter.PrepareData();
+
+		UPCGGenericUserParameterGetSettings* Settings = Cast<UPCGGenericUserParameterGetSettings>(NodeConverter.GetGeneratedSettings());
+		Settings->PropertyPath = PropertyName.ToString();
+		Settings->OutputAttributeName = PropertyName;
+
+		NodeConverter.ApplyStructural();
+		NodeConverter.Finalize();
+
+		return NodeConverter.IsComplete();
+	}
+
+	return false;
+}
+
 TArray<FPCGPinProperties> UPCGGenericUserParameterGetSettings::OutputPinProperties() const
 {
 	TArray<FPCGPinProperties> PinProperties;
@@ -134,11 +190,96 @@ FString UPCGGenericUserParameterGetSettings::GetAdditionalTitleInformation() con
 	return PropertyPath;
 }
 
+#if WITH_EDITOR
+EPCGChangeType UPCGGenericUserParameterGetSettings::GetChangeTypeForProperty(const FName& InPropertyName) const
+{
+	EPCGChangeType ChangeType = Super::GetChangeTypeForProperty(InPropertyName);
+
+	if (InPropertyName == GET_MEMBER_NAME_CHECKED(UPCGGenericUserParameterGetSettings, PropertyPath))
+	{
+		ChangeType |= EPCGChangeType::Cosmetic;
+	}
+
+	return ChangeType;
+}
+#endif // WITH_EDITOR
+
 FPCGElementPtr UPCGGenericUserParameterGetSettings::CreateElement() const
 {
 	return MakeShared<FPCGUserParameterGetElement>();
 }
 
+#if WITH_EDITOR
+TArray<FPCGPreconfiguredInfo> UPCGGenericUserParameterGetSettings::GetConversionInfo() const
+{
+	using namespace PCGUserParameterGet::Constants::Conversion;
+	TArray<FPCGPreconfiguredInfo> ConversionInfo;
+
+	ConversionInfo.Emplace(
+		GenericToSpecificIndex,
+		PCGConversion::Helpers::GetDefaultNodeTitle(UPCGUserParameterGetSettings::StaticClass()),
+		GenericToSpecificTooltip);
+
+	return ConversionInfo;
+}
+#endif // WITH_EDITOR
+
+bool UPCGGenericUserParameterGetSettings::ConvertNode(const FPCGPreconfiguredInfo& ConversionInfo)
+{
+	UPCGNode* Node = CastChecked<UPCGNode>(GetOuter());
+
+	if (ConversionInfo.PreconfiguredIndex == PCGUserParameterGet::Constants::Conversion::GenericToSpecificIndex)
+	{
+		// If the property is overridden or upstream, abort.
+		if (IsPropertyOverriddenByPin(
+			GET_MEMBER_NAME_CHECKED(UPCGGenericUserParameterGetSettings, PropertyPath))
+			|| Source != EPCGUserParameterSource::Current)
+		{
+			PCGLog::LogErrorOnGraph(LOCTEXT("GenericInvalidConversion", "Can't convert Get User Parameter with an overridden property or upstream source."));
+			return false;
+		}
+
+		FPCGSingleNodeConverter NodeConverter(Node, UPCGUserParameterGetSettings::StaticClass());
+		if (!NodeConverter.IsGraphInitialized() || !NodeConverter.IsSourceInitialized())
+		{
+			PCGLog::Settings::LogInvalidPreconfigurationWarning(ConversionInfo.PreconfiguredIndex, LOCTEXT("GenericGetUserParameterNodeTitle", "Get User Parameter (Generic)"));
+			return false;
+		}
+
+		NodeConverter.PrepareData();
+
+		bool bParameterSuccess = false;
+		if (UPCGUserParameterGetSettings* Settings = Cast<UPCGUserParameterGetSettings>(NodeConverter.GetGeneratedSettings()))
+		{
+			if (const UPCGGraph* Graph = NodeConverter.GetGraph())
+			{
+				if (const FInstancedPropertyBag* UserParametersStruct = Graph->GetUserParametersStruct())
+				{
+					// Use the base attribute and ignore extractors.
+					FPCGAttributePropertySelector Selector;
+					Selector.Update(PropertyPath);
+					if (const FPropertyBagPropertyDesc* Desc = UserParametersStruct->FindPropertyDescByName(Selector.GetAttributeName()))
+					{
+						Settings->PropertyName = Desc->Name;
+						Settings->PropertyGuid = Desc->ID;
+						bParameterSuccess = true;
+					}
+				}
+			}
+		}
+
+		if (bParameterSuccess)
+		{
+			NodeConverter.ApplyStructural();
+			NodeConverter.Finalize();
+
+			return NodeConverter.IsComplete();
+		}
+	}
+
+	return false;
+}
+
 bool FPCGUserParameterGetElement::ExecuteInternal(FPCGContext* Context) const
 {
 	TRACE_CPUPROFILER_EVENT_SCOPE(FPCGUserParameterGetElement::Execute);
@@ -150,7 +291,7 @@ bool FPCGUserParameterGetElement::ExecuteInternal(FPCGContext* Context) const
 
 	bool bForceQuiet = (GenericSettings && GenericSettings->bQuiet);
 
-	TArray<FConstStructView, TInlineAllocator<16>> ValidLayouts = PCGUserParameterGetSettings::GetValidLayouts(*Context);
+	TArray<FConstStructView, TInlineAllocator<16>> ValidLayouts = PCGUserParameterGet::Settings::GetValidLayouts(*Context);
 
 	// Remove layouts we're not interested in, as per settings
 	if (ValidLayouts.Num() > 1)
diff --git a/Engine/Plugins/PCG/Source/PCG/Public/Elements/PCGUserParameterGet.h b/Engine/Plugins/PCG/Source/PCG/Public/Elements/PCGUserParameterGet.h
--- a/Engine/Plugins/PCG/Source/PCG/Public/Elements/PCGUserParameterGet.h
+++ b/Engine/Plugins/PCG/Source/PCG/Public/Elements/PCGUserParameterGet.h
@@ -56,7 +56,10 @@ protected:
 	virtual FName GetDefaultNodeName() const override { return FName("GetGraphParameter"); }
 	virtual FText GetDefaultNodeTitle() const override { return NSLOCTEXT("PCGUserParameterGetSettings", "NodeTitle", "Get Graph Parameter"); }
 	virtual EPCGSettingsType GetType() const override { return EPCGSettingsType::GraphParameters; }
+	virtual TArray<FPCGPreconfiguredInfo> GetConversionInfo() const override;
 #endif
+
+	virtual bool ConvertNode(const FPCGPreconfiguredInfo& ConversionInfo) override;
 	//~End UPCGSettings interface
 };
 
@@ -91,6 +94,10 @@ public:
 	virtual bool RequiresDataFromPreTask() const override { return true; }
 	virtual FString GetAdditionalTitleInformation() const override;
 
+#if WITH_EDITOR
+	virtual EPCGChangeType GetChangeTypeForProperty(const FName& InPropertyName) const override;
+#endif // WITH_EDITOR
+
 protected:
 	virtual TArray<FPCGPinProperties> InputPinProperties() const override { return {}; }
 	virtual TArray<FPCGPinProperties> OutputPinProperties() const override;
@@ -100,7 +107,10 @@ protected:
 	virtual FName GetDefaultNodeName() const override { return FName("GetGenericGraphParameter"); }
 	virtual FText GetDefaultNodeTitle() const override { return NSLOCTEXT("UPCGGenericUserParameterGetSettings", "NodeTitle", "Get Graph Parameter"); }
 	virtual EPCGSettingsType GetType() const override { return EPCGSettingsType::GraphParameters; }
+	virtual TArray<FPCGPreconfiguredInfo> GetConversionInfo() const override;
 #endif
+
+	virtual bool ConvertNode(const FPCGPreconfiguredInfo& ConversionInfo) override;
 	//~End UPCGSettings interface
 };
 
