Numerous AnimBank PCG improvements, bug fixes, and optimizations.

diff --git a/Engine/Plugins/PCG/Source/PCG/Private/Compute/DataInterfaces/Elements/PCGAnimBankMeshSpawnerDataInterface.cpp b/Engine/Plugins/PCG/Source/PCG/Private/Compute/DataInterfaces/Elements/PCGAnimBankMeshSpawnerDataInterface.cpp
--- a/Engine/Plugins/PCG/Source/PCG/Private/Compute/DataInterfaces/Elements/PCGAnimBankMeshSpawnerDataInterface.cpp
+++ b/Engine/Plugins/PCG/Source/PCG/Private/Compute/DataInterfaces/Elements/PCGAnimBankMeshSpawnerDataInterface.cpp
@@ -190,7 +190,7 @@ UComputeDataProvider* UPCGAnimBankMeshSpawnerDataInterface::CreateDataProvider(T
 		DataProvider->SelectorAttributeId = Primitives->SelectorAttributeId;
 		// We know the name and type of the selected mesh attribute statically and declared the attribute in GetKernelAttributeKeys,
 		// so the attribute ID should be present in the attribute table.
-		DataProvider->SelectedMeshAttributeId = Binding->GetAttributeId(DataProvider->Settings->OutAttributeName, EPCGKernelAttributeType::StringKey);
+		// TODO: DataProvider->SelectedMeshAttributeId = Binding->GetAttributeId(DataProvider->Settings->OutMeshAttribute, EPCGKernelAttributeType::StringKey);
 		ensure(DataProvider->SelectedMeshAttributeId != INDEX_NONE);
 	}
 
diff --git a/Engine/Plugins/PCG/Source/PCG/Private/Elements/PCGAnimBankMeshSpawner.cpp b/Engine/Plugins/PCG/Source/PCG/Private/Elements/PCGAnimBankMeshSpawner.cpp
--- a/Engine/Plugins/PCG/Source/PCG/Private/Elements/PCGAnimBankMeshSpawner.cpp
+++ b/Engine/Plugins/PCG/Source/PCG/Private/Elements/PCGAnimBankMeshSpawner.cpp
@@ -16,9 +16,7 @@
 #include "Helpers/PCGActorHelpers.h"
 #include "Helpers/PCGHelpers.h"
 #include "InstanceDataPackers/PCGAnimBankMeshInstanceDataPackerBase.h"
-#include "MeshSelectors/PCGAnimBankMeshSelectorBase.h"
-#include "MeshSelectors/PCGAnimBankMeshSelectorByAttribute.h"
-#include "MeshSelectors/PCGAnimBankMeshSelectorWeighted.h"
+#include "MeshSelectors/PCGAnimBankMeshSelector.h"
 
 #include "Components/AnimBankMeshComponent.h"
 #include "Engine/SkinnedAsset.h"
@@ -32,13 +30,12 @@
 
 UPCGAnimBankMeshSpawnerSettings::UPCGAnimBankMeshSpawnerSettings(const FObjectInitializer &ObjectInitializer)
 {
-	MeshSelectorType = UPCGAnimBankMeshSelectorWeighted::StaticClass();
 	// Implementation note: this should not have been done here (it should have been null), as it causes issues with copy & paste
 	// when the thing to paste does not have that class for its instance.
 	// However, removing it makes it that any object actually using the instance created by default would be lost.
 	if (!this->HasAnyFlags(RF_ClassDefaultObject))
 	{
-		MeshSelectorParameters = ObjectInitializer.CreateDefaultSubobject<UPCGAnimBankMeshSelectorWeighted>(this, TEXT("DefaultSelectorInstance"));
+		MeshSelectorParameters = ObjectInitializer.CreateDefaultSubobject<UPCGAnimBankMeshSelector>(this, TEXT("DefaultSelectorInstance"));
 	}
 }
 
@@ -53,43 +50,26 @@ bool UPCGAnimBankMeshSpawnerSettings::IsKernelValid(FPCGContext* InContext, bool
 		return false;
 	}
 
-	if (UPCGAnimBankMeshSelectorWeighted* SelectorWeighted = Cast<UPCGAnimBankMeshSelectorWeighted>(MeshSelectorParameters))
+	if (MeshSelectorParameters == nullptr)
 	{
-		if (SelectorWeighted->MeshEntries.IsEmpty())
-		{
-			PCG_KERNEL_VALIDATION_ERR(InContext, this, bQuiet, LOCTEXT("NoMeshEntries", "No meshes specified in selector."));
-			return false;
-		}
-
-		for (const FPCGAnimBankMeshSelectorWeightedEntry& Entry : SelectorWeighted->MeshEntries)
-		{
-			if (Entry.Descriptor.SkinnedAsset.IsNull())
-			{
-				PCG_KERNEL_VALIDATION_ERR(InContext, this, bQuiet, LOCTEXT("UnassignedMesh", "Unassigned mesh."));
-				return false;
-			}
+		return false;
+	}
 
-			for (const FSoftAnimBankItem& BankItem : Entry.Descriptor.BankItems)
-			{
-				if (BankItem.BankAsset.IsNull())
-				{
-					PCG_KERNEL_VALIDATION_ERR(InContext, this, bQuiet, LOCTEXT("UnassignedAnimBank", "Unassigned anim bank."));
-					return false;
-				}
-			}
-		}
+	if (!MeshSelectorParameters->BankAttribute.IsValid())
+	{
+		PCG_KERNEL_VALIDATION_ERR(InContext, this, bQuiet, LOCTEXT("NoBankAttribute", "Mesh selector bank attribute is invalid."));
+		return false;
 	}
-	else if (UPCGAnimBankMeshSelectorByAttribute* SelectorByAttribute = Cast<UPCGAnimBankMeshSelectorByAttribute>(MeshSelectorParameters))
+
+	if (!MeshSelectorParameters->SequenceIndexAttribute.IsValid())
 	{
-		if (SelectorByAttribute->AttributeName == NAME_None)
-		{
-			PCG_KERNEL_VALIDATION_ERR(InContext, this, bQuiet, LOCTEXT("NoMeshAttribute", "Mesh selector attribute name is 'None'."));
-			return false;
-		}
+		PCG_KERNEL_VALIDATION_ERR(InContext, this, bQuiet, LOCTEXT("NoSequenceIndexAttribute", "Mesh selector sequence index attribute is invalid."));
+		return false;
 	}
-	else
+
+	if (!MeshSelectorParameters->MeshAttribute.IsValid())
 	{
-		PCG_KERNEL_VALIDATION_ERR(InContext, this, bQuiet, LOCTEXT("InvalidMeshSelector", "Currently GPU AnimBank Mesh Spawner nodes must use PCGAnimBankMeshSelectorWeighted or UPCGAnimBankMeshSelectorByAttribute as the mesh selector type."));
+		PCG_KERNEL_VALIDATION_ERR(InContext, this, bQuiet, LOCTEXT("NoMeshAttribute", "Mesh selector mesh attribute is invalid."));
 		return false;
 	}
 
@@ -122,6 +102,8 @@ bool UPCGAnimBankMeshSpawnerSettings::ComputeOutputPinDataDesc(const UPCGPin* Ou
 		return false;
 	}
 
+	check(MeshSelectorParameters);
+
 	// First inspect data coming from CPU. Will return description if this node is a CPU node or in a different compute graph.
 	if (Binding->ComputeCPUOutputPinDataDesc(OutputPin, OutDesc))
 	{
@@ -133,70 +115,12 @@ bool UPCGAnimBankMeshSpawnerSettings::ComputeOutputPinDataDesc(const UPCGPin* Ou
 	// First forward input data to output.
 	FPCGDataCollectionDesc OutputPinDesc = InputPinDesc;
 
-	// Add output attribute (selected mesh).
-	{
-		// Create unique value keys for the output string.
-		TArray<int32> UniqueStringKeys;
-		if (UPCGAnimBankMeshSelectorWeighted* SelectorWeighted = Cast<UPCGAnimBankMeshSelectorWeighted>(MeshSelectorParameters))
-		{
-			// Weighted selection - add explicit strings from settings.
-			for (const FPCGAnimBankMeshSelectorWeightedEntry& Entry : SelectorWeighted->MeshEntries)
-			{
-				const FString Value = Entry.Descriptor.SkinnedAsset.ToString();
-				const int32 StringIndex = Binding->GetStringTable().IndexOfByKey(Value);
-				if (ensureAlways(StringIndex != INDEX_NONE))
-				{
-					UniqueStringKeys.AddUnique(StringIndex);
-				}
-			}
-		}
-		else if (UPCGAnimBankMeshSelectorByAttribute* SelectorByAttribute = Cast<UPCGAnimBankMeshSelectorByAttribute>(MeshSelectorParameters))
-		{
-			// By-attribute selection - pass on strings from input attribute.
-			if (SelectorByAttribute->AttributeName != NAME_None)
-			{
-				for (const FPCGDataDesc& InputDataDesc : InputPinDesc.DataDescs)
-				{
-					for (const FPCGKernelAttributeDesc& InputAttrDesc : InputDataDesc.AttributeDescs)
-					{
-						if (InputAttrDesc.Name == SelectorByAttribute->AttributeName)
-						{
-							for (int32 InputStringKey : InputAttrDesc.GetUniqueStringKeys())
-							{
-								UniqueStringKeys.AddUnique(InputStringKey);
-							}
-						}
-					}
-				}
-			}
-		}
-		else if (MeshSelectorParameters)
-		{
-			UE_LOG(LogPCG, Warning, TEXT("Mesh selector not supported by GPU AnimBank Mesh Spawner: %s"), *MeshSelectorParameters->GetName());
-		}
-
-		OutputPinDesc.AddAttributeToAllData(FPCGKernelAttributeKey(OutAttributeName, EPCGKernelAttributeType::StringKey), Binding, &UniqueStringKeys);
-	}
-
 	OutDesc = OutputPinDesc;
 	return true;
 }
 
 void UPCGAnimBankMeshSpawnerSettings::GetKernelAttributeKeys(TArray<FPCGKernelAttributeKey>& OutKeys) const
 {
-	// Register the attribute this node creates.
-	OutKeys.AddUnique(FPCGKernelAttributeKey(OutAttributeName, EPCGKernelAttributeType::StringKey));
-}
-
-void UPCGAnimBankMeshSpawnerSettings::AddStaticCreatedStrings(TArray<FString>& InOutStringTable) const
-{
-	if (UPCGAnimBankMeshSelectorWeighted* SelectorWeighted = Cast<UPCGAnimBankMeshSelectorWeighted>(MeshSelectorParameters))
-	{
-		for (const FPCGAnimBankMeshSelectorWeightedEntry& Entry : SelectorWeighted->MeshEntries)
-		{
-			InOutStringTable.AddUnique(Entry.Descriptor.SkinnedAsset.ToString());
-		}
-	}
 }
 
 int UPCGAnimBankMeshSpawnerSettings::ComputeKernelThreadCount(const UPCGDataBinding* Binding) const
@@ -380,14 +304,6 @@ bool FPCGAnimBankMeshSpawnerElement::PrepareDataInternal(FPCGContext* InContext)
 				UPCGPointData* OutputPointData = FPCGContext::NewObject_AnyThread<UPCGPointData>(Context);
 				OutputPointData->InitializeFromData(PointData);
 
-				if (OutputPointData->Metadata->HasAttribute(Settings->OutAttributeName))
-				{
-					OutputPointData->Metadata->DeleteAttribute(Settings->OutAttributeName);
-					PCGE_LOG(Verbose, LogOnly, FText::Format(LOCTEXT("AttributeOverwritten", "Metadata attribute '{0}' is being overwritten in the output data"), FText::FromName(Settings->OutAttributeName)));
-				}
-
-				OutputPointData->Metadata->CreateStringAttribute(Settings->OutAttributeName, FName(NAME_None).ToString(), /*bAllowsInterpolation=*/false);
-
 				Output.Data = OutputPointData;
 				check(!Context->CurrentOutputPointData);
 				Context->CurrentOutputPointData = OutputPointData;
@@ -650,8 +566,21 @@ void FPCGAnimBankMeshSpawnerElement::SpawnAnimBankMeshInstances(FPCGAnimBankMesh
 	check((ABMC->GetNumCustomDataFloats() == 0 && PreExistingInstanceCount == 0) || ABMC->GetNumCustomDataFloats() == NumCustomDataFloats);
 	ABMC->SetNumCustomDataFloats(NumCustomDataFloats);
 
+	TArray<FTransform> Transforms;
+	TArray<int32> BankIndices;
+
+	Transforms.Reserve(NewInstanceCount);
+	BankIndices.Reserve(NewInstanceCount);
+
+	// TODO: Remove allocs/copies
+	for (const FPCGAnimBankMeshInstance& Instance : InstanceList.Instances)
+	{
+		Transforms.Emplace(Instance.Transform);
+		BankIndices.Emplace(Instance.BankIndex);
+	}
+
 	// Populate the instances
-	TArray<FPrimitiveInstanceId> NewIds = ABMC->AddInstances(InstanceList.Instances, /*bShouldReturnIndices=*/true, /*bWorldSpace=*/true);
+	TArray<FPrimitiveInstanceId> NewIds = ABMC->AddInstances(Transforms, BankIndices, /*bShouldReturnIndices=*/true, /*bWorldSpace=*/true);
 
 	// Copy new CustomData into the ABMC PerInstanceSMCustomData
 	if (NumCustomDataFloats > 0)
@@ -698,10 +627,10 @@ void UPCGAnimBankMeshSpawnerSettings::PostLoad()
 	Super::PostLoad();
 
 	const EObjectFlags Flags = GetMaskedFlags(RF_PropagateToSubObjects) | RF_Transactional;
-	
+
 	if (!MeshSelectorParameters)
 	{
-		RefreshMeshSelector();
+		MeshSelectorParameters = NewObject<UPCGAnimBankMeshSelector>(this, UPCGAnimBankMeshSelector::StaticClass(), NAME_None, GetMaskedFlags(RF_PropagateToSubObjects));
 	}
 	else
 	{
@@ -725,11 +654,7 @@ void UPCGAnimBankMeshSpawnerSettings::PostEditChangeProperty(FPropertyChangedEve
 	{
 		const FName& PropertyName = PropertyChangedEvent.Property->GetFName();
 
-		if (PropertyName == GET_MEMBER_NAME_CHECKED(UPCGAnimBankMeshSpawnerSettings, MeshSelectorType))
-		{
-			RefreshMeshSelector();
-		} 
-		else if (PropertyName == GET_MEMBER_NAME_CHECKED(UPCGAnimBankMeshSpawnerSettings, InstanceDataPackerType))
+		if (PropertyName == GET_MEMBER_NAME_CHECKED(UPCGAnimBankMeshSpawnerSettings, InstanceDataPackerType))
 		{
 			RefreshInstancePacker();
 		}
@@ -740,32 +665,10 @@ void UPCGAnimBankMeshSpawnerSettings::PostEditChangeProperty(FPropertyChangedEve
 
 bool UPCGAnimBankMeshSpawnerSettings::CanEditChange(const FProperty* InProperty) const
 {
-	// TODO: In place temporarily, until the other two modes are supported
-	if (InProperty->GetFName() == GET_MEMBER_NAME_CHECKED(UPCGAnimBankMeshSpawnerSettings, AnimBankMeshComponentPropertyOverrides))
-	{
-		if (!MeshSelectorType->IsChildOf(UPCGAnimBankMeshSelectorByAttribute::StaticClass()))
-		{
-			return false;
-		}
-	}
-
 	return Super::CanEditChange(InProperty);
 }
 #endif
 
-void UPCGAnimBankMeshSpawnerSettings::SetMeshSelectorType(TSubclassOf<UPCGAnimBankMeshSelectorBase> InMeshSelectorType) 
-{
-	if (!MeshSelectorParameters || InMeshSelectorType != MeshSelectorType)
-	{
-		if (InMeshSelectorType != MeshSelectorType)
-		{
-			MeshSelectorType = InMeshSelectorType;
-		}
-		
-		RefreshMeshSelector();
-	}
-}
-
 void UPCGAnimBankMeshSpawnerSettings::SetInstancePackerType(TSubclassOf<UPCGAnimBankMeshInstanceDataPackerBase> InInstancePackerType) 
 {
 	if (!InstanceDataPackerParameters || InInstancePackerType != InstanceDataPackerType)
@@ -779,30 +682,6 @@ void UPCGAnimBankMeshSpawnerSettings::SetInstancePackerType(TSubclassOf<UPCGAnim
 	}
 }
 
-void UPCGAnimBankMeshSpawnerSettings::RefreshMeshSelector()
-{
-	if (MeshSelectorType)
-	{
-		ensure(IsInGameThread());
-
-		if (MeshSelectorParameters)
-		{
-		#if WITH_EDITOR
-			MeshSelectorParameters->Rename(nullptr, GetTransientPackage(), REN_DontCreateRedirectors);
-		#endif
-			MeshSelectorParameters->MarkAsGarbage();
-			MeshSelectorParameters = nullptr;
-		}
-
-		const EObjectFlags Flags = GetMaskedFlags(RF_PropagateToSubObjects);
-		MeshSelectorParameters = NewObject<UPCGAnimBankMeshSelectorBase>(this, MeshSelectorType, NAME_None, Flags);
-	}
-	else
-	{
-		MeshSelectorParameters = nullptr;
-	}
-}
-
 void UPCGAnimBankMeshSpawnerSettings::RefreshInstancePacker()
 {
 	if (InstanceDataPackerType)
diff --git a/Engine/Plugins/PCG/Source/PCG/Private/Elements/PCGAnimBankMeshSpawnerContext.cpp b/Engine/Plugins/PCG/Source/PCG/Private/Elements/PCGAnimBankMeshSpawnerContext.cpp
--- a/Engine/Plugins/PCG/Source/PCG/Private/Elements/PCGAnimBankMeshSpawnerContext.cpp
+++ b/Engine/Plugins/PCG/Source/PCG/Private/Elements/PCGAnimBankMeshSpawnerContext.cpp
@@ -11,10 +11,7 @@ void FPCGAnimBankMeshSpawnerContext::ResetInputIterationData()
 	CurrentOutputPointData = nullptr;
 	MaterialOverrideHelper.Reset();
 	CurrentPointIndex = 0;
-	WeightedMeshInstances.Reset();
 	MeshToValueKey.Reset();
-	CumulativeWeights.Reset();
-	CategoryEntryToInstancesAndWeights.Reset();
 	AttributeOverridePartition.Reset();
 	OverriddenDescriptors.Reset();
 }
\ No newline at end of file
diff --git a/Engine/Plugins/PCG/Source/PCG/Private/InstanceDataPackers/PCGAnimBankMeshInstanceDataPackerBase.cpp b/Engine/Plugins/PCG/Source/PCG/Private/InstanceDataPackers/PCGAnimBankMeshInstanceDataPackerBase.cpp
--- a/Engine/Plugins/PCG/Source/PCG/Private/InstanceDataPackers/PCGAnimBankMeshInstanceDataPackerBase.cpp
+++ b/Engine/Plugins/PCG/Source/PCG/Private/InstanceDataPackers/PCGAnimBankMeshInstanceDataPackerBase.cpp
@@ -4,7 +4,7 @@
 
 #include "PCGElement.h"
 #include "Elements/Metadata/PCGMetadataElementCommon.h"
-#include "MeshSelectors/PCGAnimBankMeshSelectorBase.h"
+#include "MeshSelectors/PCGAnimBankMeshSelector.h"
 #include "Metadata/PCGMetadata.h"
 #include "Metadata/Accessors/IPCGAttributeAccessor.h"
 #include "Metadata/Accessors/PCGAttributeAccessorKeys.h"
@@ -65,7 +65,7 @@ void UPCGAnimBankMeshInstanceDataPackerBase::PackCustomDataFromAttributes(const
 void UPCGAnimBankMeshInstanceDataPackerBase::PackCustomDataFromAttributes(const FPCGAnimBankMeshInstanceList& InstanceList, const TArray<const FPCGMetadataAttributeBase*>& Attributes, FPCGAnimBankMeshPackedCustomData& OutPackedCustomData) const
 {
 	const UPCGPointData* PointData = InstanceList.PointData.Get();
-	for (const int32 Index : InstanceList.InstancesIndices)
+	for (const int32 Index : InstanceList.InstancePointIndices)
 	{
 		const PCGMetadataEntryKey EntryKey = PointData ? PointData->GetPoints()[Index].MetadataEntry : Index;
 		for (const FPCGMetadataAttributeBase* AttributeBase : Attributes)
@@ -184,7 +184,7 @@ void UPCGAnimBankMeshInstanceDataPackerBase::PackCustomDataFromAttributes(const
 
 void UPCGAnimBankMeshInstanceDataPackerBase::PackCustomDataFromAccessors(const FPCGAnimBankMeshInstanceList& InstanceList, TArray<TUniquePtr<const IPCGAttributeAccessor>> Accessors, TArray<TUniquePtr<const IPCGAttributeAccessorKeys>> AccessorKeys, FPCGAnimBankMeshPackedCustomData& OutPackedCustomData) const
 {
-	const int TotalNumCustomDataFloats = OutPackedCustomData.NumCustomDataFloats * InstanceList.Instances.Num();
+	const int32 TotalNumCustomDataFloats = OutPackedCustomData.NumCustomDataFloats * InstanceList.Instances.Num();
 
 	if (!ensure(Accessors.Num() == AccessorKeys.Num()))
 	{
diff --git a/Engine/Plugins/PCG/Source/PCG/Private/InstanceDataPackers/PCGAnimBankMeshInstanceDataPackerByAttribute.cpp b/Engine/Plugins/PCG/Source/PCG/Private/InstanceDataPackers/PCGAnimBankMeshInstanceDataPackerByAttribute.cpp
--- a/Engine/Plugins/PCG/Source/PCG/Private/InstanceDataPackers/PCGAnimBankMeshInstanceDataPackerByAttribute.cpp
+++ b/Engine/Plugins/PCG/Source/PCG/Private/InstanceDataPackers/PCGAnimBankMeshInstanceDataPackerByAttribute.cpp
@@ -7,7 +7,7 @@
 #include "Data/PCGSpatialData.h"
 #include "InstanceDataPackers/PCGAnimBankMeshInstanceDataPackerBase.h"
 #include "Metadata/Accessors/PCGAttributeAccessorHelpers.h"
-#include "MeshSelectors/PCGAnimBankMeshSelectorBase.h"
+#include "MeshSelectors/PCGAnimBankMeshSelector.h"
 
 #include UE_INLINE_GENERATED_CPP_BY_NAME(PCGAnimBankMeshInstanceDataPackerByAttribute)
 
@@ -44,13 +44,13 @@ void UPCGAnimBankMeshInstanceDataPackerByAttribute::PackInstances_Implementation
 		
 		if (InSpatialData == PointData)
 		{
-			Keys = MakeUnique<const FPCGAttributeAccessorKeysPointsSubset>(PointData->GetPoints(), InstanceList.InstancesIndices);
+			Keys = MakeUnique<const FPCGAttributeAccessorKeysPointsSubset>(PointData->GetPoints(), InstanceList.InstancePointIndices);
 		}
 		else
 		{
 			// Convert indices to entry keys
-			ExtractedKeys.Reserve(InstanceList.InstancesIndices.Num());
-			Algo::Transform(InstanceList.InstancesIndices, ExtractedKeys, [](const int32 Index) -> PCGMetadataEntryKey{ return Index; });
+			ExtractedKeys.Reserve(InstanceList.InstancePointIndices.Num());
+			Algo::Transform(InstanceList.InstancePointIndices, ExtractedKeys, [](const int32 Index) -> PCGMetadataEntryKey{ return Index; });
 			Keys = MakeUnique<const FPCGAttributeAccessorKeysEntries>(ExtractedKeys);
 		}
 
diff --git a/Engine/Plugins/PCG/Source/PCG/Private/MeshSelectors/PCGAnimBankMeshDescriptor.cpp b/Engine/Plugins/PCG/Source/PCG/Private/MeshSelectors/PCGAnimBankMeshDescriptor.cpp
--- a/Engine/Plugins/PCG/Source/PCG/Private/MeshSelectors/PCGAnimBankMeshDescriptor.cpp
+++ b/Engine/Plugins/PCG/Source/PCG/Private/MeshSelectors/PCGAnimBankMeshDescriptor.cpp
@@ -30,3 +30,25 @@ void FPCGSoftAnimBankMeshComponentDescriptor::InitComponent(UAnimBankMeshCompone
 	ABMComponent->ComponentTags = ComponentTags;
 	Super::InitComponent(ABMComponent);
 }
+
+int32 FPCGSoftAnimBankMeshComponentDescriptor::GetOrAddBankIndex(const FSoftAnimBankItem& BankItem)
+{
+	int32 BankIndex = INDEX_NONE;
+
+	for (int32 ExistingBankIndex = 0; ExistingBankIndex < BankItems.Num(); ++ExistingBankIndex)
+	{
+		if (BankItems[ExistingBankIndex] == BankItem)
+		{
+			BankIndex = ExistingBankIndex;
+			break;
+		}
+	}
+
+	if (BankIndex == INDEX_NONE)
+	{
+		BankIndex = BankItems.Num();
+		BankItems.Emplace(BankItem);
+	}
+
+	return BankIndex;
+}
diff --git a/Engine/Plugins/PCG/Source/PCG/Private/MeshSelectors/PCGAnimBankMeshSelector.cpp b/Engine/Plugins/PCG/Source/PCG/Private/MeshSelectors/PCGAnimBankMeshSelector.cpp
new file mode 100644
--- /dev/null
+++ b/Engine/Plugins/PCG/Source/PCG/Private/MeshSelectors/PCGAnimBankMeshSelector.cpp
@@ -0,0 +1,382 @@
+// Copyright Epic Games, Inc. All Rights Reserved.
+
+#include "MeshSelectors/PCGAnimBankMeshSelector.h"
+
+#include "Data/PCGPointData.h"
+#include "Data/PCGSpatialData.h"
+#include "Elements/PCGAnimBankMeshSpawner.h"
+#include "Elements/PCGAnimBankMeshSpawnerContext.h"
+#include "Elements/Metadata/PCGMetadataElementCommon.h"
+#include "MeshSelectors/PCGAnimBankMeshSelector.h"
+#include "Metadata/PCGObjectPropertyOverride.h"
+#include "Metadata/PCGMetadataPartitionCommon.h"
+#include "Metadata/Accessors/PCGAttributeAccessorHelpers.h"
+
+#include "Algo/AnyOf.h"
+#include "Animation/AnimBank.h"
+
+#include UE_INLINE_GENERATED_CPP_BY_NAME(PCGAnimBankMeshSelector)
+
+#define LOCTEXT_NAMESPACE "PCGAnimBankMeshSelector"
+
+namespace PCGAnimBankMeshSelector
+{
+	// Returns variation based on mesh, material overrides and reverse culling
+	FPCGAnimBankMeshInstanceList& GetInstanceList(
+		TArray<FPCGAnimBankMeshInstanceList>& InstanceLists,
+		const FPCGSoftAnimBankMeshComponentDescriptor& TemplateDescriptor,
+		TSoftObjectPtr<USkinnedAsset> Asset,
+		const TArray<TSoftObjectPtr<UMaterialInterface>>& MaterialOverrides,
+		bool bReverseCulling,
+		const UPCGPointData* InPointData,
+		const int AttributePartitionIndex = INDEX_NONE)
+	{
+		for (FPCGAnimBankMeshInstanceList& InstanceList : InstanceLists)
+		{
+			if (InstanceList.Descriptor.SkinnedAsset == Asset &&
+				// GW-TODO: InstanceList.Descriptor.bReverseCulling == bReverseCulling &&
+				// GW-TODO: InstanceList.Descriptor.OverrideMaterials == MaterialOverrides &&
+				InstanceList.AttributePartitionIndex == AttributePartitionIndex)
+			{
+				return InstanceList;
+			}
+		}
+
+		FPCGAnimBankMeshInstanceList& NewInstanceList = InstanceLists.Emplace_GetRef(TemplateDescriptor);
+		NewInstanceList.Descriptor.SkinnedAsset = Asset;
+		// GW-TODO: NewInstanceList.Descriptor.OverrideMaterials = MaterialOverrides;
+		// GW-TODO: NewInstanceList.Descriptor.bReverseCulling = bReverseCulling;
+		NewInstanceList.AttributePartitionIndex = AttributePartitionIndex;
+		NewInstanceList.PointData = InPointData;
+
+		return NewInstanceList;
+	}
+}
+
+void UPCGAnimBankMeshSelector::PostLoad()
+{
+	Super::PostLoad();
+
+	// TODO: Remove if/when FBodyInstance is updated or replaced
+	// Necessary to update the collision Response Container from the Response Array
+	TemplateDescriptor.PostLoadFixup(this);
+}
+
+bool UPCGAnimBankMeshSelector::SelectInstances(
+	FPCGAnimBankMeshSpawnerContext& Context,
+	const UPCGAnimBankMeshSpawnerSettings* Settings,
+	const UPCGPointData* InPointData,
+	TArray<FPCGAnimBankMeshInstanceList>& OutMeshInstances,
+	UPCGPointData* OutPointData) const
+{
+	TRACE_CPUPROFILER_EVENT_SCOPE(UPCGMeshSelector::SelectInstances);
+
+	if (!InPointData)
+	{
+		PCGE_LOG_C(Error, GraphAndLog, &Context, LOCTEXT("InputMissingData", "Missing input data"));
+		return true;
+	}
+
+	if (!InPointData->Metadata)
+	{
+		PCGE_LOG_C(Error, GraphAndLog, &Context, LOCTEXT("InputMissingMetadata", "Unable to get metadata from input"));
+		return true;
+	}
+
+	if (!InPointData->Metadata->HasAttribute(BankAttribute.GetAttributeName()))
+	{
+		PCGE_LOG_C(Error, GraphAndLog, &Context, FText::Format(LOCTEXT("AttributeNotInMetadata", "Attribute '{0}' is not in the metadata"), BankAttribute.GetDisplayText()));
+		return true;
+	}
+
+	if (!InPointData->Metadata->HasAttribute(SequenceIndexAttribute.GetAttributeName()))
+	{
+		PCGE_LOG_C(Error, GraphAndLog, &Context, FText::Format(LOCTEXT("AttributeNotInMetadata", "Attribute '{0}' is not in the metadata"), SequenceIndexAttribute.GetDisplayText()));
+		return true;
+	}
+
+	if (!InPointData->Metadata->HasAttribute(MeshAttribute.GetAttributeName()))
+	{
+		PCGE_LOG_C(Error, GraphAndLog, &Context, FText::Format(LOCTEXT("AttributeNotInMetadata", "Attribute '{0}' is not in the metadata"), MeshAttribute.GetDisplayText()));
+		return true;
+	}
+
+	const FPCGMetadataAttributeBase* BankAttributeBase = InPointData->Metadata->GetConstAttribute(BankAttribute.GetAttributeName());
+	check(BankAttributeBase);
+
+	const FPCGMetadataAttributeBase* SequenceIndexAttributeBase = InPointData->Metadata->GetConstAttribute(SequenceIndexAttribute.GetAttributeName());
+	check(SequenceIndexAttributeBase);
+
+	const FPCGMetadataAttributeBase* MeshAttributeBase = InPointData->Metadata->GetConstAttribute(MeshAttribute.GetAttributeName());
+	check(MeshAttributeBase);
+
+	// Validate that the "bank" attribute is of the right type
+	if (!PCG::Private::IsOfTypes<FSoftObjectPath, FString>(BankAttributeBase->GetTypeId()))
+	{
+		PCGE_LOG_C(Error, GraphAndLog, &Context, FText::Format(LOCTEXT("AttributeInvalidType", "Attribute '{0}' is not of valid type (must be FString or FSoftObjectPath)"), BankAttribute.GetDisplayText()));
+		return true;
+	}
+
+	// Validate that the "sequence index" attribute is of the right type
+	if (!PCG::Private::IsOfTypes<int32>(SequenceIndexAttributeBase->GetTypeId()))
+	{
+		PCGE_LOG_C(Error, GraphAndLog, &Context, FText::Format(LOCTEXT("AttributeInvalidType", "Attribute '{0}' is not of valid type (must be int32)"), SequenceIndexAttribute.GetDisplayText()));
+		return true;
+	}
+
+	// Validate that the "mesh" attribute is of the right type
+	if (!PCG::Private::IsOfTypes<FSoftObjectPath, FString>(MeshAttributeBase->GetTypeId()))
+	{
+		PCGE_LOG_C(Error, GraphAndLog, &Context, FText::Format(LOCTEXT("AttributeInvalidType", "Attribute '{0}' is not of valid type (must be FString or FSoftObjectPath)"), MeshAttribute.GetDisplayText()));
+		return true;
+	}
+
+	FPCGMeshMaterialOverrideHelper& MaterialOverrideHelper = Context.MaterialOverrideHelper;
+	if (!MaterialOverrideHelper.IsInitialized())
+	{
+		MaterialOverrideHelper.Initialize(Context, bUseAttributeMaterialOverrides, {} /*TemplateDescriptor.OverrideMaterials*/, MaterialOverrideAttributes, InPointData->Metadata);
+	}
+
+	if (!MaterialOverrideHelper.IsValid())
+	{
+		return true;
+	}
+
+	// ByAttribute takes in SoftObjectPaths per point in the metadata, so we can pass those directly into the outgoing pin if it exists
+	if (Context.CurrentPointIndex == 0 && OutPointData)
+	{
+		TRACE_CPUPROFILER_EVENT_SCOPE(UPCGMeshSelector::SetupOutPointData);
+		OutPointData->SetPoints(InPointData->GetPoints());
+	}
+
+	TRACE_CPUPROFILER_EVENT_SCOPE(UPCGAnimBankMeshSelector::SelectEntries);
+
+	if (!Context.bPartitionDone)
+	{
+		TRACE_CPUPROFILER_EVENT_SCOPE(UPCGAnimBankMeshSelector::SelectEntries::BuildingPartition);
+
+		TArray<FPCGObjectPropertyOverrideDescription> PropertyOverrides = Settings->AnimBankMeshComponentPropertyOverrides;
+		const FString SkinnedAssetPropertyString = GET_MEMBER_NAME_CHECKED(FPCGSoftAnimBankMeshComponentDescriptor, SkinnedAsset).ToString();
+
+		// Add the skinned asset override to the list only if not already provided
+		if (!Algo::AnyOf(PropertyOverrides, [&SkinnedAssetPropertyString](const FPCGObjectPropertyOverrideDescription& PropertyOverride) { return PropertyOverride.PropertyTarget == SkinnedAssetPropertyString; }))
+		{
+			PropertyOverrides.Emplace(MeshAttribute, SkinnedAssetPropertyString);
+		}
+
+		// Validate all the selectors are actual FPCGSoftAnimBankMeshComponentDescriptor properties
+		TArray<FPCGAttributePropertySelector> ValidSelectorOverrides;
+
+		for (const FPCGObjectPropertyOverrideDescription& PropertyOverride : PropertyOverrides)
+		{
+			if (FPCGSoftAnimBankMeshComponentDescriptor::StaticStruct()->FindPropertyByName(FName(PropertyOverride.PropertyTarget)))
+			{
+				ValidSelectorOverrides.Emplace_GetRef() = FPCGAttributePropertySelector::CreateFromOtherSelector<FPCGAttributePropertySelector>(PropertyOverride.InputSource);
+			}
+			else
+			{
+				PCGE_LOG_C(Error, GraphAndLog, &Context, FText::Format(LOCTEXT("OverriddenPropertyNotFound", "Property '{0}' not a valid property with an ABM Descriptor. It will be ignored."), FText::FromString(PropertyOverride.PropertyTarget)));
+			}
+		}
+
+		// If there are valid overrides, partition the points on those attributes so that an instance can be created for each
+		if (!ValidSelectorOverrides.IsEmpty())
+		{
+			Context.AttributeOverridePartition = PCGMetadataPartitionCommon::AttributeGenericPartition(InPointData, ValidSelectorOverrides, &Context, Settings->bSilenceOverrideAttributeNotFoundErrors);
+		}
+
+		// Set the descriptors to match the partition count. Uninitialized, because it will be copied from the template below
+		Context.OverriddenDescriptors.Reserve(Context.AttributeOverridePartition.Num());
+		for (int I = 0; I < Context.AttributeOverridePartition.Num(); ++I)
+		{
+			FPCGSoftAnimBankMeshComponentDescriptor& Descriptor = Context.OverriddenDescriptors.Add_GetRef(TemplateDescriptor);
+
+			// If partition is empty (which can happen, esp. for the default partition on the default value, we'll just skip it here.
+			if (Context.AttributeOverridePartition[I].IsEmpty())
+			{
+				continue;
+			}
+
+			// Use the Object Override to map the user's input selector and property to the descriptor
+			FPCGObjectOverrides Overrides(&Descriptor);
+			Overrides.Initialize(PropertyOverrides, &Descriptor, InPointData, &Context);
+
+			// Since they are already partitioned and identical, we can just use the value on the first point
+			check(!Context.AttributeOverridePartition[I].IsEmpty());
+			int32 AnyPointIndexOnThisPartition = Context.AttributeOverridePartition[I][0];
+			Overrides.Apply(AnyPointIndexOnThisPartition);
+		}
+
+		Context.PointBankValues.Reset();
+		Context.PointSequenceIndexValues.Reset();
+
+		// Given partitioning is expensive, check if we're out of time for this frame
+		Context.bPartitionDone = true;
+		if (Context.ShouldStop())
+		{
+			return false;
+		}
+	}
+
+	// Assign points to entries
+	int32 CurrentPartitionIndex = Context.CurrentPointIndex; // misnomer but we're reusing another concept from the context
+	const TArray<FPCGPoint>& Points = InPointData->GetPoints();
+
+	{
+		TRACE_CPUPROFILER_EVENT_SCOPE(UPCGAnimBankMeshSelector::SelectEntries::PushingPointsToInstanceLists);
+
+		if (Context.PointBankValues.IsEmpty() || Context.PointSequenceIndexValues.IsEmpty())
+		{
+			PCGAttributeAccessorHelpers::ExtractAllValues(InPointData, BankAttribute, Context.PointBankValues, &Context);
+			PCGAttributeAccessorHelpers::ExtractAllValues(InPointData, SequenceIndexAttribute, Context.PointSequenceIndexValues, &Context);
+		}
+
+		// TODO: Revisit this when attribute partitioning is returned in a more optimized form
+		// The partition index is used to assign the point to the correct partition's instance
+		while (CurrentPartitionIndex < Context.AttributeOverridePartition.Num())
+		{
+			const int32 PartitionIndex = CurrentPartitionIndex++;
+			TArray<int32>& Partition = Context.AttributeOverridePartition[PartitionIndex];
+			const FPCGSoftAnimBankMeshComponentDescriptor& CurrentPartitionDescriptor = Context.OverriddenDescriptors[PartitionIndex];
+			
+			if (Partition.IsEmpty() || CurrentPartitionDescriptor.SkinnedAsset.IsNull())
+			{
+				continue;
+			}
+
+			bool bSkip = false;
+
+			for (const FSoftAnimBankItem& BankItem : CurrentPartitionDescriptor.BankItems)
+			{
+				if (BankItem.BankAsset.IsNull())
+				{
+					bSkip = true;
+					break;
+				}
+			}
+
+			if (bSkip)
+			{
+				continue;
+			}
+
+			// Setup data for mesh bounds computation
+			if (OutPointData && Settings->bApplyMeshBoundsToPoints)
+			{
+				TRACE_CPUPROFILER_EVENT_SCOPE(UPCGAnimBankMeshSelector::SelectEntries::PushPointsToApplyMeshBounds);
+				TArray<int32>& PointIndices = Context.MeshToOutPoints.FindOrAdd(CurrentPartitionDescriptor.SkinnedAsset).FindOrAdd(OutPointData);
+				PointIndices.Append(Partition);
+			}
+
+			// Separate the inverse determinant instances so we can push them to a different ABM
+			TArray<int32> ReverseInstances;
+			{
+				TRACE_CPUPROFILER_EVENT_SCOPE(UPCGAnimBankMeshSelector::SelectEntries::BuildReverseInstances);
+				int32 WriteIndex = 0;
+				for (int32 ReadIndex = 0; ReadIndex < Partition.Num(); ++ReadIndex)
+				{
+					const FPCGPoint& Point = Points[Partition[ReadIndex]];
+					if (Point.Transform.GetDeterminant() < 0)
+					{
+						ReverseInstances.Add(Partition[ReadIndex]);
+					}
+					else
+					{
+						if (WriteIndex != ReadIndex)
+						{
+							Partition[WriteIndex] = Partition[ReadIndex];
+						}
+						++WriteIndex;
+					}
+				}
+
+				Partition.SetNum(WriteIndex);
+			}
+
+			auto AddPointsToInstanceList = [&OutMeshInstances, &CurrentPartitionDescriptor, &Context, &MaterialOverrideHelper, &Points, PartitionIndex, InPointData](const TArray<int32>& PointIndices, bool bReverseTransform)
+			{
+				if (MaterialOverrideHelper.OverridesMaterials())
+				{
+					for (int32 PointIndex = 0; PointIndex < PointIndices.Num(); ++PointIndex)
+					{
+						const int32 InstanceToPoint = PointIndices[PointIndex];
+
+						const FPCGPoint& Point = Points[InstanceToPoint];
+						FPCGAnimBankMeshInstanceList& InstanceList = PCGAnimBankMeshSelector::GetInstanceList(
+							OutMeshInstances,
+							CurrentPartitionDescriptor,
+							CurrentPartitionDescriptor.SkinnedAsset,
+							MaterialOverrideHelper.GetMaterialOverrides(Point.MetadataEntry),
+							bReverseTransform,
+							InPointData,
+							PartitionIndex
+						);
+
+						FSoftAnimBankItem BankItem;
+						{
+							BankItem.BankAsset = Context.PointBankValues[InstanceToPoint];
+							BankItem.SequenceIndex = Context.PointSequenceIndexValues[InstanceToPoint];
+						}
+
+						FPCGAnimBankMeshInstance Instance;
+						Instance.BankIndex = InstanceList.Descriptor.GetOrAddBankIndex(BankItem);
+						Instance.Transform = Point.Transform;
+
+						InstanceList.Instances.Emplace(Instance);
+						InstanceList.InstancePointIndices.Emplace(PointIndex);
+					}
+				}
+				else
+				{
+					TArray<TSoftObjectPtr<UMaterialInterface>> DummyMaterialList;
+					FPCGAnimBankMeshInstanceList& InstanceList = PCGAnimBankMeshSelector::GetInstanceList(
+						OutMeshInstances,
+						CurrentPartitionDescriptor,
+						CurrentPartitionDescriptor.SkinnedAsset,
+						DummyMaterialList,
+						bReverseTransform,
+						InPointData,
+						PartitionIndex
+					);
+
+					check(InstanceList.Instances.Num() == InstanceList.InstancePointIndices.Num());
+					const int32 InstanceOffset = InstanceList.Instances.Num();
+					InstanceList.Instances.SetNum(InstanceOffset + PointIndices.Num());
+					InstanceList.InstancePointIndices.Append(PointIndices);
+
+					for (int32 PointIndex = 0; PointIndex < PointIndices.Num(); ++PointIndex)
+					{
+						FPCGAnimBankMeshInstance& Instance = InstanceList.Instances[InstanceOffset + PointIndex];
+
+						const int32 InstanceToPoint = PointIndices[PointIndex];
+
+						FSoftAnimBankItem BankItem;
+						{
+							BankItem.BankAsset = Context.PointBankValues[InstanceToPoint];
+							BankItem.SequenceIndex = Context.PointSequenceIndexValues[InstanceToPoint];
+						}
+
+						Instance.BankIndex = InstanceList.Descriptor.GetOrAddBankIndex(BankItem);
+						Instance.Transform = Points[InstanceToPoint].Transform;
+					}
+				}
+			};
+
+			{
+				TRACE_CPUPROFILER_EVENT_SCOPE(UPCGAnimBankMeshSelector::SelectEntries::AddPointsToInstanceList);
+				AddPointsToInstanceList(Partition, false);
+				AddPointsToInstanceList(ReverseInstances, true);
+			}
+
+			if (Context.ShouldStop())
+			{
+				break;
+			}
+		}
+	}
+
+	Context.CurrentPointIndex = CurrentPartitionIndex; // misnomer, but we're using the same context
+	return CurrentPartitionIndex == Context.AttributeOverridePartition.Num();
+}
+
+#undef LOCTEXT_NAMESPACE
diff --git a/Engine/Plugins/PCG/Source/PCG/Private/MeshSelectors/PCGAnimBankMeshSelectorBase.cpp b/Engine/Plugins/PCG/Source/PCG/Private/MeshSelectors/PCGAnimBankMeshSelectorBase.cpp
deleted file mode 100644
diff --git a/Engine/Plugins/PCG/Source/PCG/Private/MeshSelectors/PCGAnimBankMeshSelectorByAttribute.cpp b/Engine/Plugins/PCG/Source/PCG/Private/MeshSelectors/PCGAnimBankMeshSelectorByAttribute.cpp
deleted file mode 100644
diff --git a/Engine/Plugins/PCG/Source/PCG/Private/MeshSelectors/PCGAnimBankMeshSelectorWeighted.cpp b/Engine/Plugins/PCG/Source/PCG/Private/MeshSelectors/PCGAnimBankMeshSelectorWeighted.cpp
deleted file mode 100644
diff --git a/Engine/Plugins/PCG/Source/PCG/Private/MeshSelectors/PCGAnimBankMeshSelectorWeightedByCategory.cpp b/Engine/Plugins/PCG/Source/PCG/Private/MeshSelectors/PCGAnimBankMeshSelectorWeightedByCategory.cpp
deleted file mode 100644
diff --git a/Engine/Plugins/PCG/Source/PCG/Public/Elements/PCGAnimBankMeshSpawner.h b/Engine/Plugins/PCG/Source/PCG/Public/Elements/PCGAnimBankMeshSpawner.h
--- a/Engine/Plugins/PCG/Source/PCG/Public/Elements/PCGAnimBankMeshSpawner.h
+++ b/Engine/Plugins/PCG/Source/PCG/Public/Elements/PCGAnimBankMeshSpawner.h
@@ -8,7 +8,7 @@
 #include "PCGAnimBankMeshSpawner.generated.h"
 
 class UPCGAnimBankMeshInstanceDataPackerBase;
-class UPCGAnimBankMeshSelectorBase;
+class UPCGAnimBankMeshSelector;
 class UPCGSpatialData;
 struct FPCGContext;
 struct FPCGAnimBankMeshInstanceList;
@@ -29,7 +29,6 @@ public:
 	virtual bool IsKernelValid(FPCGContext* InContext = nullptr, bool bQuiet = true) const;
 	virtual FString GetCookedKernelSource(const FPCGKernelAttributeTable& InAttributeStaticTable) const override;
 	virtual void GetKernelAttributeKeys(TArray<FPCGKernelAttributeKey>& OutKeys) const override;
-	virtual void AddStaticCreatedStrings(TArray<FString>& InOutStringTable) const override;
 	virtual int ComputeKernelThreadCount(const UPCGDataBinding* Binding) const override;
 	virtual bool ComputeOutputPinDataDesc(const UPCGPin* OutputPin, const UPCGDataBinding* Binding, FPCGDataCollectionDesc& OutDesc) const override;
 #if WITH_EDITOR
@@ -58,19 +57,12 @@ public:
 	// ~End UObject interface
 #endif
 
-	UFUNCTION(BlueprintCallable, Category = Settings)
-	void SetMeshSelectorType(TSubclassOf<UPCGAnimBankMeshSelectorBase> InMeshSelectorType);
-
 	UFUNCTION(BlueprintCallable, Category = Settings)
 	void SetInstancePackerType(TSubclassOf<UPCGAnimBankMeshInstanceDataPackerBase> InInstancePackerType);
 
 public:
-	/** Defines the method of mesh selection per input data */
-	UPROPERTY(BlueprintReadOnly, EditAnywhere, NoClear, Category = MeshSelector)
-	TSubclassOf<UPCGAnimBankMeshSelectorBase> MeshSelectorType;
-
 	UPROPERTY(BlueprintReadOnly, VisibleAnywhere, Instanced, Category = MeshSelector)
-	TObjectPtr<UPCGAnimBankMeshSelectorBase> MeshSelectorParameters;
+	TObjectPtr<UPCGAnimBankMeshSelector> MeshSelectorParameters;
 
 	/**
 	 * Defines the method of custom data packing for spawned anim bank meshes. Note, Rotators are treated as 3 floats, while Quaternions are
@@ -88,10 +80,6 @@ public:
 	UPROPERTY(BlueprintReadWrite, EditAnywhere, Category = Settings)
 	TArray<FPCGObjectPropertyOverrideDescription> AnimBankMeshComponentPropertyOverrides;
 
-	/** Attribute name to store mesh SoftObjectPaths inside if the output pin is connected. Note: Will overwrite existing data if the attribute name already exists. */
-	UPROPERTY(BlueprintReadWrite, EditAnywhere, Category = Settings)
-	FName OutAttributeName = NAME_None;
-
 	/** Sets the BoundsMin and BoundsMax attributes of each point to reflect the AnimBank mesh spawned at its location */
 	UPROPERTY(BlueprintReadWrite, EditAnywhere, Category = Settings)
 	bool bApplyMeshBoundsToPoints = true;
@@ -116,7 +104,6 @@ public:
 	bool bWarnOnIdenticalSpawn = true;
 
 protected:
-	void RefreshMeshSelector();
 	void RefreshInstancePacker();
 
 private:
diff --git a/Engine/Plugins/PCG/Source/PCG/Public/Elements/PCGAnimBankMeshSpawnerContext.h b/Engine/Plugins/PCG/Source/PCG/Public/Elements/PCGAnimBankMeshSpawnerContext.h
--- a/Engine/Plugins/PCG/Source/PCG/Public/Elements/PCGAnimBankMeshSpawnerContext.h
+++ b/Engine/Plugins/PCG/Source/PCG/Public/Elements/PCGAnimBankMeshSpawnerContext.h
@@ -5,18 +5,12 @@
 #include "PCGContext.h"
 #include "PCGManagedResource.h"
 #include "Async/PCGAsyncLoadingContext.h"
-#include "MeshSelectors/PCGAnimBankMeshSelectorBase.h"
+#include "MeshSelectors/PCGAnimBankMeshSelector.h"
 #include "InstanceDataPackers/PCGAnimBankMeshInstanceDataPackerBase.h"
 
 class AActor;
 class UPCGSpatialData;
 
-struct FPCGAnimBankMeshInstancesAndWeights
-{
-	TArray<TArray<FPCGAnimBankMeshInstanceList>> MeshInstances;
-	TArray<int> CumulativeWeights;
-};
-
 struct FPCGAnimBankMeshSpawnerContext : public FPCGContext, public IPCGAsyncLoadingContext
 {
 	struct FPackedInstanceListData
@@ -53,14 +47,7 @@ struct FPCGAnimBankMeshSpawnerContext : public FPCGContext, public IPCGAsyncLoad
 	// We need to keep all point indices that will spawn this mesh, in all output point data.
 	TMap<TSoftObjectPtr<USkinnedAsset>, TMap<UPCGPointData*, TArray<int32>>> MeshToOutPoints;
 
-	// Used in weighted selector
-	TArray<TArray<FPCGAnimBankMeshInstanceList>> WeightedMeshInstances;
-	//TMap<FSoftAnimBankItem, PCGMetadataValueKey> MeshToValueKey;
 	TMap<TSoftObjectPtr<USkinnedAsset>, PCGMetadataValueKey> MeshToValueKey;
-	TArray<int> CumulativeWeights;
-
-	// Used in the weighted by category selector
-	TMap<PCGMetadataValueKey, FPCGAnimBankMeshInstancesAndWeights> CategoryEntryToInstancesAndWeights;
 
 	// Used for mesh property overrides
 	TArray<TArray<int32>> AttributeOverridePartition;
@@ -69,5 +56,8 @@ struct FPCGAnimBankMeshSpawnerContext : public FPCGContext, public IPCGAsyncLoad
 	// Keeping track of all touched resources to allow for correct cleanup on abort
 	TArray<TWeakObjectPtr<UPCGManagedAnimBankMeshComponent>> TouchedResources;
 
+	TArray<FSoftObjectPath> PointBankValues;
+	TArray<int32> PointSequenceIndexValues;
+
 	void ResetInputIterationData();
 };
\ No newline at end of file
diff --git a/Engine/Plugins/PCG/Source/PCG/Public/MeshSelectors/PCGAnimBankMeshDescriptor.h b/Engine/Plugins/PCG/Source/PCG/Public/MeshSelectors/PCGAnimBankMeshDescriptor.h
--- a/Engine/Plugins/PCG/Source/PCG/Public/MeshSelectors/PCGAnimBankMeshDescriptor.h
+++ b/Engine/Plugins/PCG/Source/PCG/Public/MeshSelectors/PCGAnimBankMeshDescriptor.h
@@ -3,10 +3,26 @@
 #pragma once
 
 #include "Engine/CollisionProfile.h"
+#include "Engine/DataTable.h"
 #include "Animation/AnimBank.h"
 
 #include "PCGAnimBankMeshDescriptor.generated.h"
 
+USTRUCT()
+struct FPCGAnimBankDataRow : public FTableRowBase
+{
+	GENERATED_BODY()
+
+	UPROPERTY(EditAnywhere, Category=AnimBank)
+	TObjectPtr<USkeletalMesh> Mesh = nullptr;
+
+	UPROPERTY(EditAnywhere, Category=AnimBank)
+	TObjectPtr<UAnimBank> Bank = nullptr;
+
+	UPROPERTY(EditAnywhere, Category=AnimBank)
+	uint32 SequenceIndex = 0;
+};
+
 /** Convenience PCG-side component descriptor so we can adjust defaults to the most common use cases. */
 // Implementation note: the tags don't really need to contribute to the hash, so we will retain the base class !=, == and ComputeHash implementations
 USTRUCT()
@@ -17,7 +33,9 @@ struct FPCGSoftAnimBankMeshComponentDescriptor : public FSoftAnimBankMeshCompone
 	PCG_API FPCGSoftAnimBankMeshComponentDescriptor();
 	PCG_API explicit FPCGSoftAnimBankMeshComponentDescriptor(const FAnimBankMeshComponentDescriptor& Other);
 	PCG_API virtual void InitFrom(const UAnimBankMeshComponent* Component, bool bInitBodyInstance = true) override;
-	PCG_API virtual void InitComponent(UAnimBankMeshComponent* ISMComponent) const override;
+	PCG_API virtual void InitComponent(UAnimBankMeshComponent* Component) const override;
+
+	PCG_API int32 GetOrAddBankIndex(const FSoftAnimBankItem& BankItem);
 
 public:
 	UPROPERTY(EditAnywhere, Category="Component Settings|Tags", meta = (DisplayAfter = "ComponentClass"))
diff --git a/Engine/Plugins/PCG/Source/PCG/Public/MeshSelectors/PCGAnimBankMeshSelector.h b/Engine/Plugins/PCG/Source/PCG/Public/MeshSelectors/PCGAnimBankMeshSelector.h
new file mode 100644
--- /dev/null
+++ b/Engine/Plugins/PCG/Source/PCG/Public/MeshSelectors/PCGAnimBankMeshSelector.h
@@ -0,0 +1,100 @@
+// Copyright Epic Games, Inc. All Rights Reserved.
+
+#pragma once
+
+#include "PCGPoint.h"
+#include "Data/PCGPointData.h"
+#include "Engine/CollisionProfile.h"
+#include "MeshSelectors/PCGAnimBankMeshDescriptor.h"
+#include "MeshSelectors/PCGMeshMaterialOverrideHelper.h"
+#include "Metadata/PCGMetadata.h"
+
+#include "MeshSelectors/PCGAnimBankMeshDescriptor.h"
+
+#include "PCGAnimBankMeshSelector.generated.h"
+
+class UPCGPointData;
+class UPCGSpatialData;
+class UAnimBank;
+struct FPCGContext;
+struct FPCGAnimBankMeshSpawnerContext;
+class UPCGAnimBankMeshSpawnerSettings;
+
+USTRUCT(BlueprintType)
+struct FPCGAnimBankMeshInstance
+{
+	GENERATED_BODY()
+
+	FTransform Transform;
+	int32 BankIndex = 0;
+};
+
+USTRUCT(BlueprintType)
+struct FPCGAnimBankMeshInstanceList
+{
+	GENERATED_BODY()
+
+	explicit FPCGAnimBankMeshInstanceList(const FPCGSoftAnimBankMeshComponentDescriptor& InDescriptor)
+	: Descriptor(InDescriptor)
+	, AttributePartitionIndex(INDEX_NONE)
+	{}
+	
+	FPCGAnimBankMeshInstanceList() = default;
+	~FPCGAnimBankMeshInstanceList() = default;
+	FPCGAnimBankMeshInstanceList(const FPCGAnimBankMeshInstanceList&) = default;
+	FPCGAnimBankMeshInstanceList(FPCGAnimBankMeshInstanceList&&) = default;
+	FPCGAnimBankMeshInstanceList& operator=(const FPCGAnimBankMeshInstanceList&) = default;
+	FPCGAnimBankMeshInstanceList& operator=(FPCGAnimBankMeshInstanceList&&) = default;
+
+	UPROPERTY(EditAnywhere, Category = Settings)
+	FPCGSoftAnimBankMeshComponentDescriptor Descriptor;
+
+	UPROPERTY(BlueprintReadWrite, EditAnywhere, Category = Settings)
+	TArray<FPCGAnimBankMeshInstance> Instances;
+
+	/** Tracks which partition the instance list belongs to. */
+	int64 AttributePartitionIndex;
+
+	UPROPERTY(BlueprintReadWrite, EditAnywhere, Category = Settings)
+	TWeakObjectPtr<const UPCGPointData> PointData;
+	
+	UPROPERTY(BlueprintReadWrite, EditAnywhere, Category = Settings)
+	TArray<int32> InstancePointIndices;
+};
+
+UCLASS(MinimalAPI, BlueprintType, ClassGroup = (Procedural))
+class UPCGAnimBankMeshSelector : public UObject 
+{
+	GENERATED_BODY()
+
+public:
+	// ~Begin UObject interface
+	virtual void PostLoad() override;
+	// ~End UObject interface
+
+	virtual bool SelectInstances(
+		FPCGAnimBankMeshSpawnerContext& Context,
+		const UPCGAnimBankMeshSpawnerSettings* Settings,
+		const UPCGPointData* InPointData,
+		TArray<FPCGAnimBankMeshInstanceList>& OutMeshInstances,
+		UPCGPointData* OutPointData) const;
+
+public:
+	UPROPERTY(BlueprintReadWrite, EditAnywhere, Category = MeshSelector)
+	FPCGAttributePropertyInputSelector BankAttribute;
+
+	UPROPERTY(BlueprintReadWrite, EditAnywhere, Category = MeshSelector)
+	FPCGAttributePropertyInputSelector SequenceIndexAttribute;
+
+	UPROPERTY(BlueprintReadWrite, EditAnywhere, Category = MeshSelector)
+	FPCGAttributePropertyInputSelector MeshAttribute;
+
+	UPROPERTY(EditAnywhere, Category = MeshSelector)
+	FPCGSoftAnimBankMeshComponentDescriptor TemplateDescriptor;
+
+	UPROPERTY(BlueprintReadWrite, EditAnywhere, Category = MeshSelector, meta = (InlineEditConditionToggle))
+	bool bUseAttributeMaterialOverrides = false;
+
+	UPROPERTY(BlueprintReadWrite, EditAnywhere, DisplayName = "By Attribute Material Overrides", Category = MeshSelector, meta = (EditCondition = "bUseAttributeMaterialOverrides"))
+	TArray<FName> MaterialOverrideAttributes;
+};
diff --git a/Engine/Plugins/PCG/Source/PCG/Public/MeshSelectors/PCGAnimBankMeshSelectorBase.h b/Engine/Plugins/PCG/Source/PCG/Public/MeshSelectors/PCGAnimBankMeshSelectorBase.h
deleted file mode 100644
diff --git a/Engine/Plugins/PCG/Source/PCG/Public/MeshSelectors/PCGAnimBankMeshSelectorByAttribute.h b/Engine/Plugins/PCG/Source/PCG/Public/MeshSelectors/PCGAnimBankMeshSelectorByAttribute.h
deleted file mode 100644
diff --git a/Engine/Plugins/PCG/Source/PCG/Public/MeshSelectors/PCGAnimBankMeshSelectorWeighted.h b/Engine/Plugins/PCG/Source/PCG/Public/MeshSelectors/PCGAnimBankMeshSelectorWeighted.h
deleted file mode 100644
diff --git a/Engine/Plugins/PCG/Source/PCG/Public/MeshSelectors/PCGAnimBankMeshSelectorWeightedByCategory.h b/Engine/Plugins/PCG/Source/PCG/Public/MeshSelectors/PCGAnimBankMeshSelectorWeightedByCategory.h
deleted file mode 100644
diff --git a/Engine/Source/Runtime/Engine/Classes/Animation/AnimBank.h b/Engine/Source/Runtime/Engine/Classes/Animation/AnimBank.h
--- a/Engine/Source/Runtime/Engine/Classes/Animation/AnimBank.h
+++ b/Engine/Source/Runtime/Engine/Classes/Animation/AnimBank.h
@@ -445,10 +445,10 @@ public:
 	UPROPERTY(EditAnywhere, Category = "Component Settings")
 	TSubclassOf<UAnimBankMeshComponent> ComponentClass;
 
-	UPROPERTY(EditAnywhere, Category = "Component Settings", meta = (DisplayAfter = "ComponentClass"))
+	UPROPERTY()
 	TSoftObjectPtr<class USkinnedAsset> SkinnedAsset;
 
-	UPROPERTY(EditAnywhere, Category = "Component Settings", meta = (DisplayAfter = "SkinnedAsset"))
+	UPROPERTY()
 	TArray<FSoftAnimBankItem> BankItems;
 };
 
diff --git a/Engine/Source/Runtime/Engine/Classes/Components/AnimBankMeshComponent.h b/Engine/Source/Runtime/Engine/Classes/Components/AnimBankMeshComponent.h
--- a/Engine/Source/Runtime/Engine/Classes/Components/AnimBankMeshComponent.h
+++ b/Engine/Source/Runtime/Engine/Classes/Components/AnimBankMeshComponent.h
@@ -136,7 +136,7 @@ public:
 
 	/** Add multiple instances to this component. Transform is given in local space of this component unless bWorldSpace is set. */
 	UFUNCTION(BlueprintCallable, Category = "Components|AnimBankMesh")
-	ENGINE_API TArray<FPrimitiveInstanceId> AddInstances(const TArray<FTransform>& InstanceTransforms, bool bShouldReturnIds, bool bWorldSpace = false);
+	ENGINE_API TArray<FPrimitiveInstanceId> AddInstances(const TArray<FTransform>& Transforms, const TArray<int32>& BankIndices, bool bShouldReturnIds, bool bWorldSpace = false);
 
 	/** Update custom data for specific instance */
 	UFUNCTION(BlueprintCallable, Category = "Components|AnimBankMesh")
diff --git a/Engine/Source/Runtime/Engine/Private/Components/AnimBankMeshComponent.cpp b/Engine/Source/Runtime/Engine/Private/Components/AnimBankMeshComponent.cpp
--- a/Engine/Source/Runtime/Engine/Private/Components/AnimBankMeshComponent.cpp
+++ b/Engine/Source/Runtime/Engine/Private/Components/AnimBankMeshComponent.cpp
@@ -1377,24 +1377,23 @@ FPrimitiveInstanceId UAnimBankMeshComponent::AddInstance(const FTransform& Insta
 	return AddInstanceInternal(InstanceData.Num(), InstanceTransform, BankIndex, bWorldSpace);
 }
 
-TArray<FPrimitiveInstanceId> UAnimBankMeshComponent::AddInstances(const TArray<FTransform>& InstanceTransforms, bool bShouldReturnIds, bool bWorldSpace)
+TArray<FPrimitiveInstanceId> UAnimBankMeshComponent::AddInstances(const TArray<FTransform>& Transforms, const TArray<int32>& BankIndices, bool bShouldReturnIds, bool bWorldSpace)
 {
 	TArray<FPrimitiveInstanceId> NewInstanceIds;
-	if (InstanceTransforms.IsEmpty())
+	if (Transforms.IsEmpty() || (Transforms.Num() != BankIndices.Num()))
 	{
 		return NewInstanceIds;
 	}
 
 	Modify();
 
-	const int32 NumToAdd = InstanceTransforms.Num();
+	const int32 NumToAdd = Transforms.Num();
 
 	if (bShouldReturnIds)
 	{
 		NewInstanceIds.SetNumUninitialized(NumToAdd);
 	}
 
-
 	// Reserve memory space
 	const int32 NewNumInstances = InstanceData.Num() + NumToAdd;
 	InstanceData.Reserve(NewNumInstances);
@@ -1405,9 +1404,9 @@ TArray<FPrimitiveInstanceId> UAnimBankMeshComponent::AddInstances(const TArray<F
 
 	for (int32 AddIndex = 0; AddIndex < NumToAdd; ++AddIndex)
 	{
-		const FTransform& InstanceTransform = InstanceTransforms[AddIndex];
-		int32 BankIndex = 0; // TODO
-		FPrimitiveInstanceId InstanceId = AddInstanceInternal(InstanceData.Num(), InstanceTransform, BankIndex, bWorldSpace);
+		const FTransform& Transform = Transforms[AddIndex];
+		const int32 BankIndex = BankIndices[AddIndex];
+		FPrimitiveInstanceId InstanceId = AddInstanceInternal(InstanceData.Num(), Transform, BankIndex, bWorldSpace);
 		if (bShouldReturnIds)
 		{
 			NewInstanceIds[AddIndex] = InstanceId;
