[PCG] Add mecanism to not react to objects, managed by a PCG Component, changes if this component is generating.

* Add virtual function on the PCG managed resources to know if an object is a managed resource.
* Add safe-guard in the OnObjectChanged function in the tracking to discard components that are generating and the object changed is a managed resource of themselves.
* Update ignoring origin check in the tracking, to also check the InOriginatingObject
* Use this mecanism on dynamic mesh component spawn by PCG.
* Small change to make sure we are not re-using a resource if the owner of the resource is different from the target actor.

#rnx

diff --git a/Engine/Plugins/PCG/Source/PCG/Private/PCGActorAndComponentMapping.cpp b/Engine/Plugins/PCG/Source/PCG/Private/PCGActorAndComponentMapping.cpp
--- a/Engine/Plugins/PCG/Source/PCG/Private/PCGActorAndComponentMapping.cpp
+++ b/Engine/Plugins/PCG/Source/PCG/Private/PCGActorAndComponentMapping.cpp
@@ -5,6 +5,7 @@
 #include "PCGComponent.h"
 #include "PCGDataAsset.h"
 #include "PCGGraph.h"
+#include "PCGManagedResource.h"
 #include "PCGModule.h"
 #include "PCGSubsystem.h"
 #include "PCGWorldActor.h"
@@ -1817,6 +1818,37 @@ void FPCGActorAndComponentMapping::OnObjectChanged(UObject* InObject, const FAct
 		return;
 	}
 
+	TArray<const UObject*, TInlineAllocator<2>> ChangedObjects = {InObject};
+	ChangedObjects.AddUnique(InOriginatingChangeObject);
+
+	// We discard the component if:
+	//   * It is null.
+	//   * It is one of the objects that changed.
+	//   * We should not refresh the owner and its owner is one of the objects that changed.
+	//   * It is generating and one of the objects that changed is managed by it (safe-guard, would probably induce an infinite loop).
+	//   * One of the objects that changed is ignored by its original component.
+	auto ShouldDiscardComponent = [&ChangedObjects, bNoRefreshOwner](const UPCGComponent* InComponent)
+	{
+		if (!InComponent
+			|| ChangedObjects.Contains(InComponent)
+			|| (bNoRefreshOwner && ChangedObjects.Contains(InComponent->GetOwner()))
+			|| (InComponent->IsGenerating() && InComponent->IsAnyObjectManagedByResource(ChangedObjects)))
+		{
+			return true;
+		}
+
+		const UObject* ObjectIgnored = nullptr;
+		const UPCGComponent* OriginalComponent = InComponent->GetConstOriginalComponent();
+
+		if (ensure(OriginalComponent) && OriginalComponent->IsIgnoringAnyChangeOrigins(ChangedObjects, ObjectIgnored))
+		{
+			PCGGraphExecutionLogging::LogChangeOriginIgnoredForComponent(ObjectIgnored, InComponent);
+			return true;
+		}
+
+		return false;
+	};
+
 	TSet<UPCGComponent*> DirtyComponents;
 
 	EPCGComponentDirtyFlag DirtyFlag = EPCGComponentDirtyFlag::Actor;
@@ -1829,23 +1861,15 @@ void FPCGActorAndComponentMapping::OnObjectChanged(UObject* InObject, const FAct
 	const FBox ActorBounds = PCGActorAndComponentMapping::GetActorBounds(Actor);
 
 	{
-		TRACE_CPUPROFILER_EVENT_SCOPE(FPCGActorAndComponentMapping::OnActorChanged::AlwaysTrackedUpdate);
+		TRACE_CPUPROFILER_EVENT_SCOPE(FPCGActorAndComponentMapping::OnObjectChanged::AlwaysTrackedUpdate);
 
 		for (UPCGComponent* PCGComponent : AlwaysTrackedComponents)
 		{
-			if (!PCGComponent || PCGComponent == InOriginatingChangeObject || (bNoRefreshOwner && PCGComponent->GetOwner() == InObject))
+			if (ShouldDiscardComponent(PCGComponent))
 			{
 				continue;
 			}
 
-			UPCGComponent* OriginalComponent = PCGComponent->GetOriginalComponent();
-
-			if (ensure(OriginalComponent) && OriginalComponent->IsIgnoringChangeOrigin(InObject))
-			{
-				PCGGraphExecutionLogging::LogChangeOriginIgnoredForComponent(InObject, PCGComponent);
-				continue;
-			}
-
 			// Don't mark "Owner changed" if the change originate from a PCG Component. It will be delegated to the ClearCacheForActor.
 			// It is necessary to avoid infine loops when there are multiple PCG components on one actor, and one component was generated.
 			const bool bOwnerChanged = (PCGComponent->GetOwner() == InObject) && (!InOriginatingChangeObject || !InOriginatingChangeObject->IsA<UPCGComponent>());
@@ -1875,24 +1899,15 @@ void FPCGActorAndComponentMapping::OnObjectChanged(UObject* InObject, const FAct
 		// Then do an octree find to get all components that intersect with this actor.
 		// If the actor has moved, we also need to find components that intersected with it before
 		// We first do it for non-partitioned, then we do it for partitioned
-		auto UpdateNonPartitioned = [this, &DirtyComponents, InObject, CulledTrackedComponents, &RemovedTags, DirtyFlag, InOriginatingChangeObject, bNoRefreshOwner, &MatchedKeys](const FPCGComponentRef& ComponentRef) -> void
+		auto UpdateNonPartitioned = [this, &DirtyComponents, InObject, CulledTrackedComponents, &RemovedTags, DirtyFlag, InOriginatingChangeObject, bNoRefreshOwner, &MatchedKeys, &ShouldDiscardComponent](const FPCGComponentRef& ComponentRef) -> void
 		{
-			TRACE_CPUPROFILER_EVENT_SCOPE(FPCGActorAndComponentMapping::OnActorChanged::UpdateNonPartitioned);
+			TRACE_CPUPROFILER_EVENT_SCOPE(FPCGActorAndComponentMapping::OnObjectChanged::UpdateNonPartitioned);
 
 			// Don't dirty if the component was already dirtied, not tracked, the origin of the change or its owner is the changed actor and we should not refresh.
 			if (DirtyComponents.Contains(ComponentRef.Component) ||
 				!CulledTrackedComponents.Contains(ComponentRef.Component) ||
-				(InOriginatingChangeObject == ComponentRef.Component) ||
-				(bNoRefreshOwner && ComponentRef.Component->GetOwner() == InObject))
-			{
-				return;
-			}
-
-			UPCGComponent* OriginalComponent = ComponentRef.Component->GetOriginalComponent();
-
-			if (ensure(OriginalComponent) && OriginalComponent->IsIgnoringChangeOrigin(InObject))
+				ShouldDiscardComponent(ComponentRef.Component.Get()))
 			{
-				PCGGraphExecutionLogging::LogChangeOriginIgnoredForComponent(InObject, ComponentRef.Component);
 				return;
 			}
 
@@ -1909,25 +1924,16 @@ void FPCGActorAndComponentMapping::OnObjectChanged(UObject* InObject, const FAct
 
 		// For partitioned, we first need check if the original component intersect with the bounds, then forward the dirty call only to locals that intersect with the bounds.
 		// Note: CurrentActorBoundsPtr is passed by reference because it will be modified between lambda calls (cf comment above).
-		auto UpdatePartitioned = [this, &DirtyComponents, InObject, CulledTrackedComponents, &CurrentActorBoundsPtr, &RemovedTags, DirtyFlag, InOriginatingChangeObject, bNoRefreshOwner, &MatchedKeys](const FPCGComponentRef& ComponentRef)  -> void
+		auto UpdatePartitioned = [this, &DirtyComponents, InObject, CulledTrackedComponents, &CurrentActorBoundsPtr, &RemovedTags, DirtyFlag, InOriginatingChangeObject, bNoRefreshOwner, &MatchedKeys, &ShouldDiscardComponent](const FPCGComponentRef& ComponentRef)  -> void
 		{
-			TRACE_CPUPROFILER_EVENT_SCOPE(FPCGActorAndComponentMapping::OnActorChanged::UpdatePartitioned);
+			TRACE_CPUPROFILER_EVENT_SCOPE(FPCGActorAndComponentMapping::OnObjectChanged::UpdatePartitioned);
 
 			// Don't dirty if the component is not tracked, the origin of the change or its owner is the changed actor and we should not refresh.
 			// We can "re-dirty" it because changes can impact different local components, from the same
 			// original component.
 			if (!CulledTrackedComponents.Contains(ComponentRef.Component) ||
-				(InOriginatingChangeObject == ComponentRef.Component) ||
-				(bNoRefreshOwner && ComponentRef.Component->GetOwner() == InObject))
-			{
-				return;
-			}
-
-			UPCGComponent* OriginalComponent = ComponentRef.Component->GetOriginalComponent();
-
-			if (ensure(OriginalComponent) && OriginalComponent->IsIgnoringChangeOrigin(InObject))
+				ShouldDiscardComponent(ComponentRef.Component.Get()))
 			{
-				PCGGraphExecutionLogging::LogChangeOriginIgnoredForComponent(InObject, ComponentRef.Component);
 				return;
 			}
 
@@ -1973,7 +1979,7 @@ void FPCGActorAndComponentMapping::OnObjectChanged(UObject* InObject, const FAct
 			const FBox& PreviousBounds = InPreviousData->Get<0>();
 			if (PreviousBounds.IsValid && !PreviousBounds.Equals(ActorBounds))
 			{
-				TRACE_CPUPROFILER_EVENT_SCOPE(FPCGActorAndComponentMapping::OnActorChanged::SecondUpdateHasMoved);
+				TRACE_CPUPROFILER_EVENT_SCOPE(FPCGActorAndComponentMapping::OnObjectChanged::SecondUpdateHasMoved);
 
 				// Set the actor bounds with the old one, to have the right Overlap in the Partition case.
 				CurrentActorBoundsPtr = &PreviousBounds;
diff --git a/Engine/Plugins/PCG/Source/PCG/Private/PCGComponent.cpp b/Engine/Plugins/PCG/Source/PCG/Private/PCGComponent.cpp
--- a/Engine/Plugins/PCG/Source/PCG/Private/PCGComponent.cpp
+++ b/Engine/Plugins/PCG/Source/PCG/Private/PCGComponent.cpp
@@ -29,6 +29,7 @@
 #include "CoreGlobals.h"
 #include "LandscapeProxy.h"
 #include "Algo/AnyOf.h"
+#include "Algo/IndexOf.h"
 #include "Algo/Transform.h"
 #include "Components/InstancedStaticMeshComponent.h"
 #include "Engine/BlueprintGeneratedClass.h"
@@ -335,7 +336,7 @@ void UPCGComponent::AddActorsToManagedResources(const TArray<AActor*>& InActors)
 	AddToManagedResources(ManagedResource);
 }
 
-void UPCGComponent::ForEachManagedResource(TFunctionRef<void(UPCGManagedResource*)> Func)
+void UPCGComponent::ForEachManagedResource(TFunctionRef<void(UPCGManagedResource*)> InFunction)
 {
 	FScopeLock ResourcesLock(&GeneratedResourcesLock);
 	check(!GeneratedResourcesInaccessible);
@@ -343,11 +344,49 @@ void UPCGComponent::ForEachManagedResource(TFunctionRef<void(UPCGManagedResource
 	{
 		if (ManagedResource)
 		{
-			Func(ManagedResource);
+			InFunction(ManagedResource);
 		}
 	}
 }
 
+void UPCGComponent::ForEachConstManagedResource(TFunctionRef<void(const UPCGManagedResource*)> InFunction) const
+{
+	FScopeLock ResourcesLock(&GeneratedResourcesLock);
+	check(!GeneratedResourcesInaccessible);
+	for (const TObjectPtr<UPCGManagedResource>& ManagedResource : GeneratedResources)
+	{
+		if (ManagedResource)
+		{
+			InFunction(ManagedResource);
+		}
+	}
+}
+
+bool UPCGComponent::IsAnyObjectManagedByResource(const TArrayView<const UObject*> InObjects) const
+{
+	TRACE_CPUPROFILER_EVENT_SCOPE(UPCGComponent::IsAnyObjectManagedByResource);
+	
+	bool bIsManagedByResource = false;
+	ForEachConstManagedResource([&bIsManagedByResource, &InObjects](const UPCGManagedResource* ManagedResource)
+	{
+		if (bIsManagedByResource || !ManagedResource)
+		{
+			return;
+		}
+
+		for (const UObject* Object : InObjects)
+		{
+			if (ManagedResource->IsManaging(Object))
+			{
+				bIsManagedByResource = true;
+				break;
+			}
+		}
+	});
+	
+	return bIsManagedByResource;
+}
+
 bool UPCGComponent::ShouldGenerate(bool bForce, EPCGComponentGenerationTrigger RequestedGenerationTrigger) const
 {
 	if (!bActivated || !GetGraph() || !GetSubsystem())
@@ -3008,6 +3047,17 @@ UPCGComponent* UPCGComponent::GetOriginalComponent()
 	return ensure(PartitionActor) ? PartitionActor->GetOriginalComponent(this) : this;
 }
 
+const UPCGComponent* UPCGComponent::GetConstOriginalComponent() const
+{
+	if (!IsLocalComponent())
+	{
+		return this;
+	}
+
+	const APCGPartitionActor* PartitionActor = Cast<APCGPartitionActor>(GetOwner());
+	return ensure(PartitionActor) ? PartitionActor->GetOriginalComponent(this) : this;
+}
+
 UPCGData* UPCGComponent::CreateActorPCGData() const
 {
 	return CreateActorPCGData(GetOwner(), bParseActorComponents);
@@ -3447,51 +3497,85 @@ void UPCGComponent::UpdateDynamicTracking()
 	}
 }
 
-void UPCGComponent::StartIgnoringChangeOriginDuringGeneration(UObject* InChangeOriginToIgnore)
+void UPCGComponent::StartIgnoringChangeOriginDuringGeneration(const UObject* InChangeOriginToIgnore)
 {
-	FWriteScopeLock Lock(IgnoredChangeOriginsLock);
+	StartIgnoringChangeOriginsDuringGeneration(TArrayView<const UObject*>(&InChangeOriginToIgnore, 1));
+}
 
-	if (int32* FoundCounter = IgnoredChangeOriginsToCounters.Find(InChangeOriginToIgnore))
+void UPCGComponent::StartIgnoringChangeOriginsDuringGeneration(const TArrayView<const UObject*> InChangeOriginsToIgnore)
+{
+	FWriteScopeLock Lock(IgnoredChangeOriginsLock);
+	for (const UObject* InChangeOriginToIgnore : InChangeOriginsToIgnore)
 	{
-		int32& Counter = *FoundCounter; // Put in local variable to evade SA warnings.
-		ensure(Counter >= 0);
+		if (int32* FoundCounter = IgnoredChangeOriginsToCounters.Find(InChangeOriginToIgnore))
+		{
+			int32& Counter = *FoundCounter; // Put in local variable to evade SA warnings.
+			ensure(Counter >= 0);
 
-		Counter = FMath::Max(0, Counter) + 1;
+			Counter = FMath::Max(0, Counter) + 1;
+		}
+		else
+		{
+			IgnoredChangeOriginsToCounters.Add(InChangeOriginToIgnore, 1);
+		}
 	}
-	else
-	{
+}
 
-		IgnoredChangeOriginsToCounters.Add(InChangeOriginToIgnore, 1);
-	}
+void UPCGComponent::StopIgnoringChangeOriginDuringGeneration(const UObject* InChangeOriginToIgnore)
+{
+	StopIgnoringChangeOriginsDuringGeneration(TArrayView<const UObject*>(&InChangeOriginToIgnore, 1));
 }
 
-void UPCGComponent::StopIgnoringChangeOriginDuringGeneration(UObject* InChangeOriginToIgnore)
+void UPCGComponent::StopIgnoringChangeOriginsDuringGeneration(const TArrayView<const UObject*> InChangeOriginsToIgnore)
 {
 	FWriteScopeLock Lock(IgnoredChangeOriginsLock);
 
-	int32* FoundCounter = IgnoredChangeOriginsToCounters.Find(InChangeOriginToIgnore);
-	if (ensure(FoundCounter))
+	for (const UObject* InChangeOriginToIgnore : InChangeOriginsToIgnore)
 	{
-		int32& Counter = *FoundCounter; // Put in local variable to evade SA warnings.
-		ensure(Counter > 0);
-
-		if (--Counter <= 0)
+		int32* FoundCounter = IgnoredChangeOriginsToCounters.Find(InChangeOriginToIgnore);
+		if (ensure(FoundCounter))
 		{
-			IgnoredChangeOriginsToCounters.Remove(InChangeOriginToIgnore);
+			int32& Counter = *FoundCounter; // Put in local variable to evade SA warnings.
+			ensure(Counter > 0);
+
+			if (--Counter <= 0)
+			{
+				IgnoredChangeOriginsToCounters.Remove(InChangeOriginToIgnore);
+			}
 		}
 	}
 }
 
-bool UPCGComponent::IsIgnoringChangeOrigin(UObject* InChangeOrigin)
+bool UPCGComponent::IsIgnoringChangeOrigin(const UObject* InChangeOrigin) const 
+{
+	const UObject* Dummy = nullptr;
+	return IsIgnoringAnyChangeOrigins(TArrayView<const UObject*>(&InChangeOrigin, 1), Dummy);
+}
+
+bool UPCGComponent::IsIgnoringAnyChangeOrigins(const TArrayView<const UObject*> InChangeOrigins, const UObject*& OutFirstObjectFound) const
 {
-	if ((bIgnoreLandscapeTracking || (GetGraph() && GetGraph()->bIgnoreLandscapeTracking)) && Cast<ALandscapeProxy>(InChangeOrigin))
+	if ((bIgnoreLandscapeTracking || (GetGraph() && GetGraph()->bIgnoreLandscapeTracking)))
 	{
-		return true;
+		const int32 Index = Algo::IndexOfByPredicate(InChangeOrigins, [](const UObject* ChangeOrigin) { return Cast<ALandscapeProxy>(ChangeOrigin); });
+		if (Index != INDEX_NONE)
+		{
+			OutFirstObjectFound = InChangeOrigins[Index];
+			return true;
+		}
 	}
 
 	FReadScopeLock Lock(IgnoredChangeOriginsLock);
-	const int32* Counter = IgnoredChangeOriginsToCounters.Find(InChangeOrigin);
-	return Counter && ensure(*Counter > 0);
+	for (const UObject* ChangeOrigin : InChangeOrigins)
+	{
+		const int32* Counter = IgnoredChangeOriginsToCounters.Find(ChangeOrigin);
+		if (Counter && ensure(*Counter > 0))
+		{
+			OutFirstObjectFound = ChangeOrigin;
+			return true;
+		}
+	}
+	
+	return false;
 }
 
 void UPCGComponent::ResetIgnoredChangeOrigins(bool bLogIfAnyPresent)
@@ -3508,7 +3592,6 @@ void UPCGComponent::ResetIgnoredChangeOrigins(bool bLogIfAnyPresent)
 
 	IgnoredChangeOriginsToCounters.Reset();
 }
-
 #endif // WITH_EDITOR
 
 void UPCGComponent::SetManagedResources(const TArray<TObjectPtr<UPCGManagedResource>>& Resources)
diff --git a/Engine/Plugins/PCG/Source/PCG/Private/PCGManagedResource.cpp b/Engine/Plugins/PCG/Source/PCG/Private/PCGManagedResource.cpp
--- a/Engine/Plugins/PCG/Source/PCG/Private/PCGManagedResource.cpp
+++ b/Engine/Plugins/PCG/Source/PCG/Private/PCGManagedResource.cpp
@@ -6,6 +6,7 @@
 #include "PCGModule.h"
 #include "Helpers/PCGHelpers.h"
 
+#include "Algo/AnyOf.h"
 #include "Components/InstancedStaticMeshComponent.h"
 #include "Components/AnimBankMeshComponent.h"
 #include "Components/SplineMeshComponent.h"
@@ -200,6 +201,20 @@ void UPCGManagedActors::MarkAsReused()
 	}
 }
 
+bool UPCGManagedActors::IsManaging(const UObject* InObject) const
+{
+	if (!InObject || !InObject->IsA<AActor>())
+	{
+		return false;
+	}
+
+	return Algo::AnyOf(GeneratedActors, [InObject](const TSoftObjectPtr<AActor>& SoftGeneratedActor)
+		{
+			const AActor* GeneratedActor = SoftGeneratedActor.Get();
+			return GeneratedActor && GeneratedActor == InObject;
+		});
+}
+
 #if WITH_EDITOR
 void UPCGManagedActors::ChangeTransientState(EPCGEditorDirtyMode NewEditingMode)
 {
@@ -582,6 +597,21 @@ bool UPCGManagedComponentBase::MoveResourceToNewActor(AActor* NewActor, const AA
 	return bMovedResources;
 }
 
+bool UPCGManagedComponentBase::IsManaging(const UObject* InObject) const
+{
+	if (!InObject || !InObject->IsA<UActorComponent>())
+	{
+		return false;
+	}
+
+	// Const-cast because GetComponentsArray is not const.
+	return Algo::AnyOf(const_cast<UPCGManagedComponentBase*>(this)->GetComponentsArray(), [InObject](const TSoftObjectPtr<UActorComponent>& SoftActorComponent)
+		{
+			const UActorComponent* ActorComponent = SoftActorComponent.Get();
+			return ActorComponent && ActorComponent == InObject;
+		});
+}
+
 void UPCGManagedComponentBase::MarkAsUsed()
 {
 	if (!bIsMarkedUnused)
diff --git a/Engine/Plugins/PCG/Source/PCG/Public/PCGComponent.h b/Engine/Plugins/PCG/Source/PCG/Public/PCGComponent.h
--- a/Engine/Plugins/PCG/Source/PCG/Public/PCGComponent.h
+++ b/Engine/Plugins/PCG/Source/PCG/Public/PCGComponent.h
@@ -12,6 +12,7 @@
 #include "ComponentInstanceDataCache.h"
 #include "Components/ActorComponent.h"
 #include "Containers/Map.h"
+#include "Misc/ScopeExit.h"
 #include "UObject/ObjectKey.h"
 
 #include "PCGComponent.generated.h"
@@ -134,6 +135,7 @@ public:
 
 	/** If this is a local component returns self, otherwise returns the original component. */
 	UPCGComponent* GetOriginalComponent();
+	const UPCGComponent* GetConstOriginalComponent() const;
 
 	UPCGSchedulingPolicyBase* GetRuntimeGenSchedulingPolicy() const { return SchedulingPolicy; }
 
@@ -161,6 +163,10 @@ public:
 
 public:
 	void ForEachManagedResource(TFunctionRef<void(UPCGManagedResource*)> InFunction);
+	void ForEachConstManagedResource(TFunctionRef<void(const UPCGManagedResource*)> InFunction) const;
+
+	/** Will scan the managed resources to check if any resource manage one of the objects. */
+	bool IsAnyObjectManagedByResource(const TArrayView<const UObject*> InObjects) const;
 
 	void Generate();
 	void Cleanup();
@@ -417,12 +423,43 @@ public:
 	void RegisterDynamicTracking(const UPCGSettings* InSettings, const TArrayView<TPair<FPCGSelectionKey, bool>>& InDynamicKeysAndCulling);
 
 	/** For duration of the current/next generation, any change triggers from this change origin will be discarded. */
-	void StartIgnoringChangeOriginDuringGeneration(UObject* InChangeOriginToIgnore);
-	void StopIgnoringChangeOriginDuringGeneration(UObject* InChangeOriginToIgnore);
-	bool IsIgnoringChangeOrigin(UObject* InChangeOrigin);
+	void StartIgnoringChangeOriginDuringGeneration(const UObject* InChangeOriginToIgnore);
+	void StartIgnoringChangeOriginsDuringGeneration(const TArrayView<const UObject*> InChangeOriginsToIgnore);
+
+	void StopIgnoringChangeOriginDuringGeneration(const UObject* InChangeOriginToIgnore);
+	void StopIgnoringChangeOriginsDuringGeneration(const TArrayView<const UObject*> InChangeOriginsToIgnore);
+
+	bool IsIgnoringChangeOrigin(const UObject* InChangeOrigin) const;
+	bool IsIgnoringAnyChangeOrigins(const TArrayView<const UObject*> InChangeOrigins, const UObject*& OutFirstObjectFound) const;
 	void ResetIgnoredChangeOrigins(bool bLogIfAnyPresent);
 #endif
 
+	// Allow for the function to be defined outside of editor, so it's lighter at the calling site and will just call the InFunc.
+	// Note that this will mark the original component.
+	template <typename Func>
+	void IgnoreChangeOriginsDuringGenerationWithScope(const TArrayView<const UObject*> InChangeOriginsToIgnore, Func InFunc)
+	{
+#if WITH_EDITOR
+		UPCGComponent* OriginalComponent = GetOriginalComponent();
+		check(OriginalComponent);
+		
+		OriginalComponent->StartIgnoringChangeOriginsDuringGeneration(InChangeOriginsToIgnore);
+
+		ON_SCOPE_EXIT
+		{
+			OriginalComponent->StopIgnoringChangeOriginsDuringGeneration(InChangeOriginsToIgnore);
+		};
+#endif // WITH_EDITOR
+		
+		InFunc();
+	}
+
+	template <typename Func>
+	void IgnoreChangeOriginDuringGenerationWithScope(const UObject* InChangeOriginToIgnore, Func InFunc)
+	{
+		IgnoreChangeOriginsDuringGenerationWithScope(TArrayView<const UObject*>(&InChangeOriginToIgnore, 1), std::move(InFunc));
+	}
+
 	/** Utility function (mostly for tests) to properly set the value of bIsComponentPartitioned.
 	*   Will do an immediate cleanup first and then register/unregister the component to the subsystem.
 	*   It's your responsibility after to regenerate the graph if you want to.
@@ -693,7 +730,7 @@ private:
 	* of an element. Entries removed when counter is decremented to 0, so empty map means no active ignores.
 	*/
 	TMap<TObjectKey<UObject>, int32> IgnoredChangeOriginsToCounters;
-	FRWLock IgnoredChangeOriginsLock;
+	mutable FRWLock IgnoredChangeOriginsLock;
 #endif
 
 	mutable FCriticalSection GeneratedResourcesLock;
diff --git a/Engine/Plugins/PCG/Source/PCG/Public/PCGManagedResource.h b/Engine/Plugins/PCG/Source/PCG/Public/PCGManagedResource.h
--- a/Engine/Plugins/PCG/Source/PCG/Public/PCGManagedResource.h
+++ b/Engine/Plugins/PCG/Source/PCG/Public/PCGManagedResource.h
@@ -55,6 +55,9 @@ public:
 	const FPCGCrc& GetCrc() const { return Crc; }
 	void SetCrc(const FPCGCrc& InCrc) { Crc = InCrc; }
 
+	/** Returns true if this resource manages this object. */
+	virtual bool IsManaging(const UObject* InObject) const { return false; }
+
 #if WITH_EDITOR
 	virtual void ChangeTransientState(EPCGEditorDirtyMode NewEditingMode);
 	virtual void MarkTransientOnLoad() { bMarkedTransientOnLoad = true; }
@@ -96,6 +99,8 @@ public:
 	virtual void MarkAsUsed() override;
 	virtual void MarkAsReused() override;
 
+	virtual bool IsManaging(const UObject* InObject) const override;
+
 #if WITH_EDITOR
 	virtual void ChangeTransientState(EPCGEditorDirtyMode NewEditingMode) override;
 #endif
@@ -120,6 +125,7 @@ public:
 	virtual void MarkAsUsed() override;
 	virtual void MarkAsReused() override;
 	virtual bool MoveResourceToNewActor(AActor* NewActor, const AActor* ExpectedPreviousOwner) override;
+	virtual bool IsManaging(const UObject* InObject) const override;
 
 #if WITH_EDITOR
 	virtual void ChangeTransientState(EPCGEditorDirtyMode NewEditingMode) override;
diff --git a/Engine/Plugins/PCGInterops/PCGGeometryScriptInterop/Source/PCGGeometryScriptInterop/Private/Editor/PCGDynamicMeshDataVisualization.cpp b/Engine/Plugins/PCGInterops/PCGGeometryScriptInterop/Source/PCGGeometryScriptInterop/Private/Editor/PCGDynamicMeshDataVisualization.cpp
--- a/Engine/Plugins/PCGInterops/PCGGeometryScriptInterop/Source/PCGGeometryScriptInterop/Private/Editor/PCGDynamicMeshDataVisualization.cpp
+++ b/Engine/Plugins/PCGInterops/PCGGeometryScriptInterop/Source/PCGGeometryScriptInterop/Private/Editor/PCGDynamicMeshDataVisualization.cpp
@@ -4,6 +4,8 @@
 
 #if WITH_EDITOR
 
+#include "PCGComponent.h"
+#include "PCGContext.h"
 #include "Data/PCGDynamicMeshData.h"
 #include "Helpers/PCGHelpers.h"
 #include "Resources/PCGDynamicMeshManagedComponent.h"
@@ -16,7 +18,8 @@
 void FPCGDynamicMeshDataVisualization::ExecuteDebugDisplay(FPCGContext* Context, const UPCGSettingsInterface* SettingsInterface, const UPCGData* Data, AActor* TargetActor) const
 {
 	const UPCGDynamicMeshData* DynMeshData = CastChecked<UPCGDynamicMeshData>(Data);
-	if (!DynMeshData)
+	UPCGComponent* SourcePCGComponent = Context->SourceComponent.Get();
+	if (!DynMeshData || !SourcePCGComponent)
 	{
 		return;
 	}
@@ -27,7 +30,8 @@ void FPCGDynamicMeshDataVisualization::ExecuteDebugDisplay(FPCGContext* Context,
 	
 	if (Component)
 	{
-		DynMeshData->InitializeDynamicMeshComponentFromData(Component);
+		// Modifying the dynamic mesh component would trigger a refresh, so make sure we notify the PCG Component to ignore any changes to the dyn mesh component during that time.
+		SourcePCGComponent->IgnoreChangeOriginDuringGenerationWithScope(Component, [&DynMeshData, &Component]() { DynMeshData->InitializeDynamicMeshComponentFromData(Component); });
 	}
 }
 
diff --git a/Engine/Plugins/PCGInterops/PCGGeometryScriptInterop/Source/PCGGeometryScriptInterop/Private/Elements/PCGSpawnDynamicMesh.cpp b/Engine/Plugins/PCGInterops/PCGGeometryScriptInterop/Source/PCGGeometryScriptInterop/Private/Elements/PCGSpawnDynamicMesh.cpp
--- a/Engine/Plugins/PCGInterops/PCGGeometryScriptInterop/Source/PCGGeometryScriptInterop/Private/Elements/PCGSpawnDynamicMesh.cpp
+++ b/Engine/Plugins/PCGInterops/PCGGeometryScriptInterop/Source/PCGGeometryScriptInterop/Private/Elements/PCGSpawnDynamicMesh.cpp
@@ -2,6 +2,7 @@
 
 #include "Elements/PCGSpawnDynamicMesh.h"
 
+#include "PCGComponent.h"
 #include "PCGContext.h"
 #include "Components/DynamicMeshComponent.h"
 #include "Data/PCGDynamicMeshData.h"
@@ -49,6 +50,12 @@ bool FPCGSpawnDynamicMeshElement::ExecuteInternal(FPCGContext* InContext) const
 		PCGLog::LogErrorOnGraph(LOCTEXT("InvalidTargetActor", "Invalid target actor."), InContext);
 		return true;
 	}
+	
+	UPCGComponent* SourcePCGComponent = InContext->SourceComponent.Get();
+	if (!SourcePCGComponent)
+	{
+		return true;
+	}
 
 	for (const FPCGTaggedData& Input : InContext->InputData.GetInputsByPin(PCGPinConstants::DefaultInputLabel))
 	{
@@ -66,8 +73,7 @@ bool FPCGSpawnDynamicMeshElement::ExecuteInternal(FPCGContext* InContext) const
 			continue;
 		}
 
-		// Needs to be a copy.
-		DynMeshData->InitializeDynamicMeshComponentFromData(Component);
+		SourcePCGComponent->IgnoreChangeOriginDuringGenerationWithScope(Component, [&DynMeshData, &Component]() { DynMeshData->InitializeDynamicMeshComponentFromData(Component); });
 
 		for (const FString& Tag : Input.Tags)
 		{
diff --git a/Engine/Plugins/PCGInterops/PCGGeometryScriptInterop/Source/PCGGeometryScriptInterop/Private/Resources/PCGDynamicMeshManagedComponent.cpp b/Engine/Plugins/PCGInterops/PCGGeometryScriptInterop/Source/PCGGeometryScriptInterop/Private/Resources/PCGDynamicMeshManagedComponent.cpp
--- a/Engine/Plugins/PCGInterops/PCGGeometryScriptInterop/Source/PCGGeometryScriptInterop/Private/Resources/PCGDynamicMeshManagedComponent.cpp
+++ b/Engine/Plugins/PCGInterops/PCGGeometryScriptInterop/Source/PCGGeometryScriptInterop/Private/Resources/PCGDynamicMeshManagedComponent.cpp
@@ -82,7 +82,7 @@ UPCGDynamicMeshManagedComponent* PCGDynamicMeshManagedComponent::GetOrCreateDyna
 	const FPCGCrc CRC = InMeshData->GetOrComputeCrc(/*bFullDataCrc=*/true);
 
 	UPCGDynamicMeshManagedComponent* ExistingResource = nullptr;
-	SourceComponent->ForEachManagedResource([&ExistingResource, &CRC](UPCGManagedResource* Resource)
+	SourceComponent->ForEachManagedResource([&ExistingResource, &CRC, TargetActor](UPCGManagedResource* Resource)
 	{
 		// If we already found a valid resource, just skip until the end.
 		if (ExistingResource)
@@ -98,7 +98,12 @@ UPCGDynamicMeshManagedComponent* PCGDynamicMeshManagedComponent::GetOrCreateDyna
 		
 		if (DynMeshResource->GetDataUID() == CRC.GetValue() && DynMeshResource->CanBeUsed())
 		{
-			ExistingResource = DynMeshResource;
+			UDynamicMeshComponent* Component = DynMeshResource->GetComponent();
+			// Make sure the component is valid and on the right target actor.
+			if (Component && Component->GetOwner() == TargetActor)
+			{
+				ExistingResource = DynMeshResource;
+			}
 		}
 	});
 
