Optimize instance data after anim banks are spawned via PCG

diff --git a/Engine/Plugins/PCG/Source/PCG/Private/Elements/PCGAnimBankMeshSpawner.cpp b/Engine/Plugins/PCG/Source/PCG/Private/Elements/PCGAnimBankMeshSpawner.cpp
--- a/Engine/Plugins/PCG/Source/PCG/Private/Elements/PCGAnimBankMeshSpawner.cpp
+++ b/Engine/Plugins/PCG/Source/PCG/Private/Elements/PCGAnimBankMeshSpawner.cpp
@@ -592,6 +592,7 @@ void FPCGAnimBankMeshSpawnerElement::SpawnAnimBankMeshInstances(FPCGAnimBankMesh
 	}
 
 	ABMC->UpdateBounds();
+	ABMC->OptimizeInstanceData();
 
 	{
 		PCGE_LOG(Verbose, LogOnly, FText::Format(LOCTEXT("GenerationInfo", "Added {0} instances of '{1}' on actor '{2}'"),
