[PCG] Enable debug display of proxy data running through CPU nodes.

Proxy data needs to be read back which can take a frame or two, so debug display may need to be called more than once. Luckily the graph executor V2 code already allowed deferred debug display, so simply chain calls.

#rnx

diff --git a/Engine/Plugins/PCG/Source/PCG/Private/Graph/PCGGraphExecutor.cpp b/Engine/Plugins/PCG/Source/PCG/Private/Graph/PCGGraphExecutor.cpp
--- a/Engine/Plugins/PCG/Source/PCG/Private/Graph/PCGGraphExecutor.cpp
+++ b/Engine/Plugins/PCG/Source/PCG/Private/Graph/PCGGraphExecutor.cpp
@@ -984,23 +984,20 @@ void FPCGGraphExecutor::PostTaskExecute(TSharedPtr<FPCGGraphActiveTask> ActiveTa
 	if (bTaskWasBypassed || bTaskFullyExecuted)
 	{
 #if WITH_EDITOR
-		// No need to add if we are post executing off of main thread
-		if(!IsInGameThread())
+		// Execute debug display.
+		if (ActiveTaskSettingsInterface && ActiveTaskSettingsInterface->bDebug)
 		{
-			PCGGraphExecutor::TScopeLock ScopeLock(CollectGCReferenceTasksLock);
-			CollectGCReferenceTasks.Add(ActiveTaskPtr);
-		}
-
-		PCGHelpers::ExecuteOnGameThread(UE_SOURCE_LOCATION, [WeakHandler = GameThreadHandler->AsWeak(), ActiveTaskPtr]()
-		{
-			if (TSharedPtr<FPCGGraphExecutor::FGameThreadHandler> Handler = WeakHandler.Pin())
+			// Debug display can span multiple frames so hold onto reference to task.
 			{
-				ActiveTaskPtr->Element->DebugDisplay(ActiveTaskPtr->Context.Get());
-
-				PCGGraphExecutor::TScopeLock ScopeLock(Handler->GetExecutor()->CollectGCReferenceTasksLock);
-				Handler->GetExecutor()->CollectGCReferenceTasks.Remove(ActiveTaskPtr);
+				PCGGraphExecutor::TScopeLock ScopeLock(CollectGCReferenceTasksLock);
+				CollectGCReferenceTasks.Add(ActiveTaskPtr);
 			}
-		});
+
+			PCGHelpers::ExecuteOnGameThread(UE_SOURCE_LOCATION, [WeakHandler = GameThreadHandler->AsWeak(), ActiveTaskPtr, WorldWeak = MakeWeakObjectPtr(World)]()
+			{
+				DebugDisplayTask(WeakHandler, ActiveTaskPtr, WorldWeak);
+			});
+		}
 		
 		if (UPCGComponent* SourceComponent = ActiveTask.Context->SourceComponent.Get())
 		{
@@ -1252,6 +1249,35 @@ void FPCGGraphExecutor::PrepareForExecute(FPCGGraphTask& Task, FCachedResult*& O
 #endif
 }
 
+#if WITH_EDITOR
+void FPCGGraphExecutor::DebugDisplayTask(TWeakPtr<FPCGGraphExecutor::FGameThreadHandler> WeakHandler, TSharedPtr<FPCGGraphActiveTask> ActiveTaskPtr, TWeakObjectPtr<UWorld> WeakWorld)
+{
+	TSharedPtr<FPCGGraphExecutor::FGameThreadHandler> Handler = WeakHandler.Pin();
+	if (!Handler)
+	{
+		return;
+	}
+
+	if (!ActiveTaskPtr->Element->DebugDisplay(ActiveTaskPtr->Context.Get()))
+	{
+		if (UPCGSubsystem* PCGSubsystem = UWorld::GetSubsystem<UPCGSubsystem>(WeakWorld.Get()))
+		{
+			// Not done, run on next game thread next tick instead. This is safe to call even if running inside a begin tick action.
+			PCGSubsystem->RegisterBeginTickAction([WeakHandler, ActiveTaskPtr, WeakWorld]()
+			{
+				DebugDisplayTask(WeakHandler, ActiveTaskPtr, WeakWorld);
+			});
+
+			return;
+		}
+	}
+
+	// If we fall through to here, debug display completed. Cleanup.
+	PCGGraphExecutor::TScopeLock ScopeLock(Handler->GetExecutor()->CollectGCReferenceTasksLock);
+	Handler->GetExecutor()->CollectGCReferenceTasks.Remove(ActiveTaskPtr);
+}
+#endif // WITH_EDITOR
+
 double FPCGGraphExecutor::GetTickBudgetInSeconds() const
 {
 #if WITH_EDITOR
diff --git a/Engine/Plugins/PCG/Source/PCG/Private/Graph/PCGGraphExecutor.h b/Engine/Plugins/PCG/Source/PCG/Private/Graph/PCGGraphExecutor.h
--- a/Engine/Plugins/PCG/Source/PCG/Private/Graph/PCGGraphExecutor.h
+++ b/Engine/Plugins/PCG/Source/PCG/Private/Graph/PCGGraphExecutor.h
@@ -183,7 +183,20 @@ private:
 		bool bDoDynamicTaskCulling = false;
 		bool bIsPostGraphTask = false;		
 	};
-		
+
+	// Handler that we can use as a Weak ptr to determine if the Executor is still valid
+	class FGameThreadHandler : public TSharedFromThis<FGameThreadHandler>
+	{
+	public:
+		FGameThreadHandler(FPCGGraphExecutor* InExecutor)
+			: Executor(InExecutor) { }
+
+		FPCGGraphExecutor* GetExecutor() { return Executor; }
+
+	private:
+		FPCGGraphExecutor* Executor = nullptr;
+	};
+
 	void PostTaskExecute(TSharedPtr<FPCGGraphActiveTask> ActiveTask, bool bIsDone);
 	bool ProcessScheduledTasks();
 	void ExecuteTasksEnded();
@@ -207,6 +220,10 @@ private:
 	/** Will check the cache for existing result or create and initialize the FPCGContext to the task */
 	void PrepareForExecute(FPCGGraphTask& Task, FCachedResult*& OutCachedResult, bool bLiveTasksLockAlreadyLocked);
 
+#if WITH_EDITOR
+	static void DebugDisplayTask(TWeakPtr<FGameThreadHandler> GameThreadHandler, TSharedPtr<FPCGGraphActiveTask> ActiveTaskPtr, TWeakObjectPtr<UWorld> WeakWorld);
+#endif
+
 	/** Store cache results and Queue next tasks */
 	void ProcessCachedResults(TArray<FCachedResult*> CachedResults);
 	TArray<FPCGTaskId> ProcessCachedResultsInternal(TArray<FCachedResult*> CachedResults);
@@ -329,19 +346,6 @@ private:
 
 	TObjectPtr<UWorld> World = nullptr;
 
-	// Handler that we can use as a Weak ptr to determine if the Executor is still valid
-	class FGameThreadHandler : public TSharedFromThis<FGameThreadHandler>
-	{
-	public:
-		FGameThreadHandler(FPCGGraphExecutor* InExecutor)
-			: Executor(InExecutor) { }
-
-		FPCGGraphExecutor* GetExecutor() { return Executor; }
-
-	private:
-		FPCGGraphExecutor* Executor = nullptr;
-	};
-
 	TSharedPtr<FGameThreadHandler> GameThreadHandler;
 };
 
diff --git a/Engine/Plugins/PCG/Source/PCG/Private/PCGElement.cpp b/Engine/Plugins/PCG/Source/PCG/Private/PCGElement.cpp
--- a/Engine/Plugins/PCG/Source/PCG/Private/PCGElement.cpp
+++ b/Engine/Plugins/PCG/Source/PCG/Private/PCGElement.cpp
@@ -725,13 +725,13 @@ FPCGContext* IPCGElement::CreateContext()
 }
 
 #if WITH_EDITOR
-void IPCGElement::DebugDisplay(FPCGContext* Context) const
+bool IPCGElement::DebugDisplay(FPCGContext* Context) const
 {
 	// Check Debug flag.
 	const UPCGSettingsInterface* SettingsInterface = Context->GetInputSettingsInterface();
 	if (!SettingsInterface || !SettingsInterface->bDebug)
 	{
-		return;
+		return true;
 	}
 
 	// If graph is being inspected, only display Debug if the component is being inspected, or in the HiGen case also display if
@@ -746,7 +746,7 @@ void IPCGElement::DebugDisplay(FPCGContext* Context) const
 		// then do not display debug.
 		if (!Graph->IsHierarchicalGenerationEnabled())
 		{
-			return;
+			return true;
 		}
 
 		// If a child of this component is being inspected (a local component on smaller grid and overlapping) then we still show debug,
@@ -767,39 +767,71 @@ void IPCGElement::DebugDisplay(FPCGContext* Context) const
 			// If no inspected child component then don't display debug.
 			if (!bFoundInspectedChildComponent)
 			{
-				return;
+				return true;
 			}
 		}
 	}
 
-	FPCGDataCollection ElementInputs = Context->InputData;
-	FPCGDataCollection ElementOutputs = Context->OutputData;
-
-	Context->InputData = ElementOutputs;
-	Context->OutputData = FPCGDataCollection();
-
 	// In the case of a node with multiple output pins, we will select only the inputs from the first non-empty pin.
 	const UPCGPin* FirstOutPin = Context->Node ? Context->Node->GetFirstConnectedOutputPin() : nullptr;
 
-	const FPCGDataVisualizationRegistry& DataVisRegistry = FPCGModule::GetConstPCGDataVisualizationRegistry();
-	TArray<FPCGTaggedData> Inputs = Context->InputData.GetInputs();
+	bool bHasPendingReadbacks = false;
+
+	TArray<FPCGTaggedData> DataToDebug;
+	DataToDebug.Reserve(Context->OutputData.TaggedData.Num());
 
-	for (const FPCGTaggedData& Input : Inputs)
+	// If there are any proxies in the input data, request readback to CPU.
+	for (const FPCGTaggedData& Output : Context->OutputData.TaggedData)
 	{
 		// Skip output if we're filtering on the first pin or the the data is null.
-		if (!Input.Data || (FirstOutPin && FirstOutPin->Properties.Label != Input.Pin))
+		if (!Output.Data || (FirstOutPin && FirstOutPin->Properties.Label != Output.Pin))
 		{
 			continue;
 		}
 
-		if (const IPCGDataVisualization* DataVis = DataVisRegistry.GetDataVisualization(Input.Data->GetClass()))
+		if (const UPCGProxyForGPUData* DataWithGPUSupport = Cast<UPCGProxyForGPUData>(Output.Data))
+		{
+			// Poll until readback is done.
+			UPCGProxyForGPUData::FReadbackResult Result = DataWithGPUSupport->GetCPUData(Context);
+
+			if (Result.bComplete)
+			{
+				// Readback done. Queue up data for debug (unless we've found pending readbacks, in which case we'll run again later).
+				if (ensure(Result.Data) && !bHasPendingReadbacks)
+				{
+					FPCGTaggedData Data = Output;
+					Data.Data = Result.Data;
+					DataToDebug.Add(Data);
+				}
+			}
+			else
+			{
+				bHasPendingReadbacks = true;
+			}
+		}
+		else
+		{
+			DataToDebug.Add(Output);
+		}
+	}
+
+	if (bHasPendingReadbacks)
+	{
+		// Signal that data was not ready, will run again later.
+		return false;
+	}
+
+	const FPCGDataVisualizationRegistry& DataVisRegistry = FPCGModule::GetConstPCGDataVisualizationRegistry();
+
+	for (const FPCGTaggedData& Data : DataToDebug)
+	{
+		if (const IPCGDataVisualization* DataVis = DataVisRegistry.GetDataVisualization(Data.Data->GetClass()))
 		{
-			DataVis->ExecuteDebugDisplay(Context, SettingsInterface, Input.Data, Context->GetTargetActor(nullptr));
+			DataVis->ExecuteDebugDisplay(Context, SettingsInterface, Data.Data, Context->GetTargetActor(nullptr));
 		}
 	}
 
-	Context->InputData = ElementInputs;
-	Context->OutputData = ElementOutputs;
+	return true;
 }
 #endif // WITH_EDITOR
 
diff --git a/Engine/Plugins/PCG/Source/PCG/Public/PCGElement.h b/Engine/Plugins/PCG/Source/PCG/Public/PCGElement.h
--- a/Engine/Plugins/PCG/Source/PCG/Public/PCGElement.h
+++ b/Engine/Plugins/PCG/Source/PCG/Public/PCGElement.h
@@ -141,8 +141,8 @@ public:
 	/** Is this element used for dispatching compute graphs. Used as rudimentary RTTI for element object comparisons (editor only). */
 	virtual bool IsComputeGraphElement() const { return false; }
 
-	/** Note: must be called from the main thread. */
-	void DebugDisplay(FPCGContext* Context) const;
+	/** Performs debug display. Returns true when complete. Note: must be called from the main thread. */
+	bool DebugDisplay(FPCGContext* Context) const;
 #endif
 
 protected:
