[PCG] Add TODO for investigating adding a task for doing debug display.

#rnx

diff --git a/Engine/Plugins/PCG/Source/PCG/Private/Graph/PCGGraphExecutor.cpp b/Engine/Plugins/PCG/Source/PCG/Private/Graph/PCGGraphExecutor.cpp
--- a/Engine/Plugins/PCG/Source/PCG/Private/Graph/PCGGraphExecutor.cpp
+++ b/Engine/Plugins/PCG/Source/PCG/Private/Graph/PCGGraphExecutor.cpp
@@ -988,6 +988,7 @@ void FPCGGraphExecutor::PostTaskExecute(TSharedPtr<FPCGGraphActiveTask> ActiveTa
 		if (ActiveTaskSettingsInterface && ActiveTaskSettingsInterface->bDebug)
 		{
 			// Debug display can span multiple frames so hold onto reference to task.
+			// TODO: Investigate scheduling an additional task to perform the debug display rather than manually holding onto references here.
 			{
 				PCGGraphExecutor::TScopeLock ScopeLock(CollectGCReferenceTasksLock);
 				CollectGCReferenceTasks.Add(ActiveTaskPtr);
