[PCG] UPCGManagedComponentBase::PostEditImport can now skip remapping with copy/paste fix (38477387)

#rnx

diff --git a/Engine/Plugins/PCG/Source/PCG/Private/PCGManagedResource.cpp b/Engine/Plugins/PCG/Source/PCG/Private/PCGManagedResource.cpp
--- a/Engine/Plugins/PCG/Source/PCG/Private/PCGManagedResource.cpp
+++ b/Engine/Plugins/PCG/Source/PCG/Private/PCGManagedResource.cpp
@@ -315,26 +315,20 @@ void UPCGManagedComponentBase::PostEditImport()
 		{
 			TSoftObjectPtr<UActorComponent> GeneratedComponent = GeneratedComponents[ComponentIndex];
 
-			UActorComponent* MatchingComponent = nullptr;
-
-			if (GeneratedComponent.IsValid())
+			// Do not erase if we haven't yet imported the properties (value is explicitely null)
+			if (GeneratedComponent.IsNull())
 			{
-				if (UActorComponent** FoundMatchingComponent = ActorComponents.FindByPredicate([&GeneratedComponent](UActorComponent* ActorComponent) { return ActorComponent && ActorComponent->GetFName() == GeneratedComponent->GetFName(); }))
-				{
-					MatchingComponent = *FoundMatchingComponent;
-				}
+				continue;
 			}
 
-			if (MatchingComponent)
+			// If Generated component is part of the current actor components, we keep it
+			if (UActorComponent* GeneratedComponentPtr = GeneratedComponent.Get(); GeneratedComponentPtr && ActorComponents.Contains(GeneratedComponentPtr))
 			{
-				GeneratedComponents[ComponentIndex] = MatchingComponent;
-			}
-			else
-			{
-				// Not quite clear what to do when we have a component that cannot be remapped.
-				// Maybe we should check against guids instead?
-				ForgetComponent(ComponentIndex);
+				continue;
 			}
+
+			// Forget components that are not owned by this actor
+			ForgetComponent(ComponentIndex);
 		}
 	}
 }
