[PCG] Properly rename out PCG Component sub-objects (unused Managed Resources and cleared Generated Output Data)
- One of the side-effects of not doing that was bloated export text (prior to a GC)
- Also remove no longer needed duplication code of generated output because property is now visible and taken care of by the base component instance data class

#rnx

diff --git a/Engine/Plugins/PCG/Source/PCG/Private/PCGComponent.cpp b/Engine/Plugins/PCG/Source/PCG/Private/PCGComponent.cpp
--- a/Engine/Plugins/PCG/Source/PCG/Private/PCGComponent.cpp
+++ b/Engine/Plugins/PCG/Source/PCG/Private/PCGComponent.cpp
@@ -632,7 +632,7 @@ void UPCGComponent::PostProcessGraph(const FBox& InNewBounds, bool bInGenerated,
 
 	CleanupUnusedManagedResources();
 
-	GeneratedGraphOutput.Reset();
+	ClearGraphGeneratedOutput();
 
 #if WITH_EDITOR
 	ResetIgnoredChangeOrigins(/*bLogIfAnyPresent=*/true);
@@ -849,7 +849,8 @@ void UPCGComponent::PostCleanupGraph()
 	CurrentCleanupTask = InvalidPCGTaskId;
 
 	const bool bHadGeneratedGraphOutput = GeneratedGraphOutput.TaggedData.Num() > 0;
-	GeneratedGraphOutput.Reset();
+
+	ClearGraphGeneratedOutput();
 
 #if WITH_EDITOR
 	DynamicallyTrackedKeysToSettings.Reset();
@@ -1378,7 +1379,7 @@ void UPCGComponent::CleanupLocalImmediate(bool bRemoveComponents, bool bCleanupL
 #if WITH_EDITOR
 					if (Resource && !Resource->IsMarkedTransientOnLoad())
 					{
-						Resource->Rename(nullptr, nullptr, REN_DoNotDirty | REN_DontCreateRedirectors | REN_NonTransactional);
+						Resource->Rename(nullptr, GetTransientPackage(), REN_DoNotDirty | REN_DontCreateRedirectors | REN_NonTransactional);
 					}
 #endif
 
@@ -1483,7 +1484,7 @@ FPCGTaskId UPCGComponent::CreateCleanupTask(bool bRemoveComponents, const TArray
 #if WITH_EDITOR
 					if (Resource && !Resource->IsMarkedTransientOnLoad())
 					{
-						Resource->Rename(nullptr, nullptr, REN_DoNotDirty | REN_DontCreateRedirectors | REN_NonTransactional);
+						Resource->Rename(nullptr, GetTransientPackage(), REN_DoNotDirty | REN_DontCreateRedirectors | REN_NonTransactional);
 					}
 #endif
 
@@ -1612,7 +1613,7 @@ void UPCGComponent::CleanupUnusedManagedResources()
 #if WITH_EDITOR
 				if (Resource && !Resource->IsMarkedTransientOnLoad())
 				{
-					Resource->Rename(nullptr, nullptr, REN_DoNotDirty | REN_DontCreateRedirectors | REN_NonTransactional);
+					Resource->Rename(nullptr, GetTransientPackage(), REN_DoNotDirty | REN_DontCreateRedirectors | REN_NonTransactional);
 				}
 #endif
 				GeneratedResources.RemoveAtSwap(ResourceIndex);
@@ -1625,6 +1626,19 @@ void UPCGComponent::CleanupUnusedManagedResources()
 	PCGGeneratedResourcesLogging::LogCleanupUnusedManagedResourcesFinished(this, GeneratedResources);
 }
 
+void UPCGComponent::ClearGraphGeneratedOutput()
+{
+	for (FPCGTaggedData& GeneratedData : GeneratedGraphOutput.TaggedData)
+	{
+		if (GeneratedData.Data && GeneratedData.Data->GetOuter() == this)
+		{
+			const_cast<UPCGData*>(GeneratedData.Data.Get())->Rename(nullptr, GetTransientPackage(), REN_DoNotDirty | REN_DontCreateRedirectors | REN_NonTransactional);
+		}
+	}
+
+	GeneratedGraphOutput.Reset();
+}
+
 void UPCGComponent::BeginPlay()
 {
 	Super::BeginPlay();
@@ -3833,16 +3847,7 @@ void FPCGComponentInstanceData::ApplyToComponent(UActorComponent* Component, con
 		// This means that some properties need to be reapplied here manually unless we make them visible
 		if (SourceComponent)
 		{
-			// Critical: GeneratedGraphOutput
-			PCGComponent->GeneratedGraphOutput = SourceComponent->GeneratedGraphOutput;
-			// Re-outer any data moved here
-			for (FPCGTaggedData& TaggedData : PCGComponent->GeneratedGraphOutput.TaggedData)
-			{
-				if (TaggedData.Data)
-				{
-					const_cast<UPCGData*>(TaggedData.Data.Get())->Rename(nullptr, PCGComponent, REN_DoNotDirty | REN_DontCreateRedirectors | REN_NonTransactional);
-				}
-			}
+			// No need to copy PCGComponent->GeneratedGraphOutput anymore as it is now a visible property
 
 #if WITH_EDITOR
 			// While this is serialized, it is not properly copied because it is not visible. This is needed otherwise a refresh can retoggle from not generated to generated
diff --git a/Engine/Plugins/PCG/Source/PCG/Public/PCGComponent.h b/Engine/Plugins/PCG/Source/PCG/Public/PCGComponent.h
--- a/Engine/Plugins/PCG/Source/PCG/Public/PCGComponent.h
+++ b/Engine/Plugins/PCG/Source/PCG/Public/PCGComponent.h
@@ -581,6 +581,7 @@ private:
 	void PostCleanupGraph();
 	void OnProcessGraphAborted(bool bQuiet = false, bool bCleanupUnusedResources = true);
 	void CleanupUnusedManagedResources();
+	void ClearGraphGeneratedOutput();
 	bool MoveResourcesToNewActor(AActor* InNewActor, bool bCreateChild);
 
 	/** Called as part of the RerunConstructionScript, just takes resources as-is, assumes current state is empty + the resources have been retargeted if needed */
