[PCG] Add node to generate landscape grass maps.

Uses generation on the fly rather than serialized grass map data currently.  Runtime generation is not yet supported on mobile.

Add support for pseudo-bindless to the texture data interface so that multiple textures can be sampled from a custom HLSL node.

#rnx

diff --git a/Engine/Plugins/PCG/Shaders/Private/PCGGrassMapUnpackerCS.usf b/Engine/Plugins/PCG/Shaders/Private/PCGGrassMapUnpackerCS.usf
new file mode 100644
--- /dev/null
+++ b/Engine/Plugins/PCG/Shaders/Private/PCGGrassMapUnpackerCS.usf
@@ -0,0 +1,75 @@
+// Copyright Epic Games, Inc. All Rights Reserved.
+
+#include "/Engine/Private/Common.ush"
+#include "/Engine/Public/Platform.ush"
+
+// Must match FPCGGrassMapUnpackerCS::MaxNumLandscapeComponents
+#define MAX_NUM_LANDSCAPE_COMPONENTS 64
+
+// Input:
+// First 2 channels of first texture always reserved for heightmap (see constructor of FLandscapeGrassWeightExporter_RenderThread::FComponentInfo).
+// So first two grass maps will be in z & w channels of first set of pass, otherwise packed in remaining channels.
+// Generated grass maps all end up in one big row (see FLandscapeGrassWeightExporter::FLandscapeGrassWeightExporter()).
+//
+//  -----------------------------------------
+// | LC0  | LC0  | LC1  | LC1  | LC2  | LC2  |
+// | x=h0 | x=g2 | x=h0 | x=g2 | x=h0 | x=g2 |  NumComponents=3, NumGrassMaps=3, InNumGrassMapPasses=2
+// | y=h1 | y=   | y=h1 | y=   | y=h1 | y=   |
+// | z=g0 | z=   | z=g0 | z=   | z=g0 | z=   |
+// | w=g1 | w=   | w=g1 | w=   | w=g1 | w=   |
+//  -----------------------------------------
+Texture2D<float4> InPackedGrassMaps;
+
+// Output:
+// Each slice is a single grass combined grass map. Dimension is (NumComponentsX * ResX, NumComponentsY * ResY, NumGrassMaps).
+//
+//                -------------
+//              -------------  |
+//             |             | |
+// Num Comps Y |             | |
+//   * ResY    |             | |
+//             |             | |
+//             |             |-  NumGrassMaps
+//              -------------
+//               Num Comps X
+//                 * ResX
+RWTexture2DArray<float> OutUnpackedGrassMaps;
+
+int4 InLinearTileIndexToComponentIndex[MAX_NUM_LANDSCAPE_COMPONENTS]; // 16 byte alignment on arrays
+uint InNumTilesX;
+uint InLandscapeComponentResolution;
+uint InNumGrassMapPasses;
+uint2 InOutputResolution;
+
+// Dispatch group count is (NumComponentsX * ResX / 8, NumComponentsY * ResY / 8, NumGrassMaps).
+[numthreads(THREADGROUPSIZE_X, THREADGROUPSIZE_Y, THREADGROUPSIZE_Z)]
+void PCGGrassMapUnpacker_CS(uint3 DispatchThreadId : SV_DispatchThreadID)
+{
+	if (any(DispatchThreadId.xy >= InOutputResolution))
+	{
+		return;
+	}
+
+	// The rendered landscape components in the input are not in any particular order, so lookup which component index for this output tile.
+	const uint2 TileCoords = DispatchThreadId.xy / InLandscapeComponentResolution;
+	const uint TileLinearIndex = TileCoords.y * InNumTilesX + TileCoords.x;
+	const int ComponentIndex = InLinearTileIndexToComponentIndex[TileLinearIndex][0];
+
+	// Invalid component index means there was no component present for this output tile.
+	if (ComponentIndex == -1)
+	{
+		OutUnpackedGrassMaps[DispatchThreadId] = 0.0f;
+		return;
+	}
+
+	// Coordinates into input grass maps (see packing above).
+	const uint GrassMapIndex = DispatchThreadId.z;
+	const uint PassIndex = (GrassMapIndex + 2) / 4;
+	const uint ComponentXOffset = (ComponentIndex * InNumGrassMapPasses + PassIndex) * InLandscapeComponentResolution;
+	const uint2 InputCoords = uint2(DispatchThreadId.x % InLandscapeComponentResolution + ComponentXOffset, DispatchThreadId.y % InLandscapeComponentResolution);
+
+	// Source channel (see packing above).
+	const uint InputChannelIndex = (GrassMapIndex + 2) % 4;
+
+	OutUnpackedGrassMaps[DispatchThreadId] = InPackedGrassMaps[InputCoords][InputChannelIndex];
+}
diff --git a/Engine/Plugins/PCG/Shaders/Private/PCGTextureDataInterface.ush b/Engine/Plugins/PCG/Shaders/Private/PCGTextureDataInterface.ush
--- a/Engine/Plugins/PCG/Shaders/Private/PCGTextureDataInterface.ush
+++ b/Engine/Plugins/PCG/Shaders/Private/PCGTextureDataInterface.ush
@@ -1,7 +1,29 @@
 // Copyright Epic Games, Inc. All Rights Reserved.
 
-Texture2D {DataInterfaceName}_Texture;
-SamplerState {DataInterfaceName}_Sampler;
+// We can't use full bindless on all platforms, so fallback to emulating. Supports maximum 8 texture data objects.
+#define MAX_NUM_TEXTURES 8
+
+// TODO: Define out multiple bindings if we know statically there is only one.
+Texture2D {DataInterfaceName}_Texture_0;
+Texture2D {DataInterfaceName}_Texture_1;
+Texture2D {DataInterfaceName}_Texture_2;
+Texture2D {DataInterfaceName}_Texture_3;
+Texture2D {DataInterfaceName}_Texture_4;
+Texture2D {DataInterfaceName}_Texture_5;
+Texture2D {DataInterfaceName}_Texture_6;
+Texture2D {DataInterfaceName}_Texture_7;
+
+// TODO: We only need 2 samplers + an indirection array, de-duplicate these.
+SamplerState {DataInterfaceName}_Sampler_0;
+SamplerState {DataInterfaceName}_Sampler_1;
+SamplerState {DataInterfaceName}_Sampler_2;
+SamplerState {DataInterfaceName}_Sampler_3;
+SamplerState {DataInterfaceName}_Sampler_4;
+SamplerState {DataInterfaceName}_Sampler_5;
+SamplerState {DataInterfaceName}_Sampler_6;
+SamplerState {DataInterfaceName}_Sampler_7;
+
+float4 {DataInterfaceName}_TextureBounds[MAX_NUM_TEXTURES];
 
 float2 GetTexCoords_{DataInterfaceName}(float2 WorldPos, float2 Min, float2 Max)
 {
@@ -17,7 +39,61 @@ float2 GetTexCoords_{DataInterfaceName}(float2 WorldPos, float2 Min, float2 Max)
 	return Pos / Extents;
 }
 
+#define PCG_CONDITIONAL_TEXTURE_SAMPLE(Index) \
+	if (InDataIndex == Index) \
+	{ \
+		return {DataInterfaceName}_Texture_##Index.Sample({DataInterfaceName}_Sampler_##Index, InTextureUV); \
+	}
+
+float4 Sample_{DataInterfaceName}(int InDataIndex, float2 InTextureUV)
+{
+	PCG_CONDITIONAL_TEXTURE_SAMPLE(0)
+	PCG_CONDITIONAL_TEXTURE_SAMPLE(1)
+	PCG_CONDITIONAL_TEXTURE_SAMPLE(2)
+	PCG_CONDITIONAL_TEXTURE_SAMPLE(3)
+	PCG_CONDITIONAL_TEXTURE_SAMPLE(4)
+	PCG_CONDITIONAL_TEXTURE_SAMPLE(5)
+	PCG_CONDITIONAL_TEXTURE_SAMPLE(6)
+	PCG_CONDITIONAL_TEXTURE_SAMPLE(7)
+
+	// TODO: Runtime diagnostic error.
+	return (float4)0;
+}
+
+#undef PCG_CONDITIONAL_TEXTURE_SAMPLE
+
+#define PCG_CONDITIONAL_TEXTURE_SAMPLE_WORLD_POS(Index) \
+	if (InDataIndex == Index) \
+	{ \
+		const float2 TextureUV = GetTexCoords_{DataInterfaceName}(WorldPos, {DataInterfaceName}_TextureBounds[Index].xy, {DataInterfaceName}_TextureBounds[Index].zw); \
+		return {DataInterfaceName}_Texture_##Index.Sample({DataInterfaceName}_Sampler_##Index, TextureUV); \
+	}
+
+float4 SampleWorldPos_{DataInterfaceName}(int InDataIndex, float2 WorldPos)
+{
+	PCG_CONDITIONAL_TEXTURE_SAMPLE_WORLD_POS(0)
+	PCG_CONDITIONAL_TEXTURE_SAMPLE_WORLD_POS(1)
+	PCG_CONDITIONAL_TEXTURE_SAMPLE_WORLD_POS(2)
+	PCG_CONDITIONAL_TEXTURE_SAMPLE_WORLD_POS(3)
+	PCG_CONDITIONAL_TEXTURE_SAMPLE_WORLD_POS(4)
+	PCG_CONDITIONAL_TEXTURE_SAMPLE_WORLD_POS(5)
+	PCG_CONDITIONAL_TEXTURE_SAMPLE_WORLD_POS(6)
+	PCG_CONDITIONAL_TEXTURE_SAMPLE_WORLD_POS(7)
+
+	// TODO: Runtime diagnostic error.
+	return (float4)0;
+}
+
+#undef PCG_CONDITIONAL_TEXTURE_SAMPLE_WORLD_POS
+
+// Deprecated (5.6). Fallback to using first data.
 float4 Sample_{DataInterfaceName}(float2 TexCoords)
 {
-	return {DataInterfaceName}_Texture.SampleLevel({DataInterfaceName}_Sampler, TexCoords, 0);
+	return Sample_{DataInterfaceName}(/*InDataIndex=*/0, TexCoords);
+}
+
+// Deprecated (5.6). Fallback to using first data.
+float4 SampleWorldPos_{DataInterfaceName}(float2 WorldPos)
+{
+	return SampleWorldPos_{DataInterfaceName}(/*InDataIndex=*/0, WorldPos);
 }
diff --git a/Engine/Plugins/PCG/Source/PCG/Private/Compute/DataInterfaces/PCGTextureDataInterface.cpp b/Engine/Plugins/PCG/Source/PCG/Private/Compute/DataInterfaces/PCGTextureDataInterface.cpp
--- a/Engine/Plugins/PCG/Source/PCG/Private/Compute/DataInterfaces/PCGTextureDataInterface.cpp
+++ b/Engine/Plugins/PCG/Source/PCG/Private/Compute/DataInterfaces/PCGTextureDataInterface.cpp
@@ -9,18 +9,22 @@
 #include "GlobalRenderResources.h"
 #include "RHIResources.h"
 #include "RHIStaticStates.h"
+#include "RenderGraphBuilder.h"
 #include "ShaderCompilerCore.h"
 #include "ShaderCore.h"
 #include "ShaderParameterMetadataBuilder.h"
+#include "SystemTextures.h"
 #include "TextureResource.h"
 #include "ComputeFramework/ShaderParamTypeDefinition.h"
 #include "Engine/Texture.h"
 
 #include UE_INLINE_GENERATED_CPP_BY_NAME(PCGTextureDataInterface)
 
+#define LOCTEXT_NAMESPACE "PCGTextureDataInterface"
+
 void UPCGTextureDataInterface::GetSupportedInputs(TArray<FShaderFunctionDefinition>& OutFunctions) const
 {
-	// TODO: Support texture2d array, texture3d. UNiagaraDataInterfaceRenderTarget2D and related are useful references.
+	// TODO: Support texture3d, etc. UNiagaraDataInterfaceRenderTarget2D and related are useful references.
 
 	OutFunctions.AddDefaulted_GetRef()
 		.SetName(TEXT("GetTexCoords"))
@@ -29,15 +33,33 @@ void UPCGTextureDataInterface::GetSupportedInputs(TArray<FShaderFunctionDefiniti
 		.AddParam(EShaderFundamentalType::Float, 2)
 		.AddParam(EShaderFundamentalType::Float, 2);
 
+	OutFunctions.AddDefaulted_GetRef()
+		.SetName(TEXT("Sample"))
+		.AddReturnType(EShaderFundamentalType::Float, 4)
+		.AddParam(EShaderFundamentalType::Uint)
+		.AddParam(EShaderFundamentalType::Float, 2);
+
 	OutFunctions.AddDefaulted_GetRef()
 		.SetName(TEXT("Sample"))
 		.AddReturnType(EShaderFundamentalType::Float, 4)
 		.AddParam(EShaderFundamentalType::Float, 2);
+
+	OutFunctions.AddDefaulted_GetRef()
+		.SetName(TEXT("SampleWorldPos"))
+		.AddReturnType(EShaderFundamentalType::Float, 4)
+		.AddParam(EShaderFundamentalType::Uint)
+		.AddParam(EShaderFundamentalType::Float, 2);
+
+	OutFunctions.AddDefaulted_GetRef()
+		.SetName(TEXT("SampleWorldPos"))
+		.AddReturnType(EShaderFundamentalType::Float, 4)
+		.AddParam(EShaderFundamentalType::Float, 2);
 }
 
 BEGIN_SHADER_PARAMETER_STRUCT(FPCGTextureDataInterfaceParameters,)
-	SHADER_PARAMETER_TEXTURE(Texture2D, Texture)
-	SHADER_PARAMETER_SAMPLER(SamplerState, Sampler)
+	SHADER_PARAMETER_RDG_TEXTURE_SRV_ARRAY(Texture2D, Texture, [PCGTextureDataInterfaceConstants::MAX_NUM_TEXTURES]) // TODO: Pick a max value and document somewhere + do validation
+	SHADER_PARAMETER_SAMPLER_ARRAY(SamplerState, Sampler, [PCGTextureDataInterfaceConstants::MAX_NUM_TEXTURES])
+	SHADER_PARAMETER_ARRAY(FVector4f, TextureBounds, [PCGTextureDataInterfaceConstants::MAX_NUM_TEXTURES])
 END_SHADER_PARAMETER_STRUCT()
 
 void UPCGTextureDataInterface::GetShaderParameters(TCHAR const* UID, FShaderParametersMetadataBuilder& InOutBuilder, FShaderParametersMetadataAllocations& InOutAllocations) const
@@ -80,47 +102,57 @@ UComputeDataProvider* UPCGTextureDataInterface::CreateDataProvider(TObjectPtr<UO
 
 	// Take any input pin label alias to obtain the data from the input data collection.
 	check(!DownstreamInputPinLabelAliases.IsEmpty());
-	const TArray<FPCGTaggedData> TaggedDatas = Binding->InputDataCollection.GetInputsByPin(DownstreamInputPinLabelAliases[0]);
+	const TArray<FPCGTaggedData> InputTaggedData = Binding->InputDataCollection.GetInputsByPin(DownstreamInputPinLabelAliases[0]);
+	
+	if (InputTaggedData.Num() >= PCGTextureDataInterfaceConstants::MAX_NUM_TEXTURES)
+	{
+		UE_LOG(LogPCG, Warning, TEXT("Texture data interface on pin '%s' received too many data items (%d), only first %d valid textures will be used."),
+			*DownstreamInputPinLabelAliases[0].ToString(),
+			InputTaggedData.Num(),
+			PCGTextureDataInterfaceConstants::MAX_NUM_TEXTURES);
+	}
+
+	const int NumInputTextures = FMath::Min(InputTaggedData.Num(), (int)PCGTextureDataInterfaceConstants::MAX_NUM_TEXTURES);
+	TArray<FPCGInputTextureInfo> TextureInfos;
+	TextureInfos.SetNum(NumInputTextures);
 
-	if (!TaggedDatas.IsEmpty() && ensure(TaggedDatas[0].Data))
+	for (int InputTextureIndex = 0; InputTextureIndex < NumInputTextures; ++InputTextureIndex)
 	{
-		// There should only be one texture data
-		// TODO: Bindless resources to handle multiple textures
-		ensure(TaggedDatas.Num() == 1); 
+		const UPCGTextureData* TextureData = Cast<UPCGTextureData>(InputTaggedData[InputTextureIndex].Data);
+		if (!ensure(TextureData))
+		{
+			UE_LOG(LogPCG, Error, TEXT("Unsupported data type encountered by texture data interface: '%s'"),
+				InputTaggedData[InputTextureIndex].Data ? *InputTaggedData[InputTextureIndex].Data->GetName() : TEXT("NULL"));
+			continue;
+		}
 
-		if (const UPCGBaseTextureData* BaseTextureData = Cast<UPCGBaseTextureData>(TaggedDatas[0].Data))
+		FPCGInputTextureInfo& TextureInfo = TextureInfos[InputTextureIndex];
+
+		TextureInfo.Texture = TextureData->GetTextureRHI();
+		TextureInfo.ResourceType = TextureData->GetTextureResourceType();
+		TextureInfo.Filter = TextureData->Filter;
+		TextureInfo.TextureBounds = TextureData->GetBounds();
+		TextureInfo.SliceIndex = TextureData->GetTextureSlice();
+
+		if (TextureInfo.ResourceType == EPCGTextureResourceType::ExportedTexture)
 		{
-			DataProvider->TextureRHI = BaseTextureData->GetTextureRHI();
-			DataProvider->Filter = BaseTextureData->Filter;
+			TextureInfo.ExportedTexture = TextureData->GetRefCountedTexture();
 		}
-		else
+
+		if (TextureInfo.Texture)
 		{
-			ensureMsgf(false, TEXT("Unsupported data type encountered by texture data interface: '%s'"), *TaggedDatas[0].Data->GetName());
+			TextureInfo.bIsTextureArray = TextureInfo.Texture->GetDesc().Dimension == ETextureDimension::Texture2DArray;
 		}
 	}
 
-	return DataProvider;
-}
+	DataProvider->TextureInfos = MoveTemp(TextureInfos);
 
-FPCGTextureDataProviderProxy::FPCGTextureDataProviderProxy(FTextureRHIRef InTextureRHI, EPCGTextureFilter InFilter)
-	: TextureRHI(InTextureRHI)
-{
-	switch (InFilter)
-	{
-	case EPCGTextureFilter::Point:
-		SamplerRHI = TStaticSamplerState<SF_Point, AM_Clamp, AM_Clamp, AM_Clamp>::GetRHI();
-		break;
-	default:
-		ensureMsgf(false, TEXT("Unrecognized texture filter setting, falling back to bilinear."));
-	case EPCGTextureFilter::Bilinear: // Falls through
-		SamplerRHI = TStaticSamplerState<SF_Bilinear, AM_Clamp, AM_Clamp, AM_Clamp>::GetRHI();
-		break;
-	}
+	return DataProvider;
 }
 
 FComputeDataProviderRenderProxy* UPCGTextureDataProvider::GetRenderProxy()
 {
-	return new FPCGTextureDataProviderProxy(TextureRHI, Filter);
+	return new FPCGTextureDataProviderProxy(TextureInfos);
 }
 
 bool FPCGTextureDataProviderProxy::IsValid(FValidationData const& InValidationData) const
@@ -139,7 +171,95 @@ void FPCGTextureDataProviderProxy::GatherDispatchData(FDispatchData const& InDis
 	for (int32 InvocationIndex = 0; InvocationIndex < ParameterArray.Num(); ++InvocationIndex)
 	{
 		FParameters& Parameters = ParameterArray[InvocationIndex];
-		Parameters.Texture = TextureRHI ? (FRHITexture*)TextureRHI : (FRHITexture*)GBlackTexture->TextureRHI;
-		Parameters.Sampler = SamplerRHI;
+
+		for (int BindingIndex = 0; BindingIndex < PCGTextureDataInterfaceConstants::MAX_NUM_TEXTURES; ++BindingIndex)
+		{
+			Parameters.Texture[BindingIndex] = TextureSRV[BindingIndex];
+
+			if (BindingIndex < TextureInfos.Num())
+			{
+				Parameters.TextureBounds[BindingIndex] = FVector4f(
+					TextureInfos[BindingIndex].TextureBounds.Min.X,
+					TextureInfos[BindingIndex].TextureBounds.Min.Y,
+					TextureInfos[BindingIndex].TextureBounds.Max.X,
+					TextureInfos[BindingIndex].TextureBounds.Max.Y);
+
+				switch (TextureInfos[BindingIndex].Filter)
+				{
+				case EPCGTextureFilter::Point:
+					Parameters.Sampler[BindingIndex] = TStaticSamplerState<SF_Point, AM_Clamp, AM_Clamp, AM_Clamp>::GetRHI();
+					break;
+				default:
+					ensureMsgf(false, TEXT("Unrecognized texture filter setting, falling back to bilinear."));
+				case EPCGTextureFilter::Bilinear: // Falls through
+					Parameters.Sampler[BindingIndex] = TStaticSamplerState<SF_Bilinear, AM_Clamp, AM_Clamp, AM_Clamp>::GetRHI();
+					break;
+				}
+			}
+			else
+			{
+				Parameters.TextureBounds[BindingIndex] = FVector4f::Zero();
+				Parameters.Sampler[BindingIndex] = TStaticSamplerState<SF_Bilinear, AM_Clamp, AM_Clamp, AM_Clamp>::GetRHI();
+			}
+		}
 	}
 }
+
+void FPCGTextureDataProviderProxy::AllocateResources(FRDGBuilder& GraphBuilder, FAllocationData const& InAllocationData)
+{
+	// TODO: Move this to somewhere more generic where it runs once, rather than for every PCGTextureDataInterface?
+	// Copied from GSystemTextures.GetBlackDummy(GraphBuilder));
+	const FRHITextureCreateDesc DummyTextureDesc =
+		FRHITextureCreateDesc::Create2D(TEXT("PCGDummyTexture"), 1, 1, PF_G8)
+		.SetFlags(ETextureCreateFlags::ShaderResource);
+
+	FTextureRHIRef Texture = RHICreateTexture(DummyTextureDesc);
+
+	{
+		uint32 DestStride;
+		float* DestBuffer = (float*)RHILockTexture2D(Texture, 0, RLM_WriteOnly, DestStride, false);
+		*DestBuffer = 0;
+		RHIUnlockTexture2D(Texture, 0, false);
+	}
+
+	// Create a dummy texture to bind.
+	TRefCountPtr<IPooledRenderTarget> BlackDummyTarget = CreateRenderTarget(Texture, TEXT("PCGDummyRenderTarget"));
+	check(BlackDummyTarget.IsValid());
+
+	FRDGTextureRef DummyTextureRDG = GraphBuilder.RegisterExternalTexture(BlackDummyTarget);
+
+	for (int DataIndex = 0; DataIndex < PCGTextureDataInterfaceConstants::MAX_NUM_TEXTURES; ++DataIndex)
+	{
+		FRDGTextureRef TextureRDG = nullptr;
+
+		if (DataIndex < TextureInfos.Num())
+		{
+			if (TextureInfos[DataIndex].ResourceType == EPCGTextureResourceType::TextureObject && TextureInfos[DataIndex].Texture)
+			{
+				TextureRDG = GraphBuilder.RegisterExternalTexture(
+					CreateRenderTarget(TextureInfos[DataIndex].Texture, *FText::Format(LOCTEXT("RenderTargetDebugName", "PCGDataInterfaceRenderTarget_{0}"), DataIndex).ToString()));
+			}
+			else if (TextureInfos[DataIndex].ResourceType == EPCGTextureResourceType::ExportedTexture && TextureInfos[DataIndex].ExportedTexture)
+			{
+				TextureRDG = GraphBuilder.RegisterExternalTexture(TextureInfos[DataIndex].ExportedTexture);
+			}
+		}
+
+		if (!TextureRDG)
+		{
+			// Fallback to blank dummy texture if no texture was provided.
+			TextureRDG = DummyTextureRDG;
+		}
+
+		if (DataIndex < TextureInfos.Num() && TextureInfos[DataIndex].bIsTextureArray && TextureRDG != DummyTextureRDG)
+		{
+			TextureSRV[DataIndex] = GraphBuilder.CreateSRV(FRDGTextureSRVDesc::CreateForSlice(TextureRDG, TextureInfos[DataIndex].SliceIndex));
+		}
+		else
+		{
+			TextureSRV[DataIndex] = GraphBuilder.CreateSRV(TextureRDG);
+		}
+	}
+}
+
+#undef LOCTEXT_NAMESPACE
diff --git a/Engine/Plugins/PCG/Source/PCG/Private/Compute/DataInterfaces/PCGTextureDataInterface.h b/Engine/Plugins/PCG/Source/PCG/Private/Compute/DataInterfaces/PCGTextureDataInterface.h
--- a/Engine/Plugins/PCG/Source/PCG/Private/Compute/DataInterfaces/PCGTextureDataInterface.h
+++ b/Engine/Plugins/PCG/Source/PCG/Private/Compute/DataInterfaces/PCGTextureDataInterface.h
@@ -2,17 +2,37 @@
 
 #pragma once
 
+#include "ComputeFramework/ComputeDataProvider.h"
 #include "Compute/DataInterfaces/PCGComputeDataInterface.h"
 
 #include "Data/PCGTextureData.h"
 
 #include "RHIResources.h"
-#include "ComputeFramework/ComputeDataProvider.h"
 
 #include "PCGTextureDataInterface.generated.h"
 
 class FPCGTextureDataInterfaceParameters;
 
+namespace PCGTextureDataInterfaceConstants
+{
+	/**
+	 * The maximum number of textures which can be bound to a PCGTextureDataInterface.
+	 * If you change this number, make sure update PCGTextureDataInterface.ush as well.
+	 */
+	constexpr uint32 MAX_NUM_TEXTURES = 8;
+}
+
+struct FPCGInputTextureInfo
+{
+	EPCGTextureResourceType ResourceType = EPCGTextureResourceType::TextureObject;
+	FTextureRHIRef Texture = nullptr;
+	TRefCountPtr<IPooledRenderTarget> ExportedTexture = nullptr;
+	EPCGTextureFilter Filter = EPCGTextureFilter::Bilinear;
+	FBox TextureBounds = FBox(EForceInit::ForceInit);
+	int SliceIndex = 0;
+	bool bIsTextureArray = false;
+};
+
 /** Data Interface allowing sampling of a texture. */
 UCLASS(ClassGroup = (Procedural))
 class UPCGTextureDataInterface : public UPCGComputeDataInterface
@@ -45,25 +65,26 @@ public:
 	FComputeDataProviderRenderProxy* GetRenderProxy() override;
 	//~ End UComputeDataProvider Interface
 
-	FTextureRHIRef TextureRHI = nullptr;
-
-	EPCGTextureFilter Filter = EPCGTextureFilter::Bilinear;
+	TArray<FPCGInputTextureInfo> TextureInfos;
 };
 
 class FPCGTextureDataProviderProxy : public FComputeDataProviderRenderProxy
 {
 public:
-	FPCGTextureDataProviderProxy(FTextureRHIRef InTextureRHI, EPCGTextureFilter InFilter);
+	FPCGTextureDataProviderProxy(const TArray<FPCGInputTextureInfo>& InTextureInfos)
+		: TextureInfos(InTextureInfos)
+	{}
 
 	//~ Begin FComputeDataProviderRenderProxy Interface
 	bool IsValid(FValidationData const& InValidationData) const override;
 	void GatherDispatchData(FDispatchData const& InDispatchData) override;
+	void AllocateResources(FRDGBuilder& GraphBuilder, FAllocationData const& InAllocationData) override;
 	//~ End FComputeDataProviderRenderProxy Interface
 
 protected:
 	using FParameters = FPCGTextureDataInterfaceParameters;
 
-	FTextureRHIRef TextureRHI = nullptr;
+	TArray<FPCGInputTextureInfo> TextureInfos;
 
-	FSamplerStateRHIRef SamplerRHI = nullptr;
+	FRDGTextureSRVRef TextureSRV[PCGTextureDataInterfaceConstants::MAX_NUM_TEXTURES];
 };
diff --git a/Engine/Plugins/PCG/Source/PCG/Private/Compute/Elements/PCGCustomHLSL.cpp b/Engine/Plugins/PCG/Source/PCG/Private/Compute/Elements/PCGCustomHLSL.cpp
--- a/Engine/Plugins/PCG/Source/PCG/Private/Compute/Elements/PCGCustomHLSL.cpp
+++ b/Engine/Plugins/PCG/Source/PCG/Private/Compute/Elements/PCGCustomHLSL.cpp
@@ -828,7 +828,9 @@ void UPCGCustomHLSLSettings::UpdateInputDeclarations()
 
 		InputDeclarations += FString::Format(TEXT(
 			"float2 {0}_GetTexCoords(float2 WorldPos, float2 Min, float2 Max);\n"
-			"float4 {0}_Sample(float2 TexCoords);\n"),
+			"float4 {0}_Sample(uint DataIndex, float2 TexCoords);\n"
+			"// Computes sample coordinates of the WorldPos relative to the texture data's bounds."
+			"float4 {0}_SampleWorldPos(uint DataIndex, float2 WorldPos);\n"),
 			{ bMultiPin ? PCGHLSLElement::PinDeclTemplateStr : TextureDataPins[0] });
 
 		InputDeclarations += TEXT("\n");
@@ -1045,7 +1047,6 @@ void UPCGCustomHLSLSettings::UpdatePinSettings()
 		}
 
 		if (!!(Properties.AllowedTypes & EPCGDataType::Landscape)
-			|| !!(Properties.AllowedTypes & EPCGDataType::Texture)
 			|| !!(Properties.AllowedTypes & EPCGDataType::Param))
 		{
 			// Don't allow multiple data on this pin because we do not support a dynamic number of textures/landscapes bound to a
diff --git a/Engine/Plugins/PCG/Source/PCG/Private/Data/PCGTextureData.cpp b/Engine/Plugins/PCG/Source/PCG/Private/Data/PCGTextureData.cpp
--- a/Engine/Plugins/PCG/Source/PCG/Private/Data/PCGTextureData.cpp
+++ b/Engine/Plugins/PCG/Source/PCG/Private/Data/PCGTextureData.cpp
@@ -398,6 +398,7 @@ void UPCGTextureData::InitializeInternal(UTexture* InTexture, uint32 InTextureIn
 		return;
 	}
 
+	ResourceType = EPCGTextureResourceType::TextureObject;
 	Texture = InTexture;
 	TextureIndex = InTextureIndex;
 	Transform = InTransform;
@@ -499,6 +500,8 @@ void UPCGTextureData::InitializeInternal(UTexture* InTexture, uint32 InTextureIn
 
 bool UPCGTextureData::Initialize(UTexture* InTexture, uint32 InTextureIndex, const FTransform& InTransform, bool bCreateCPUDuplicateEditorOnly)
 {
+	ResourceType = EPCGTextureResourceType::TextureObject;
+
 	bool bInitializeDone = false;
 	InitializeInternal(InTexture, InTextureIndex, InTransform, &bInitializeDone, bCreateCPUDuplicateEditorOnly);
 	return bInitializeDone;
@@ -506,6 +509,8 @@ bool UPCGTextureData::Initialize(UTexture* InTexture, uint32 InTextureIndex, con
 
 void UPCGTextureData::Initialize(UTexture* InTexture, uint32 InTextureIndex, const FTransform& InTransform, const TFunction<void()>& InPostInitializeCallback, bool bCreateCPUDuplicateEditorOnly)
 {
+	ResourceType = EPCGTextureResourceType::TextureObject;
+
 	bool bInitializeDone = false;
 	InitializeInternal(InTexture, InTextureIndex, InTransform, &bInitializeDone, bCreateCPUDuplicateEditorOnly);
 
@@ -519,6 +524,23 @@ void UPCGTextureData::Initialize(UTexture* InTexture, uint32 InTextureIndex, con
 	}
 }
 
+bool UPCGTextureData::Initialize(TRefCountPtr<IPooledRenderTarget> InTextureHandle, uint32 InTextureIndex, const FTransform& InTransform)
+{
+	ResourceType = EPCGTextureResourceType::ExportedTexture;
+	TextureHandle = InTextureHandle;
+	TextureIndex = InTextureIndex;
+	Transform = InTransform;
+
+	Bounds = FBox(EForceInit::ForceInit);
+	Bounds += FVector(-1.0f, -1.0f, 0.0f);
+	Bounds += FVector(1.0f, 1.0f, 0.0f);
+	Bounds = Bounds.TransformBy(Transform);
+
+	InitializeFromGPUTexture();
+
+	return bSuccessfullyInitialized;
+}
+
 void UPCGTextureData::AddToCrc(FArchiveCrc32& Ar, bool bFullDataCrc) const
 {
 	Super::AddToCrc(Ar, bFullDataCrc);
@@ -529,8 +551,15 @@ void UPCGTextureData::AddToCrc(FArchiveCrc32& Ar, bool bFullDataCrc) const
 
 FTextureRHIRef UPCGTextureData::GetTextureRHI() const
 {
-	FTextureResource* Resource = Texture.IsValid() ? Texture->GetResource() : nullptr;
-	return Resource ? Resource->GetTextureRHI() : nullptr;
+	if (ResourceType == EPCGTextureResourceType::TextureObject)
+	{
+		FTextureResource* Resource = Texture.IsValid() ? Texture->GetResource() : nullptr;
+		return Resource ? Resource->GetTextureRHI() : nullptr;
+	}
+	else
+	{
+		return TextureHandle ? TextureHandle->GetRHI() : nullptr;
+	}
 }
 
 UPCGSpatialData* UPCGTextureData::CopyInternal(FPCGContext* Context) const
@@ -539,7 +568,9 @@ UPCGSpatialData* UPCGTextureData::CopyInternal(FPCGContext* Context) const
 
 	CopyBaseTextureData(NewTextureData);
 
+	NewTextureData->ResourceType = ResourceType;
 	NewTextureData->Texture = Texture;
+	NewTextureData->TextureHandle = TextureHandle;
 
 	return NewTextureData;
 }
@@ -679,52 +710,45 @@ TOptional<bool> UPCGTextureData::InitializeFromCPUTexture()
 
 bool UPCGTextureData::InitializeFromGPUTexture()
 {
-	if (!Texture.IsValid())
-	{
-		return true;
-	}
-
 	if (bReadbackFromGPUInitiated)
 	{
 		return false;
 	}
 	bReadbackFromGPUInitiated = true;
 
+	if (ResourceType == EPCGTextureResourceType::TextureObject)
+	{
+		if (!Texture.IsValid())
+		{
+			return true;
+		}
 #if WITH_EDITOR
-	// Force a wait on any incomplete async texture compilation and caching operations
-	Texture->FinishCachePlatformData();
+		// Force a wait on any incomplete async texture compilation and caching operations
+		Texture->FinishCachePlatformData();
 #endif
 
-	// TODO: Forcing a resource update is problematic. Clears cooked platform data and may break other usages of this texture
-	// But it is necessary to force the texture memory to be present on the first dispatch
-	Texture->UpdateResource();
-	Texture->WaitForPendingInitOrStreaming();
-
-	FTexturePlatformData* PlatformData = nullptr;
-	FTextureResource* TextureResource = Texture->GetResource();
-
-	if (UTexture2D* Texture2D = Cast<UTexture2D>(Texture))
-	{
-		PlatformData = Texture2D->GetPlatformData();
-	}
-	else if (UTexture2DArray* Texture2DArray = Cast<UTexture2DArray>(Texture))
-	{
-		PlatformData = Texture2DArray->GetPlatformData();
+		// TODO: Forcing a resource update is problematic. Clears cooked platform data and may break other usages of this texture
+		// But it is necessary to force the texture memory to be present on the first dispatch
+		Texture->UpdateResource();
+		Texture->WaitForPendingInitOrStreaming();
 	}
 
-	if (PlatformData && TextureResource && TextureResource->TextureRHI)
+	if (FTextureRHIRef TextureRHI = GetTextureRHI())
 	{
 		FPCGTextureReadbackDispatchParams Params;
-		Params.SourceTexture = TextureResource->TextureRHI;
+		Params.SourceTexture = TextureRHI;
 
 		// We should always use a point filter sampler since we are trying to get a 1 to 1 copy of the texture. We will do our own filtering later.
 		Params.SourceSampler = TStaticSamplerState<SF_Point, AM_Clamp, AM_Clamp, AM_Clamp>::GetRHI();
-		Params.SourceDimensions = FIntPoint(PlatformData->SizeX, PlatformData->SizeY);
 		Params.SourceTextureIndex = TextureIndex;
 
+		const FIntVector TextureSize = TextureRHI->GetDesc().GetSize();
+		Params.SourceDimensions = FIntPoint(TextureSize.X, TextureSize.Y);
+
 		TWeakObjectPtr<UPCGTextureData> ThisWeakPtr(this);
+		const FString ResourceName = *TextureRHI->GetName().ToString();
 
-		FPCGTextureReadbackInterface::Dispatch(Params, [ThisWeakPtr, PlatformData](void* OutBuffer, int32 ReadbackWidth, int32 ReadbackHeight)
+		FPCGTextureReadbackInterface::Dispatch(Params, [ThisWeakPtr, TextureSize, ResourceName](void* OutBuffer, int32 ReadbackWidth, int32 ReadbackHeight)
 		{
 			TRACE_CPUPROFILER_EVENT_SCOPE(UPCGTextureData::Initialize::DispatchCallback);
 
@@ -736,12 +760,12 @@ bool UPCGTextureData::InitializeFromGPUTexture()
 
 			const int32 PixelCount = ReadbackWidth * ReadbackHeight;
 
-			if (PlatformData->SizeX * PlatformData->SizeY != PixelCount)
+			if (TextureSize.X * TextureSize.Y != PixelCount)
 			{
 				UE_LOG(LogPCG, Error,
 					TEXT("PCGTextureData readback has different dimensions than the source texture '%s'. Expected (%d, %d), received (%d, %d)."),
-					*This->Texture->GetFName().ToString(),
-					PlatformData->SizeX, PlatformData->SizeY,
+					*ResourceName,
+					TextureSize.X, TextureSize.Y,
 					ReadbackWidth, ReadbackHeight);
 			}
 
@@ -758,7 +782,7 @@ bool UPCGTextureData::InitializeFromGPUTexture()
 			}
 			else
 			{
-				UE_LOG(LogPCG, Error, TEXT("PCGTextureData unable to get readback results from '%s'"), *This->Texture->GetFName().ToString());
+				UE_LOG(LogPCG, Error, TEXT("PCGTextureData unable to get readback results from '%s'."), *ResourceName);
 			}
 
 			This->bSuccessfullyInitialized = true;
@@ -772,7 +796,16 @@ bool UPCGTextureData::InitializeFromGPUTexture()
 	}
 	else
 	{
-		UE_LOG(LogPCG, Error, TEXT("PCGTextureData failed to acquire texture resource for '%s'"), *Texture->GetFName().ToString());
+		if (Texture.IsValid())
+		{
+			UE_LOG(LogPCG, Error, TEXT("PCGTextureData failed to acquire texture resource for '%s'."), *Texture->GetFName().ToString());
+		}
+		else
+		{
+			UE_LOG(LogPCG, Error, TEXT("PCGTextureData failed to acquire texture resource."));
+
+		}
+
 		return true;
 	}
 
diff --git a/Engine/Plugins/PCG/Source/PCG/Private/Elements/PCGGenerateGrassMaps.cpp b/Engine/Plugins/PCG/Source/PCG/Private/Elements/PCGGenerateGrassMaps.cpp
new file mode 100644
--- /dev/null
+++ b/Engine/Plugins/PCG/Source/PCG/Private/Elements/PCGGenerateGrassMaps.cpp
@@ -0,0 +1,487 @@
+// Copyright Epic Games, Inc. All Rights Reserved.
+
+#include "Elements/PCGGenerateGrassMaps.h"
+
+#include "PCGComponent.h"
+#include "PCGGrassMapUnpackerCS.h"
+#include "PCGSubsystem.h"
+#include "Data/PCGLandscapeData.h"
+#include "Data/PCGTextureData.h"
+#include "Helpers/PCGHelpers.h"
+
+#include "DrawDebugHelpers.h"
+#include "Landscape.h"
+#include "LandscapeComponent.h"
+#include "LandscapeGrassType.h"
+#include "LandscapeGrassWeightExporter.h"
+#include "LandscapeProxy.h"
+#include "RenderCaptureInterface.h"
+#include "RenderGraphBuilder.h"
+#include "RenderGraphUtils.h"
+#include "Engine/World.h"
+
+#include UE_INLINE_GENERATED_CPP_BY_NAME(PCGGenerateGrassMaps)
+
+#define LOCTEXT_NAMESPACE "PCGGenerateGrassMapsElement"
+
+namespace PCGGenerateGrassMaps
+{
+	const FName InputPinLabel = TEXT("Landscape");
+
+	const FText LandscapeComponentLostError = LOCTEXT("LandscapeComponentLost", "Reference to one or more landscape components lost, grass maps will not be generated.");
+
+	static int32 GTriggerGPUCaptureDispatches = 0;
+	static FAutoConsoleVariableRef CVarTriggerGPUCapture(
+		TEXT("pcg.Grass.TriggerGPUCaptures"),
+		GTriggerGPUCaptureDispatches,
+		TEXT("Trigger GPU captures for this many of the subsequent grass generations."));
+
+#if WITH_EDITOR
+	TAutoConsoleVariable<bool> CVarDebugDrawGeneratedComponents(
+		TEXT("pcg.Grass.DebugDrawGeneratedComponents"),
+		false,
+		TEXT("Draws debug boxes around landscapes for which grass maps are generated, colored by the current task ID."));
+#endif
+}
+
+TArray<FPCGPinProperties> UPCGGenerateGrassMapsSettings::InputPinProperties() const
+{
+	TArray<FPCGPinProperties> PinProperties;
+
+	// Single landscape only for now. In the future we should iterate over all landscapes and generate grass maps for each.
+	FPCGPinProperties& PinProp = PinProperties.Emplace_GetRef(
+		PCGGenerateGrassMaps::InputPinLabel,
+		EPCGDataType::Landscape,
+		/*bAllowMultipleConnections=*/true,
+		/*bInAllowMultipleData=*/false);
+	PinProp.SetRequiredPin();
+
+	return PinProperties;
+}
+
+TArray<FPCGPinProperties> UPCGGenerateGrassMapsSettings::OutputPinProperties() const
+{
+	TArray<FPCGPinProperties> PinProperties;
+
+	FPCGPinProperties& PinProp = PinProperties.Emplace_GetRef(
+		PCGPinConstants::DefaultOutputLabel,
+		EPCGDataType::Texture,
+		/*bAllowMultipleConnections=*/true,
+		/*bInAllowMultipleData=*/false);
+	PinProp.SetRequiredPin();
+
+	return PinProperties;
+}
+
+FPCGElementPtr UPCGGenerateGrassMapsSettings::CreateElement() const
+{
+	return MakeShared<FPCGGenerateGrassMapsElement>();
+}
+
+FPCGGenerateGrassMapsContext::~FPCGGenerateGrassMapsContext()
+{
+	delete LandscapeGrassWeightExporter;
+}
+
+void FPCGGenerateGrassMapsContext::AddExtraStructReferencedObjects(FReferenceCollector& Collector)
+{
+	for (TObjectPtr<UPCGTextureData> Data : TextureDatas)
+	{
+		if (Data)
+		{
+			Collector.AddReferencedObject(Data);
+		}
+	}
+}
+
+FPCGContext* FPCGGenerateGrassMapsElement::CreateContext()
+{
+	return new FPCGGenerateGrassMapsContext();
+}
+
+bool FPCGGenerateGrassMapsElement::ExecuteInternal(FPCGContext* InContext) const
+{
+	TRACE_CPUPROFILER_EVENT_SCOPE(FPCGGenerateGrassMapsElement::ExecuteInternal);
+
+	FPCGGenerateGrassMapsContext* Context = static_cast<FPCGGenerateGrassMapsContext*>(InContext);
+	check(Context);
+
+	// 1. Select landscape components that overlap the given bounds.
+	if (!Context->bLandscapeComponentsFiltered)
+	{
+		const UPCGLandscapeData* LandscapeData = nullptr;
+		for (const FPCGTaggedData& Data : Context->InputData.TaggedData)
+		{
+			if (const UPCGLandscapeData* InputLandscapeData = Cast<UPCGLandscapeData>(Data.Data))
+			{
+				if (!LandscapeData)
+				{
+					LandscapeData = InputLandscapeData;
+				}
+				else
+				{
+					PCGLog::LogWarningOnGraph(LOCTEXT("MultipleLandscapesNotSpported", "Multiple landscape data inputs not currently supported, only first will be used."), Context);
+				}
+			}
+		}
+
+		if (!LandscapeData)
+		{
+			// No input landscape, done.
+			return true;
+		}
+
+		const UPCGComponent* SourceComponent = Context->SourceComponent.IsValid() ? Context->SourceComponent.Get() : nullptr;
+		check(SourceComponent);
+
+		// TODO: Remove when grass map generation has been enabled in game worlds.
+		if (SourceComponent->GetWorld() && SourceComponent->GetWorld()->IsGameWorld())
+		{
+			UE_LOG(LogPCG, Warning, TEXT("Grass map generation is temporarily editor-only and disabled in game worlds."));
+			return true;
+		}
+
+		const FBox ComponentBounds = SourceComponent->GetGridBounds();
+
+		for (TSoftObjectPtr<ALandscapeProxy> LandscapeProxyPtr : LandscapeData->Landscapes)
+		{
+			if (LandscapeProxyPtr.Get())
+			{
+				Context->LandscapeProxy = LandscapeProxyPtr.Get();
+				break;
+			}
+		}
+
+		if (ALandscape* Landscape = Context->LandscapeProxy.Get() ? Context->LandscapeProxy->GetLandscapeActor() : nullptr)
+		{
+			TArray<ULandscapeComponent*> AllLandscapeComponents;
+			Landscape->GetComponents<ULandscapeComponent>(AllLandscapeComponents);
+
+			for (ULandscapeComponent* LandscapeComponent : AllLandscapeComponents)
+			{
+				if (!ensure(LandscapeComponent))
+				{
+					continue;
+				}
+
+				// Only generate grass map if there is meaningful overlap with our domain of interest.
+				const FBox LandscapeComponentBounds = LandscapeComponent->Bounds.GetBox();
+				if (LandscapeComponentBounds.Overlap(ComponentBounds).GetVolume() > UE_KINDA_SMALL_NUMBER)
+				{
+					if (Context->LandscapeComponents.Num() >= FPCGGrassMapUnpackerCS::MaxNumLandscapeComponents)
+					{
+						PCGLog::LogWarningOnGraph(LOCTEXT("MaxLandscapeComponentsExceeded", "Too many landscape components overlap the generation domain. Consider partitioning the component onto a smaller grid size."), Context);
+						break;
+					}
+
+					LandscapeComponent->UpdateGrassTypes();
+
+					Context->LandscapeComponents.Add(LandscapeComponent);
+
+					if (Context->LandscapeComponents.Num() == 1)
+					{
+						for (TObjectPtr<ULandscapeGrassType> GrassType : LandscapeComponent->GetGrassTypes())
+						{
+							Context->GrassTypes.Add(GrassType);
+						}
+
+						Context->GrassMapBounds = LandscapeComponentBounds;
+
+						Context->LandscapeComponentExtent = LandscapeComponent->Bounds.BoxExtent.X * 2.0;
+
+						if (!ensure(FMath::IsNearlyEqual(LandscapeComponent->Bounds.BoxExtent.X, LandscapeComponent->Bounds.BoxExtent.Y)))
+						{
+							return true;
+						}
+
+						if (!ensure(Context->LandscapeComponentExtent > 0.0))
+						{
+							return true;
+						}
+					}
+					else
+					{
+						Context->GrassMapBounds += LandscapeComponentBounds;
+
+						// We expect all landscape components for a single landscape to have the same grass types. Do weak validation here.
+						const uint32 ExpectedNumGrassTypes = Context->GrassTypes.Num();
+						const uint32 ActualNumGrassTypes = LandscapeComponent->GetGrassTypes().Num();
+						ensure(ExpectedNumGrassTypes == ActualNumGrassTypes);
+
+						// Currently assuming all landscape have the same extents.
+						ensure(FMath::IsNearlyEqual(Context->LandscapeComponentExtent, LandscapeComponent->Bounds.BoxExtent.X * 2.0));
+					}
+
+#if WITH_EDITOR
+					if (PCGGenerateGrassMaps::CVarDebugDrawGeneratedComponents.GetValueOnGameThread())
+					{
+						DrawDebugBox(SourceComponent->GetWorld(),
+							LandscapeComponentBounds.GetCenter(),
+							LandscapeComponentBounds.GetExtent(),
+							FColor::MakeRandomSeededColor(Context->TaskId),
+							/*bPersistentLines=*/false,
+							/*LifeTime=*/3.0f);
+					}
+#endif
+				}
+			}
+		}
+
+		if (!Context->LandscapeProxy.Get() || Context->LandscapeComponents.IsEmpty() || Context->GrassTypes.IsEmpty())
+		{
+			return true;
+		}
+
+		Context->bLandscapeComponentsFiltered = true;
+	}
+
+	// 2. Wait for landscape components to be ready for grass map rendering.
+	if (!Context->bReadyToRender)
+	{
+		bool bAllReady = true;
+
+		for (TWeakObjectPtr<ULandscapeComponent> LandscapeComponentWeak : Context->LandscapeComponents)
+		{
+			ULandscapeComponent* LandscapeComponent = LandscapeComponentWeak.Get();
+			if (!LandscapeComponent)
+			{
+				PCGLog::LogErrorOnGraph(PCGGenerateGrassMaps::LandscapeComponentLostError, InContext);
+				return true;
+			}
+
+			bAllReady &= LandscapeComponent->CanRenderGrassMap();
+		}
+
+		if (!bAllReady)
+		{
+			// Sleep until next frame, no use spinning on this.
+			Context->bIsPaused = true;
+			Context->SourceComponent->GetSubsystem()->RegisterBeginTickAction([ContextHandle = InContext->GetOrCreateHandle()]()
+			{
+				if (TSharedPtr<FPCGContextHandle> SharedHandle = ContextHandle.Pin())
+				{
+					if (FPCGContext* ContextPtr = SharedHandle->GetContext())
+					{
+						ContextPtr->bIsPaused = false;
+					}
+				}
+			});
+
+			return false;
+		}
+
+		Context->bReadyToRender = true;
+	}
+
+	// 3. Schedule grass map generation.
+	if (!Context->bGenerationScheduled)
+	{
+		ALandscapeProxy* LandscapeProxy = Context->LandscapeProxy.IsValid() ? Context->LandscapeProxy.Get() : nullptr;
+		if (!LandscapeProxy)
+		{
+			PCGLog::LogErrorOnGraph(LOCTEXT("LandscapeProxyLost", "Reference to landscape proxy actor lost, grass maps will not be generated."), InContext);
+			return true;
+		}
+
+		TArray<ULandscapeComponent*> LandscapeComponents;
+		TArray<FIntVector2> LandscapeTileCoords;
+		LandscapeComponents.Reserve(Context->LandscapeComponents.Num());
+		LandscapeTileCoords.Reserve(Context->LandscapeComponents.Num());
+
+		for (TWeakObjectPtr<ULandscapeComponent> LandscapeComponentPtr : Context->LandscapeComponents)
+		{
+			ULandscapeComponent* LandscapeComponent = LandscapeComponentPtr.IsValid() ? LandscapeComponentPtr.Get() : nullptr;
+
+			if (!LandscapeComponent)
+			{
+				PCGLog::LogErrorOnGraph(PCGGenerateGrassMaps::LandscapeComponentLostError, InContext);
+				return true;
+			}
+
+			LandscapeComponents.Add(LandscapeComponent);
+
+			// Landscape components are not ordered, so store a 2d index of each component within the grass map.
+			// TODO could likely sort the LandscapeComponents array so the order is known and no indices need looking up.
+			LandscapeTileCoords.Add(FIntVector2(
+				(LandscapeComponent->Bounds.Origin.X - Context->GrassMapBounds.Min.X) / Context->LandscapeComponentExtent,
+				(LandscapeComponent->Bounds.Origin.Y - Context->GrassMapBounds.Min.Y) / Context->LandscapeComponentExtent));
+		}
+
+		Context->LandscapeGrassWeightExporter = new FLandscapeGrassWeightExporter(
+			LandscapeProxy,
+			LandscapeComponents,
+			/*bInNeedsGrassmap=*/true,
+			/*bInNeedsHeightmap=*/false,
+			/*HeightMips=*/{},
+			/*bExport=*/false,
+			/*bShouldReadback=*/false);
+
+		UE::RenderCommandPipe::FSyncScope SyncScope;
+
+		RenderCaptureInterface::FScopedCapture RenderCapture(PCGGenerateGrassMaps::GTriggerGPUCaptureDispatches > 0, TEXT("PCGLandscapeGrassmapCapture"));
+		PCGGenerateGrassMaps::GTriggerGPUCaptureDispatches = FMath::Max(PCGGenerateGrassMaps::GTriggerGPUCaptureDispatches - 1, 0);
+		
+		TWeakPtr<FPCGContextHandle> ContextHandle = Context->GetOrCreateHandle();
+
+		ENQUEUE_RENDER_COMMAND(GenerateGrassMaps)([ContextHandle, LandscapeComponentExtent=Context->LandscapeComponentExtent, LandscapeTileCoords=MoveTemp(LandscapeTileCoords)](FRHICommandListImmediate& RHICmdList)
+		{
+			FPCGGenerateGrassMapsContext* Context = FPCGContext::GetContextFromHandle<FPCGGenerateGrassMapsContext>(ContextHandle);
+			if (!Context)
+			{
+				return;
+			}
+
+			FRDGBuilder GraphBuilder(RHICmdList);
+
+			FRDGTextureDesc GrassMapTextureDesc = FRDGTextureDesc::Create2D(
+				Context->LandscapeGrassWeightExporter->GetTargetSize(),
+				PF_B8G8R8A8,
+				FClearValueBinding(),
+				ETextureCreateFlags::RenderTargetable | ETextureCreateFlags::ShaderResource);
+
+			FRDGTextureRef GrassMapTexture = GraphBuilder.CreateTexture(GrassMapTextureDesc, TEXT("PCGLandscapeGrassMapRenderTarget"));
+			FRDGTextureSRVRef GrassMapTextureSRV = GraphBuilder.CreateSRV(GrassMapTexture);
+
+			// Generate grass maps. All will be generated to a single texture.
+			Context->LandscapeGrassWeightExporter->RenderLandscapeComponentToTexture_RenderThread(GraphBuilder, GrassMapTexture);
+
+			const FVector GrassMapExtent = Context->GrassMapBounds.GetExtent() * 2;
+			const uint32 NumTilesX = FMath::RoundToInt(GrassMapExtent.X / LandscapeComponentExtent);
+			const uint32 NumTilesY = FMath::RoundToInt(GrassMapExtent.Y / LandscapeComponentExtent);
+
+			check(NumTilesX * NumTilesY == Context->LandscapeComponents.Num());
+
+			// Each corner of a quad in the landscape corresponds to one texel in the grass map.
+			const uint32 LandscapeComponentResolution = Context->LandscapeProxy->ComponentSizeQuads + 1;
+			const FIntPoint GrassMapResolution(LandscapeComponentResolution * NumTilesX, LandscapeComponentResolution * NumTilesY);
+			const uint32 NumGrassTypes = Context->GrassTypes.Num();
+
+			// Output texture is array of textures, one per grass map.
+			FRDGTextureDesc GrassMapDesc = FRDGTextureDesc::Create2DArray(
+				GrassMapResolution,
+				PF_G8,
+				FClearValueBinding(),
+				ETextureCreateFlags::ShaderResource | ETextureCreateFlags::UAV,
+				NumGrassTypes);
+
+			FRDGTextureRef GrassMap = GraphBuilder.CreateTexture(GrassMapDesc, TEXT("PCGLandscapeGrassMapUnpacked"));
+			FRDGTextureUAVRef GrassMapUAV = GraphBuilder.CreateUAV(GrassMap);
+
+			// Export the output texture so it can be used downstream.
+			TRefCountPtr<IPooledRenderTarget> GrassMapExported = GraphBuilder.ConvertToExternalTexture(GrassMap);
+			GraphBuilder.SetTextureAccessFinal(GrassMap, ERHIAccess::SRVCompute);
+
+			// Unpack the generated results to simple world aligned textures.
+			FPCGGrassMapUnpackerCS::FParameters* Parameters = GraphBuilder.AllocParameters<FPCGGrassMapUnpackerCS::FParameters>();
+			Parameters->InPackedGrassMaps = GrassMapTextureSRV;
+			Parameters->OutUnpackedGrassMaps = GrassMapUAV;
+			Parameters->InNumTilesX = NumTilesX;
+			Parameters->InOutputResolution = FUintVector2(GrassMapResolution.X, GrassMapResolution.Y);
+			Parameters->InLandscapeComponentResolution = LandscapeComponentResolution;
+
+			// The first 2 channels are reserved for height data. See illustration of packing in PCGGrassMapUnpackerCS.usf.
+			Parameters->InNumGrassMapPasses = FMath::DivideAndRoundUp<int>(NumGrassTypes + 2, 4);
+
+			// Initialize to invalid component indices.
+			for (int Index = 0; Index < FPCGGrassMapUnpackerCS::MaxNumLandscapeComponents; ++Index)
+			{
+				Parameters->InLinearTileIndexToComponentIndex[Index].X = -1;
+			}
+
+			// Now write component mapping.
+			for (int Index = 0; Index < LandscapeTileCoords.Num(); ++Index)
+			{
+				Parameters->InLinearTileIndexToComponentIndex[LandscapeTileCoords[Index].Y * NumTilesX + LandscapeTileCoords[Index].X].X = Index;
+			}
+
+			TShaderMapRef<FPCGGrassMapUnpackerCS> Shader(GetGlobalShaderMap(GMaxRHIFeatureLevel));
+			const int GroupCountX = FMath::DivideAndRoundUp<int>(GrassMapResolution.X, FPCGGrassMapUnpackerCS::ThreadGroupDim);
+			const int GroupCountY = FMath::DivideAndRoundUp<int>(GrassMapResolution.Y, FPCGGrassMapUnpackerCS::ThreadGroupDim);
+			const int GroupCountZ = NumGrassTypes;
+
+			FComputeShaderUtils::AddPass(GraphBuilder, RDG_EVENT_NAME("PCGUnpackGrassMap"), ERDGPassFlags::Compute, Shader, Parameters, FIntVector(GroupCountX, GroupCountY, GroupCountZ));
+
+			GraphBuilder.Execute();
+
+			// Pass exported buffer back to game thread and wake up this element.
+			PCGHelpers::ExecuteOnGameThread(UE_SOURCE_LOCATION, [ContextHandle, GrassMapExported]()
+			{
+				if (FPCGGenerateGrassMapsContext* Context = FPCGContext::GetContextFromHandle<FPCGGenerateGrassMapsContext>(ContextHandle))
+				{
+					Context->GrassMapHandle = GrassMapExported;
+					Context->bIsPaused = false;
+				}
+			});
+		});
+
+		Context->bGenerationScheduled = true;
+
+		// Render command will wake this task up after completing.
+		Context->bIsPaused = true;
+
+		return false;
+	}
+
+	// 4. Initialize texture data objects.
+
+	if (Context->TextureDatas.IsEmpty())
+	{
+		// Create the texture data objects if they haven't been created already.
+		for (int DataIndex = 0; DataIndex < Context->GrassTypes.Num(); ++DataIndex)
+		{
+			UPCGTextureData* TextureData = FPCGContext::NewObject_AnyThread<UPCGTextureData>(Context);
+			Context->TextureDatas.Add(TextureData);
+		}
+	}
+
+	const FTransform GrassMapTransform(FQuat::Identity, Context->GrassMapBounds.GetCenter(), Context->GrassMapBounds.GetExtent());
+	bool bAllTexturesInitialized = true;
+
+	check(Context->GrassMapHandle.IsValid());
+
+	for (int DataIndex = 0; DataIndex < Context->GrassTypes.Num(); ++DataIndex)
+	{
+		check(Context->TextureDatas.IsValidIndex(DataIndex));
+		UPCGTextureData* TextureData = Context->TextureDatas[DataIndex];
+		check(TextureData);
+
+		// Poll initialize (fine to be called even when initialization was already complete).
+		bAllTexturesInitialized &= TextureData->Initialize(Context->GrassMapHandle, /*TextureIndex=*/DataIndex, GrassMapTransform);
+	}
+
+	if (!bAllTexturesInitialized)
+	{
+		// Initialization not complete. Could be waiting on async texture processing or for GPU readback. Sleep until next frame.
+		// TODO: Ideally we do lazy readback on texture data in the future so we don't have to read it back to CPU unless it's needed.
+		Context->bIsPaused = true;
+		Context->SourceComponent->GetSubsystem()->RegisterBeginTickAction([ContextHandle = Context->GetOrCreateHandle()]()
+		{
+			if (FPCGGenerateGrassMapsContext* ContextPtr = FPCGContext::GetContextFromHandle<FPCGGenerateGrassMapsContext>(ContextHandle))
+			{
+				ContextPtr->bIsPaused = false;
+			}
+		});
+
+		return false;
+	}
+
+	// 5. Emit texture data objects.
+
+	for (UPCGTextureData* TextureData : Context->TextureDatas)
+	{
+		if (!TextureData->IsSuccessfullyInitialized())
+		{
+			PCGE_LOG(Warning, LogOnly, LOCTEXT("TextureInitFailed", "Data could not be retrieved for this texture, initialization failed."));
+			continue;
+		}
+
+		Context->OutputData.TaggedData.Emplace_GetRef().Data = TextureData;
+
+		// TODO: Tag data with grass type name, if we can get it.
+		// TODO: Also output attribute set for grass type varieties object paths?
+	}
+
+	return true;
+}
+
+#undef LOCTEXT_NAMESPACE
diff --git a/Engine/Plugins/PCG/Source/PCG/Public/Data/PCGTextureData.h b/Engine/Plugins/PCG/Source/PCG/Public/Data/PCGTextureData.h
--- a/Engine/Plugins/PCG/Source/PCG/Public/Data/PCGTextureData.h
+++ b/Engine/Plugins/PCG/Source/PCG/Public/Data/PCGTextureData.h
@@ -4,6 +4,7 @@
 
 #include "Data/PCGSurfaceData.h"
 
+#include "RendererInterface.h"
 #include "RHI.h"
 
 #include "PCGTextureData.generated.h"
@@ -43,6 +44,13 @@ enum class EPCGTextureAddressMode : uint8
 	Wrap UMETA(ToolTip = "Tiles the texture to fit.")
 };
 
+UENUM()
+enum class EPCGTextureResourceType : uint8
+{
+	TextureObject UMETA(ToolTip = "UObject texture us as UTexture2D."),
+	ExportedTexture UMETA(ToolTip = "Texture handle exported from a texture on the GPU.")
+};
+
 namespace PCGTextureSamplingHelpers
 {
 	/** Returns true if a texture is CPU-accessible. */
@@ -82,6 +90,8 @@ public:
 	virtual bool IsValid() const;
 
 	virtual FTextureRHIRef GetTextureRHI() const PURE_VIRTUAL(UPCGBaseTextureData::GetTextureResource, return nullptr;);
+	virtual TRefCountPtr<IPooledRenderTarget> GetRefCountedTexture() const { return nullptr; }
+	virtual int GetTextureSlice() const { return 0; }
 
 protected:
 	PRAGMA_DISABLE_DEPRECATION_WARNINGS
@@ -162,20 +172,24 @@ public:
 	 * and then IsInitialized() is used to verify the initialization was successful and data is ready to use.
 	 */
 	PCG_API bool Initialize(UTexture* InTexture, uint32 InTextureIndex, const FTransform& InTransform, bool bCreateCPUDuplicateEditorOnly = false);
-
-	UE_DEPRECATED(5.5, "Will be removed. Poll the alternate Initialize API until it returns true instead of passing in a callback.")
-	PCG_API void Initialize(UTexture* InTexture, uint32 InTextureIndex, const FTransform& InTransform, const TFunction<void()>& InPostInitializeCallback, bool bCreateCPUDuplicateEditorOnly = false);
+	PCG_API bool Initialize(TRefCountPtr<IPooledRenderTarget> InTextureHandle, uint32 InTextureIndex, const FTransform& InTransform);
 
 	/** Data is successfully initialized and is ready to use. */
 	PCG_API bool IsSuccessfullyInitialized() const { return bSuccessfullyInitialized; }
 
+	EPCGTextureResourceType GetTextureResourceType() const { return ResourceType; }
+
 	//~Begin UPCGData interface
 	virtual EPCGDataType GetDataType() const override { return EPCGDataType::Texture; }
 	virtual void AddToCrc(FArchiveCrc32& Ar, bool bFullDataCrc) const override;
+	// Don't hold onto exported buffers currently as graphics memory usage (and lifetimes) may cause issues.
+	virtual bool IsCacheable() const { return Super::IsCacheable() && ResourceType != EPCGTextureResourceType::ExportedTexture; }
 	//~End UPCGData interface
 
 	//~Begin UPCGBaseTextureData interface
 	virtual FTextureRHIRef GetTextureRHI() const override;
+	virtual TRefCountPtr<IPooledRenderTarget> GetRefCountedTexture() const { return TextureHandle; }
+	virtual int GetTextureSlice() const override { return TextureIndex; }
 	//~End UPCGBaseTextureData interface
 	
 	//~Begin UPCGSpatialData interface
@@ -221,4 +235,16 @@ public:
 
 	// Added to help deprecation in 5.5. To be removed when the deprecated Initialized function is removed.
 	TFunction<void()> PostInitializeCallback;
+
+protected:
+	/** The type of underlying resource that this texture data represents. */
+	UPROPERTY()
+	EPCGTextureResourceType ResourceType = EPCGTextureResourceType::TextureObject;
+
+	/** If initialized from an exported texture this holds a reference to the resource. */
+	TRefCountPtr<IPooledRenderTarget> TextureHandle = nullptr;
+
+public:
+	UE_DEPRECATED(5.5, "Will be removed. Poll the alternate Initialize API until it returns true instead of passing in a callback.")
+	PCG_API void Initialize(UTexture* InTexture, uint32 InTextureIndex, const FTransform& InTransform, const TFunction<void()>& InPostInitializeCallback, bool bCreateCPUDuplicateEditorOnly = false);
 };
diff --git a/Engine/Plugins/PCG/Source/PCG/Public/Elements/PCGGenerateGrassMaps.h b/Engine/Plugins/PCG/Source/PCG/Public/Elements/PCGGenerateGrassMaps.h
new file mode 100644
--- /dev/null
+++ b/Engine/Plugins/PCG/Source/PCG/Public/Elements/PCGGenerateGrassMaps.h
@@ -0,0 +1,85 @@
+// Copyright Epic Games, Inc. All Rights Reserved.
+
+#pragma once
+
+#include "PCGContext.h"
+#include "PCGSettings.h"
+
+#include "PCGGenerateGrassMaps.generated.h"
+
+class UPCGTextureData;
+class ALandscapeProxy;
+class FLandscapeGrassWeightExporter;
+class ULandscapeComponent;
+class ULandscapeGrassType;
+struct IPooledRenderTarget;
+
+UCLASS(MinimalAPI, BlueprintType, ClassGroup = (Procedural))
+class UPCGGenerateGrassMapsSettings : public UPCGSettings
+{
+	GENERATED_BODY()
+
+public:
+	//~Begin UPCGSettings interface
+#if WITH_EDITOR
+	virtual FName GetDefaultNodeName() const override { return FName(TEXT("GenerateGrassMaps")); }
+	virtual FText GetDefaultNodeTitle() const override { return NSLOCTEXT("PCGGenerateGrassMapsElement", "NodeTitle", "Generate Grass Maps"); }
+	virtual FText GetNodeTooltipText() const override { return NSLOCTEXT("PCGGenerateGrassMapsElement", "NodeTooltip", "Generates landscape grass maps on the GPU."); }
+	virtual EPCGSettingsType GetType() const override { return EPCGSettingsType::GPU; }
+#endif
+
+protected:
+	virtual TArray<FPCGPinProperties> InputPinProperties() const override;
+	virtual TArray<FPCGPinProperties> OutputPinProperties() const override;
+	virtual bool IsInputPinRequiredByExecution(const UPCGPin* InPin) const override { return true; }
+	virtual FPCGElementPtr CreateElement() const override;
+	//~End UPCGSettings interface
+};
+
+struct FPCGGenerateGrassMapsContext : public FPCGContext
+{
+public:
+	~FPCGGenerateGrassMapsContext();
+
+	FLandscapeGrassWeightExporter* LandscapeGrassWeightExporter = nullptr;
+
+	/** Output texture data objects. */
+	TArray<TObjectPtr<UPCGTextureData>> TextureDatas;
+
+	/** Exported result texture array. */
+	TRefCountPtr<IPooledRenderTarget> GrassMapHandle = nullptr;
+
+	TArray<TWeakObjectPtr<ULandscapeGrassType>> GrassTypes;
+
+	TWeakObjectPtr<ALandscapeProxy> LandscapeProxy = nullptr;
+	TArray<TWeakObjectPtr<ULandscapeComponent>> LandscapeComponents;
+
+	/** World-space bounds containing all of the landscape components given to the grass weight exporter. */
+	FBox GrassMapBounds = FBox(EForceInit::ForceInit);
+
+	/** Extent (side length) of each landscape component. */
+	double LandscapeComponentExtent = 0.0;
+
+	/** True when we have filtered all of the incoming landscape components down to the ones which overlap the given bounds. */
+	bool bLandscapeComponentsFiltered = false;
+
+	/** True when the landscape components are ready for rendering. */
+	bool bReadyToRender = false;
+
+	/** True when grass map generation has been scheduled on the render thread. */
+	bool bGenerationScheduled = false;
+
+protected:
+	virtual void AddExtraStructReferencedObjects(FReferenceCollector& Collector) override;
+};
+
+class FPCGGenerateGrassMapsElement : public IPCGElement
+{
+public:
+	/** FLandscapeGrassWeightExporter expects to exist in scope only on the game thread. */
+	virtual bool CanExecuteOnlyOnMainThread(FPCGContext* Context) const { return true; }
+
+protected:
+	virtual FPCGContext* CreateContext() override;
+	virtual bool ExecuteInternal(FPCGContext* InContext) const override;
+};
diff --git a/Engine/Plugins/PCG/Source/PCGCompute/Internal/PCGGrassMapUnpackerCS.h b/Engine/Plugins/PCG/Source/PCGCompute/Internal/PCGGrassMapUnpackerCS.h
new file mode 100644
--- /dev/null
+++ b/Engine/Plugins/PCG/Source/PCGCompute/Internal/PCGGrassMapUnpackerCS.h
@@ -0,0 +1,33 @@
+// Copyright Epic Games, Inc. All Rights Reserved.
+
+#pragma once
+
+#include "GlobalShader.h"
+#include "ShaderParameterStruct.h"
+
+/** Compute shader for unpacking grass map textures rendered by FLandscapeGrassWeightExporter.
+ * Note: This class is subject to change without deprecation.
+ */
+class FPCGGrassMapUnpackerCS : public FGlobalShader
+{
+public:
+	static constexpr uint32 ThreadGroupDim = 8;
+	static constexpr uint32 MaxNumLandscapeComponents = 64;
+
+public:
+	DECLARE_EXPORTED_GLOBAL_SHADER(FPCGGrassMapUnpackerCS, PCGCOMPUTE_API);
+	SHADER_USE_PARAMETER_STRUCT(FPCGGrassMapUnpackerCS, FGlobalShader);
+
+	BEGIN_SHADER_PARAMETER_STRUCT(FParameters, PCGCOMPUTE_API)
+		SHADER_PARAMETER_RDG_TEXTURE_SRV(Texture2D<float4>, InPackedGrassMaps)
+		SHADER_PARAMETER_RDG_TEXTURE_UAV(RWTexture2DArray<float>, OutUnpackedGrassMaps)
+		SHADER_PARAMETER_ARRAY(FIntVector4, InLinearTileIndexToComponentIndex, [MaxNumLandscapeComponents])
+		SHADER_PARAMETER(uint32, InNumTilesX)
+		SHADER_PARAMETER(uint32, InLandscapeComponentResolution)
+		SHADER_PARAMETER(uint32, InNumGrassMapPasses)
+		SHADER_PARAMETER(FUintVector2, InOutputResolution)
+	END_SHADER_PARAMETER_STRUCT()
+
+public:
+	static void ModifyCompilationEnvironment(const FGlobalShaderPermutationParameters& Parameters, FShaderCompilerEnvironment& OutEnvironment);
+};
diff --git a/Engine/Plugins/PCG/Source/PCGCompute/Private/PCGGrassMapUnpackerCS.cpp b/Engine/Plugins/PCG/Source/PCGCompute/Private/PCGGrassMapUnpackerCS.cpp
new file mode 100644
--- /dev/null
+++ b/Engine/Plugins/PCG/Source/PCGCompute/Private/PCGGrassMapUnpackerCS.cpp
@@ -0,0 +1,14 @@
+// Copyright Epic Games, Inc. All Rights Reserved.
+
+#include "PCGGrassMapUnpackerCS.h"
+
+void FPCGGrassMapUnpackerCS::ModifyCompilationEnvironment(const FGlobalShaderPermutationParameters& Parameters, FShaderCompilerEnvironment& OutEnvironment)
+{
+	FGlobalShader::ModifyCompilationEnvironment(Parameters, OutEnvironment);
+
+	OutEnvironment.SetDefine(TEXT("THREADGROUPSIZE_X"), ThreadGroupDim);
+	OutEnvironment.SetDefine(TEXT("THREADGROUPSIZE_Y"), ThreadGroupDim);
+	OutEnvironment.SetDefine(TEXT("THREADGROUPSIZE_Z"), 1);
+}
+
+IMPLEMENT_GLOBAL_SHADER(FPCGGrassMapUnpackerCS, "/PCGComputeShaders/PCGGrassMapUnpackerCS.usf", "PCGGrassMapUnpacker_CS", SF_Compute);
