[PCG] Add header include to fix non unity compile on Android.

#rnx

diff --git a/Engine/Plugins/PCG/Source/PCG/Public/Elements/PCGGenerateGrassMaps.h b/Engine/Plugins/PCG/Source/PCG/Public/Elements/PCGGenerateGrassMaps.h
--- a/Engine/Plugins/PCG/Source/PCG/Public/Elements/PCGGenerateGrassMaps.h
+++ b/Engine/Plugins/PCG/Source/PCG/Public/Elements/PCGGenerateGrassMaps.h
@@ -5,6 +5,8 @@
 #include "PCGContext.h"
 #include "PCGSettings.h"
 
+#include "RendererInterface.h"
+
 #include "PCGGenerateGrassMaps.generated.h"
 
 class UPCGTextureData;
