[PCG] Fix Custom HLSL node declarations

#rnx

diff --git a/Engine/Plugins/PCG/Source/PCG/Private/Compute/Elements/PCGCustomHLSL.cpp b/Engine/Plugins/PCG/Source/PCG/Private/Compute/Elements/PCGCustomHLSL.cpp
--- a/Engine/Plugins/PCG/Source/PCG/Private/Compute/Elements/PCGCustomHLSL.cpp
+++ b/Engine/Plugins/PCG/Source/PCG/Private/Compute/Elements/PCGCustomHLSL.cpp
@@ -829,7 +829,7 @@ void UPCGCustomHLSLSettings::UpdateInputDeclarations()
 		InputDeclarations += FString::Format(TEXT(
 			"float2 {0}_GetTexCoords(float2 WorldPos, float2 Min, float2 Max);\n"
 			"float4 {0}_Sample(uint DataIndex, float2 TexCoords);\n"
-			"// Computes sample coordinates of the WorldPos relative to the texture data's bounds."
+			"// Computes sample coordinates of the WorldPos relative to the texture data's bounds.\n"
 			"float4 {0}_SampleWorldPos(uint DataIndex, float2 WorldPos);\n"),
 			{ bMultiPin ? PCGHLSLElement::PinDeclTemplateStr : TextureDataPins[0] });
 
