PCG: optimize FPCGSelectionKey::IsMatching to avoid recreating a soft object path for the same object over and over.

diff --git a/Engine/Plugins/PCG/Source/PCG/Private/Elements/PCGActorSelector.cpp b/Engine/Plugins/PCG/Source/PCG/Private/Elements/PCGActorSelector.cpp
--- a/Engine/Plugins/PCG/Source/PCG/Private/Elements/PCGActorSelector.cpp
+++ b/Engine/Plugins/PCG/Source/PCG/Private/Elements/PCGActorSelector.cpp
@@ -316,7 +316,7 @@ void FPCGSelectionKey::SetExtraDependency(const UClass* InExtraDependency)
 	OptionalExtraDependency = InExtraDependency;
 }
 
-bool FPCGSelectionKey::IsMatching(const UObject* InObject, const UPCGComponent* InComponent) const
+bool FPCGSelectionKey::IsMatching(const UObject* InObject, const FSoftObjectPath& InObjectPath, const UPCGComponent* InComponent) const
 {
 	if (!InObject)
 	{
@@ -353,13 +353,13 @@ bool FPCGSelectionKey::IsMatching(const UObject* InObject, const UPCGComponent*
 	case EPCGActorSelection::ByClass:
 		return InObject && InObject->GetClass()->IsChildOf(SelectionClass);
 	case EPCGActorSelection::ByPath:
-		return FSoftObjectPath(InObject) == ObjectPath;
+		return InObjectPath == ObjectPath;
 	default:
 		return false;
 	}
 }
 
-bool FPCGSelectionKey::IsMatching(const UObject* InObject, const TSet<FName>& InRemovedTags, const TSet<UPCGComponent*>& InComponents, TSet<UPCGComponent*>* OptionalMatchedComponents) const
+bool FPCGSelectionKey::IsMatching(const UObject* InObject, const FSoftObjectPath& InObjectPath, const TSet<FName>& InRemovedTags, const TSet<UPCGComponent*>& InComponents, TSet<UPCGComponent*>* OptionalMatchedComponents) const
 {
 	if (!InObject)
 	{
@@ -434,7 +434,7 @@ bool FPCGSelectionKey::IsMatching(const UObject* InObject, const TSet<FName>& In
 		bIsMatched = InObject->IsA(SelectionClass);
 		break;
 	case EPCGActorSelection::ByPath:
-		bIsMatched = FSoftObjectPath(InObject) == ObjectPath;
+		bIsMatched = InObjectPath == ObjectPath;
 		break;
 	default:
 		bIsMatched = false;
diff --git a/Engine/Plugins/PCG/Source/PCG/Private/Elements/PCGDataFromActor.cpp b/Engine/Plugins/PCG/Source/PCG/Private/Elements/PCGDataFromActor.cpp
--- a/Engine/Plugins/PCG/Source/PCG/Private/Elements/PCGDataFromActor.cpp
+++ b/Engine/Plugins/PCG/Source/PCG/Private/Elements/PCGDataFromActor.cpp
@@ -413,7 +413,7 @@ bool FPCGDataFromActorElement::ExecuteInternal(FPCGContext* InContext) const
 					{
 						if (AActor* Actor = Component->GetOwner())
 						{
-							if (Key.IsMatching(Actor, Component))
+							if (Key.IsMatching(Actor, FSoftObjectPath(Actor), Component))
 							{
 								Context->FoundActors.AddUnique(Actor);
 							}
@@ -426,7 +426,7 @@ bool FPCGDataFromActorElement::ExecuteInternal(FPCGContext* InContext) const
 					{
 						if (AActor* Actor = Component->GetOwner())
 						{
-							if (Key.IsMatching(Actor, Component))
+							if (Key.IsMatching(Actor, FSoftObjectPath(Actor), Component))
 							{
 								Context->FoundActors.AddUnique(Actor);
 							}
diff --git a/Engine/Plugins/PCG/Source/PCG/Private/PCGActorAndComponentMapping.cpp b/Engine/Plugins/PCG/Source/PCG/Private/PCGActorAndComponentMapping.cpp
--- a/Engine/Plugins/PCG/Source/PCG/Private/PCGActorAndComponentMapping.cpp
+++ b/Engine/Plugins/PCG/Source/PCG/Private/PCGActorAndComponentMapping.cpp
@@ -1203,9 +1203,10 @@ bool FPCGActorAndComponentMapping::IsActorTracked(const AActor* InActor) const
 
 	TSet<FName> EmptySet;
 
-	auto Matching = [InActor, &EmptySet](const TPair<FPCGSelectionKey, TSet<UPCGComponent*>>& It) -> bool
+	FSoftObjectPath ActorPath(InActor);
+	auto Matching = [InActor, &ActorPath, &EmptySet](const TPair<FPCGSelectionKey, TSet<UPCGComponent*>>& It) -> bool
 	{
-		return It.Key.IsMatching(InActor, EmptySet, It.Value, nullptr);
+		return It.Key.IsMatching(InActor, ActorPath, EmptySet, It.Value, nullptr);
 	};
 
 	return Algo::AnyOf(CulledTrackedKeysToComponentsMap, Matching) || Algo::AnyOf(AlwaysTrackedKeysToComponentsMap, Matching);
@@ -1798,11 +1799,12 @@ void FPCGActorAndComponentMapping::OnObjectChanged(UObject* InObject, const FAct
 		RemovedTags = InPreviousData->Get<1>().Difference(TSet<FName>(Actor->Tags));
 	}
 
-	auto Gather = [InObject, &RemovedTags, &MatchedKeys](TMap<FPCGSelectionKey, TSet<UPCGComponent*>> InMap, TSet<UPCGComponent*>& OutSet)
+	FSoftObjectPath ObjectPath(InObject);
+	auto Gather = [InObject, &ObjectPath, &RemovedTags, &MatchedKeys](TMap<FPCGSelectionKey, TSet<UPCGComponent*>> InMap, TSet<UPCGComponent*>& OutSet)
 	{
 		for (auto& It : InMap)
 		{
-			if (It.Key.IsMatching(InObject, RemovedTags, It.Value, &OutSet))
+			if (It.Key.IsMatching(InObject, ObjectPath, RemovedTags, It.Value, &OutSet))
 			{
 				MatchedKeys.Add(It.Key);
 			}
@@ -2142,6 +2144,9 @@ void FPCGActorAndComponentMapping::OnPCGGraphStartsGenerating(UPCGComponent* InC
 
 	auto IsKeyTrackingPCGData = [](const FPCGSelectionKey& InKey) { return InKey.OptionalExtraDependency && InKey.OptionalExtraDependency->IsChildOf(UPCGComponent::StaticClass()); };
 
+	AActor* Actor = InComponent->GetOwner();
+	FSoftObjectPath ActorPath(Actor);
+
 	for (TPair<FPCGSelectionKey, TSet<UPCGComponent*>>& It : CulledTrackedKeysToComponentsMap)
 	{
 		if (!IsKeyTrackingPCGData(It.Key))
@@ -2150,7 +2155,7 @@ void FPCGActorAndComponentMapping::OnPCGGraphStartsGenerating(UPCGComponent* InC
 		}
 
 		TSet<UPCGComponent*> TempTrackedComponents;
-		It.Key.IsMatching(InComponent->GetOwner(), RemovedTags, It.Value, &TempTrackedComponents);
+		It.Key.IsMatching(Actor, ActorPath, RemovedTags, It.Value, &TempTrackedComponents);
 
 		// Removing all components that aren't intersecting with it, since it won't contribute to refresh
 		for (UPCGComponent* TrackedComponent : TempTrackedComponents)
@@ -2169,7 +2174,7 @@ void FPCGActorAndComponentMapping::OnPCGGraphStartsGenerating(UPCGComponent* InC
 			continue;
 		}
 
-		It.Key.IsMatching(InComponent->GetOwner(), RemovedTags, It.Value, &TrackedComponents);
+		It.Key.IsMatching(Actor, ActorPath, RemovedTags, It.Value, &TrackedComponents);
 	}
 
 	for (UPCGComponent* Component : TrackedComponents)
@@ -2181,7 +2186,7 @@ void FPCGActorAndComponentMapping::OnPCGGraphStartsGenerating(UPCGComponent* InC
 		}
 
 		// If the tracked component is currently ignoring the refresh from InComponent, don't add to the dependencies
-		if (Component->IsIgnoringChangeOrigin(InComponent->GetOwner()))
+		if (Component->IsIgnoringChangeOrigin(Actor))
 		{
 			continue;
 		}
diff --git a/Engine/Plugins/PCG/Source/PCG/Private/PCGComponent.cpp b/Engine/Plugins/PCG/Source/PCG/Private/PCGComponent.cpp
--- a/Engine/Plugins/PCG/Source/PCG/Private/PCGComponent.cpp
+++ b/Engine/Plugins/PCG/Source/PCG/Private/PCGComponent.cpp
@@ -2830,11 +2830,12 @@ bool UPCGComponent::IsObjectTracked(const UObject* InObject, bool& bOutIsCulled)
 		return true;
 	}
 
-	auto CheckMap = [this, &InObject, &bOutIsCulled](const FPCGSelectionKeyToSettingsMap& InMap) -> bool
+	FSoftObjectPath ObjectPath(InObject);
+	auto CheckMap = [this, &InObject, &ObjectPath, &bOutIsCulled](const FPCGSelectionKeyToSettingsMap& InMap) -> bool
 	{
 		for (const auto& It : InMap)
 		{
-			bool bFound = It.Key.IsMatching(InObject, this);
+			bool bFound = It.Key.IsMatching(InObject, ObjectPath, this);
 
 			if (bFound)
 			{
diff --git a/Engine/Plugins/PCG/Source/PCG/Public/Elements/PCGActorSelector.h b/Engine/Plugins/PCG/Source/PCG/Public/Elements/PCGActorSelector.h
--- a/Engine/Plugins/PCG/Source/PCG/Public/Elements/PCGActorSelector.h
+++ b/Engine/Plugins/PCG/Source/PCG/Public/Elements/PCGActorSelector.h
@@ -78,8 +78,9 @@ struct PCG_API FPCGSelectionKey
 
 	// Friend functions need to be explicitly exported.
 	PCG_API friend uint32 GetTypeHash(const FPCGSelectionKey& In);
-	bool IsMatching(const UObject* InObject, const UPCGComponent* InComponent) const;
-	bool IsMatching(const UObject* InObject, const TSet<FName>& InRemovedTags, const TSet<UPCGComponent*>& InComponents, TSet<UPCGComponent*>* OptionalMatchedComponents = nullptr) const;
+
+	bool IsMatching(const UObject* InObject, const FSoftObjectPath& InObjectPath, const UPCGComponent* InComponent) const;
+	bool IsMatching(const UObject* InObject, const FSoftObjectPath& InObjectPath, const TSet<FName>& InRemovedTags, const TSet<UPCGComponent*>& InComponents, TSet<UPCGComponent*>* OptionalMatchedComponents = nullptr) const;
 
 	void SetExtraDependency(const UClass* InExtraDependency);
 
