[PCG] Fix cook error in PCGTextureDataInterface.ush
* Replace Sample with SampleLevel

#rnx

diff --git a/Engine/Plugins/PCG/Shaders/Private/PCGTextureDataInterface.ush b/Engine/Plugins/PCG/Shaders/Private/PCGTextureDataInterface.ush
--- a/Engine/Plugins/PCG/Shaders/Private/PCGTextureDataInterface.ush
+++ b/Engine/Plugins/PCG/Shaders/Private/PCGTextureDataInterface.ush
@@ -42,7 +42,7 @@ float2 GetTexCoords_{DataInterfaceName}(float2 WorldPos, float2 Min, float2 Max)
 #define PCG_CONDITIONAL_TEXTURE_SAMPLE(Index) \
 	if (InDataIndex == Index) \
 	{ \
-		return {DataInterfaceName}_Texture_##Index.Sample({DataInterfaceName}_Sampler_##Index, InTextureUV); \
+		return {DataInterfaceName}_Texture_##Index.SampleLevel({DataInterfaceName}_Sampler_##Index, InTextureUV, 0); \
 	}
 
 float4 Sample_{DataInterfaceName}(int InDataIndex, float2 InTextureUV)
@@ -66,7 +66,7 @@ float4 Sample_{DataInterfaceName}(int InDataIndex, float2 InTextureUV)
 	if (InDataIndex == Index) \
 	{ \
 		const float2 TextureUV = GetTexCoords_{DataInterfaceName}(WorldPos, {DataInterfaceName}_TextureBounds[Index].xy, {DataInterfaceName}_TextureBounds[Index].zw); \
-		return {DataInterfaceName}_Texture_##Index.Sample({DataInterfaceName}_Sampler_##Index, TextureUV); \
+		return {DataInterfaceName}_Texture_##Index.SampleLevel({DataInterfaceName}_Sampler_##Index, TextureUV, 0); \
 	}
 
 float4 SampleWorldPos_{DataInterfaceName}(int InDataIndex, float2 WorldPos)
