[PCG] Fix failing attribute copy and write debug value from Mac and Linux

[RN] [PCG] Linux & Mac: Fixed uninitialized data in the output of GPU nodes for point properties that are not explicitly written in the node.

diff --git a/Engine/Plugins/PCG/Shaders/Private/PCGDataCollectionDataInterface.ush b/Engine/Plugins/PCG/Shaders/Private/PCGDataCollectionDataInterface.ush
--- a/Engine/Plugins/PCG/Shaders/Private/PCGDataCollectionDataInterface.ush
+++ b/Engine/Plugins/PCG/Shaders/Private/PCGDataCollectionDataInterface.ush
@@ -109,7 +109,7 @@ uint LoadBufferInternal_{DataInterfaceName}(uint Address)
 
 void StoreBufferInternal_{DataInterfaceName}(uint Address, uint Value)
 {
-	return {DataInterfaceName}_DataCollectionBuffer.Store(Address, Value);
+	{DataInterfaceName}_DataCollectionBuffer.Store(Address, Value);
 }
 
 uint GetFirstElementAddressInternal_{DataInterfaceName}(uint InDataIndex, int InAttributeId)
diff --git a/Engine/Plugins/PCG/Source/PCG/Private/Compute/DataInterfaces/PCGDebugDataInterface.cpp b/Engine/Plugins/PCG/Source/PCG/Private/Compute/DataInterfaces/PCGDebugDataInterface.cpp
--- a/Engine/Plugins/PCG/Source/PCG/Private/Compute/DataInterfaces/PCGDebugDataInterface.cpp
+++ b/Engine/Plugins/PCG/Source/PCG/Private/Compute/DataInterfaces/PCGDebugDataInterface.cpp
@@ -44,7 +44,7 @@ void UPCGDebugDataInterface::GetHLSL(FString& OutHLSL, FString const& InDataInte
 			"\n{\n"
 			"	if (Index >= 0 && Index < {0})\n"
 			"	{\n"
-			"		return {1}_DebugBuffer.Store(Index * 4, Value);\n"
+			"		{1}_DebugBuffer.Store(Index * 4, asuint(Value));\n"
 			"	}\n"
 			"}\n"),
 			{ DebugBufferSizeFloats, *InDataInterfaceName })
