[PCG] Support runtime grass map generation in PCG

diff --git a/Engine/Plugins/PCG/Source/PCG/Private/Elements/PCGGenerateGrassMaps.cpp b/Engine/Plugins/PCG/Source/PCG/Private/Elements/PCGGenerateGrassMaps.cpp
--- a/Engine/Plugins/PCG/Source/PCG/Private/Elements/PCGGenerateGrassMaps.cpp
+++ b/Engine/Plugins/PCG/Source/PCG/Private/Elements/PCGGenerateGrassMaps.cpp
@@ -134,13 +134,6 @@ bool FPCGGenerateGrassMapsElement::ExecuteInternal(FPCGContext* InContext) const
 		const UPCGComponent* SourceComponent = Context->SourceComponent.IsValid() ? Context->SourceComponent.Get() : nullptr;
 		check(SourceComponent);
 
-		// TODO: Remove when grass map generation has been enabled in game worlds.
-		if (SourceComponent->GetWorld() && SourceComponent->GetWorld()->IsGameWorld())
-		{
-			UE_LOG(LogPCG, Warning, TEXT("Grass map generation is temporarily editor-only and disabled in game worlds."));
-			return true;
-		}
-
 		const FBox ComponentBounds = SourceComponent->GetGridBounds();
 
 		for (TSoftObjectPtr<ALandscapeProxy> LandscapeProxyPtr : LandscapeData->Landscapes)
@@ -249,7 +242,7 @@ bool FPCGGenerateGrassMapsElement::ExecuteInternal(FPCGContext* InContext) const
 				return true;
 			}
 
-			bAllReady &= LandscapeComponent->CanRenderGrassMap();
+			bAllReady &= UE::Landscape::CanRenderGrassMap(LandscapeComponent);
 		}
 
 		if (!bAllReady)
diff --git a/Engine/Source/Runtime/Landscape/Internal/LandscapeGrassWeightExporter.h b/Engine/Source/Runtime/Landscape/Internal/LandscapeGrassWeightExporter.h
--- a/Engine/Source/Runtime/Landscape/Internal/LandscapeGrassWeightExporter.h
+++ b/Engine/Source/Runtime/Landscape/Internal/LandscapeGrassWeightExporter.h
@@ -13,6 +13,11 @@ class FTextureRenderTarget2DResource;
 class UTextureRenderTarget2D;
 using FRDGTextureRef = FRDGTexture*;
 
+namespace UE::Landscape
+{
+	LANDSCAPE_API bool CanRenderGrassMap(ULandscapeComponent* Component);
+}
+
 // Hacky base class to avoid 8 bytes of padding after the vtable
 class FLandscapeGrassWeightExporter_RenderThread_FixLayout
 {
