Added back deprecated API on actor selector objects

#rnx

diff --git a/Engine/Plugins/PCG/Source/PCG/Private/Elements/PCGActorSelector.cpp b/Engine/Plugins/PCG/Source/PCG/Private/Elements/PCGActorSelector.cpp
--- a/Engine/Plugins/PCG/Source/PCG/Private/Elements/PCGActorSelector.cpp
+++ b/Engine/Plugins/PCG/Source/PCG/Private/Elements/PCGActorSelector.cpp
@@ -316,6 +316,12 @@ void FPCGSelectionKey::SetExtraDependency(const UClass* InExtraDependency)
 	OptionalExtraDependency = InExtraDependency;
 }
 
+// DEPRECATED 5.6
+bool FPCGSelectionKey::IsMatching(const UObject* InObject, const UPCGComponent* InComponent) const
+{
+	return IsMatching(InObject, FSoftObjectPath(InObject), InComponent);
+}
+
 bool FPCGSelectionKey::IsMatching(const UObject* InObject, const FSoftObjectPath& InObjectPath, const UPCGComponent* InComponent) const
 {
 	if (!InObject)
@@ -359,6 +365,12 @@ bool FPCGSelectionKey::IsMatching(const UObject* InObject, const FSoftObjectPath
 	}
 }
 
+// DEPRECATED 5.6
+bool FPCGSelectionKey::IsMatching(const UObject* InObject, const TSet<FName>& InRemovedTags, const TSet<UPCGComponent*>& InComponents, TSet<UPCGComponent*>* OptionalMatchedComponents) const
+{
+	return IsMatching(InObject, FSoftObjectPath(InObject), InRemovedTags, InComponents, OptionalMatchedComponents);
+}
+
 bool FPCGSelectionKey::IsMatching(const UObject* InObject, const FSoftObjectPath& InObjectPath, const TSet<FName>& InRemovedTags, const TSet<UPCGComponent*>& InComponents, TSet<UPCGComponent*>* OptionalMatchedComponents) const
 {
 	if (!InObject)
diff --git a/Engine/Plugins/PCG/Source/PCG/Public/Elements/PCGActorSelector.h b/Engine/Plugins/PCG/Source/PCG/Public/Elements/PCGActorSelector.h
--- a/Engine/Plugins/PCG/Source/PCG/Public/Elements/PCGActorSelector.h
+++ b/Engine/Plugins/PCG/Source/PCG/Public/Elements/PCGActorSelector.h
@@ -103,6 +103,12 @@ struct PCG_API FPCGSelectionKey
 	// If it should track a specific object dependency instead of an actor. For example, GetActorData with GetPCGComponent data.
 	UPROPERTY()
 	TObjectPtr<const UClass> OptionalExtraDependency = nullptr;
+
+	// Deprecated methods
+	UE_DEPRECATED(5.6, "Replaced by the version with a soft object path.")
+	bool IsMatching(const UObject* InObject, const UPCGComponent* InComponent) const;
+	UE_DEPRECATED(5.6, "Replaced by the version with a soft object path.")
+	bool IsMatching(const UObject* InObject, const TSet<FName>& InRemovedTags, const TSet<UPCGComponent*>& InComponents, TSet<UPCGComponent*>* OptionalMatchedComponents = nullptr) const;
 };
 
 PCG_API FArchive& operator<<(FArchive& Ar, FPCGSelectionKey& Key);
