[PCG] Node Conversion - Create Attribute -> Get User Parameter

+ Also directly adds a user graph parameter

diff --git a/Engine/Plugins/PCG/Source/PCG/Private/Elements/PCGCreateAttribute.cpp b/Engine/Plugins/PCG/Source/PCG/Private/Elements/PCGCreateAttribute.cpp
--- a/Engine/Plugins/PCG/Source/PCG/Private/Elements/PCGCreateAttribute.cpp
+++ b/Engine/Plugins/PCG/Source/PCG/Private/Elements/PCGCreateAttribute.cpp
@@ -4,10 +4,14 @@
 
 #include "PCGContext.h"
 #include "PCGCustomVersion.h"
+#include "PCGGraph.h"
 #include "PCGParamData.h"
 #include "PCGPin.h"
 #include "Data/PCGSpatialData.h"
+#include "Elements/PCGUserParameterGet.h"
 #include "Elements/Metadata/PCGMetadataElementCommon.h"
+#include "Helpers/PCGConversion.h"
+#include "Helpers/PCGPropertyHelpers.h"
 #include "Metadata/Accessors/PCGAttributeAccessorHelpers.h"
 
 #include UE_INLINE_GENERATED_CPP_BY_NAME(PCGCreateAttribute)
@@ -23,6 +27,14 @@ namespace PCGCreateAttributeConstants
 	const FName AttributesLabel = TEXT("Attributes");
 	const FText AttributesTooltip = LOCTEXT("AttributesTooltip", "Optional Attribute Set to create from. Not used if not connected.");
 	const FText ErrorCreatingAttributeMessage = LOCTEXT("ErrorCreatingAttribute", "Error while creating attribute '{0}'");
+
+	namespace Conversion
+	{
+		constexpr int32 ToGetGraphParameterIndex = 0;
+
+		const FText ToGetGraphParameterTooltip = LOCTEXT("ToGetGraphParameterTooltip", "Convert constant to graph parameter.");
+		const FName ConversionDefaultName(LOCTEXT("ConversionDefaultName", "NewProperty").ToString());
+	}
 }
 
 namespace PCGCreateAttribute
@@ -266,6 +278,84 @@ FString UPCGCreateAttributeSetSettings::GetAdditionalTitleInformation() const
 	return FString::Printf(TEXT("%s: %s"), *OutputTarget.GetName().ToString(), *AttributeTypes.ToString());
 }
 
+#if WITH_EDITOR
+TArray<FPCGPreconfiguredInfo> UPCGCreateAttributeSetSettings::GetConversionInfo() const
+{
+	using namespace PCGCreateAttributeConstants::Conversion;
+	TArray<FPCGPreconfiguredInfo> ConversionInfo;
+
+	ConversionInfo.Emplace(
+		ToGetGraphParameterIndex,
+		PCGConversion::Helpers::GetDefaultNodeTitle(UPCGUserParameterGetSettings::StaticClass()),
+		ToGetGraphParameterTooltip);
+
+	return ConversionInfo;
+}
+#endif // WITH_EDITOR
+
+bool UPCGCreateAttributeSetSettings::ConvertNode(const FPCGPreconfiguredInfo& ConversionInfo)
+{
+	UPCGNode* Node = CastChecked<UPCGNode>(GetOuter());
+
+	switch (ConversionInfo.PreconfiguredIndex)
+	{
+		case PCGCreateAttributeConstants::Conversion::ToGetGraphParameterIndex:
+		{
+			FPCGSingleNodeConverter NodeConverter(Node, UPCGUserParameterGetSettings::StaticClass());
+			if (!NodeConverter.IsGraphInitialized() || !NodeConverter.IsSourceInitialized())
+			{
+				PCGLog::Settings::LogInvalidPreconfigurationWarning(ConversionInfo.PreconfiguredIndex, PCGCreateAttributeConstants::NodeTitleCreateAttribute);
+				return false;
+			}
+
+			NodeConverter.PrepareData();
+
+			// The property will determine the structure of the Graph Parameter node.
+			if (UPCGGraph* Graph = NodeConverter.GetGraph())
+			{
+				const FInstancedPropertyBag* UserParametersStruct = Graph->GetUserParametersStruct();
+				FName PropertyName = (OutputTarget.GetName() != NAME_None)
+					? OutputTarget.GetName()
+					: PCGCreateAttributeConstants::Conversion::ConversionDefaultName;
+
+				// Validate that the property name doesn't exist... if it does, increment until it doesn't
+				uint64 NameIteration = 1;
+				while (UserParametersStruct->FindPropertyDescByName(PropertyName))
+				{
+					if (++NameIteration > 100)
+					{
+						return false;
+					}
+
+					PropertyName.SetNumber(NameIteration);
+				}
+
+				FPropertyBagPropertyDesc PropertyDesc = PCGPropertyHelpers::GetPropertyBagDescWithMetadataType(PropertyName, AttributeTypes.Type);
+				Graph->AddUserParameters({PropertyDesc});
+				const FPropertyBagPropertyDesc* Desc = UserParametersStruct->FindPropertyDescByName(PropertyName);
+				check(Desc);
+
+				// Update the generated node with the new user parameter
+				UPCGUserParameterGetSettings* Settings = CastChecked<UPCGUserParameterGetSettings>(NodeConverter.GetGeneratedSettings());
+				Settings->PropertyName = Desc->Name;
+				Settings->PropertyGuid = Desc->ID;
+
+				AttributeTypes.Dispatcher([PropertyName, Graph]<typename T>(T&& Value) { Graph->SetGraphParameter(PropertyName, std::forward<T>(Value)); });
+			}
+
+			NodeConverter.ApplyStructural();
+			NodeConverter.Finalize();
+
+			return NodeConverter.IsComplete();
+		}
+
+		default:
+			break;
+	}
+
+	return false;
+}
+
 FPCGElementPtr UPCGCreateAttributeSetSettings::CreateElement() const
 {
 	return MakeShared<FPCGCreateAttributeElement>();
diff --git a/Engine/Plugins/PCG/Source/PCG/Public/Elements/PCGCreateAttribute.h b/Engine/Plugins/PCG/Source/PCG/Public/Elements/PCGCreateAttribute.h
--- a/Engine/Plugins/PCG/Source/PCG/Public/Elements/PCGCreateAttribute.h
+++ b/Engine/Plugins/PCG/Source/PCG/Public/Elements/PCGCreateAttribute.h
@@ -98,7 +98,9 @@ public:
 protected:
 #if WITH_EDITOR
 	virtual EPCGChangeType GetChangeTypeForProperty(const FName& InPropertyName) const override { return Super::GetChangeTypeForProperty(InPropertyName) | EPCGChangeType::Cosmetic; }
+	virtual TArray<FPCGPreconfiguredInfo> GetConversionInfo() const override;
 #endif
+	virtual bool ConvertNode(const FPCGPreconfiguredInfo& ConversionInfo) override;
 	//~End UPCGSettings interface
 
 public:
