[PCG] GPU: Add a macro for copying attributes across pins.

#rnx

diff --git a/Engine/Plugins/PCG/Shaders/Private/Elements/PCGCopyPoints.usf b/Engine/Plugins/PCG/Shaders/Private/Elements/PCGCopyPoints.usf
--- a/Engine/Plugins/PCG/Shaders/Private/Elements/PCGCopyPoints.usf
+++ b/Engine/Plugins/PCG/Shaders/Private/Elements/PCGCopyPoints.usf
@@ -3,6 +3,7 @@
 #include "/Engine/Private/ComputeShaderUtils.ush"
 #include "/Engine/Private/Quaternion.ush"
 #include "/Plugin/PCG/Private/PCGShaderUtils.ush"
+#include "/Plugin/PCG/Private/PCGShaderUtilsInternal.ush"
 
 // Mirrors EPCGCopyPointsInheritanceMode
 #define COPY_POINTS_INHERITANCE_MODE_Relative 0
@@ -20,96 +21,6 @@
 #define MAX_NUM_ATTRS 128
 #define NUM_RESERVED_ATTRS 32
 
-// TODO pass in attribute IDs to copy.
-void CopyAttributesFromSource(uint Source_DataIndex, uint Source_ElementIndex, uint Out_DataIndex, uint Out_ElementIndex, bool bInitializeNonCopiedAttributes)
-{
-	const uint NumAttributes = Out_GetDataNumAttributesInternal(Out_DataIndex);
-	uint NumAttributesProcessed = 0;
-
-	for (int AttributeId = NUM_RESERVED_ATTRS; AttributeId < MAX_NUM_ATTRS; ++AttributeId)
-	{
-		const uint Out_Stride = Out_GetAttributeStrideInternal(Out_DataIndex, AttributeId);
-
-		if (Out_Stride == 0)
-		{
-			// No output attribute to write to.
-			continue;
-		}
-		
-		const uint Source_Stride = Source_GetAttributeStrideInternal(Source_DataIndex, AttributeId);
-		const uint Out_ElementAddress = Out_GetElementAddressInternal(Out_DataIndex, Out_ElementIndex, AttributeId);
-		
-		// Copy the attribute if it exists in the Source data.
-		if (Source_Stride == Out_Stride)
-		{
-			const uint Source_ElementAddress = Source_GetElementAddressInternal(Source_DataIndex, Source_ElementIndex, AttributeId);
-			
-			for (uint Offset = 0; Offset < Out_Stride; Offset += 4)
-			{
-				Out_StoreBufferInternal(Out_ElementAddress + Offset, Source_LoadBufferInternal(Source_ElementAddress + Offset));
-			}
-		}
-		else if (bInitializeNonCopiedAttributes)
-		{
-			for (uint Offset = 0; Offset < Out_Stride; Offset += 4)
-			{
-				// Initialize output data (only header part of buffer will be initialized).
-				Out_StoreBufferInternal(Out_ElementAddress + Offset, 0u);
-			}
-		}
-		
-		if (++NumAttributesProcessed >= NumAttributes)
-		{
-			break; // We can early-out when we've looked at all the possible attributes
-		}
-	}
-}
-
-// TODO pass in attribute IDs to copy.
-void CopyAttributesFromTarget(uint Target_DataIndex, uint Target_ElementIndex, uint Out_DataIndex, uint Out_ElementIndex, bool bInitializeNonCopiedAttributes)
-{
-	const uint NumAttributes = Out_GetDataNumAttributesInternal(Out_DataIndex);
-	uint NumAttributesProcessed = 0;
-	
-	for (int AttributeId = NUM_RESERVED_ATTRS; AttributeId < MAX_NUM_ATTRS; ++AttributeId)
-	{
-		const uint Out_Stride = Out_GetAttributeStrideInternal(Out_DataIndex, AttributeId);
-		
-		if (Out_Stride == 0)
-		{
-			// No output attribute to write to.
-			continue;
-		}
-		
-		const uint Target_Stride = Target_GetAttributeStrideInternal(Target_DataIndex, AttributeId);
-		const uint Out_ElementAddress = Out_GetElementAddressInternal(Out_DataIndex, Out_ElementIndex, AttributeId);
-		
-		// Copy the attribute if it exists in the Target data.
-		if (Target_Stride == Out_Stride)
-		{
-			const uint Target_ElementAddress = Target_GetElementAddressInternal(Target_DataIndex, Target_ElementIndex, AttributeId);
-			
-			for (uint Offset = 0; Offset < Out_Stride; Offset += 4)
-			{
-				Out_StoreBufferInternal(Out_ElementAddress + Offset, Target_LoadBufferInternal(Target_ElementAddress + Offset));
-			}
-		}
-		else if (bInitializeNonCopiedAttributes)
-		{
-			for (uint Offset = 0; Offset < Out_Stride; Offset += 4)
-			{
-				// Initialize output data (only header part of buffer will be initialized).
-				Out_StoreBufferInternal(Out_ElementAddress + Offset, 0u);
-			}
-		}
-		
-		if (++NumAttributesProcessed >= NumAttributes)
-		{
-			break; // We can early-out when we've looked at all the possible attributes
-		}
-	}
-}
-
 [numthreads(64, 1, 1)]
 void Main(uint3 GroupId : SV_GroupID, uint GroupIndex : SV_GroupIndex)
 {
@@ -304,24 +215,24 @@ void Main(uint3 GroupId : SV_GroupID, uint GroupIndex : SV_GroupIndex)
 	// Copy attributes
 	if (AttributeInheritance == COPY_POINTS_METADATA_INHERITANCE_MODE_SourceOnly)
 	{
-		CopyAttributesFromSource(SourceDataIndex, SourcePointIndexInData, OutDataIndex, OutElemIndex, /*bInitializeNonCopiedAttributes=*/true);
+		PCG_COPY_ATTRIBUTES_TO_OUTPUT(Out, Source, OutDataIndex, OutElemIndex, SourceDataIndex, SourcePointIndexInData, /*bMetadataOnly=*/true, /*bInitNonCopiedAttributes=*/true);
 	}
 	else if (AttributeInheritance == COPY_POINTS_METADATA_INHERITANCE_MODE_TargetOnly)
 	{
-		CopyAttributesFromTarget(TargetDataIndex, TargetPointIndexInData, OutDataIndex, OutElemIndex, /*bInitializeNonCopiedAttributes=*/true);
+		PCG_COPY_ATTRIBUTES_TO_OUTPUT(Out, Target, OutDataIndex, OutElemIndex, TargetDataIndex, TargetPointIndexInData, /*bMetadataOnly=*/true, /*bInitNonCopiedAttributes=*/true);
 	}
 	else if (AttributeInheritance == COPY_POINTS_METADATA_INHERITANCE_MODE_SourceFirst)
 	{
-		CopyAttributesFromTarget(TargetDataIndex, TargetPointIndexInData, OutDataIndex, OutElemIndex, /*bInitializeNonCopiedAttributes=*/true);
+		PCG_COPY_ATTRIBUTES_TO_OUTPUT(Out, Target, OutDataIndex, OutElemIndex, TargetDataIndex, TargetPointIndexInData, /*bMetadataOnly=*/true, /*bInitNonCopiedAttributes=*/true);
 
 		// Don't allow Source to initialize non-copied attributes in the output, otherwise it will overwrite attributes from Target.
-		CopyAttributesFromSource(SourceDataIndex, SourcePointIndexInData, OutDataIndex, OutElemIndex, /*bInitializeNonCopiedAttributes=*/false);
+		PCG_COPY_ATTRIBUTES_TO_OUTPUT(Out, Source, OutDataIndex, OutElemIndex, SourceDataIndex, SourcePointIndexInData, /*bMetadataOnly=*/true, /*bInitNonCopiedAttributes=*/false);
 	}
 	else if (AttributeInheritance == COPY_POINTS_METADATA_INHERITANCE_MODE_TargetFirst)
 	{
-		CopyAttributesFromSource(SourceDataIndex, SourcePointIndexInData, OutDataIndex, OutElemIndex, /*bInitializeNonCopiedAttributes=*/true);
+		PCG_COPY_ATTRIBUTES_TO_OUTPUT(Out, Source, OutDataIndex, OutElemIndex, SourceDataIndex, SourcePointIndexInData, /*bMetadataOnly=*/true, /*bInitNonCopiedAttributes=*/true);
 
 		// Don't allow Target to initialize non-copied attributes in the output, otherwise it will overwrite attributes from Source.
-		CopyAttributesFromTarget(TargetDataIndex, TargetPointIndexInData, OutDataIndex, OutElemIndex, /*bInitializeNonCopiedAttributes=*/false);
+		PCG_COPY_ATTRIBUTES_TO_OUTPUT(Out, Target, OutDataIndex, OutElemIndex, TargetDataIndex, TargetPointIndexInData, /*bMetadataOnly=*/true, /*bInitNonCopiedAttributes=*/false);
 	}
-}
\ No newline at end of file
+}
diff --git a/Engine/Plugins/PCG/Shaders/Private/Elements/PCGStaticMeshSpawner.usf b/Engine/Plugins/PCG/Shaders/Private/Elements/PCGStaticMeshSpawner.usf
--- a/Engine/Plugins/PCG/Shaders/Private/Elements/PCGStaticMeshSpawner.usf
+++ b/Engine/Plugins/PCG/Shaders/Private/Elements/PCGStaticMeshSpawner.usf
@@ -8,13 +8,12 @@
 #include "/Engine/Private/ComputeShaderUtils.ush"
 #include "/Engine/Private/Quaternion.ush"
 #include "/Plugin/PCG/Private/PCGShaderUtils.ush"
+#include "/Plugin/PCG/Private/PCGShaderUtilsInternal.ush"
 
 // TODO: should be code-gen'd from PCGComputeCommon.h
 #define MAX_NUM_ATTRS 128
 #define NUM_RESERVED_ATTRS 32
 
-void CopyAttributesToOutput(uint In_DataIndex, uint In_ElementIndex, uint Out_DataIndex, uint Out_ElementIndex);
-
 [numthreads(64, 1, 1)]
 void Main(uint3 GroupId : SV_GroupID, uint GroupIndex : SV_GroupIndex)
 {
@@ -134,7 +133,7 @@ void Main(uint3 GroupId : SV_GroupID, uint GroupIndex : SV_GroupIndex)
 		Out_SetBoundsMax(Out_DataIndex, Out_ElementIndex, In_GetBoundsMax(In_DataIndex, In_ElementIndex));
 	}
 
-	CopyAttributesToOutput(In_DataIndex, In_ElementIndex, Out_DataIndex, Out_ElementIndex);
+	PCG_COPY_METADATA_ATTRIBUTES_TO_OUTPUT(Out, In, Out_DataIndex, Out_ElementIndex, In_DataIndex, In_ElementIndex);
 
 	if (SMSpawner_GetSelectedMeshAttributeId() != (uint)-1)
 	{
@@ -142,47 +141,3 @@ void Main(uint3 GroupId : SV_GroupID, uint GroupIndex : SV_GroupIndex)
 		Out_SetStringKey(Out_DataIndex, Out_ElementIndex, SMSpawner_GetSelectedMeshAttributeId(), OutputMeshStringKey);
 	}
 }
-
-// TODO pass in attribute IDs to copy.
-void CopyAttributesToOutput(uint In_DataIndex, uint In_ElementIndex, uint Out_DataIndex, uint Out_ElementIndex)
-{
-	const uint NumAttributes = Out_GetDataNumAttributesInternal(Out_DataIndex);
-	uint NumAttributesProcessed = 0;
-	
-	for (int AttributeId = NUM_RESERVED_ATTRS; AttributeId < MAX_NUM_ATTRS; ++AttributeId)
-	{
-		const uint Out_Stride = Out_GetAttributeStrideInternal(Out_DataIndex, AttributeId);
-		
-		if (Out_Stride == 0)
-		{
-			// No output attribute to write to.
-			continue;
-		}
-		
-		const uint In_Stride = In_GetAttributeStrideInternal(In_DataIndex, AttributeId);
-		const uint Out_ElementAddress = Out_GetElementAddressInternal(Out_DataIndex, Out_ElementIndex, AttributeId);
-		
-		if (In_Stride == Out_Stride)
-		{
-			const uint In_ElementAddress = In_GetElementAddressInternal(In_DataIndex, In_ElementIndex, AttributeId);
-			
-			for (uint Offset = 0; Offset < Out_Stride; Offset += 4)
-			{
-				Out_StoreBufferInternal(Out_ElementAddress + Offset, In_LoadBufferInternal(In_ElementAddress + Offset));
-			}
-		}
-		else
-		{
-			for (uint Offset = 0; Offset < Out_Stride; Offset += 4)
-			{
-				// Initialize output data (only header part of buffer will be initialized).
-				Out_StoreBufferInternal(Out_ElementAddress + Offset, 0u);
-			}
-		}
-		
-		if (++NumAttributesProcessed >= NumAttributes)
-		{
-			break; // We can early-out when we've looked at all the possible attributes
-		}
-	}
-}
diff --git a/Engine/Plugins/PCG/Shaders/Private/PCGShaderUtilsInternal.ush b/Engine/Plugins/PCG/Shaders/Private/PCGShaderUtilsInternal.ush
new file mode 100644
--- /dev/null
+++ b/Engine/Plugins/PCG/Shaders/Private/PCGShaderUtilsInternal.ush
@@ -0,0 +1,62 @@
+// Copyright Epic Games, Inc. All Rights Reserved.
+
+// Must match C++ constants in PCGComputeCommon.h
+#ifndef NUM_RESERVED_ATTRS
+#define NUM_RESERVED_ATTRS 32
+#endif
+#ifndef MAX_NUM_ATTRS
+#define MAX_NUM_ATTRS 128
+#endif
+
+// TODO pass in attribute IDs to copy?
+#define PCG_COPY_ATTRIBUTES_TO_OUTPUT(TargetPin, SourcePin, Target_DataIndex, Target_ElementIndex, Source_DataIndex, Source_ElementIndex, bMetadataOnly, bInitNonCopiedAttributes) \
+	{ \
+		const uint NumAttributes = TargetPin##_GetDataNumAttributesInternal(Target_DataIndex); \
+		uint NumAttributesProcessed = 0; \
+		\
+		const int FirstAttributeId = (bMetadataOnly) ? NUM_RESERVED_ATTRS : 0; \
+		\
+		for (int AttributeId = FirstAttributeId; AttributeId < MAX_NUM_ATTRS; ++AttributeId) \
+		{ \
+			const uint Target_Stride = TargetPin##_GetAttributeStrideInternal(Target_DataIndex, AttributeId); \
+			\
+			if (Target_Stride == 0) \
+			{ \
+				/* No output attribute to write to. */ \
+				continue; \
+			} \
+			\
+			const uint Source_Stride = SourcePin##_GetAttributeStrideInternal(Source_DataIndex, AttributeId); \
+			const uint Target_ElementAddress = TargetPin##_GetElementAddressInternal(Target_DataIndex, Target_ElementIndex, AttributeId); \
+			\
+			if (Source_Stride == Target_Stride) \
+			{ \
+				const uint Source_ElementAddress = SourcePin##_GetElementAddressInternal(Source_DataIndex, Source_ElementIndex, AttributeId); \
+				\
+				for (uint Offset = 0; Offset < Target_Stride; Offset += 4) \
+				{ \
+					TargetPin##_StoreBufferInternal(Target_ElementAddress + Offset, SourcePin##_LoadBufferInternal(Source_ElementAddress + Offset)); \
+				} \
+			} \
+			else if (bInitNonCopiedAttributes) \
+			{ \
+				for (uint Offset = 0; Offset < Target_Stride; Offset += 4) \
+				{ \
+					/* Initialize output data (only header part of buffer will be initialized). */ \
+					TargetPin##_StoreBufferInternal(Target_ElementAddress + Offset, 0u); \
+				} \
+			} \
+			\
+			if (++NumAttributesProcessed >= NumAttributes) \
+			{ \
+				/* We can early-out when we've looked at all the possible attributes. */ \
+				break; \
+			} \
+		} \
+	}
+
+#define PCG_COPY_METADATA_ATTRIBUTES_TO_OUTPUT(TargetPin, SourcePin, Target_DataIndex, Target_ElementIndex, Source_DataIndex, Source_ElementIndex) \
+	PCG_COPY_ATTRIBUTES_TO_OUTPUT(TargetPin, SourcePin, Target_DataIndex, Target_ElementIndex, Source_DataIndex, Source_ElementIndex, /*bMetadataOnly=*/true, /*bInitNonCopiedAttributes=*/true)
+
+#define PCG_COPY_ALL_ATTRIBUTES_TO_OUTPUT(TargetPin, SourcePin, Target_DataIndex, Target_ElementIndex, Source_DataIndex, Source_ElementIndex) \
+	PCG_COPY_ATTRIBUTES_TO_OUTPUT(TargetPin, SourcePin, Target_DataIndex, Target_ElementIndex, Source_DataIndex, Source_ElementIndex, /*bMetadataOnly=*/false, /*bInitNonCopiedAttributes=*/true)
diff --git a/Engine/Plugins/PCG/Source/PCG/Private/Compute/Elements/PCGCustomHLSL.cpp b/Engine/Plugins/PCG/Source/PCG/Private/Compute/Elements/PCGCustomHLSL.cpp
--- a/Engine/Plugins/PCG/Source/PCG/Private/Compute/Elements/PCGCustomHLSL.cpp
+++ b/Engine/Plugins/PCG/Source/PCG/Private/Compute/Elements/PCGCustomHLSL.cpp
@@ -1479,7 +1479,7 @@ bool UPCGCustomHLSLSettings::IsKernelValid(FPCGContext* InContext, bool bQuiet)
 			return false;
 		}
 
-		if (!GetGeneratedFunctions(/*OutFunctions=*/nullptr, ErrorTextPtr))
+		if (!ProcessShaderSource(/*InOutSource=*/nullptr, ErrorTextPtr))
 		{
 			if (ErrorTextPtr)
 			{
@@ -1721,16 +1721,10 @@ bool UPCGCustomHLSLSettings::AreAllOutputPinsWritten(FText* OutWarningText) cons
 	return true;
 }
 
-bool UPCGCustomHLSLSettings::GetGeneratedFunctions(FString* OutFunctions, FText* OutErrorText) const
+bool UPCGCustomHLSLSettings::ProcessShaderSource(FString* InOutSource, FText* OutErrorText) const
 {
 	TSet<TPair<FString, FString>> PinsRequiringElementCopy;
 
-	// Point processor nodes always copy input to output at the beginning of the kernel.
-	if (KernelType == EPCGKernelType::PointProcessor && GetPointProcessingInputPin() && GetFirstPointOutputPin())
-	{
-		PinsRequiringElementCopy.Emplace({ GetPointProcessingInputPin()->Properties.Label.ToString(), GetFirstPointOutputPin()->Properties.Label.ToString() });
-	}
-
 	const FString Source = ShaderFunctions + TEXT("\n") + ShaderSource;
 
 	/** Regex for parsing usages of CopyElementFrom function. */
@@ -1782,65 +1776,20 @@ bool UPCGCustomHLSLSettings::GetGeneratedFunctions(FString* OutFunctions, FText*
 		PinsRequiringElementCopy.Emplace({ SourcePinStr, TargetPinStr });
 	}
 
-	if (!OutFunctions)
+	if (!InOutSource)
 	{
 		return true;
 	}
 
-	FString Functions = TEXT("// Generated functions\n");
-
-	for (TPair<FString, FString>& Pair : PinsRequiringElementCopy)
+	// Replace function calls like Out_CopyElementFrom_In(...) with macro PCG_COPY_ALL_ATTRIBUTES_TO_OUTPUT(Out, In, ...).
+	for (const TPair<FString, FString>& SourcePinTargetPin : PinsRequiringElementCopy)
 	{
-		Functions += FString::Format(TEXT(
-			"\n"
-			"// Copy element from pin {0} to pin {1}.\n"
-			"void {1}_CopyElementFrom_{0}(uint {1}_DataIndex, uint {1}_ElementIndex, uint {0}_DataIndex, uint {0}_ElementIndex)\n"
-			"{\n"
-			"    // Copy elements across pins. Loop over all attribute headers, if the address is non-zero, then copy it from pin {0} to pin {1}.\n"
-			"    int NumAttributesRemaining = (int){1}_GetDataNumAttributesInternal({1}_DataIndex);\n"
-			"\n"
-			"    for (int AttributeId = 0; AttributeId < {2}; ++AttributeId)\n"
-			"    {\n"
-			"        const uint {1}_Stride = {1}_GetAttributeStrideInternal({1}_DataIndex, AttributeId);\n"
-			"\n"
-			"        if ({1}_Stride == 0)\n"
-			"        {\n"
-			"            // No output attribute to write to.\n"
-			"            continue;\n"
-			"        }\n"
-			"\n"
-			"        const uint {0}_Stride = {0}_GetAttributeStrideInternal({0}_DataIndex, AttributeId);\n"
-			"        const uint {1}_ElementAddress = {1}_GetElementAddressInternal({1}_DataIndex, {1}_ElementIndex, AttributeId);\n"
-			"\n"
-			"        if ({0}_Stride == {1}_Stride)\n"
-			"        {\n"
-			"            const uint {0}_ElementAddress = {0}_GetElementAddressInternal({0}_DataIndex, {0}_ElementIndex, AttributeId);\n"
-			"\n"
-			"            for (uint Offset = 0; Offset < {1}_Stride; Offset += 4)\n"
-			"            {\n"
-			"                {1}_StoreBufferInternal({1}_ElementAddress + Offset, {0}_LoadBufferInternal({0}_ElementAddress + Offset));\n"
-			"            }\n"
-			"\n"
-			"        }\n"
-			"        else\n"
-			"        {\n"
-			"            for (uint Offset = 0; Offset < {1}_Stride; Offset += 4)\n"
-			"            {\n"
-			"                // Initialize output data (only header part of buffer will be initialized).\n"
-			"                {1}_StoreBufferInternal({1}_ElementAddress + Offset, 0u);\n"
-			"            }\n"
-			"        }\n"
-			"\n"
-			"        if (--NumAttributesRemaining <= 0)\n"
-			"        {\n"
-			"            break; // We can early-out when we've looked at all the possible attributes\n"
-			"        }\n"
-			"    }\n"
-			"}\n"),
-			{ Pair.Key, Pair.Value, PCGComputeConstants::MAX_NUM_ATTRS });
+		InOutSource->ReplaceInline(
+			*FString::Format(TEXT("{1}_CopyElementFrom_{0}("), { SourcePinTargetPin.Key, SourcePinTargetPin.Value }),
+			*FString::Format(TEXT("PCG_COPY_ALL_ATTRIBUTES_TO_OUTPUT({1}, {0}, "), { SourcePinTargetPin.Key, SourcePinTargetPin.Value })
+		);
 	}
 
-	*OutFunctions = Functions;
 	return true;
 }
 
@@ -1849,10 +1798,7 @@ FString UPCGCustomHLSLSettings::GetCookedKernelSource(const FPCGKernelAttributeT
 	const FIntVector GroupSize(PCGComputeConstants::THREAD_GROUP_SIZE, 1, 1);
 
 	FString Source = ShaderSource;
-
-	FString Functions;
-	GetGeneratedFunctions(&Functions);
-	Functions += (TEXT("\n\n// User kernel functions\n") + ShaderFunctions);
+	FString Functions = TEXT("\n\n// User kernel functions\n") + ShaderFunctions;
 
 #if PLATFORM_WINDOWS
 	// Remove old-school stuff.
@@ -1860,6 +1806,9 @@ FString UPCGCustomHLSLSettings::GetCookedKernelSource(const FPCGKernelAttributeT
 	Functions.ReplaceInline(TEXT("\r"), TEXT(""));
 #endif
 
+	ProcessShaderSource(&Source);
+	ProcessShaderSource(&Functions);
+
 	// We inject attribute IDs directly into the source. This is most efficient and saves us passing them into the kernel. However
 	// the trade off is that a shader with a shared source will generate multiple variants if used in different compute graphs with
 	// different attribute IDs.
@@ -1887,6 +1836,7 @@ FString UPCGCustomHLSLSettings::GetCookedKernelSource(const FPCGKernelAttributeT
 		Includes += TEXT("#include \"/Engine/Private/ComputeShaderUtils.ush\"\n");
 		Includes += TEXT("#include \"/Engine/Private/Quaternion.ush\"\n");
 		Includes += TEXT("#include \"/Plugin/PCG/Private/PCGShaderUtils.ush\"\n");
+		Includes += TEXT("#include \"/Plugin/PCG/Private/PCGShaderUtilsInternal.ush\"\n");
 	}
 
 	const FString KernelFunc = FString::Printf(
@@ -1943,7 +1893,7 @@ FString UPCGCustomHLSLSettings::GetCookedKernelSource(const FPCGKernelAttributeT
 			KernelSpecificPreamble += FString::Format(TEXT(
 				"\n"
 				"// Point processor always initializes outputs by copying input data elements.\n"
-				"{1}_CopyElementFrom_{0}({1}_DataIndex, ElementIndex, {0}_DataIndex, ElementIndex);\n"),
+				"PCG_COPY_ALL_ATTRIBUTES_TO_OUTPUT({1}, {0}, {1}_DataIndex, ElementIndex, {0}_DataIndex, ElementIndex);\n"),
 				{ InputPin->Properties.Label.ToString(), OutputPin->Properties.Label.ToString() });
 		}
 	}
diff --git a/Engine/Plugins/PCG/Source/PCG/Private/Compute/Elements/PCGCustomHLSL.h b/Engine/Plugins/PCG/Source/PCG/Private/Compute/Elements/PCGCustomHLSL.h
--- a/Engine/Plugins/PCG/Source/PCG/Private/Compute/Elements/PCGCustomHLSL.h
+++ b/Engine/Plugins/PCG/Source/PCG/Private/Compute/Elements/PCGCustomHLSL.h
@@ -134,8 +134,8 @@ protected:
 	/** Will the ThreadCountMultiplier value be applied when calculating the dispatch thread count. */
 	bool IsThreadCountMultiplierInUse() const { return KernelType == EPCGKernelType::Custom && DispatchThreadCount != EPCGDispatchThreadCount::Fixed; }
 
-	/** Generate any functions that require code-gen for the shader source. */
-	virtual bool GetGeneratedFunctions(FString* OutFunctions, FText* OutErrorText = nullptr) const;
+	/** Apply any code gen or transformations or expansions to the shader source. */
+	virtual bool ProcessShaderSource(FString* InOutSource, FText* OutErrorText = nullptr) const;
 
 protected:
 	UPROPERTY(EditAnywhere, Category = "Settings")
