PCG: GPU - cache results of building data descriptions.

Parsing data collections to build the info can be very expensive.

#rnx

diff --git a/Engine/Plugins/PCG/Source/PCG/Private/Compute/DataInterfaces/PCGDataCollectionDataInterface.cpp b/Engine/Plugins/PCG/Source/PCG/Private/Compute/DataInterfaces/PCGDataCollectionDataInterface.cpp
--- a/Engine/Plugins/PCG/Source/PCG/Private/Compute/DataInterfaces/PCGDataCollectionDataInterface.cpp
+++ b/Engine/Plugins/PCG/Source/PCG/Private/Compute/DataInterfaces/PCGDataCollectionDataInterface.cpp
@@ -431,7 +431,7 @@ UComputeDataProvider* UPCGDataCollectionDataInterface::CreateDataProvider(TObjec
 	UPCGDataBinding* Binding = CastChecked<UPCGDataBinding>(InBinding);
 
 	FPCGDataCollectionDesc PinDataDesc;
-	ensure(ProducerSettings->ComputeOutputPinDataDesc(OutputPinLabel, Binding, PinDataDesc));
+	ensure(Binding->GetOutputPinDataDesc_GameThread(OutputPinLabel, ProducerSettings, PinDataDesc));
 
 	if (PCGComputeHelpers::IsBufferSizeTooLarge(PinDataDesc.ComputePackedSizeBytes()))
 	{
diff --git a/Engine/Plugins/PCG/Source/PCG/Private/Compute/DataInterfaces/PCGDataCollectionUploadDataInterface.cpp b/Engine/Plugins/PCG/Source/PCG/Private/Compute/DataInterfaces/PCGDataCollectionUploadDataInterface.cpp
--- a/Engine/Plugins/PCG/Source/PCG/Private/Compute/DataInterfaces/PCGDataCollectionUploadDataInterface.cpp
+++ b/Engine/Plugins/PCG/Source/PCG/Private/Compute/DataInterfaces/PCGDataCollectionUploadDataInterface.cpp
@@ -184,12 +184,19 @@ UComputeDataProvider* UPCGDataCollectionUploadDataInterface::CreateDataProvider(
 	Provider->Binding = Binding;
 
 	// Pick the data items from input data collection using any of the compute graph element virtual input pin labels.
-	// TODO: Ideally we could call ProducerSettings->ComputeOutputPinDataDesc() but some settings do not have associated nodes/pins.
 	check(!DownstreamInputPinLabelAliases.IsEmpty());
-	Provider->PinDesc = FPCGDataCollectionDesc::BuildFromInputDataCollectionAndInputPinLabel(
-		Binding->InputDataCollection,
-		DownstreamInputPinLabelAliases[0],
-		Binding);
+
+	if (ProducerSettings && Cast<UPCGNode>(ProducerSettings->GetOuter()))
+	{
+		ensure(Binding->GetOutputPinDataDesc_GameThread(OutputPinLabel, ProducerSettings, Provider->PinDesc));
+	}
+	else
+	{
+		Provider->PinDesc = FPCGDataCollectionDesc::BuildFromInputDataCollectionAndInputPinLabel(
+			Binding->InputDataCollection,
+			DownstreamInputPinLabelAliases[0],
+			Binding);
+	}
 
 	Provider->DownstreamInputPinLabels = DownstreamInputPinLabelAliases;
 
diff --git a/Engine/Plugins/PCG/Source/PCG/Private/Compute/Elements/PCGCustomHLSL.cpp b/Engine/Plugins/PCG/Source/PCG/Private/Compute/Elements/PCGCustomHLSL.cpp
--- a/Engine/Plugins/PCG/Source/PCG/Private/Compute/Elements/PCGCustomHLSL.cpp
+++ b/Engine/Plugins/PCG/Source/PCG/Private/Compute/Elements/PCGCustomHLSL.cpp
@@ -224,7 +224,7 @@ int UPCGCustomHLSLSettings::GetProcessingElemCountForInputPin(const UPCGPin* Inp
 			check(UpstreamSettings);
 
 			FPCGDataCollectionDesc PinDesc;
-			if (ensure(UpstreamSettings->ComputeOutputPinDataDesc(UpstreamOutputPin, Binding, PinDesc)))
+			if (ensure(Binding->GetOutputPinDataDesc_GameThread(UpstreamOutputPin, UpstreamSettings, PinDesc)))
 			{
 				ProcessingElemCount += PinDesc.ComputeDataElementCount(InputPin->Properties.AllowedTypes);
 			}
@@ -258,7 +258,7 @@ int UPCGCustomHLSLSettings::ComputeKernelThreadCount(const UPCGDataBinding* Bind
 			if (const UPCGPin* OutputPin = GetFirstOutputPin())
 			{
 				FPCGDataCollectionDesc Desc;
-				if (ensure(ComputeOutputPinDataDesc(OutputPin, Binding, Desc)))
+				if (ensure(Binding->GetOutputPinDataDesc_GameThread(OutputPin, this, Desc)))
 				{
 					ThreadCount = Desc.ComputeDataElementCount(EPCGDataType::Any);
 				}
@@ -1535,7 +1535,7 @@ bool UPCGCustomHLSLSettings::AreKernelAttributesValid(FPCGContext* InContext, FT
 			check(OutputPin);
 
 			FPCGDataCollectionDesc PinDesc;
-			if (ensure(ComputeOutputPinDataDesc(OutputPin, DataBinding, PinDesc)))
+			if (ensure(DataBinding->GetOutputPinDataDesc_GameThread(OutputPin, this, PinDesc)))
 			{
 				OutputPinDescs.Add(OutputPin->Properties.Label, PinDesc);
 			}
diff --git a/Engine/Plugins/PCG/Source/PCG/Private/Compute/PCGDataBinding.cpp b/Engine/Plugins/PCG/Source/PCG/Private/Compute/PCGDataBinding.cpp
--- a/Engine/Plugins/PCG/Source/PCG/Private/Compute/PCGDataBinding.cpp
+++ b/Engine/Plugins/PCG/Source/PCG/Private/Compute/PCGDataBinding.cpp
@@ -40,6 +40,12 @@ bool UPCGDataBinding::ComputeCPUOutputPinDataDesc(const UPCGPin* OutputPin, FPCG
 	check(OutputPin);
 	check(Graph);
 
+	if (const FPCGDataCollectionDesc* FoundDescription = OutputPinDataDescriptionCache.Find(OutputPin))
+	{
+		OutDesc = *FoundDescription;
+		return true;
+	}
+
 	// Find the unique alias for the virtual pin created during compilation to receive the data from the upstream OutputPin.
 	// OutputPin can either be on a CPU node or on a GPU node in a different compute graph.
 	const FName* FoundPinAlias = Graph->OutputCPUPinToInputGPUPinAlias.Find(OutputPin);
@@ -51,6 +57,8 @@ bool UPCGDataBinding::ComputeCPUOutputPinDataDesc(const UPCGPin* OutputPin, FPCG
 	// Create description from compute graph element input data.
 	OutDesc = FPCGDataCollectionDesc::BuildFromInputDataCollectionAndInputPinLabel(InputDataCollection, *FoundPinAlias, this);
 
+	OutputPinDataDescriptionCache.Add(OutputPin, OutDesc);
+
 	return true;
 }
 
@@ -85,6 +93,37 @@ void UPCGDataBinding::ReceiveDataFromGPU_GameThread(
 	}
 }
 
+bool UPCGDataBinding::GetOutputPinDataDesc_GameThread(FName InProducerPinLabel, const UPCGSettings* InProducerSettings, FPCGDataCollectionDesc& OutDescription) const
+{
+	check(InProducerSettings);
+	const UPCGNode* Node = CastChecked<UPCGNode>(InProducerSettings->GetOuter());
+
+	return GetOutputPinDataDesc_GameThread(Node->GetOutputPin(InProducerPinLabel), InProducerSettings, OutDescription);
+}
+
+bool UPCGDataBinding::GetOutputPinDataDesc_GameThread(const UPCGPin* InProducerPin, const UPCGSettings* InProducerSettings, FPCGDataCollectionDesc& OutDescription) const
+{
+	ensure(IsInGameThread());
+
+	if (const FPCGDataCollectionDesc* FoundDescription = OutputPinDataDescriptionCache.Find(InProducerPin))
+	{
+		OutDescription = *FoundDescription;
+		return true;
+	}
+
+	FPCGDataCollectionDesc Description;
+	if (InProducerSettings->ComputeOutputPinDataDesc(InProducerPin, this, Description))
+	{
+		OutputPinDataDescriptionCache.Add(InProducerPin, Description);
+		OutDescription = MoveTemp(Description);
+		return true;
+	}
+	else
+	{
+		return false;
+	}
+}
+
 void UPCGDataBinding::InitializeInputData(const FPCGDataCollection& InComputeGraphElementInputData)
 {
 	TRACE_CPUPROFILER_EVENT_SCOPE(UPCGDataBinding::InitializeInputData);
@@ -292,7 +331,7 @@ void UPCGDataBinding::DebugLogDataDescriptions()
 			for (const UPCGPin* Pin : Node->GetOutputPins())
 			{
 				FPCGDataCollectionDesc DataDesc;
-				if (ensure(Settings->ComputeOutputPinDataDesc(Pin, this, DataDesc)))
+				if (ensure(GetOutputPinDataDesc_GameThread(Pin, Settings, DataDesc)))
 				{
 					UE_LOG(LogPCG, Display, TEXT("\tOutput Pin: %s (%d data)"), *Pin->Properties.Label.ToString(), DataDesc.DataDescs.Num());
 					LogDataDescription(DataDesc);
diff --git a/Engine/Plugins/PCG/Source/PCG/Private/Compute/PCGDataBinding.h b/Engine/Plugins/PCG/Source/PCG/Private/Compute/PCGDataBinding.h
--- a/Engine/Plugins/PCG/Source/PCG/Private/Compute/PCGDataBinding.h
+++ b/Engine/Plugins/PCG/Source/PCG/Private/Compute/PCGDataBinding.h
@@ -70,6 +70,10 @@ public:
 
 	void ReceiveDataFromGPU_GameThread(const UPCGProxyForGPUData* InData, const UPCGSettings* InProducerSettings, EPCGExportMode InExportMode, FName InPinLabel, FName InPinLabelAlias);
 
+	/** Get description of data produced by an output pin. */
+	bool GetOutputPinDataDesc_GameThread(FName InProducerPinLabel, const UPCGSettings* InProducerSettings, FPCGDataCollectionDesc& OutDescription) const;
+	bool GetOutputPinDataDesc_GameThread(const UPCGPin* InProducerPin, const UPCGSettings* InProducerSettings, FPCGDataCollectionDesc& OutDescription) const;
+
 private:
 	void InitializeInputData(const FPCGDataCollection& InComputeGraphElementInputData);
 
@@ -119,4 +123,7 @@ protected:
 
 	UPROPERTY()
 	TArray<FString> StringTable;
+
+	/** Cache of data descriptions computed for output pins to amortized cost of building data descriptions. */
+	mutable TMap<const UPCGPin*, FPCGDataCollectionDesc> OutputPinDataDescriptionCache;
 };
diff --git a/Engine/Plugins/PCG/Source/PCG/Private/Compute/PCGDataForGPU.cpp b/Engine/Plugins/PCG/Source/PCG/Private/Compute/PCGDataForGPU.cpp
--- a/Engine/Plugins/PCG/Source/PCG/Private/Compute/PCGDataForGPU.cpp
+++ b/Engine/Plugins/PCG/Source/PCG/Private/Compute/PCGDataForGPU.cpp
@@ -634,7 +634,7 @@ namespace PCGDataForGPUHelpers
 
 			// Add data from connected upstream output pin.
 			FPCGDataCollectionDesc EdgeDesc;
-			if (ensure(UpstreamSettings->ComputeOutputPinDataDesc(UpstreamOutputPin, Binding, EdgeDesc)))
+			if (ensure(Binding->GetOutputPinDataDesc_GameThread(UpstreamOutputPin, UpstreamSettings, EdgeDesc)))
 			{
 				PinDesc.Combine(EdgeDesc);
 			}
diff --git a/Engine/Plugins/PCG/Source/PCG/Private/Elements/PCGCopyPoints.cpp b/Engine/Plugins/PCG/Source/PCG/Private/Elements/PCGCopyPoints.cpp
--- a/Engine/Plugins/PCG/Source/PCG/Private/Elements/PCGCopyPoints.cpp
+++ b/Engine/Plugins/PCG/Source/PCG/Private/Elements/PCGCopyPoints.cpp
@@ -114,7 +114,7 @@ int UPCGCopyPointsSettings::ComputeKernelThreadCount(const UPCGDataBinding* Bind
 	int ThreadCount = 0;
 
 	FPCGDataCollectionDesc OutputPinDesc;
-	if (ensure(ComputeOutputPinDataDesc(OutPin, Binding, OutputPinDesc)))
+	if (ensure(Binding->GetOutputPinDataDesc_GameThread(OutPin, this, OutputPinDesc)))
 	{
 		for (const FPCGDataDesc& DataDesc : OutputPinDesc.DataDescs)
 		{
