[PCG] GPU: Add experimental GPU support for attribute partition support.

* Adds a path for reading back & updating data description as a form of compaction to fix over-allocation of instances in simple cases.
* Only supports partitioning string keys - needs to know statically how many partitions there will be.
* The point order is not stable currently, can reorder points.

This is a rudimentary approach in terms of the algorithm and the unstable order likely won't cut it long term, but it gets us off the ground and involved fixing and improving a number of issues around data descriptions and around graph splits that will be a foundation for next steps.

#rnx

diff --git a/Engine/Plugins/PCG/Shaders/Private/Elements/PCGMetadataPartition.usf b/Engine/Plugins/PCG/Shaders/Private/Elements/PCGMetadataPartition.usf
new file mode 100644
--- /dev/null
+++ b/Engine/Plugins/PCG/Shaders/Private/Elements/PCGMetadataPartition.usf
@@ -0,0 +1,53 @@
+// Copyright Epic Games, Inc. All Rights Reserved.
+
+// Kernel for GPU implementation of Metadata Partition node.
+
+#include "/Engine/Private/ComputeShaderUtils.ush"
+#include "/Engine/Private/Quaternion.ush"
+#include "/Plugin/PCG/Private/PCGShaderUtils.ush"
+#include "/Plugin/PCG/Private/PCGShaderUtilsInternal.ush"
+
+// TODO: should be code-gen'd from PCGComputeCommon.h
+#define MAX_NUM_ATTRS 128
+#define NUM_RESERVED_ATTRS 32
+
+[numthreads(64, 1, 1)]
+void Main(uint3 GroupId : SV_GroupID, uint GroupIndex : SV_GroupIndex)
+{
+	// Mark the kernel as having executed. Must run before we early out via thread index, because the kernel is still 'executed' even if the number of
+	// threads to iterate on is zero. Even if GetNumThreads() returns 0, the kernel will still have been dispatched on a single thread to set this flag.
+	if (all(GroupId == 0) && GroupIndex == 0)
+	{
+		Out_SetAsExecutedInternal();
+	}
+
+	const uint ThreadIndex = GetUnWrappedDispatchThreadId(GroupId, GroupIndex, 64);
+
+	uint In_DataIndex, In_ElementIndex;
+	if (!In_GetThreadData(ThreadIndex, In_DataIndex, In_ElementIndex))
+	{
+		return;
+	}
+
+	if (In_IsPointRemoved(In_DataIndex, In_ElementIndex))
+	{
+		return;
+	}
+
+	const int PartitionAttributeId = MetadataPartition_GetPartitionAttributeId();
+	const int NumPartitions = MetadataPartition_GetNumPartitions();
+	
+	const int PartitionValueKey = In_GetStringKey(In_DataIndex, In_ElementIndex, PartitionAttributeId);
+
+	const int PartitionIndex = MetadataPartition_GetPartitionIndex(PartitionValueKey);
+
+	const uint Out_DataIndex = In_DataIndex * NumPartitions + PartitionIndex;
+	const int Out_ElementIndex = Out_AddToElementCounterInternal(Out_DataIndex, 1);
+	if (Out_ElementIndex == -1)
+	{
+		return;
+	}
+
+	// Write output point.
+	PCG_COPY_ALL_ATTRIBUTES_TO_OUTPUT(Out, In, Out_DataIndex, Out_ElementIndex, In_DataIndex, In_ElementIndex);
+}
diff --git a/Engine/Plugins/PCG/Shaders/Private/PCGDataCollectionAdaptor.usf b/Engine/Plugins/PCG/Shaders/Private/PCGDataCollectionAdaptor.usf
new file mode 100644
--- /dev/null
+++ b/Engine/Plugins/PCG/Shaders/Private/PCGDataCollectionAdaptor.usf
@@ -0,0 +1,83 @@
+// Copyright Epic Games, Inc. All Rights Reserved.
+
+#include "/Engine/Private/Common.ush"
+#include "/Engine/Public/Platform.ush"
+
+// Buffer format (sizes in bytes):
+//
+// -- START HEADER --
+// PackedExecutionFlagAndNumData (4)
+// For each data:
+//   TypeId (4)
+//   Attribute Count (4)
+//   Element Count (4)
+//   For attribute ID 0 to MAX_NUM_ATTR:
+//     PackedIdAndStride (4)
+//     AttributeElementStartAddress (4)
+// -- END HEADER --
+// For each data:
+//   For each present attribute:
+//     <Packed element data> (element count * stride)
+
+#define DATA_COLLECTION_HEADER_SIZE_BYTES 4
+#define DATA_HEADER_SIZE_BYTES 1036
+#define ATTRIBUTE_HEADER_SIZE_BYTES 8
+
+StructuredBuffer<int2> InSourceToTargetAttributeId;
+StructuredBuffer<uint> InDataElementCounts;
+ByteAddressBuffer InDataCollection;
+RWByteAddressBuffer OutDataCollection;
+uint InNumData;
+uint InNumRemappedAttributes;
+
+void Out_SetNumElements(uint InDataIndex, uint InNumElements)
+{
+	const uint Address =
+		DATA_COLLECTION_HEADER_SIZE_BYTES +
+		InDataIndex * DATA_HEADER_SIZE_BYTES +
+		/*TypeId*/4 + /*Attribute Count*/4;
+
+	return OutDataCollection.Store(Address, InNumElements);
+}
+
+uint GetAttributeHeaderAddress(uint InDataIndex, int InAttributeId)
+{
+	return DATA_COLLECTION_HEADER_SIZE_BYTES +
+		InDataIndex * DATA_HEADER_SIZE_BYTES +
+		/*This TypeId*/4 + /*This Attribute Count*/4 + /*This Element Count*/4 +
+		InAttributeId * ATTRIBUTE_HEADER_SIZE_BYTES;
+}
+
+[numthreads(NUM_THREADS_PER_GROUP, 1, 1)]
+void MainCS(uint InThreadIndex : SV_DispatchThreadID)
+{
+	if (InThreadIndex >= InNumRemappedAttributes * InNumData + InNumData)
+	{
+		return;
+	}
+	
+	// Distribute update tasks among threads.
+
+	if (InThreadIndex < InNumData)
+	{
+		// Task: Update data element count.
+
+		const uint DataIndexToUpdate = InThreadIndex;
+		Out_SetNumElements(DataIndexToUpdate, InDataElementCounts[DataIndexToUpdate]);
+	}
+	else if ((InThreadIndex - InNumData) < (InNumRemappedAttributes * InNumData))
+	{
+		// Task: Remap attribute. Set attribute ID and address in output data.
+
+		const uint RemapDataIndex = (InThreadIndex - InNumData) / InNumRemappedAttributes;
+		const uint RemapAttrIndex = (InThreadIndex - InNumData) % InNumRemappedAttributes;
+
+		const int2 SourceIdToTargetId = InSourceToTargetAttributeId[RemapAttrIndex];
+		
+		const uint SourceAddress = GetAttributeHeaderAddress(RemapDataIndex, SourceIdToTargetId[0]);
+		const uint TargetAddress = GetAttributeHeaderAddress(RemapDataIndex, SourceIdToTargetId[1]);
+
+		// Write ID, stride, address to target to apply remap.
+		OutDataCollection.Store2(TargetAddress, InDataCollection.Load2(SourceAddress));
+	}
+}
diff --git a/Engine/Plugins/PCG/Shaders/Private/PCGDataCollectionDataInterface.ush b/Engine/Plugins/PCG/Shaders/Private/PCGDataCollectionDataInterface.ush
--- a/Engine/Plugins/PCG/Shaders/Private/PCGDataCollectionDataInterface.ush
+++ b/Engine/Plugins/PCG/Shaders/Private/PCGDataCollectionDataInterface.ush
@@ -1,7 +1,7 @@
 // Copyright Epic Games, Inc. All Rights Reserved.
 
 // Buffer format (sizes in bytes):
-
+//
 // -- START HEADER --
 // PackedExecutionFlagAndNumData (4)
 // For each data:
@@ -12,14 +12,12 @@
 //     PackedIdAndStride (4)
 //     AttributeElementStartAddress (4)
 // -- END HEADER --
-//
-//   For all present attributes across all data:
+// For each data:
+//   For each present attribute:
 //     <Packed element data> (element count * stride)
 
 #include "/Engine/Private/Quaternion.ush"
 
-RWByteAddressBuffer {DataInterfaceName}_DataCollectionBuffer;
-
 // TODO: These defines could be code-gen'd to stay in sync with our PCGComputeConstants, see UComputeDataInterface::GetDefines()
 #define POINT_POSITION_ID 0
 #define POINT_ROTATION_ID 1
@@ -45,6 +43,11 @@ RWByteAddressBuffer {DataInterfaceName}_DataCollectionBuffer;
 #define MAX_NUM_ATTRS 128
 #define NUM_RESERVED_ATTRS 32
 
+RWByteAddressBuffer {DataInterfaceName}_DataCollectionBuffer;
+
+RWStructuredBuffer<uint> {DataInterfaceName}_ElementCountersBuffer;
+uint {DataInterfaceName}_ElementCountersPresent;
+
 // #################### HEADER READERS ##########################
 
 uint GetNumData_{DataInterfaceName}()
@@ -63,6 +66,19 @@ uint GetNumElements_{DataInterfaceName}(uint InDataIndex)
 	return {DataInterfaceName}_DataCollectionBuffer.Load(ReadAddress);
 }
 
+int AddToElementCounterInternal_{DataInterfaceName}(uint InDataIndex, uint InElementCount)
+{
+	if ({DataInterfaceName}_ElementCountersPresent == 0)
+	{
+		return -1;
+	}
+
+	uint ValueBefore;
+	InterlockedAdd({DataInterfaceName}_ElementCountersBuffer[InDataIndex], InElementCount, ValueBefore);
+
+	return (int)ValueBefore;
+}
+
 bool GetThreadData_{DataInterfaceName}(uint InThreadIndex, out uint OutDataIndex, out uint OutElementIndex)
 {
 	int ElementIndex = InThreadIndex;
diff --git a/Engine/Plugins/PCG/Source/PCG/Internal/Compute/Data/PCGProxyForGPUData.h b/Engine/Plugins/PCG/Source/PCG/Internal/Compute/Data/PCGProxyForGPUData.h
--- a/Engine/Plugins/PCG/Source/PCG/Internal/Compute/Data/PCGProxyForGPUData.h
+++ b/Engine/Plugins/PCG/Source/PCG/Internal/Compute/Data/PCGProxyForGPUData.h
@@ -50,8 +50,13 @@ public:
 	/** Returns a description of this data. Does not trigger a GPU->CPU readback. */
 	bool GetDescription(FPCGDataDesc& OutDescription) const;
 
+	/** Apply the given per-data element counts to the data description. */
+	void UpdateElementCountsFromReadback(const TArray<uint32>& InElementCounts);
+
+	TSharedPtr<const FPCGProxyForGPUDataCollection> GetGPUInfo() const;
+
 protected:
-	TSharedPtr<FPCGProxyForGPUDataCollection> GetGPUInfoValidated() const;
+	TSharedPtr<FPCGProxyForGPUDataCollection> GetGPUInfoMutable() const;
 
 protected:
 	UPROPERTY()
diff --git a/Engine/Plugins/PCG/Source/PCG/Internal/Compute/PCGComputeCommon.h b/Engine/Plugins/PCG/Source/PCG/Internal/Compute/PCGComputeCommon.h
--- a/Engine/Plugins/PCG/Source/PCG/Internal/Compute/PCGComputeCommon.h
+++ b/Engine/Plugins/PCG/Source/PCG/Internal/Compute/PCGComputeCommon.h
@@ -113,4 +113,7 @@ namespace PCGComputeHelpers
 	bool IsBufferSizeTooLarge(uint64 InBufferSizeBytes, bool bInLogError = true);
 
 	bool ComputeDataDescription(const UPCGData* InData, const UPCGDataBinding* InBinding, FPCGDataDesc& OutDescription);
+
+	int32 GetAttributeIdFromMetadataAttributeIndex(int32 InAttributeIndex);
+	int32 GetMetadataAttributeIndexFromAttributeId(int32 InAttributeId);
 }
diff --git a/Engine/Plugins/PCG/Source/PCG/Internal/Compute/PCGDataForGPU.h b/Engine/Plugins/PCG/Source/PCG/Internal/Compute/PCGDataForGPU.h
--- a/Engine/Plugins/PCG/Source/PCG/Internal/Compute/PCGDataForGPU.h
+++ b/Engine/Plugins/PCG/Source/PCG/Internal/Compute/PCGDataForGPU.h
@@ -81,7 +81,8 @@ struct FPCGKernelAttributeTable
 
 	int32 GetAttributeId(FName InName, EPCGKernelAttributeType InType) const;
 
-	void AddAttribute(FName InName, EPCGKernelAttributeType InType);
+	/** Adds an attribute of given name and type. Returns added attribute ID. */
+	int32 AddAttribute(FName InName, EPCGKernelAttributeType InType);
 
 #if PCG_KERNEL_LOGGING_ENABLED
 	void DebugLog();
@@ -107,22 +108,28 @@ public:
 	{
 	}
 
-	explicit FPCGKernelAttributeDesc(int32 InIndex, EPCGKernelAttributeType InType, FName InName, TArray<int32>&& InUniqueStringKeys)
+	explicit FPCGKernelAttributeDesc(int32 InIndex, EPCGKernelAttributeType InType, FName InName, const TArray<int32>* InUniqueStringKeys)
 		: Index(InIndex)
 		, Type(InType)
 		, Name(InName)
-		, UniqueStringKeys(MoveTemp(InUniqueStringKeys))
 	{
+		if (InUniqueStringKeys)
+		{
+			UniqueStringKeys = *InUniqueStringKeys;
+		}
 	}
 
 	const TArray<int32>& GetUniqueStringKeys() const { return UniqueStringKeys; }
 
 	void AddUniqueStringKeys(const TArray<int32>& InOtherStringKeys);
 
+	void SetStringKeys(const TArrayView<int32>& InStringKeys);
+
 	bool operator==(const FPCGKernelAttributeDesc& Other) const;
 
 public:
 	// Attribute Id
+	// TODO: Rename to Id to disambiguate absolute attribute Id from index into metadata attribute table.
 	int32 Index = INDEX_NONE;
 
 	EPCGKernelAttributeType Type = EPCGKernelAttributeType::Float;
@@ -148,6 +155,8 @@ public:
 
 	bool ContainsAttribute(FName InAttributeName) const;
 
+	void AddAttribute(FPCGKernelAttributeKey InAttribute, const UPCGDataBinding* InBinding, const TArray<int32>* InOptionalUniqueStringKeys = nullptr);
+
 public:
 	EPCGDataType Type = EPCGDataType::Point;
 	TArray<FPCGKernelAttributeDesc> AttributeDescs;
@@ -196,7 +205,9 @@ struct FPCGDataCollectionDesc
 	bool ContainsAttributeOnAnyData(FName InAttributeName) const;
 
 	/** Makes attribute present on all data. If data has existing attribute with same name then the given type will be applied. */
-	void AddAttributeToAllData(FPCGKernelAttributeKey InAttribute, const UPCGDataBinding* InBinding, TArray<int32>* InOptionalUniqueStringKeys = nullptr);
+	void AddAttributeToAllData(FPCGKernelAttributeKey InAttribute, const UPCGDataBinding* InBinding, const TArray<int32>* InOptionalUniqueStringKeys = nullptr);
+
+	void GetUniqueStringKeyValues(int32 InAttributeId, TArray<int32>& OutUniqueStringKeys) const;
 
 	/** Description of each data in this data collection. */
 	TArray<FPCGDataDesc> DataDescs;
@@ -213,8 +224,14 @@ public:
 
 	const TRefCountPtr<FRDGPooledBuffer>& GetBuffer() const { return Buffer; }
 
+	uint32 GetBufferSizeBytes() const { return BufferSizeBytes; }
+
 	const FPCGDataCollectionDesc& GetDescription() const { return Description; }
 
+	void UpdateElementCountsFromReadback(const TArray<uint32>& InElementCounts);
+
+	const TArray<FString>& GetStringTable() const { return StringTable; }
+
 protected:
 	/** Persistent GPU buffer that can be read back. Buffer will be freed when this ref count is 0. */
 	TRefCountPtr<FRDGPooledBuffer> Buffer;
diff --git a/Engine/Plugins/PCG/Source/PCG/Private/Compute/Data/PCGProxyForGPUData.cpp b/Engine/Plugins/PCG/Source/PCG/Private/Compute/Data/PCGProxyForGPUData.cpp
--- a/Engine/Plugins/PCG/Source/PCG/Private/Compute/Data/PCGProxyForGPUData.cpp
+++ b/Engine/Plugins/PCG/Source/PCG/Private/Compute/Data/PCGProxyForGPUData.cpp
@@ -20,7 +20,7 @@ void UPCGProxyForGPUData::AddToCrc(FArchiveCrc32& Ar, bool bFullDataCrc) const
 
 EPCGDataType UPCGProxyForGPUData::GetUnderlyingDataType() const
 {
-	if (TSharedPtr<FPCGProxyForGPUDataCollection> DataCollectionGPU = GetGPUInfoValidated())
+	if (TSharedPtr<const FPCGProxyForGPUDataCollection> DataCollectionGPU = GetGPUInfo())
 	{
 		return DataCollectionGPU->GetDescription().DataDescs[DataIndexInCollection].Type;
 	}
@@ -35,7 +35,7 @@ TSharedPtr<const FPCGProxyForGPUDataCollection> UPCGProxyForGPUData::GetInputDat
 
 int UPCGProxyForGPUData::GetElementCount() const
 {
-	if (TSharedPtr<FPCGProxyForGPUDataCollection> DataCollectionGPU = GetGPUInfoValidated())
+	if (TSharedPtr<const FPCGProxyForGPUDataCollection> DataCollectionGPU = GetGPUInfo())
 	{
 		return DataCollectionGPU->GetDescription().DataDescs[DataIndexInCollection].ElementCount;
 	}
@@ -45,7 +45,7 @@ int UPCGProxyForGPUData::GetElementCount() const
 
 bool UPCGProxyForGPUData::GetDescription(FPCGDataDesc& OutDescription) const
 {
-	if (TSharedPtr<FPCGProxyForGPUDataCollection> DataCollectionGPU = GetGPUInfoValidated())
+	if (TSharedPtr<const FPCGProxyForGPUDataCollection> DataCollectionGPU = GetGPUInfo())
 	{
 		OutDescription = DataCollectionGPU->GetDescription().DataDescs[DataIndexInCollection];
 		return true;
@@ -54,9 +54,17 @@ bool UPCGProxyForGPUData::GetDescription(FPCGDataDesc& OutDescription) const
 	return false;
 }
 
+void UPCGProxyForGPUData::UpdateElementCountsFromReadback(const TArray<uint32>& InElementCounts)
+{
+	if (TSharedPtr<FPCGProxyForGPUDataCollection> GPUInfo = GetGPUInfoMutable())
+	{
+		GPUInfo->UpdateElementCountsFromReadback(InElementCounts);
+	}
+}
+
 UPCGProxyForGPUData::FReadbackResult UPCGProxyForGPUData::GetCPUData(FPCGContext* InContext) const
 {
-	TSharedPtr<FPCGProxyForGPUDataCollection> DataOnGPU = GetGPUInfoValidated();
+	TSharedPtr<FPCGProxyForGPUDataCollection> DataOnGPU = GetGPUInfoMutable();
 	if (!DataOnGPU)
 	{
 		UE_LOG(LogPCG, Error, TEXT("Data collection lost!"));
@@ -73,7 +81,12 @@ UPCGProxyForGPUData::FReadbackResult UPCGProxyForGPUData::GetCPUData(FPCGContext
 	return FReadbackResult{ .bComplete = true, .Data = ResultData };
 }
 
-TSharedPtr<FPCGProxyForGPUDataCollection> UPCGProxyForGPUData::GetGPUInfoValidated() const
+TSharedPtr<const FPCGProxyForGPUDataCollection> UPCGProxyForGPUData::GetGPUInfo() const
+{
+	return GetGPUInfoMutable();
+}
+
+TSharedPtr<FPCGProxyForGPUDataCollection> UPCGProxyForGPUData::GetGPUInfoMutable() const
 {
 	if (!DataCollectionOnGPU)
 	{
diff --git a/Engine/Plugins/PCG/Source/PCG/Private/Compute/DataInterfaces/Elements/PCGMetadataPartitionDataInterface.cpp b/Engine/Plugins/PCG/Source/PCG/Private/Compute/DataInterfaces/Elements/PCGMetadataPartitionDataInterface.cpp
new file mode 100644
--- /dev/null
+++ b/Engine/Plugins/PCG/Source/PCG/Private/Compute/DataInterfaces/Elements/PCGMetadataPartitionDataInterface.cpp
@@ -0,0 +1,165 @@
+// Copyright Epic Games, Inc. All Rights Reserved.
+
+#include "PCGMetadataPartitionDataInterface.h"
+
+#include "Compute/PCGComputeCommon.h"
+#include "Compute/PCGDataBinding.h"
+#include "Elements/Metadata/PCGMetadataPartition.h"
+
+#include "GlobalRenderResources.h"
+#include "RenderGraphBuilder.h"
+#include "RenderGraphResources.h"
+#include "RHIResources.h"
+#include "ShaderCompilerCore.h"
+#include "ShaderCore.h"
+#include "ShaderParameterMetadataBuilder.h"
+#include "ComputeFramework/ShaderParamTypeDefinition.h"
+
+void UPCGMetadataPartitionDataInterface::GetSupportedInputs(TArray<FShaderFunctionDefinition>& OutFunctions) const
+{
+	OutFunctions.AddDefaulted_GetRef()
+		.SetName(TEXT("MetadataPartition_GetPartitionAttributeId"))
+		.AddReturnType(EShaderFundamentalType::Int);
+
+	OutFunctions.AddDefaulted_GetRef()
+		.SetName(TEXT("MetadataPartition_GetNumPartitions"))
+		.AddReturnType(EShaderFundamentalType::Uint);
+
+	OutFunctions.AddDefaulted_GetRef()
+		.SetName(TEXT("MetadataPartition_GetPartitionIndex"))
+		.AddReturnType(EShaderFundamentalType::Uint)
+		.AddParam(EShaderFundamentalType::Uint); // InAttributeValue
+
+	OutFunctions.AddDefaulted_GetRef()
+		.SetName(TEXT("MetadataPartition_GetElementIndexToWriteTo"))
+		.AddReturnType(EShaderFundamentalType::Uint)
+		.AddParam(EShaderFundamentalType::Int) // InInputDataIndex
+		.AddParam(EShaderFundamentalType::Uint) // InPartitionIndex
+		.AddParam(EShaderFundamentalType::Uint); // InNumPartitions
+}
+
+BEGIN_SHADER_PARAMETER_STRUCT(FPCGMetadataPartitionDataInterfaceParameters,)
+	SHADER_PARAMETER(int32, PartitionAttributeId)
+	SHADER_PARAMETER(uint32, NumPartitions)
+	SHADER_PARAMETER_RDG_BUFFER_SRV(StructuredBuffer<uint32>, AttributeValueToPartitionIndex)
+END_SHADER_PARAMETER_STRUCT()
+
+void UPCGMetadataPartitionDataInterface::GetShaderParameters(TCHAR const* UID, FShaderParametersMetadataBuilder& InOutBuilder, FShaderParametersMetadataAllocations& InOutAllocations) const
+{
+	InOutBuilder.AddNestedStruct<FPCGMetadataPartitionDataInterfaceParameters>(UID);
+}
+
+void UPCGMetadataPartitionDataInterface::GetHLSL(FString& OutHLSL, FString const& InDataInterfaceName) const
+{
+	TMap<FString, FStringFormatArg> TemplateArgs =
+	{
+		{ TEXT("DataInterfaceName"), InDataInterfaceName },
+	};
+
+	OutHLSL += FString::Format(TEXT(
+		"StructuredBuffer<uint> {DataInterfaceName}_AttributeValueToPartitionIndex;\n"
+		"int {DataInterfaceName}_PartitionAttributeId;\n"
+		"uint {DataInterfaceName}_NumPartitions;\n"
+		"\n"
+		"int MetadataPartition_GetPartitionAttributeId_{DataInterfaceName}()\n"
+		"{\n"
+		"	return {DataInterfaceName}_PartitionAttributeId;\n"
+		"}\n"
+		"\n"
+		"uint MetadataPartition_GetNumPartitions_{DataInterfaceName}()\n"
+		"{\n"
+		"	return {DataInterfaceName}_NumPartitions;\n"
+		"}\n"
+		"\n"
+		"uint MetadataPartition_GetPartitionIndex_{DataInterfaceName}(int InAttributeValue)\n"
+		"{\n"
+		"	return {DataInterfaceName}_AttributeValueToPartitionIndex[InAttributeValue];\n"
+		"}\n"
+		"\n"
+	), TemplateArgs);
+}
+
+UComputeDataProvider* UPCGMetadataPartitionDataInterface::CreateDataProvider(TObjectPtr<UObject> InBinding, uint64 InInputMask, uint64 InOutputMask) const
+{
+	TRACE_CPUPROFILER_EVENT_SCOPE(UPCGMetadataPartitionDataInterface::CreateDataProvider);
+	UPCGDataBinding* Binding = CastChecked<UPCGDataBinding>(InBinding);
+
+	TObjectPtr<UPCGMetaDataPartitionDataProvider> DataProvider = NewObject<UPCGMetaDataPartitionDataProvider>();
+
+	if (ensure(Settings->PartitionAttributeSelectors.Num() == 1 && Settings->PartitionAttributeSelectors[0].IsBasicAttribute()))
+	{
+		const FPCGDataCollectionDesc InputDataDesc = PCGDataForGPUHelpers::ComputeInputPinDataDesc(Settings, PCGPinConstants::DefaultInputLabel, Binding);
+
+		const int32 PartitionAttributeId = Binding->GetAttributeId(Settings->PartitionAttributeSelectors[0].GetAttributeName(), EPCGKernelAttributeType::StringKey);
+		if (PartitionAttributeId != INDEX_NONE)
+		{
+			InputDataDesc.GetUniqueStringKeyValues(PartitionAttributeId, DataProvider->UniqueStringKeyValues);
+		}
+
+		int32 LargestValue = INDEX_NONE;
+		for (int32 Key : DataProvider->UniqueStringKeyValues)
+		{
+			LargestValue = FMath::Max(LargestValue, Key);
+		}
+
+		DataProvider->NumInputData = InputDataDesc.DataDescs.Num();
+		DataProvider->PartitionAttributeId = PartitionAttributeId;
+		DataProvider->NumPartitions = DataProvider->UniqueStringKeyValues.Num();
+		DataProvider->MaxAttributeValue = LargestValue;
+	}
+
+	return DataProvider;
+}
+
+FComputeDataProviderRenderProxy* UPCGMetaDataPartitionDataProvider::GetRenderProxy()
+{
+	return new FPCGMetaDataPartitionProviderProxy(PartitionAttributeId, NumInputData, NumPartitions, MaxAttributeValue, UniqueStringKeyValues);
+}
+
+bool FPCGMetaDataPartitionProviderProxy::IsValid(FValidationData const& InValidationData) const
+{
+	return (InValidationData.ParameterStructSize == sizeof(FParameters));
+}
+
+void FPCGMetaDataPartitionProviderProxy::AllocateResources(FRDGBuilder& GraphBuilder, FAllocationData const& InAllocationData)
+{
+	const int32 NumPartitionsToAllocate = FMath::Max(1u, NumInputData) * FMath::Max(1u, NumPartitions);
+	const int32 NumIndicesToAllocate = FMath::Max(1, MaxAttributeValue + 1);
+
+	TArray<uint32> Zeros;
+	Zeros.SetNumZeroed(FMath::Max(NumPartitionsToAllocate, NumIndicesToAllocate));
+
+	{
+		FRDGBufferDesc Desc = FRDGBufferDesc::CreateStructuredDesc(sizeof(uint32), NumIndicesToAllocate);
+		Desc.Usage |= BUF_SourceCopy;
+
+		AttributeValueToPartitionIndex = GraphBuilder.CreateBuffer(Desc, TEXT("PCGAttributeValueToPartitionIndex"));
+		AttributeValueToPartitionIndexSRV = GraphBuilder.CreateSRV(AttributeValueToPartitionIndex);
+
+		int PartitionCounter = 0;
+		// Skipping default empty string.
+		for (int Index = 1; Index <= MaxAttributeValue; ++Index)
+		{
+			const int PartitionIndex = UniqueStringKeyValues.IndexOfByKey(Index);
+			if (PartitionIndex != INDEX_NONE) // TODO - probably should make it uint and give it index_nones.
+			{
+				Zeros[Index] = PartitionIndex;
+			}
+		}
+
+		GraphBuilder.QueueBufferUpload(AttributeValueToPartitionIndex, Zeros.GetData(), NumIndicesToAllocate * Zeros.GetTypeSize());
+	}
+}
+
+void FPCGMetaDataPartitionProviderProxy::GatherDispatchData(FDispatchData const& InDispatchData)
+{
+	const TStridedView<FParameters> ParameterArray = MakeStridedParameterView<FParameters>(InDispatchData);
+	for (int32 InvocationIndex = 0; InvocationIndex < ParameterArray.Num(); ++InvocationIndex)
+	{
+		FParameters& Parameters = ParameterArray[InvocationIndex];
+
+		Parameters.PartitionAttributeId = PartitionAttributeId;
+		Parameters.NumPartitions = NumPartitions;
+		Parameters.AttributeValueToPartitionIndex = AttributeValueToPartitionIndexSRV;
+	}
+}
diff --git a/Engine/Plugins/PCG/Source/PCG/Private/Compute/DataInterfaces/Elements/PCGMetadataPartitionDataInterface.h b/Engine/Plugins/PCG/Source/PCG/Private/Compute/DataInterfaces/Elements/PCGMetadataPartitionDataInterface.h
new file mode 100644
--- /dev/null
+++ b/Engine/Plugins/PCG/Source/PCG/Private/Compute/DataInterfaces/Elements/PCGMetadataPartitionDataInterface.h
@@ -0,0 +1,90 @@
+// Copyright Epic Games, Inc. All Rights Reserved.
+
+#pragma once
+
+#include "Compute/DataInterfaces/PCGComputeDataInterface.h"
+#include "ComputeFramework/ComputeDataProvider.h"
+
+#include "PCGMetadataPartitionDataInterface.generated.h"
+
+class FPCGMetadataPartitionDataInterfaceParameters;
+class UPCGMetadataPartitionSettings;
+
+/** Data Interface to marshal Attribute Partition settings to the GPU. */
+UCLASS(ClassGroup = (Procedural))
+class UPCGMetadataPartitionDataInterface : public UPCGComputeDataInterface
+{
+	GENERATED_BODY()
+
+public:
+	//~ Begin UComputeDataInterface Interface
+	TCHAR const* GetClassName() const override { return TEXT("PCGMetadataPartition"); }
+	void GetSupportedInputs(TArray<FShaderFunctionDefinition>& OutFunctions) const override;
+	void GetShaderParameters(TCHAR const* UID, FShaderParametersMetadataBuilder& InOutBuilder, FShaderParametersMetadataAllocations& InOutAllocations) const override;
+	void GetHLSL(FString& OutHLSL, FString const& InDataInterfaceName) const override;
+	UComputeDataProvider* CreateDataProvider(TObjectPtr<UObject> InBinding, uint64 InInputMask, uint64 InOutputMask) const override;
+	//~ End UComputeDataInterface Interface
+
+	UPROPERTY()
+	TObjectPtr<const UPCGMetadataPartitionSettings> Settings = nullptr;
+};
+
+UCLASS()
+class UPCGMetaDataPartitionDataProvider : public UComputeDataProvider
+{
+	GENERATED_BODY()
+
+public:
+	//~ Begin UComputeDataProvider Interface
+	FComputeDataProviderRenderProxy* GetRenderProxy() override;
+	//~ End UComputeDataProvider Interface
+
+	UPROPERTY()
+	int32 PartitionAttributeId = INDEX_NONE;
+
+	UPROPERTY()
+	int32 NumInputData = 0;
+
+	UPROPERTY()
+	uint32 NumPartitions = 0;
+
+	UPROPERTY()
+	int32 MaxAttributeValue = INDEX_NONE;
+
+	UPROPERTY()
+	TArray<int32> UniqueStringKeyValues;
+};
+
+class FPCGMetaDataPartitionProviderProxy : public FComputeDataProviderRenderProxy
+{
+public:
+	FPCGMetaDataPartitionProviderProxy(int32 InPartitionAttributeId, int32 InNumInputData, uint32 InNumPartitions, int32 InMaxAttributeValue, TArray<int32> InUniqueStringKeyValues)
+		: PartitionAttributeId(InPartitionAttributeId)
+		, NumInputData(InNumInputData)
+		, NumPartitions(InNumPartitions)
+		, MaxAttributeValue(InMaxAttributeValue)
+		, UniqueStringKeyValues(MoveTemp(InUniqueStringKeyValues))
+	{}
+
+	//~ Begin FComputeDataProviderRenderProxy Interface
+	bool IsValid(FValidationData const& InValidationData) const override;
+	void AllocateResources(FRDGBuilder& GraphBuilder, FAllocationData const& InAllocationData);
+	void GatherDispatchData(FDispatchData const& InDispatchData) override;
+	//~ End FComputeDataProviderRenderProxy Interface
+
+protected:
+	using FParameters = FPCGMetadataPartitionDataInterfaceParameters;
+
+	int32 PartitionAttributeId = INDEX_NONE;
+
+	uint32 NumInputData = 0;
+
+	uint32 NumPartitions = 0;
+
+	int32 MaxAttributeValue = INDEX_NONE;
+
+	TArray<int32> UniqueStringKeyValues;
+
+	FRDGBufferRef AttributeValueToPartitionIndex = nullptr;
+	FRDGBufferSRVRef AttributeValueToPartitionIndexSRV = nullptr;
+};
diff --git a/Engine/Plugins/PCG/Source/PCG/Private/Compute/DataInterfaces/PCGComputeDataInterface.h b/Engine/Plugins/PCG/Source/PCG/Private/Compute/DataInterfaces/PCGComputeDataInterface.h
--- a/Engine/Plugins/PCG/Source/PCG/Private/Compute/DataInterfaces/PCGComputeDataInterface.h
+++ b/Engine/Plugins/PCG/Source/PCG/Private/Compute/DataInterfaces/PCGComputeDataInterface.h
@@ -7,6 +7,8 @@
 
 #include "PCGComputeDataInterface.generated.h"
 
+class UPCGDataBinding;
+
 UCLASS(Abstract, ClassGroup = (Procedural))
 class UPCGComputeDataInterface : public UComputeDataInterface
 {
@@ -45,5 +47,5 @@ class UPCGComputeDataProvider : public UComputeDataProvider
 	GENERATED_BODY()
 
 public:
-	virtual bool PrepareForExecute_GameThread() { return true; }
+	virtual bool PrepareForExecute_GameThread(const UPCGDataBinding* InBinding) { return true; }
 };
diff --git a/Engine/Plugins/PCG/Source/PCG/Private/Compute/DataInterfaces/PCGDataCollectionDataInterface.cpp b/Engine/Plugins/PCG/Source/PCG/Private/Compute/DataInterfaces/PCGDataCollectionDataInterface.cpp
--- a/Engine/Plugins/PCG/Source/PCG/Private/Compute/DataInterfaces/PCGDataCollectionDataInterface.cpp
+++ b/Engine/Plugins/PCG/Source/PCG/Private/Compute/DataInterfaces/PCGDataCollectionDataInterface.cpp
@@ -52,6 +52,12 @@ void UPCGDataCollectionDataInterface::GetSupportedInputs(TArray<FShaderFunctionD
 			.AddReturnType(EShaderFundamentalType::Uint)
 			.AddParam(EShaderFundamentalType::Uint); // InDataIndex
 
+		OutFunctions.AddDefaulted_GetRef()
+			.SetName(TEXT("AddToElementCounterInternal"))
+			.AddReturnType(EShaderFundamentalType::Int) // Value before increment. Returns -1 if counters disabled.
+			.AddParam(EShaderFundamentalType::Uint) // InDataIndex
+			.AddParam(EShaderFundamentalType::Uint); // InElementCount
+
 		OutFunctions.AddDefaulted_GetRef()
 			.SetName(TEXT("GetThreadData"))
 			.AddParam(EShaderFundamentalType::Uint) // InThreadIndex
@@ -388,6 +394,8 @@ void UPCGDataCollectionDataInterface::GetSupportedInputs(TArray<FShaderFunctionD
 
 BEGIN_SHADER_PARAMETER_STRUCT(FPCGDataCollectionDataInterfaceParameters, )
 	SHADER_PARAMETER_RDG_BUFFER_UAV(RWByteAddressBuffer, DataCollectionBuffer)
+	SHADER_PARAMETER_RDG_BUFFER_UAV(RWStructuredBuffer<uint32>, ElementCountersBuffer)
+	SHADER_PARAMETER(uint32, ElementCountersPresent)
 END_SHADER_PARAMETER_STRUCT()
 
 void UPCGDataCollectionDataInterface::GetShaderParameters(TCHAR const* UID, FShaderParametersMetadataBuilder& InOutBuilder, FShaderParametersMetadataAllocations& InOutAllocations) const
@@ -445,6 +453,8 @@ UComputeDataProvider* UPCGDataCollectionDataInterface::CreateDataProvider(TObjec
 
 	Provider->ExportMode = bRequiresExport ? EPCGExportMode::ComputeGraphOutput : EPCGExportMode::NoExport;
 
+	Provider->bExposeElementCounters = bExposeElementCounters;
+
 #if WITH_EDITOR
 	if (Binding->SourceComponent->IsInspecting())
 	{
@@ -470,24 +480,73 @@ FComputeDataProviderRenderProxy* UPCGDataCollectionDataProvider::GetRenderProxy(
 {
 	TRACE_CPUPROFILER_EVENT_SCOPE(UPCGDataCollectionDataProvider::GetRenderProxy);
 
-	FPCGDataCollectionDataProviderProxy* Proxy = new FPCGDataCollectionDataProviderProxy(this, PinDesc, ExportMode, OutputPinLabel, OutputPinLabelAlias);
+	FPCGDataCollectionDataProviderProxy* Proxy = new FPCGDataCollectionDataProviderProxy(this, PinDesc, ExportMode, bExposeElementCounters, OutputPinLabel, OutputPinLabelAlias);
 
 #if WITH_EDITOR
 	if (const UPCGNode* Node = Cast<UPCGNode>(ProducerSettings ? ProducerSettings->GetOuter() : nullptr))
 	{
-		Proxy->SetBufferName(FString::Format(TEXT("PCG_DC_{0}_{1}"), { Node->GetNodeTitle(EPCGNodeTitleType::ListView).ToString(), OutputPinLabel.ToString()}));
+		const FString NodeTitle = Node->GetNodeTitle(EPCGNodeTitleType::ListView).ToString();
+		Proxy->SetBufferName(FString::Format(TEXT("PCG_DC_{0}_{1}"), { NodeTitle, OutputPinLabel.ToString() }));
+		Proxy->SetCountersBufferName(FString::Format(TEXT("PCG_DC_{0}_{1}_Counters"), { NodeTitle, OutputPinLabel.ToString()}));
 	}
 	else
 	{
 		Proxy->SetBufferName(TEXT("PCG_DC"));
+		Proxy->SetCountersBufferName(TEXT("PCG_DC_Counters"));
 	}
 #endif
 
+	if (bExposeElementCounters)
+	{
+		Proxy->AsyncReadbackCallback_RenderThread = [ThisWeak=TWeakObjectPtr<UPCGDataCollectionDataProvider>(this)](const void* InData, int InNumBytes)
+		{
+			TRACE_CPUPROFILER_EVENT_SCOPE(UPCGDataCollectionDataProvider::ProcessReadbackData_RenderThread);
+			TStrongObjectPtr<UPCGDataCollectionDataProvider> This = ThisWeak.Pin();
+			if (!This)
+			{
+				return;
+			}
+
+			// We should never find ourselves stomping existing data.
+			check(This->ReadbackCounters.IsEmpty());
+
+			if (ensure(InData && InNumBytes > 0))
+			{
+				// Copy the data to temp storage for the game thread to pick up.
+				This->ReadbackCounters.SetNumUninitialized(InNumBytes / This->ReadbackCounters.GetTypeSize());
+				FMemory::Memcpy(This->ReadbackCounters.GetData(), InData, InNumBytes);
+			}
+
+			This->bReadbackComplete = true;
+
+			PCGHelpers::ExecuteOnGameThread(UE_SOURCE_LOCATION, [ThisWeak=ThisWeak, OutputPinLabelAlias=This->OutputPinLabelAlias]()
+			{
+				TStrongObjectPtr<UPCGDataCollectionDataProvider> This = ThisWeak.Pin();
+				if (This)
+				{
+					if (This->Binding.IsValid())
+					{
+						This->Binding->ReceiveElementCountsFromGPU_GameThread(OutputPinLabelAlias, This->ReadbackCounters);
+					}
+
+					This->OnCountersReadBack.Broadcast();
+				}
+			});
+		};
+	}
+
 	return Proxy;
 }
 
-FPCGDataCollectionDataProviderProxy::FPCGDataCollectionDataProviderProxy(TWeakObjectPtr<UPCGDataCollectionDataProvider> InDataProvider, const FPCGDataCollectionDesc& InPinDesc, EPCGExportMode InExportMode, FName InOutputPinLabel, FName InOutputPinLabelAlias)
+FPCGDataCollectionDataProviderProxy::FPCGDataCollectionDataProviderProxy(
+	TWeakObjectPtr<UPCGDataCollectionDataProvider> InDataProvider,
+	const FPCGDataCollectionDesc& InPinDesc,
+	EPCGExportMode InExportMode,
+	bool bInExposeElementCounters,
+	FName InOutputPinLabel,
+	FName InOutputPinLabelAlias)
 	: ExportMode(InExportMode)
+	, bExposeElementCounters(bInExposeElementCounters)
 	, PinDesc(InPinDesc)
 	, DataProviderWeakPtr(InDataProvider)
 	, OutputPinLabel(InOutputPinLabel)
@@ -515,12 +574,15 @@ bool FPCGDataCollectionDataProviderProxy::IsValid(FValidationData const& InValid
 void FPCGDataCollectionDataProviderProxy::GatherDispatchData(FDispatchData const& InDispatchData)
 {
 	check(BufferUAV);
+	check(CountersUAV);
 
 	const TStridedView<FParameters> ParameterArray = MakeStridedParameterView<FParameters>(InDispatchData);
 	for (int32 InvocationIndex = 0; InvocationIndex < ParameterArray.Num(); ++InvocationIndex)
 	{
 		FParameters& Parameters = ParameterArray[InvocationIndex];
 		Parameters.DataCollectionBuffer = BufferUAV;
+		Parameters.ElementCountersBuffer = CountersUAV;
+		Parameters.ElementCountersPresent = bExposeElementCounters ? 1 : 0;
 	}
 }
 
@@ -529,11 +591,42 @@ void FPCGDataCollectionDataProviderProxy::AllocateResources(FRDGBuilder& GraphBu
 	TRACE_CPUPROFILER_EVENT_SCOPE(FPCGDataCollectionDataProviderProxy::AllocateResources);
 	check(SizeBytes > 0);
 
+	{
+		const TCHAR* DebugNameCounters =
+#if WITH_EDITOR
+			*CountersBufferName;
+#else
+			TEXT("PCGDataCollection_Counters");
+#endif
+
+		// Ensure we always allocate at least 1 element so we have something to bind.
+		const uint32 NumCounters = bExposeElementCounters ? FMath::Max(PinDesc.DataDescs.Num(), 1) : 1;
+
+		FRDGBufferDesc CountersDesc = FRDGBufferDesc::CreateStructuredDesc(sizeof(uint32), NumCounters);
+		CountersDesc.Usage |= BUF_SourceCopy;
+
+		Counters = GraphBuilder.CreateBuffer(CountersDesc, DebugNameCounters);
+		CountersUAV = GraphBuilder.CreateUAV(Counters);
+
+		// A single counter, or a small number of counters, will be a very common case.
+		if (NumCounters <= 4)
+		{
+			GraphBuilder.QueueBufferUpload(Counters, &FUintVector4::ZeroValue, NumCounters * sizeof(uint32));
+		}
+		else
+		{
+			TArray<uint32> Zeroes;
+			Zeroes.SetNumZeroed(NumCounters);
+
+			GraphBuilder.QueueBufferUpload(Counters, Zeroes.GetData(), Zeroes.Num() * Zeroes.GetTypeSize());
+		}
+	}
+
 	const TCHAR* DebugName =
 #if WITH_EDITOR
 		*BufferName;
 #else
-		TEXT("PCGDataCollectionBuffer");
+		TEXT("PCGDataCollection");
 #endif
 
 	FRDGBufferDesc Desc = FRDGBufferDesc::CreateByteAddressDesc(SizeBytes);
@@ -592,11 +685,28 @@ void FPCGDataCollectionDataProviderProxy::AllocateResources(FRDGBuilder& GraphBu
 	}
 }
 
+void FPCGDataCollectionDataProviderProxy::GetReadbackData(TArray<FReadbackData>& OutReadbackData) const
+{
+	if (bExposeElementCounters)
+	{
+		FReadbackData Data;
+		Data.Buffer = Counters;
+		Data.NumBytes = PinDesc.DataDescs.Num() * sizeof(uint32);
+		Data.ReadbackCallback_RenderThread = &AsyncReadbackCallback_RenderThread;
+
+		OutReadbackData.Add(MoveTemp(Data));
+	}
+}
+
 #if WITH_EDITOR
 void FPCGDataCollectionDataProviderProxy::SetBufferName(const FString& InBufferName)
 {
 	BufferName = InBufferName;
 }
+void FPCGDataCollectionDataProviderProxy::SetCountersBufferName(const FString& InCountersBufferName)
+{
+	CountersBufferName = InCountersBufferName;
+}
 #endif
 
 #undef LOCTEXT_NAMESPACE
diff --git a/Engine/Plugins/PCG/Source/PCG/Private/Compute/DataInterfaces/PCGDataCollectionDataInterface.h b/Engine/Plugins/PCG/Source/PCG/Private/Compute/DataInterfaces/PCGDataCollectionDataInterface.h
--- a/Engine/Plugins/PCG/Source/PCG/Private/Compute/DataInterfaces/PCGDataCollectionDataInterface.h
+++ b/Engine/Plugins/PCG/Source/PCG/Private/Compute/DataInterfaces/PCGDataCollectionDataInterface.h
@@ -33,6 +33,8 @@ public:
 	TCHAR const* GetShaderVirtualPath() const override; // TODO probably easier to just inline rather than external source?
 	void GetShaderHash(FString& InOutKey) const override;
 	void GetHLSL(FString& OutHLSL, FString const& InDataInterfaceName) const override;
+	// This data interface only uses the CF readback functionality for the element counters.
+	bool GetRequiresReadback() const override { return bExposeElementCounters; }
 	UComputeDataProvider* CreateDataProvider(TObjectPtr<UObject> InBinding, uint64 InInputMask, uint64 InOutputMask) const override;
 	//~ End UComputeDataInterface Interface
 
@@ -40,6 +42,8 @@ public:
 
 	void SetRequiresExport(bool bInRequiresExport) { bRequiresExport = bInRequiresExport; }
 
+	void SetExposeElementCounters(bool bInExposeElementCounters) { bExposeElementCounters = bInExposeElementCounters; }
+
 protected:
 	/** Settings of node that produces this data, normally the upstream node. */
 	UPROPERTY()
@@ -49,6 +53,9 @@ protected:
 	UPROPERTY()
 	bool bRequiresExport = false;
 
+	UPROPERTY()
+	bool bExposeElementCounters = false;
+
 	static TCHAR const* TemplateFilePath;
 };
 
@@ -69,6 +76,9 @@ public:
 	DECLARE_EVENT(UPCGDataCollectionDataProvider, FOnDataExported);
 	FOnDataExported& OnDataExported_GameThread() { return OnDataExported; }
 
+	DECLARE_EVENT(UPCGDataCollectionDataProvider, FOnCountersReadBack);
+	FOnCountersReadBack& OnCountersReadBack_GameThread() { return OnCountersReadBack; }
+
 public:
 	/** Settings of node that produces this data, normally the upstream node. */
 	UPROPERTY()
@@ -81,6 +91,9 @@ public:
 	UPROPERTY()
 	EPCGExportMode ExportMode = EPCGExportMode::NoExport;
 
+	UPROPERTY()
+	bool bExposeElementCounters = false;
+
 	UPROPERTY()
 	FName OutputPinLabel;
 
@@ -89,28 +102,47 @@ public:
 
 private:
 	FOnDataExported OnDataExported;
+
+	FOnCountersReadBack OnCountersReadBack;
+
+	std::atomic<bool> bReadbackComplete = false;
+
+	TArray<uint32> ReadbackCounters;
 };
 
 class FPCGDataCollectionDataProviderProxy : public FComputeDataProviderRenderProxy
 {
 public:
-	FPCGDataCollectionDataProviderProxy(TWeakObjectPtr<UPCGDataCollectionDataProvider> InDataProvider, const FPCGDataCollectionDesc& InPinDesc, EPCGExportMode InExportMode, FName InOutputPinLabel, FName InOutputPinLabelAlias);
+	FPCGDataCollectionDataProviderProxy(
+		TWeakObjectPtr<UPCGDataCollectionDataProvider> InDataProvider,
+		const FPCGDataCollectionDesc& InPinDesc,
+		EPCGExportMode InExportMode,
+		bool bInExposeElementCounters,
+		FName InOutputPinLabel,
+		FName InOutputPinLabelAlias);
 
 	//~ Begin FComputeDataProviderRenderProxy Interface
 	bool IsValid(FValidationData const& InValidationData) const override;
 	void GatherDispatchData(FDispatchData const& InDispatchData) override;
 	void AllocateResources(FRDGBuilder& GraphBuilder, FAllocationData const& InAllocationData) override;
+	void GetReadbackData(TArray<FReadbackData>& OutReadbackData) const override;
 	//~ End FComputeDataProviderRenderProxy Interface
 
 #if WITH_EDITOR
 	void SetBufferName(const FString& InBufferName);
+	void SetCountersBufferName(const FString& InCountersBufferName);
 #endif
 
+public:
+	FReadbackCallback AsyncReadbackCallback_RenderThread;
+
 protected:
 	uint32 SizeBytes = 0;
 
 	EPCGExportMode ExportMode = EPCGExportMode::NoExport;
 
+	bool bExposeElementCounters = false;
+
 	using FParameters = FPCGDataCollectionDataInterfaceParameters;
 
 	FPCGDataCollectionDesc PinDesc;
@@ -118,6 +150,9 @@ protected:
 	FRDGBufferRef Buffer = nullptr;
 	FRDGBufferUAVRef BufferUAV = nullptr;
 
+	FRDGBufferRef Counters = nullptr;
+	FRDGBufferUAVRef CountersUAV = nullptr;
+
 	// Weak pointer useful for passing back buffer handles. Do not access this directly from the render thread.
 	TWeakObjectPtr<UPCGDataCollectionDataProvider> DataProviderWeakPtr;
 
@@ -126,5 +161,6 @@ protected:
 
 #if WITH_EDITOR
 	FString BufferName;
+	FString CountersBufferName;
 #endif
 };
diff --git a/Engine/Plugins/PCG/Source/PCG/Private/Compute/DataInterfaces/PCGDataCollectionUploadDataInterface.cpp b/Engine/Plugins/PCG/Source/PCG/Private/Compute/DataInterfaces/PCGDataCollectionUploadDataInterface.cpp
--- a/Engine/Plugins/PCG/Source/PCG/Private/Compute/DataInterfaces/PCGDataCollectionUploadDataInterface.cpp
+++ b/Engine/Plugins/PCG/Source/PCG/Private/Compute/DataInterfaces/PCGDataCollectionUploadDataInterface.cpp
@@ -7,173 +7,17 @@
 #include "PCGNode.h"
 #include "PCGSettings.h"
 #include "PCGSubsystem.h"
+#include "PCGDataCollectionAdaptorCS.h"
 #include "Compute/PCGComputeCommon.h"
 #include "Compute/PCGComputeGraph.h"
 #include "Compute/PCGDataBinding.h"
+#include "Compute/PCGDataUploadAdaptor.h"
 #include "Compute/Data/PCGProxyForGPUData.h"
 
+#include "RenderCaptureInterface.h"
 #include "RenderGraphBuilder.h"
 #include "RenderGraphResources.h"
-
-/** A helper to ingest incoming data objects and pass the required data to the data interface. */
-class FPCGDataUploadAdaptor
-{
-public:
-	FPCGDataUploadAdaptor(UPCGDataBinding* InDataBinding, const FPCGDataCollectionDesc& InPinDesc, FName InPinLabel)
-	{
-		PinDesc = InPinDesc;
-		PinLabel = InPinLabel;
-		DataBinding = InDataBinding;
-
-		// Are all our data items coming direct from GPU - and importantly are a complete, single data collection?
-		bool bMultipleBuffersEncountered = false;
-		TSharedPtr<const FPCGProxyForGPUDataCollection> UniqueDataCollection;
-		for (int InputDataIndex = 0; InputDataIndex < DataBinding->InputDataCollection.TaggedData.Num(); ++InputDataIndex)
-		{
-			const UPCGProxyForGPUData* DataGPU = Cast<UPCGProxyForGPUData>(DataBinding->InputDataCollection.TaggedData[InputDataIndex].Data);
-			if (!DataGPU)
-			{
-				// Encountered CPU data, no reuse for now.
-				// TODO: Support multiple input buffers and only upload required data.
-				UniqueDataCollection.Reset();
-				break;
-			}
-
-			TSharedPtr<const FPCGProxyForGPUDataCollection> InputDataCollection = DataGPU->GetInputDataCollectionInfo();
-
-			if (!InputDataCollection)
-			{
-				UE_LOG(LogPCG, Error, TEXT("No reuse: Missing data collection buffer!"));
-
-				UniqueDataCollection.Reset();
-				break;
-			}
-
-			if (DataGPU->GetDataIndexInCollection() != InputDataIndex)
-			{
-				// Out of order data items, no reuse for now.
-				// TODO: support indirection of data items.
-				UniqueDataCollection.Reset();
-				break;
-			}
-
-			if (UniqueDataCollection)
-			{
-				if (UniqueDataCollection != InputDataCollection)
-				{
-					// Multiple buffer items, no reuse for now.
-					// TODO: Support multiple input buffers.
-					UniqueDataCollection.Reset();
-					break;
-				}
-			}
-			else
-			{
-				UniqueDataCollection = InputDataCollection;
-			}
-		}
-
-		if (UniqueDataCollection && UniqueDataCollection->GetDescription().DataDescs.Num() != DataBinding->InputDataCollection.TaggedData.Num())
-		{
-			// We're not using all the data items from the buffer, no reuse for now.
-			// TODO: Support indirection of data items.
-			UniqueDataCollection.Reset();
-		}
-
-		if (UniqueDataCollection)
-		{
-			ExternalBufferForReuse = UniqueDataCollection->GetBuffer();
-		}
-	}
-
-	/** Do any preparation work such as data readbacks. Returns true when preparation is complete. */
-	bool PrepareData_GameThread()
-	{
-		check(IsInGameThread());
-
-		if (!ExternalBufferForReuse)
-		{
-			UPCGDataBinding* Binding = DataBinding.Get();
-			if (!Binding)
-			{
-				return true;
-			}
-
-			const uint64 BufferSize = PinDesc.ComputePackedSizeBytes();
-			if (PCGComputeHelpers::IsBufferSizeTooLarge(BufferSize))
-			{
-				return true;
-			}
-
-			bool bAllDataReady = true;
-
-			for (FPCGTaggedData& Data : Binding->InputDataCollection.TaggedData)
-			{
-				if (const UPCGProxyForGPUData* Proxy = Cast<UPCGProxyForGPUData>(Data.Data))
-				{
-					UPCGProxyForGPUData::FReadbackResult Result = Proxy->GetCPUData(nullptr);
-					if (!Result.bComplete)
-					{
-						bAllDataReady = false;
-					}
-					else
-					{
-						Data.Data = Result.Data;
-					}
-				}
-			}
-
-			if (!bAllDataReady)
-			{
-				return false;
-			}
-
-			PinDesc.PackDataCollection(Binding->InputDataCollection, PinLabel, Binding->GetStringTable(), PackedDataCollection);
-		}
-
-		return true;
-	}
-
-	/** Gets the buffer that can then be used to read the data in kernels. */
-	FRDGBufferRef GetBuffer(FRDGBuilder& InGraphBuilder, EPCGExportMode InExportMode, const TCHAR* InBufferDebugName)
-	{
-		FRDGBufferRef Buffer = nullptr;
-
-		if (ExternalBufferForReuse)
-		{
-			ensure(PackedDataCollection.IsEmpty());
-
-			Buffer = InGraphBuilder.RegisterExternalBuffer(ExternalBufferForReuse);
-		}
-		else
-		{
-			FRDGBufferDesc Desc = FRDGBufferDesc::CreateByteAddressDesc(sizeof(uint32) * PackedDataCollection.Num());
-			if (InExportMode != EPCGExportMode::NoExport)
-			{
-				// We don't know for sure whether buffer will be read back or not, so need to flag the possibility if the buffer will be passed downstream.
-				Desc.Usage |= BUF_SourceCopy;
-			}
-
-			Buffer = InGraphBuilder.CreateBuffer(Desc, InBufferDebugName);
-
-			InGraphBuilder.QueueBufferUpload(Buffer, PackedDataCollection.GetData(), PackedDataCollection.Num() * PackedDataCollection.GetTypeSize(), ERDGInitialDataFlags::None);
-		}
-
-		return Buffer;
-	}
-
-private:
-	TRefCountPtr<FRDGPooledBuffer> ExternalBufferForReuse;
-
-	TWeakObjectPtr<UPCGDataBinding> DataBinding;
-
-	FPCGDataCollectionDesc PinDesc;
-
-	FName PinLabel;
-
-	/** If data cannot be reused from GPU, will be packed into this buffer prior to upload. */
-	TArray<uint32> PackedDataCollection;
-};
+#include "RenderGraphUtils.h"
 
 UComputeDataProvider* UPCGDataCollectionUploadDataInterface::CreateDataProvider(TObjectPtr<UObject> InBinding, uint64 InInputMask, uint64 InOutputMask) const
 {
@@ -206,16 +50,16 @@ UComputeDataProvider* UPCGDataCollectionUploadDataInterface::CreateDataProvider(
 	return Provider;
 }
 
-bool UPCGDataProviderDataCollectionUpload::PrepareForExecute_GameThread()
+bool UPCGDataProviderDataCollectionUpload::PrepareForExecute_GameThread(const UPCGDataBinding* InBinding)
 {
-	return !DataAdaptor || DataAdaptor->PrepareData_GameThread();
+	return !DataAdaptor || DataAdaptor->PrepareData_GameThread(InBinding);
 }
 
 FComputeDataProviderRenderProxy* UPCGDataProviderDataCollectionUpload::GetRenderProxy()
 {
 	TRACE_CPUPROFILER_EVENT_SCOPE(UPCGDataProviderDataCollectionUpload::GetRenderProxy);
 
-	FPCGDataProviderDataCollectionUploadProxy* Proxy = new FPCGDataProviderDataCollectionUploadProxy(DataAdaptor, this, PinDesc, DownstreamInputPinLabels, OutputPinLabel, OutputPinLabelAlias, ExportMode);
+	FPCGDataProviderDataCollectionUploadProxy* Proxy = new FPCGDataProviderDataCollectionUploadProxy(DataAdaptor, this, PinDesc, DownstreamInputPinLabels, OutputPinLabel, OutputPinLabelAlias, ExportMode, bExposeElementCounters);
 
 #if WITH_EDITOR
 	if (const UPCGNode* Node = Cast<UPCGNode>(ProducerSettings ? ProducerSettings->GetOuter() : nullptr))
@@ -238,8 +82,9 @@ FPCGDataProviderDataCollectionUploadProxy::FPCGDataProviderDataCollectionUploadP
 	const TArray<FName>& InDownstreamInputPinLabels,
 	FName InOutputPinLabel,
 	FName InOutputPinLabelAlias,
-	EPCGExportMode InExportMode)
-	: FPCGDataCollectionDataProviderProxy(InDataProvider, InPinDesc, InExportMode, InOutputPinLabel, InOutputPinLabelAlias)
+	EPCGExportMode InExportMode,
+	bool bInExposeElementCounters)
+	: FPCGDataCollectionDataProviderProxy(InDataProvider, InPinDesc, InExportMode, bInExposeElementCounters, InOutputPinLabel, InOutputPinLabelAlias)
 	, DataAdaptor(InDataAdaptor)
 	, DownstreamInputPinLabels(InDownstreamInputPinLabels)
 {
@@ -249,13 +94,46 @@ void FPCGDataProviderDataCollectionUploadProxy::AllocateResources(FRDGBuilder& G
 {
 	TRACE_CPUPROFILER_EVENT_SCOPE(FPCGDataProviderDataCollectionUploadProxy::AllocateResources);
 
-	const TCHAR* DebugName =
+	{
+		const TCHAR* DebugName =
 #if WITH_EDITOR
-		*BufferName;
+			*BufferName;
 #else
-		TEXT("PCGDataCollectionUploadBuffer");
+			TEXT("PCGDataCollectionUpload");
 #endif
 
-	Buffer = DataAdaptor->GetBuffer(GraphBuilder, ExportMode, DebugName);
-	BufferUAV = GraphBuilder.CreateUAV(Buffer);
+		Buffer = DataAdaptor->GetBuffer_RenderThread(GraphBuilder, ExportMode, DebugName);
+		BufferUAV = GraphBuilder.CreateUAV(Buffer);
+	}
+
+	{
+		const TCHAR* DebugName =
+#if WITH_EDITOR
+			*CountersBufferName;
+#else
+			TEXT("PCGDataCollectionUpload_Counters");
+#endif
+
+		// Ensure we always allocate at least 1 element so we have something to bind.
+		const uint32 NumCounters = bExposeElementCounters ? FMath::Max(PinDesc.DataDescs.Num(), 1) : 1;
+
+		FRDGBufferDesc CountersDesc = FRDGBufferDesc::CreateStructuredDesc(sizeof(uint32), NumCounters);
+		CountersDesc.Usage |= BUF_SourceCopy;
+
+		Counters = GraphBuilder.CreateBuffer(CountersDesc, DebugName);
+		CountersUAV = GraphBuilder.CreateUAV(Counters);
+
+		// A single counter, or a small number of counters, will be a very common case.
+		if (NumCounters <= 4)
+		{
+			GraphBuilder.QueueBufferUpload(Counters, &FUintVector4::ZeroValue, NumCounters * sizeof(uint32));
+		}
+		else
+		{
+			TArray<uint32> Zeroes;
+			Zeroes.SetNumZeroed(NumCounters);
+
+			GraphBuilder.QueueBufferUpload(Counters, Zeroes.GetData(), Zeroes.Num() * Zeroes.GetTypeSize());
+		}
+	}
 }
diff --git a/Engine/Plugins/PCG/Source/PCG/Private/Compute/DataInterfaces/PCGDataCollectionUploadDataInterface.h b/Engine/Plugins/PCG/Source/PCG/Private/Compute/DataInterfaces/PCGDataCollectionUploadDataInterface.h
--- a/Engine/Plugins/PCG/Source/PCG/Private/Compute/DataInterfaces/PCGDataCollectionUploadDataInterface.h
+++ b/Engine/Plugins/PCG/Source/PCG/Private/Compute/DataInterfaces/PCGDataCollectionUploadDataInterface.h
@@ -34,7 +34,7 @@ public:
 	//~ End UComputeDataProvider Interface
 
 	//~ Begin UPCGComputeDataProvider Interface
-	virtual bool PrepareForExecute_GameThread() override;
+	virtual bool PrepareForExecute_GameThread(const UPCGDataBinding* InBinding) override;
 	//~ End UPCGComputeDataProvider Interface
 
 	TArray<FName> DownstreamInputPinLabels;
@@ -52,7 +52,8 @@ public:
 		const TArray<FName>& InDownstreamInputPinLabels,
 		FName InOutputPinLabel,
 		FName InOutputPinLabelAlias,
-		EPCGExportMode InExportMode);
+		EPCGExportMode InExportMode,
+		bool bInExposeElementCounters);
 
 	//~ Begin FComputeDataProviderRenderProxy Interface
 	void AllocateResources(FRDGBuilder& GraphBuilder, FAllocationData const& InAllocationData) override;
diff --git a/Engine/Plugins/PCG/Source/PCG/Private/Compute/DataInterfaces/PCGDebugDataInterface.cpp b/Engine/Plugins/PCG/Source/PCG/Private/Compute/DataInterfaces/PCGDebugDataInterface.cpp
--- a/Engine/Plugins/PCG/Source/PCG/Private/Compute/DataInterfaces/PCGDebugDataInterface.cpp
+++ b/Engine/Plugins/PCG/Source/PCG/Private/Compute/DataInterfaces/PCGDebugDataInterface.cpp
@@ -91,8 +91,6 @@ FComputeDataProviderRenderProxy* UPCGDebugDataProvider::GetRenderProxy()
 		{
 			UE_LOG(LogPCG, Warning, TEXT("\t%d:\t%f"), Index, DataAsFloat[Index]);
 		}
-
-		DataProvider->bOperationComplete = true;
 	};
 
 	return new FPCGDebugDataProviderProxy(DebugBufferSizeFloats, ProcessReadbackData_RenderThread);
diff --git a/Engine/Plugins/PCG/Source/PCG/Private/Compute/DataInterfaces/PCGDebugDataInterface.h b/Engine/Plugins/PCG/Source/PCG/Private/Compute/DataInterfaces/PCGDebugDataInterface.h
--- a/Engine/Plugins/PCG/Source/PCG/Private/Compute/DataInterfaces/PCGDebugDataInterface.h
+++ b/Engine/Plugins/PCG/Source/PCG/Private/Compute/DataInterfaces/PCGDebugDataInterface.h
@@ -19,7 +19,7 @@ class PCG_API UPCGDebugDataInterface : public UComputeDataInterface
 
 public:
 	//~ Begin UComputeDataInterface Interface
-	TCHAR const* GetClassName() const override { return TEXT("KernelDebug"); }
+	TCHAR const* GetClassName() const override { return TEXT("PCGKernelDebug"); }
 	void GetSupportedOutputs(TArray<FShaderFunctionDefinition>& OutFunctions) const override;
 	void GetShaderParameters(TCHAR const* UID, FShaderParametersMetadataBuilder& InOutBuilder, FShaderParametersMetadataAllocations& InOutAllocations) const override;
 	void GetShaderHash(FString& InOutKey) const override;
@@ -46,7 +46,6 @@ public:
 	//~ End UComputeDataProvider Interface
 
 	uint32 DebugBufferSizeFloats = 0;
-	std::atomic<bool> bOperationComplete = false;
 };
 
 class FPCGDebugDataProviderProxy : public FComputeDataProviderRenderProxy
diff --git a/Engine/Plugins/PCG/Source/PCG/Private/Compute/DataInterfaces/PCGInstanceDataInterface.cpp b/Engine/Plugins/PCG/Source/PCG/Private/Compute/DataInterfaces/PCGInstanceDataInterface.cpp
--- a/Engine/Plugins/PCG/Source/PCG/Private/Compute/DataInterfaces/PCGInstanceDataInterface.cpp
+++ b/Engine/Plugins/PCG/Source/PCG/Private/Compute/DataInterfaces/PCGInstanceDataInterface.cpp
@@ -16,6 +16,7 @@
 #include "InstanceDataSceneProxy.h"
 #include "PrimitiveSceneInfo.h"
 #include "PrimitiveSceneProxy.h"
+#include "RenderCaptureInterface.h"
 #include "RendererUtils.h"
 #include "RenderGraphBuilder.h"
 #include "RenderGraphUtils.h"
@@ -37,6 +38,19 @@
 
 #define PCG_INSTANCE_DATA_LOGGING 0
 
+namespace PCGInstanceDataInterface
+{
+#if !UE_BUILD_SHIPPING
+	static int32 TriggerGPUCaptureDispatchIndex = 0;
+	static FAutoConsoleVariableRef CVarTriggerGPUCaptureDispatchIndex(
+		TEXT("pcg.GPU.Instancing.TriggerSceneWriterGPUCaptureDispatchIndex"),
+		TriggerGPUCaptureDispatchIndex,
+		TEXT("Index of the next dispatch to capture. I.e. if set to 2, will ignore one dispatch and then trigger a capture on the next one."),
+		ECVF_RenderThreadSafe
+	);
+#endif // !UE_BUILD_SHIPPING
+}
+
 void UPCGInstanceDataInterface::GetSupportedOutputs(TArray<FShaderFunctionDefinition>& OutFunctions) const
 {
 	OutFunctions.AddDefaulted_GetRef()
@@ -250,14 +264,25 @@ void FPCGInstanceDataProviderProxy::PostGraphDispatch(FRDGBuilder& GraphBuilder)
 
 		const int32 NumInstancesAllocatedInGPUScene = PrimitiveSceneProxy->GetInstanceDataHeader().NumInstances;
 
+#if UE_BUILD_SHIPPING
+		const bool bTriggerCapture = false;
+#else
+		const bool bTriggerCapture = (PCGInstanceDataInterface::TriggerGPUCaptureDispatchIndex == 1);
+		PCGInstanceDataInterface::TriggerGPUCaptureDispatchIndex = FMath::Max(PCGInstanceDataInterface::TriggerGPUCaptureDispatchIndex - 1, 0);
+#endif // UE_BUILD_SHIPPING
+
 		// Post graph dispatch, the instance data will have been populated. Enqueue a scene update to write the instance data into the scene.
 		// Will get executed on same frame and in same builder.
 		SceneComputeUpdatesBB->SceneComputeUpdates->EnqueueUpdate(SceneInfo,
-			FGPUSceneWriteDelegate::CreateLambda([PrimitiveIndex, NumInstancesAllocatedInGPUScene, CumulativeInstanceCount, InstanceDataSRV=InstanceDataSRV, InstanceCustomFloatDataSRV=InstanceCustomFloatDataSRV, WriteCountersSRV=WriteCountersSRV](FRDGBuilder& GraphBuilder, const FGPUSceneWriteDelegateParams& Params)
+			FGPUSceneWriteDelegate::CreateLambda([PrimitiveIndex, NumInstancesAllocatedInGPUScene, CumulativeInstanceCount, InstanceDataSRV=InstanceDataSRV, InstanceCustomFloatDataSRV=InstanceCustomFloatDataSRV, WriteCountersSRV=WriteCountersSRV, bTriggerCapture](FRDGBuilder& GraphBuilder, const FGPUSceneWriteDelegateParams& Params)
 			{
 				RDG_EVENT_SCOPE(GraphBuilder, "SceneComputeUpdateInterface->EnqueueUpdate");
 				check(Params.PersistentPrimitiveId != (uint32)INDEX_NONE);
 
+#if !UE_BUILD_SHIPPING
+				RenderCaptureInterface::FScopedCapture RenderCapture(bTriggerCapture, GraphBuilder, TEXT("SceneComputeUpdateInterface->EnqueueUpdate"));
+#endif // !UE_BUILD_SHIPPING
+
 				// Write instances
 				FPCGSceneWriterCS::FParameters* Parameters = GraphBuilder.AllocParameters<FPCGSceneWriterCS::FParameters>();
 				Parameters->InPrimitiveIndex = PrimitiveIndex;
diff --git a/Engine/Plugins/PCG/Source/PCG/Private/Compute/Elements/PCGComputeGraphElement.cpp b/Engine/Plugins/PCG/Source/PCG/Private/Compute/Elements/PCGComputeGraphElement.cpp
--- a/Engine/Plugins/PCG/Source/PCG/Private/Compute/Elements/PCGComputeGraphElement.cpp
+++ b/Engine/Plugins/PCG/Source/PCG/Private/Compute/Elements/PCGComputeGraphElement.cpp
@@ -113,9 +113,12 @@ bool FPCGComputeGraphElement::ExecuteInternal(FPCGContext* InContext) const
 		});
 	};
 
-	// 7. Execution is complete when graph was enqueued/dispatched, and any exports are prepared.
-	if (Context->bGraphEnqueued && Context->ProvidersWithBufferExports.IsEmpty())
+	// 7. Execution is complete when graph was enqueued/dispatched and all exports and readbacks are complete.
+	if (Context->bGraphEnqueued && Context->ProvidersWithBufferExports.IsEmpty() && Context->ProvidersWithReadbacks.IsEmpty())
 	{
+		// Apply any read back element counts to the final data description.
+		Context->DataBinding->ApplyElementCounts_GameThread();
+
 #if WITH_EDITOR
 		// Call both to start all required readbacks.
 		const bool bDebugDrawDone = ExecuteDebugDraw(Context);
@@ -152,7 +155,7 @@ bool FPCGComputeGraphElement::ExecuteInternal(FPCGContext* InContext) const
 	{
 		// Likely we need a frame to pass in order to make progress.
 		SleepUntilNextFrame(Context);
-		return true;
+		return false;
 	}
 
 	// 4. Initialize and parse incoming data for data sizes, attributes, etc that will drive buffer allocations and dispatch thread counts.
@@ -190,16 +193,20 @@ bool FPCGComputeGraphElement::ExecuteInternal(FPCGContext* InContext) const
 		for (UComputeDataProvider* ComputeDataProvider : Context->ComputeGraphInstance.GetDataProviders())
 		{
 			UPCGDataCollectionDataProvider* DataProvider = Cast<UPCGDataCollectionDataProvider>(ComputeDataProvider);
+			if (!DataProvider)
+			{
+				continue;
+			}
 
-			const bool bDataProviderDebugged = DataProvider && DataProvider->ProducerSettings && DataProvider->ProducerSettings->bDebug;
+			const bool bDataProviderDebugged = DataProvider->ProducerSettings && DataProvider->ProducerSettings->bDebug;
 
-			if (DataProvider && (DataProvider->IsExportRequired() || bDataProviderDebugged))
+			if (DataProvider->IsExportRequired() || bDataProviderDebugged)
 			{
 				Context->ProvidersWithBufferExports.Add(DataProvider);
 
 				TWeakObjectPtr<UPCGDataCollectionDataProvider> DataProviderWeak = DataProvider;
 
-				DataProvider->OnDataExported_GameThread().AddLambda([ContextHandle = Context->GetOrCreateHandle(), DataProviderWeak]()
+				DataProvider->OnDataExported_GameThread().AddLambda([ContextHandle=Context->GetOrCreateHandle(), DataProviderWeak]()
 				{
 					check(IsInGameThread());
 
@@ -222,6 +229,34 @@ bool FPCGComputeGraphElement::ExecuteInternal(FPCGContext* InContext) const
 					}
 				});
 			}
+
+			if (DataProvider->bExposeElementCounters)
+			{
+				Context->ProvidersWithReadbacks.Add(DataProvider);
+
+				TWeakObjectPtr<UPCGDataCollectionDataProvider> DataProviderWeak = DataProvider;
+
+				DataProvider->OnCountersReadBack_GameThread().AddLambda([ContextHandle=Context->GetOrCreateHandle(), DataProviderWeak]()
+				{
+					check(IsInGameThread());
+
+					TSharedPtr<FPCGContextHandle> SharedHandle = ContextHandle.Pin();
+					FPCGComputeGraphContext* ContextPtr = static_cast<FPCGComputeGraphContext*>(SharedHandle ? SharedHandle->GetContext() : nullptr);
+					UPCGDataCollectionDataProvider* DataProvider = DataProviderWeak.Get();
+					if (!DataProvider || !ContextPtr)
+					{
+						return;
+					}
+
+					ensure(ContextPtr->ProvidersWithReadbacks.Contains(DataProvider));
+					ContextPtr->ProvidersWithReadbacks.Remove(DataProvider);
+
+					if (ContextPtr->ProvidersWithReadbacks.IsEmpty())
+					{
+						ContextPtr->bIsPaused = false;
+					}
+				});
+			}
 		}
 
 		if (bAnyComponentsSetup)
@@ -273,7 +308,7 @@ bool FPCGComputeGraphElement::ExecuteInternal(FPCGContext* InContext) const
 		{
 			if (UPCGComputeDataProvider* DataProvider = Cast<UPCGComputeDataProvider>(ComputeDataProvider))
 			{
-				if (!DataProvider->PrepareForExecute_GameThread())
+				if (!DataProvider->PrepareForExecute_GameThread(Context->DataBinding.Get()))
 				{
 					bAllReady = false;
 				}
@@ -382,8 +417,8 @@ bool FPCGComputeGraphElement::SetupProceduralISMComponents(FPCGContext* InContex
 
 		const FPCGDataCollectionDesc InputDataDesc = PCGDataForGPUHelpers::ComputeInputPinDataDesc(SpawnerSettings, PCGPinConstants::DefaultInputLabel, InBinding);
 
-		const uint32 InputPointCount = InputDataDesc.ComputeDataElementCount(EPCGDataType::Point);
-		if (InputPointCount == 0)
+		const uint32 TotalInputPointCount = InputDataDesc.ComputeDataElementCount(EPCGDataType::Point);
+		if (TotalInputPointCount == 0)
 		{
 			continue;
 		}
@@ -430,28 +465,43 @@ bool FPCGComputeGraphElement::SetupProceduralISMComponents(FPCGContext* InContex
 		{
 			const FName SelectorName = SelectorByAttribute->AttributeName;
 
-			// Compute how many unique incoming values we will receive.
-			TSet<int32> UniqueStringKeys;
+			// Compute how many instances we expect for each mesh.
+			TMap<int32, uint32> StringKeyToInstanceCount;
 
 			for (const FPCGDataDesc& Desc : InputDataDesc.DataDescs)
 			{
+				if (Desc.ElementCount <= 0)
+				{
+					continue;
+				}
+
 				for (const FPCGKernelAttributeDesc& AttributeDesc : Desc.AttributeDescs)
 				{
 					if (AttributeDesc.Name == SelectorName)
 					{
-						if (AttributeDesc.Type != EPCGKernelAttributeType::StringKey)
+						if (AttributeDesc.Type == EPCGKernelAttributeType::StringKey)
+						{
+							ensure(SelectorAttributeId == INDEX_NONE || SelectorAttributeId == AttributeDesc.Index);
+							SelectorAttributeId = AttributeDesc.Index;
+
+							for (int32 AttributeStringKeyValue : AttributeDesc.GetUniqueStringKeys())
+							{
+								if (AttributeStringKeyValue > 0)
+								{
+									StringKeyToInstanceCount.FindOrAdd(AttributeStringKeyValue, 0) += Desc.ElementCount;
+								}
+							}
+
+							// We could early out but we currently continue looping to generate the below warning for mismatched types.
+						}
+						else
 						{
 							// Currently only a single type per attribute name is supported (the name/type in the attribute table). It's possible to wire a graph
 							// for which an attribute is present with multiple types. Warn if this is encountered.
 							PCG_KERNEL_VALIDATION_WARN(Context, SpawnerSettings, /*bQuiet=*/false, FText::Format(
 								LOCTEXT("MeshSelectorAttributeNotUsable", "Attribute '{0}' not usable for mesh selection, only attributes of type String Key are supported."),
 								FText::FromName(SelectorName)));
-
-							continue;
 						}
-
-						SelectorAttributeId = AttributeDesc.Index;
-						UniqueStringKeys.Append(AttributeDesc.GetUniqueStringKeys());
 					}
 				}
 			}
@@ -465,14 +515,14 @@ bool FPCGComputeGraphElement::SetupProceduralISMComponents(FPCGContext* InContex
 				continue;
 			}
 
-			for (int32 StringKey : UniqueStringKeys)
+			for (const TPair<int32, uint32>& StringKeyAndInstanceCount : StringKeyToInstanceCount)
 			{
-				if (!ensure(InBinding->GetStringTable().IsValidIndex(StringKey)))
+				if (!ensure(InBinding->GetStringTable().IsValidIndex(StringKeyAndInstanceCount.Key)))
 				{
 					continue;
 				}
 
-				const FString& MeshPathString = InBinding->GetStringTable()[StringKey];
+				const FString& MeshPathString = InBinding->GetStringTable()[StringKeyAndInstanceCount.Key];
 				if (MeshPathString.IsEmpty())
 				{
 					continue;
@@ -497,12 +547,12 @@ bool FPCGComputeGraphElement::SetupProceduralISMComponents(FPCGContext* InContex
 
 				FPCGProceduralISMComponentDescriptor Descriptor;
 				Descriptor = SelectorByAttribute->TemplateDescriptor;
-				Descriptor.NumInstances = InputPointCount;
+				Descriptor.NumInstances = StringKeyAndInstanceCount.Value;
 				Descriptor.LocalBounds = LocalBounds;
 				Descriptor.NumCustomFloats = CustomFloatCount;
 				Descriptor.StaticMesh = StaticMesh;
 
-				PrimitiveStringKeys.Emplace(StringKey);
+				PrimitiveStringKeys.Add(StringKeyAndInstanceCount.Key);
 				ComponentsToCreate.Add(MoveTemp(Descriptor));
 
 				if (SpawnerSettings->bApplyMeshBoundsToPoints)
@@ -565,24 +615,24 @@ bool FPCGComputeGraphElement::SetupProceduralISMComponents(FPCGContext* InContex
 					Descriptor.NumCustomFloats = CustomFloatCount;
 					Descriptor.StaticMesh = StaticMesh;
 
-					uint32 InstanceCount = FMath::CeilToInt(InputPointCount * Weight);
+					uint32 InstanceCount = FMath::CeilToInt(TotalInputPointCount * Weight);
 
 					if (SelectorWeighted->MeshEntries.Num() > 1)
 					{
 						// Since we'll be selecting meshes based on random draws using the point random seeds which we don't have on CPU,
 						// we may pick more or less than the expected number of instances for each mesh. Use binomial variance to calculate
 						// the overallocation that gives 99.7% confidence (3 sigma).
-						const double Variance = double(InputPointCount) * Weight * (1.0f - Weight);
+						const double Variance = double(TotalInputPointCount) * Weight * (1.0f - Weight);
 						const double Sigma = FMath::Sqrt(Variance);
 						const uint32 AdditionalAllocation = FMath::CeilToInt(3.0f * Sigma);
 
 						// Useful for debugging instance counts.
-						//UE_LOG(LogPCG, Warning, TEXT("Over allocation, N = %d, P = %.2f, NP = %d, 3sig = %f, Extra = %d (%.2f%%), Calculated = %d, Final = %d"), InputPointCount, Weight, InstanceCount, float(3.0f * Sigma), AdditionalAllocation, float(AdditionalAllocation) / InputPointCount, InstanceCount + AdditionalAllocation, FMath::Min(InstanceCount + AdditionalAllocation, InputPointCount));
+						//UE_LOG(LogPCG, Warning, TEXT("Over allocation, N = %d, P = %.2f, NP = %d, 3sig = %f, Extra = %d (%.2f%%), Calculated = %d, Final = %d"), TotalInputPointCount, Weight, InstanceCount, float(3.0f * Sigma), AdditionalAllocation, float(AdditionalAllocation) / TotalInputPointCount, InstanceCount + AdditionalAllocation, FMath::Min(InstanceCount + AdditionalAllocation, TotalInputPointCount));
 
 						InstanceCount += AdditionalAllocation;
 					}
 
-					Descriptor.NumInstances = FMath::Min(InstanceCount, InputPointCount);
+					Descriptor.NumInstances = FMath::Min(InstanceCount, TotalInputPointCount);
 
 					ComponentsToCreate.Add(MoveTemp(Descriptor));
 
@@ -647,7 +697,7 @@ bool FPCGComputeGraphElement::SetupProceduralISMComponents(FPCGContext* InContex
 		}
 
 		// Useful for debugging instance counts.
-		//UE_LOG(LogPCG, Log, TEXT("Input points %d, total instance count %d (+ %.2f%%)"), InputPointCount, Primitives.NumInstancesAllPrimitives, (float(Primitives.NumInstancesAllPrimitives - InputPointCount) / InputPointCount) * 100.0f);
+		//UE_LOG(LogPCG, Log, TEXT("Input points %d, total instance count %d (+ %.2f%%)"), TotalInputPointCount, Primitives.NumInstancesAllPrimitives, (float(Primitives.NumInstancesAllPrimitives - TotalInputPointCount) / TotalInputPointCount) * 100.0f);
 	}
 
 	return bAnyComponentsSetup;
@@ -805,9 +855,9 @@ bool FPCGComputeGraphElement::StoreDataForInspection(FPCGComputeGraphContext* In
 	for (TPair<const UPCGSettings*, FPCGDataCollection>& SettingsAndData : SettingsToDataCollection)
 	{
 		// Required by inspection code.
-		SettingsAndData.Get<1>().ComputeCrcs(/*bFullDataCrc=*/false);
+		SettingsAndData.Value.ComputeCrcs(/*bFullDataCrc=*/false);
 
-		const UPCGNode* Node = Cast<UPCGNode>(SettingsAndData.Get<0>()->GetOuter());
+		const UPCGNode* Node = Cast<UPCGNode>(SettingsAndData.Key->GetOuter());
 
 		// TODO: Input data not yet supported.
 		InContext->SourceComponent->StoreInspectionData(InContext->Stack, Node, /*InTimer=*/nullptr, /*InInputData=*/{}, SettingsAndData.Get<1>(), /*bUsedCache*/false);
diff --git a/Engine/Plugins/PCG/Source/PCG/Private/Compute/Elements/PCGComputeGraphElement.h b/Engine/Plugins/PCG/Source/PCG/Private/Compute/Elements/PCGComputeGraphElement.h
--- a/Engine/Plugins/PCG/Source/PCG/Private/Compute/Elements/PCGComputeGraphElement.h
+++ b/Engine/Plugins/PCG/Source/PCG/Private/Compute/Elements/PCGComputeGraphElement.h
@@ -43,6 +43,8 @@ public:
 	* back to main thread, upon which a reference is taken to the buffer and the provider is removed from this set to signal completion.
 	*/
 	TSet<TObjectPtr<UComputeDataProvider>> ProvidersWithBufferExports;
+
+	TSet<TObjectPtr<UComputeDataProvider>> ProvidersWithReadbacks;
 };
 
 /** Executes a CF graph. Created by the compiler when collapsing GPU nodes rather than by a settings/node. */
diff --git a/Engine/Plugins/PCG/Source/PCG/Private/Compute/PCGComputeCommon.cpp b/Engine/Plugins/PCG/Source/PCG/Private/Compute/PCGComputeCommon.cpp
--- a/Engine/Plugins/PCG/Source/PCG/Private/Compute/PCGComputeCommon.cpp
+++ b/Engine/Plugins/PCG/Source/PCG/Private/Compute/PCGComputeCommon.cpp
@@ -6,8 +6,9 @@
 #include "PCGContext.h"
 #include "PCGParamData.h"
 #include "PCGSubsystem.h"
-#include "Compute/Data/PCGProxyForGPUData.h"
+#include "Compute/PCGDataBinding.h"
 #include "Compute/PCGPinPropertiesGPU.h"
+#include "Compute/Data/PCGProxyForGPUData.h"
 #include "Data/PCGPointData.h"
 
 #include "DynamicRHI.h"
@@ -160,7 +161,7 @@ namespace PCGComputeHelpers
 							ClampedDataIndex = 0;
 						}
 
-						const FPCGDataDesc& InputDataDesc = InputDesc.DataDescs.IsValidIndex(ClampedDataIndex) ? InputDesc.DataDescs[DataIndex] : FPCGDataDesc(EPCGDataType::Any, 0);
+						const FPCGDataDesc& InputDataDesc = InputDesc.DataDescs.IsValidIndex(ClampedDataIndex) ? InputDesc.DataDescs[ClampedDataIndex] : FPCGDataDesc(EPCGDataType::Any, 0);
 
 						if (Props.ElementCountMode == EPCGElementCountMode::FromInputData)
 						{
@@ -332,12 +333,45 @@ namespace PCGComputeHelpers
 	{
 		if (const UPCGProxyForGPUData* Proxy = Cast<UPCGProxyForGPUData>(InData))
 		{
-			return Proxy->GetDescription(OutDescription);
+			if (!Proxy->GetDescription(OutDescription))
+			{
+				return false;
+			}
+
+			if (TSharedPtr<const FPCGProxyForGPUDataCollection> DataGPU = Proxy->GetGPUInfo())
+			{
+				for (FPCGKernelAttributeDesc& AttrDesc : OutDescription.AttributeDescs)
+				{
+					const int32 Index = PCGComputeHelpers::GetMetadataAttributeIndexFromAttributeId(AttrDesc.Index);
+					if (Index == INDEX_NONE)
+					{
+						continue;
+					}
+
+					if (const int32* AttributeIndex = InBinding->GetAttributeIdRemap().Find(UPCGDataBinding::FSourceBufferAndAttributeIndex(DataGPU, Index)))
+					{
+						const int32 NewId = PCGComputeHelpers::GetAttributeIdFromMetadataAttributeIndex(*AttributeIndex);
+						//UE_LOG(LogPCG, Warning, TEXT("Remapped attribute '%s' ID: %d -> %d"), *AttrDesc.Name.ToString(), AttrDesc.Index, NewId);
+						AttrDesc.Index = NewId;
+					}
+				}
+			}
 		}
 		else
 		{
 			OutDescription = FPCGDataDesc(InData, InBinding);
-			return true;
 		}
+
+		return true;
+	}
+
+	int32 GetAttributeIdFromMetadataAttributeIndex(int32 InAttributeIndex)
+	{
+		return InAttributeIndex >= 0 ? (InAttributeIndex + PCGComputeConstants::NUM_RESERVED_ATTRS) : INDEX_NONE;
+	}
+
+	int32 GetMetadataAttributeIndexFromAttributeId(int32 InAttributeId)
+	{
+		return InAttributeId >= PCGComputeConstants::NUM_RESERVED_ATTRS ? (InAttributeId - PCGComputeConstants::NUM_RESERVED_ATTRS) : INDEX_NONE;
 	}
 }
diff --git a/Engine/Plugins/PCG/Source/PCG/Private/Compute/PCGDataBinding.cpp b/Engine/Plugins/PCG/Source/PCG/Private/Compute/PCGDataBinding.cpp
--- a/Engine/Plugins/PCG/Source/PCG/Private/Compute/PCGDataBinding.cpp
+++ b/Engine/Plugins/PCG/Source/PCG/Private/Compute/PCGDataBinding.cpp
@@ -35,6 +35,37 @@ void UPCGDataBinding::Initialize(
 #endif
 }
 
+bool UPCGDataBinding::GetOutputPinDataDesc_GameThread(FName InProducerPinLabel, const UPCGSettings* InProducerSettings, FPCGDataCollectionDesc& OutDescription) const
+{
+	check(InProducerSettings);
+	const UPCGNode* Node = CastChecked<UPCGNode>(InProducerSettings->GetOuter());
+
+	return GetOutputPinDataDesc_GameThread(Node->GetOutputPin(InProducerPinLabel), InProducerSettings, OutDescription);
+}
+
+bool UPCGDataBinding::GetOutputPinDataDesc_GameThread(const UPCGPin* InProducerPin, const UPCGSettings* InProducerSettings, FPCGDataCollectionDesc& OutDescription) const
+{
+	ensure(IsInGameThread());
+
+	if (const FPCGDataCollectionDesc* FoundDescription = OutputPinDataDescriptionCache.Find(InProducerPin))
+	{
+		OutDescription = *FoundDescription;
+		return true;
+	}
+
+	FPCGDataCollectionDesc Description;
+	if (InProducerSettings->ComputeOutputPinDataDesc(InProducerPin, this, Description))
+	{
+		OutputPinDataDescriptionCache.Add(InProducerPin, Description);
+		OutDescription = MoveTemp(Description);
+		return true;
+	}
+	else
+	{
+		return false;
+	}
+}
+
 bool UPCGDataBinding::ComputeCPUOutputPinDataDesc(const UPCGPin* OutputPin, FPCGDataCollectionDesc& OutDesc) const
 {
 	check(OutputPin);
@@ -93,34 +124,22 @@ void UPCGDataBinding::ReceiveDataFromGPU_GameThread(
 	}
 }
 
-bool UPCGDataBinding::GetOutputPinDataDesc_GameThread(FName InProducerPinLabel, const UPCGSettings* InProducerSettings, FPCGDataCollectionDesc& OutDescription) const
+void UPCGDataBinding::ReceiveElementCountsFromGPU_GameThread(FName InPinLabel, const TArray<uint32>& InCounters)
 {
-	check(InProducerSettings);
-	const UPCGNode* Node = CastChecked<UPCGNode>(InProducerSettings->GetOuter());
-
-	return GetOutputPinDataDesc_GameThread(Node->GetOutputPin(InProducerPinLabel), InProducerSettings, OutDescription);
+	PinToElementCounts.FindOrAdd(InPinLabel) = InCounters;
 }
 
-bool UPCGDataBinding::GetOutputPinDataDesc_GameThread(const UPCGPin* InProducerPin, const UPCGSettings* InProducerSettings, FPCGDataCollectionDesc& OutDescription) const
+void UPCGDataBinding::ApplyElementCounts_GameThread()
 {
-	ensure(IsInGameThread());
-
-	if (const FPCGDataCollectionDesc* FoundDescription = OutputPinDataDescriptionCache.Find(InProducerPin))
+	for (FPCGTaggedData& Output : OutputDataCollection.TaggedData)
 	{
-		OutDescription = *FoundDescription;
-		return true;
-	}
-
-	FPCGDataCollectionDesc Description;
-	if (InProducerSettings->ComputeOutputPinDataDesc(InProducerPin, this, Description))
-	{
-		OutputPinDataDescriptionCache.Add(InProducerPin, Description);
-		OutDescription = MoveTemp(Description);
-		return true;
-	}
-	else
-	{
-		return false;
+		if (const TArray<uint32>* FoundElementCounts = PinToElementCounts.Find(Output.Pin))
+		{
+			if (UPCGProxyForGPUData* Data = const_cast<UPCGProxyForGPUData*>(Cast<UPCGProxyForGPUData>(Output.Data)))
+			{
+				Data->UpdateElementCountsFromReadback(*FoundElementCounts);
+			}
+		}
 	}
 }
 
@@ -163,22 +182,45 @@ void UPCGDataBinding::AddInputDataAttributesToTable(const FPCGDataCollection& In
 {
 	for (const FPCGTaggedData& Data : InComputeGraphElementInputData.TaggedData)
 	{
-		const UPCGMetadata* Metadata = Data.Data ? Data.Data->ConstMetadata() : nullptr;
-		if (!Metadata)
+		if (const UPCGProxyForGPUData* Proxy = Cast<UPCGProxyForGPUData>(Data.Data))
 		{
-			continue;
-		}
+			FPCGDataDesc SourceDesc;
+			if (Proxy->GetDescription(SourceDesc))
+			{
+				for (FPCGKernelAttributeDesc& SourceAttrDesc : SourceDesc.AttributeDescs)
+				{
+					// Convert from absolute ID to index into metadata attributes.
+					const int32 SourceAttributeTableIndex = PCGComputeHelpers::GetMetadataAttributeIndexFromAttributeId(SourceAttrDesc.Index);
 
-		TArray<FName> AttributeNames;
-		TArray<EPCGMetadataTypes> AttributeTypes;
-		Metadata->GetAttributes(AttributeNames, AttributeTypes);
+					// Negative index means the attribute is a property rather than metadata.
+					if (SourceAttributeTableIndex < 0)
+					{
+						continue;
+					}
 
-		for (int Index = 0; Index < AttributeNames.Num(); ++Index)
+					const int32 ExistingAttributeId = AttributeTable.AddAttribute(SourceAttrDesc.Name, SourceAttrDesc.Type);
+
+					if (ExistingAttributeId != SourceAttributeTableIndex)
+					{
+						AttributeIdRemap.Add(FSourceBufferAndAttributeIndex(Proxy->GetGPUInfo(), SourceAttributeTableIndex), ExistingAttributeId);
+					}
+				}
+			}
+
+		}
+		else if (const UPCGMetadata* Metadata = Data.Data ? Data.Data->ConstMetadata() : nullptr)
 		{
-			const FName AttributeName = AttributeNames[Index];
-			const EPCGKernelAttributeType AttributeType = PCGDataForGPUHelpers::GetAttributeTypeFromMetadataType(AttributeTypes[Index]);
+			TArray<FName> AttributeNames;
+			TArray<EPCGMetadataTypes> AttributeTypes;
+			Metadata->GetAttributes(AttributeNames, AttributeTypes);
 
-			AttributeTable.AddAttribute(AttributeName, AttributeType);
+			for (int Index = 0; Index < AttributeNames.Num(); ++Index)
+			{
+				const FName AttributeName = AttributeNames[Index];
+				const EPCGKernelAttributeType AttributeType = PCGDataForGPUHelpers::GetAttributeTypeFromMetadataType(AttributeTypes[Index]);
+
+				AttributeTable.AddAttribute(AttributeName, AttributeType);
+			}
 		}
 	}
 }
@@ -196,6 +238,19 @@ void UPCGDataBinding::BuildStringTable()
 
 	for (const FPCGTaggedData& Data : InputDataCollection.TaggedData)
 	{
+		if (const UPCGProxyForGPUData* Proxy = Cast<UPCGProxyForGPUData>(Data.Data))
+		{
+			if (TSharedPtr<const FPCGProxyForGPUDataCollection> GPUInfo = Proxy->GetGPUInfo())
+			{
+				for (const FString& String : GPUInfo->GetStringTable())
+				{
+					StringTable.AddUnique(String);
+				}
+
+				continue;
+			}
+		}
+
 		const UPCGMetadata* Metadata = Data.Data ? Data.Data->ConstMetadata() : nullptr;
 		if (!Metadata)
 		{
diff --git a/Engine/Plugins/PCG/Source/PCG/Private/Compute/PCGDataBinding.h b/Engine/Plugins/PCG/Source/PCG/Private/Compute/PCGDataBinding.h
--- a/Engine/Plugins/PCG/Source/PCG/Private/Compute/PCGDataBinding.h
+++ b/Engine/Plugins/PCG/Source/PCG/Private/Compute/PCGDataBinding.h
@@ -65,14 +65,22 @@ public:
 	/** Get the set of all unique strings entering the compute graph. */
 	const TArray<FString>& GetStringTable() const { return StringTable; }
 
+	/** Get description of data produced by an output pin. */
+	bool GetOutputPinDataDesc_GameThread(FName InProducerPinLabel, const UPCGSettings* InProducerSettings, FPCGDataCollectionDesc& OutDescription) const;
+	bool GetOutputPinDataDesc_GameThread(const UPCGPin* InProducerPin, const UPCGSettings* InProducerSettings, FPCGDataCollectionDesc& OutDescription) const;
+
 	/** If the output pin is a connected upstream pin of the compute graph (on a CPU node or a GPU node from a different compute graph), returns a description of this data. */
 	bool ComputeCPUOutputPinDataDesc(const UPCGPin* OutputPin, FPCGDataCollectionDesc& OutDesc) const;
 
 	void ReceiveDataFromGPU_GameThread(const UPCGProxyForGPUData* InData, const UPCGSettings* InProducerSettings, EPCGExportMode InExportMode, FName InPinLabel, FName InPinLabelAlias);
 
-	/** Get description of data produced by an output pin. */
-	bool GetOutputPinDataDesc_GameThread(FName InProducerPinLabel, const UPCGSettings* InProducerSettings, FPCGDataCollectionDesc& OutDescription) const;
-	bool GetOutputPinDataDesc_GameThread(const UPCGPin* InProducerPin, const UPCGSettings* InProducerSettings, FPCGDataCollectionDesc& OutDescription) const;
+	void ReceiveElementCountsFromGPU_GameThread(FName InPinLabel, const TArray<uint32>& InCounters);
+
+	void ApplyElementCounts_GameThread();
+
+	using FSourceBufferAndAttributeIndex = TPair<TSharedPtr<const FPCGProxyForGPUDataCollection>, /*Source Metadata Attribute Index*/int32>;
+
+	const TMap<FSourceBufferAndAttributeIndex, /*Attribute Index*/int32>& GetAttributeIdRemap() const { return AttributeIdRemap; }
 
 private:
 	void InitializeInputData(const FPCGDataCollection& InComputeGraphElementInputData);
@@ -115,6 +123,9 @@ public:
 	TArray<FPCGDataToDebug> DataToDebug;
 	TArray<FPCGDataToDebug> DataToInspect;
 
+	/** Final element counts read back from GPU. */
+	TMap<FName, TArray<uint32>> PinToElementCounts;
+
 protected:
 	/** All attributes present in graph at execution time. Coherent across all graph branches. Seeded using statically known created attributes at compile time, then
 	 * augmented with incoming attributes from input data collection at runtime. */
@@ -124,6 +135,9 @@ protected:
 	UPROPERTY()
 	TArray<FString> StringTable;
 
+	/** Remaps attribute IDs to avoid collisions with this compute graph's attribute IDs. */
+	TMap<FSourceBufferAndAttributeIndex, /*Attribute ID*/int32> AttributeIdRemap;
+
 	/** Cache of data descriptions computed for output pins to amortized cost of building data descriptions. */
 	mutable TMap<const UPCGPin*, FPCGDataCollectionDesc> OutputPinDataDescriptionCache;
 };
diff --git a/Engine/Plugins/PCG/Source/PCG/Private/Compute/PCGDataForGPU.cpp b/Engine/Plugins/PCG/Source/PCG/Private/Compute/PCGDataForGPU.cpp
--- a/Engine/Plugins/PCG/Source/PCG/Private/Compute/PCGDataForGPU.cpp
+++ b/Engine/Plugins/PCG/Source/PCG/Private/Compute/PCGDataForGPU.cpp
@@ -656,17 +656,25 @@ uint32 GetTypeHash(const FPCGKernelAttributeKey& In)
 
 int32 FPCGKernelAttributeTable::GetAttributeId(const FPCGKernelAttributeKey& InAttribute) const
 {
-	return AttributeTable.IndexOfByKey(InAttribute) + PCGComputeConstants::NUM_RESERVED_ATTRS;
+	const int TableIndex = AttributeTable.IndexOfByKey(InAttribute);
+	if (ensureMsgf(TableIndex >= 0, TEXT("Missing attribute '%s', has this attribute been declared using GetKernelAttributeKeys()?"), *InAttribute.Name.ToString()))
+	{
+		return PCGComputeHelpers::GetAttributeIdFromMetadataAttributeIndex(TableIndex);
+	}
+	else
+	{
+		return INDEX_NONE;
+	}
 }
 
 int32 FPCGKernelAttributeTable::GetAttributeId(FName InName, EPCGKernelAttributeType InType) const
 {
-	return AttributeTable.IndexOfByKey(FPCGKernelAttributeKey(InName, InType)) + PCGComputeConstants::NUM_RESERVED_ATTRS;
+	return GetAttributeId(FPCGKernelAttributeKey(InName, InType));
 }
 
-void FPCGKernelAttributeTable::AddAttribute(FName InName, EPCGKernelAttributeType InType)
+int32 FPCGKernelAttributeTable::AddAttribute(FName InName, EPCGKernelAttributeType InType)
 {
-	AttributeTable.AddUnique(FPCGKernelAttributeKey(InName, InType));
+	return AttributeTable.AddUnique(FPCGKernelAttributeKey(InName, InType));
 }
 
 #if PCG_KERNEL_LOGGING_ENABLED
@@ -692,6 +700,11 @@ void FPCGKernelAttributeDesc::AddUniqueStringKeys(const TArray<int32>& InOtherSt
 	}
 }
 
+void FPCGKernelAttributeDesc::SetStringKeys(const TArrayView<int32>& InStringKeys)
+{
+	UniqueStringKeys = InStringKeys;
+}
+
 bool FPCGKernelAttributeDesc::operator==(const FPCGKernelAttributeDesc& Other) const
 {
 	return Index == Other.Index && Type == Other.Type && Name == Other.Name;
@@ -744,6 +757,33 @@ bool FPCGDataDesc::ContainsAttribute(FName InAttributeName) const
 	});
 }
 
+void FPCGDataDesc::AddAttribute(FPCGKernelAttributeKey InAttribute, const UPCGDataBinding* InBinding, const TArray<int32>* InOptionalUniqueStringKeys)
+{
+	const int32 AttributeId = InBinding->GetAttributeId(InAttribute);
+	if (AttributeId == INDEX_NONE)
+	{
+		return;
+	}
+
+	bool bAttributeFound = false;
+
+	for (FPCGKernelAttributeDesc& AttributeDesc : AttributeDescs)
+	{
+		if (AttributeDesc.Name == InAttribute.Name)
+		{
+			ensureMsgf(!bAttributeFound, TEXT("Attribute '%s' encountered multiple times in a single data."), *AttributeDesc.Name.ToString());
+			bAttributeFound = true;
+
+			AttributeDesc.Type = InAttribute.Type;
+		}
+	}
+
+	if (!bAttributeFound)
+	{
+		AttributeDescs.Emplace(AttributeId, InAttribute.Type, InAttribute.Name, InOptionalUniqueStringKeys);
+	}
+}
+
 void FPCGDataDesc::InitializeAttributeDescs(const UPCGData* InData, const UPCGDataBinding* InBinding)
 {
 	TRACE_CPUPROFILER_EVENT_SCOPE(FPCGDataDesc::InitializeAttributeDescs);
@@ -878,7 +918,7 @@ void FPCGDataDesc::InitializeAttributeDescs(const UPCGData* InData, const UPCGDa
 
 			if (AttributeId != INDEX_NONE)
 			{
-				AttributeDescs.Emplace(AttributeId, AttributeType, AttributeName, MoveTemp(UniqueStringKeys));
+				AttributeDescs.Emplace(AttributeId, AttributeType, AttributeName, &UniqueStringKeys);
 			}
 		}
 	}
@@ -1268,7 +1308,13 @@ EPCGUnpackDataCollectionResult FPCGDataCollectionDesc::UnpackDataCollection(FPCG
 
 		const uint32 TypeId =        DataAsUint[DataHeaderAddress + 0];
 		const uint32 NumAttributes = DataAsUint[DataHeaderAddress + 1];
-		const uint32 NumElements =   DataAsUint[DataHeaderAddress + 2];
+		
+		uint32 NumElements = DataAsUint[DataHeaderAddress + 2];
+		// Use tighter/more refined element count if we have one:
+		if (DataDescs[DataIndex].ElementCount >= 0)
+		{
+			NumElements = FMath::Min(NumElements, static_cast<uint32>(DataDescs[DataIndex].ElementCount));
+		}
 
 		const FPCGDataDesc& DataDesc = DataDescs[DataIndex];
 		const TArray<FPCGKernelAttributeDesc>& AttributeDescs = DataDesc.AttributeDescs;
@@ -1588,33 +1634,30 @@ bool FPCGDataCollectionDesc::ContainsAttributeOnAnyData(FName InAttributeName) c
 	});
 }
 
-void FPCGDataCollectionDesc::AddAttributeToAllData(FPCGKernelAttributeKey InAttribute, const UPCGDataBinding* InBinding, TArray<int32>* InOptionalUniqueStringKeys)
+void FPCGDataCollectionDesc::AddAttributeToAllData(FPCGKernelAttributeKey InAttribute, const UPCGDataBinding* InBinding, const TArray<int32>* InOptionalUniqueStringKeys)
 {
-	const int32 AttributeId = InBinding->GetAttributeId(InAttribute);
-	if (AttributeId == INDEX_NONE)
+	for (FPCGDataDesc& DataDesc : DataDescs)
 	{
-		return;
+		DataDesc.AddAttribute(InAttribute, InBinding, InOptionalUniqueStringKeys);
 	}
+}
 
-	for (FPCGDataDesc& DataDesc : DataDescs)
+void FPCGDataCollectionDesc::GetUniqueStringKeyValues(int32 InAttributeId, TArray<int32>& OutUniqueStringKeys) const
+{
+	for (const FPCGDataDesc& DataDesc : DataDescs)
 	{
-		bool bAttributeFound = false;
-
-		for (FPCGKernelAttributeDesc& AttributeDesc : DataDesc.AttributeDescs)
+		for (const FPCGKernelAttributeDesc& AttributeDesc : DataDesc.AttributeDescs)
 		{
-			if (AttributeDesc.Name == InAttribute.Name)
+			if (AttributeDesc.Index == InAttributeId)
 			{
-				ensureMsgf(!bAttributeFound, TEXT("Attribute '%s' encountered multiple times in a single data."), *AttributeDesc.Name.ToString());
-				bAttributeFound = true;
+				for (int32 StringKey : AttributeDesc.GetUniqueStringKeys())
+				{
+					OutUniqueStringKeys.AddUnique(StringKey);
+				}
 
-				AttributeDesc.Type = InAttribute.Type;
+				break;
 			}
 		}
-
-		if (!bAttributeFound)
-		{
-			DataDesc.AttributeDescs.Emplace(AttributeId, InAttribute.Type, InAttribute.Name, InOptionalUniqueStringKeys ? *InOptionalUniqueStringKeys : TArray<int32>());
-		}
 	}
 }
 
@@ -1723,4 +1766,19 @@ bool FPCGProxyForGPUDataCollection::GetCPUData(FPCGContext* InContext, int32 InD
 	return true;
 }
 
+void FPCGProxyForGPUDataCollection::UpdateElementCountsFromReadback(const TArray<uint32>& InElementCounts)
+{
+	uint32 DataCounter = 0;
+
+	if (!ensure(Description.DataDescs.Num() == InElementCounts.Num()))
+	{
+		return;
+	}
+
+	for (FPCGDataDesc& DataDesc : Description.DataDescs)
+	{
+		DataDesc.ElementCount = InElementCounts[DataCounter++];
+	}
+}
+
 #undef LOCTEXT_NAMESPACE
diff --git a/Engine/Plugins/PCG/Source/PCG/Private/Compute/PCGDataUploadAdaptor.cpp b/Engine/Plugins/PCG/Source/PCG/Private/Compute/PCGDataUploadAdaptor.cpp
new file mode 100644
--- /dev/null
+++ b/Engine/Plugins/PCG/Source/PCG/Private/Compute/PCGDataUploadAdaptor.cpp
@@ -0,0 +1,271 @@
+// Copyright Epic Games, Inc. All Rights Reserved.
+
+#include "Compute/PCGDataUploadAdaptor.h"
+
+#include "PCGDataCollectionAdaptorCS.h"
+#include "PCGModule.h"
+#include "Compute/PCGComputeCommon.h"
+#include "Compute/PCGComputeGraph.h"
+#include "Compute/PCGDataBinding.h"
+#include "Compute/Data/PCGProxyForGPUData.h"
+
+#include "RenderCaptureInterface.h"
+#include "RenderGraphBuilder.h"
+#include "RenderGraphResources.h"
+#include "RenderGraphUtils.h"
+
+namespace PCGDataUploadAdaptor
+{
+#if !UE_BUILD_SHIPPING
+	static int32 TriggerGPUCaptureDispatches = 0;
+	static FAutoConsoleVariableRef CVarTriggerGPUCapture(
+		TEXT("pcg.GPU.DataAdaptor.TriggerGPUCaptures"),
+		TriggerGPUCaptureDispatches,
+		TEXT("Trigger GPU captures for this many of the subsequent data collection update dispatches."));
+#endif // !UE_BUILD_SHIPPING
+}
+
+FPCGDataUploadAdaptor::FPCGDataUploadAdaptor(UPCGDataBinding* InDataBinding, const FPCGDataCollectionDesc& InTargetDataCollectionDesc, FName InInputPinLabel)
+{
+	check(IsInGameThread());
+
+	TargetDataCollectionDesc = InTargetDataCollectionDesc;
+	InputPinLabel = InInputPinLabel;
+	DataBinding = InDataBinding;
+
+	// Are all our data items coming direct from GPU - and importantly are a complete, single data collection?
+	bool bMultipleBuffersEncountered = false;
+	TSharedPtr<const FPCGProxyForGPUDataCollection> UniqueDataCollection;
+	for (int InputDataIndex = 0; InputDataIndex < DataBinding->InputDataCollection.TaggedData.Num(); ++InputDataIndex)
+	{
+		const UPCGProxyForGPUData* DataGPU = Cast<UPCGProxyForGPUData>(DataBinding->InputDataCollection.TaggedData[InputDataIndex].Data);
+		if (!DataGPU)
+		{
+			// Encountered CPU data, no reuse for now.
+			// TODO: Support multiple input buffers and only upload required data.
+			UniqueDataCollection.Reset();
+			break;
+		}
+
+		TSharedPtr<const FPCGProxyForGPUDataCollection> InputDataCollection = DataGPU->GetInputDataCollectionInfo();
+		if (!InputDataCollection)
+		{
+			UE_LOG(LogPCG, Error, TEXT("No reuse: Missing data collection buffer!"));
+
+			UniqueDataCollection.Reset();
+			break;
+		}
+
+		if (!UniqueDataCollection)
+		{
+			UniqueDataCollection = InputDataCollection;
+		}
+		else if (UniqueDataCollection != InputDataCollection || DataGPU->GetDataIndexInCollection() != InputDataIndex)
+		{
+			// Multiple input buffers or out of order data items, no reuse for now.
+			// TODO: support indirection of data items.
+			UniqueDataCollection.Reset();
+			break;
+		}
+	}
+
+	if (UniqueDataCollection && UniqueDataCollection->GetDescription().DataDescs.Num() != DataBinding->InputDataCollection.TaggedData.Num())
+	{
+		// We're not using all the data items from the buffer, no reuse for now.
+		// TODO: Support indirection of data items.
+		UniqueDataCollection.Reset();
+	}
+
+	if (!UniqueDataCollection)
+	{
+		return;
+	}
+
+	ExternalBufferForReuse = UniqueDataCollection->GetBuffer();
+	ExternalBufferSizeBytes = UniqueDataCollection->GetBufferSizeBytes();
+}
+
+/** Do any preparation work such as data readbacks. Returns true when preparation is complete. */
+bool FPCGDataUploadAdaptor::PrepareData_GameThread(const UPCGDataBinding* InBinding)
+{
+	check(InBinding);
+	check(IsInGameThread());
+
+	if (ExternalBufferForReuse)
+	{
+		SourceToTargetAttributeId.Reserve(InBinding->GetAttributeIdRemap().Num());
+
+		Algo::Transform(InBinding->GetAttributeIdRemap(), SourceToTargetAttributeId, [](const TPair<UPCGDataBinding::FSourceBufferAndAttributeIndex, int32>& Remapping)
+		{
+			const int32 TargetAttributeId = PCGComputeHelpers::GetAttributeIdFromMetadataAttributeIndex(Remapping.Value);
+			const int32 SourceAttributeId = PCGComputeHelpers::GetAttributeIdFromMetadataAttributeIndex(Remapping.Key.Value);
+
+			ensure(TargetAttributeId >= PCGComputeConstants::NUM_RESERVED_ATTRS);
+			ensure(SourceAttributeId >= PCGComputeConstants::NUM_RESERVED_ATTRS);
+
+			return FIntVector2(SourceAttributeId, TargetAttributeId);
+		});
+	}
+	else
+	{
+		UPCGDataBinding* Binding = DataBinding.Get();
+		if (!Binding)
+		{
+			return true;
+		}
+
+		const uint64 BufferSize = TargetDataCollectionDesc.ComputePackedSizeBytes();
+		if (PCGComputeHelpers::IsBufferSizeTooLarge(BufferSize))
+		{
+			return true;
+		}
+
+		bool bAllDataReady = true;
+
+		for (FPCGTaggedData& Data : Binding->InputDataCollection.TaggedData)
+		{
+			if (const UPCGProxyForGPUData* Proxy = Cast<UPCGProxyForGPUData>(Data.Data))
+			{
+				UPCGProxyForGPUData::FReadbackResult Result = Proxy->GetCPUData(nullptr);
+				if (!Result.bComplete)
+				{
+					bAllDataReady = false;
+				}
+				else
+				{
+					Data.Data = Result.Data;
+				}
+			}
+		}
+
+		if (!bAllDataReady)
+		{
+			return false;
+		}
+
+		TargetDataCollectionDesc.PackDataCollection(Binding->InputDataCollection, InputPinLabel, Binding->GetStringTable(), PackedDataCollection);
+	}
+
+	return true;
+}
+
+FRDGBufferSRVRef FPCGDataUploadAdaptor::GetAttributeRemapBufferSRV(FRDGBuilder& InGraphBuilder)
+{
+	const uint32 SizeBytes = FMath::Max(SourceToTargetAttributeId.Num(), 1) * SourceToTargetAttributeId.GetTypeSize();
+	const FRDGBufferDesc Desc = FRDGBufferDesc::CreateByteAddressDesc(SizeBytes);
+
+	if (SourceToTargetAttributeId.IsEmpty())
+	{
+		FRDGBufferRef NullBuffer = InGraphBuilder.CreateBuffer(Desc, TEXT("EmptyRemap"));
+		InGraphBuilder.QueueBufferUpload(NullBuffer, &FUintVector2::ZeroValue, sizeof(FUintVector2::ZeroValue));
+
+		return InGraphBuilder.CreateSRV(NullBuffer);
+	}
+	else
+	{
+		FRDGBufferRef Buffer = InGraphBuilder.CreateBuffer(Desc, TEXT("PCGDataCollectionUpload_AttributeRemap"));
+
+		InGraphBuilder.QueueBufferUpload(Buffer, SourceToTargetAttributeId.GetData(), SizeBytes);
+
+		return InGraphBuilder.CreateSRV(Buffer);
+	}
+}
+
+/** Create buffer with the element counts of each data. Stored in a buffer because we do not constrain the max data count. */
+FRDGBufferSRVRef FPCGDataUploadAdaptor::GetDataElementCountsBufferSRV(FRDGBuilder& InGraphBuilder)
+{
+	TArray<uint32> DataElementCounts;
+	DataElementCounts.Reserve(TargetDataCollectionDesc.DataDescs.Num());
+
+	for (const FPCGDataDesc& DataDesc : TargetDataCollectionDesc.DataDescs)
+	{
+		DataElementCounts.Add(DataDesc.ElementCount);
+	}
+
+	if (DataElementCounts.IsEmpty())
+	{
+		const FRDGBufferDesc Desc = FRDGBufferDesc::CreateByteAddressDesc(DataElementCounts.GetTypeSize());
+		FRDGBufferRef NullBuffer = InGraphBuilder.CreateBuffer(Desc, TEXT("EmptyDataElementCounts"));
+
+		uint32 Zero = 0;
+		InGraphBuilder.QueueBufferUpload(NullBuffer, &Zero, sizeof(Zero));
+
+		return InGraphBuilder.CreateSRV(NullBuffer);
+	}
+	else
+	{
+		const uint32 SizeBytes = DataElementCounts.Num() * DataElementCounts.GetTypeSize();
+		const FRDGBufferDesc Desc = FRDGBufferDesc::CreateByteAddressDesc(SizeBytes);
+
+		FRDGBufferRef Buffer = InGraphBuilder.CreateBuffer(Desc, TEXT("PCGDataCollectionUpload_DataElementCounts"));
+
+		InGraphBuilder.QueueBufferUpload(Buffer, DataElementCounts.GetData(), SizeBytes);
+
+		return InGraphBuilder.CreateSRV(Buffer);
+	}
+}
+
+/** Gets the buffer that can then be used to read the data in kernels. */
+FRDGBufferRef FPCGDataUploadAdaptor::GetBuffer_RenderThread(FRDGBuilder& InGraphBuilder, EPCGExportMode InExportMode, const TCHAR* InBufferDebugName)
+{
+	if (!ExternalBufferForReuse)
+	{
+		// Not reusing an existing buffer, create a new buffer and pack the input data collection.
+		FRDGBufferDesc Desc = FRDGBufferDesc::CreateByteAddressDesc(sizeof(uint32) * PackedDataCollection.Num());
+		if (InExportMode != EPCGExportMode::NoExport)
+		{
+			// We don't know for sure whether buffer will be read back or not, so need to flag the possibility if the buffer will be passed downstream.
+			Desc.Usage |= BUF_SourceCopy;
+		}
+
+		FRDGBufferRef Buffer = InGraphBuilder.CreateBuffer(Desc, InBufferDebugName);
+
+		InGraphBuilder.QueueBufferUpload(Buffer, PackedDataCollection.GetData(), PackedDataCollection.Num() * PackedDataCollection.GetTypeSize());
+
+		return Buffer;
+	}
+
+	ensure(PackedDataCollection.IsEmpty());
+
+	// Copy the external input buffer to a transient buffer, then run a CS to adapt it:
+	// - Update data counts
+	// - Remap attribute IDs
+
+#if !UE_BUILD_SHIPPING
+	RenderCaptureInterface::FScopedCapture RenderCapture(PCGDataUploadAdaptor::TriggerGPUCaptureDispatches > 0, InGraphBuilder, TEXT("FPCGDataUploadAdaptor"));
+	PCGDataUploadAdaptor::TriggerGPUCaptureDispatches = FMath::Max(PCGDataUploadAdaptor::TriggerGPUCaptureDispatches - 1, 0);
+#endif // !UE_BUILD_SHIPPING
+
+	FRDGBufferRef InputBuffer = InGraphBuilder.RegisterExternalBuffer(ExternalBufferForReuse);
+	const uint32 InputBufferSizeBytes = ExternalBufferSizeBytes;
+	const uint32 InputBufferSizeUints = InputBufferSizeBytes / sizeof(uint32);
+
+	const uint32 OutputSizeBytes = InputBufferSizeBytes;
+	FRDGBufferDesc OutputDesc = FRDGBufferDesc::CreateByteAddressDesc(OutputSizeBytes);
+	if (InExportMode != EPCGExportMode::NoExport)
+	{
+		// We don't know for sure whether buffer will be read back or not, so need to flag the possibility if the buffer will be passed downstream.
+		OutputDesc.Usage |= BUF_SourceCopy;
+	}
+
+	FRDGBufferRef OutputBuffer = InGraphBuilder.CreateBuffer(OutputDesc, InBufferDebugName);
+
+	// Populate using copy. We can't do everything in one pass trivially due to race conditions between groups.
+	AddCopyBufferPass(InGraphBuilder, OutputBuffer, InputBuffer);
+
+	FPCGDataCollectionAdaptorCS::FParameters* Parameters = InGraphBuilder.AllocParameters<FPCGDataCollectionAdaptorCS::FParameters>();
+	Parameters->InSourceToTargetAttributeId = GetAttributeRemapBufferSRV(InGraphBuilder);
+	Parameters->InDataElementCounts = GetDataElementCountsBufferSRV(InGraphBuilder);
+	Parameters->InDataCollection = InGraphBuilder.CreateSRV(InputBuffer);
+	Parameters->OutDataCollection = InGraphBuilder.CreateUAV(OutputBuffer);
+	Parameters->InNumData = TargetDataCollectionDesc.DataDescs.Num();
+	Parameters->InNumRemappedAttributes = SourceToTargetAttributeId.Num();
+
+	const uint32 NumUpdateTasks = Parameters->InNumData + Parameters->InNumRemappedAttributes * Parameters->InNumData;
+	const int GroupCount = FMath::DivideAndRoundUp(NumUpdateTasks, FPCGDataCollectionAdaptorCS::NumThreadsPerGroup);
+
+	TShaderMapRef<FPCGDataCollectionAdaptorCS> Shader(GetGlobalShaderMap(GMaxRHIFeatureLevel));
+	FComputeShaderUtils::AddPass(InGraphBuilder, RDG_EVENT_NAME("PCGDataCollectionAdaptor"), ERDGPassFlags::Compute | ERDGPassFlags::NeverCull /*TODO*/, Shader, Parameters, FIntVector(GroupCount, 1, 1));
+
+	return OutputBuffer;
+}
diff --git a/Engine/Plugins/PCG/Source/PCG/Private/Compute/PCGDataUploadAdaptor.h b/Engine/Plugins/PCG/Source/PCG/Private/Compute/PCGDataUploadAdaptor.h
new file mode 100644
--- /dev/null
+++ b/Engine/Plugins/PCG/Source/PCG/Private/Compute/PCGDataUploadAdaptor.h
@@ -0,0 +1,46 @@
+// Copyright Epic Games, Inc. All Rights Reserved.
+
+#pragma once
+
+#include "Compute/PCGDataForGPU.h"
+
+#include "CoreMinimal.h"
+#include "RendererInterface.h"
+
+class FRDGBuilder;
+class UPCGDataBinding;
+enum class EPCGExportMode : uint8;
+
+/** A helper to ingest incoming data objects and pass the required data to the data interface. */
+class FPCGDataUploadAdaptor
+{
+public:
+	FPCGDataUploadAdaptor(UPCGDataBinding* InDataBinding, const FPCGDataCollectionDesc& InTargetDataCollectionDesc, FName InInputPinLabel);
+
+	/** Do any preparation work such as data readbacks. Returns true when preparation is complete. */
+	bool PrepareData_GameThread(const UPCGDataBinding* InBinding);
+
+	FRDGBufferSRVRef GetAttributeRemapBufferSRV(FRDGBuilder& InGraphBuilder);
+
+	/** Create buffer with the element counts of each data. Stored in a buffer because we do not constrain the max data count. */
+	FRDGBufferSRVRef GetDataElementCountsBufferSRV(FRDGBuilder& InGraphBuilder);
+
+	/** Gets the buffer that can then be used to read the data in kernels. */
+	FRDGBufferRef GetBuffer_RenderThread(FRDGBuilder& InGraphBuilder, EPCGExportMode InExportMode, const TCHAR* InBufferDebugName);
+
+private:
+	TRefCountPtr<FRDGPooledBuffer> ExternalBufferForReuse;
+
+	TWeakObjectPtr<UPCGDataBinding> DataBinding;
+
+	FPCGDataCollectionDesc TargetDataCollectionDesc;
+
+	/** If data cannot be reused from GPU, will be packed into this buffer prior to upload. */
+	TArray<uint32> PackedDataCollection;
+
+	TArray<FIntVector2> SourceToTargetAttributeId;
+
+	FName InputPinLabel;
+
+	uint32 ExternalBufferSizeBytes = 0;
+};
diff --git a/Engine/Plugins/PCG/Source/PCG/Private/Elements/Metadata/PCGMetadataPartition.cpp b/Engine/Plugins/PCG/Source/PCG/Private/Elements/Metadata/PCGMetadataPartition.cpp
--- a/Engine/Plugins/PCG/Source/PCG/Private/Elements/Metadata/PCGMetadataPartition.cpp
+++ b/Engine/Plugins/PCG/Source/PCG/Private/Elements/Metadata/PCGMetadataPartition.cpp
@@ -4,6 +4,11 @@
 
 #include "PCGContext.h"
 #include "PCGCustomVersion.h"
+#include "Compute/PCGComputeCommon.h"
+#include "Compute/PCGDataBinding.h"
+#include "Compute/Elements/PCGComputeGraphElement.h"
+#include "Compute/DataInterfaces/Elements/PCGMetadataPartitionDataInterface.h"
+#include "Graph/PCGGPUGraphCompilationContext.h"
 #include "Helpers/PCGHelpers.h"
 #include "Metadata/PCGMetadata.h"
 #include "Metadata/PCGMetadataAttribute.h"
@@ -13,10 +18,17 @@
 #include "Metadata/Accessors/PCGAttributeAccessorKeys.h"
 #include "Utils/PCGLogErrors.h"
 
+#include "ShaderCompilerCore.h"
+
 #include UE_INLINE_GENERATED_CPP_BY_NAME(PCGMetadataPartition)
 
 #define LOCTEXT_NAMESPACE "PCGMetadataPartitionElement"
 
+namespace PCGMetadataPartitionSettings
+{
+	static TCHAR const* TemplateFilePath = TEXT("/Plugin/PCG/Private/Elements/PCGMetadataPartition.usf");
+}
+
 TArray<FPCGPinProperties> UPCGMetadataPartitionSettings::OutputPinProperties() const
 {
 	TArray<FPCGPinProperties> PinProperties;
@@ -53,7 +65,176 @@ void UPCGMetadataPartitionSettings::PostLoad()
 #endif // WITH_EDITOR
 }
 
+bool UPCGMetadataPartitionSettings::IsKernelValid(FPCGContext* InContext, bool bQuiet) const
+{
+	TRACE_CPUPROFILER_EVENT_SCOPE(UPCGMetadataPartitionSettings::IsKernelValid);
+	check(InContext);
+
+	if (!Super::IsKernelValid(InContext, bQuiet))
+	{
+		return false;
+	}
+
+	if (PartitionAttributeSelectors.Num() != 1)
+	{
+		PCG_KERNEL_VALIDATION_ERR(InContext, this, bQuiet, LOCTEXT("MustProvideOnePartitionAttribute", "GPU implementation currently only supports a single partition attribute."));
+		return false;
+	}
+
+	if (!PartitionAttributeSelectors[0].IsBasicAttribute())
+	{
+		PCG_KERNEL_VALIDATION_ERR(InContext, this, bQuiet, LOCTEXT("OnlyBasicAttributesSupported", "GPU implementation currently only supports basic attributes."));
+		return false;
+	}
+
+	FPCGComputeGraphContext* Context = InContext->IsComputeContext() ? static_cast<FPCGComputeGraphContext*>(InContext) : nullptr;
+
+	if (Context && Context->DataBinding)
+	{
+		const UPCGNode* Node = CastChecked<UPCGNode>(GetOuter());
+		const UPCGPin* InputPin = Node->GetInputPin(PCGPinConstants::DefaultInputLabel);
+
+		FName AttributeName = PartitionAttributeSelectors[0].GetAttributeName();
+
+		if (ensure(InputPin))
+		{
+			FPCGDataCollectionDesc CollectionDesc = PCGDataForGPUHelpers::ComputeInputPinDataDesc(InputPin, Context->DataBinding.Get());
+
+			FPCGKernelAttributeDesc AttributeDesc;
+			bool bConflictingTypesInData = false;
+			const bool bAttributePresent = CollectionDesc.GetAttributeDesc(AttributeName, AttributeDesc, bConflictingTypesInData);
+
+			if (!bAttributePresent)
+			{
+				PCG_KERNEL_VALIDATION_ERR(InContext, this, bQuiet, LOCTEXT("MustProvideOnePartitionAttribute", "GPU implementation currently only supports a single partition attribute."));
+				return false;
+			}
+
+			if (bConflictingTypesInData)
+			{
+				PCG_KERNEL_VALIDATION_ERR(InContext, this, bQuiet,
+					FText::Format(LOCTEXT("PartitionAttributeTypeConflict", "Attribute '{0}' encountered with multiple different types in input data."), FText::FromName(AttributeName)));
+				return false;
+			}
+
+			if (AttributeDesc.Type != EPCGKernelAttributeType::StringKey)
+			{
+				PCG_KERNEL_VALIDATION_ERR(Context, this, bQuiet, FText::Format(
+					LOCTEXT("PartitionAttributeTypeInvalid", "Attribute '{0}' not usable for partitioning, only attributes of type String Key are currently supported."),
+					FText::FromName(AttributeName)));
+				return false;
+			}
+		}
+	}
+
+	return true;
+}
+
+FString UPCGMetadataPartitionSettings::GetCookedKernelSource(const FPCGKernelAttributeTable& InAttributeStaticTable) const
+{
+	FString TemplateFile;
+	ensure(LoadShaderSourceFile(PCGMetadataPartitionSettings::TemplateFilePath, EShaderPlatform::SP_PCD3D_SM5, &TemplateFile, nullptr));
+
+	return TemplateFile;
+}
+
+void UPCGMetadataPartitionSettings::GetKernelAttributeKeys(TArray<FPCGKernelAttributeKey>& OutKeys) const
+{
+	// Register the attributes this node reads or writes for which we know the attribute type. Currently
+	// only StringKey attributes can be partitioned when executing on GPU, so declare the attribute here
+	// rather than doing work to resolve at runtime.
+	for (const FPCGAttributePropertyInputSelector& Selector : PartitionAttributeSelectors)
+	{
+		if (Selector.IsBasicAttribute())
+		{
+			OutKeys.Add(FPCGKernelAttributeKey(Selector.GetAttributeName(), EPCGKernelAttributeType::StringKey));
+		}
+	}
+
+	// Register output attribute created by this node.
+	if (bAssignIndexPartition)
+	{
+		OutKeys.AddUnique(FPCGKernelAttributeKey(PartitionIndexAttributeName, EPCGKernelAttributeType::Int));
+	}
+}
+
+int UPCGMetadataPartitionSettings::ComputeKernelThreadCount(const UPCGDataBinding* Binding) const
+{
+	const UPCGNode* Node = CastChecked<UPCGNode>(GetOuter());
+	const UPCGPin* Pin = Node->GetInputPin(PCGPinConstants::DefaultInputLabel);
+
+	return PCGDataForGPUHelpers::ComputeInputPinDataDesc(Pin, Binding).ComputeDataElementCount(EPCGDataType::Point);
+}
+
+bool UPCGMetadataPartitionSettings::ComputeOutputPinDataDesc(const UPCGPin* OutputPin, const UPCGDataBinding* Binding, FPCGDataCollectionDesc& OutDesc) const
+{
+	check(OutputPin);
+	check(Binding);
+
+	if (!ensure(OutputPin->Properties.Label == PCGPinConstants::DefaultOutputLabel))
+	{
+		return false;
+	}
+
+	// First inspect data coming from CPU. Will return description if this node is a CPU node or in a different compute graph.
+	if (Binding->ComputeCPUOutputPinDataDesc(OutputPin, OutDesc))
+	{
+		return true;
+	}
+
+	const FPCGDataCollectionDesc InputPinDesc = PCGDataForGPUHelpers::ComputeInputPinDataDesc(this, PCGPinConstants::DefaultInputLabel, Binding);
+
+	const int32 PartitionAttributeId = Binding->GetAttributeId(PartitionAttributeSelectors[0].GetAttributeName(), EPCGKernelAttributeType::StringKey);
+
+	FPCGDataCollectionDesc OutputPinDesc;
+
+	// Emit a point data for each potential partition.
+	if (PartitionAttributeId != INDEX_NONE)
+	{
+		TArray<int32> UniqueStringKeyValues;
+		InputPinDesc.GetUniqueStringKeyValues(PartitionAttributeId, UniqueStringKeyValues);
+
+		for (const FPCGDataDesc& InputData : InputPinDesc.DataDescs)
+		{
+			// Before execution we don't know what string key value each attribute has, so we have to assume worst case - could be up to N points in M partitions.
+			for (int32 StringKeyValue : UniqueStringKeyValues)
+			{
+				FPCGDataDesc& Desc = OutputPinDesc.DataDescs.Add_GetRef(InputData);
+
+				// Set the unique string keys to reflect the partition operation.
+				for (FPCGKernelAttributeDesc& AttrDesc : Desc.AttributeDescs)
+				{
+					if (AttrDesc.Index == PartitionAttributeId)
+					{
+						AttrDesc.SetStringKeys(MakeArrayView(&StringKeyValue, 1));
+						break;
+					}
+				}
+			}
+		}
+	}
+
+	// Add output attribute (partition index).
+	if (bAssignIndexPartition)
+	{
+		OutputPinDesc.AddAttributeToAllData(FPCGKernelAttributeKey(PartitionIndexAttributeName, EPCGKernelAttributeType::Int), Binding);
+	}
+
+	OutDesc = OutputPinDesc;
+
+	return true;
+}
+
 #if WITH_EDITOR
+void UPCGMetadataPartitionSettings::CreateAdditionalInputDataInterfaces(FPCGGPUCompilationContext& InOutContext, UObject* InObjectOuter, TArray<TObjectPtr<UComputeDataInterface>>& OutDataInterfaces) const
+{
+	Super::CreateAdditionalInputDataInterfaces(InOutContext, InObjectOuter, OutDataInterfaces);
+
+	TObjectPtr<UPCGMetadataPartitionDataInterface> NodeDI = InOutContext.NewObject_AnyThread<UPCGMetadataPartitionDataInterface>(InObjectOuter);
+	NodeDI->Settings = this;
+	OutDataInterfaces.Add(NodeDI);
+}
+
 void UPCGMetadataPartitionSettings::ApplyDeprecation(UPCGNode* InOutNode)
 {
 	PRAGMA_DISABLE_DEPRECATION_WARNINGS
diff --git a/Engine/Plugins/PCG/Source/PCG/Private/Elements/PCGStaticMeshSpawner.cpp b/Engine/Plugins/PCG/Source/PCG/Private/Elements/PCGStaticMeshSpawner.cpp
--- a/Engine/Plugins/PCG/Source/PCG/Private/Elements/PCGStaticMeshSpawner.cpp
+++ b/Engine/Plugins/PCG/Source/PCG/Private/Elements/PCGStaticMeshSpawner.cpp
@@ -126,10 +126,10 @@ bool UPCGStaticMeshSpawnerSettings::ComputeOutputPinDataDesc(const UPCGPin* Outp
 	// Add output attribute (selected mesh).
 	{
 		// Create unique value keys for the output string.
-		TArray<int32> UniqueStringKeys;
-
 		if (UPCGMeshSelectorWeighted* SelectorWeighted = Cast<UPCGMeshSelectorWeighted>(MeshSelectorParameters))
 		{
+			TArray<int32> UniqueStringKeys;
+
 			// Weighted selection - add explicit strings from settings.
 			for (const FPCGMeshSelectorWeightedEntry& Entry : SelectorWeighted->MeshEntries)
 			{
@@ -140,24 +140,28 @@ bool UPCGStaticMeshSpawnerSettings::ComputeOutputPinDataDesc(const UPCGPin* Outp
 					UniqueStringKeys.AddUnique(StringIndex);
 				}
 			}
+
+			OutputPinDesc.AddAttributeToAllData(FPCGKernelAttributeKey(OutAttributeName, EPCGKernelAttributeType::StringKey), Binding, &UniqueStringKeys);
 		}
 		else if (UPCGMeshSelectorByAttribute* SelectorByAttribute = Cast<UPCGMeshSelectorByAttribute>(MeshSelectorParameters))
 		{
 			// By-attribute selection - pass on strings from input attribute.
 			if (SelectorByAttribute->AttributeName != NAME_None)
 			{
-				for (const FPCGDataDesc& InputDataDesc : InputPinDesc.DataDescs)
+				for (FPCGDataDesc& DataDesc : OutputPinDesc.DataDescs)
 				{
-					for (const FPCGKernelAttributeDesc& InputAttrDesc : InputDataDesc.AttributeDescs)
+					const TArray<int32>* StringKeys = nullptr;
+
+					for (const FPCGKernelAttributeDesc& AttrDesc : DataDesc.AttributeDescs)
 					{
-						if (InputAttrDesc.Name == SelectorByAttribute->AttributeName)
+						if (AttrDesc.Name == SelectorByAttribute->AttributeName)
 						{
-							for (int32 InputStringKey : InputAttrDesc.GetUniqueStringKeys())
-							{
-								UniqueStringKeys.AddUnique(InputStringKey);
-							}
+							StringKeys = &AttrDesc.GetUniqueStringKeys();
+							break;
 						}
 					}
+
+					DataDesc.AddAttribute(FPCGKernelAttributeKey(OutAttributeName, EPCGKernelAttributeType::StringKey), Binding, StringKeys);
 				}
 			}
 		}
@@ -165,8 +169,6 @@ bool UPCGStaticMeshSpawnerSettings::ComputeOutputPinDataDesc(const UPCGPin* Outp
 		{
 			UE_LOG(LogPCG, Warning, TEXT("Mesh selector not supported by GPU Static Mesh Spawner: %s"), *MeshSelectorParameters->GetName());
 		}
-
-		OutputPinDesc.AddAttributeToAllData(FPCGKernelAttributeKey(OutAttributeName, EPCGKernelAttributeType::StringKey), Binding, &UniqueStringKeys);
 	}
 
 	OutDesc = OutputPinDesc;
diff --git a/Engine/Plugins/PCG/Source/PCG/Private/Graph/PCGGraphCompilerGPU.cpp b/Engine/Plugins/PCG/Source/PCG/Private/Graph/PCGGraphCompilerGPU.cpp
--- a/Engine/Plugins/PCG/Source/PCG/Private/Graph/PCGGraphCompilerGPU.cpp
+++ b/Engine/Plugins/PCG/Source/PCG/Private/Graph/PCGGraphCompilerGPU.cpp
@@ -62,18 +62,22 @@ void FPCGGraphCompilerGPU::LabelConnectedGPUNodeIslands(
 			}
 		}
 
-		if (const TArray<FPCGTaskId>* Successors = InTaskSuccessors.Find(InTaskId))
+		const UPCGSettings* Settings = InCompiledTasks[InTaskId].Node ? InCompiledTasks[InTaskId].Node->GetSettings() : nullptr;
+		if (!Settings || !Settings->SplitGraphAtOutput())
 		{
-			for (FPCGTaskId Successor : *Successors)
+			if (const TArray<FPCGTaskId>* Successors = InTaskSuccessors.Find(InTaskId))
 			{
-				if (Successor == InTraversedFromTaskId)
+				for (FPCGTaskId Successor : *Successors)
 				{
-					continue;
-				}
+					if (Successor == InTraversedFromTaskId)
+					{
+						continue;
+					}
 
-				if (OutIslandIDs[Successor] == 0 && InGPUCompatibleTaskIds.Contains(Successor))
-				{
-					RecursiveCall(Successor, InIslandID, InTaskId, RecursiveCall);
+					if (OutIslandIDs[Successor] == 0 && InGPUCompatibleTaskIds.Contains(Successor))
+					{
+						RecursiveCall(Successor, InIslandID, InTaskId, RecursiveCall);
+					}
 				}
 			}
 		}
@@ -523,7 +527,7 @@ void FPCGGraphCompilerGPU::BuildComputeGraphTask(
 	uint32 InGridSize,
 	uint32 InComputeGraphIndex,
 	const TSet<FPCGTaskId>& InCollapsedTasks,
-	const TSet<FPCGTaskId>& InAllGPUCompatibleTasks,
+	const TArray<FPCGTaskId>& InComputeElementTasks,
 	const FTaskToSuccessors& InTaskSuccessors,
 	TArray<FPCGGraphTask>& InOutCompiledTasks,
 	const FOriginalToVirtualPin& InOriginalToVirtualPin,
@@ -556,9 +560,9 @@ void FPCGGraphCompilerGPU::BuildComputeGraphTask(
 
 		FPinToDataInterface OutputPinDataInterfaces;
 
-		CreateDataInterfaces(InOutContext, InOutCompiledTasks, InCollapsedTasks, InAllGPUCompatibleTasks, InTaskSuccessors, InOriginalToVirtualPin, ComputeGraph, OutputPinDataInterfaces);
+		CreateDataInterfaces(InOutContext, InOutCompiledTasks, InCollapsedTasks, InComputeElementTasks, InTaskSuccessors, InOriginalToVirtualPin, ComputeGraph, OutputPinDataInterfaces);
 
-		CompileComputeGraph(InOutContext, InOutCompiledTasks, InCollapsedTasks, InAllGPUCompatibleTasks, InTaskSuccessors, InOriginalToVirtualPin, OutputPinDataInterfaces, ComputeGraph);
+		CompileComputeGraph(InOutContext, InOutCompiledTasks, InCollapsedTasks, InComputeElementTasks, InTaskSuccessors, InOriginalToVirtualPin, OutputPinDataInterfaces, ComputeGraph);
 
 		// Remove empty strings (at execution time empty string is placed in table slot 0).
 		// NOTE: This can scramble order but order is not important at this stage (key values created at execution time in data binding).
@@ -593,7 +597,7 @@ void FPCGGraphCompilerGPU::CreateDataInterfaces(
 	FPCGGPUCompilationContext& InOutContext,
 	const TArray<FPCGGraphTask>& InCompiledTasks,
 	const TSet<FPCGTaskId>& InCollapsedTasks,
-	const TSet<FPCGTaskId>& InAllGPUCompatibleTasks,
+	const TArray<FPCGTaskId>& InComputeElementTasks,
 	const FTaskToSuccessors& InTaskSuccessors,
 	const FOriginalToVirtualPin& InOriginalToVirtualPin,
 	UPCGComputeGraph* InOutComputeGraph,
@@ -603,20 +607,20 @@ void FPCGGraphCompilerGPU::CreateDataInterfaces(
 
 	// Create data interfaces which allow kernels to read or write data. Each data interface is associated with a node output pin.
 	// For CPU->GPU edges, an upload data interface is created. For GPU->CPU edges, a readback data interface is created.
-	auto CreateDataInterface = [&InOutContext, &InCollapsedTasks, &InCompiledTasks, InOutComputeGraph](FPCGTaskId InTaskId, bool bRequiresExport, const FPCGPinProperties& InOutputPinProperties) -> UPCGComputeDataInterface*
+	auto CreateDataInterface = [&InOutContext, &InCollapsedTasks, &InCompiledTasks, InOutComputeGraph](FPCGTaskId InProducerTaskId, bool bRequiresExport, const FPCGPinProperties& InOutputPinProperties) -> UPCGComputeDataInterface*
 	{
 		TRACE_CPUPROFILER_EVENT_SCOPE(CreateDataInterface);
 
-		const bool bUpstreamIsGPUTask = InCollapsedTasks.Contains(InTaskId);
+		const bool bUpstreamIsGPUTask = InCollapsedTasks.Contains(InProducerTaskId);
 
 		EPCGDataType PinType = InOutputPinProperties.AllowedTypes;
 
 		// Dynamically typed pins could have a different type
-		const UPCGSettings* Settings = InCompiledTasks[InTaskId].Node ? InCompiledTasks[InTaskId].Node->GetSettings() : nullptr;
-		const UPCGPin* Pin = InCompiledTasks[InTaskId].Node ? InCompiledTasks[InTaskId].Node->GetOutputPin(InOutputPinProperties.Label) : nullptr;
-		if (Settings && Pin)
+		const UPCGSettings* ProducerSettings = InCompiledTasks[InProducerTaskId].Node ? InCompiledTasks[InProducerTaskId].Node->GetSettings() : nullptr;
+		const UPCGPin* ProducerPin = InCompiledTasks[InProducerTaskId].Node ? InCompiledTasks[InProducerTaskId].Node->GetOutputPin(InOutputPinProperties.Label) : nullptr;
+		if (ProducerSettings && ProducerPin)
 		{
-			PinType = Settings->GetCurrentPinTypes(Pin);
+			PinType = ProducerSettings->GetCurrentPinTypes(ProducerPin);
 		}
 
 		UPCGComputeDataInterface* DataInterface = nullptr;
@@ -640,8 +644,9 @@ void FPCGGraphCompilerGPU::CreateDataInterfaces(
 			}
 
 			check(DataInterfacePCGData);
-			DataInterfacePCGData->SetProducerSettings(InCompiledTasks[InTaskId].Node ? InCompiledTasks[InTaskId].Node->GetSettings() : nullptr);
+			DataInterfacePCGData->SetProducerSettings(ProducerSettings);
 			DataInterfacePCGData->SetRequiresExport(bRequiresExport);
+			DataInterfacePCGData->SetExposeElementCounters(ProducerSettings && ProducerPin && ProducerSettings->DoesOutputPinRequireElementCounters(ProducerPin));
 
 			DataInterface = DataInterfacePCGData;
 			break;
@@ -662,7 +667,7 @@ void FPCGGraphCompilerGPU::CreateDataInterfaces(
 		{
 			UE_LOG(LogPCG, Error, TEXT("Unsupported connected upstream pin '%s' on node '%s' with type %s. Consider adding a conversion to a supported type such as Point."),
 				*InOutputPinProperties.Label.ToString(),
-				InCompiledTasks[InTaskId].Node ? *InCompiledTasks[InTaskId].Node->GetNodeTitle(EPCGNodeTitleType::ListView).ToString() : TEXT("MISSING"),
+				InCompiledTasks[InProducerTaskId].Node ? *InCompiledTasks[InProducerTaskId].Node->GetNodeTitle(EPCGNodeTitleType::ListView).ToString() : TEXT("MISSING"),
 				StaticEnum<EPCGDataType>() ? *StaticEnum<EPCGDataType>()->GetNameStringByValue(static_cast<int64>(PinType)) : TEXT("MISSING")
 			);
 			break;
@@ -760,7 +765,7 @@ void FPCGGraphCompilerGPU::CreateDataInterfaces(
 		for (const FPCGGraphTaskInput& Input : InCompiledTasks[TaskId].Inputs)
 		{
 			// Never wire directly to the GPU-compatible tasks, these are collapsed into compute graphs (and will be culled in a next step).
-			if (InAllGPUCompatibleTasks.Contains(Input.TaskId))
+			if (InCollapsedTasks.Contains(Input.TaskId) || InComputeElementTasks.Contains(Input.TaskId))
 			{
 				continue;
 			}
@@ -816,7 +821,7 @@ void FPCGGraphCompilerGPU::CompileComputeGraph(
 	FPCGGPUCompilationContext& InOutContext,
 	const TArray<FPCGGraphTask>& InCompiledTasks,
 	const TSet<FPCGTaskId>& InCollapsedTasks,
-	const TSet<FPCGTaskId>& InAllGPUCompatibleTasks,
+	const TArray<FPCGTaskId>& InComputeElementTasks,
 	const FTaskToSuccessors& InTaskSuccessors,
 	const FOriginalToVirtualPin& InOriginalToVirtualPin,
 	const FPinToDataInterface& InOutputPinToDataInterface,
@@ -907,7 +912,7 @@ void FPCGGraphCompilerGPU::CompileComputeGraph(
 			// Currently the new compute graph tasks are wired into the graph in parallel to each GPU node task. The GPU node tasks
 			// will be culled at the end, leaving only the compute graphs. Only create DIs for tasks within this compute graph (in InCollapsedTasks),
 			// or for tasks that will not be culled (not in InAllGPUCompatibleTasks).
-			const bool bValidInput = InCollapsedTasks.Contains(Input.TaskId) || !InAllGPUCompatibleTasks.Contains(Input.TaskId);
+			const bool bValidInput = InCollapsedTasks.Contains(Input.TaskId) || !InComputeElementTasks.Contains(Input.TaskId);
 			if (!bValidInput)
 			{
 				continue;
@@ -1292,6 +1297,8 @@ void FPCGGraphCompilerGPU::CreateGPUNodes(FPCGGraphCompiler& InOutCompiler, UPCG
 			OutputCPUPinToVirtualPin[ComputeGraphIndex]);
 	}
 
+	TArray<FPCGTaskId> ComputeElementTasks;
+
 	// Build each compute graph.
 	for (uint32 ComputeGraphIndex = 0; ComputeGraphIndex < NumComputeGraphs; ++ComputeGraphIndex)
 	{
@@ -1305,6 +1312,7 @@ void FPCGGraphCompilerGPU::CreateGPUNodes(FPCGGraphCompiler& InOutCompiler, UPCG
 
 		// Add a new compute graph task. The original GPU tasks will then be culled later below.
 		const FPCGTaskId ComputeGraphTaskId = InOutCompiledTasks.Num();
+		ComputeElementTasks.Add(ComputeGraphTaskId);
 		FPCGGraphTask& ComputeGraphTask = InOutCompiledTasks.Emplace_GetRef();
 		ComputeGraphTask.NodeId = ComputeGraphTaskId;
 
@@ -1332,7 +1340,7 @@ void FPCGGraphCompilerGPU::CreateGPUNodes(FPCGGraphCompiler& InOutCompiler, UPCG
 			InGridSize,
 			ComputeGraphIndex,
 			NodeSubsetToConvertToCFGraph,
-			GPUCompatibleTaskIds,
+			ComputeElementTasks,
 			TaskSuccessors,
 			InOutCompiledTasks,
 			OriginalToVirtualPin[ComputeGraphIndex],
diff --git a/Engine/Plugins/PCG/Source/PCG/Private/Graph/PCGGraphCompilerGPU.h b/Engine/Plugins/PCG/Source/PCG/Private/Graph/PCGGraphCompilerGPU.h
--- a/Engine/Plugins/PCG/Source/PCG/Private/Graph/PCGGraphCompilerGPU.h
+++ b/Engine/Plugins/PCG/Source/PCG/Private/Graph/PCGGraphCompilerGPU.h
@@ -70,7 +70,7 @@ private:
 		uint32 InGridSize,
 		uint32 InComputeGraphIndex,
 		const TSet<FPCGTaskId>& InCollapsedTasks,
-		const TSet<FPCGTaskId>& InAllGPUCompatibleTasks,
+		const TArray<FPCGTaskId>& InComputeElementTasks,
 		const FTaskToSuccessors& InTaskSuccessors,
 		TArray<FPCGGraphTask>& InOutCompiledTasks,
 		const FOriginalToVirtualPin& InOriginalToVirtualPin,
@@ -81,7 +81,7 @@ private:
 		FPCGGPUCompilationContext& InOutContext,
 		const TArray<FPCGGraphTask>& InCompiledTasks,
 		const TSet<FPCGTaskId>& InCollapsedTasks,
-		const TSet<FPCGTaskId>& InAllGPUCompatibleTasks,
+		const TArray<FPCGTaskId>& InComputeElementTasks,
 		const FTaskToSuccessors& InTaskSuccessors,
 		const FOriginalToVirtualPin& InOriginalToVirtualPin,
 		UPCGComputeGraph* InOutComputeGraph, // In
@@ -92,7 +92,7 @@ private:
 		FPCGGPUCompilationContext& InOutContext,
 		const TArray<FPCGGraphTask>& InCompiledTasks,
 		const TSet<FPCGTaskId>& InCollapsedTasks,
-		const TSet<FPCGTaskId>& InAllGPUCompatibleTasks,
+		const TArray<FPCGTaskId>& InComputeElementTasks,
 		const FTaskToSuccessors& InTaskSuccessors,
 		const FOriginalToVirtualPin& InOriginalToVirtualPin,
 		const FPinToDataInterface& InOutputPinToDataInterface,
diff --git a/Engine/Plugins/PCG/Source/PCG/Public/Elements/Metadata/PCGMetadataPartition.h b/Engine/Plugins/PCG/Source/PCG/Public/Elements/Metadata/PCGMetadataPartition.h
--- a/Engine/Plugins/PCG/Source/PCG/Public/Elements/Metadata/PCGMetadataPartition.h
+++ b/Engine/Plugins/PCG/Source/PCG/Public/Elements/Metadata/PCGMetadataPartition.h
@@ -18,9 +18,17 @@ public:
 	//~End UObject interface
 
 	//~Begin UPCGSettings interface
+	virtual bool IsKernelValid(FPCGContext* InContext = nullptr, bool bQuiet = true) const;
+	virtual FString GetCookedKernelSource(const FPCGKernelAttributeTable& InAttributeStaticTable) const override;
+	virtual void GetKernelAttributeKeys(TArray<FPCGKernelAttributeKey>& OutKeys) const override;
+	virtual int ComputeKernelThreadCount(const UPCGDataBinding* Binding) const override;
+	virtual bool ComputeOutputPinDataDesc(const UPCGPin* OutputPin, const UPCGDataBinding* Binding, FPCGDataCollectionDesc& OutDesc) const override;
+	virtual bool DoesOutputPinRequireElementCounters(const UPCGPin* OutputPin) const override { return true; }
 #if WITH_EDITOR
+	virtual bool SplitGraphAtOutput() const { return true; }
+	virtual void CreateAdditionalInputDataInterfaces(FPCGGPUCompilationContext& InOutContext, UObject* InObjectOuter, TArray<TObjectPtr<UComputeDataInterface>>& OutDataInterfaces) const override;
+	virtual bool DisplayExecuteOnGPUSetting() const override { return true; }
 	virtual void ApplyDeprecation(UPCGNode* InOutNode) override;
-
 	virtual FName GetDefaultNodeName() const override { return FName(TEXT("AttributePartition")); }
 	virtual FText GetDefaultNodeTitle() const override { return NSLOCTEXT("PCGMetadataPartitionSettings", "NodeTitle", "Attribute Partition"); }
 	virtual EPCGSettingsType GetType() const override { return EPCGSettingsType::Metadata; }
@@ -49,7 +57,7 @@ public:
 	bool bTokenizeOnWhiteSpace = false;
 
 	/** Assign an index partition as an extra attribute. */
-	UPROPERTY(EditAnywhere, Category = Settings, meta = (PCG_Overridable))
+	UPROPERTY(EditAnywhere, DisplayName = "Output Partition Index", Category = Settings, meta = (PCG_Overridable))
 	bool bAssignIndexPartition = false;
 
 	/** If we assign an index, we can also not partition (and only assign the partition index to the original data). */
diff --git a/Engine/Plugins/PCG/Source/PCG/Public/PCGSettings.h b/Engine/Plugins/PCG/Source/PCG/Public/PCGSettings.h
--- a/Engine/Plugins/PCG/Source/PCG/Public/PCGSettings.h
+++ b/Engine/Plugins/PCG/Source/PCG/Public/PCGSettings.h
@@ -543,6 +543,11 @@ public:
 	*/
 	virtual bool ComputeOutputPinDataDesc(const UPCGPin* OutputPin, const UPCGDataBinding* InBinding, FPCGDataCollectionDesc& OutDesc) const;
 
+	/** [EXPERIMENTAL] Whether the output pin should have a buffer of counters attached, useful for compaction and other counting requirements.
+	* Note: This API function is subject to change without deprecation.
+	*/
+	virtual bool DoesOutputPinRequireElementCounters(const UPCGPin* OutputPin) const { return false; }
+
 #if WITH_EDITOR
 	/** [EXPERIMENTAL] Create additional input data interfaces to marshal any required input data.
 	* Note: This API function is subject to change without deprecation.
@@ -555,6 +560,11 @@ public:
 	virtual void CreateAdditionalOutputDataInterfaces(FPCGGPUCompilationContext& InOutContext, UObject* InObjectOuter, TArray<TObjectPtr<UComputeDataInterface>>& OutDataInterfaces) const {}
 	virtual void CreateAdditionalOutputDataInterfaces(TArray<TObjectPtr<UComputeDataInterface>>& OutDataInterfaces) const {}
 
+	/** [EXPERIMENTAL] Whether compute graphs should be split at the output of this node. Useful for reading back runtime statistics or diagnostic info before continuing.
+	* Note: This API function is subject to change without deprecation.
+	*/
+	virtual bool SplitGraphAtOutput() const { return false; }
+
 	/** [EXPERIMENTAL] Whether to display GPU execution option in node settings UI. 
 	* Note: This API function is subject to change without deprecation.
 	*/
diff --git a/Engine/Plugins/PCG/Source/PCGCompute/Internal/PCGDataCollectionAdaptorCS.h b/Engine/Plugins/PCG/Source/PCGCompute/Internal/PCGDataCollectionAdaptorCS.h
new file mode 100644
--- /dev/null
+++ b/Engine/Plugins/PCG/Source/PCGCompute/Internal/PCGDataCollectionAdaptorCS.h
@@ -0,0 +1,31 @@
+// Copyright Epic Games, Inc. All Rights Reserved.
+
+#pragma once
+
+#include "GlobalShader.h"
+#include "GPUSceneWriter.h"
+#include "ShaderParameterStruct.h"
+
+/** [EXPERIMENTAL] Compute shader for updating a data collection buffer with attribute ID remaps and/or element counts.
+ * Note: This class is subject to change without deprecation.
+ */
+class FPCGDataCollectionAdaptorCS : public FGlobalShader
+{
+public:
+	DECLARE_EXPORTED_GLOBAL_SHADER(FPCGDataCollectionAdaptorCS, PCGCOMPUTE_API);
+	SHADER_USE_PARAMETER_STRUCT(FPCGDataCollectionAdaptorCS, FGlobalShader);
+
+	BEGIN_SHADER_PARAMETER_STRUCT(FParameters, PCGCOMPUTE_API)
+		SHADER_PARAMETER_RDG_BUFFER_SRV(StructuredBuffer<FInt32Vector2>, InSourceToTargetAttributeId)
+		SHADER_PARAMETER_RDG_BUFFER_SRV(StructuredBuffer<uint32>, InDataElementCounts)
+		SHADER_PARAMETER_RDG_BUFFER_SRV(ByteAddressBuffer, InDataCollection)
+		SHADER_PARAMETER_RDG_BUFFER_UAV(RWByteAddressBuffer, OutDataCollection)
+		SHADER_PARAMETER(uint32, InNumData)
+		SHADER_PARAMETER(uint32, InNumRemappedAttributes)
+	END_SHADER_PARAMETER_STRUCT()
+
+public:
+	static constexpr uint32 NumThreadsPerGroup = 64;
+
+	static void ModifyCompilationEnvironment(const FGlobalShaderPermutationParameters& Parameters, FShaderCompilerEnvironment& OutEnvironment);
+};
diff --git a/Engine/Plugins/PCG/Source/PCGCompute/Private/PCGDataCollectionAdaptorCS.cpp b/Engine/Plugins/PCG/Source/PCGCompute/Private/PCGDataCollectionAdaptorCS.cpp
new file mode 100644
--- /dev/null
+++ b/Engine/Plugins/PCG/Source/PCGCompute/Private/PCGDataCollectionAdaptorCS.cpp
@@ -0,0 +1,12 @@
+// Copyright Epic Games, Inc. All Rights Reserved.
+
+#include "PCGDataCollectionAdaptorCS.h"
+
+void FPCGDataCollectionAdaptorCS::ModifyCompilationEnvironment(const FGlobalShaderPermutationParameters& Parameters, FShaderCompilerEnvironment& OutEnvironment)
+{
+	FGlobalShader::ModifyCompilationEnvironment(Parameters, OutEnvironment);
+
+	OutEnvironment.SetDefine(TEXT("NUM_THREADS_PER_GROUP"), NumThreadsPerGroup);
+}
+
+IMPLEMENT_GLOBAL_SHADER(FPCGDataCollectionAdaptorCS, "/PCGComputeShaders/PCGDataCollectionAdaptor.usf", "MainCS", SF_Compute);
