[PCG] Update custom hlsl declarations
* Use squiggly braces instead of angle braces
* Add usage examples

#rnx

diff --git a/Engine/Plugins/PCG/Source/PCG/Private/Compute/Elements/PCGCustomHLSL.cpp b/Engine/Plugins/PCG/Source/PCG/Private/Compute/Elements/PCGCustomHLSL.cpp
--- a/Engine/Plugins/PCG/Source/PCG/Private/Compute/Elements/PCGCustomHLSL.cpp
+++ b/Engine/Plugins/PCG/Source/PCG/Private/Compute/Elements/PCGCustomHLSL.cpp
@@ -27,7 +27,7 @@
 
 namespace PCGHLSLElement
 {
-	const FString PinDeclTemplateStr = TEXT("<pin>");
+	const FString PinDeclTemplateStr = TEXT("{pin}");
 
 	/** First capture: Pin name (supports a - z, A - Z, and 0 - 9) */
 	constexpr int AttributePinCaptureGroup = 1;
@@ -761,8 +761,10 @@ void UPCGCustomHLSLSettings::UpdateInputDeclarations()
 			"\n"
 			"// Valid types: bool, int, float, float2, float3, float4, Rotator (float3), Quat (float4), Transform (float4x4), StringKey (int), Name (uint2)\n"
 			"\n"
-			"<type> {0}_Get<type>(uint DataIndex, uint ElementIndex, int AttributeId);\n"
-			"<type> {0}_Get<type>(uint DataIndex, uint ElementIndex, 'AttributeName');\n"),
+			"{type} {0}_Get{type}(uint DataIndex, uint ElementIndex, int AttributeId);\n"
+			"{type} {0}_Get{type}(uint DataIndex, uint ElementIndex, 'AttributeName');\n"
+			"\n"
+			"// Example: {0}_GetFloat({0}_DataIndex, ElementIndex, 'MyFloatAttr');\n"),
 			{ bMultiPin ? PCGHLSLElement::PinDeclTemplateStr : DataCollectionDataPins[0] });
 
 		InputDeclarations += TEXT("\n");
@@ -899,8 +901,10 @@ void UPCGCustomHLSLSettings::UpdateOutputDeclarations()
 		OutputDeclarations += FString::Format(TEXT(
 			"// Valid types: bool, int, float, float2, float3, float4, Rotator (float3), Quat (float4), Transform (float4x4), StringKey (int), Name (uint2)\n"
 			"\n"
-			"void {0}_Set<type>(uint DataIndex, uint ElementIndex, int AttributeId, <type> Value);\n"
-			"void {0}_Set<type>(uint DataIndex, uint ElementIndex, 'AttributeName', <type> Value);\n"),
+			"void {0}_Set{type}(uint DataIndex, uint ElementIndex, int AttributeId, {type} Value);\n"
+			"void {0}_Set{type}(uint DataIndex, uint ElementIndex, 'AttributeName', {type} Value);\n"
+			"\n"
+			"// Example: {0}_SetFloat({0}_DataIndex, ElementIndex, 'MyFloatAttr', MyValue);\n"),
 			{ bMultiPin ? PCGHLSLElement::PinDeclTemplateStr : DataCollectionDataPins[0] });
 
 		OutputDeclarations += TEXT("\n");
@@ -919,7 +923,7 @@ void UPCGCustomHLSLSettings::UpdateOutputDeclarations()
 
 		OutputDeclarations += FString::Format(TEXT(
 			"void {0}_InitializePoint(uint DataIndex, uint ElementIndex);\n"
-			"void {0}_CopyElementFrom_<input pin>(uint TargetDataIndex, uint TargetElementIndex, uint SourceDataIndex, uint SourceElementIndex);\n"
+			"void {0}_CopyElementFrom_{input pin}(uint TargetDataIndex, uint TargetElementIndex, uint SourceDataIndex, uint SourceElementIndex);\n"
 			"bool {0}_RemovePoint(uint DataIndex, uint ElementIndex);\n"
 			"\n"
 			"void {0}_SetPosition(uint DataIndex, uint ElementIndex, float3 Position);\n"
