[PCG] Fix tests for Graph Parameters

* Broken after the change in the PCG Graph to use another PostEditChange

#rnx

diff --git a/Engine/Plugins/PCG/Source/PCG/Private/Tests/PCGGraphParametersTest.cpp b/Engine/Plugins/PCG/Source/PCG/Private/Tests/PCGGraphParametersTest.cpp
--- a/Engine/Plugins/PCG/Source/PCG/Private/Tests/PCGGraphParametersTest.cpp
+++ b/Engine/Plugins/PCG/Source/PCG/Private/Tests/PCGGraphParametersTest.cpp
@@ -53,6 +53,17 @@ namespace PCGTests
 
 	constexpr const TCHAR* UserParametersName = TEXT("UserParameters");
 
+	void CallPostEditChange(UPCGGraph* Graph, FProperty* Property)
+	{
+		FEditPropertyChain PropertyChain;
+		PropertyChain.AddHead(Property);
+		FPropertyChangedEvent PropertyChangedEvent{ Property };
+
+		FPropertyChangedChainEvent PropertyChangedChainEvent{ PropertyChain, PropertyChangedEvent };
+
+		Graph->PostEditChangeChainProperty(PropertyChangedChainEvent);
+	}
+
 	template <typename Func>
 	void EmulateModifyingUserParameters(UPCGGraph* Graph, Func Callback)
 	{
@@ -62,8 +73,7 @@ namespace PCGTests
 
 		Callback();
 
-		FPropertyChangedEvent PropertyChangedEvent{Property};
-		Graph->PostEditChangeProperty(PropertyChangedEvent);
+		CallPostEditChange(Graph, Property);
 	}
 
 	template <typename Func>
@@ -78,8 +88,7 @@ namespace PCGTests
 
 		Callback();
 
-		FPropertyChangedEvent PropertyChangedEvent{ ValueProperty };
-		Graph->PostEditChangeProperty(PropertyChangedEvent);
+		CallPostEditChange(Graph, ValueProperty);
 	}
 }
 
