[PCG] Refactor Reroute node conversion to new Conversion system

#rnx

diff --git a/Engine/Plugins/PCG/Source/PCG/Private/Elements/PCGReroute.cpp b/Engine/Plugins/PCG/Source/PCG/Private/Elements/PCGReroute.cpp
--- a/Engine/Plugins/PCG/Source/PCG/Private/Elements/PCGReroute.cpp
+++ b/Engine/Plugins/PCG/Source/PCG/Private/Elements/PCGReroute.cpp
@@ -3,11 +3,33 @@
 #include "Elements/PCGReroute.h"
 
 #include "PCGContext.h"
-#include "PCGEdge.h"
+#include "PCGGraph.h"
 #include "PCGModule.h"
 #include "PCGNode.h"
 #include "PCGPin.h"
 #include "Compute/PCGDataForGPU.h"
+#include "Helpers/PCGConversion.h"
+
+#define LOCTEXT_NAMESPACE "PCGRerouteElement"
+
+namespace PCGReroute::Constants::Conversion
+{
+	constexpr int32 RerouteToNamedDeclarationIndex = 0;
+	constexpr int32 RerouteToNamedPairingIndex = 1;
+	constexpr int32 NamedDeclarationToRerouteIndex = 0;
+	constexpr int32 NamedUsageToRerouteIndex = 0;
+
+	const FText ToRerouteLabel = LOCTEXT("ToRerouteLabel", "Reroute Node");
+	const FText ToNamedDeclarationLabel = LOCTEXT("ToNamedDeclarationLabel", "Named Declaration");
+	const FText ToNamedPairingLabel = LOCTEXT("ToNamedPairingLabel", "Named Pairing");
+
+	const FText ToRerouteTooltip = LOCTEXT("ToRerouteTooltip", "Converts the named reroute to a normal reroute node.");
+	const FText ToNamedDeclarationTooltip = LOCTEXT("ToNamedDeclarationTooltip", "Converts the reroute node to a named reroute declaration.");
+	const FText ToNamedPairingTooltip = LOCTEXT("ToNamedPairingTooltip", "Converts the reroute node to a named reroute pair (declaration & usage).");
+
+	const FName DefaultRerouteTitle(LOCTEXT("DefaultRerouteTitle", "Reroute").ToString());
+	const FName DefaultNamedRerouteTitle(LOCTEXT("DefaultNamedRerouteTitle", "NamedReroute").ToString());
+}
 
 UPCGRerouteSettings::UPCGRerouteSettings()
 {
@@ -22,6 +44,25 @@ UPCGRerouteSettings::UPCGRerouteSettings()
 #endif
 }
 
+TOptional<FName> UPCGRerouteSettings::GetCollisionFreeNodeName(const UPCGGraph* InGraph, FName BaseName)
+{
+	FName NamedRerouteTitle(BaseName);
+
+	uint64 TitleIteration = 1;
+	while (InGraph->FindNodeByTitleName(NamedRerouteTitle))
+	{
+		constexpr uint64 MaxIterationsBeforeAbort = 100;
+		if (TitleIteration > MaxIterationsBeforeAbort)
+		{
+			return {};
+		}
+
+		NamedRerouteTitle.SetNumber(++TitleIteration);
+	}
+
+	return NamedRerouteTitle;
+}
+
 TArray<FPCGPinProperties> UPCGRerouteSettings::InputPinProperties() const
 {
 	FPCGPinProperties PinProperties;
@@ -62,6 +103,99 @@ bool UPCGRerouteSettings::ComputeOutputPinDataDesc(const UPCGPin* OutputPin, con
 	return false;
 }
 
+#if WITH_EDITOR
+TArray<FPCGPreconfiguredInfo> UPCGRerouteSettings::GetConversionInfo() const
+{
+	using namespace PCGReroute::Constants::Conversion;
+	TArray<FPCGPreconfiguredInfo> ConversionInfo;
+
+	ConversionInfo.Emplace(RerouteToNamedDeclarationIndex, ToNamedDeclarationLabel, ToNamedDeclarationTooltip);
+	ConversionInfo.Emplace(RerouteToNamedPairingIndex, ToNamedPairingLabel, ToNamedPairingTooltip);
+
+	return ConversionInfo;
+}
+#endif // WITH_EDITOR
+
+bool UPCGRerouteSettings::ConvertNode(const FPCGPreconfiguredInfo& ConversionInfo)
+{
+	using namespace PCGReroute::Constants::Conversion;
+	UPCGNode* Node = CastChecked<UPCGNode>(GetOuter());
+	UPCGGraph* Graph = CastChecked<UPCGGraph>(Node->GetOuter());
+
+	switch (ConversionInfo.PreconfiguredIndex)
+	{
+		case RerouteToNamedDeclarationIndex:
+		{
+			if (TOptional<FName> NamedRerouteTitle = GetCollisionFreeNodeName(Graph, DefaultNamedRerouteTitle))
+			{
+				Node->NodeTitle = *NamedRerouteTitle;
+			}
+			else
+			{
+				return false;
+			}
+
+			FPCGSingleNodeConverter NodeConverter(Node, UPCGNamedRerouteDeclarationSettings::StaticClass());
+			if (!NodeConverter.IsGraphInitialized() || !NodeConverter.IsSourceInitialized())
+			{
+				PCGLog::Settings::LogInvalidPreconfigurationWarning(ConversionInfo.PreconfiguredIndex, ToRerouteLabel);
+				return false;
+			}
+
+			NodeConverter.PrepareData();
+			NodeConverter.ApplyStructural();
+			NodeConverter.Finalize();
+
+			if (NodeConverter.IsComplete())
+			{
+				return true;
+			}
+			else
+			{
+				// Back out the title change
+				Node->NodeTitle = NAME_None;
+				return false;
+			}
+		}
+
+		case RerouteToNamedPairingIndex:
+		{
+			TOptional<FName> NamedRerouteTitle = GetCollisionFreeNodeName(Graph, DefaultNamedRerouteTitle);
+			if (!NamedRerouteTitle)
+			{
+				return false;
+			}
+
+			FPCGReroutePairNodeConverter NodeConverter(Node, *NamedRerouteTitle);
+			if (!NodeConverter.IsGraphInitialized() || !NodeConverter.IsSourceInitialized())
+			{
+				PCGLog::Settings::LogInvalidPreconfigurationWarning(ConversionInfo.PreconfiguredIndex, ToRerouteLabel);
+				return false;
+			}
+
+			NodeConverter.PrepareData();
+
+			UPCGNamedRerouteDeclarationSettings* DeclarationSettings = CastChecked<UPCGNamedRerouteDeclarationSettings>(NodeConverter.GetGeneratedDeclarationSettings());
+			// TODO: Intentionally handled here as it will be in the future, when the builder replaces the converter class. Remove this comment at that time.
+			for (UPCGSettings* Settings : NodeConverter.GetGeneratedUsageSettings())
+			{
+				UPCGNamedRerouteUsageSettings* UsageSettings = CastChecked<UPCGNamedRerouteUsageSettings>(Settings);
+				UsageSettings->Declaration = DeclarationSettings;
+			}
+
+			NodeConverter.ApplyStructural();
+			NodeConverter.Finalize();
+
+			return NodeConverter.IsComplete();
+		}
+
+		default:
+			break;
+	}
+
+	return false;
+}
+
 TArray<FPCGPinProperties> UPCGNamedRerouteDeclarationSettings::OutputPinProperties() const
 {
 	TArray<FPCGPinProperties> PinProperties;
@@ -82,6 +216,41 @@ TArray<FPCGPinProperties> UPCGNamedRerouteDeclarationSettings::OutputPinProperti
 	return PinProperties;
 }
 
+#if WITH_EDITOR
+TArray<FPCGPreconfiguredInfo> UPCGNamedRerouteDeclarationSettings::GetConversionInfo() const
+{
+	using namespace PCGReroute::Constants::Conversion;
+	TArray<FPCGPreconfiguredInfo> ConversionInfo;
+
+	ConversionInfo.Emplace(NamedDeclarationToRerouteIndex, ToRerouteLabel, ToRerouteLabel);
+
+	return ConversionInfo;
+}
+#endif // WITH_EDITOR
+
+bool UPCGNamedRerouteDeclarationSettings::ConvertNode(const FPCGPreconfiguredInfo& ConversionInfo)
+{
+	using namespace PCGReroute::Constants::Conversion;
+	UPCGNode* Node = CastChecked<UPCGNode>(GetOuter());
+	if (ConversionInfo.PreconfiguredIndex == NamedDeclarationToRerouteIndex)
+	{
+		FPCGRerouteDeclarationConverter NodeConverter(Node, DefaultRerouteTitle);
+		if (!NodeConverter.IsGraphInitialized() || !NodeConverter.IsSourceInitialized())
+		{
+			PCGLog::Settings::LogInvalidPreconfigurationWarning(ConversionInfo.PreconfiguredIndex, ToNamedDeclarationLabel);
+			return false;
+		}
+
+		NodeConverter.PrepareData();
+		NodeConverter.ApplyStructural();
+		NodeConverter.Finalize();
+
+		return NodeConverter.IsComplete();
+	}
+
+	return false;
+}
+
 TArray<FPCGPinProperties> UPCGNamedRerouteUsageSettings::InputPinProperties() const
 {
 	TArray<FPCGPinProperties> PinProperties = Super::InputPinProperties();
@@ -115,3 +284,5 @@ bool FPCGRerouteElement::ExecuteInternal(FPCGContext* Context) const
 
 	return true;
 }
+
+#undef LOCTEXT_NAMESPACE
\ No newline at end of file
diff --git a/Engine/Plugins/PCG/Source/PCG/Private/Helpers/PCGConversion.cpp b/Engine/Plugins/PCG/Source/PCG/Private/Helpers/PCGConversion.cpp
--- a/Engine/Plugins/PCG/Source/PCG/Private/Helpers/PCGConversion.cpp
+++ b/Engine/Plugins/PCG/Source/PCG/Private/Helpers/PCGConversion.cpp
@@ -6,6 +6,8 @@
 #include "PCGGraph.h"
 #include "PCGNode.h"
 #include "PCGSettings.h"
+// TODO: Included temporarily until a graph builder replaces the need for the reroute converter
+#include "Elements/PCGReroute.h"
 
 namespace PCGConversion
 {
@@ -26,7 +28,22 @@ namespace PCGConversion
 
 	namespace Node::Helpers
 	{
-		bool RewireEdges(UPCGPin* OldPin, UPCGPin* NewPin)
+		bool RewireSingleEdge(UPCGEdge* Edge, UPCGPin* OldPin, UPCGPin* NewPin)
+		{
+			if (!Edge || !OldPin || !NewPin || OldPin == NewPin)
+			{
+				return false;
+			}
+
+			UPCGPin* OtherPin = Edge->InputPin == OldPin ? Edge->OutputPin : Edge->InputPin;
+			check(OtherPin);
+			OtherPin->BreakEdgeTo(OldPin);
+			OtherPin->AddEdgeTo(NewPin);
+
+			return true;
+		}
+
+		bool RewireAllEdges(UPCGPin* OldPin, UPCGPin* NewPin)
 		{
 			if (!OldPin || !NewPin || OldPin == NewPin)
 			{
@@ -35,11 +52,7 @@ namespace PCGConversion
 
 			while (!OldPin->Edges.IsEmpty())
 			{
-				UPCGEdge* Edge = OldPin->Edges[0];
-				UPCGPin* OtherPin = Edge->InputPin == OldPin ? Edge->OutputPin : Edge->InputPin;
-				check(OtherPin);
-				OtherPin->BreakEdgeTo(OldPin);
-				OtherPin->AddEdgeTo(NewPin);
+				RewireSingleEdge(OldPin->Edges[0], OldPin, NewPin);
 			}
 
 			return true;
@@ -114,7 +127,7 @@ bool FPCGSingleNodeConverter::ExecutePrepareData(bool& bOutMarkGraphAsDirty)
 
 		if (GeneratedNode && GeneratedSettings)
 		{
-			if (GeneratedNode->HasAuthoredTitle())
+			if (SourceNode->HasAuthoredTitle())
 			{
 				GeneratedNode->NodeTitle = SourceNode->NodeTitle;
 			}
@@ -144,14 +157,14 @@ bool FPCGSingleNodeConverter::ExecuteApplyStructural(bool& bOutMarkGraphAsDirty)
 	const TArray<TObjectPtr<UPCGPin>>& GeneratedInputPins = GeneratedNode->GetInputPins();
 	for (int32 InputIndex = 0; InputIndex < GeneratedInputPins.Num() && InputIndex < SourceInputPins.Num(); ++InputIndex)
 	{
-		bResult &= PCGConversion::Node::Helpers::RewireEdges(SourceInputPins[InputIndex], GeneratedInputPins[InputIndex]);
+		bResult &= PCGConversion::Node::Helpers::RewireAllEdges(SourceInputPins[InputIndex], GeneratedInputPins[InputIndex]);
 	}
 
 	const TArray<TObjectPtr<UPCGPin>>& SourceOutputPins = SourceNode->GetOutputPins();
 	const TArray<TObjectPtr<UPCGPin>>& GeneratedOutputPins = GeneratedNode->GetOutputPins();
 	for (int32 OutputIndex = 0; OutputIndex < GeneratedOutputPins.Num() && OutputIndex < SourceOutputPins.Num(); ++OutputIndex)
 	{
-		bResult &= PCGConversion::Node::Helpers::RewireEdges(SourceOutputPins[OutputIndex], GeneratedOutputPins[OutputIndex]);
+		bResult &= PCGConversion::Node::Helpers::RewireAllEdges(SourceOutputPins[OutputIndex], GeneratedOutputPins[OutputIndex]);
 	}
 
 	SourceGraph->RemoveNode(SourceNode);
@@ -165,4 +178,199 @@ void FPCGSingleNodeConverter::Finalize()
 	{
 		GeneratedNode->UpdateAfterSettingsChangeDuringCreation();
 	}
-}
\ No newline at end of file
+}
+
+FPCGRerouteDeclarationConverter::FPCGRerouteDeclarationConverter(UPCGNode* InNode, FName InNodeTitle)
+	: FPCGSingleNodeConverter(InNode, UPCGRerouteSettings::StaticClass())
+	, RerouteNodeTitle(InNodeTitle)
+{}
+
+bool FPCGRerouteDeclarationConverter::ExecutePrepareData(bool& bOutMarkGraphAsDirty)
+{
+	if (FPCGSingleNodeConverter::ExecutePrepareData(bOutMarkGraphAsDirty))
+	{
+		GeneratedNode->NodeTitle = RerouteNodeTitle;
+		return true;
+	}
+
+	return false;
+}
+
+bool FPCGRerouteDeclarationConverter::ExecuteApplyStructural(bool& bOutMarkGraphAsDirty)
+{
+	if (!ensure(IsValid()))
+	{
+		return false;
+	}
+
+	bool bResult = true;
+	const TArray<TObjectPtr<UPCGPin>>& SourceInputPins = SourceNode->GetInputPins();
+	const TArray<TObjectPtr<UPCGPin>>& GeneratedInputPins = GeneratedNode->GetInputPins();
+	check(SourceInputPins.Num() == 1 && GeneratedInputPins.Num() == 1);
+	bResult &= PCGConversion::Node::Helpers::RewireAllEdges(SourceInputPins[0], GeneratedInputPins[0]);
+
+	const TArray<TObjectPtr<UPCGPin>>& GeneratedOutputPins = GeneratedNode->GetOutputPins();
+	checkf(GeneratedOutputPins.Num() == 1, TEXT("%s"), TEXT("This class currently only supports converting reroute nodes."));
+	if (UPCGNamedRerouteDeclarationSettings* SourceSettings = Cast<UPCGNamedRerouteDeclarationSettings>(SourceNode->GetSettings()))
+	{
+		const TArray<TObjectPtr<UPCGPin>>& SourceOutputPins = SourceNode->GetOutputPins();
+		checkf(SourceOutputPins.Num() == 2, TEXT("%s"), TEXT("This class currently only supports converting reroute nodes."));
+		bResult &= PCGConversion::Node::Helpers::RewireAllEdges(SourceOutputPins[0], GeneratedOutputPins[0]);
+
+		for (UPCGNode* UsageNode : SourceGraph->GetNodes())
+		{
+			const UPCGNamedRerouteUsageSettings* Settings = Cast<UPCGNamedRerouteUsageSettings>(UsageNode->GetSettings());
+			if (Settings && Settings->Declaration == SourceSettings)
+			{
+				check(UsageNode->GetOutputPins().Num() == 1);
+				bResult &= PCGConversion::Node::Helpers::RewireAllEdges(UsageNode->GetOutputPins()[0], GeneratedOutputPins[0]);
+				SourceGraph->RemoveNode(UsageNode);
+			}
+		}
+	}
+
+	SourceGraph->RemoveNode(SourceNode);
+
+	return bResult;
+}
+
+FPCGReroutePairNodeConverter::FPCGReroutePairNodeConverter(UPCGNode* InRerouteNode, FName InNodeTitle)
+	: FPCGConverterBase(InRerouteNode ? InRerouteNode->GetGraph() : nullptr)
+	, SourceNode(InRerouteNode)
+	, RerouteNodeTitle(InNodeTitle)
+{
+	if (SourceNode && SourceNode->GetSettings() && SourceNode->GetSettings()->IsA<UPCGRerouteSettings>())
+	{
+		SetSourceInitialized();
+	}
+}
+
+void FPCGReroutePairNodeConverter::Finalize()
+{
+	if (ensure(IsValid()))
+	{
+		GeneratedDeclarationNode->UpdateAfterSettingsChangeDuringCreation();
+
+		TArray<UPCGNode*, TInlineAllocator<16>> UsageNodes;
+		for (auto [DownstreamNode, UsageNode] : DownstreamToUsageNodeMapping)
+		{
+			if (!UsageNodes.Contains(UsageNode))
+			{
+				UsageNode->UpdateAfterSettingsChangeDuringCreation();
+				UsageNodes.Add(UsageNode);
+			}
+		}
+	}
+}
+
+bool FPCGReroutePairNodeConverter::IsValid()
+{
+	bool bIsValid = SourceGraph != nullptr;
+	if (!IsDataPrepared())
+	{
+		return bIsValid;
+	}
+
+	bIsValid &= GeneratedDeclarationNode != nullptr;
+
+	for (auto [DownstreamNode, UsageNode] : DownstreamToUsageNodeMapping)
+	{
+		bIsValid &= (DownstreamNode != nullptr) && (UsageNode != nullptr);
+	}
+
+	return bIsValid;
+}
+
+bool FPCGReroutePairNodeConverter::ExecutePrepareData(bool& bOutMarkGraphAsDirty)
+{
+	if (ensure(IsValid()))
+	{
+		bOutMarkGraphAsDirty = true;
+		GeneratedDeclarationNode = SourceGraph->AddNodeOfType(UPCGNamedRerouteDeclarationSettings::StaticClass(), GeneratedDeclarationSettings);
+		DownstreamToUsageNodeMapping.Empty();
+
+		const TArray<TObjectPtr<UPCGPin>>& SourceOutputPins = SourceNode->GetOutputPins();
+		checkf(SourceOutputPins.Num() == 1, TEXT("%s"), TEXT("This class currently only supports converting reroute nodes."));
+
+		auto TransferProperties = [this](UPCGNode* Node, UPCGSettings* Settings)
+		{
+			if (Node && Settings)
+			{
+#if WITH_EDITOR
+				SourceNode->TransferEditorProperties(GeneratedDeclarationNode);
+#endif // WITH_EDITOR
+				GeneratedDeclarationNode->NodeTitle = RerouteNodeTitle;
+				GeneratedDeclarationSettings->bEnabled = SourceNode->GetSettings()->bEnabled;
+
+				return true;
+			}
+
+			return false;
+		};
+
+		bool bSuccess = TransferProperties(GeneratedDeclarationNode, GeneratedDeclarationSettings);
+
+		for (const UPCGEdge* Edge : SourceOutputPins[0]->Edges)
+		{
+			const UPCGNode* DownstreamNode = Edge->OutputPin->Node;
+			check(DownstreamNode);
+			if (!DownstreamToUsageNodeMapping.Contains(DownstreamNode))
+			{
+				UPCGNode* UsageNode = SourceGraph->AddNodeOfType(UPCGNamedRerouteUsageSettings::StaticClass(), GeneratedUsageSettings.Emplace_GetRef());
+				DownstreamToUsageNodeMapping.Add(DownstreamNode, UsageNode);
+				bSuccess &= TransferProperties(UsageNode, GeneratedUsageSettings.Last());
+				UsageNode->NodeTitle = RerouteNodeTitle;
+#if WITH_EDITOR
+				// TODO: Temporary hardcoded offset. In the future builder, we should have helpers that will position the
+				// node relative to the downstream node, check the graph for overlaps, etc.
+				UsageNode->SetNodePosition(DownstreamNode->PositionX - 200, DownstreamNode->PositionY);
+#endif // WITH_EDITOR
+				UPCGPin* DeclarationInvisiblePin = GeneratedDeclarationNode->GetOutputPin(PCGNamedRerouteConstants::InvisiblePinLabel);
+				UPCGPin* UsageInvisiblePin = UsageNode->GetInputPin(PCGPinConstants::DefaultInputLabel);
+				DeclarationInvisiblePin->AddEdgeTo(UsageInvisiblePin);
+			}
+		}
+
+		return bSuccess;
+	}
+
+	return false;
+}
+
+bool FPCGReroutePairNodeConverter::ExecuteApplyStructural(bool& bOutMarkGraphAsDirty)
+{
+	if (!ensure(IsValid()))
+	{
+		return false;
+	}
+
+	bool bResult = true;
+	const TArray<TObjectPtr<UPCGPin>>& SourceInputPins = SourceNode->GetInputPins();
+	const TArray<TObjectPtr<UPCGPin>>& GeneratedInputPins = GeneratedDeclarationNode->GetInputPins();
+	for (int32 InputIndex = 0; InputIndex < GeneratedInputPins.Num() && InputIndex < SourceInputPins.Num(); ++InputIndex)
+	{
+		bResult &= PCGConversion::Node::Helpers::RewireAllEdges(SourceInputPins[InputIndex], GeneratedInputPins[InputIndex]);
+	}
+
+	if (!DownstreamToUsageNodeMapping.IsEmpty())
+	{
+		const TArray<TObjectPtr<UPCGPin>>& SourceOutputPins = SourceNode->GetOutputPins();
+		checkf(SourceOutputPins.Num() == 1, TEXT("%s"), TEXT("This class currently only supports converting reroute nodes."));
+		for (UPCGPin* Pin : SourceOutputPins)
+		{
+			check(Pin);
+			for (int32 EdgeIndex = 0; !Pin->Edges.IsEmpty(); ++EdgeIndex)
+			{
+				check(Pin->Edges[0] && Pin->Edges[0]->OutputPin && Pin->Edges[0]->OutputPin->Node)
+				UPCGNode* UsageNode = DownstreamToUsageNodeMapping[Pin->Edges[0]->OutputPin->Node];
+				check(UsageNode->GetOutputPins().Num() == 1);
+				UPCGPin* UsagePin = UsageNode->GetOutputPin(PCGPinConstants::DefaultOutputLabel);
+				bResult &= PCGConversion::Node::Helpers::RewireSingleEdge(Pin->Edges[0], Pin, UsagePin);
+			}
+		}
+	}
+
+	SourceGraph->RemoveNode(SourceNode);
+
+	return bResult;
+}
diff --git a/Engine/Plugins/PCG/Source/PCG/Private/Helpers/PCGConversion.h b/Engine/Plugins/PCG/Source/PCG/Private/Helpers/PCGConversion.h
--- a/Engine/Plugins/PCG/Source/PCG/Private/Helpers/PCGConversion.h
+++ b/Engine/Plugins/PCG/Source/PCG/Private/Helpers/PCGConversion.h
@@ -95,7 +95,7 @@ private:
 };
 
 /** Converts a single source node into a single compatible target node and rewires edges. */
-class FPCGSingleNodeConverter final : public FPCGConverterBase // TODO: final can be removed when the class becomes public
+class FPCGSingleNodeConverter : public FPCGConverterBase
 {
 public:
 	explicit FPCGSingleNodeConverter(UPCGNode* InNode, TSubclassOf<UPCGSettings> InTargetSettingsClass);
@@ -122,12 +122,75 @@ protected:
 
 	/** The node being converted. */
 	UPCGNode* const SourceNode;
+	/** The node created during the conversion. */
+	UPCGNode* GeneratedNode = nullptr;
 	/** The class of the target node. */
 	const TSubclassOf<UPCGSettings> TargetSettingsClass;
 
 private:
-	/** The node created during the conversion. */
-	UPCGNode* GeneratedNode = nullptr;
 	/** The settings created during the conversion. */
 	UPCGSettings* GeneratedSettings = nullptr;
-};
\ No newline at end of file
+};
+
+/**
+ * [EXPERIMENTAL] Converts a single source node into a single compatible target node and rewires edges.
+ * TODO: This class is temporary and will likely be replaced by a non-explicit converter that utilizes a graph building system.
+ */
+class FPCGRerouteDeclarationConverter : public FPCGSingleNodeConverter
+{
+public:
+	explicit FPCGRerouteDeclarationConverter(UPCGNode* InNode, FName InNodeTitle);
+
+	//~Begin FPCGConverterBase interface
+	/** Create the new node and settings. */
+	virtual bool ExecutePrepareData(bool& bOutMarkGraphAsDirty) override;
+	/** Rewire the input and output pins. */
+	virtual bool ExecuteApplyStructural(bool& bOutMarkGraphAsDirty) override;
+	//~Begin FPCGConverterBase interface
+
+private:
+	FName RerouteNodeTitle = NAME_None;
+};
+
+/**
+ * [EXPERIMENTAL] Converts a single reroute node into named reroute pairs--the declaration to the upstream node and one
+ * usage node for every downstream node.
+ * TODO: This class is temporary and will likely be replaced by a non-explicit converter that utilizes a graph building system.
+ */
+class FPCGReroutePairNodeConverter final : public FPCGConverterBase
+{
+public:
+	explicit FPCGReroutePairNodeConverter(UPCGNode* InRerouteNode, FName InNodeTitle);
+
+	//~Begin FPCGConverterBase interface
+	/** The node conversion still has valid references. */
+	virtual bool IsValid() override;
+	//~Begin FPCGConverterBase interface
+
+	/** Finalizes the node conversion. */
+	void Finalize();
+
+	UPCGNode* GetSourceNode() const { return SourceNode; }
+	UPCGNode* GetGeneratedDeclarationNode() const { return GeneratedDeclarationNode; }
+	UPCGSettings* GetGeneratedDeclarationSettings() const { return GeneratedDeclarationSettings; }
+	const TMap<const UPCGNode*, UPCGNode*>& GetDownstreamToUsageNodeMapping() const { return DownstreamToUsageNodeMapping; }
+	const TArray<UPCGSettings*>& GetGeneratedUsageSettings() const { return GeneratedUsageSettings; }
+
+protected:
+	//~Begin FPCGConverterBase interface
+	/** Create the new node and settings. */
+	virtual bool ExecutePrepareData(bool& bOutMarkGraphAsDirty) override;
+	/** Rewire the input and output pins. */
+	virtual bool ExecuteApplyStructural(bool& bOutMarkGraphAsDirty) override;
+	//~Begin FPCGConverterBase interface
+
+	UPCGNode* const SourceNode;
+
+private:
+	UPCGNode* GeneratedDeclarationNode = nullptr;
+	UPCGSettings* GeneratedDeclarationSettings = nullptr;
+
+	FName RerouteNodeTitle = NAME_None;
+	TArray<UPCGSettings*> GeneratedUsageSettings;
+	TMap<const UPCGNode*, UPCGNode*> DownstreamToUsageNodeMapping;
+};
diff --git a/Engine/Plugins/PCG/Source/PCG/Private/PCGGraph.cpp b/Engine/Plugins/PCG/Source/PCG/Private/PCGGraph.cpp
--- a/Engine/Plugins/PCG/Source/PCG/Private/PCGGraph.cpp
+++ b/Engine/Plugins/PCG/Source/PCG/Private/PCGGraph.cpp
@@ -555,6 +555,35 @@ UPCGNode* UPCGGraph::FindNodeWithSettings(const UPCGSettingsInterface* InSetting
 	return NodeFound;
 }
 
+UPCGNode* UPCGGraph::FindNodeByTitleName(FName NodeTitle, bool bRecursive) const
+{
+	UPCGNode* NodeFound = nullptr;
+
+	auto FindNode = [&NodeFound, NodeTitle](UPCGNode* InNode)
+	{
+		if (InNode && InNode->NodeTitle == NodeTitle)
+		{
+			NodeFound = InNode;
+			return false; // stop execution
+		}
+		else
+		{
+			return true;
+		}
+	};
+
+	if (bRecursive)
+	{
+		ForEachNodeRecursively(FindNode);
+	}
+	else
+	{
+		ForEachNode(FindNode);
+	}
+
+	return NodeFound;
+}
+
 #if WITH_EDITOR
 void UPCGGraph::PreSave(FObjectPreSaveContext ObjectSaveContext)
 {
diff --git a/Engine/Plugins/PCG/Source/PCG/Private/PCGPin.cpp b/Engine/Plugins/PCG/Source/PCG/Private/PCGPin.cpp
--- a/Engine/Plugins/PCG/Source/PCG/Private/PCGPin.cpp
+++ b/Engine/Plugins/PCG/Source/PCG/Private/PCGPin.cpp
@@ -7,6 +7,8 @@
 #include "PCGNode.h"
 #include "PCGSettings.h"
 
+#include "UObject/Package.h"
+
 #include UE_INLINE_GENERATED_CPP_BY_NAME(PCGPin)
 
 namespace PCGPin
@@ -234,6 +236,7 @@ bool UPCGPin::BreakEdgeTo(UPCGPin* OtherPin, TSet<UPCGNode*>* InTouchedNodes/*=
 
 			ensure(OtherPin->Edges.Remove(Edge));
 			Edges.Remove(Edge);
+			Edge->Rename(/*NewName=*/nullptr, GetTransientPackage(), REN_DontCreateRedirectors);
 
 			if (InTouchedNodes)
 			{
@@ -268,6 +271,7 @@ bool UPCGPin::BreakAllEdges(TSet<UPCGNode*>* InTouchedNodes/*= nullptr*/)
 		{
 			OtherPin->Modify();
 			ensure(OtherPin->Edges.Remove(Edge));
+			Edge->Rename(/*NewName=*/nullptr, GetTransientPackage(), REN_DontCreateRedirectors);
 			bChanged = true;
 
 			if (InTouchedNodes)
diff --git a/Engine/Plugins/PCG/Source/PCG/Public/Elements/PCGReroute.h b/Engine/Plugins/PCG/Source/PCG/Public/Elements/PCGReroute.h
--- a/Engine/Plugins/PCG/Source/PCG/Public/Elements/PCGReroute.h
+++ b/Engine/Plugins/PCG/Source/PCG/Public/Elements/PCGReroute.h
@@ -36,11 +36,18 @@ public:
 	virtual bool CanUserEditTitle() const override { return false; }
 #endif
 
+	static PCG_API TOptional<FName> GetCollisionFreeNodeName(const UPCGGraph* InGraph, FName BaseName);
+
 protected:
 	virtual TArray<FPCGPinProperties> InputPinProperties() const override;
 	virtual TArray<FPCGPinProperties> OutputPinProperties() const override;
 	virtual FPCGElementPtr CreateElement() const override;
 	virtual bool ComputeOutputPinDataDesc(const UPCGPin* OutputPin, const UPCGDataBinding* Binding, FPCGDataCollectionDesc& OutDesc) const override;
+
+#if WITH_EDITOR
+	virtual TArray<FPCGPreconfiguredInfo> GetConversionInfo() const override;
+#endif // WITH_EDITOR
+	virtual bool ConvertNode(const FPCGPreconfiguredInfo& ConversionInfo) override;
 	//~End UPCGSettings interface
 };
 
@@ -65,6 +72,10 @@ public:
 
 protected:
 	virtual TArray<FPCGPinProperties> OutputPinProperties() const override;
+#if WITH_EDITOR
+	virtual TArray<FPCGPreconfiguredInfo> GetConversionInfo() const override;
+#endif // WITH_EDITOR
+	virtual bool ConvertNode(const FPCGPreconfiguredInfo& ConversionInfo) override;
 
 public:
 #if WITH_EDITOR
@@ -86,6 +97,10 @@ public:
 
 protected:
 	virtual TArray<FPCGPinProperties> InputPinProperties() const override;
+#if WITH_EDITOR
+	virtual TArray<FPCGPreconfiguredInfo> GetConversionInfo() const override { return {}; }
+#endif // WITH_EDITOR
+	virtual bool ConvertNode(const FPCGPreconfiguredInfo& ConversionInfo) override { return false; }
 
 public:
 	/** Very counter-intuitive but reroute nodes are normally culled by other means, if they aren't we want to make sure they log errors. */
diff --git a/Engine/Plugins/PCG/Source/PCG/Public/PCGGraph.h b/Engine/Plugins/PCG/Source/PCG/Public/PCGGraph.h
--- a/Engine/Plugins/PCG/Source/PCG/Public/PCGGraph.h
+++ b/Engine/Plugins/PCG/Source/PCG/Public/PCGGraph.h
@@ -335,6 +335,9 @@ public:
 	/** Returns the node with the given settings in the graph, if any */
 	UPCGNode* FindNodeWithSettings(const UPCGSettingsInterface* InSettings, bool bRecursive = false) const;
 
+	/** Returns the first node that matches the given name in the graph, if any.*/
+	UPCGNode* FindNodeByTitleName(FName NodeTitle, bool bRecursive = false) const;
+
 	const TArray<UPCGNode*>& GetNodes() const { return Nodes; }
 	void AddNode(UPCGNode* InNode);
 	void AddNodes(TArray<UPCGNode*>& InNodes);
diff --git a/Engine/Plugins/PCG/Source/PCGEditor/Private/PCGEditor.cpp b/Engine/Plugins/PCG/Source/PCGEditor/Private/PCGEditor.cpp
--- a/Engine/Plugins/PCG/Source/PCGEditor/Private/PCGEditor.cpp
+++ b/Engine/Plugins/PCG/Source/PCGEditor/Private/PCGEditor.cpp
@@ -820,11 +820,6 @@ void FPCGEditor::BindCommands()
 		FExecuteAction::CreateSP(this, &FPCGEditor::OnRenameNode),
 		FCanExecuteAction::CreateSP(this, &FPCGEditor::CanRenameNode));
 
-	GraphEditorCommands->MapAction(
-		PCGEditorCommands.ConvertNamedRerouteToReroute,
-		FExecuteAction::CreateSP(this, &FPCGEditor::OnConvertNamedRerouteToReroute),
-		FCanExecuteAction::CreateSP(this, &FPCGEditor::CanConvertNamedRerouteToReroute));
-
 	GraphEditorCommands->MapAction(
 		PCGEditorCommands.SelectNamedRerouteUsages,
 		FExecuteAction::CreateSP(this, &FPCGEditor::OnSelectNamedRerouteUsages),
@@ -834,11 +829,6 @@ void FPCGEditor::BindCommands()
 		PCGEditorCommands.SelectNamedRerouteDeclaration,
 		FExecuteAction::CreateSP(this, &FPCGEditor::OnSelectNamedRerouteDeclaration),
 		FCanExecuteAction::CreateSP(this, &FPCGEditor::CanSelectNamedRerouteDeclaration));
-
-	GraphEditorCommands->MapAction(
-		PCGEditorCommands.ConvertRerouteToNamedReroute,
-		FExecuteAction::CreateSP(this, &FPCGEditor::OnConvertRerouteToNamedReroute),
-		FCanExecuteAction::CreateSP(this, &FPCGEditor::CanConvertRerouteToNamedReroute));
 }
 
 void FPCGEditor::OnFind()
@@ -1311,132 +1301,6 @@ bool FPCGEditor::InternalValidationOnAction()
 	return true;
 }
 
-void FPCGEditor::OnConvertNamedRerouteToReroute()
-{
-	if (!InternalValidationOnAction())
-	{
-		return;
-	}
-
-	UPCGGraph* PCGGraph = PCGEditorGraph->GetPCGGraph();
-	check(PCGGraph);
-
-	TArray<UPCGNode*> DeclarationsToConvert;
-	for (UObject* Object : GraphEditorWidget->GetSelectedNodes())
-	{
-		if (!Object)
-		{
-			continue;
-		}
-
-		if (UPCGEditorGraphNodeNamedRerouteDeclaration* DeclarationNode = Cast<UPCGEditorGraphNodeNamedRerouteDeclaration>(Object))
-		{
-			check(DeclarationNode->GetPCGNode());
-			DeclarationsToConvert.Add(DeclarationNode->GetPCGNode());
-		}
-	}
-
-	if (DeclarationsToConvert.IsEmpty())
-	{
-		return;
-	}
-
-	// Disable graph notifications
-	PCGGraph->DisableNotificationsForEditor();
-
-	// For all found declarations, replace by a normal reroute node and forward the connection from the usages to this reroute instead.
-	TArray<UPCGNode*> NodesToDelete;
-
-	for (UPCGNode* Declaration : DeclarationsToConvert)
-	{
-		NodesToDelete.Add(Declaration);
-
-		// Create reroute node
-		UPCGSettings* RerouteNodeSettings = nullptr;
-		UPCGNode* RerouteNode = PCGGraph->AddNodeOfType(UPCGRerouteSettings::StaticClass(), RerouteNodeSettings);
-
-		// Copy any non-functional information
-		Declaration->TransferEditorProperties(RerouteNode);
-
-		// Force a graph refresh - needed for the other operations/editor callbacks to go through, a bit unfortunate.
-		// TODO: improve this
-		PCGEditorGraph->ReconstructGraph();
-
-		UPCGPin* RerouteInput = RerouteNode->GetInputPin(PCGPinConstants::DefaultInputLabel);
-		UPCGPin* RerouteOutput = RerouteNode->GetOutputPin(PCGPinConstants::DefaultOutputLabel);
-
-		// Copy input edge to the declaration to the new reroute node
-		check(Declaration->GetInputPin(PCGPinConstants::DefaultInputLabel));
-		
-		TArray<UPCGEdge*> InputEdges = Declaration->GetInputPin(PCGPinConstants::DefaultInputLabel)->Edges;
-		ensure(InputEdges.Num() <= 1);
-
-		if(InputEdges.Num() >= 1)
-		{
-			PCGGraph->AddEdge(InputEdges[0]->InputPin->Node, InputEdges[0]->InputPin->Properties.Label, RerouteNode, RerouteInput->Properties.Label);
-		}
-
-		// The reroute declaration has edges to all its usages, and potentially other edges through its out pin.
-		// Keep track of all usages to be deleted but also forward edges from the usages to the newly created reroute.
-		if (UPCGPin* OutputPin = Declaration->GetOutputPin(PCGPinConstants::DefaultOutputLabel))
-		{
-			TArray<UPCGEdge*> Edges = OutputPin->Edges;
-			for(UPCGEdge* Edge : Edges)
-			{
-				PCGGraph->AddEdge(RerouteNode, RerouteOutput->Properties.Label, Edge->OutputPin->Node, Edge->OutputPin->Properties.Label);
-			}
-		}
-
-		if (UPCGPin* InvisiblePin = Declaration->GetOutputPin(PCGNamedRerouteConstants::InvisiblePinLabel))
-		{
-			TArray<UPCGEdge*> Edges = InvisiblePin->Edges;
-			for(UPCGEdge* Edge : Edges)
-			{
-				check(Edge->OutputPin && Edge->OutputPin->Node);
-				UPCGNode* Usage = Edge->OutputPin->Node;
-				NodesToDelete.Add(Usage);
-
-				check(Usage->GetOutputPin(PCGPinConstants::DefaultOutputLabel));
-				TArray<UPCGEdge*> UsageEdges = Usage->GetOutputPin(PCGPinConstants::DefaultOutputLabel)->Edges;
-				for(UPCGEdge* UsageEdge : UsageEdges)
-				{
-					PCGGraph->AddEdge(RerouteNode, RerouteOutput->Properties.Label, UsageEdge->OutputPin->Node, UsageEdge->OutputPin->Properties.Label);
-				}
-			}
-		}
-	}
-
-	// Delete removed ndoes
-	PCGGraph->RemoveNodes(NodesToDelete);
-
-	// Re-enable graph notifications
-	PCGGraph->EnableNotificationsForEditor();
-
-	// Force a graph refresh
-	PCGEditorGraph->ReconstructGraph();
-
-	// Notify the widget
-	GraphEditorWidget->NotifyGraphChanged();	
-}
-
-bool FPCGEditor::CanConvertNamedRerouteToReroute() const
-{
-	if (!GraphEditorWidget)
-	{
-		return false;
-	}
-
-	for (UObject* Object : GraphEditorWidget->GetSelectedNodes())
-	{
-		if (UPCGEditorGraphNodeNamedRerouteDeclaration* DeclarationNode = Cast<UPCGEditorGraphNodeNamedRerouteDeclaration>(Object))
-		{
-			return true;
-		}
-	}
-
-	return false;
-}
-
 void FPCGEditor::OnSelectNamedRerouteUsages()
 {
 	if (!InternalValidationOnAction())
@@ -1503,7 +1367,7 @@ void FPCGEditor::OnSelectNamedRerouteDeclaration()
 	}
 
 	const FGraphPanelSelectionSet SelectedNodes = GraphEditorWidget->GetSelectedNodes();
-	
+
 	if (SelectedNodes.Num() != 1)
 	{
 		return;
@@ -1552,126 +1416,6 @@ bool FPCGEditor::CanSelectNamedRerouteDeclaration() const
 	return false;
 }
 
-void FPCGEditor::OnConvertRerouteToNamedReroute()
-{
-	if (!InternalValidationOnAction())
-	{
-		return;
-	}
-
-	UPCGGraph* PCGGraph = PCGEditorGraph->GetPCGGraph();
-	check(PCGGraph);
-
-	TArray<UPCGNode*> ReroutesToConvert;
-	for (UObject* Object : GraphEditorWidget->GetSelectedNodes())
-	{
-		if(UPCGEditorGraphNodeReroute* RerouteNode = Cast<UPCGEditorGraphNodeReroute>(Object))
-		{
-			// Skip named reroutes
-			if (RerouteNode->IsA<UPCGEditorGraphNodeNamedRerouteBase>())
-			{
-				continue;
-			}
-
-			check(RerouteNode->GetPCGNode());
-			ReroutesToConvert.Add(RerouteNode->GetPCGNode());
-		}
-	}
-
-	if (ReroutesToConvert.IsEmpty())
-	{
-		return;
-	}
-
-	// Disable graph notifications
-	PCGGraph->DisableNotificationsForEditor();
-
-	// For all found reroutes, replace by a named reroute pair and forward inputs to the declaration and outputs to the usage.
-	TArray<UPCGNode*> NodesToDelete;
-
-	for (UPCGNode* Reroute : ReroutesToConvert)
-	{
-		NodesToDelete.Add(Reroute);
-
-		// Create named reroute declaration
-		UPCGSettings* DeclarationSettings = nullptr;
-		UPCGNode* Declaration = PCGGraph->AddNodeOfType(UPCGNamedRerouteDeclarationSettings::StaticClass(), DeclarationSettings);
-
-		// Create named reroute usage
-		UPCGSettings* UsageSettings = nullptr;
-		UPCGNode* Usage = PCGGraph->AddNodeOfType(UPCGNamedRerouteUsageSettings::StaticClass(), UsageSettings);
-		Cast<UPCGNamedRerouteUsageSettings>(UsageSettings)->Declaration = Cast<UPCGNamedRerouteDeclarationSettings>(DeclarationSettings);
-
-		// Setup non-fonctional information
-		Reroute->TransferEditorProperties(Declaration);
-		Reroute->TransferEditorProperties(Usage);
-		constexpr float PositionOffsetIncrementY = 50.f;
-		Usage->PositionY += PositionOffsetIncrementY;
-
-		// Force a graph refresh - needed for the other operations to go through which a bit unfortunate
-		// TODO improve this
-		PCGEditorGraph->ReconstructGraph();
-
-		// Copy all inputs to the reroute to the declaration inputs
-		check(Reroute->GetInputPin(PCGPinConstants::DefaultInputLabel));
-		TArray<UPCGEdge*> RerouteInputEdges = Reroute->GetInputPin(PCGPinConstants::DefaultInputLabel)->Edges;
-		ensure(RerouteInputEdges.Num() <= 1);
-
-		if (RerouteInputEdges.Num() >= 1)
-		{
-			UPCGEdge* InputEdge = RerouteInputEdges[0];
-			PCGGraph->AddEdge(InputEdge->InputPin->Node, InputEdge->InputPin->Properties.Label, Declaration, PCGPinConstants::DefaultInputLabel);
-			UPCGEditorGraphNodeNamedRerouteDeclaration* DeclarationNode = CastChecked<UPCGEditorGraphNodeNamedRerouteDeclaration>(GetEditorNode(Declaration));
-			const FString NodeName = DeclarationNode->GenerateNodeName(InputEdge->InputPin->Node, InputEdge->InputPin->Properties.Label);
-			DeclarationNode->OnRenameNode(NodeName);
-		}
-
-		// Add invisible edge between declaration and usage
-		PCGGraph->AddEdge(Declaration, PCGNamedRerouteConstants::InvisiblePinLabel, Usage, PCGPinConstants::DefaultInputLabel);
-
-		// Copy all outputs from the reroute to the usage outputs
-		check(Reroute->GetOutputPin(PCGPinConstants::DefaultOutputLabel));
-		TArray<UPCGEdge*> OutputEdges = Reroute->GetOutputPin(PCGPinConstants::DefaultOutputLabel)->Edges;
-		for (UPCGEdge* OutputEdge : OutputEdges)
-		{
-			PCGGraph->AddEdge(Usage, PCGPinConstants::DefaultOutputLabel, OutputEdge->OutputPin->Node, OutputEdge->OutputPin->Properties.Label);
-		}
-	}
-
-	// Delete removed ndoes
-	PCGGraph->RemoveNodes(NodesToDelete);
-
-	// Re-enable graph notifications
-	PCGGraph->EnableNotificationsForEditor();
-
-	// Force a graph refresh
-	PCGEditorGraph->ReconstructGraph();
-
-	// Notify the widget
-	GraphEditorWidget->NotifyGraphChanged();
-}
-
-bool FPCGEditor::CanConvertRerouteToNamedReroute() const
-{
-	if (!GraphEditorWidget)
-	{
-		return false;
-	}
-
-	for (UObject* Object : GraphEditorWidget->GetSelectedNodes())
-	{
-		if (UPCGEditorGraphNodeReroute* RerouteNode = Cast<UPCGEditorGraphNodeReroute>(Object))
-		{
-			if (!RerouteNode->IsA<UPCGEditorGraphNodeNamedRerouteBase>())
-			{
-				return true;
-			}
-		}
-	}
-
-	return false;
-}
-
 bool FPCGEditor::IsNodeSourceEditorTabClosed() const
 {
 	return !TabManager.IsValid() || !TabManager->FindExistingLiveTab(FPCGEditor_private::NodeSourceID).IsValid();
diff --git a/Engine/Plugins/PCG/Source/PCGEditor/Private/PCGEditor.h b/Engine/Plugins/PCG/Source/PCGEditor/Private/PCGEditor.h
--- a/Engine/Plugins/PCG/Source/PCGEditor/Private/PCGEditor.h
+++ b/Engine/Plugins/PCG/Source/PCGEditor/Private/PCGEditor.h
@@ -256,11 +256,6 @@ private:
 	/** Whether the user can rename the selected node */
 	bool CanRenameNode() const;
 
-	/** Converts a named reroute to reroute, rewiring its edges to a newly created reroute */
-	void OnConvertNamedRerouteToReroute();
-	/** Whether the user can convert the selection to a reroute */
-	bool CanConvertNamedRerouteToReroute() const;
-
 	/** Selects the associated usages of a given reroute declaration */
 	void OnSelectNamedRerouteUsages();
 	/** Whether the user can find the usages from the selection */
@@ -271,11 +266,6 @@ private:
 	/** Whether the user can find the declaration from the selection */
 	bool CanSelectNamedRerouteDeclaration() const;
 
-	/** Converts a selected reroute to a named reroute pair; the inputs will go to the declaration and the outputs will go to the usage */
-	void OnConvertRerouteToNamedReroute();
-	/** Whether the user can convert to a named reroute from the selection */
-	bool CanConvertRerouteToNamedReroute() const;
-
 	/** Internal method that validates a few things (& logs errors) prior to executing actions. */
 	bool InternalValidationOnAction();
 
diff --git a/Engine/Plugins/PCG/Source/PCGEditor/Private/PCGEditorCommands.cpp b/Engine/Plugins/PCG/Source/PCGEditor/Private/PCGEditorCommands.cpp
--- a/Engine/Plugins/PCG/Source/PCGEditor/Private/PCGEditorCommands.cpp
+++ b/Engine/Plugins/PCG/Source/PCGEditor/Private/PCGEditorCommands.cpp
@@ -36,10 +36,8 @@ void FPCGEditorCommands::RegisterCommands()
 	UI_COMMAND(ToggleInspect, "Toggle Inspection", "Toggle node inspection for selected node", EUserInterfaceActionType::ToggleButton, FInputChord(EKeys::A));
 	UI_COMMAND(AddSourcePin, "Add Source Pin", "Add new source pin to the current node", EUserInterfaceActionType::Button, FInputChord());
 	UI_COMMAND(RenameNode, "Rename Node", "Rename the selected node", EUserInterfaceActionType::Button, FInputChord(EKeys::F2));
-	UI_COMMAND(ConvertNamedRerouteToReroute, "Convert Named Reroute to Reroute", "Convers the selected named reroute declaration to a normal reroute node", EUserInterfaceActionType::Button, FInputChord());
 	UI_COMMAND(SelectNamedRerouteUsages, "Select Named Reroute Usages", "Selects all usages of this Named Reroute Declaration", EUserInterfaceActionType::Button, FInputChord());
 	UI_COMMAND(SelectNamedRerouteDeclaration, "Select Named Reroute Declaration", "Selects the associated Named Reroute Declaration matching this Named Reroute Usage", EUserInterfaceActionType::Button, FInputChord());
-	UI_COMMAND(ConvertRerouteToNamedReroute, "Convert Reroute to Named Reroute", "Convers the selected reroute node to a Named reroute pair (declaration & usage)", EUserInterfaceActionType::Button, FInputChord());
 }
 
 #undef LOCTEXT_NAMESPACE
diff --git a/Engine/Plugins/PCG/Source/PCGEditor/Private/PCGEditorCommands.h b/Engine/Plugins/PCG/Source/PCGEditor/Private/PCGEditorCommands.h
--- a/Engine/Plugins/PCG/Source/PCGEditor/Private/PCGEditorCommands.h
+++ b/Engine/Plugins/PCG/Source/PCGEditor/Private/PCGEditorCommands.h
@@ -32,8 +32,6 @@ public:
 	TSharedPtr<FUICommandInfo> ToggleInspect;
 	TSharedPtr<FUICommandInfo> AddSourcePin;
 	TSharedPtr<FUICommandInfo> RenameNode;
-	TSharedPtr<FUICommandInfo> ConvertNamedRerouteToReroute;
 	TSharedPtr<FUICommandInfo> SelectNamedRerouteUsages;
 	TSharedPtr<FUICommandInfo> SelectNamedRerouteDeclaration;
-	TSharedPtr<FUICommandInfo> ConvertRerouteToNamedReroute;
 };
diff --git a/Engine/Plugins/PCG/Source/PCGEditor/Private/PCGEditorGraphNodeBase.cpp b/Engine/Plugins/PCG/Source/PCGEditor/Private/PCGEditorGraphNodeBase.cpp
--- a/Engine/Plugins/PCG/Source/PCGEditor/Private/PCGEditorGraphNodeBase.cpp
+++ b/Engine/Plugins/PCG/Source/PCGEditor/Private/PCGEditorGraphNodeBase.cpp
@@ -220,17 +220,12 @@ void UPCGEditorGraphNodeBase::GetNodeContextMenuActions(UToolMenu* Menu, class U
 		{
 			if (Settings->IsA<UPCGNamedRerouteDeclarationSettings>())
 			{
-				Section.AddMenuEntry(FPCGEditorCommands::Get().ConvertNamedRerouteToReroute);
 				Section.AddMenuEntry(FPCGEditorCommands::Get().SelectNamedRerouteUsages);
 			}
 			else if (Settings->IsA<UPCGNamedRerouteUsageSettings>())
 			{
 				Section.AddMenuEntry(FPCGEditorCommands::Get().SelectNamedRerouteDeclaration);
 			}
-			else if (Settings->IsA<UPCGRerouteSettings>())
-			{
-				Section.AddMenuEntry(FPCGEditorCommands::Get().ConvertRerouteToNamedReroute);
-			}
 
 			// Operation to convert a node in place into something else, such as a different node.
 			const TArray<FPCGPreconfiguredInfo> ConversionInfo = Settings->GetConversionInfo();
diff --git a/Engine/Plugins/PCG/Source/PCGEditor/Private/PCGEditorGraphNodeReroute.cpp b/Engine/Plugins/PCG/Source/PCGEditor/Private/PCGEditorGraphNodeReroute.cpp
--- a/Engine/Plugins/PCG/Source/PCGEditor/Private/PCGEditorGraphNodeReroute.cpp
+++ b/Engine/Plugins/PCG/Source/PCGEditor/Private/PCGEditorGraphNodeReroute.cpp
@@ -191,19 +191,23 @@ FString UPCGEditorGraphNodeNamedRerouteDeclaration::GenerateNodeName(const UPCGN
 
 void UPCGEditorGraphNodeNamedRerouteDeclaration::OnRenameNode(const FString& NewName)
 {
-	const FString ModifiedName = GetCollisionFreeNodeName(NewName);
-	if (CanRenameNode(ModifiedName))
+	check(PCGNode);
+	if (TOptional<FName> ModifiedName = UPCGRerouteSettings::GetCollisionFreeNodeName(PCGNode->GetGraph(), FName(NewName)))
 	{
-		CommitNodeRename(ModifiedName);
-		// Propagate the name change to downstream usage nodes
-		ApplyToUsageNodes([&ModifiedName](UPCGEditorGraphNodeNamedRerouteUsage* RerouteNode)
+		FString NameString = ModifiedName->ToString();
+		if (CanRenameNode(NameString))
 		{
-			RerouteNode->bCanRenameNode = true;
-			RerouteNode->OnRenameNode(ModifiedName);
-			RerouteNode->bCanRenameNode = false;
-		});
+			CommitNodeRename(NameString);
+			// Propagate the name change to downstream usage nodes
+			ApplyToUsageNodes([&NameString](UPCGEditorGraphNodeNamedRerouteUsage* RerouteNode)
+			{
+				RerouteNode->bCanRenameNode = true;
+				RerouteNode->OnRenameNode(NameString);
+				RerouteNode->bCanRenameNode = false;
+			});
 
-		ReconstructNodeOnChange();
+			ReconstructNodeOnChange();
+		}
 	}
 }
 
@@ -231,41 +235,4 @@ void UPCGEditorGraphNodeNamedRerouteDeclaration::FixNodeNameCollision()
 	bCanRenameNode = true;
 	OnRenameNode(BaseName);
 	bCanRenameNode = false;
-}
-
-FString UPCGEditorGraphNodeNamedRerouteDeclaration::GetCollisionFreeNodeName(const FString& BaseName) const
-{
-	// If there is another declaration in the graph with the same name, append a _N to the name here.
-	UPCGEditorGraph* EditorGraph = Cast<UPCGEditorGraph>(GetGraph());
-
-	if (!EditorGraph)
-	{
-		return BaseName;
-	}
-
-	FString TentativeName = BaseName;
-	bool bHasNameCollision = true;
-	int32 NameSuffix = 1;
-
-	while (bHasNameCollision)
-	{
-		bHasNameCollision = false;
-
-		for (const TObjectPtr<UEdGraphNode>& EdGraphNode : EditorGraph->Nodes)
-		{
-			if (EdGraphNode && EdGraphNode != this && EdGraphNode->GetNodeTitle(ENodeTitleType::ListView).ToString() == TentativeName)
-			{
-				bHasNameCollision = true;
-				break;
-			}
-		}
-
-		if (bHasNameCollision)
-		{
-			// Implementation note: since the node title is a display string, we need to craft the tentative name to be "stable" against it, otherwise the comparisons will fail.
-			TentativeName = FString::Printf(TEXT("%s %d"), *BaseName, NameSuffix++);
-		}
-	}
-
-	return TentativeName;
 }
\ No newline at end of file
diff --git a/Engine/Plugins/PCG/Source/PCGEditor/Private/PCGEditorGraphNodeReroute.h b/Engine/Plugins/PCG/Source/PCGEditor/Private/PCGEditorGraphNodeReroute.h
--- a/Engine/Plugins/PCG/Source/PCGEditor/Private/PCGEditorGraphNodeReroute.h
+++ b/Engine/Plugins/PCG/Source/PCGEditor/Private/PCGEditorGraphNodeReroute.h
@@ -62,7 +62,6 @@ public:
 
 	FString GenerateNodeName(const UPCGNode* FromNode, FName FromPinName);
 	void FixNodeNameCollision();
-	FString GetCollisionFreeNodeName(const FString& BaseName) const;
 
 protected:
 	virtual FText GetPinFriendlyName(const UPCGPin* InPin) const override;
