[PCG] Emit warning from grass map generation node when generating grass maps outside of the editor world and runtime grass generation is not supported.

#rnx

diff --git a/Engine/Plugins/PCG/Source/PCG/Private/Elements/PCGGenerateGrassMaps.cpp b/Engine/Plugins/PCG/Source/PCG/Private/Elements/PCGGenerateGrassMaps.cpp
--- a/Engine/Plugins/PCG/Source/PCG/Private/Elements/PCGGenerateGrassMaps.cpp
+++ b/Engine/Plugins/PCG/Source/PCG/Private/Elements/PCGGenerateGrassMaps.cpp
@@ -134,6 +134,12 @@ bool FPCGGenerateGrassMapsElement::ExecuteInternal(FPCGContext* InContext) const
 		const UPCGComponent* SourceComponent = Context->SourceComponent.IsValid() ? Context->SourceComponent.Get() : nullptr;
 		check(SourceComponent);
 
+		if (!UE::Landscape::IsRuntimeGrassMapGenerationSupported() && (!SourceComponent->GetWorld() || SourceComponent->GetWorld()->IsGameWorld()))
+		{
+			UE_LOG(LogPCG, Warning, TEXT("Grass map generation is disabled outside of editor worlds. Try enabling the CVar 'grass.GrassMap.AlwaysBuildRuntimeGenerationResources'."));
+			return true;
+		}
+
 		const FBox ComponentBounds = SourceComponent->GetGridBounds();
 
 		for (TSoftObjectPtr<ALandscapeProxy> LandscapeProxyPtr : LandscapeData->Landscapes)
diff --git a/Engine/Source/Runtime/Landscape/Internal/LandscapeGrassWeightExporter.h b/Engine/Source/Runtime/Landscape/Internal/LandscapeGrassWeightExporter.h
--- a/Engine/Source/Runtime/Landscape/Internal/LandscapeGrassWeightExporter.h
+++ b/Engine/Source/Runtime/Landscape/Internal/LandscapeGrassWeightExporter.h
@@ -16,6 +16,7 @@ using FRDGTextureRef = FRDGTexture*;
 namespace UE::Landscape
 {
 	LANDSCAPE_API bool CanRenderGrassMap(ULandscapeComponent* Component);
+	LANDSCAPE_API bool IsRuntimeGrassMapGenerationSupported();
 }
 
 // Hacky base class to avoid 8 bytes of padding after the vtable
diff --git a/Engine/Source/Runtime/Landscape/Private/LandscapeGrassMapsBuilder.cpp b/Engine/Source/Runtime/Landscape/Private/LandscapeGrassMapsBuilder.cpp
--- a/Engine/Source/Runtime/Landscape/Private/LandscapeGrassMapsBuilder.cpp
+++ b/Engine/Source/Runtime/Landscape/Private/LandscapeGrassMapsBuilder.cpp
@@ -247,6 +247,11 @@ namespace UE::Landscape
 		return true;
 	}
 
+	bool IsRuntimeGrassMapGenerationSupported()
+	{
+		return GGrassMapUseRuntimeGeneration || GGrassMapAlwaysBuildRuntimeGenerationResources;
+	}
+
 	// calculates the minimum distance between any cameras and the specified Worldbounds.
 	static inline double CalculateMinDistanceToCamerasSquared(const TArray<FVector>& Cameras, const FBoxSphereBounds& WorldBounds)
 	{
