[PCG] Struct of Arrays ground work: Move to an Indexed Octree to support different Point data structures
- Add UPCGBasePointData abstract class (this will be the base class to support Struct of Arrays point type)
- Move PCG codebase to use new Octree
- Deprecation of previous Octree type
- A follow up CL will add new value range which will replace the slow virtual indexed getters introduced in this CL

#rnx

diff --git a/Engine/Plugins/PCG/Source/PCG/Private/Data/PCGBasePointData.cpp b/Engine/Plugins/PCG/Source/PCG/Private/Data/PCGBasePointData.cpp
new file mode 100644
--- /dev/null
+++ b/Engine/Plugins/PCG/Source/PCG/Private/Data/PCGBasePointData.cpp
@@ -0,0 +1,4 @@
+// Copyright Epic Games, Inc. All Rights Reserved.
+
+#include "Data/PCGBasePointData.h"
+
diff --git a/Engine/Plugins/PCG/Source/PCG/Private/Data/PCGCollisionWrapperData.cpp b/Engine/Plugins/PCG/Source/PCG/Private/Data/PCGCollisionWrapperData.cpp
--- a/Engine/Plugins/PCG/Source/PCG/Private/Data/PCGCollisionWrapperData.cpp
+++ b/Engine/Plugins/PCG/Source/PCG/Private/Data/PCGCollisionWrapperData.cpp
@@ -320,28 +320,30 @@ bool UPCGCollisionWrapperData::SamplePoint(const FTransform& InTransform, const
 	// For all touching points, check against the actual physics shape we have for that point
 	// E.g. get the point collision attribute value, match to our internal list, get the shape (might require loading), then test against the matching shapes
 	// note that we need to place the stuff in the right referential too, since our shapes will all be at the origin
-	const UPCGPointData::PointOctree& Octree = GetPointData()->GetOctree();
+	const PCGPointOctree::FPointOctree& Octree = GetPointData()->GetPointOctree();
 
 	float Density = 0;
 	FBox TransformedBounds = InBounds.TransformBy(InTransform);
 
-	Octree.FindElementsWithBoundsTest(FBoxCenterAndExtent(TransformedBounds.GetCenter(), TransformedBounds.GetExtent()), [this, &InTransform, &InBounds, &Density](const FPCGPointRef& InPointRef)
+	Octree.FindElementsWithBoundsTest(FBoxCenterAndExtent(TransformedBounds.GetCenter(), TransformedBounds.GetExtent()), [this, &InTransform, &InBounds, &Density](const PCGPointOctree::FPointRef& InPointRef)
 	{
-		if (InPointRef.Point->Density < Density)
+		const float PointDensity = GetPointData()->GetDensity(InPointRef.Index);
+
+		if (PointDensity < Density)
 		{
 			return;
 		}
 
 		bool bHasOverlap = false;
 
-		FTransform RelativeTransform = InTransform.GetRelativeTransform(InPointRef.Point->Transform);
+		FTransform RelativeTransform = InTransform.GetRelativeTransform(GetPointData()->GetTransform(InPointRef.Index));
 		FCollisionShape CollisionShape;
 		CollisionShape.SetBox(FVector3f(InBounds.GetExtent() * RelativeTransform.GetScale3D()));
 		FPhysicsShapeAdapter CollAdapter(RelativeTransform.GetRotation(), CollisionShape);
 		const FPhysicsGeometry& Geom = CollAdapter.GetGeometry();
 		FTransform GeomTransform = CollAdapter.GetGeomPose(RelativeTransform.GetLocation());
 
-		const int32 EntryIndex = InPointRef.Point - GetPointData()->GetPoints().GetData();
+		const int32 EntryIndex = InPointRef.Index;
 		const PhysicsInterfaceTypes::FInlineShapeArray& Shapes = GetCachedShapes(EntryIndex);
 
 		if (!Shapes.IsEmpty())
@@ -366,7 +368,7 @@ bool UPCGCollisionWrapperData::SamplePoint(const FTransform& InTransform, const
 
 		if (bHasOverlap)
 		{
-			Density = InPointRef.Point->Density;
+			Density = PointDensity;
 		}
 	});
 
diff --git a/Engine/Plugins/PCG/Source/PCG/Private/Data/PCGPointData.cpp b/Engine/Plugins/PCG/Source/PCG/Private/Data/PCGPointData.cpp
--- a/Engine/Plugins/PCG/Source/PCG/Private/Data/PCGPointData.cpp
+++ b/Engine/Plugins/PCG/Source/PCG/Private/Data/PCGPointData.cpp
@@ -9,6 +9,7 @@
 #include "Metadata/Accessors/PCGAttributeAccessorHelpers.h"
 #include "Metadata/Accessors/PCGAttributeAccessorKeys.h"
 
+#include "Algo/Transform.h"
 #include "GameFramework/Actor.h"
 #include "HAL/IConsoleManager.h"
 #include "Serialization/ArchiveCrc32.h"
@@ -253,24 +254,37 @@ void UPCGPointData::GetResourceSizeEx(FResourceSizeEx& CumulativeResourceSize)
 {
 	Super::GetResourceSizeEx(CumulativeResourceSize);
 
-	CumulativeResourceSize.AddDedicatedSystemMemoryBytes(Points.GetAllocatedSize() + Octree.GetSizeBytes() + sizeof(Bounds));
+	CumulativeResourceSize.AddDedicatedSystemMemoryBytes(Points.GetAllocatedSize() + Octree.GetSizeBytes() + PCGPointOctree.GetSizeBytes() + sizeof(Bounds));
 }
 
 TArray<FPCGPoint>& UPCGPointData::GetMutablePoints()
 {
 	bOctreeIsDirty = true;
+	bOctreeOldIsDirty = true;
 	bBoundsAreDirty = true;
 	return Points;
 }
 
+PRAGMA_DISABLE_DEPRECATION_WARNINGS
 const UPCGPointData::PointOctree& UPCGPointData::GetOctree() const
+{
+	if (bOctreeOldIsDirty)
+	{
+		RebuildOctreeOld();
+	}
+
+	return Octree;
+}
+PRAGMA_ENABLE_DEPRECATION_WARNINGS
+
+const PCGPointOctree::FPointOctree& UPCGPointData::GetPointOctree() const
 {
 	if (bOctreeIsDirty)
 	{
 		RebuildOctree();
 	}
 
-	return Octree;
+	return PCGPointOctree;
 }
 
 void UPCGPointData::AddToCrc(FArchiveCrc32& Ar, bool bFullDataCrc) const
@@ -369,6 +383,94 @@ FPCGAttributeAccessorMethods UPCGPointData::GetPointAccessorMethods()
 	return Methods;
 }
 
+TArray<FTransform> UPCGPointData::GetTransformsCopy() const
+{
+	TArray<FTransform> Transforms;
+	Transforms.Reserve(Points.Num());
+	Algo::Transform(Points, Transforms, [](const FPCGPoint& Point) { return Point.Transform; });
+	return Transforms;
+}
+
+const FTransform& UPCGPointData::GetTransform(int32 InPointIndex) const
+{
+	return Points[InPointIndex].Transform;
+}
+
+float UPCGPointData::GetDensity(int32 InPointIndex) const
+{
+	return Points[InPointIndex].Density;
+}
+
+const FVector& UPCGPointData::GetBoundsMin(int32 InPointIndex) const
+{
+	return Points[InPointIndex].BoundsMin;
+}
+
+const FVector& UPCGPointData::GetBoundsMax(int32 InPointIndex) const
+{
+	return Points[InPointIndex].BoundsMax;
+}
+
+const FVector4& UPCGPointData::GetColor(int32 InPointIndex) const
+{
+	return Points[InPointIndex].Color;
+}
+
+float UPCGPointData::GetSteepness(int32 InPointIndex) const
+{
+	return Points[InPointIndex].Steepness;
+}
+
+int32 UPCGPointData::GetSeed(int32 InPointIndex) const
+{
+	return Points[InPointIndex].Seed;
+}
+
+int64 UPCGPointData::GetMetadataEntry(int32 InPointIndex) const
+{
+	return Points[InPointIndex].MetadataEntry;
+}
+
+FBoxSphereBounds UPCGPointData::GetDensityBounds(int32 InPointIndex) const
+{
+	return Points[InPointIndex].GetDensityBounds();
+}
+
+FBox UPCGPointData::GetLocalDensityBounds(int32 InPointIndex) const
+{
+	return Points[InPointIndex].GetLocalDensityBounds();
+}
+
+FBox UPCGPointData::GetLocalBounds(int32 InPointIndex) const
+{
+	return Points[InPointIndex].GetLocalBounds();
+}
+
+FVector UPCGPointData::GetLocalCenter(int32 InPointIndex) const
+{
+	return Points[InPointIndex].GetLocalCenter();
+}
+
+FVector UPCGPointData::GetExtents(int32 InPointIndex) const
+{
+	return Points[InPointIndex].GetExtents();
+}
+
+FVector UPCGPointData::GetScaledExtents(int32 InPointIndex) const
+{
+	return Points[InPointIndex].GetScaledExtents();
+}
+
+FVector UPCGPointData::GetLocalSize(int32 InPointIndex) const
+{
+	return Points[InPointIndex].GetLocalSize();
+}
+
+FVector UPCGPointData::GetScaledLocalSize(int32 InPointIndex) const
+{
+	return Points[InPointIndex].GetScaledLocalSize();
+}
+
 FBox UPCGPointData::GetBounds() const
 {
 	if (bBoundsAreDirty)
@@ -413,6 +515,7 @@ void UPCGPointData::CopyPointsFrom(const UPCGPointData* InData, const TArray<int
 
 	bBoundsAreDirty = true;
 	bOctreeIsDirty = true;
+	bOctreeOldIsDirty = true;
 }
 
 void UPCGPointData::SetPoints(const TArray<FPCGPoint>& InPoints)
@@ -506,15 +609,15 @@ bool UPCGPointData::ProjectPoint(const FTransform& InTransform, const FBox& InBo
 		RebuildOctree();
 	}
 
-	TArray<TPair<const FPCGPoint*, FVector::FReal>, TInlineAllocator<4>> Contributions;
+	TArray<TPair<int32, FVector::FReal>, TInlineAllocator<4>> Contributions;
 	const bool bSampleInVolume = (InBounds.GetExtent() != FVector::ZeroVector);
 
 	if (!bSampleInVolume)
 	{
 		const FVector InPosition = InTransform.GetLocation();
-		Octree.FindElementsWithBoundsTest(FBoxCenterAndExtent(InPosition, FVector::Zero()), [&InPosition, &Contributions](const FPCGPointRef& InPointRef) 
+		PCGPointOctree.FindElementsWithBoundsTest(FBoxCenterAndExtent(InPosition, FVector::Zero()), [this, &InPosition, &Contributions](const PCGPointOctree::FPointRef& InPointRef) 
 		{
-			Contributions.Emplace(InPointRef.Point, PCGPointHelpers::InverseEuclidianDistance(*InPointRef.Point, InPosition));
+			Contributions.Emplace(InPointRef.Index, PCGPointHelpers::InverseEuclidianDistance(Points[InPointRef.Index], InPosition));
 		});
 	}
 	else
@@ -522,21 +625,21 @@ bool UPCGPointData::ProjectPoint(const FTransform& InTransform, const FBox& InBo
 		FBox TransformedBounds = InBounds.TransformBy(InTransform);
 		FMatrix InTransformInverseMatrix = InTransform.ToMatrixWithScale().Inverse();
 
-		Octree.FindElementsWithBoundsTest(FBoxCenterAndExtent(TransformedBounds.GetCenter(), TransformedBounds.GetExtent()), [bUseBounds, &InBounds, &InTransformInverseMatrix, &Contributions](const FPCGPointRef& InPointRef) 
+		PCGPointOctree.FindElementsWithBoundsTest(FBoxCenterAndExtent(TransformedBounds.GetCenter(), TransformedBounds.GetExtent()), [this, bUseBounds, &InBounds, &InTransformInverseMatrix, &Contributions](const PCGPointOctree::FPointRef& InPointRef)
 		{
-			const FVector::FReal Contribution = bUseBounds ? PCGPointHelpers::VolumeOverlap(*InPointRef.Point, InBounds, InTransformInverseMatrix) : 1.0;
+			const FVector::FReal Contribution = bUseBounds ? PCGPointHelpers::VolumeOverlap(Points[InPointRef.Index], InBounds, InTransformInverseMatrix) : 1.0;
 			if (Contribution > 0)
 			{
-				Contributions.Emplace(InPointRef.Point, Contribution);
+				Contributions.Emplace(InPointRef.Index, Contribution);
 			}
 		});
 	}
 
 	FVector::FReal SumContributions = 0;
 	FVector::FReal MaxContribution = 0;
-	const FPCGPoint* MaxContributor = nullptr;
+	int32 MaxContributor = INDEX_NONE;
 
-	for (const TPair<const FPCGPoint*, FVector::FReal>& Contribution : Contributions)
+	for (const TPair<int32, FVector::FReal>& Contribution : Contributions)
 	{
 		SumContributions += Contribution.Value;
 
@@ -573,9 +676,9 @@ bool UPCGPointData::ProjectPoint(const FTransform& InTransform, const FBox& InBo
 	FVector4 WeightedColor = FVector4::Zero();
 	float WeightedSteepness = 0;
 
-	for (const TPair<const FPCGPoint*, FVector::FReal> Contribution : Contributions)
+	for (const TPair<int32, FVector::FReal>& Contribution : Contributions)
 	{
-		const FPCGPoint& SourcePoint = *Contribution.Key;
+		const FPCGPoint& SourcePoint = Points[Contribution.Key];
 		const FVector::FReal Weight = Contribution.Value / SumContributions;
 
 		WeightedPosition += SourcePoint.Transform.GetLocation() * Weight;
@@ -596,7 +699,7 @@ bool UPCGPointData::ProjectPoint(const FTransform& InTransform, const FBox& InBo
 		WeightedColor += SourcePoint.Color * Weight;
 		WeightedSteepness += SourcePoint.Steepness * Weight;
 
-		ContributionsForMetadata.Emplace(Contribution.Key, static_cast<float>(Weight));
+		ContributionsForMetadata.Emplace(&Points[Contribution.Key], static_cast<float>(Weight));
 	}
 
 	// Finally, apply changes to point, based on the projection settings
@@ -638,7 +741,7 @@ bool UPCGPointData::ProjectPoint(const FTransform& InTransform, const FBox& InBo
 	{
 		//TRACE_CPUPROFILER_EVENT_SCOPE(UPCGPointData::SamplePoint::SetupMetadata);
 		// Initialise metadata entry for this temporary point
-		UPCGMetadataAccessorHelpers::InitializeMetadataWithParent(OutPoint, OutMetadata, *MaxContributor, Metadata);
+		UPCGMetadataAccessorHelpers::InitializeMetadataWithParent(OutPoint, OutMetadata, Points[MaxContributor], Metadata);
 
 		if (ContributionsForMetadata.Num() > 1)
 		{
@@ -649,17 +752,18 @@ bool UPCGPointData::ProjectPoint(const FTransform& InTransform, const FBox& InBo
 	return true;
 }
 
-void UPCGPointData::RebuildOctree() const
+void UPCGPointData::RebuildOctreeOld() const
 {
+PRAGMA_DISABLE_DEPRECATION_WARNINGS
 	FScopeLock Lock(&CachedDataLock);
 
-	if (!bOctreeIsDirty)
+	if (!bOctreeOldIsDirty)
 	{
 		return;
 	}
 
 	TRACE_CPUPROFILER_EVENT_SCOPE(UPCGPointData::RebuildOctree)
-	check(bOctreeIsDirty);
+	check(bOctreeOldIsDirty);
 
 	FBox PointBounds = GetBounds();
 	TOctree2<FPCGPointRef, FPCGPointRefSemantics> NewOctree(PointBounds.GetCenter(), PointBounds.GetExtent().Length());
@@ -670,6 +774,31 @@ void UPCGPointData::RebuildOctree() const
 	}
 
 	Octree = MoveTemp(NewOctree);
+	bOctreeOldIsDirty = false;
+PRAGMA_ENABLE_DEPRECATION_WARNINGS
+}
+
+void UPCGPointData::RebuildOctree() const
+{
+	FScopeLock Lock(&CachedDataLock);
+
+	if (!bOctreeIsDirty)
+	{
+		return;
+	}
+
+	TRACE_CPUPROFILER_EVENT_SCOPE(UPCGPointData::RebuildOctree)
+	check(bOctreeIsDirty);
+
+	FBox PointBounds = GetBounds();
+	PCGPointOctree::FPointOctree NewOctree(PointBounds.GetCenter(), PointBounds.GetExtent().Length());
+
+	for (int32 Index = 0; Index < Points.Num(); ++Index)
+	{
+		NewOctree.AddElement(PCGPointOctree::FPointRef(Index, Points[Index].GetDensityBounds()));
+	}
+
+	PCGPointOctree = MoveTemp(NewOctree);
 	bOctreeIsDirty = false;
 }
 
diff --git a/Engine/Plugins/PCG/Source/PCG/Private/Elements/PCGAttractElement.cpp b/Engine/Plugins/PCG/Source/PCG/Private/Elements/PCGAttractElement.cpp
--- a/Engine/Plugins/PCG/Source/PCG/Private/Elements/PCGAttractElement.cpp
+++ b/Engine/Plugins/PCG/Source/PCG/Private/Elements/PCGAttractElement.cpp
@@ -127,9 +127,8 @@ namespace PCGAttractElement
 						OutAttractData.TargetPointData,
 						SourcePoints[Index].Transform.GetLocation(),
 						AttractSettings->Distance,
-						[Index, &OutAttractData, &MinDistanceSquared, &ClosestTargetIndex, &bFoundColocatedPoint](const FPCGPointRef& PointRef, const double DistanceSquared)
+						[Index, &OutAttractData, &MinDistanceSquared, &ClosestTargetIndex, &bFoundColocatedPoint](const UPCGBasePointData* PointData, int32 TargetIndex, const double DistanceSquared)
 						{
-							const int32 TargetIndex = PointRef.Point - OutAttractData.TargetPointData->GetPoints().GetData();
 							// Ignore self-selection
 							if (TargetIndex == Index)
 							{
@@ -183,10 +182,8 @@ namespace PCGAttractElement
 					OutAttractData.TargetPointData,
 					SourcePoints[Index].Transform.GetLocation(),
 					AttractSettings->Distance,
-					[Index, &OutAttractData](const FPCGPointRef& PointRef, const double DistanceSquared)
+					[Index, &OutAttractData](const UPCGBasePointData* PointData, int32 TargetIndex, const double DistanceSquared)
 					{
-						const int32 TargetIndex = PointRef.Point - OutAttractData.TargetPointData->GetPoints().GetData();
-
 						if (OutAttractData.Mapping[Index] == INDEX_NONE)
 						{
 							OutAttractData.Mapping[Index] = TargetIndex;
diff --git a/Engine/Plugins/PCG/Source/PCG/Private/Elements/PCGBlurElement.cpp b/Engine/Plugins/PCG/Source/PCG/Private/Elements/PCGBlurElement.cpp
--- a/Engine/Plugins/PCG/Source/PCG/Private/Elements/PCGBlurElement.cpp
+++ b/Engine/Plugins/PCG/Source/PCG/Private/Elements/PCGBlurElement.cpp
@@ -194,11 +194,11 @@ bool FPCGBlurElement::ExecuteInternal(FPCGContext* InContext) const
 				{
 					const FPCGPoint& Point = IterState.InputPointData->GetPoints()[ReadIndex];
 					const FVector PointLocation = Point.Transform.GetLocation();
-					TArray<TTuple<const FPCGPoint*, double>> NeighborsAndDistance;
+					TArray<TTuple<int32, double>> NeighborsAndDistance;
 
-					UPCGOctreeQueries::ForEachPointInsideSphere(IterState.InputPointData, PointLocation, Settings->SearchDistance, [&NeighborsAndDistance](const FPCGPoint& Neighbor, double SquaredDistance)
+					UPCGOctreeQueries::ForEachPointInsideSphere(IterState.InputPointData, PointLocation, Settings->SearchDistance, [&NeighborsAndDistance](const UPCGBasePointData* PointData, int32 Neighbor, double SquaredDistance)
 					{
-						NeighborsAndDistance.Emplace(&Neighbor, SquaredDistance);
+						NeighborsAndDistance.Emplace(Neighbor, SquaredDistance);
 					});
 
 					if (!ensure(!NeighborsAndDistance.IsEmpty()))
@@ -206,10 +206,9 @@ bool FPCGBlurElement::ExecuteInternal(FPCGContext* InContext) const
 						return true;
 					}
 
-					auto GetNeighborValue = [TypedReadBuffer, &IterState](const FPCGPoint* Neighbor) -> AttributeType
+					auto GetNeighborValue = [TypedReadBuffer, &IterState](int32 NeighborIndex) -> AttributeType
 					{
 						// Pointer arithmetic to know the neighbor index.
-						const int32 NeighborIndex = static_cast<int32>(Neighbor - IterState.InputPointData->GetPoints().GetData());
 						return TypedReadBuffer[NeighborIndex];
 					};
 
@@ -228,7 +227,7 @@ bool FPCGBlurElement::ExecuteInternal(FPCGContext* InContext) const
 							switch (Settings->BlurMode)
 							{
 							case EPCGBlurElementMode::Linear:
-								Weight = 1.0 - FVector::Distance(PointLocation, Neighbor->Transform.GetLocation()) / Settings->SearchDistance;
+								Weight = 1.0 - FVector::Distance(PointLocation, IterState.InputPointData->GetTransform(Neighbor).GetLocation()) / Settings->SearchDistance;
 								break;
 							case EPCGBlurElementMode::Gaussian:
 								Weight = GaussianRatio * FMath::Exp(-SquaredDistance * GaussianExpRatio);
diff --git a/Engine/Plugins/PCG/Source/PCG/Private/Elements/PCGCollapsePoints.cpp b/Engine/Plugins/PCG/Source/PCG/Private/Elements/PCGCollapsePoints.cpp
--- a/Engine/Plugins/PCG/Source/PCG/Private/Elements/PCGCollapsePoints.cpp
+++ b/Engine/Plugins/PCG/Source/PCG/Private/Elements/PCGCollapsePoints.cpp
@@ -35,9 +35,9 @@ namespace PCGCollapsePoints
 				const double WeightSum = PrimaryWeight + SecondaryWeight;
 
 				double Alpha = FMath::IsNearlyZero(WeightSum) ? 0.5 : (SecondaryWeight / WeightSum);
-				const FVector DeltaPosition = Settings.GetPositionFunc(OutState.Points[SecondaryPointIndex]) - Settings.GetPositionFunc(OutState.Points[PrimaryPointIndex]);
+				const FVector DeltaPosition = Settings.GetPointPositionFunc(OutState.PointTransforms, OutState.SourceData, SecondaryPointIndex) - Settings.GetPointPositionFunc(OutState.PointTransforms, OutState.SourceData, PrimaryPointIndex);
 
-				OutState.Points[PrimaryPointIndex].Transform.AddToTranslation(Alpha * DeltaPosition);
+				OutState.PointTransforms[PrimaryPointIndex].AddToTranslation(Alpha * DeltaPosition);
 				OutState.Weights[PrimaryPointIndex] = WeightSum;
 				OutState.Merged[SecondaryPointIndex] = PrimaryPointIndex;
 			}
@@ -46,18 +46,18 @@ namespace PCGCollapsePoints
 		void RebuildOctree(FPCGContext* InContext, const FCollapsePointsSettings& Settings, FCollapsePointsState& OutState)
 		{
 			TRACE_CPUPROFILER_EVENT_SCOPE(PCGCollapsePointsElement::Algo::RebuildOctree);
-			check(OutState.SourceData && OutState.Merged.Num() == OutState.Points.Num());
+			check(OutState.SourceData && OutState.Merged.Num() == OutState.SourceData->GetNumPoints());
 
-			UPCGPointData::PointOctree NewOctree(OutState.SourceData->GetBounds().GetCenter(), OutState.SourceData->GetBounds().GetExtent().Length());
-			for (int PointIndex = 0; PointIndex < OutState.Points.Num(); ++PointIndex)
+			PCGPointOctree::FPointOctree NewOctree(OutState.SourceData->GetBounds().GetCenter(), OutState.SourceData->GetBounds().GetExtent().Length());
+			for (int PointIndex = 0; PointIndex < OutState.SourceData->GetNumPoints(); ++PointIndex)
 			{
 				if (OutState.Merged[PointIndex] == INDEX_NONE)
 				{
-					NewOctree.AddElement(Settings.GetPointRefFunc(OutState.Points[PointIndex]));
+					NewOctree.AddElement(Settings.GetPointReferenceFunc(OutState.PointTransforms, OutState.SourceData, PointIndex));
 				}
 			}
 
-			OutState.Octree = MoveTemp(NewOctree);
+			OutState.PointOctree = MoveTemp(NewOctree);
 		}
 	}
 
@@ -95,17 +95,17 @@ namespace PCGCollapsePoints
 				bool bHasColocatedPoint = false;
 
 				const double Extents = UE_DOUBLE_SQRT_2 * Settings.DistanceThreshold;
-				FBoxCenterAndExtent SearchBounds = Settings.GetSearchBoundsFunc(OutState.Points[PointIndex], Extents);
-				OutState.Octree.FindElementsWithBoundsTest(SearchBounds, [&Settings, PointIndex, &MinSqrDistance, &ClosestUnvisitedIndex, &bHasColocatedPoint, &OutState](const FPCGPointRef& PointRef)
+				FBoxCenterAndExtent SearchBounds = Settings.GetPointSearchBoundsFunc(OutState.PointTransforms, OutState.SourceData, PointIndex, Extents);
+				OutState.PointOctree.FindElementsWithBoundsTest(SearchBounds, [&Settings, PointIndex, &MinSqrDistance, &ClosestUnvisitedIndex, &bHasColocatedPoint, &OutState](const PCGPointOctree::FPointRef& PointRef)
 				{
-					int NeighborIndex = PointRef.Point - OutState.Points.GetData();
+					int NeighborIndex = PointRef.Index;
 
 					if (OutState.Merged[NeighborIndex] != INDEX_NONE || OutState.Visited[NeighborIndex])
 					{
 						return;
 					}
 
-					double SqrDistance = (Settings.GetPositionFunc(OutState.Points[PointIndex]) - Settings.GetPositionFunc(OutState.Points[NeighborIndex])).SquaredLength();
+					double SqrDistance = (Settings.GetPointPositionFunc(OutState.PointTransforms, OutState.SourceData, PointIndex) - Settings.GetPointPositionFunc(OutState.PointTransforms, OutState.SourceData, NeighborIndex)).SquaredLength();
 
 					if (FMath::IsNearlyZero(SqrDistance))
 					{
@@ -154,34 +154,39 @@ namespace PCGCollapsePoints
 
 	namespace ComparisonModes
 	{
-		FVector GetPosition(const FPCGPoint& InPoint)
+		FVector GetPosition(const TArray<FTransform>& InTransforms, const UPCGPointData* InPointData, int32 InIndex)
 		{
-			return InPoint.Transform.GetLocation();
+			return InTransforms[InIndex].GetLocation();
 		}
 
-		FVector GetCenter(const FPCGPoint& InPoint)
+		FVector GetCenter(const TArray<FTransform>& InTransforms, const UPCGPointData* InPointData, int32 InIndex)
 		{
-			return InPoint.Transform.TransformPosition(InPoint.GetLocalCenter());
+			return InTransforms[InIndex].TransformPosition(InPointData->GetLocalCenter(InIndex));
 		}
 
-		FPCGPointRef GetPositionPointRef(const FPCGPoint& InPoint)
+		PCGPointOctree::FPointRef GetPositionPointRef(const TArray<FTransform>& InTransforms, const UPCGPointData* InPointData, int32 InIndex)
 		{
-			return FPCGPointRef(InPoint, FBox(FVector::ZeroVector, FVector::ZeroVector));
+			check(InPointData);
+			return PCGPointOctree::FPointRef(InIndex, FBox(FVector::ZeroVector, FVector::ZeroVector).TransformBy(InTransforms[InIndex]));
 		}
 
-		FPCGPointRef GetCenterPointRef(const FPCGPoint& InPoint)
+		PCGPointOctree::FPointRef GetCenterPointRef(const TArray<FTransform>& InTransforms, const UPCGPointData* InPointData, int32 InIndex)
 		{
-			return FPCGPointRef(InPoint, FBox(InPoint.GetLocalCenter(), InPoint.GetLocalCenter()));
+			check(InPointData);
+			const FVector LocalCenter = InPointData->GetLocalCenter(InIndex);
+			return PCGPointOctree::FPointRef(InIndex, FBox(LocalCenter, LocalCenter).TransformBy(InTransforms[InIndex]));
 		}
 
-		FBoxCenterAndExtent GetPositionSearchBounds(const FPCGPoint& InPoint, const double& Extents)
+		FBoxCenterAndExtent GetPositionSearchBounds(const TArray<FTransform>& InTransforms, const UPCGPointData* InPointData, int32 InIndex, const double& Extents)
 		{
-			return FBoxCenterAndExtent(InPoint.Transform.GetLocation(), FVector(Extents, Extents, Extents));
+			check(InPointData);
+			return FBoxCenterAndExtent(InTransforms[InIndex].GetLocation(), FVector(Extents, Extents, Extents));
 		}
 
-		FBoxCenterAndExtent GetCenterSearchBounds(const FPCGPoint& InPoint, const double& Extents)
+		FBoxCenterAndExtent GetCenterSearchBounds(const TArray<FTransform>& InTransforms, const UPCGPointData* InPointData, int32 InIndex, const double& Extents)
 		{
-			return FBoxCenterAndExtent(InPoint.Transform.TransformPosition(InPoint.GetLocalCenter()), FVector(Extents, Extents, Extents));
+			check(InPointData);
+			return FBoxCenterAndExtent(InTransforms[InIndex].TransformPosition(InPointData->GetLocalCenter(InIndex)), FVector(Extents, Extents, Extents));
 		}
 	}
 } // namespace PCGCollapsePoints
@@ -236,9 +241,9 @@ bool FPCGCollapsePointsElement::PrepareDataInternal(FPCGContext* InContext) cons
 	Context->InitializePerExecutionState([Settings](const ContextType*, ExecStateType& OutState)
 	{
 		PCGCollapsePoints::FCollapsePointsSettings::PairSelectionFuncType PairSelectionFunc = nullptr;
-		PCGCollapsePoints::FCollapsePointsSettings::GetPositionFuncType GetPositionFunc = nullptr;
-		PCGCollapsePoints::FCollapsePointsSettings::GetPointRefFuncType GetPointRefFunc = nullptr;
-		PCGCollapsePoints::FCollapsePointsSettings::GetSearchBoundsFuncType GetSearchBoundsFunc = nullptr;
+		PCGCollapsePoints::FCollapsePointsSettings::GetPointPositionFuncType GetPointPositionFunc = nullptr;
+		PCGCollapsePoints::FCollapsePointsSettings::GetPointReferenceFuncType GetPointReferenceFunc = nullptr;
+		PCGCollapsePoints::FCollapsePointsSettings::GetPointSearchBoundsFuncType GetPointSearchBoundsFunc = nullptr;
 
 		if (Settings->Mode == EPCGCollapseMode::PairwiseClosest)
 		{
@@ -255,15 +260,15 @@ bool FPCGCollapsePointsElement::PrepareDataInternal(FPCGContext* InContext) cons
 
 		if (Settings->ComparisonMode == EPCGCollapseComparisonMode::Position)
 		{
-			GetPositionFunc = PCGCollapsePoints::ComparisonModes::GetPosition;
-			GetPointRefFunc = PCGCollapsePoints::ComparisonModes::GetPositionPointRef;
-			GetSearchBoundsFunc = PCGCollapsePoints::ComparisonModes::GetPositionSearchBounds;
+			GetPointPositionFunc = PCGCollapsePoints::ComparisonModes::GetPosition;
+			GetPointReferenceFunc = PCGCollapsePoints::ComparisonModes::GetPositionPointRef;
+			GetPointSearchBoundsFunc = PCGCollapsePoints::ComparisonModes::GetPositionSearchBounds;
 		}
 		else if (Settings->ComparisonMode == EPCGCollapseComparisonMode::Center)
 		{
-			GetPositionFunc = PCGCollapsePoints::ComparisonModes::GetCenter;
-			GetPointRefFunc = PCGCollapsePoints::ComparisonModes::GetCenterPointRef;
-			GetSearchBoundsFunc = PCGCollapsePoints::ComparisonModes::GetCenterSearchBounds;
+			GetPointPositionFunc = PCGCollapsePoints::ComparisonModes::GetCenter;
+			GetPointReferenceFunc = PCGCollapsePoints::ComparisonModes::GetCenterPointRef;
+			GetPointSearchBoundsFunc = PCGCollapsePoints::ComparisonModes::GetCenterSearchBounds;
 		}
 		else
 		{
@@ -271,14 +276,15 @@ bool FPCGCollapsePointsElement::PrepareDataInternal(FPCGContext* InContext) cons
 		}
 
 		check(PairSelectionFunc);
+
 		OutState = PCGCollapsePoints::FCollapsePointsSettings
 		{
 			.Settings = Settings,
 			.PairSelectionFunc = PairSelectionFunc,
 			.MergeSelectionFunc = PCGCollapsePoints::Algo::MergePairs,
-			.GetPositionFunc = GetPositionFunc,
-			.GetPointRefFunc = GetPointRefFunc,
-			.GetSearchBoundsFunc = GetSearchBoundsFunc,
+			.GetPointPositionFunc = GetPointPositionFunc,
+			.GetPointReferenceFunc = GetPointReferenceFunc,
+			.GetPointSearchBoundsFunc = GetPointSearchBoundsFunc,
 			.DistanceThreshold = Settings->DistanceThreshold
 		};
 
@@ -296,15 +302,15 @@ bool FPCGCollapsePointsElement::PrepareDataInternal(FPCGContext* InContext) cons
 
 		OutState.SourceData = PointData;
 
-		const TArray<FPCGPoint>& Points = PointData->GetPoints();
-		const int32 NumPoints = Points.Num();
-
 		// Get points, octree will be rebuilt later
-		OutState.Points = Points;
+		OutState.PointTransforms = PointData->GetTransformsCopy();
 
 		// Get weights
+		const int32 NumPoints = PointData->GetNumPoints();
 		OutState.Weights.SetNumUninitialized(NumPoints);
 
+		ensure(OutState.PointTransforms.Num() == OutState.Weights.Num());
+
 		if (Settings->bUseMergeWeightAttribute)
 		{
 			const FPCGAttributePropertyInputSelector Selector = Settings->MergeWeightAttribute.CopyAndFixLast(PointData);
diff --git a/Engine/Plugins/PCG/Source/PCG/Private/Elements/PCGCullPointsOutsideActorBounds.cpp b/Engine/Plugins/PCG/Source/PCG/Private/Elements/PCGCullPointsOutsideActorBounds.cpp
--- a/Engine/Plugins/PCG/Source/PCG/Private/Elements/PCGCullPointsOutsideActorBounds.cpp
+++ b/Engine/Plugins/PCG/Source/PCG/Private/Elements/PCGCullPointsOutsideActorBounds.cpp
@@ -98,12 +98,13 @@ bool FPCGCullPointsOutsideActorBoundsElement::ExecuteInternal(FPCGContext* Conte
 			TArray<FPCGPoint>& CulledPoints = CulledPointsData->GetMutablePoints();
 
 			// TODO: split up this query across mutiple threads, but will make gathering the point a bit harder
-			const UPCGPointData::PointOctree& InputOctree = InputPointData->GetOctree();
+			const PCGPointOctree::FPointOctree& InputOctree = InputPointData->GetPointOctree();
+			const TArray<FPCGPoint>& InputPoints = InputPointData->GetPoints();
 			FBoxCenterAndExtent BoxCenterAndExtents(BoundsBox);
 
-			InputOctree.FindElementsWithBoundsTest(BoxCenterAndExtents, [&CulledPoints](const FPCGPointRef& PointRef)
+			InputOctree.FindElementsWithBoundsTest(BoxCenterAndExtents, [&InputPoints , &CulledPoints](const PCGPointOctree::FPointRef& PointRef)
 			{
-				CulledPoints.Add(*PointRef.Point);
+				CulledPoints.Add(InputPoints[PointRef.Index]);
 			});
 
 			if (!CulledPoints.IsEmpty())
diff --git a/Engine/Plugins/PCG/Source/PCG/Private/Elements/PCGDistance.cpp b/Engine/Plugins/PCG/Source/PCG/Private/Elements/PCGDistance.cpp
--- a/Engine/Plugins/PCG/Source/PCG/Private/Elements/PCGDistance.cpp
+++ b/Engine/Plugins/PCG/Source/PCG/Private/Elements/PCGDistance.cpp
@@ -278,16 +278,15 @@ bool FPCGDistanceElement::ExecuteInternal(FPCGContext* Context) const
 			FVector MinDistanceVector = FVector::ZeroVector;
 
 			// Signed distance field for calculating the closest point of source and target
-			auto CalculateSDF = [&MinDistanceSquared, &MinDistanceVector, &SourcePoint, SourceCenter, SourceShape, TargetShape](const FPCGPointRef& TargetPointRef)
+			auto CalculateSDF = [&MinDistanceSquared, &MinDistanceVector, &SourcePoint, SourceCenter, SourceShape, TargetShape](const FPCGPoint& TargetPoint, const FBoxSphereBounds& Bounds)
 			{
 				// If the source pointer and target pointer are the same, ignore distance to the exact same point
-				if (&SourcePoint == TargetPointRef.Point)
+				if (&SourcePoint == &TargetPoint)
 				{
 					return;
 				}
 
-				const FPCGPoint& TargetPoint = *TargetPointRef.Point;
-				const FVector& TargetCenter = TargetPointRef.Bounds.Origin;
+				const FVector& TargetCenter = Bounds.Origin;
 
 				const FVector SourceShapePos = PCGDistance::CalcPosition(SourceShape, SourcePoint, TargetPoint, SourceCenter, TargetCenter);
 				const FVector TargetShapePos = PCGDistance::CalcPosition(TargetShape, TargetPoint, SourcePoint, TargetCenter, SourceCenter);
@@ -309,11 +308,14 @@ bool FPCGDistanceElement::ExecuteInternal(FPCGContext* Context) const
 			{
 				check(TargetPointData);
 				
-				const UPCGPointData::PointOctree& Octree = TargetPointData->GetOctree();
-
-				Octree.FindElementsWithBoundsTest(
-						FBoxCenterAndExtent(SourceQueryBounds.Origin, SourceQueryBounds.BoxExtent),
-						CalculateSDF);
+				const PCGPointOctree::FPointOctree& Octree = TargetPointData->GetPointOctree();
+				const TArray<FPCGPoint>& TargetPoints = TargetPointData->GetPoints();
+
+				Octree.FindElementsWithBoundsTest(FBoxCenterAndExtent(SourceQueryBounds.Origin, SourceQueryBounds.BoxExtent),
+					[&TargetPoints, &CalculateSDF](const PCGPointOctree::FPointRef& InPointRef)
+					{
+						CalculateSDF(TargetPoints[InPointRef.Index], InPointRef.Bounds);
+					});
 			};
 
 			if (Settings->bCheckSourceAgainstRespectiveTarget)
diff --git a/Engine/Plugins/PCG/Source/PCG/Private/Elements/PCGSelfPruning.cpp b/Engine/Plugins/PCG/Source/PCG/Private/Elements/PCGSelfPruning.cpp
--- a/Engine/Plugins/PCG/Source/PCG/Private/Elements/PCGSelfPruning.cpp
+++ b/Engine/Plugins/PCG/Source/PCG/Private/Elements/PCGSelfPruning.cpp
@@ -4,6 +4,7 @@
 
 #include "PCGContext.h"
 #include "PCGCustomVersion.h"
+#include "Data/PCGBasePointData.h"
 #include "Data/PCGPointData.h"
 #include "Data/PCGSpatialData.h"
 #include "Metadata/Accessors/IPCGAttributeAccessor.h"
@@ -22,17 +23,20 @@
 
 namespace PCGSelfPruningAlgorithms
 {
-	bool RandomSort(const FPCGPointRef& A, const FPCGPointRef& B)
+	bool RandomSort(const UPCGBasePointData* PointData, const PCGPointOctree::FPointRef& A, const PCGPointOctree::FPointRef& B)
 	{
-		return A.Point->Seed < B.Point->Seed;
+		const int32 SeedA = PointData->GetSeed(A.Index);
+		const int32 SeedB = PointData->GetSeed(B.Index);
+
+		return SeedA < SeedB;
 	}
 
-	bool SortSmallToLargeNoRandom(const FPCGPointRef& A, const FPCGPointRef& B, FVector::FReal SquaredRadiusEquality)
+	bool SortSmallToLargeNoRandom(const UPCGBasePointData* PointData, const PCGPointOctree::FPointRef& A, const PCGPointOctree::FPointRef& B, FVector::FReal SquaredRadiusEquality)
 	{
 		return A.Bounds.BoxExtent.SquaredLength() * SquaredRadiusEquality < B.Bounds.BoxExtent.SquaredLength();
 	}
 
-	bool SortSmallToLargeWithRandom(const FPCGPointRef& A, const FPCGPointRef& B, FVector::FReal SquaredRadiusEquality)
+	bool SortSmallToLargeWithRandom(const UPCGBasePointData* PointData, const PCGPointOctree::FPointRef& A, const PCGPointOctree::FPointRef& B, FVector::FReal SquaredRadiusEquality)
 	{
 		const FVector::FReal SqrLenA = A.Bounds.BoxExtent.SquaredLength();
 		const FVector::FReal SqrLenB = B.Bounds.BoxExtent.SquaredLength();
@@ -46,14 +50,14 @@ namespace PCGSelfPruningAlgorithms
 		}
 		else
 		{
-			return RandomSort(A, B);
+			return RandomSort(PointData, A, B);
 		}
 	}
 
 	/**
 	* Keep Extents sort alive since it is a bit faster than the SortByAttribute.
 	*/
-	void SortExtents(TArray<FPCGPointRef>& SortedPoints, const FPCGSelfPruningParameters& Parameters, FVector::FReal SquaredRadiusEquality)
+	void SortExtents(const UPCGBasePointData* PointData, TArray<PCGPointOctree::FPointRef>& SortedPointRefs, const FPCGSelfPruningParameters& Parameters, FVector::FReal SquaredRadiusEquality)
 	{
 		TRACE_CPUPROFILER_EVENT_SCOPE(FPCGSelfPruningElement::Execute::SortExtents);
 
@@ -61,22 +65,22 @@ namespace PCGSelfPruningAlgorithms
 		{
 			if (Parameters.bRandomizedPruning)
 			{
-				Algo::Sort(SortedPoints, [SquaredRadiusEquality](const FPCGPointRef& A, const FPCGPointRef& B) { return !PCGSelfPruningAlgorithms::SortSmallToLargeWithRandom(A, B, SquaredRadiusEquality); });
+				Algo::Sort(SortedPointRefs, [PointData, SquaredRadiusEquality](const PCGPointOctree::FPointRef& A, const PCGPointOctree::FPointRef& B) { return !PCGSelfPruningAlgorithms::SortSmallToLargeWithRandom(PointData, A, B, SquaredRadiusEquality); });
 			}
 			else
 			{
-				Algo::Sort(SortedPoints, [SquaredRadiusEquality](const FPCGPointRef& A, const FPCGPointRef& B) { return !PCGSelfPruningAlgorithms::SortSmallToLargeNoRandom(A, B, SquaredRadiusEquality); });
+				Algo::Sort(SortedPointRefs, [PointData, SquaredRadiusEquality](const PCGPointOctree::FPointRef& A, const PCGPointOctree::FPointRef& B) { return !PCGSelfPruningAlgorithms::SortSmallToLargeNoRandom(PointData, A, B, SquaredRadiusEquality); });
 			}
 		}
 		else //if (Parameters.PruningType == EPCGSelfPruningType::SmallToLarge)
 		{
 			if (Parameters.bRandomizedPruning)
 			{
-				Algo::Sort(SortedPoints, [SquaredRadiusEquality](const FPCGPointRef& A, const FPCGPointRef& B) { return PCGSelfPruningAlgorithms::SortSmallToLargeWithRandom(A, B, SquaredRadiusEquality); });
+				Algo::Sort(SortedPointRefs, [PointData, SquaredRadiusEquality](const PCGPointOctree::FPointRef& A, const PCGPointOctree::FPointRef& B) { return PCGSelfPruningAlgorithms::SortSmallToLargeWithRandom(PointData, A, B, SquaredRadiusEquality); });
 			}
 			else
 			{
-				Algo::Sort(SortedPoints, [SquaredRadiusEquality](const FPCGPointRef& A, const FPCGPointRef& B) { return PCGSelfPruningAlgorithms::SortSmallToLargeNoRandom(A, B, SquaredRadiusEquality); });
+				Algo::Sort(SortedPointRefs, [PointData, SquaredRadiusEquality](const PCGPointOctree::FPointRef& A, const PCGPointOctree::FPointRef& B) { return PCGSelfPruningAlgorithms::SortSmallToLargeNoRandom(PointData, A, B, SquaredRadiusEquality); });
 			}
 		}
 	}
@@ -84,24 +88,21 @@ namespace PCGSelfPruningAlgorithms
 
 namespace PCGSelfPruningElement
 {
-	void FPointBitSet::Initialize(const TArray<FPCGPoint>& Points)
+	void FPointBitSet::Initialize(int32 NumPoints)
 	{
-		Bits.SetNumZeroed(Points.Num() / 32 + 1);
-		FirstPoint = Points.Num() > 0 ? &Points[0] : nullptr;
+		Bits.SetNumZeroed(NumPoints / 32 + 1);
 	}
 
-	void FPointBitSet::Add(const FPCGPoint* Point)
+	void FPointBitSet::Add(int32 Index)
 	{
-		const uint32 Index = Point - FirstPoint;
 		const uint32 BitsIndex = Index / 32;
 		const uint32 Bit = Index % 32;
 
 		Bits[BitsIndex] |= (1 << Bit);
 	}
 
-	bool FPointBitSet::Contains(const FPCGPoint* Point)
+	bool FPointBitSet::Contains(int32 Index)
 	{
-		const uint32 Index = Point - FirstPoint;
 		const uint32 BitsIndex = Index / 32;
 		const uint32 Bit = Index % 32;
 
@@ -120,11 +121,11 @@ namespace PCGSelfPruningElement
 		TRACE_CPUPROFILER_EVENT_SCOPE(FPCGSelfPruningElement::Execute::DensityBoundsExclusion);
 
 		check(IterationState.InputPointData);
-		const UPCGPointData::PointOctree& Octree = IterationState.InputPointData->GetOctree();
+		const PCGPointOctree::FPointOctree& Octree = IterationState.InputPointData->GetPointOctree();
 
 		int32 CheckTimeSlicingCount = 0;
 
-		while (IterationState.CurrentPointIndex < IterationState.SortedPoints.Num())
+		while (IterationState.CurrentPointIndex < IterationState.SortedPointRefs.Num())
 		{
 			// Don't check too many times. Pre-increment will make sure we always at least process "TimeSliceFrequencyCheck" elements at each iteration.
 			if (++CheckTimeSlicingCount >= TimeSliceFrequencyCheck)
@@ -137,20 +138,20 @@ namespace PCGSelfPruningElement
 				CheckTimeSlicingCount = 0;
 			}
 
-			const FPCGPointRef& PointRef = IterationState.SortedPoints[IterationState.CurrentPointIndex++];
-			if (IterationState.ExcludedPoints.Contains(PointRef.Point))
+			const PCGPointOctree::FPointRef& PointRef = IterationState.SortedPointRefs[IterationState.CurrentPointIndex++];
+			if (IterationState.ExcludedPoints.Contains(PointRef.Index))
 			{
 				continue;
 			}
 
-			IterationState.ExclusionPoints.Add(PointRef.Point);
+			IterationState.ExclusionPoints.Add(PointRef.Index);
 
-			Octree.FindElementsWithBoundsTest(FBoxCenterAndExtent(PointRef.Bounds.Origin, PointRef.Bounds.BoxExtent), [&IterationState](const FPCGPointRef& InPointRef)
+			Octree.FindElementsWithBoundsTest(FBoxCenterAndExtent(PointRef.Bounds.Origin, PointRef.Bounds.BoxExtent), [&IterationState](const PCGPointOctree::FPointRef& InPointRef)
 			{
 				// TODO: check on an oriented-box basis?
-				if (!IterationState.ExclusionPoints.Contains(InPointRef.Point))
+				if (!IterationState.ExclusionPoints.Contains(InPointRef.Index))
 				{
-					IterationState.ExcludedPoints.Add(InPointRef.Point);
+					IterationState.ExcludedPoints.Add(InPointRef.Index);
 				}
 			});
 		}
@@ -164,15 +165,15 @@ namespace PCGSelfPruningElement
 		TRACE_CPUPROFILER_EVENT_SCOPE(FPCGSelfPruningElement::Execute::CollisionExclusion);
 
 		check(IterationState.InputPointData);
-		const UPCGPointData::PointOctree& Octree = IterationState.InputPointData->GetOctree();
+		const PCGPointOctree::FPointOctree& Octree = IterationState.InputPointData->GetPointOctree();
 
 		int32 CheckTimeSlicingCount = 0;
 
-		TArray<const FPCGPointRef*, TInlineAllocator<256>> ElementsToTest;
+		TArray<const PCGPointOctree::FPointRef*, TInlineAllocator<256>> ElementsToTest;
 
-		auto SetupQueryInfo = [&IterationState](const FPCGPointRef& PointRef, FBodyInstance* OtherBodyInstance, EPCGCollisionQueryFlag CollisionQueryFlag, FBodyInstance*& OutInstance, PhysicsInterfaceTypes::FInlineShapeArray& OutShapes, FCollisionShape& OutSimpleShape)
+		auto SetupQueryInfo = [&IterationState](const PCGPointOctree::FPointRef& PointRef, FBodyInstance* OtherBodyInstance, EPCGCollisionQueryFlag CollisionQueryFlag, FBodyInstance*& OutInstance, PhysicsInterfaceTypes::FInlineShapeArray& OutShapes, FCollisionShape& OutSimpleShape)
 		{
-			const int32 EntryIndex = PointRef.Point - IterationState.InputPointData->GetPoints().GetData();
+			const int32 EntryIndex = PointRef.Index;
 			OutInstance = IterationState.CollisionWrapper.GetBodyInstance(EntryIndex);
 			bool bHasComplexShapes = false;
 
@@ -196,7 +197,7 @@ namespace PCGSelfPruningElement
 					}
 				}
 
-				OutInstance->UpdateBodyScale(PointRef.Point->Transform.GetScale3D());
+				OutInstance->UpdateBodyScale(IterationState.InputPointData->GetTransform(PointRef.Index).GetScale3D());
 				const bool bFirstChoice = FPCGCollisionWrapper::GetShapeArray(OutInstance, CollisionQueryFlag, OutShapes);
 
 				if (OutShapes.IsEmpty())
@@ -220,7 +221,7 @@ namespace PCGSelfPruningElement
 			return bHasComplexShapes;
 		};
 
-		while (IterationState.CurrentPointIndex < IterationState.SortedPoints.Num())
+		while (IterationState.CurrentPointIndex < IterationState.SortedPointRefs.Num())
 		{
 			// Don't check too many times. Pre-increment will make sure we always at least process "TimeSliceFrequencyCheck" elements at each iteration.
 			if (++CheckTimeSlicingCount >= TimeSliceFrequencyCheck)
@@ -233,20 +234,20 @@ namespace PCGSelfPruningElement
 				CheckTimeSlicingCount = 0;
 			}
 
-			const FPCGPointRef& PointRef = IterationState.SortedPoints[IterationState.CurrentPointIndex++];
-			if (IterationState.ExcludedPoints.Contains(PointRef.Point))
+			const PCGPointOctree::FPointRef& PointRef = IterationState.SortedPointRefs[IterationState.CurrentPointIndex++];
+			if (IterationState.ExcludedPoints.Contains(PointRef.Index))
 			{
 				continue; // Point discarded from previous iteration
 			}
 
 			// Select point
-			IterationState.ExclusionPoints.Add(PointRef.Point);
+			IterationState.ExclusionPoints.Add(PointRef.Index);
 
 			// 1. Gather point refs to test against - similar to the DensityBoundsExclusion, except we write to an array temporarily
 			ElementsToTest.Reset();
-			Octree.FindElementsWithBoundsTest(FBoxCenterAndExtent(PointRef.Bounds.Origin, PointRef.Bounds.BoxExtent), [&IterationState, &ElementsToTest](const FPCGPointRef& OtherPointRef)
+			Octree.FindElementsWithBoundsTest(FBoxCenterAndExtent(PointRef.Bounds.Origin, PointRef.Bounds.BoxExtent), [&IterationState, &ElementsToTest](const PCGPointOctree::FPointRef& OtherPointRef)
 			{
-				if (!IterationState.ExclusionPoints.Contains(OtherPointRef.Point) && !IterationState.ExcludedPoints.Contains(OtherPointRef.Point))
+				if (!IterationState.ExclusionPoints.Contains(OtherPointRef.Index) && !IterationState.ExcludedPoints.Contains(OtherPointRef.Index))
 				{
 					ElementsToTest.Add(&OtherPointRef);
 				}
@@ -267,19 +268,21 @@ namespace PCGSelfPruningElement
 			PhysicsInterfaceTypes::FInlineShapeArray OtherShapes;
 			FCollisionShape OtherSimpleShape;
 
+			const FTransform& PointTransform = IterationState.InputPointData->GetTransform(PointRef.Index);
 			const bool bThisHasComplexShapes = SetupQueryInfo(PointRef, nullptr, InCollisionQueryFlag, ThisInstance, ThisShapes, ThisSimpleShape);
-			FTransform TransformNoScale = FTransform(PointRef.Point->Transform.GetRotation(), PointRef.Point->Transform.GetLocation());
+			FTransform TransformNoScale = FTransform(PointTransform.GetRotation(), PointTransform.GetLocation());
 
 			// We must force the other collision flag to simple if we have complex shapes in the leading shape here, because complex-complex overlaps aren't supported.
 			EPCGCollisionQueryFlag OtherCollisionQueryFlag = (bThisHasComplexShapes ? EPCGCollisionQueryFlag::Simple : InCollisionQueryFlag);
 
-			for (const FPCGPointRef* OtherPointRef : ElementsToTest)
+			for (const PCGPointOctree::FPointRef* OtherPointRef : ElementsToTest)
 			{
 				bool bOverlaps = false;
 
 				bool bOtherHasComplexShapes = SetupQueryInfo(*OtherPointRef, ThisInstance, OtherCollisionQueryFlag, OtherInstance, OtherShapes, OtherSimpleShape);
 
-				FTransform OtherTransformNoScale = FTransform(OtherPointRef->Point->Transform.GetRotation(), OtherPointRef->Point->Transform.GetLocation());
+				const FTransform& OtherPointTransform = IterationState.InputPointData->GetTransform(OtherPointRef->Index);
+				FTransform OtherTransformNoScale = FTransform(OtherPointTransform.GetRotation(), OtherPointTransform.GetLocation());
 
 				// Four cases here:
 				// 1 - shapes vs shapes
@@ -359,13 +362,13 @@ namespace PCGSelfPruningElement
 
 				if (bOverlaps)
 				{
-					IterationState.ExcludedPoints.Add(OtherPointRef->Point);
+					IterationState.ExcludedPoints.Add(OtherPointRef->Index);
 				}
 			}
 		}
 
 		// Clean up temporary instances
-		if (IterationState.CurrentPointIndex == IterationState.SortedPoints.Num())
+		if (IterationState.CurrentPointIndex == IterationState.SortedPointRefs.Num())
 		{
 			for (TPair<FBodyInstance*, FBodyInstance*>& TemporaryInstance : IterationState.TemporaryBodyInstances)
 			{
@@ -382,11 +385,11 @@ namespace PCGSelfPruningElement
 		TRACE_CPUPROFILER_EVENT_SCOPE(FPCGSelfPruningElement::Execute::DuplicatePointsExclusion);
 
 		check(IterationState.InputPointData);
-		const UPCGPointData::PointOctree& Octree = IterationState.InputPointData->GetOctree();
+		const PCGPointOctree::FPointOctree& Octree = IterationState.InputPointData->GetPointOctree();
 
 		int32 CheckTimeSlicingCount = 0;
 
-		while (IterationState.CurrentPointIndex < IterationState.SortedPoints.Num())
+		while (IterationState.CurrentPointIndex < IterationState.SortedPointRefs.Num())
 		{
 			// Don't check too many times. Pre-increment will make sure we always at least process "TimeSliceFrequencyCheck" elements at each iteration.
 			if (++CheckTimeSlicingCount >= TimeSliceFrequencyCheck)
@@ -399,20 +402,22 @@ namespace PCGSelfPruningElement
 				CheckTimeSlicingCount = 0;
 			}
 
-			const FPCGPointRef& PointRef = IterationState.SortedPoints[IterationState.CurrentPointIndex++];
-			if (IterationState.ExcludedPoints.Contains(PointRef.Point))
+			const PCGPointOctree::FPointRef& PointRef = IterationState.SortedPointRefs[IterationState.CurrentPointIndex++];
+			if (IterationState.ExcludedPoints.Contains(PointRef.Index))
 			{
 				continue;
 			}
 
-			IterationState.ExclusionPoints.Add(PointRef.Point);
+			IterationState.ExclusionPoints.Add(PointRef.Index);
+			const FTransform& PointTransform = IterationState.InputPointData->GetTransform(PointRef.Index);
+			const FVector PointLocalCenter = IterationState.InputPointData->GetLocalCenter(PointRef.Index);
 
-			Octree.FindElementsWithBoundsTest(FBoxCenterAndExtent(PointRef.Point->Transform.TransformPosition(PointRef.Point->GetLocalCenter()), FVector::Zero()), [&IterationState, PointRef](const FPCGPointRef& InPointRef)
+			Octree.FindElementsWithBoundsTest(FBoxCenterAndExtent(PointTransform.TransformPosition(PointLocalCenter), FVector::Zero()), [&IterationState, &PointTransform](const PCGPointOctree::FPointRef InPointRef)
 			{
-				if ((PointRef.Point->Transform.GetLocation() - InPointRef.Point->Transform.GetLocation()).SquaredLength() <= SMALL_NUMBER &&
-				!IterationState.ExclusionPoints.Contains(InPointRef.Point))
+				const FTransform& OtherPointTransform = IterationState.InputPointData->GetTransform(InPointRef.Index);
+				if ((PointTransform.GetLocation() - OtherPointTransform.GetLocation()).SquaredLength() <= SMALL_NUMBER && !IterationState.ExclusionPoints.Contains(InPointRef.Index))
 				{
-					IterationState.ExcludedPoints.Add(InPointRef.Point);
+					IterationState.ExcludedPoints.Add(InPointRef.Index);
 				}
 			});
 		}
@@ -484,8 +489,8 @@ namespace PCGSelfPruningElement
 		const TArray<FPCGPoint>& Points = InState.InputPointData->GetPoints();
 
 		// Force octree computation, and check if we need to stop after that if it was dirty in the first place.
-		const bool bOctreeWasDirty = InState.InputPointData->IsOctreeDirty();
-		const UPCGPointData::PointOctree& Octree = InState.InputPointData->GetOctree();
+		const bool bOctreeWasDirty = InState.InputPointData->IsPointOctreeDirty();
+		const PCGPointOctree::FPointOctree& Octree = InState.InputPointData->GetPointOctree();
 		if (bOctreeWasDirty && ShouldStop(InOptionalContext))
 		{
 			return false;
@@ -499,11 +504,11 @@ namespace PCGSelfPruningElement
 		if (!InState.bSortDone)
 		{
 			// In the case of the collision-driven self-pruning, we have to populate the sorted points array earlier since we're playing with the bounds.
-			InState.SortedPoints.Empty();
-			InState.SortedPoints.Reserve(Points.Num());
-			for (const FPCGPoint& Point : Points)
+			InState.SortedPointRefs.Empty();
+			InState.SortedPointRefs.Reserve(Points.Num());
+			for (int32 PointIndex = 0; PointIndex < Points.Num(); ++PointIndex)
 			{
-				InState.SortedPoints.Emplace(Point);
+				InState.SortedPointRefs.Emplace(PCGPointOctree::FPointRef(PointIndex, Points[PointIndex].GetDensityBounds()));
 			}
 
 			FPCGAttributePropertySelector ComparisonSource = InParameters.ComparisonSource.CopyAndFixLast(InState.InputPointData);
@@ -518,12 +523,12 @@ namespace PCGSelfPruningElement
 			if (InParameters.bRandomizedPruning && !bSortAttribute)
 			{
 				TRACE_CPUPROFILER_EVENT_SCOPE(FPCGSelfPruningElement::Execute::SortingSeeds);
-				Algo::Sort(InState.SortedPoints, PCGSelfPruningAlgorithms::RandomSort);
+				Algo::Sort(InState.SortedPointRefs, [&InState](const PCGPointOctree::FPointRef& A, const PCGPointOctree::FPointRef& B) { return PCGSelfPruningAlgorithms::RandomSort(InState.InputPointData, A, B); });
 			}
 			else if (ComparisonSource.GetSelection() == EPCGAttributePropertySelection::PointProperty && ComparisonSource.GetPointProperty() == EPCGPointProperties::Extents && ComparisonSource.GetExtraNames().IsEmpty())
 			{
 				TRACE_CPUPROFILER_EVENT_SCOPE(FPCGSelfPruningElement::Execute::SortingExtents);
-				PCGSelfPruningAlgorithms::SortExtents(InState.SortedPoints, InParameters, SquaredRadiusEquality);
+				PCGSelfPruningAlgorithms::SortExtents(InState.InputPointData, InState.SortedPointRefs, InParameters, SquaredRadiusEquality);
 			}
 			else
 			{
@@ -590,7 +595,7 @@ namespace PCGSelfPruningElement
 						return Points[IndexA].Seed < Points[IndexB].Seed;
 					};
 
-					PCGAttributeAccessorHelpers::SortByAttribute(*InputAccessor, FPCGAttributeAccessorKeysPoints(Points), InState.SortedPoints, bAscending, PCGAttributeAccessorHelpers::Private::DefaultIndexGetter, std::move(CompareLessWithRandom));
+					PCGAttributeAccessorHelpers::SortByAttribute(*InputAccessor, FPCGAttributeAccessorKeysPoints(Points), InState.SortedPointRefs, bAscending, PCGAttributeAccessorHelpers::Private::DefaultIndexGetter, std::move(CompareLessWithRandom));
 				}
 				else
 				{
@@ -611,12 +616,12 @@ namespace PCGSelfPruningElement
 						return PCG::Private::MetadataTraits<ValueType>::Less(ACopy, B);
 					};
 
-					PCGAttributeAccessorHelpers::SortByAttribute(*InputAccessor, FPCGAttributeAccessorKeysPoints(Points), InState.SortedPoints, bAscending, PCGAttributeAccessorHelpers::Private::DefaultIndexGetter, std::move(CompareLess));
+					PCGAttributeAccessorHelpers::SortByAttribute(*InputAccessor, FPCGAttributeAccessorKeysPoints(Points), InState.SortedPointRefs, bAscending, PCGAttributeAccessorHelpers::Private::DefaultIndexGetter, std::move(CompareLess));
 				}
 			}
 
-			InState.ExclusionPoints.Initialize(Points);
-			InState.ExcludedPoints.Initialize(Points);
+			InState.ExclusionPoints.Initialize(Points.Num());
+			InState.ExcludedPoints.Initialize(Points.Num());
 
 			InState.bSortDone = true;
 
@@ -652,11 +657,11 @@ namespace PCGSelfPruningElement
 
 			TArray<FPCGPoint>& OutputPoints = PrunedData->GetMutablePoints();
 
-			for (const FPCGPoint& Point : Points)
+			for (int32 PointIndex = 0; PointIndex < Points.Num(); ++PointIndex)
 			{
-				if (InState.ExclusionPoints.Contains(&Point))
+				if (InState.ExclusionPoints.Contains(PointIndex))
 				{
-					OutputPoints.Add(Point);
+					OutputPoints.Add(Points[PointIndex]);
 				}
 			}
 
diff --git a/Engine/Plugins/PCG/Source/PCG/Private/SpatialAlgo/PCGAStar.cpp b/Engine/Plugins/PCG/Source/PCG/Private/SpatialAlgo/PCGAStar.cpp
--- a/Engine/Plugins/PCG/Source/PCG/Private/SpatialAlgo/PCGAStar.cpp
+++ b/Engine/Plugins/PCG/Source/PCG/Private/SpatialAlgo/PCGAStar.cpp
@@ -189,14 +189,17 @@ ESearchResult ExecuteSearchIteration(const FSearchSettings& SearchSettings, FSea
 		PerPointProcessing(CurrentGoal, MinSquaredDistanceToGoal);
 	}
 
+	const TArray<FPCGPoint>& Points = SearchState.OriginatingPointData->GetPoints();
+
 	// Gather neighbors within the search radius.
 	UPCGOctreeQueries::ForEachPointInsideSphere(
 		SearchState.OriginatingPointData,
 		NodeList[CurrentNodeIndex].PCGPoint->Transform.GetLocation(),
-		SearchSettings.SearchDistance,
-		[&PerPointProcessing](const FPCGPointRef& PointRef, const double DistanceToPointSquared)
+ 		SearchSettings.SearchDistance,
+		[&Points, &PerPointProcessing](const UPCGBasePointData* BasePointData, int32 PointIndex, const double DistanceToPointSquared)
 		{
-			PerPointProcessing(PointRef.Point, DistanceToPointSquared);
+			const FPCGPoint* PCGPoint = &Points[PointIndex];
+			PerPointProcessing(PCGPoint, DistanceToPointSquared);
 		});
 
 	// This node has been completely evaluated.
diff --git a/Engine/Plugins/PCG/Source/PCG/Private/SpatialAlgo/PCGOctreeQueries.cpp b/Engine/Plugins/PCG/Source/PCG/Private/SpatialAlgo/PCGOctreeQueries.cpp
--- a/Engine/Plugins/PCG/Source/PCG/Private/SpatialAlgo/PCGOctreeQueries.cpp
+++ b/Engine/Plugins/PCG/Source/PCG/Private/SpatialAlgo/PCGOctreeQueries.cpp
@@ -5,7 +5,7 @@
 namespace PCGOctreeQueries
 {
 	/** Factorize the logic for all GetClosest/GetFarthest, as the only thing that changes is the StartingDistance and the Condition. */
-	template <typename Func = TFunctionRef<bool(double ChosenPointDistance, const FPCGPoint& CurrentPoint, double SquaredDistance)>>
+	template <typename Func = TFunctionRef<bool(double ChosenPointDistance, int32 PointIndex, double SquaredDistance)>>
 	const FPCGPoint* GetPointInSphereUnderCondition(const UPCGPointData* InPointData, const FVector& InCenter, const double InSearchDistance, const double InStartingDistance, Func Condition)
 	{
 		if (!InPointData)
@@ -15,14 +15,16 @@ namespace PCGOctreeQueries
 
 		const FPCGPoint* ChosenPoint = nullptr;
 		const double SearchDistance = FMath::Max(InSearchDistance, 0.0);
+		const TArray<FPCGPoint>& Points = InPointData->GetPoints();
 
 		double ChosenPointDistance = InStartingDistance;
-		UPCGOctreeQueries::ForEachPointInsideSphere(InPointData, InCenter, SearchDistance, [&ChosenPointDistance, &ChosenPoint, Condition = std::move(Condition)](const FPCGPoint& Point, double SquaredDistance)
+		UPCGOctreeQueries::ForEachPointInsideSphere(InPointData, InCenter, SearchDistance, [&Points, &ChosenPointDistance, &ChosenPoint, Condition = std::move(Condition)](const UPCGBasePointData* BasePointData, int32 PointIndex, double SquaredDistance)
 		{
-			if (Condition(ChosenPointDistance, Point, SquaredDistance))
+			const FPCGPoint& PCGPoint = Points[PointIndex];
+			if (Condition(ChosenPointDistance, PCGPoint, SquaredDistance))
 			{
 				ChosenPointDistance = SquaredDistance;
-				ChosenPoint = &Point;
+				ChosenPoint = &PCGPoint;
 			}
 		});
 
@@ -38,11 +40,13 @@ TArray<FPCGPoint> UPCGOctreeQueries::GetPointsInsideBounds(const UPCGPointData*
 	}
 
 	TArray<FPCGPoint> Result;
-	InPointData->GetOctree().FindElementsWithBoundsTest(InBounds, [&Result](const FPCGPointRef& PointRef)
+	const TArray<FPCGPoint>& Points = InPointData->GetPoints();
+
+	InPointData->GetPointOctree().FindElementsWithBoundsTest(InBounds, [&Points, &Result](const PCGPointOctree::FPointRef& PointRef)
 	{
-		if (ensure(PointRef.Point))
+		if (Points.IsValidIndex(PointRef.Index))
 		{
-			Result.Add(*PointRef.Point);
+			Result.Add(Points[PointRef.Index]);
 		}
 	});
 
diff --git a/Engine/Plugins/PCG/Source/PCG/Private/Tests/SpatialAlgo/PCGOctreeQueriesTest.cpp b/Engine/Plugins/PCG/Source/PCG/Private/Tests/SpatialAlgo/PCGOctreeQueriesTest.cpp
--- a/Engine/Plugins/PCG/Source/PCG/Private/Tests/SpatialAlgo/PCGOctreeQueriesTest.cpp
+++ b/Engine/Plugins/PCG/Source/PCG/Private/Tests/SpatialAlgo/PCGOctreeQueriesTest.cpp
@@ -58,6 +58,17 @@ bool FPCGOctreeQueriesTest_Sphere::RunTest(const FString& Parameters)
 
 	UTEST_EQUAL("We found the expected number of points", CountFound, ExpectedIndexes.Num());
 
+	CountFound = 0;
+	UPCGOctreeQueries::ForEachPointInsideSphere(InputPointData, FVector::ZeroVector, 350.0, [&ExpectedIndexes, &CountFound](const UPCGBasePointData* PointData, int32 PointIndex, double)
+	{
+		if (ExpectedIndexes.Contains(PointIndex))
+		{
+			++CountFound;
+		}
+	});
+
+	UTEST_EQUAL("We found the expected number of points", CountFound, ExpectedIndexes.Num());
+
 	return true;
 }
 
diff --git a/Engine/Plugins/PCG/Source/PCG/Public/Data/PCGBasePointData.h b/Engine/Plugins/PCG/Source/PCG/Public/Data/PCGBasePointData.h
new file mode 100644
--- /dev/null
+++ b/Engine/Plugins/PCG/Source/PCG/Public/Data/PCGBasePointData.h
@@ -0,0 +1,44 @@
+// Copyright Epic Games, Inc. All Rights Reserved.
+
+#pragma once
+
+#include "PCGSpatialData.h"
+
+#include "Utils/PCGPointOctree.h"
+
+#include "PCGBasePointData.generated.h"
+
+
+UCLASS(Abstract)
+class PCG_API UPCGBasePointData : public UPCGSpatialData
+{
+	GENERATED_BODY()
+
+public:
+	virtual bool IsValidRef(const PCGPointOctree::FPointRef& InPointRef) const PURE_VIRTUAL(UPCGBasePointData::IsValidRef, return false;);
+
+	virtual bool IsPointOctreeDirty() const PURE_VIRTUAL(UPCGBasePointData::IsPointOctreeDirty, return false;);
+	virtual const PCGPointOctree::FPointOctree& GetPointOctree() const PURE_VIRTUAL(UPCGBasePointData::GetPointOctree, static PCGPointOctree::FPointOctree Empty; return Empty;);
+
+	virtual int32 GetNumPoints() const PURE_VIRTUAL(UPCGBasePointData::GetNumPoints, return 0;);
+
+	virtual TArray<FTransform> GetTransformsCopy() const PURE_VIRTUAL(UPCGBasePointData::GetTransforms, return TArray<FTransform>(););
+	
+	virtual const FTransform& GetTransform(int32 InPointIndex) const PURE_VIRTUAL(UPCGBasePointData::GetTransform, return FTransform::Identity;);
+	virtual float GetDensity(int32 InPointIndex) const PURE_VIRTUAL(UPCGBasePointData::GetDensity, return 0.f;);
+	virtual const FVector& GetBoundsMin(int32 InPointIndex) const PURE_VIRTUAL(UPCGBasePointData::GetBoundsMin, return FVector::ZeroVector;);
+	virtual const FVector& GetBoundsMax(int32 InPointIndex) const PURE_VIRTUAL(UPCGBasePointData::GetBoundsMin, return FVector::ZeroVector;);
+	virtual const FVector4& GetColor(int32 InPointIndex) const PURE_VIRTUAL(UPCGBasePointData::GetColor, static FVector4 Zero(0.0f); return Zero;);
+	virtual float GetSteepness(int32 InPointIndex) const PURE_VIRTUAL(UPCGBasePointData::GetSteepness, return 0.f;);
+	virtual int32 GetSeed(int32 InPointIndex) const PURE_VIRTUAL(UPCGBasePointData::GetSeed, return 0;);
+	virtual int64 GetMetadataEntry(int32 InPointIndex) const PURE_VIRTUAL(UPCGBasePointData::GetMetadataEntry, return 0;);
+
+	virtual FBoxSphereBounds GetDensityBounds(int32 InPointIndex) const PURE_VIRTUAL(UPCGBasePointData::GetDensityBounds, return FBoxSphereBounds(););
+	virtual FBox GetLocalDensityBounds(int32 InPointIndex) const PURE_VIRTUAL(UPCGBasePointData::GetLocalDensityBounds, return FBox(););
+	virtual FBox GetLocalBounds(int32 InPointIndex) const PURE_VIRTUAL(UPCGBasePointData::GetLocalBounds, return FBox(););
+	virtual FVector GetLocalCenter(int32 InPointIndex) const PURE_VIRTUAL(UPCGBasePointData::GetLocalCenter, return FVector::ZeroVector;);
+	virtual FVector GetExtents(int32 InPointIndex) const PURE_VIRTUAL(UPCGBasePointData::GetExtents, return FVector::ZeroVector;);
+	virtual FVector GetScaledExtents(int32 InPointIndex) const PURE_VIRTUAL(UPCGBasePointData::GetScaledExtents, return FVector::ZeroVector;);
+	virtual FVector GetLocalSize(int32 InPointIndex) const PURE_VIRTUAL(UPCGBasePointData::GetLocalSize, return FVector::ZeroVector;);
+	virtual FVector GetScaledLocalSize(int32 InPointIndex) const PURE_VIRTUAL(UPCGBasePointData::GetScaledLocalSize, return FVector::ZeroVector;);
+};
\ No newline at end of file
diff --git a/Engine/Plugins/PCG/Source/PCG/Public/Data/PCGPointData.h b/Engine/Plugins/PCG/Source/PCG/Public/Data/PCGPointData.h
--- a/Engine/Plugins/PCG/Source/PCG/Public/Data/PCGPointData.h
+++ b/Engine/Plugins/PCG/Source/PCG/Public/Data/PCGPointData.h
@@ -2,7 +2,7 @@
 
 #pragma once
 
-#include "PCGSpatialData.h"
+#include "Data/PCGBasePointData.h"
 
 #include "Math/GenericOctreePublic.h"
 #include "Math/GenericOctree.h"
@@ -25,7 +25,9 @@ namespace PCGPointDataConstants
 	const FName ActorReferenceAttribute = TEXT("ActorReference");
 }
 
-struct PCG_API FPCGPointRef
+PRAGMA_DISABLE_DEPRECATION_WARNINGS
+
+struct PCG_API UE_DEPRECATED(5.6, "Use PCGPointOctree::FPointRef instead") FPCGPointRef
 {
 	FPCGPointRef(const FPCGPoint& InPoint);
 	FPCGPointRef(const FPCGPoint& InPoint, const FBox& InBoundsOverride);
@@ -34,7 +36,7 @@ struct PCG_API FPCGPointRef
 	FBoxSphereBounds Bounds;
 };
 
-struct PCG_API FPCGPointRefSemantics
+struct PCG_API UE_DEPRECATED(5.6, "Use PCGPointOctree::FPointRefSemantics instead") FPCGPointRefSemantics
 {
 	enum { MaxElementsPerLeaf = 16 };
 	enum { MinInclusiveElementsPerNode = 7 };
@@ -63,14 +65,19 @@ struct PCG_API FPCGPointRefSemantics
 	}
 };
 
+PRAGMA_ENABLE_DEPRECATION_WARNINGS
+
 // TODO: Split this in "concrete" vs "api" class (needed for views)
 UCLASS(BlueprintType, ClassGroup = (Procedural))
-class PCG_API UPCGPointData : public UPCGSpatialData
+class PCG_API UPCGPointData : public UPCGBasePointData
 {
 	GENERATED_BODY()
 
 public:
+PRAGMA_DISABLE_DEPRECATION_WARNINGS
+	UE_DEPRECATED(5.6, "Use PCGPointOctree::FPointOctree instead")
 	typedef TOctree2<FPCGPointRef, FPCGPointRefSemantics> PointOctree;
+PRAGMA_ENABLE_DEPRECATION_WARNINGS
 
 	//~Begin UObject Interface
 	virtual void GetResourceSizeEx(FResourceSizeEx& CumulativeResourceSize) override;
@@ -87,6 +94,30 @@ public:
 	// Get the functions to the accessor factory
 	static FPCGAttributeAccessorMethods GetPointAccessorMethods();
 
+	// ~Begin UPCGBasePointData interface
+	virtual bool IsValidRef(const PCGPointOctree::FPointRef& InPointRef) const override { return Points.IsValidIndex(InPointRef.Index); }
+
+	virtual TArray<FTransform> GetTransformsCopy() const override;
+
+	virtual const FTransform& GetTransform(int32 InPointIndex) const override;
+	virtual float GetDensity(int32 InPointIndex) const override;
+	virtual const FVector& GetBoundsMin(int32 InPointIndex) const override;
+	virtual const FVector& GetBoundsMax(int32 InPointIndex) const override;
+	virtual const FVector4& GetColor(int32 InPointIndex) const override;
+	virtual float GetSteepness(int32 InPointIndex) const override;
+	virtual int32 GetSeed(int32 InPointIndex) const override;
+	virtual int64 GetMetadataEntry(int32 InPointIndex) const override;
+
+	virtual FBoxSphereBounds GetDensityBounds(int32 InPointIndex) const override;
+	virtual FBox GetLocalDensityBounds(int32 InPointIndex) const override;
+	virtual FBox GetLocalBounds(int32 InPointIndex) const override;
+	virtual FVector GetLocalCenter(int32 InPointIndex) const override;
+	virtual FVector GetExtents(int32 InPointIndex) const override;
+	virtual FVector GetScaledExtents(int32 InPointIndex) const override;
+	virtual FVector GetLocalSize(int32 InPointIndex) const override;
+	virtual FVector GetScaledLocalSize(int32 InPointIndex) const override;
+	// ~End UPCGBasePointData interface
+
 	// ~Begin UPCGSpatialData interface
 	virtual int GetDimension() const override { return 0; }
 	virtual FBox GetBounds() const override;
@@ -116,7 +147,7 @@ public:
 	TArray<FPCGPoint> GetPointsCopy() const { return Points; }
 
 	UFUNCTION(BlueprintCallable, Category = SpatialData)
-	int32 GetNumPoints() const { return Points.Num(); }
+	virtual int32 GetNumPoints() const override { return Points.Num(); }
 
 	UFUNCTION(BlueprintCallable, Category = SpatialData)
 	bool IsEmpty() const { return Points.IsEmpty(); }
@@ -134,15 +165,22 @@ public:
 	virtual void Flatten() override;
 
 	TArray<FPCGPoint>& GetMutablePoints();
-
-	const PointOctree& GetOctree() const;
+	
+	/** Get the dirty status of the Octree. Note that the Point Octree can be rebuilt from another thread, so this info can be invalidated at anytime. */
+	virtual bool IsPointOctreeDirty() const { return bOctreeIsDirty; }
+	virtual const PCGPointOctree::FPointOctree& GetPointOctree() const override;
 
 	bool ProjectPoint(const FTransform& InTransform, const FBox& InBounds, const FPCGProjectionParams& InParams, FPCGPoint& OutPoint, UPCGMetadata* OutMetadata, bool bUseBounds) const;
 
-	/** Get the dirty status of the Octree. Note that the Octree can be rebuilt from another thread, so this info can be invalidated at anytime. */
-	bool IsOctreeDirty() const { return bOctreeIsDirty; }
+PRAGMA_DISABLE_DEPRECATION_WARNINGS
+	UE_DEPRECATED(5.6, "Use GetPointOctree instead")
+	const PointOctree& GetOctree() const;
+	UE_DEPRECATED(5.6, "Use IsPointOctreeDirty instead")
+	bool IsOctreeDirty() const { return bOctreeOldIsDirty; }
+PRAGMA_ENABLE_DEPRECATION_WARNINGS
 
 protected:
+	void RebuildOctreeOld() const;
 	void RebuildOctree() const;
 	void RecomputeBounds() const;
 
@@ -150,8 +188,14 @@ protected:
 	TArray<FPCGPoint> Points;
 
 	mutable FCriticalSection CachedDataLock;
-	mutable PointOctree Octree;
+	mutable PCGPointOctree::FPointOctree PCGPointOctree;
 	mutable FBox Bounds; // TODO: review if this needs to be threadsafe
 	mutable bool bBoundsAreDirty = true;
+	mutable bool bOctreeOldIsDirty = true;
 	mutable bool bOctreeIsDirty = true;
+
+PRAGMA_DISABLE_DEPRECATION_WARNINGS
+	// deprecated
+	mutable PointOctree Octree;
+PRAGMA_ENABLE_DEPRECATION_WARNINGS
 };
diff --git a/Engine/Plugins/PCG/Source/PCG/Public/Elements/PCGCollapsePoints.h b/Engine/Plugins/PCG/Source/PCG/Public/Elements/PCGCollapsePoints.h
--- a/Engine/Plugins/PCG/Source/PCG/Public/Elements/PCGCollapsePoints.h
+++ b/Engine/Plugins/PCG/Source/PCG/Public/Elements/PCGCollapsePoints.h
@@ -5,6 +5,7 @@
 #include "PCGSettings.h"
 
 #include "Data/PCGPointData.h"
+#include "Utils/PCGPointOctree.h"
 #include "Elements/PCGTimeSlicedElementBase.h"
 #include "Metadata/PCGMetadataAttributeTpl.h"
 #include "Metadata/Accessors/PCGPropertyAccessor.h"
@@ -94,11 +95,11 @@ public:
 
 namespace PCGCollapsePoints
 {
+PRAGMA_DISABLE_DEPRECATION_WARNINGS
 	struct FCollapsePointsState
 	{
-		// Current iterative state during merge
-		TArray<FPCGPoint> Points;
-		UPCGPointData::PointOctree Octree;
+		TArray<FTransform> PointTransforms;
+		PCGPointOctree::FPointOctree PointOctree;
 
 		TArray<double> Weights;
 		TArray<int> Merged;
@@ -118,8 +119,14 @@ namespace PCGCollapsePoints
 		TArray<TUniquePtr<const IPCGAttributeAccessor>> SourceMergeAccessors;
 		TUniquePtr<const IPCGAttributeAccessorKeys> SourceMergeKeys;
 		TArray<TUniquePtr<IPCGAttributeAccessor>> OutputMergeAccessors;
-	};
 
+		UE_DEPRECATED(5.6, "Use PointTransforms instead")
+		TArray<FPCGPoint> Points;
+
+		UE_DEPRECATED(5.6, "Use PointOctree instead")
+		UPCGPointData::PointOctree Octree;
+	};
+		
 	struct FCollapsePointsSettings
 	{
 		const UPCGCollapsePointsSettings* Settings = nullptr;
@@ -130,17 +137,33 @@ namespace PCGCollapsePoints
 		using MergeSelectionFuncType = void(*)(FPCGContext* /*InContext*/, const FCollapsePointsSettings& /*Settings*/, FCollapsePointsState& /*OutState*/);
 		MergeSelectionFuncType MergeSelectionFunc = nullptr;
 
+		using GetPointPositionFuncType = FVector(*)(const TArray<FTransform>& /*InTransforms*/,const UPCGPointData* /*InPointData*/, int32 /*InPointIndex*/);
+		GetPointPositionFuncType GetPointPositionFunc = nullptr;
+
+		using GetPointReferenceFuncType = PCGPointOctree::FPointRef(*)(const TArray<FTransform>& /*InTransforms*/, const UPCGPointData* /*InPointData*/, int32 /*InPointIndex*/);
+		GetPointReferenceFuncType GetPointReferenceFunc = nullptr;
+
+		using GetPointSearchBoundsFuncType = FBoxCenterAndExtent(*)(const TArray<FTransform>& /*InTransforms*/, const UPCGPointData* /*InPointData*/, int32 /*InPointIndex*/, const double& /*InExtents*/);
+		GetPointSearchBoundsFuncType GetPointSearchBoundsFunc = nullptr;
+
+		double DistanceThreshold = 0.0;
+
 		using GetPositionFuncType = FVector(*)(const FPCGPoint&);
+		
+		UE_DEPRECATED(5.6, "Use GetPointPositionFunc instead")
 		GetPositionFuncType GetPositionFunc = nullptr;
 
 		using GetPointRefFuncType = FPCGPointRef(*)(const FPCGPoint&);
+		
+		UE_DEPRECATED(5.6, "Use GetPointSearchBoundsFunc instead")
 		GetPointRefFuncType GetPointRefFunc = nullptr;
 
 		using GetSearchBoundsFuncType = FBoxCenterAndExtent(*)(const FPCGPoint&, const double&);
+		
+		UE_DEPRECATED(5.6, "Use GetPointSearchBoundsFunc instead")
 		GetSearchBoundsFuncType GetSearchBoundsFunc = nullptr;
-
-		double DistanceThreshold = 0.0;
 	};
+PRAGMA_ENABLE_DEPRECATION_WARNINGS
 } // namespace PCGCollapsePoints
 
 class FPCGCollapsePointsElement : public TPCGTimeSlicedElementBase<PCGCollapsePoints::FCollapsePointsSettings, PCGCollapsePoints::FCollapsePointsState>
diff --git a/Engine/Plugins/PCG/Source/PCG/Public/Elements/PCGSelfPruning.h b/Engine/Plugins/PCG/Source/PCG/Public/Elements/PCGSelfPruning.h
--- a/Engine/Plugins/PCG/Source/PCG/Public/Elements/PCGSelfPruning.h
+++ b/Engine/Plugins/PCG/Source/PCG/Public/Elements/PCGSelfPruning.h
@@ -69,18 +69,24 @@ namespace PCGSelfPruningElement
 	struct FPointBitSet
 	{
 		TArray<uint32> Bits;
-		const FPCGPoint* FirstPoint = nullptr;
 
-		void Initialize(const TArray<FPCGPoint>& Points);
-		void Add(const FPCGPoint* Point);
-		bool Contains(const FPCGPoint* Point);
+		void Initialize(int32 NumPoints);
+		void Add(int32 Index);
+		bool Contains(int32 Index);
+
+		UE_DEPRECATED(5.6, "Use int32 param version")
+		void Initialize(const TArray<FPCGPoint>& Points) {}
+		UE_DEPRECATED(5.6, "Use int32 param version")
+		void Add(const FPCGPoint* Point) {}
+		UE_DEPRECATED(5.6, "Use int32 param version")
+		bool Contains(const FPCGPoint* Point) { return false; }
 	};
 
 	struct FIterationState
 	{
 		const UPCGPointData* InputPointData = nullptr;
 		UPCGPointData* OutputPointData = nullptr;
-		TArray<FPCGPointRef> SortedPoints;
+		TArray<PCGPointOctree::FPointRef> SortedPointRefs;
 		FPointBitSet ExcludedPoints;
 		FPointBitSet ExclusionPoints;
 		int32 CurrentPointIndex = 0;
@@ -90,6 +96,11 @@ namespace PCGSelfPruningElement
 		// Additionally, in the cases where two points to compare use the same instance, we'll have to do a temporary copy.
 		FPCGCollisionWrapper CollisionWrapper;
 		TMap<FBodyInstance*, FBodyInstance*> TemporaryBodyInstances;
+
+PRAGMA_DISABLE_DEPRECATION_WARNINGS
+		UE_DEPRECATED(5.6, "Use SortedPointRefs instead")
+		TArray<FPCGPointRef> SortedPoints;
+PRAGMA_ENABLE_DEPRECATION_WARNINGS
 	};
 
 	/** Will do the self pruning on all inputs in the context. Blocking call (no time slicing). Be careful, it can be very costly. */
diff --git a/Engine/Plugins/PCG/Source/PCG/Public/SpatialAlgo/PCGOctreeQueries.h b/Engine/Plugins/PCG/Source/PCG/Public/SpatialAlgo/PCGOctreeQueries.h
--- a/Engine/Plugins/PCG/Source/PCG/Public/SpatialAlgo/PCGOctreeQueries.h
+++ b/Engine/Plugins/PCG/Source/PCG/Public/SpatialAlgo/PCGOctreeQueries.h
@@ -21,8 +21,8 @@ private:
 
 public:
 	/** Query the internal octree to call a func on all points within some sphere. Callback takes the point (as const ref) and its squared distance to the center. */
-	template <typename Func = TFunctionRef<void(const FPCGPoint& Point, double SquaredDistance)>>
-	static void ForEachPointInsideSphere(const UPCGPointData* InPointData, const FVector& InCenter, const double InRadius, Func Callback);
+	template <typename Func = TFunctionRef<void(const UPCGBasePointData* PointData, int32 Index, double SquaredDistance)>>
+	static PCG_API void ForEachPointInsideSphere(const UPCGBasePointData* InPointData, const FVector& InCenter, const double InRadius, Func InCallback);
 
 	/** Get the closest point to a given position within the search distance, If bInDiscardCenter is true, will reject any points that are exactly at the center. Returns null if not found. */
 	static PCG_API const FPCGPoint* GetClosestPoint(const UPCGPointData* InPointData, const FVector& InCenter, const bool bInDiscardCenter, const double InSearchDistance = DefaultSearchDistance);
@@ -65,7 +65,7 @@ private:
 };
 
 template <typename Func>
-void UPCGOctreeQueries::ForEachPointInsideSphere(const UPCGPointData* InPointData, const FVector& InCenter, const double InRadius, Func Callback)
+void UPCGOctreeQueries::ForEachPointInsideSphere(const UPCGBasePointData* InPointData, const FVector& InCenter, const double InRadius, Func Callback)
 {
 	if (!InPointData)
 	{
@@ -75,18 +75,30 @@ void UPCGOctreeQueries::ForEachPointInsideSphere(const UPCGPointData* InPointDat
 	const double Extents = UE_DOUBLE_SQRT_2 * InRadius;
 	const double SquaredRadius = InRadius * InRadius;
 	FBoxCenterAndExtent SearchBounds(InCenter, FVector(Extents, Extents, Extents));
-
-	InPointData->GetOctree().FindElementsWithBoundsTest(SearchBounds, [&InCenter, SquaredRadius, Callback = std::move(Callback)](const FPCGPointRef& PointRef)
+	
+	// To support previous callback
+	const UPCGPointData* PointData = Cast<UPCGPointData>(InPointData);
+	const TArray<FPCGPoint>* Points = PointData ? &PointData->GetPoints() : nullptr;
+		
+	InPointData->GetPointOctree().FindElementsWithBoundsTest(SearchBounds, [InPointData, Points, &InCenter, SquaredRadius, Callback = std::move(Callback)](const PCGPointOctree::FPointRef& PointRef)
 	{
-		if (!PointRef.Point)
+		if (!InPointData->IsValidRef(PointRef))
 		{
 			return;
 		}
 
-		const double SquaredDistance = FVector::DistSquared(InCenter, PointRef.Point->Transform.GetLocation());
+		const double SquaredDistance = FVector::DistSquared(InCenter, Points ? (*Points)[PointRef.Index].Transform.GetLocation() : InPointData->GetTransform(PointRef.Index).GetLocation());
 		if (SquaredDistance <= SquaredRadius)
 		{
-			Callback(*PointRef.Point, SquaredDistance);
+			if constexpr (std::is_invocable_v<Func, const FPCGPoint&, double>)
+			{
+				check(Points);
+				Callback((*Points)[PointRef.Index], SquaredDistance);
+			}
+			else
+			{
+				Callback(InPointData, PointRef.Index, SquaredDistance);
+			}
 		}
 	});
 }
\ No newline at end of file
diff --git a/Engine/Plugins/PCG/Source/PCG/Public/Utils/PCGPointOctree.h b/Engine/Plugins/PCG/Source/PCG/Public/Utils/PCGPointOctree.h
new file mode 100644
--- /dev/null
+++ b/Engine/Plugins/PCG/Source/PCG/Public/Utils/PCGPointOctree.h
@@ -0,0 +1,54 @@
+// Copyright Epic Games, Inc. All Rights Reserved.
+
+#pragma once
+
+#include "Math/GenericOctreePublic.h"
+#include "Math/GenericOctree.h"
+
+namespace PCGPointOctree
+{
+	struct PCG_API FPointRef
+	{
+		FPointRef() = default;
+
+		FPointRef(int32 InIndex, const FBoxSphereBounds& InBounds)
+		{
+			Index = InIndex;
+			Bounds = InBounds;
+		}
+
+		int32 Index = INDEX_NONE;
+		FBoxSphereBounds Bounds;
+	};
+
+	struct PCG_API FPointRefSemantics
+	{
+		enum { MaxElementsPerLeaf = 16 };
+		enum { MinInclusiveElementsPerNode = 7 };
+		enum { MaxNodeDepth = 12 };
+
+		typedef TInlineAllocator<MaxElementsPerLeaf> ElementAllocator;
+
+		FORCEINLINE static const FBoxSphereBounds& GetBoundingBox(const FPointRef& InPoint)
+		{
+			return InPoint.Bounds;
+		}
+
+		FORCEINLINE static const bool AreElementsEqual(const FPointRef& A, const FPointRef& B)
+		{
+			// TODO: verify if that's sufficient
+			return A.Index == B.Index;
+		}
+
+		FORCEINLINE static void ApplyOffset(FPointRef& InPoint)
+		{
+			ensureMsgf(false, TEXT("Not implemented"));
+		}
+
+		FORCEINLINE static void SetElementId(const FPointRef& Element, FOctreeElementId2 OctreeElementID)
+		{
+		}
+	};
+
+	typedef TOctree2<FPointRef, FPointRefSemantics> FPointOctree;
+}
\ No newline at end of file
