[PCG] Made EPCGPrintVerbosity BlueprintType so it can be used as a graph parameter, and split NoLogging/PrintToScreen

#rnx

diff --git a/Engine/Plugins/PCG/Source/PCG/Private/Elements/PCGPrintElement.cpp b/Engine/Plugins/PCG/Source/PCG/Private/Elements/PCGPrintElement.cpp
--- a/Engine/Plugins/PCG/Source/PCG/Private/Elements/PCGPrintElement.cpp
+++ b/Engine/Plugins/PCG/Source/PCG/Private/Elements/PCGPrintElement.cpp
@@ -69,7 +69,12 @@ bool FPCGPrintElement::ExecuteInternal(FPCGContext* Context) const
 	const UPCGPrintElementSettings* Settings = Context->GetInputSettings<UPCGPrintElementSettings>();
 	check(Settings);
 
-	if (!Settings->bEnablePrint)
+	if (!Settings->bEnablePrint
+		|| (Settings->Verbosity == EPCGPrintVerbosity::NoLogging
+#if WITH_EDITOR
+			&& !Settings->bPrintToScreen
+#endif // WITH_EDITOR
+		))
 	{
 		return true;
 	}
@@ -141,6 +146,8 @@ bool FPCGPrintElement::ExecuteInternal(FPCGContext* Context) const
 		case EPCGPrintVerbosity::Display:
 			PCGE_LOG_C(Display, LogOnly, Context, FText::FromString(FinalString));
 			break;
+		case EPCGPrintVerbosity::NoLogging:
+			break;
 		default:
 			PCGLog::LogErrorOnGraph(LOCTEXT("InvalidVerbosity", "Invalid verbosity selected. Please use Log, Display, Warning, or Error. See EPCGPrintVerbosity for more details."), Context);
 			return true;
diff --git a/Engine/Plugins/PCG/Source/PCG/Public/Elements/PCGPrintElement.h b/Engine/Plugins/PCG/Source/PCG/Public/Elements/PCGPrintElement.h
--- a/Engine/Plugins/PCG/Source/PCG/Public/Elements/PCGPrintElement.h
+++ b/Engine/Plugins/PCG/Source/PCG/Public/Elements/PCGPrintElement.h
@@ -11,9 +11,10 @@
 
 #include "PCGPrintElement.generated.h"
 
-UENUM()
+UENUM(BlueprintType)
 enum class EPCGPrintVerbosity : uint8
 {
+	NoLogging = ELogVerbosity::NoLogging, // 0
 	Log = ELogVerbosity::Log,
 	Warning = ELogVerbosity::Warning,
 	Error = ELogVerbosity::Error,
