[PCG] Update WRHQ with Sweep support

diff --git a/Engine/Plugins/PCG/Source/PCG/Private/Data/PCGWorldData.cpp b/Engine/Plugins/PCG/Source/PCG/Private/Data/PCGWorldData.cpp
--- a/Engine/Plugins/PCG/Source/PCG/Private/Data/PCGWorldData.cpp
+++ b/Engine/Plugins/PCG/Source/PCG/Private/Data/PCGWorldData.cpp
@@ -285,17 +285,24 @@ bool UPCGWorldRayHitData::SamplePoint(const FTransform& InTransform, const FBox&
 	check(World.IsValid());
 
 	// Todo: consider prebuilding this
-	FCollisionObjectQueryParams ObjectQueryParams(QueryParams.CollisionChannel);
-	FCollisionQueryParams Params; // TODO: apply properties from the settings when/if they exist
-	Params.bTraceComplex = QueryParams.bTraceComplex;
-	Params.bReturnPhysicalMaterial = QueryParams.bGetReferenceToPhysicalMaterial;
-	Params.bReturnFaceIndex = QueryParams.bGetFaceIndex || QueryParams.bGetUVCoords;
+	const FCollisionObjectQueryParams ObjectQueryParams(QueryParams.CollisionChannel);
+	FCollisionQueryParams CollisionQueryParams; // TODO: apply properties from the settings when/if they exist
+	CollisionQueryParams.bTraceComplex = QueryParams.bTraceComplex;
+	CollisionQueryParams.bReturnPhysicalMaterial = QueryParams.bGetReferenceToPhysicalMaterial;
+	CollisionQueryParams.bReturnFaceIndex = QueryParams.bGetFaceIndex || QueryParams.bGetUVCoords;
 
 	FVector RayOrigin = InTransform.GetLocation() - ((InTransform.GetLocation() - QueryParams.RayOrigin) | QueryParams.RayDirection) * QueryParams.RayDirection;
 	FVector RayEnd = RayOrigin + QueryParams.RayDirection * QueryParams.RayLength;
 
 	TArray<FHitResult> Hits;
-	if (World->LineTraceMultiByObjectType(Hits, RayOrigin, RayEnd, ObjectQueryParams, Params))
+	if (World->SweepMultiByObjectType(
+				Hits,
+				RayOrigin,
+				RayEnd,
+				FQuat(CollisionShape.ShapeRotation),
+				ObjectQueryParams,
+				CollisionShape.ToCollisionShape(),
+				CollisionQueryParams))
 	{
 		TOptional<FHitResult> HitResult = PCGWorldQueryHelpers::FilterRayHitResults(&QueryParams, OriginatingComponent, Hits, ActorFilter.GetCachedFilterActors());
 
diff --git a/Engine/Plugins/PCG/Source/PCG/Private/Elements/PCGWorldQuery.cpp b/Engine/Plugins/PCG/Source/PCG/Private/Elements/PCGWorldQuery.cpp
--- a/Engine/Plugins/PCG/Source/PCG/Private/Elements/PCGWorldQuery.cpp
+++ b/Engine/Plugins/PCG/Source/PCG/Private/Elements/PCGWorldQuery.cpp
@@ -209,6 +209,7 @@ bool FPCGWorldRayHitQueryElement::ExecuteInternal(FPCGContext* Context) const
 	Data->QueryParams = QueryParams;
 	Data->QueryParams.Initialize();
 	Data->OriginatingComponent = Context->SourceComponent;
+	Data->CollisionShape = Settings->CollisionShape;
 	Data->Initialize(World, Transform, /*InBounds=*/FBox(EForceInit::ForceInit), LocalBounds);
 
 	bool bHasLandscapeMetadata = false;
diff --git a/Engine/Plugins/PCG/Source/PCG/Public/Data/PCGWorldData.h b/Engine/Plugins/PCG/Source/PCG/Public/Data/PCGWorldData.h
--- a/Engine/Plugins/PCG/Source/PCG/Public/Data/PCGWorldData.h
+++ b/Engine/Plugins/PCG/Source/PCG/Public/Data/PCGWorldData.h
@@ -2,6 +2,7 @@
 
 #pragma once
 
+#include "PCGCollisionShape.h"
 #include "PCGVolumeData.h"
 #include "PCGSurfaceData.h"
 #include "Helpers/PCGWorldQueryHelpers.h"
@@ -372,6 +373,10 @@ public:
 	UPROPERTY()
 	FBox Bounds = FBox(EForceInit::ForceInit);
 
+	/** Parameters for either using a line trace or specifying a collision shape for a sweep. */
+	UPROPERTY(BlueprintReadWrite, EditAnywhere, Category = Data, meta = (ShowOnlyInnerProperties))
+	FPCGCollisionShape CollisionShape;
+
 	UPROPERTY(BlueprintReadWrite, EditAnywhere, Category = Data, meta = (ShowOnlyInnerProperties))
 	FPCGWorldRayHitQueryParams QueryParams;
 	
diff --git a/Engine/Plugins/PCG/Source/PCG/Public/Elements/PCGWorldQuery.h b/Engine/Plugins/PCG/Source/PCG/Public/Elements/PCGWorldQuery.h
--- a/Engine/Plugins/PCG/Source/PCG/Public/Elements/PCGWorldQuery.h
+++ b/Engine/Plugins/PCG/Source/PCG/Public/Elements/PCGWorldQuery.h
@@ -3,6 +3,7 @@
 #pragma once
 
 #include "PCGSettings.h"
+#include "Data/PCGCollisionShape.h"
 #include "Data/PCGWorldData.h"
 
 #include "PCGWorldQuery.generated.h"
@@ -71,6 +72,10 @@ protected:
 	//~End UPCGSettings interface
 
 public:
+	/** Parameters for either using a line trace or specifying a collision shape for a sweep. */
+	UPROPERTY(BlueprintReadWrite, EditAnywhere, Category = Settings, meta = (ShowOnlyInnerProperties, PCG_Overridable))
+	FPCGCollisionShape CollisionShape;
+
 	UPROPERTY(BlueprintReadWrite, EditAnywhere, Category = Data, meta = (PCG_Overridable, ShowOnlyInnerProperties))
 	FPCGWorldRayHitQueryParams QueryParams;
 };
