Added missing tooltips in tools menu

#rnx

diff --git a/Engine/Plugins/PCG/Source/PCGEditor/Private/PCGEditorModule.cpp b/Engine/Plugins/PCG/Source/PCGEditor/Private/PCGEditorModule.cpp
--- a/Engine/Plugins/PCG/Source/PCGEditor/Private/PCGEditorModule.cpp
+++ b/Engine/Plugins/PCG/Source/PCGEditor/Private/PCGEditorModule.cpp
@@ -504,113 +504,119 @@ void FPCGEditorModule::PopulateMenuActions(FMenuBuilder& MenuBuilder)
 			})),
 		NAME_None);
 
-	MenuBuilder.AddSubMenu(LOCTEXT("PCGSubMenuDelete", "Delete"), FText(), FNewMenuDelegate::CreateLambda([](FMenuBuilder& SubMenuBuilder)
-	{
-		SubMenuBuilder.AddMenuEntry(
-			LOCTEXT("DeletePCGPartitionActors", "All PCG Partition Grid Actors & Generated Actors"),
-			LOCTEXT("DeletePCGPartitionActors_Tooltip", "Deletes all PCG Partition Grid Actors and PCG Partition Generated Actors in the current world"),
-			FSlateIcon(),
-			FUIAction(
-				FExecuteAction::CreateLambda([]() 
-				{
-					if (GEditor)
-					{
-						if (UPCGSubsystem* PCGSubsystem = UPCGSubsystem::GetInstance(GEditor->GetEditorWorldContext().World()))
-						{
-							PCGSubsystem->DeleteSerializedPartitionActors(/*bOnlyDeleteUnused=*/false);
-						}
-					}
-				})),
-			NAME_None);
-
-		SubMenuBuilder.AddMenuEntry(
-			LOCTEXT("DeletePCGPartitionActorsChildren", "All PCG Partition Generated Actors"),
-			LOCTEXT("DeletePCGPartitionActorsChildren_Tooltip", "Deletes all PCG Partition Generated Actors in the current world (not the Partition Grid Actors themselves)"),
-			FSlateIcon(),
-			FUIAction(
-				FExecuteAction::CreateLambda([]() 
-				{
-					if (GEditor)
+	MenuBuilder.AddSubMenu(
+		LOCTEXT("PCGSubMenuDelete", "Delete"), 
+		LOCTEXT("PCGSubMenuDelete_Tooltip", "Editor commands to delete PCG Partition Actors & World Actors."),
+		FNewMenuDelegate::CreateLambda([](FMenuBuilder& SubMenuBuilder)
+		{
+			SubMenuBuilder.AddMenuEntry(
+				LOCTEXT("DeletePCGPartitionActors", "All PCG Partition Grid Actors & Generated Actors"),
+				LOCTEXT("DeletePCGPartitionActors_Tooltip", "Deletes all PCG Partition Grid Actors and PCG Partition Generated Actors in the current world"),
+				FSlateIcon(),
+				FUIAction(
+					FExecuteAction::CreateLambda([]() 
 					{
-						if (UPCGSubsystem* PCGSubsystem = UPCGSubsystem::GetInstance(GEditor->GetEditorWorldContext().World()))
+						if (GEditor)
 						{
-							PCGSubsystem->DeleteSerializedPartitionActors(/*bOnlyDeleteUnused=*/false, /*bOnlyChildren=*/true);
+							if (UPCGSubsystem* PCGSubsystem = UPCGSubsystem::GetInstance(GEditor->GetEditorWorldContext().World()))
+							{
+								PCGSubsystem->DeleteSerializedPartitionActors(/*bOnlyDeleteUnused=*/false);
+							}
 						}
-					}
-				})),
-			NAME_None);
-
-		SubMenuBuilder.AddMenuEntry(
-			LOCTEXT("DeletePCGWorldActor", "All PCG World Actors"),
-			LOCTEXT("DeletePCGWorldActor_Tooltip", "Deletes all PCG World Actors (This also deletes all PCG Partition Grid Actors & Generated Actors)"),
-			FSlateIcon(),
-			FUIAction(
-				FExecuteAction::CreateLambda([]()
+					})),
+				NAME_None);
+
+			SubMenuBuilder.AddMenuEntry(
+				LOCTEXT("DeletePCGPartitionActorsChildren", "All PCG Partition Generated Actors"),
+				LOCTEXT("DeletePCGPartitionActorsChildren_Tooltip", "Deletes all PCG Partition Generated Actors in the current world (not the Partition Grid Actors themselves)"),
+				FSlateIcon(),
+				FUIAction(
+					FExecuteAction::CreateLambda([]() 
 					{
 						if (GEditor)
 						{
 							if (UPCGSubsystem* PCGSubsystem = UPCGSubsystem::GetInstance(GEditor->GetEditorWorldContext().World()))
 							{
-								PCGSubsystem->DestroyAllPCGWorldActors();
+								PCGSubsystem->DeleteSerializedPartitionActors(/*bOnlyDeleteUnused=*/false, /*bOnlyChildren=*/true);
 							}
 						}
 					})),
-			NAME_None);
-
-		SubMenuBuilder.AddMenuEntry(
-			LOCTEXT("DeleteUnusedPCGPartitionActors", "Unused PCG Partition Grid Actors"),
-			LOCTEXT("DeleteUnusedPCGPartitionActors_Tooltip", "Deletes unused PCG Partition Grid Actors in the current world"),
-			FSlateIcon(),
-			FUIAction(
-				FExecuteAction::CreateLambda([]() 
-				{
-					if (GEditor)
+				NAME_None);
+
+			SubMenuBuilder.AddMenuEntry(
+				LOCTEXT("DeletePCGWorldActor", "All PCG World Actors"),
+				LOCTEXT("DeletePCGWorldActor_Tooltip", "Deletes all PCG World Actors (This also deletes all PCG Partition Grid Actors & Generated Actors)"),
+				FSlateIcon(),
+				FUIAction(
+					FExecuteAction::CreateLambda([]()
+						{
+							if (GEditor)
+							{
+								if (UPCGSubsystem* PCGSubsystem = UPCGSubsystem::GetInstance(GEditor->GetEditorWorldContext().World()))
+								{
+									PCGSubsystem->DestroyAllPCGWorldActors();
+								}
+							}
+						})),
+				NAME_None);
+
+			SubMenuBuilder.AddMenuEntry(
+				LOCTEXT("DeleteUnusedPCGPartitionActors", "Unused PCG Partition Grid Actors"),
+				LOCTEXT("DeleteUnusedPCGPartitionActors_Tooltip", "Deletes unused PCG Partition Grid Actors in the current world"),
+				FSlateIcon(),
+				FUIAction(
+					FExecuteAction::CreateLambda([]() 
 					{
-						if (UPCGSubsystem* PCGSubsystem = UPCGSubsystem::GetInstance(GEditor->GetEditorWorldContext().World()))
+						if (GEditor)
 						{
-							PCGSubsystem->DeleteSerializedPartitionActors(/*bOnlyDeleteUnused=*/true);
+							if (UPCGSubsystem* PCGSubsystem = UPCGSubsystem::GetInstance(GEditor->GetEditorWorldContext().World()))
+							{
+								PCGSubsystem->DeleteSerializedPartitionActors(/*bOnlyDeleteUnused=*/true);
+							}
 						}
-					}
-				})),
-			NAME_None);
-	}));
+					})),
+				NAME_None);
+		}));
 
-	MenuBuilder.AddSubMenu(LOCTEXT("PCGSubMenuLandscape", "Landscape"), FText(), FNewMenuDelegate::CreateLambda([](FMenuBuilder& SubMenuBuilder)
-	{
-		SubMenuBuilder.AddMenuEntry(
-			LOCTEXT("BuildLandscapeCache", "Build Cache"),
-			LOCTEXT("BuildLandscapeCache_Tooltip", "Caches the landscape data in the PCG World Actor"),
-			FSlateIcon(),
-			FUIAction(
-				FExecuteAction::CreateLambda([]() 
-				{
-					if (GEditor)
+	MenuBuilder.AddSubMenu(
+		LOCTEXT("PCGSubMenuLandscape", "Landscape"), 
+		LOCTEXT("PCGSubMenuLandscape_Tooltip", "PCG Landscape cache related editor commands"),
+		FNewMenuDelegate::CreateLambda([](FMenuBuilder& SubMenuBuilder)
+		{
+			SubMenuBuilder.AddMenuEntry(
+				LOCTEXT("BuildLandscapeCache", "Build Cache"),
+				LOCTEXT("BuildLandscapeCache_Tooltip", "Caches the landscape data in the PCG World Actor"),
+				FSlateIcon(),
+				FUIAction(
+					FExecuteAction::CreateLambda([]() 
 					{
-						if (UPCGSubsystem* PCGSubsystem = UPCGSubsystem::GetInstance(GEditor->GetEditorWorldContext().World()))
+						if (GEditor)
 						{
-							PCGSubsystem->BuildLandscapeCache();
+							if (UPCGSubsystem* PCGSubsystem = UPCGSubsystem::GetInstance(GEditor->GetEditorWorldContext().World()))
+							{
+								PCGSubsystem->BuildLandscapeCache();
+							}
 						}
-					}
-				})),
-			NAME_None);
-
-		SubMenuBuilder.AddMenuEntry(
-			LOCTEXT("ClearLandscapeCache", "Clear Cache"),
-			LOCTEXT("ClearLandscapeCache_Tooltip", "Clears the landscape data cache in the PCG World Actor"),
-			FSlateIcon(),
-			FUIAction(
-				FExecuteAction::CreateLambda([]() 
-				{
-					if (GEditor)
+					})),
+				NAME_None);
+
+			SubMenuBuilder.AddMenuEntry(
+				LOCTEXT("ClearLandscapeCache", "Clear Cache"),
+				LOCTEXT("ClearLandscapeCache_Tooltip", "Clears the landscape data cache in the PCG World Actor"),
+				FSlateIcon(),
+				FUIAction(
+					FExecuteAction::CreateLambda([]() 
 					{
-						if (UPCGSubsystem* PCGSubsystem = UPCGSubsystem::GetInstance(GEditor->GetEditorWorldContext().World()))
+						if (GEditor)
 						{
-							PCGSubsystem->ClearLandscapeCache();
+							if (UPCGSubsystem* PCGSubsystem = UPCGSubsystem::GetInstance(GEditor->GetEditorWorldContext().World()))
+							{
+								PCGSubsystem->ClearLandscapeCache();
+							}
 						}
-					}
-				})),
-			NAME_None);
-	}));
+					})),
+				NAME_None);
+		}));
 	
 	MenuBuilder.AddMenuEntry(
 		LOCTEXT("CancelAllGeneration", "Cancel all PCG tasks"),
