[PCG] UX Spawn Node Shortcuts for PCG Graph

diff --git a/Engine/Plugins/PCG/Config/EditorPerProjectUserSettings.ini b/Engine/Plugins/PCG/Config/EditorPerProjectUserSettings.ini
new file mode 100644
--- /dev/null
+++ b/Engine/Plugins/PCG/Config/EditorPerProjectUserSettings.ini
@@ -0,0 +1,27 @@
+[PCGEditorSpawnNodes]
+; Preset
++Node = (Class=/Script/PCG.PCGCreateAttributeSetSettings Key=One Shift=false Ctrl=false Alt=false Index=1 Label="Create Constant (Double)")
++Node = (Class=/Script/PCG.PCGCreateAttributeSetSettings Key=Two Shift=false Ctrl=false Alt=false Index=4 Label="Create Constant (Vector 2D)")
++Node = (Class=/Script/PCG.PCGCreateAttributeSetSettings Key=Three Shift=false Ctrl=false Alt=false Index=5 Label="Create Constant (Vector 3D)")
++Node = (Class=/Script/PCG.PCGCreateAttributeSetSettings Key=Four Shift=false Ctrl=false Alt=false Index=6 Label="Create Constant (Vector 4D)")
++Node = (Class=/Script/PCG.PCGDebugSettings Key=D Shift=true Ctrl=false Alt=false)
++Node = (Class=/Script/PCG.PCGSpawnActorSettings Key=A Shift=true Ctrl=false Alt=false)
++Node = (Class=/Script/PCG.PCGStaticMeshSpawnerSettings Key=M Shift=false Ctrl=false Alt=false)
++Node = (Class=/Script/PCG.PCGPrintElementSettings Key=P Shift=false Ctrl=false Alt=false)
++Node = (Class=/Script/PCG.PCGBranchSettings Key=B Shift=false Ctrl=false Alt=false)
++Node = (Class=/Script/PCG.PCGGetLandscapeSettings Key=L Shift=false Ctrl=false Alt=false)
++Node = (Class=/Script/PCG.PCGGetVolumeSettings Key=V Shift=false Ctrl=false Alt=false)
++Node = (Class=/Script/PCG.PCGGetPrimitiveSettings Key=R Shift=false Ctrl=false Alt=false)
++Node = (Class=/Script/PCG.PCGGetSplineSettings Key=S Shift=false Ctrl=false Alt=false)
++Node = (Class=/Script/PCG.PCGCreateAttributeSetSettings Key=Z Shift=false Ctrl=false Alt=false Index=9 Label="Create Constant (String)")
+; Uninitialized
++Node = (Class=/Script/PCG.PCGCreateAttributeSetSettings Key= Shift=false Ctrl=false Alt=false Index=0 Label="Create Constant (Float)")
++Node = (Class=/Script/PCG.PCGCreateAttributeSetSettings Key= Shift=false Ctrl=false Alt=false Index=2 Label="Create Constant (Int32)")
++Node = (Class=/Script/PCG.PCGCreateAttributeSetSettings Key= Shift=false Ctrl=false Alt=false Index=3 Label="Create Constant (Int64)")
++Node = (Class=/Script/PCG.PCGCreateAttributeSetSettings Key= Shift=false Ctrl=false Alt=false Index=7 Label="Create Constant (Quaternion)")
++Node = (Class=/Script/PCG.PCGCreateAttributeSetSettings Key= Shift=false Ctrl=false Alt=false Index=8 Label="Create Constant (Transform)")
++Node = (Class=/Script/PCG.PCGCreateAttributeSetSettings Key= Shift=false Ctrl=false Alt=false Index=10 Label="Create Constant (Boolean)")
++Node = (Class=/Script/PCG.PCGCreateAttributeSetSettings Key= Shift=false Ctrl=false Alt=false Index=11 Label="Create Constant (Rotator)")
++Node = (Class=/Script/PCG.PCGCreateAttributeSetSettings Key= Shift=false Ctrl=false Alt=false Index=12 Label="Create Constant (Name)")
++Node = (Class=/Script/PCG.PCGCreateAttributeSetSettings Key= Shift=false Ctrl=false Alt=false Index=13 Label="Create Constant (SoftObjectPath)")
++Node = (Class=/Script/PCG.PCGCreateAttributeSetSettings Key= Shift=false Ctrl=false Alt=false Index=14 Label="Create Constant (SoftClassPath)")
\ No newline at end of file
diff --git a/Engine/Plugins/PCG/Source/PCG/Private/Elements/PCGCreateAttribute.cpp b/Engine/Plugins/PCG/Source/PCG/Private/Elements/PCGCreateAttribute.cpp
--- a/Engine/Plugins/PCG/Source/PCG/Private/Elements/PCGCreateAttribute.cpp
+++ b/Engine/Plugins/PCG/Source/PCG/Private/Elements/PCGCreateAttribute.cpp
@@ -278,6 +278,16 @@ FString UPCGCreateAttributeSetSettings::GetAdditionalTitleInformation() const
 	return FString::Printf(TEXT("%s: %s"), *OutputTarget.GetName().ToString(), *AttributeTypes.ToString());
 }
 
+void UPCGCreateAttributeSetSettings::ApplyPreconfiguredSettings(const FPCGPreConfiguredSettingsInfo& PreconfigureInfo)
+{
+	if (PreconfigureInfo.PreconfiguredIndex < 0 || PreconfigureInfo.PreconfiguredIndex >= static_cast<uint8>(EPCGMetadataTypes::Count))
+	{
+		return;
+	}
+
+	AttributeTypes.Type = static_cast<EPCGMetadataTypes>(PreconfigureInfo.PreconfiguredIndex);
+}
+
 #if WITH_EDITOR
 TArray<FPCGPreconfiguredInfo> UPCGCreateAttributeSetSettings::GetConversionInfo() const
 {
diff --git a/Engine/Plugins/PCG/Source/PCG/Public/Elements/PCGCreateAttribute.h b/Engine/Plugins/PCG/Source/PCG/Public/Elements/PCGCreateAttribute.h
--- a/Engine/Plugins/PCG/Source/PCG/Public/Elements/PCGCreateAttribute.h
+++ b/Engine/Plugins/PCG/Source/PCG/Public/Elements/PCGCreateAttribute.h
@@ -94,6 +94,7 @@ public:
 	virtual EPCGSettingsType GetType() const override { return EPCGSettingsType::Param; }
 #endif
 	virtual FString GetAdditionalTitleInformation() const override;
+	virtual void ApplyPreconfiguredSettings(const FPCGPreConfiguredSettingsInfo& PreconfigureInfo) override;
 
 protected:
 #if WITH_EDITOR
diff --git a/Engine/Plugins/PCG/Source/PCGEditor/Private/PCGEditor.cpp b/Engine/Plugins/PCG/Source/PCGEditor/Private/PCGEditor.cpp
--- a/Engine/Plugins/PCG/Source/PCGEditor/Private/PCGEditor.cpp
+++ b/Engine/Plugins/PCG/Source/PCGEditor/Private/PCGEditor.cpp
@@ -1429,6 +1429,21 @@ void FPCGEditor::SpawnNodeSourceEditorTab()
 	}
 }
 
+FReply FPCGEditor::OnSpawnNodeByShortcut(FInputChord InChord, const FVector2D& InPosition, UPCGEditorGraph* InGraph)
+{
+	const TSharedPtr<FEdGraphSchemaAction> Action = FPCGEditorSpawnNodeCommands::Get().GetGraphActionByChord(InChord);
+	if (Action.IsValid())
+	{
+		TArray<UEdGraphPin*> DummyPins;
+		Action->PerformAction(InGraph, DummyPins, InPosition);
+		return FReply::Handled();
+	}
+	else
+	{
+		return FReply::Unhandled();
+	}
+}
+
 void FPCGEditor::OnCollapseNodesInSubgraph()
 {
 	if (!InternalValidationOnAction())
@@ -2591,6 +2606,7 @@ TSharedRef<SGraphEditor> FPCGEditor::CreateGraphEditorWidget()
 	InEvents.OnSelectionChanged = SGraphEditor::FOnSelectionChanged::CreateSP(this, &FPCGEditor::OnSelectedNodesChanged);
 	InEvents.OnTextCommitted = FOnNodeTextCommitted::CreateSP(this, &FPCGEditor::OnNodeTitleCommitted);
 	InEvents.OnNodeDoubleClicked = FSingleNodeEvent::CreateSP(this, &FPCGEditor::OnNodeDoubleClicked);
+	InEvents.OnSpawnNodeByShortcut = SGraphEditor::FOnSpawnNodeByShortcut::CreateSP(this, &FPCGEditor::OnSpawnNodeByShortcut, PCGEditorGraph);
 
 	return SNew(SGraphEditor)
 		.AdditionalCommands(GraphEditorCommands)
diff --git a/Engine/Plugins/PCG/Source/PCGEditor/Private/PCGEditor.h b/Engine/Plugins/PCG/Source/PCGEditor/Private/PCGEditor.h
--- a/Engine/Plugins/PCG/Source/PCGEditor/Private/PCGEditor.h
+++ b/Engine/Plugins/PCG/Source/PCGEditor/Private/PCGEditor.h
@@ -103,6 +103,17 @@ public:
 	/** Opens the Node Source Editor tab if it is not already visible. */
 	void SpawnNodeSourceEditorTab();
 
+	/**
+	 * Handles spawning a graph node in the current graph using the passed in chord.
+	 *
+	 * @param	InChord		Chord that was just performed.
+	 * @param	InPosition	Current cursor position.
+	 * @param	InGraph		Graph that chord was performed in.
+	 *
+	 * @return	FReply	Whether chord was handled.
+	 */
+	FReply OnSpawnNodeByShortcut(FInputChord InChord, const FVector2D& InPosition, UPCGEditorGraph* InGraph);
+
 	FOnInspectedStackChanged OnInspectedStackChangedDelegate;
 
 protected:
diff --git a/Engine/Plugins/PCG/Source/PCGEditor/Private/PCGEditorCommands.cpp b/Engine/Plugins/PCG/Source/PCGEditor/Private/PCGEditorCommands.cpp
--- a/Engine/Plugins/PCG/Source/PCGEditor/Private/PCGEditorCommands.cpp
+++ b/Engine/Plugins/PCG/Source/PCGEditor/Private/PCGEditorCommands.cpp
@@ -2,6 +2,11 @@
 
 #include "PCGEditorCommands.h"
 
+#include "PCGEditorGraphSchemaActions.h"
+
+#include "PCGSettings.h"
+
+#include "Misc/ConfigCacheIni.h"
 #include "Styling/AppStyle.h"
 
 #define LOCTEXT_NAMESPACE "PCGEditorCommands"
@@ -40,4 +45,102 @@ void FPCGEditorCommands::RegisterCommands()
 	UI_COMMAND(SelectNamedRerouteDeclaration, "Select Named Reroute Declaration", "Selects the associated Named Reroute Declaration matching this Named Reroute Usage", EUserInterfaceActionType::Button, FInputChord());
 }
 
+TSharedPtr<FEdGraphSchemaAction> FPCGSpawnNodeCommandInfo::GetAction() const
+{
+	TSharedPtr<FPCGEditorGraphSchemaAction_NewNativeElement> NewAction(new FPCGEditorGraphSchemaAction_NewNativeElement);
+	NewAction->SettingsClass = PCGSettingsClass;
+	NewAction->PreconfiguredInfo = PreconfiguredInfo;
+	return NewAction;
+}
+
+void FPCGEditorSpawnNodeCommands::RegisterCommands()
+{
+	const FString ConfigSection = TEXT("PCGEditorSpawnNodes");
+	const FString SettingsLabel = TEXT("Node");
+	TArray<FString> SpawnNodeConfigEntries;
+
+	GConfig->GetArray(*ConfigSection, *SettingsLabel, SpawnNodeConfigEntries, GEditorPerProjectIni);
+
+	for (const FString& CurrentEntry : SpawnNodeConfigEntries)
+	{
+		FString ClassName;
+		if (!FParse::Value(*CurrentEntry, TEXT("Class="), ClassName))
+		{
+			// Could not find a class name
+			continue;
+		}
+
+		UClass* FoundClass = UClass::TryFindTypeSlow<UClass>(ClassName, EFindFirstObjectOptions::ExactClass);
+		if (!FoundClass || !FoundClass->IsChildOf(UPCGSettings::StaticClass()))
+		{
+			continue;
+		}
+
+		TSharedPtr<FPCGSpawnNodeCommandInfo> SpawnCommandInfo = MakeShareable(new FPCGSpawnNodeCommandInfo(FoundClass));
+
+		// If command info was created, set up a UI Command for keybinding.
+		if (SpawnCommandInfo.IsValid())
+		{
+			FKey Key;
+
+			// Parse the keybinding information.
+			FString KeyString;
+			if (FParse::Value(*CurrentEntry, TEXT("Key="), KeyString))
+			{
+				Key = *KeyString;
+			}
+
+			FInputChord Chord;
+			// Parse the chord.
+			if (Key.IsValid())
+			{
+				bool bShift = false;
+				bool bCtrl = false;
+				bool bAlt = false;
+
+				FParse::Bool(*CurrentEntry, TEXT("Shift="), bShift);
+				FParse::Bool(*CurrentEntry, TEXT("Alt="), bAlt);
+				FParse::Bool(*CurrentEntry, TEXT("Ctrl="), bCtrl);
+				Chord = FInputChord(Key, EModifierKey::FromBools(bCtrl, bAlt, bShift, false));
+			}
+
+			FParse::Value(*CurrentEntry, TEXT("Index="), SpawnCommandInfo->PreconfiguredInfo.PreconfiguredIndex);
+			FString OutValue;
+			const bool bOverrideLabel = FParse::Value(*CurrentEntry, TEXT("Label="), OutValue, /*bShouldStopOnSeparator=*/true);
+			SpawnCommandInfo->PreconfiguredInfo.Label = FText::FromString(OutValue);
+
+			const UPCGSettings* CDO = CastChecked<UPCGSettings>(FoundClass->ClassDefaultObject);
+			const FText CommandLabelText = bOverrideLabel ? SpawnCommandInfo->PreconfiguredInfo.Label : CDO->GetDefaultNodeTitle();
+			const FText Description = FText::Format(NSLOCTEXT("PCGEditor", "SpawnNodeDescription", "Hold down the bound keys and left click in the graph panel to spawn a {0} node."), CommandLabelText);
+
+			FUICommandInfo::MakeCommandInfo(/*InContext=*/this->AsShared(),
+				SpawnCommandInfo->CommandInfo,
+				FName(FText::Format(INVTEXT("SpawnNode_{0}"), CommandLabelText).ToString()),
+				CommandLabelText,
+				Description,
+				/*InIcon=*/{},
+				EUserInterfaceActionType::Button,
+				Chord);
+
+			SpawnNodeCommands.Add(SpawnCommandInfo);
+		}
+	}
+}
+
+TSharedPtr<FEdGraphSchemaAction> FPCGEditorSpawnNodeCommands::GetGraphActionByChord(const FInputChord& InChord) const
+{
+	if (InChord.IsValidChord())
+	{
+		for (const TSharedPtr<FPCGSpawnNodeCommandInfo>& Command : SpawnNodeCommands)
+		{
+			if (Command.IsValid() && Command->CommandInfo.IsValid() && Command->CommandInfo->HasActiveChord(InChord))
+			{
+				return Command->GetAction();
+			}
+		}
+	}
+
+	return TSharedPtr<FEdGraphSchemaAction>();
+}
+
 #undef LOCTEXT_NAMESPACE
diff --git a/Engine/Plugins/PCG/Source/PCGEditor/Private/PCGEditorCommands.h b/Engine/Plugins/PCG/Source/PCGEditor/Private/PCGEditorCommands.h
--- a/Engine/Plugins/PCG/Source/PCGEditor/Private/PCGEditorCommands.h
+++ b/Engine/Plugins/PCG/Source/PCGEditor/Private/PCGEditorCommands.h
@@ -2,7 +2,13 @@
 
 #pragma once
 
+#include "PCGSettings.h"
 #include "Framework/Commands/Commands.h"
+#include "Styling/AppStyle.h"
+
+class UPCGSettings;
+class UEdGraph;
+struct FEdGraphSchemaAction;
 
 class FPCGEditorCommands : public TCommands<FPCGEditorCommands>
 {
@@ -35,3 +41,57 @@ public:
 	TSharedPtr<FUICommandInfo> SelectNamedRerouteUsages;
 	TSharedPtr<FUICommandInfo> SelectNamedRerouteDeclaration;
 };
+
+struct FPCGSpawnNodeCommandInfo
+{
+	explicit FPCGSpawnNodeCommandInfo(TSubclassOf<UPCGSettings> PCGSettingsClass) : PCGSettingsClass(PCGSettingsClass) {}
+
+	/** Optional preconfigured settings info for a specific configuration of the node. */
+	FPCGPreConfiguredSettingsInfo PreconfiguredInfo;
+
+	/** Holds the UI Command to verify chords for this action are held. */
+	TSharedPtr<FUICommandInfo> CommandInfo;
+
+	/**
+	 * Creates an action to be used for placing a node into the graph.
+	 *
+	 * @return A fully prepared action containing the information to spawn the node.
+	 */
+	TSharedPtr<FEdGraphSchemaAction> GetAction() const;
+
+	/**
+	 * Gets the PCG Settings class of the spawn node action.
+	 *
+	 * @return The PCG Settings class of the spawn node action.
+	 */
+	UClass* GetClass() const { return PCGSettingsClass; }
+
+private:
+	/** Type of settings class (node) to spawn. */
+	TSubclassOf<UPCGSettings> PCGSettingsClass;
+};
+
+/** Handles spawn node commands for the PCG Graph Editor. */
+class FPCGEditorSpawnNodeCommands : public TCommands<FPCGEditorSpawnNodeCommands>
+{
+public:
+	FPCGEditorSpawnNodeCommands()
+		: TCommands(TEXT("PCGEditorSpawnNodes"), NSLOCTEXT("PCGEditorSpawnNodes", "PCGEditorSpawnNodes", "PCG Editor - Spawn Nodes"), NAME_None, FAppStyle::GetAppStyleSetName())
+	{}
+
+	// ~Begin TCommands<> interface
+	virtual void RegisterCommands() override;
+	// ~End TCommands<> interface
+
+	/**
+	 * Returns a graph action assigned to the passed in chord
+	 *
+	 * @param InChord		The chord to use for lookup.
+	 * @return				A shared pointer to the schema action.
+	 */
+	TSharedPtr<FEdGraphSchemaAction> GetGraphActionByChord(const FInputChord& InChord) const;
+
+private:
+	/** An array of all the possible commands for spawning nodes. */
+	TArray<TSharedPtr<FPCGSpawnNodeCommandInfo>> SpawnNodeCommands;
+};
diff --git a/Engine/Plugins/PCG/Source/PCGEditor/Private/PCGEditorModule.cpp b/Engine/Plugins/PCG/Source/PCGEditor/Private/PCGEditorModule.cpp
--- a/Engine/Plugins/PCG/Source/PCGEditor/Private/PCGEditorModule.cpp
+++ b/Engine/Plugins/PCG/Source/PCGEditor/Private/PCGEditorModule.cpp
@@ -74,10 +74,11 @@ void FPCGEditorModule::StartupModule()
 
 	UToolMenus::RegisterStartupCallback(FSimpleMulticastDelegate::FDelegate::CreateRaw(this, &FPCGEditorModule::RegisterMenuExtensions));
 
-	FPCGEditorCommands::Register();
 	FPCGEditorStyle::Register();
+	FPCGEditorCommands::Register();
+	FPCGEditorSpawnNodeCommands::Register();
 	FPCGNodeSourceEditorTextBoxCommands::Register();
-	
+
 	GraphNodeFactory = MakeShareable(new FPCGEditorGraphNodeFactory());
 	FEdGraphUtilities::RegisterVisualNodeFactory(GraphNodeFactory);
 
@@ -107,6 +108,8 @@ void FPCGEditorModule::ShutdownModule()
 	UnregisterDetailsCustomizations();
 	UnregisterMenuExtensions();
 
+	FPCGNodeSourceEditorTextBoxCommands::Unregister();
+	FPCGEditorSpawnNodeCommands::Unregister();
 	FPCGEditorCommands::Unregister();
 	FPCGEditorStyle::Unregister();
 
