Added option to force pcg spawned actors to be removed at the beginning of the generation

#rnx

diff --git a/Engine/Plugins/PCG/Source/PCG/Private/Elements/PCGCreateTargetActor.cpp b/Engine/Plugins/PCG/Source/PCG/Private/Elements/PCGCreateTargetActor.cpp
--- a/Engine/Plugins/PCG/Source/PCG/Private/Elements/PCGCreateTargetActor.cpp
+++ b/Engine/Plugins/PCG/Source/PCG/Private/Elements/PCGCreateTargetActor.cpp
@@ -485,6 +485,8 @@ bool FPCGCreateTargetActorElement::ExecuteInternal(FPCGContext* Context) const
 				ManagedActors->SetCrc(ResourceCrc);
 			}
 			ManagedActors->GeneratedActors.Add(GeneratedActor);
+			ManagedActors->bSupportsReset = !Settings->bDeleteActorsBeforeGeneration;
+
 			SourceComponent->AddToManagedResources(ManagedActors);
 		}
 	}
diff --git a/Engine/Plugins/PCG/Source/PCG/Private/Elements/PCGSpawnActor.cpp b/Engine/Plugins/PCG/Source/PCG/Private/Elements/PCGSpawnActor.cpp
--- a/Engine/Plugins/PCG/Source/PCG/Private/Elements/PCGSpawnActor.cpp
+++ b/Engine/Plugins/PCG/Source/PCG/Private/Elements/PCGSpawnActor.cpp
@@ -927,6 +927,7 @@ void FPCGSpawnActorElement::SpawnActors(FPCGSubgraphContext* Context, AActor* Ta
 		// Create managed resource for actor tracking
 		UPCGManagedActors* ManagedActors = NewObject<UPCGManagedActors>(Context->SourceComponent.Get());
 		ManagedActors->SetCrc(InputDependenciesCrc);
+		ManagedActors->bSupportsReset = !Settings->bDeleteActorsBeforeGeneration;
 
 		// If generated actors are not directly attached, place them in a subfolder for tidiness.
 		FString GeneratedActorsFolderPath;
diff --git a/Engine/Plugins/PCG/Source/PCG/Private/PCGManagedResource.cpp b/Engine/Plugins/PCG/Source/PCG/Private/PCGManagedResource.cpp
--- a/Engine/Plugins/PCG/Source/PCG/Private/PCGManagedResource.cpp
+++ b/Engine/Plugins/PCG/Source/PCG/Private/PCGManagedResource.cpp
@@ -80,7 +80,9 @@ bool UPCGManagedActors::Release(bool bHardRelease, TSet<TSoftObjectPtr<AActor>>&
 {
 	TRACE_CPUPROFILER_EVENT_SCOPE(UPCGManagedActors::Release);
 
-	if (!Super::Release(bHardRelease, OutActorsToDelete))
+	const bool bShouldDeleteActors = bHardRelease || !bSupportsReset;
+
+	if (!Super::Release(bShouldDeleteActors, OutActorsToDelete))
 	{
 		PCGGeneratedResourcesLogging::LogManagedActorsRelease(this, GeneratedActors, bHardRelease, /*bOnlyMarkedForCleanup=*/true);
 
diff --git a/Engine/Plugins/PCG/Source/PCG/Public/Elements/PCGCreateTargetActor.h b/Engine/Plugins/PCG/Source/PCG/Public/Elements/PCGCreateTargetActor.h
--- a/Engine/Plugins/PCG/Source/PCG/Public/Elements/PCGCreateTargetActor.h
+++ b/Engine/Plugins/PCG/Source/PCG/Public/Elements/PCGCreateTargetActor.h
@@ -85,6 +85,9 @@ public:
 	UPROPERTY(BlueprintReadWrite, EditAnywhere, Category = Settings, meta = (PCG_Overridable))
 	FPCGDataLayerSettings DataLayerSettings;
 
+	UPROPERTY(BlueprintReadWrite, EditAnywhere, Category = "Settings|Advanced")
+	bool bDeleteActorsBeforeGeneration = false;
+
 	PCG_API void SetTemplateActorClass(const TSubclassOf<AActor>& InTemplateActorClass);
 	PCG_API void SetAllowTemplateActorEditing(bool bInAllowTemplateActorEditing);
 	const TSubclassOf<AActor>& GetTemplateActorClass() const { return TemplateActorClass; }
diff --git a/Engine/Plugins/PCG/Source/PCG/Public/Elements/PCGSpawnActor.h b/Engine/Plugins/PCG/Source/PCG/Public/Elements/PCGSpawnActor.h
--- a/Engine/Plugins/PCG/Source/PCG/Public/Elements/PCGSpawnActor.h
+++ b/Engine/Plugins/PCG/Source/PCG/Public/Elements/PCGSpawnActor.h
@@ -89,6 +89,9 @@ public:
 	UPROPERTY(BlueprintReadWrite, EditAnywhere, Category = "Settings|Advanced")
 	bool bWarnOnIdenticalSpawn = true;
 
+	UPROPERTY(BlueprintReadWrite, EditAnywhere, Category = "Settings|Advanced")
+	bool bDeleteActorsBeforeGeneration = false;
+
 	UPROPERTY(BlueprintReadWrite, EditAnywhere, Category = Settings, meta = (PCG_Overridable))
 	FPCGDataLayerSettings DataLayerSettings;
 
diff --git a/Engine/Plugins/PCG/Source/PCG/Public/PCGManagedResource.h b/Engine/Plugins/PCG/Source/PCG/Public/PCGManagedResource.h
--- a/Engine/Plugins/PCG/Source/PCG/Public/PCGManagedResource.h
+++ b/Engine/Plugins/PCG/Source/PCG/Public/PCGManagedResource.h
@@ -108,6 +108,13 @@ public:
 
 	UPROPERTY(BlueprintReadWrite, EditAnywhere, Category = GeneratedData)
 	TSet<TSoftObjectPtr<AActor>> GeneratedActors;
+
+	/** Controls whether the resource will be removed at the beginning of the generation instead of being kept until the end, in the eventuality it is reused. 
+	* In practice, this means that the actors will be deleted as part of the initial cleanup in the generation instead of being removed at the end.
+	* In some instances, this might be required if some components on the actors interact negatively with the PCG processing in the graph.
+	*/
+	UPROPERTY(BlueprintReadWrite, EditAnywhere, Category = GeneratedData)
+	bool bSupportsReset = true;
 };
 
 UCLASS(Abstract)
