[PCG] Add Execution Dependency Pin to most nodes (excl. Control Flow, Compact nodes, Input/Output, etc)

* Deprecated Dependency Only pin from Gather, ApplyOnActor.
+ Implemented new icon for Execution Dependency Pin to match the asset.
+ Added graph compiler helper to determine whether task input should be culled from a pin

diff --git a/Engine/Plugins/PCG/Source/PCG/Private/Elements/PCGApplyOnActor.cpp b/Engine/Plugins/PCG/Source/PCG/Private/Elements/PCGApplyOnActor.cpp
--- a/Engine/Plugins/PCG/Source/PCG/Private/Elements/PCGApplyOnActor.cpp
+++ b/Engine/Plugins/PCG/Source/PCG/Private/Elements/PCGApplyOnActor.cpp
@@ -16,7 +16,6 @@
 
 namespace PCGApplyOnActorConstants
 {
-	const FText DependencyTooltip = LOCTEXT("DependencyPinTooltip", "Data passed to this pin will be used to order execution but will otherwise not contribute to the results of this node.");
 	const FName ObjectPropertyOverridesLabel = TEXT("Property Overrides");
 	const FText ObjectPropertyOverridesTooltip = LOCTEXT("ObjectOverrideToolTip", "Provide property overrides for the target object. The attribute name must match the InputSource name in the object property override description.");
 }
@@ -63,6 +62,10 @@ void UPCGApplyOnActorSettings::ApplyDeprecationBeforeUpdatePins(UPCGNode* InOutN
 		}
 	}
 
+PRAGMA_DISABLE_DEPRECATION_WARNINGS
+	InOutNode->RenameInputPin(PCGPinConstants::DefaultDependencyOnlyLabel, PCGPinConstants::DefaultExecutionDependencyLabel, /*bInBroadcastUpdate=*/false);
+PRAGMA_ENABLE_DEPRECATION_WARNINGS
+
 	Super::ApplyDeprecationBeforeUpdatePins(InOutNode, InputPins, OutputPins);
 }
 #endif // WITH_EDITOR
@@ -72,7 +75,6 @@ TArray<FPCGPinProperties> UPCGApplyOnActorSettings::InputPinProperties() const
 	TArray<FPCGPinProperties> PinProperties;
 	PinProperties.Emplace(PCGPinConstants::DefaultInputLabel, EPCGDataType::Any, /*bAllowMultipleConnections=*/true, /*bAllowMultipleData=*/true);
 	PinProperties.Emplace(PCGApplyOnActorConstants::ObjectPropertyOverridesLabel, EPCGDataType::Any, /*bAllowMultipleConnections=*/true, /*bAllowMultipleData=*/true, PCGApplyOnActorConstants::ObjectPropertyOverridesTooltip);
-	PinProperties.Emplace(PCGPinConstants::DefaultDependencyOnlyLabel, EPCGDataType::Any, /*bAllowMultipleConnections=*/true, /*bAllowMultipleData=*/true, PCGApplyOnActorConstants::DependencyTooltip);
 	return PinProperties;
 }
 
diff --git a/Engine/Plugins/PCG/Source/PCG/Private/Elements/PCGGather.cpp b/Engine/Plugins/PCG/Source/PCG/Private/Elements/PCGGather.cpp
--- a/Engine/Plugins/PCG/Source/PCG/Private/Elements/PCGGather.cpp
+++ b/Engine/Plugins/PCG/Source/PCG/Private/Elements/PCGGather.cpp
@@ -7,12 +7,22 @@
 
 #define LOCTEXT_NAMESPACE "PCGGatherElement"
 
+#if WITH_EDITOR
+void UPCGGatherSettings::ApplyDeprecationBeforeUpdatePins(UPCGNode* InOutNode, TArray<TObjectPtr<UPCGPin>>& InputPins, TArray<TObjectPtr<UPCGPin>>& OutputPins)
+{
+	Super::ApplyDeprecationBeforeUpdatePins(InOutNode, InputPins, OutputPins);
+
+PRAGMA_DISABLE_DEPRECATION_WARNINGS
+	InOutNode->RenameInputPin(PCGPinConstants::DefaultDependencyOnlyLabel, PCGPinConstants::DefaultExecutionDependencyLabel, /*bInBroadcastUpdate=*/false);
+PRAGMA_ENABLE_DEPRECATION_WARNINGS
+}
+#endif // WITH_EDITOR
+
 TArray<FPCGPinProperties> UPCGGatherSettings::InputPinProperties() const
 {
 	TArray<FPCGPinProperties> PinProperties;
 	FPCGPinProperties& InputPinProperty = PinProperties.Emplace_GetRef(PCGPinConstants::DefaultInputLabel, EPCGDataType::Any);
 	InputPinProperty.SetRequiredPin();
-	PinProperties.Emplace(PCGPinConstants::DefaultDependencyOnlyLabel, EPCGDataType::Any, /*bAllowMultipleConnections=*/true, /*bAllowMultipleData=*/true, LOCTEXT("DependencyPinTooltip", "Data passed to this pin will be used to order execution but will otherwise not contribute to the results of this node."));
 
 	return PinProperties;
 }
diff --git a/Engine/Plugins/PCG/Source/PCG/Private/Graph/PCGGraphCompiler.cpp b/Engine/Plugins/PCG/Source/PCG/Private/Graph/PCGGraphCompiler.cpp
--- a/Engine/Plugins/PCG/Source/PCG/Private/Graph/PCGGraphCompiler.cpp
+++ b/Engine/Plugins/PCG/Source/PCG/Private/Graph/PCGGraphCompiler.cpp
@@ -325,6 +325,8 @@ TArray<FPCGGraphTask> FPCGGraphCompiler::CompileGraph(UPCGGraph* InGraph, FPCGTa
 
 			for (const UPCGPin* InputPin : Node->InputPins)
 			{
+				// Cull dependency only data if the pin is marked as such.
+				const bool bIsDependencyOnly = InputPin->Properties.Usage == EPCGPinUsage::DependencyOnly;
 				for (const UPCGEdge* InboundEdge : InputPin->Edges)
 				{
 					if (!InboundEdge->IsValid())
@@ -333,9 +335,10 @@ TArray<FPCGGraphTask> FPCGGraphCompiler::CompileGraph(UPCGGraph* InGraph, FPCGTa
 						continue;
 					}
 
+					check(InboundEdge->OutputPin == InputPin);
 					if (FPCGTaskId* InboundId = IdMapping.Find(InboundEdge->InputPin->Node))
 					{
-						Task.Inputs.Emplace(*InboundId, InboundEdge->InputPin->Properties, InboundEdge->OutputPin->Properties);
+						Task.Inputs.Emplace(*InboundId, InboundEdge->InputPin->Properties, InputPin->Properties, !bIsDependencyOnly);
 					}
 					else
 					{
@@ -1395,7 +1398,7 @@ void FPCGGraphCompiler::CompileTopGraph(UPCGGraph* InGraph, uint32 GenerationGri
 		for (FPCGGraphTaskInput& TaskInput : CompiledTasks[TaskIndex].Inputs)
 		{
 			PCGGraphCompiler::FTaskInputOrigin InputOrigin(&TaskInput);
-			if (!InputOrigin.TaskInput)
+			if (!InputOrigin.TaskInput || !InputOrigin.TaskInput->bProvideData)
 			{
 				continue;
 			}
diff --git a/Engine/Plugins/PCG/Source/PCG/Private/PCGSettings.cpp b/Engine/Plugins/PCG/Source/PCG/Private/PCGSettings.cpp
--- a/Engine/Plugins/PCG/Source/PCG/Private/PCGSettings.cpp
+++ b/Engine/Plugins/PCG/Source/PCG/Private/PCGSettings.cpp
@@ -501,6 +501,19 @@ bool UPCGSettings::HasOverridableParam(FName InParamName) const
 	}) != nullptr;
 }
 
+void UPCGSettings::AddExecutionDependencyPin(TArray<FPCGPinProperties>& OutPins) const
+{
+	if (HasExecutionDependencyPin())
+	{
+		FPCGPinProperties& DepPin = OutPins.Emplace_GetRef(PCGPinConstants::DefaultExecutionDependencyLabel, EPCGDataType::Any, /*bInAllowMultipleConnections=*/true, /*bAllowMultipleData=*/true);
+#if WITH_EDITOR
+		DepPin.Tooltip = PCGPinConstants::Tooltips::ExecutionDependencyTooltip;
+#endif // WITH_EDITOR
+		DepPin.Usage = EPCGPinUsage::DependencyOnly;
+		DepPin.PinStatus = EPCGPinStatus::Advanced;
+	}
+}
+
 void UPCGSettings::FillOverridableParamsPins(TArray<FPCGPinProperties>& OutPins) const
 {
 	if (!HasOverridableParams())
@@ -602,6 +615,7 @@ void UPCGSettings::FillOverridableParamsPins(TArray<FPCGPinProperties>& OutPins)
 TArray<FPCGPinProperties> UPCGSettings::AllInputPinProperties() const
 {
 	TArray<FPCGPinProperties> InputPins = InputPinProperties();
+	AddExecutionDependencyPin(InputPins);
 	FillOverridableParamsPins(InputPins);
 	return InputPins;
 }
diff --git a/Engine/Plugins/PCG/Source/PCG/Public/Elements/ControlFlow/PCGBooleanSelect.h b/Engine/Plugins/PCG/Source/PCG/Public/Elements/ControlFlow/PCGBooleanSelect.h
--- a/Engine/Plugins/PCG/Source/PCG/Public/Elements/ControlFlow/PCGBooleanSelect.h
+++ b/Engine/Plugins/PCG/Source/PCG/Public/Elements/ControlFlow/PCGBooleanSelect.h
@@ -1,7 +1,7 @@
 // Copyright Epic Games, Inc. All Rights Reserved.
 #pragma once
 
-#include "PCGSettings.h"
+#include "PCGControlFlow.h"
 
 #include "PCGBooleanSelect.generated.h"
 
@@ -9,7 +9,7 @@
  * Selects data from either input pin, based on a boolean condition.
  */
 UCLASS(BlueprintType, ClassGroup = (Procedural), meta=(Keywords = "if bool select"))
-class UPCGBooleanSelectSettings : public UPCGSettings
+class UPCGBooleanSelectSettings : public UPCGControlFlowSettings
 {
 	GENERATED_BODY()
 
diff --git a/Engine/Plugins/PCG/Source/PCG/Public/Elements/ControlFlow/PCGBranch.h b/Engine/Plugins/PCG/Source/PCG/Public/Elements/ControlFlow/PCGBranch.h
--- a/Engine/Plugins/PCG/Source/PCG/Public/Elements/ControlFlow/PCGBranch.h
+++ b/Engine/Plugins/PCG/Source/PCG/Public/Elements/ControlFlow/PCGBranch.h
@@ -1,7 +1,7 @@
 // Copyright Epic Games, Inc. All Rights Reserved.
 #pragma once
 
-#include "PCGSettings.h"
+#include "PCGControlFlow.h"
 
 #include "PCGBranch.generated.h"
 
@@ -9,7 +9,7 @@
  * Routes input data to one of two outputs, based on a boolean condition.
  */
 UCLASS(BlueprintType, ClassGroup = (Procedural), meta=(Keywords = "if bool switch"))
-class UPCGBranchSettings : public UPCGSettings
+class UPCGBranchSettings : public UPCGControlFlowSettings
 {
 	GENERATED_BODY()
 
diff --git a/Engine/Plugins/PCG/Source/PCG/Public/Elements/ControlFlow/PCGControlFlow.h b/Engine/Plugins/PCG/Source/PCG/Public/Elements/ControlFlow/PCGControlFlow.h
--- a/Engine/Plugins/PCG/Source/PCG/Public/Elements/ControlFlow/PCGControlFlow.h
+++ b/Engine/Plugins/PCG/Source/PCG/Public/Elements/ControlFlow/PCGControlFlow.h
@@ -2,8 +2,21 @@
 
 #pragma once
 
+#include "PCGSettings.h"
+
 #include "PCGControlFlow.generated.h"
 
+UCLASS(Abstract, BlueprintType, ClassGroup = (Procedural))
+class PCG_API UPCGControlFlowSettings : public UPCGSettings
+{
+	GENERATED_BODY()
+
+	//~Begin UPCGSettings interface
+protected:
+	virtual bool HasExecutionDependencyPin() const override { return false; }
+	//~End UPCGSettings interface
+};
+
 USTRUCT(BlueprintType, meta = (Hidden))
 struct PCG_API FEnumSelector
 {
diff --git a/Engine/Plugins/PCG/Source/PCG/Public/Elements/ControlFlow/PCGMultiSelect.h b/Engine/Plugins/PCG/Source/PCG/Public/Elements/ControlFlow/PCGMultiSelect.h
--- a/Engine/Plugins/PCG/Source/PCG/Public/Elements/ControlFlow/PCGMultiSelect.h
+++ b/Engine/Plugins/PCG/Source/PCG/Public/Elements/ControlFlow/PCGMultiSelect.h
@@ -2,7 +2,6 @@
 
 #pragma once
 
-#include "PCGSettings.h"
 #include "PCGControlFlow.h"
 
 #include "PCGMultiSelect.generated.h"
@@ -11,7 +10,7 @@
  * Selects data from any number of input pins, based on a static selection criteria (Int/String/Enum) 
  */
 UCLASS(BlueprintType, ClassGroup = (Procedural), meta=(Keywords = "if multi switch enum"))
-class UPCGMultiSelectSettings : public UPCGSettings
+class UPCGMultiSelectSettings : public UPCGControlFlowSettings
 {
 	GENERATED_BODY()
 
diff --git a/Engine/Plugins/PCG/Source/PCG/Public/Elements/ControlFlow/PCGQualityBranch.h b/Engine/Plugins/PCG/Source/PCG/Public/Elements/ControlFlow/PCGQualityBranch.h
--- a/Engine/Plugins/PCG/Source/PCG/Public/Elements/ControlFlow/PCGQualityBranch.h
+++ b/Engine/Plugins/PCG/Source/PCG/Public/Elements/ControlFlow/PCGQualityBranch.h
@@ -2,7 +2,7 @@
 
 #pragma once
 
-#include "PCGSettings.h"
+#include "PCGControlFlow.h"
 
 #include "PCGQualityBranch.generated.h"
 
@@ -10,7 +10,7 @@
  * Control flow node that dynamically routes input data based on 'pcg.Quality' setting.
  */
 UCLASS(BlueprintType, ClassGroup = (Procedural), meta = (Keywords = "if bool switch quality"))
-class UPCGQualityBranchSettings : public UPCGSettings
+class UPCGQualityBranchSettings : public UPCGControlFlowSettings
 {
 	GENERATED_BODY()
 
diff --git a/Engine/Plugins/PCG/Source/PCG/Public/Elements/ControlFlow/PCGQualitySelect.h b/Engine/Plugins/PCG/Source/PCG/Public/Elements/ControlFlow/PCGQualitySelect.h
--- a/Engine/Plugins/PCG/Source/PCG/Public/Elements/ControlFlow/PCGQualitySelect.h
+++ b/Engine/Plugins/PCG/Source/PCG/Public/Elements/ControlFlow/PCGQualitySelect.h
@@ -2,7 +2,7 @@
 
 #pragma once
 
-#include "PCGSettings.h"
+#include "PCGControlFlow.h"
 
 #include "PCGQualitySelect.generated.h"
 
@@ -10,7 +10,7 @@
  * Selects from input pins based on 'pcg.Quality' setting.
  */
 UCLASS(BlueprintType, ClassGroup = (Procedural), meta = (Keywords = "if bool switch quality"))
-class UPCGQualitySelectSettings : public UPCGSettings
+class UPCGQualitySelectSettings : public UPCGControlFlowSettings
 {
 	GENERATED_BODY()
 
diff --git a/Engine/Plugins/PCG/Source/PCG/Public/Elements/ControlFlow/PCGSwitch.h b/Engine/Plugins/PCG/Source/PCG/Public/Elements/ControlFlow/PCGSwitch.h
--- a/Engine/Plugins/PCG/Source/PCG/Public/Elements/ControlFlow/PCGSwitch.h
+++ b/Engine/Plugins/PCG/Source/PCG/Public/Elements/ControlFlow/PCGSwitch.h
@@ -3,7 +3,6 @@
 #pragma once
 
 #include "PCGControlFlow.h"
-#include "PCGSettings.h"
 
 #include "PCGSwitch.generated.h"
 
@@ -11,7 +10,7 @@
  * Routes data from the input pin, to a specific output pin based on a selection criteria (Int/String/Enum) 
  */
 UCLASS(BlueprintType, ClassGroup = (Procedural), meta=(Keywords = "if multi enum branch"))
-class UPCGSwitchSettings : public UPCGSettings
+class UPCGSwitchSettings : public UPCGControlFlowSettings
 {
 	GENERATED_BODY()
 
diff --git a/Engine/Plugins/PCG/Source/PCG/Public/Elements/PCGCollapseElement.h b/Engine/Plugins/PCG/Source/PCG/Public/Elements/PCGCollapseElement.h
--- a/Engine/Plugins/PCG/Source/PCG/Public/Elements/PCGCollapseElement.h
+++ b/Engine/Plugins/PCG/Source/PCG/Public/Elements/PCGCollapseElement.h
@@ -22,7 +22,7 @@ public:
 	virtual bool ShouldDrawNodeCompact() const override { return true; }
 	virtual bool GetCompactNodeIcon(FName& OutCompactNodeIcon) const override;
 #endif
-	
+	virtual bool HasExecutionDependencyPin() const override { return false; }
 
 protected:
 	virtual TArray<FPCGPinProperties> InputPinProperties() const override;
diff --git a/Engine/Plugins/PCG/Source/PCG/Public/Elements/PCGConvertToAttributeSet.h b/Engine/Plugins/PCG/Source/PCG/Public/Elements/PCGConvertToAttributeSet.h
--- a/Engine/Plugins/PCG/Source/PCG/Public/Elements/PCGConvertToAttributeSet.h
+++ b/Engine/Plugins/PCG/Source/PCG/Public/Elements/PCGConvertToAttributeSet.h
@@ -22,6 +22,7 @@ public:
 	virtual bool ShouldDrawNodeCompact() const override { return true; }
 	virtual bool GetCompactNodeIcon(FName& OutCompactNodeIcon) const override;
 #endif
+	virtual bool HasExecutionDependencyPin() const override { return false; }
 
 protected:
 	virtual TArray<FPCGPinProperties> InputPinProperties() const override;
diff --git a/Engine/Plugins/PCG/Source/PCG/Public/Elements/PCGCreateSurfaceFromSpline.h b/Engine/Plugins/PCG/Source/PCG/Public/Elements/PCGCreateSurfaceFromSpline.h
--- a/Engine/Plugins/PCG/Source/PCG/Public/Elements/PCGCreateSurfaceFromSpline.h
+++ b/Engine/Plugins/PCG/Source/PCG/Public/Elements/PCGCreateSurfaceFromSpline.h
@@ -26,6 +26,7 @@ public:
 	virtual EPCGSettingsType GetType() const override { return EPCGSettingsType::Spatial; }
 	virtual bool ShouldDrawNodeCompact() const override { return bShouldDrawNodeCompact; }
 #endif
+	virtual bool HasExecutionDependencyPin() const override { return false; }
 
 protected:
 	virtual TArray<FPCGPinProperties> InputPinProperties() const override;
diff --git a/Engine/Plugins/PCG/Source/PCG/Public/Elements/PCGFilterByType.h b/Engine/Plugins/PCG/Source/PCG/Public/Elements/PCGFilterByType.h
--- a/Engine/Plugins/PCG/Source/PCG/Public/Elements/PCGFilterByType.h
+++ b/Engine/Plugins/PCG/Source/PCG/Public/Elements/PCGFilterByType.h
@@ -21,6 +21,7 @@ public:
 	virtual bool ShouldDrawNodeCompact() const override { return true; }
 	virtual bool GetCompactNodeIcon(FName& OutCompactNodeIcon) const override;
 #endif
+	virtual bool HasExecutionDependencyPin() const override { return false; }
 	
 	virtual EPCGDataType GetCurrentPinTypes(const UPCGPin* InPin) const override;
 
diff --git a/Engine/Plugins/PCG/Source/PCG/Public/Elements/PCGGather.h b/Engine/Plugins/PCG/Source/PCG/Public/Elements/PCGGather.h
--- a/Engine/Plugins/PCG/Source/PCG/Public/Elements/PCGGather.h
+++ b/Engine/Plugins/PCG/Source/PCG/Public/Elements/PCGGather.h
@@ -20,6 +20,7 @@ public:
 	virtual FText GetDefaultNodeTitle() const override { return NSLOCTEXT("PCGGatherSettings", "NodeTitle", "Gather"); }
 	virtual FText GetNodeTooltipText() const override { return NSLOCTEXT("PCGGatherSettings", "NodeTooltip", "Gathers multiple data in a single collection. Can also be used to order execution through the dependency-only pin."); }
 	virtual EPCGSettingsType GetType() const override { return EPCGSettingsType::Generic; }
+	virtual void ApplyDeprecationBeforeUpdatePins(UPCGNode* InOutNode, TArray<TObjectPtr<UPCGPin>>& InputPins, TArray<TObjectPtr<UPCGPin>>& OutputPins) override;
 #endif
 
 	virtual bool HasDynamicPins() const override { return true; }
diff --git a/Engine/Plugins/PCG/Source/PCG/Public/Elements/PCGMakeConcreteElement.h b/Engine/Plugins/PCG/Source/PCG/Public/Elements/PCGMakeConcreteElement.h
--- a/Engine/Plugins/PCG/Source/PCG/Public/Elements/PCGMakeConcreteElement.h
+++ b/Engine/Plugins/PCG/Source/PCG/Public/Elements/PCGMakeConcreteElement.h
@@ -21,7 +21,7 @@ public:
 	virtual EPCGSettingsType GetType() const override { return EPCGSettingsType::Spatial; }
 	virtual bool ShouldDrawNodeCompact() const override { return true; }
 #endif
-	
+	virtual bool HasExecutionDependencyPin() const override { return false; }
 
 protected:
 	virtual TArray<FPCGPinProperties> OutputPinProperties() const override;
diff --git a/Engine/Plugins/PCG/Source/PCG/Public/Elements/PCGReroute.h b/Engine/Plugins/PCG/Source/PCG/Public/Elements/PCGReroute.h
--- a/Engine/Plugins/PCG/Source/PCG/Public/Elements/PCGReroute.h
+++ b/Engine/Plugins/PCG/Source/PCG/Public/Elements/PCGReroute.h
@@ -43,7 +43,7 @@ protected:
 	virtual TArray<FPCGPinProperties> OutputPinProperties() const override;
 	virtual FPCGElementPtr CreateElement() const override;
 	virtual bool ComputeOutputPinDataDesc(const UPCGPin* OutputPin, const UPCGDataBinding* Binding, FPCGDataCollectionDesc& OutDesc) const override;
-
+	virtual bool HasExecutionDependencyPin() const override { return false; }
 #if WITH_EDITOR
 	virtual TArray<FPCGPreconfiguredInfo> GetConversionInfo() const override;
 #endif // WITH_EDITOR
diff --git a/Engine/Plugins/PCG/Source/PCG/Public/Elements/PCGUserParameterGet.h b/Engine/Plugins/PCG/Source/PCG/Public/Elements/PCGUserParameterGet.h
--- a/Engine/Plugins/PCG/Source/PCG/Public/Elements/PCGUserParameterGet.h
+++ b/Engine/Plugins/PCG/Source/PCG/Public/Elements/PCGUserParameterGet.h
@@ -51,6 +51,7 @@ protected:
 	virtual TArray<FPCGPinProperties> InputPinProperties() const override { return {}; }
 	virtual TArray<FPCGPinProperties> OutputPinProperties() const override;
 	virtual FPCGElementPtr CreateElement() const override;
+	virtual bool HasExecutionDependencyPin() const override { return false; }
 
 #if WITH_EDITOR
 	virtual FName GetDefaultNodeName() const override { return FName("GetGraphParameter"); }
@@ -102,6 +103,7 @@ protected:
 	virtual TArray<FPCGPinProperties> InputPinProperties() const override { return {}; }
 	virtual TArray<FPCGPinProperties> OutputPinProperties() const override;
 	virtual FPCGElementPtr CreateElement() const override;
+	virtual bool HasExecutionDependencyPin() const override { return false; }
 
 #if WITH_EDITOR
 	virtual FName GetDefaultNodeName() const override { return FName("GetGenericGraphParameter"); }
diff --git a/Engine/Plugins/PCG/Source/PCG/Public/PCGCommon.h b/Engine/Plugins/PCG/Source/PCG/Public/PCGCommon.h
--- a/Engine/Plugins/PCG/Source/PCG/Public/PCGCommon.h
+++ b/Engine/Plugins/PCG/Source/PCG/Public/PCGCommon.h
@@ -140,7 +140,9 @@ namespace PCGPinConstants
 	const FName DefaultInputLabel = TEXT("In");
 	const FName DefaultOutputLabel = TEXT("Out");
 	const FName DefaultParamsLabel = TEXT("Overrides");
+	UE_DEPRECATED(5.6, "Please use 'DefaultExecutionDependencyLabel' instead.")
 	const FName DefaultDependencyOnlyLabel = TEXT("Dependency Only");
+	const FName DefaultExecutionDependencyLabel = TEXT("Execution Dependency");
 
 	const FName DefaultInFilterLabel = TEXT("InsideFilter");
 	const FName DefaultOutFilterLabel = TEXT("OutsideFilter");
@@ -155,7 +157,14 @@ namespace Icons
 	const FName LoopPinIcon = TEXT("GraphEditor.Macro.Loop_16x");
 	const FName FeedbackPinIcon = TEXT("GraphEditor.GetSequenceBinding");
 }
+
+#if WITH_EDITOR
+namespace Tooltips
+{
+	const FText ExecutionDependencyTooltip = NSLOCTEXT("PCGCommon", "ExecutionDependencyPinTooltip", "Data passed to this pin will be used to order execution but will otherwise not contribute to the results of this node.");
 }
+#endif // WITH_EDITOR
+} // namespace PCGPinConstants
 
 namespace PCGNodeConstants
 {
diff --git a/Engine/Plugins/PCG/Source/PCG/Public/PCGInputOutputSettings.h b/Engine/Plugins/PCG/Source/PCG/Public/PCGInputOutputSettings.h
--- a/Engine/Plugins/PCG/Source/PCG/Public/PCGInputOutputSettings.h
+++ b/Engine/Plugins/PCG/Source/PCG/Public/PCGInputOutputSettings.h
@@ -49,6 +49,7 @@ protected:
 	virtual FPCGElementPtr CreateElement() const override { return MakeShared<FPCGInputOutputElement>(); }
 	virtual TArray<FPCGPinProperties> InputPinProperties() const override;
 	virtual TArray<FPCGPinProperties> OutputPinProperties() const override;
+	virtual bool HasExecutionDependencyPin() const override { return false; }
 
 #if WITH_EDITOR
 	virtual void PostEditChangeProperty(FPropertyChangedEvent& PropertyChangedEvent) override;
diff --git a/Engine/Plugins/PCG/Source/PCG/Public/PCGSettings.h b/Engine/Plugins/PCG/Source/PCG/Public/PCGSettings.h
--- a/Engine/Plugins/PCG/Source/PCG/Public/PCGSettings.h
+++ b/Engine/Plugins/PCG/Source/PCG/Public/PCGSettings.h
@@ -211,7 +211,7 @@ public:
 	/*virtual*/ FPCGElementPtr GetElement() const;
 	virtual UPCGNode* CreateNode() const;
 	
-	/** Return the concatenation of InputPinProperties and FillOverridableParamsPins */
+	/** Return the concatenation of InputPinProperties, ExecutionDependencyPin, and FillOverridableParamsPins */
 	TArray<FPCGPinProperties> AllInputPinProperties() const;
 
 	/** For symmetry reason, do the same with output pins. For now forward just the call to OutputPinProperties */
@@ -572,6 +572,15 @@ public:
 	virtual bool DisplayExecuteOnGPUSetting() const { return false; }
 #endif
 
+// Execution Dependency section
+protected:
+	/** This node should have an advanced Execution Dependency pin. */
+	virtual bool HasExecutionDependencyPin() const { return true; }
+
+private:
+	/** Adds the universal Execution Dependency pin to the node's pin properties. */
+	void AddExecutionDependencyPin(TArray<FPCGPinProperties>& OutPins) const;
+
 protected:
 	/** [EXPERIMENTAL] Note: This variable is subject to change without deprecation. */
 	UPROPERTY(EditAnywhere, Category = "GPU", meta = (Tooltip = "Whether this node should be executed on the GPU.", EditCondition = "DisplayExecuteOnGPUSetting()", EditConditionHides, HideEditConditionToggle))
diff --git a/Engine/Plugins/PCG/Source/PCGEditor/Private/PCGEditorStyle.cpp b/Engine/Plugins/PCG/Source/PCGEditor/Private/PCGEditorStyle.cpp
--- a/Engine/Plugins/PCG/Source/PCGEditor/Private/PCGEditorStyle.cpp
+++ b/Engine/Plugins/PCG/Source/PCGEditor/Private/PCGEditorStyle.cpp
@@ -70,6 +70,9 @@ FPCGEditorStyle::FPCGEditorStyle() : FSlateStyleSet("PCGEditorStyle")
 
 	Set(PCGEditorStyleConstants::Pin_Required, new IMAGE_BRUSH_SVG("Style/PCG_Graph_RequiredPin_IN", FVector2D(8.0f, 22.0f)));
 
+	Set(PCGEditorStyleConstants::Pin_Graph_Dependency_C, new IMAGE_BRUSH_SVG("Style/PCG_Graph_Dependency", PinSize));
+	Set(PCGEditorStyleConstants::Pin_Graph_Dependency_DC, new IMAGE_BRUSH_SVG("Style/PCG_Graph_Dependency_Unplugged", PinSize));
+
 	Set(PCGEditorStyleConstants::Node_Overlay_Inactive, new IMAGE_BRUSH_SVG("Style/PCG_Node_Overlay_Inactive", Icon20x20));
 
 	Set(PCGNodeConstants::Icons::CompactNodeFilter, new IMAGE_BRUSH_SVG("Style/PCG_Graph_Filter", FVector2D(28.0f, 28.0f)));
diff --git a/Engine/Plugins/PCG/Source/PCGEditor/Private/PCGEditorStyle.h b/Engine/Plugins/PCG/Source/PCGEditor/Private/PCGEditorStyle.h
--- a/Engine/Plugins/PCG/Source/PCGEditor/Private/PCGEditorStyle.h
+++ b/Engine/Plugins/PCG/Source/PCGEditor/Private/PCGEditorStyle.h
@@ -49,6 +49,8 @@ namespace PCGEditorStyleConstants
 	const FName Pin_Composite_OUT_DC = TEXT("PCG.Graph.Composite.Out.Disconnected");
 
 	const FName Pin_Required = TEXT("PCG.Graph.Pin.Required");
+	const FName Pin_Graph_Dependency_C = TEXT("PCG.Graph.Pin.Dependency.Connected");
+	const FName Pin_Graph_Dependency_DC = TEXT("PCG.Graph.Pin.Dependency.Disconnected");
 
 	const FName Node_Overlay_Inactive = TEXT("PCG.Node.Overlay.Inactive");
 
diff --git a/Engine/Plugins/PCG/Source/PCGEditor/Private/SPCGEditorGraphNode.cpp b/Engine/Plugins/PCG/Source/PCGEditor/Private/SPCGEditorGraphNode.cpp
--- a/Engine/Plugins/PCG/Source/PCGEditor/Private/SPCGEditorGraphNode.cpp
+++ b/Engine/Plugins/PCG/Source/PCGEditor/Private/SPCGEditorGraphNode.cpp
@@ -566,6 +566,13 @@ void SPCGEditorGraphNode::AddPin(const TSharedRef<SGraphPin>& PinToAdd)
 
 				PinToAdd->SetCustomPinIcon(ConnectedBrush, DisconnectedBrush);
 			}
+			else if (Pin->Properties.Usage == EPCGPinUsage::DependencyOnly)
+			{
+				const FSlateBrush* ConnectedBrush = FPCGEditorStyle::Get().GetBrush(PCGEditorStyleConstants::Pin_Graph_Dependency_C);
+				const FSlateBrush* DisconnectedBrush = FPCGEditorStyle::Get().GetBrush(PCGEditorStyleConstants::Pin_Graph_Dependency_DC);
+
+				PinToAdd->SetCustomPinIcon(ConnectedBrush, DisconnectedBrush);
+			}
 			else
 			{
 				// Node outputs are always single collection (SC).
