Fixed clashing loc text keys

diff --git a/Engine/Plugins/PCG/Source/PCG/Private/MeshSelectors/PCGAnimBankMeshSelector.cpp b/Engine/Plugins/PCG/Source/PCG/Private/MeshSelectors/PCGAnimBankMeshSelector.cpp
--- a/Engine/Plugins/PCG/Source/PCG/Private/MeshSelectors/PCGAnimBankMeshSelector.cpp
+++ b/Engine/Plugins/PCG/Source/PCG/Private/MeshSelectors/PCGAnimBankMeshSelector.cpp
@@ -113,21 +113,21 @@ bool UPCGAnimBankMeshSelector::SelectInstances(
 	// Validate that the "bank" attribute is of the right type
 	if (!PCG::Private::IsOfTypes<FSoftObjectPath, FString>(BankAttributeBase->GetTypeId()))
 	{
-		PCGE_LOG_C(Error, GraphAndLog, &Context, FText::Format(LOCTEXT("AttributeInvalidType", "Attribute '{0}' is not of valid type (must be FString or FSoftObjectPath)"), BankAttribute.GetDisplayText()));
+		PCGE_LOG_C(Error, GraphAndLog, &Context, FText::Format(LOCTEXT("BankAttributeInvalidType", "Attribute '{0}' is not of valid type (must be FString or FSoftObjectPath)"), BankAttribute.GetDisplayText()));
 		return true;
 	}
 
 	// Validate that the "sequence index" attribute is of the right type
 	if (!PCG::Private::IsOfTypes<int32>(SequenceIndexAttributeBase->GetTypeId()))
 	{
-		PCGE_LOG_C(Error, GraphAndLog, &Context, FText::Format(LOCTEXT("AttributeInvalidType", "Attribute '{0}' is not of valid type (must be int32)"), SequenceIndexAttribute.GetDisplayText()));
+		PCGE_LOG_C(Error, GraphAndLog, &Context, FText::Format(LOCTEXT("SeqIndexAttributeInvalidType", "Attribute '{0}' is not of valid type (must be int32)"), SequenceIndexAttribute.GetDisplayText()));
 		return true;
 	}
 
 	// Validate that the "mesh" attribute is of the right type
 	if (!PCG::Private::IsOfTypes<FSoftObjectPath, FString>(MeshAttributeBase->GetTypeId()))
 	{
-		PCGE_LOG_C(Error, GraphAndLog, &Context, FText::Format(LOCTEXT("AttributeInvalidType", "Attribute '{0}' is not of valid type (must be FString or FSoftObjectPath)"), MeshAttribute.GetDisplayText()));
+		PCGE_LOG_C(Error, GraphAndLog, &Context, FText::Format(LOCTEXT("MeshAttributeInvalidType", "Attribute '{0}' is not of valid type (must be FString or FSoftObjectPath)"), MeshAttribute.GetDisplayText()));
 		return true;
 	}
 
