[PCG] Fixed bug with cyclical array update in Reroute Conversion

#rnx

diff --git a/Engine/Plugins/PCG/Source/PCG/Private/Helpers/PCGConversion.cpp b/Engine/Plugins/PCG/Source/PCG/Private/Helpers/PCGConversion.cpp
--- a/Engine/Plugins/PCG/Source/PCG/Private/Helpers/PCGConversion.cpp
+++ b/Engine/Plugins/PCG/Source/PCG/Private/Helpers/PCGConversion.cpp
@@ -217,14 +217,18 @@ bool FPCGRerouteDeclarationConverter::ExecuteApplyStructural(bool& bOutMarkGraph
 		checkf(SourceOutputPins.Num() == 2, TEXT("%s"), TEXT("This class currently only supports converting reroute nodes."));
 		bResult &= PCGConversion::Node::Helpers::RewireAllEdges(SourceOutputPins[0], GeneratedOutputPins[0]);
 
-		for (UPCGNode* UsageNode : SourceGraph->GetNodes())
+		const TArray<UPCGNode*> AllNodes = SourceGraph->GetNodes();
+		for (UPCGNode* UsageNode : AllNodes)
 		{
-			const UPCGNamedRerouteUsageSettings* Settings = Cast<UPCGNamedRerouteUsageSettings>(UsageNode->GetSettings());
-			if (Settings && Settings->Declaration == SourceSettings)
+			if (UsageNode)
 			{
-				check(UsageNode->GetOutputPins().Num() == 1);
-				bResult &= PCGConversion::Node::Helpers::RewireAllEdges(UsageNode->GetOutputPins()[0], GeneratedOutputPins[0]);
-				SourceGraph->RemoveNode(UsageNode);
+				const UPCGNamedRerouteUsageSettings* Settings = Cast<UPCGNamedRerouteUsageSettings>(UsageNode->GetSettings());
+				if (Settings && Settings->Declaration == SourceSettings)
+				{
+					check(UsageNode->GetOutputPins().Num() == 1);
+					bResult &= PCGConversion::Node::Helpers::RewireAllEdges(UsageNode->GetOutputPins()[0], GeneratedOutputPins[0]);
+					SourceGraph->RemoveNode(UsageNode);
+				}
 			}
 		}
 	}
