[PCG] Runtime Scheduler: Fix leaking graph instances when deleting original component

#rnx

diff --git a/Engine/Plugins/PCG/Source/PCG/Private/RuntimeGen/PCGRuntimeGenScheduler.cpp b/Engine/Plugins/PCG/Source/PCG/Private/RuntimeGen/PCGRuntimeGenScheduler.cpp
--- a/Engine/Plugins/PCG/Source/PCG/Private/RuntimeGen/PCGRuntimeGenScheduler.cpp
+++ b/Engine/Plugins/PCG/Source/PCG/Private/RuntimeGen/PCGRuntimeGenScheduler.cpp
@@ -799,6 +799,26 @@ void FPCGRuntimeGenScheduler::OnOriginalComponentUnregistered(UPCGComponent* InO
 			CleanupComponent(GenerationKey, ComponentToCleanup);
 		}
 	}
+
+	CleanupRemainingComponents(InOriginalComponent);
+}
+
+void FPCGRuntimeGenScheduler::CleanupRemainingComponents(UPCGComponent* InOriginalComponent)
+{
+	// Check for remaining PAs to cleanup
+	// There are some cases when on Refresh of the original component that GeneratedComponents doesn't contain all PAs anymore.
+	// If the PA is not far enough to be cleaned up but no longer is considered within the Gen Source Radius for Generation.
+	if (InOriginalComponent && InOriginalComponent->IsManagedByRuntimeGenSystem())
+	{
+		TSet<TObjectPtr<APCGPartitionActor>> PartitionActors = ActorAndComponentMapping->GetPCGComponentPartitionActorMappings(InOriginalComponent);
+		for (TObjectPtr<APCGPartitionActor> PartitionActor : PartitionActors)
+		{
+			if (UPCGComponent* ComponentToCleanup = PartitionActor->GetLocalComponent(InOriginalComponent))
+			{
+				CleanupLocalComponent(PartitionActor, ComponentToCleanup);
+			}
+		}
+	}
 }
 
 void FPCGRuntimeGenScheduler::CleanupLocalComponents(const APCGWorldActor* InPCGWorldActor)
@@ -810,6 +830,7 @@ void FPCGRuntimeGenScheduler::CleanupLocalComponents(const APCGWorldActor* InPCG
 	// Generated Entry Key, Local/Generated Component
 	using PCGComponentToClean = TTuple<FGridGenerationKey, UPCGComponent*>;
 	TArray<PCGComponentToClean> ComponentsToClean;
+	TSet<UPCGComponent*> OriginalComponents;
 
 	// Find all generated local components.
 	for (FGridGenerationKey GenerationKey : GeneratedComponents)
@@ -819,6 +840,7 @@ void FPCGRuntimeGenScheduler::CleanupLocalComponents(const APCGWorldActor* InPCG
 		const FIntVector& GridCoords = GenerationKey.GetGridCoords();
 		UPCGComponent* OriginalComponent = GenerationKey.GetOriginalComponent();
 		check(OriginalComponent);
+		OriginalComponents.Add(OriginalComponent);
 
 		// Only operate on LocalComponents.
 		if (Grid != EPCGHiGenGrid::Unbounded)
@@ -832,75 +854,49 @@ void FPCGRuntimeGenScheduler::CleanupLocalComponents(const APCGWorldActor* InPCG
 	{
 		CleanupComponent(ComponentToClean.Get<0>(), ComponentToClean.Get<1>());
 	}
+
+	for (UPCGComponent* OriginalComponent : OriginalComponents)
+	{
+		CleanupRemainingComponents(OriginalComponent);
+	}
 }
 
-void FPCGRuntimeGenScheduler::CleanupComponent(const FGridGenerationKey& GenerationKey, UPCGComponent* GeneratedComponent)
+void FPCGRuntimeGenScheduler::CleanupLocalComponent(APCGPartitionActor* PartitionActor, UPCGComponent* LocalComponent)
 {
-	TRACE_CPUPROFILER_EVENT_SCOPE(FPCGRuntimeGenScheduler::CleanupComponent);
-
-	check(ActorAndComponentMapping);
-
-	const FPCGGridDescriptor GridDescriptor = GenerationKey.GetGridDescriptor();
-	const EPCGHiGenGrid Grid = GridDescriptor.GetHiGenGrid();
-
-	const FIntVector GridCoords = GenerationKey.GetGridCoords();
-
-	APCGPartitionActor* PartitionActor = nullptr;
-
-	if (!GeneratedComponent)
+	if (!PartitionActor)
 	{
-		UE_LOG(LogPCG, Warning, TEXT("Runtime generated component could not be recovered on grid %d at (%d, %d, %d). It has been lost or destroyed."), GridDescriptor.GetGridSize(), GridCoords.X, GridCoords.Y, GridCoords.Z);
-
-		// If the GeneratedComponent has been lost for some reason, get the PA directly from the ActorAndComponentMapping.
-		PartitionActor = ActorAndComponentMapping->GetPartitionActor(GridDescriptor, GridCoords);
+		return;
 	}
-	else // If the GeneratedComponent does exist, we can clean it up!.
-	{
-		PartitionActor = Cast<APCGPartitionActor>(GeneratedComponent->GetOwner());
 
-		// If the Grid is unbounded, we have a non-partitioned or unbounded component.
-		if (Grid == EPCGHiGenGrid::Unbounded)
+	if (UPCGComponent* OriginalComponent = PartitionActor->GetOriginalComponent(LocalComponent))
+	{
+		if (PCGRuntimeGenSchedulerHelpers::CVarRuntimeGenerationEnableDebugging.GetValueOnAnyThread())
 		{
-			if (PCGRuntimeGenSchedulerHelpers::CVarRuntimeGenerationEnableDebugging.GetValueOnAnyThread() && GeneratedComponent->GetOwner())
-			{
-				UE_LOG(LogPCG, Warning, TEXT("[RUNTIMEGEN] CLEANUP: '%s'"), *GeneratedComponent->GetOwner()->GetActorNameOrLabel());
-			}
-
-			GeneratedComponent->CleanupLocalImmediate(/*bRemoveComponents=*/true);
+			UE_LOG(LogPCG, Warning, TEXT("[RUNTIMEGEN] CLEANUP: '%s'"), *PartitionActor->GetActorNameOrLabel());
 		}
-		// Otherwise we have a local component.
-		else if (ensure(PartitionActor))
-		{
-			UPCGComponent* OriginalComponent = PartitionActor->GetOriginalComponent(GeneratedComponent);
 
-			if (PCGRuntimeGenSchedulerHelpers::CVarRuntimeGenerationEnableDebugging.GetValueOnAnyThread())
-			{
-				UE_LOG(LogPCG, Warning, TEXT("[RUNTIMEGEN] CLEANUP: '%s'"), *PartitionActor->GetActorNameOrLabel());
-			}
+		// This performs a CleanupLocalImmediate for us, no need to clean up ourselves.
+		PartitionActor->RemoveGraphInstance(OriginalComponent);
 
-			// This performs a CleanupLocalImmediate for us, no need to clean up ourselves.
-			PartitionActor->RemoveGraphInstance(OriginalComponent);
+		// Remove component mapping.
+		{
+			FWriteScopeLock WriteLock(ActorAndComponentMapping->ComponentToPartitionActorsMapLock);
+			TSet<TObjectPtr<APCGPartitionActor>>* PartitionActorsPtr = ActorAndComponentMapping->ComponentToPartitionActorsMap.Find(OriginalComponent);
 
-			// Remove component mapping.
+			if (PartitionActorsPtr)
 			{
-				FWriteScopeLock WriteLock(ActorAndComponentMapping->ComponentToPartitionActorsMapLock);
-				TSet<TObjectPtr<APCGPartitionActor>>* PartitionActorsPtr = ActorAndComponentMapping->ComponentToPartitionActorsMap.Find(OriginalComponent);
+				PartitionActorsPtr->Remove(PartitionActor);
 
-				if (PartitionActorsPtr)
+				if (PartitionActorsPtr->IsEmpty())
 				{
-					PartitionActorsPtr->Remove(PartitionActor);
-
-					if (PartitionActorsPtr->IsEmpty())
-					{
-						ActorAndComponentMapping->ComponentToPartitionActorsMap.Remove(OriginalComponent);
-					}
+					ActorAndComponentMapping->ComponentToPartitionActorsMap.Remove(OriginalComponent);
 				}
 			}
 		}
 	}
 
 	// Cleanup the PA if it no longer has any components (return to pool or destroy).
-	if (PartitionActor && PartitionActor->GetAllLocalPCGComponents().IsEmpty())
+	if (PartitionActor->GetAllLocalPCGComponents().IsEmpty())
 	{
 		PartitionActor->UnregisterPCG();
 
@@ -930,7 +926,46 @@ void FPCGRuntimeGenScheduler::CleanupComponent(const FGridGenerationKey& Generat
 			World->DestroyActor(PartitionActor);
 		}
 	}
+}
 
+void FPCGRuntimeGenScheduler::CleanupComponent(const FGridGenerationKey& GenerationKey, UPCGComponent* GeneratedComponent)
+{
+	TRACE_CPUPROFILER_EVENT_SCOPE(FPCGRuntimeGenScheduler::CleanupComponent);
+
+	check(ActorAndComponentMapping);
+
+	const FPCGGridDescriptor GridDescriptor = GenerationKey.GetGridDescriptor();
+	const EPCGHiGenGrid Grid = GridDescriptor.GetHiGenGrid();
+
+	const FIntVector GridCoords = GenerationKey.GetGridCoords();
+
+	APCGPartitionActor* PartitionActor = nullptr;
+
+	if (!GeneratedComponent)
+	{
+		UE_LOG(LogPCG, Warning, TEXT("Runtime generated component could not be recovered on grid %d at (%d, %d, %d). It has been lost or destroyed."), GridDescriptor.GetGridSize(), GridCoords.X, GridCoords.Y, GridCoords.Z);
+
+		// If the GeneratedComponent has been lost for some reason, get the PA directly from the ActorAndComponentMapping.
+		PartitionActor = ActorAndComponentMapping->GetPartitionActor(GridDescriptor, GridCoords);
+	}
+	else // If the GeneratedComponent does exist, we can clean it up!.
+	{
+		PartitionActor = Cast<APCGPartitionActor>(GeneratedComponent->GetOwner());
+		ensure(PartitionActor || Grid == EPCGHiGenGrid::Unbounded);
+
+		if (Grid == EPCGHiGenGrid::Unbounded)
+		{
+			if (PCGRuntimeGenSchedulerHelpers::CVarRuntimeGenerationEnableDebugging.GetValueOnAnyThread() && GeneratedComponent->GetOwner())
+			{
+				UE_LOG(LogPCG, Warning, TEXT("[RUNTIMEGEN] CLEANUP: '%s'"), *GeneratedComponent->GetOwner()->GetActorNameOrLabel());
+			}
+
+			GeneratedComponent->CleanupLocalImmediate(/*bRemoveComponents=*/true);
+		}
+	}
+
+	CleanupLocalComponent(PartitionActor, GeneratedComponent);
+	
 	GeneratedComponents.Remove(GenerationKey);
 }
 
diff --git a/Engine/Plugins/PCG/Source/PCG/Public/RuntimeGen/PCGRuntimeGenScheduler.h b/Engine/Plugins/PCG/Source/PCG/Public/RuntimeGen/PCGRuntimeGenScheduler.h
--- a/Engine/Plugins/PCG/Source/PCG/Public/RuntimeGen/PCGRuntimeGenScheduler.h
+++ b/Engine/Plugins/PCG/Source/PCG/Public/RuntimeGen/PCGRuntimeGenScheduler.h
@@ -103,6 +103,12 @@ protected:
 	void ResetPartitionActorPoolToSize(uint32 NewPoolSize);
 
 private:
+	/** Cleanup a local compnent from the Partition actor. */
+	void CleanupLocalComponent(APCGPartitionActor* PartitionActor, UPCGComponent* LocalComponent);
+
+	/** Cleanup the remaining comopnents that are not part of the GeneratedComponents */
+	void CleanupRemainingComponents(UPCGComponent* InOriginalComponent);
+
 	/** Tracks the generated components managed by the RuntimeGenScheduler. For local components, this generation key will hold the original component.
 	* For non-partitioned components, the generation key should have unbounded grid size and (0, 0, 0) grid coordinates.
 	*/
