[PCG] Fix -game usage of UActorEditorContextSubsystem

diff --git a/Engine/Plugins/PCG/Source/PCG/Private/Grid/PCGPartitionActor.cpp b/Engine/Plugins/PCG/Source/PCG/Private/Grid/PCGPartitionActor.cpp
--- a/Engine/Plugins/PCG/Source/PCG/Private/Grid/PCGPartitionActor.cpp
+++ b/Engine/Plugins/PCG/Source/PCG/Private/Grid/PCGPartitionActor.cpp
@@ -630,15 +630,19 @@ void APCGPartitionActor::PostCreation(const FPCGGridDescriptor& GridDescriptor)
 	bUse2DGrid = GridDescriptor.Is2DGrid();
 
 #if WITH_EDITOR
-	// Fetch Non External Assets and assign
-	TArray<TSoftObjectPtr<UDataLayerAsset>> DescriptorDataLayerAssets;
-	const UExternalDataLayerAsset* DescriptorExternalDataLayerAsset = nullptr;
-	GridDescriptor.GetDataLayerAssets(DescriptorDataLayerAssets, DescriptorExternalDataLayerAsset);
+	if (!GetWorld()->IsGameWorld())
+	{
+		// Fetch Non External Assets and assign
+		TArray<TSoftObjectPtr<UDataLayerAsset>> DescriptorDataLayerAssets;
+		const UExternalDataLayerAsset* DescriptorExternalDataLayerAsset = nullptr;
+		GridDescriptor.GetDataLayerAssets(DescriptorDataLayerAssets, DescriptorExternalDataLayerAsset);
+
+		// External DataLayer should have been assigned on Spawn
+		ensure(ExternalDataLayerAsset == DescriptorExternalDataLayerAsset);
 
-	// External DataLayer should have been assigned on Spawn
-	ensure(ExternalDataLayerAsset == DescriptorExternalDataLayerAsset);
+		DataLayerAssets = DescriptorDataLayerAssets;
+	}
 
-	DataLayerAssets = DescriptorDataLayerAssets;
 	// Set only once upon creation, can't change for a Partition Actor
 	RuntimeGridDescriptorHash = GridDescriptor.GetRuntimeHash();
 	SetGridSize(PCGGridSize);
diff --git a/Engine/Plugins/PCG/Source/PCG/Private/Helpers/PCGActorHelpers.cpp b/Engine/Plugins/PCG/Source/PCG/Private/Helpers/PCGActorHelpers.cpp
--- a/Engine/Plugins/PCG/Source/PCG/Private/Helpers/PCGActorHelpers.cpp
+++ b/Engine/Plugins/PCG/Source/PCG/Private/Helpers/PCGActorHelpers.cpp
@@ -687,18 +687,31 @@ AActor* UPCGActorHelpers::SpawnDefaultActor(const FSpawnDefaultActorParams& Para
 #if WITH_EDITOR
 	// Find External Data Layer if it exists so we can create scope that will allow actor to be properly created
 	const UExternalDataLayerInstance* ExternalDataLayerInstance = nullptr;
-	for (const UDataLayerInstance* DataLayerInstance : Params.DataLayerInstances)
+	bool bPushedContext = false;
+	// No need to do any DataLayer assignment in a game world
+	if (!Params.World->IsGameWorld())
 	{
-		if (const UExternalDataLayerInstance* Found = Cast<UExternalDataLayerInstance>(DataLayerInstance))
+		for (const UDataLayerInstance* DataLayerInstance : Params.DataLayerInstances)
 		{
-			ExternalDataLayerInstance = Found;
-			break;
+			if (const UExternalDataLayerInstance* Found = Cast<UExternalDataLayerInstance>(DataLayerInstance))
+			{
+				ExternalDataLayerInstance = Found;
+				break;
+			}
 		}
+
+		// Avoid relying on the Editor Context at all
+		UActorEditorContextSubsystem::Get()->PushContext();
+		bPushedContext = true;
 	}
 
-	// Avoid relying on the Editor Context at all
-	UActorEditorContextSubsystem::Get()->PushContext();
-	ON_SCOPE_EXIT{ UActorEditorContextSubsystem::Get()->PopContext(); };
+	ON_SCOPE_EXIT
+	{ 
+		if (bPushedContext)
+		{
+			UActorEditorContextSubsystem::Get()->PopContext();
+		}
+	};
 
 	// Specify EDL we want to use if any for spawning this actor
 	FScopedOverrideSpawningLevelMountPointObject EDLScope(ExternalDataLayerInstance ? ExternalDataLayerInstance->GetExternalDataLayerAsset() : nullptr);
@@ -720,12 +733,15 @@ AActor* UPCGActorHelpers::SpawnDefaultActor(const FSpawnDefaultActorParams& Para
 	NewActor->SetActorRelativeScale3D(Params.Transform.GetScale3D());
 
 #if WITH_EDITOR
-	// Add remaining DataLayers (except External which was done on spawn)
-	for (const UDataLayerInstance* DataLayerInstance : Params.DataLayerInstances)
+	if (!Params.World->IsGameWorld())
 	{
-		if (DataLayerInstance != ExternalDataLayerInstance)
+		// Add remaining DataLayers (except External which was done on spawn)
+		for (const UDataLayerInstance* DataLayerInstance : Params.DataLayerInstances)
 		{
-			DataLayerInstance->AddActor(NewActor);
+			if (DataLayerInstance != ExternalDataLayerInstance)
+			{
+				DataLayerInstance->AddActor(NewActor);
+			}
 		}
 	}
 	
diff --git a/Engine/Plugins/PCG/Source/PCG/Private/PCGSubsystem.cpp b/Engine/Plugins/PCG/Source/PCG/Private/PCGSubsystem.cpp
--- a/Engine/Plugins/PCG/Source/PCG/Private/PCGSubsystem.cpp
+++ b/Engine/Plugins/PCG/Source/PCG/Private/PCGSubsystem.cpp
@@ -1263,14 +1263,28 @@ APCGPartitionActor* UPCGSubsystem::FindOrCreatePCGPartitionActor(const FPCGGridD
 	}
 		
 #if WITH_EDITOR
-	TArray<TSoftObjectPtr<UDataLayerAsset>> DataLayerAssets;
+	bool bPushedContext = false;
 	const UExternalDataLayerAsset* ExternalDataLayerAsset = nullptr;
-			
-	GridDescriptor.GetDataLayerAssets(DataLayerAssets, ExternalDataLayerAsset);
+	
+	// No need to do any DataLayer assignment in a game world
+	if (!World->IsGameWorld())
+	{
+		TArray<TSoftObjectPtr<UDataLayerAsset>> DataLayerAssets;
+		
+		GridDescriptor.GetDataLayerAssets(DataLayerAssets, ExternalDataLayerAsset);
 
-	// Avoid relying on the Editor Context at all
-	UActorEditorContextSubsystem::Get()->PushContext();
-	ON_SCOPE_EXIT{ UActorEditorContextSubsystem::Get()->PopContext(); };
+		// Avoid relying on the Editor Context at all
+		UActorEditorContextSubsystem::Get()->PushContext();
+		bPushedContext = true;
+	}
+
+	ON_SCOPE_EXIT
+	{ 
+		if (bPushedContext)
+		{
+			UActorEditorContextSubsystem::Get()->PopContext();
+		}
+	};
 
 	// Specify EDL we want to use if any for spawning this actor
 	FScopedOverrideSpawningLevelMountPointObject EDLScope(ExternalDataLayerAsset);
