Added template functionality to graphs and added a 'initialize from template' part in the pcg graph creation process

diff --git a/Engine/Plugins/PCG/Source/PCG/Public/PCGGraph.h b/Engine/Plugins/PCG/Source/PCG/Public/PCGGraph.h
--- a/Engine/Plugins/PCG/Source/PCG/Public/PCGGraph.h
+++ b/Engine/Plugins/PCG/Source/PCG/Public/PCGGraph.h
@@ -129,6 +129,9 @@ protected:
 	UFUNCTION()
 	virtual bool AreOverridesEnabled() const { return IsExportToLibraryEnabled() && bExposeToLibrary; }
 
+	UFUNCTION()
+	virtual bool IsTemplatePropertyEnabled() const { return false; }
+
 	bool VerifyAndUpdateIfGraphParameterValueChanged(const FPropertyChangedChainEvent& PropertyChangedEvent);
 #endif // WITH_EDITOR
 
@@ -141,6 +144,9 @@ public:
 	UPROPERTY(BlueprintReadWrite, EditAnywhere, Category = AssetInfo, AssetRegistrySearchable, meta = (EditCondition = "AreOverridesEnabled", EditConditionHides))
 	bool bOverrideTitle = false;
 
+	UPROPERTY(BlueprintReadWrite, EditAnywhere, Category = AssetInfo, AssetRegistrySearchable, meta = (EditCondition = "IsTemplatePropertyEnabled", EditConditionHides))
+	bool bIsTemplate = false;
+
 	/** Override of the title for the subgraph node for this graph. */
 	UPROPERTY(BlueprintReadWrite, EditAnywhere, Category = AssetInfo, AssetRegistrySearchable, meta = (EditCondition = "AreOverridesEnabled && bOverrideTitle", EditConditionHides))
 	FText Title;
@@ -405,6 +411,10 @@ public:
 	const TObjectPtr<UPCGGraphCompilationData> GetCookedCompilationData() const { return CookedCompilationData; }
 
 protected:
+#if WITH_EDITOR
+	virtual bool IsTemplatePropertyEnabled() const override { return IsAsset(); }
+#endif
+
 	/** Internal function to react to add/remove nodes. bNotify can be set to false to not notify the world. */
 	void OnNodeAdded(UPCGNode* InNode, bool bNotify = true);
 	void OnNodesAdded(TArrayView<UPCGNode*> InNodes, bool bNotify = true);
diff --git a/Engine/Plugins/PCG/Source/PCGEditor/PCGEditor.Build.cs b/Engine/Plugins/PCG/Source/PCGEditor/PCGEditor.Build.cs
--- a/Engine/Plugins/PCG/Source/PCGEditor/PCGEditor.Build.cs
+++ b/Engine/Plugins/PCG/Source/PCGEditor/PCGEditor.Build.cs
@@ -37,6 +37,7 @@ namespace UnrealBuildTool.Rules
 					"AssetRegistry",
 					"BlueprintGraph",
 					"ContentBrowser",
+					"ContentBrowserData",
 					"DesktopWidgets",
 					"DetailCustomizations",
 					"DeveloperSettings",
diff --git a/Engine/Plugins/PCG/Source/PCGEditor/Private/PCGEditorUtils.cpp b/Engine/Plugins/PCG/Source/PCGEditor/Private/PCGEditorUtils.cpp
--- a/Engine/Plugins/PCG/Source/PCGEditor/Private/PCGEditorUtils.cpp
+++ b/Engine/Plugins/PCG/Source/PCGEditor/Private/PCGEditorUtils.cpp
@@ -7,13 +7,33 @@
 #include "Elements/PCGExecuteBlueprint.h"
 
 #include "AssetToolsModule.h"
+#include "ContentBrowserModule.h"
+#include "FrontendFilterBase.h"
+#include "IAssetTools.h"
+#include "IContentBrowserSingleton.h"
+#include "SPrimaryButton.h"
+#include "SSimpleButton.h"
+#include "Algo/Transform.h"
+#include "AssetRegistry/AssetData.h"
 #include "AssetRegistry/ARFilter.h"
 #include "AssetRegistry/AssetRegistryModule.h"
 #include "Blueprint/BlueprintSupport.h"
-#include "IAssetTools.h"
+#include "Filters/SBasicFilterBar.h"
+#include "Framework/Application/SlateApplication.h"
+#include "Interfaces/IMainFrameModule.h"
 #include "Kismet2/BlueprintEditorUtils.h"
 #include "Kismet2/KismetEditorUtilities.h"
 #include "Misc/PackageName.h"
+#include "Styling/AppStyle.h"
+#include "Types/SlateEnums.h"
+#include "Widgets/Layout/SBox.h"
+#include "Widgets/Input/SButton.h"
+#include "Widgets/Layout/SBorder.h"
+#include "Widgets/Layout/SSpacer.h"
+#include "Widgets/SBoxPanel.h"
+#include "Widgets/SWindow.h"
+
+#define LOCTEXT_NAMESPACE "PCGEditorUtils"
 
 bool PCGEditorUtils::IsAssetPCGBlueprint(const FAssetData& InAssetData)
 {
@@ -134,3 +154,199 @@ void PCGEditorUtils::ForcePCGBlueprintVariableVisibility()
 		return true;
 	});
 }
+
+class FFrontendFilter_PCGGraphTemplate : public FFrontendFilter
+{
+public:
+	FFrontendFilter_PCGGraphTemplate(TSharedPtr<FFrontendFilterCategory> FilterCategory, const FString& InCategory) 
+		: FFrontendFilter(FilterCategory)
+		, Category(InCategory)
+	{
+		TArray<FString> Tokens;
+		Category.ParseIntoArray(Tokens, TEXT("|"), true);
+
+		TArray<FString> DisplayNames;
+		Algo::Transform(Tokens, DisplayNames, [](const FString& Token) { return FName::NameToDisplayString(Token, /*bIsBool=*/false); });
+
+		const FString JoinDelimiter(" | ");
+		DisplayName = FText::FromString(FString::Join(DisplayNames, *JoinDelimiter));
+	}
+
+	virtual FString GetName() const override { return Category; }
+	virtual FText GetDisplayName() const override { return DisplayName; }
+	virtual FText GetToolTipText() const override { return FText(); }
+	virtual FLinearColor GetColor() const override { return FLinearColor(0.7f, 0.7f, 0.7f); }
+
+	virtual bool PassesFilter(const FContentBrowserItem& InItem) const override
+	{
+		FAssetData AssetData;
+		if (InItem.Legacy_TryGetAssetData(AssetData))
+		{
+			FString AssetCategory = AssetData.GetTagValueRef<FString>(GET_MEMBER_NAME_CHECKED(UPCGGraph, Category));
+
+			if (!AssetCategory.IsEmpty() && AssetCategory.StartsWith(Category))
+			{
+				return true;
+			}
+		}
+
+		return false;
+	}
+
+protected:
+	FString Category;
+	FText DisplayName;
+};
+
+bool PCGEditorUtils::PickGraphTemplate(FAssetData& OutAssetData, const FText& TitleOverride)
+{
+	auto HandleAssetSelected = [&OutAssetData](const FAssetData& InSelectedAsset)
+	{
+		OutAssetData = InSelectedAsset;
+	};
+
+	FAssetPickerConfig AssetPickerConfig;
+	AssetPickerConfig.OnAssetEnterPressed = FOnAssetEnterPressed::CreateLambda([HandleAssetSelected](const TArray<FAssetData>& SelectedAssetData)
+	{
+		if (SelectedAssetData.Num() == 1)
+		{
+			HandleAssetSelected(SelectedAssetData[0]);
+		}
+	});
+	AssetPickerConfig.OnAssetSelected = FOnAssetSelected::CreateLambda(HandleAssetSelected);
+	AssetPickerConfig.SelectionMode = ESelectionMode::Single;
+	AssetPickerConfig.bAllowNullSelection = false;
+	AssetPickerConfig.InitialAssetViewType = EAssetViewType::List;
+	AssetPickerConfig.Filter.ClassPaths.Add(UPCGGraph::StaticClass()->GetClassPathName());
+	AssetPickerConfig.OnShouldFilterAsset = FOnShouldFilterAsset::CreateLambda([](const FAssetData& AssetData)
+	{
+		return !AssetData.GetTagValueRef<bool>(GET_MEMBER_NAME_CHECKED(UPCGGraph, bIsTemplate));
+	});
+
+	TSet<FString> CategoryList;
+	ForEachPCGGraphAssetData([&CategoryList](const FAssetData& Asset)
+	{
+		if (Asset.GetTagValueRef<bool>(GET_MEMBER_NAME_CHECKED(UPCGGraph, bIsTemplate)))
+		{
+			FString Category = Asset.GetTagValueRef<FString>(GET_MEMBER_NAME_CHECKED(UPCGGraph, Category));
+			if (!Category.IsEmpty())
+			{
+				CategoryList.Add(MoveTemp(Category));
+			}
+		}
+		return true;
+	});
+
+	TSharedPtr<FFrontendFilterCategory> FilterCategory = MakeShared<FFrontendFilterCategory>(
+		LOCTEXT("GraphTemplateCategoryName", "PCG Graph Template Categories"),
+		LOCTEXT("GraphTemplateCategoryName_Tooltip", "Filter templates by categories.")
+	);
+
+	for (const FString& Category : CategoryList)
+	{
+		AssetPickerConfig.ExtraFrontendFilters.Add(MakeShared<FFrontendFilter_PCGGraphTemplate>(FilterCategory, Category));
+	}
+
+	// This is so that we can remove the "Other Filters" section easily
+	AssetPickerConfig.bUseSectionsForCustomFilterCategories = true;
+	// Make sure we only show PCG filters to avoid confusion 
+	AssetPickerConfig.OnExtendAddFilterMenu = FOnExtendAddFilterMenu::CreateLambda([](UToolMenu* InToolMenu)
+	{
+		// "AssetFilterBarFilterAdvancedAsset" taken from SAssetFilterBar.h PopulateAddFilterMenu()
+		InToolMenu->RemoveSection("AssetFilterBarFilterAdvancedAsset");
+		InToolMenu->RemoveSection("Other Filters");
+	});
+	
+	AssetPickerConfig.bAddFilterUI = !CategoryList.IsEmpty();
+
+	FContentBrowserModule& ContentBrowserModule = FModuleManager::Get().LoadModuleChecked<FContentBrowserModule>(TEXT("ContentBrowser"));
+	FText Title = (TitleOverride.IsEmpty() ? LOCTEXT("SelectTemplateDialogTitle", "Initialize from Graph Template...") : TitleOverride);
+
+	bool bClickedOk = false;
+	TSharedPtr<SWindow> Dialog;
+	Dialog = SNew(SWindow)
+		.Title(Title)
+		.SizingRule(ESizingRule::Autosized);
+
+	Dialog->SetContent(
+		SNew(SVerticalBox)
+		+ SVerticalBox::Slot()
+		.AutoHeight()
+		.Padding(5.0f)
+		[
+			SNew(SBox)
+			.WidthOverride(300.0f)
+			.HeightOverride(400.0f)
+			[
+				ContentBrowserModule.Get().CreateAssetPicker(AssetPickerConfig)
+			]
+		]
+		+ SVerticalBox::Slot()
+		.AutoHeight()
+		.Padding(5.0f)
+		[
+			SNew(SHorizontalBox)
+			+ SHorizontalBox::Slot()
+			.AutoWidth()
+			.HAlign(HAlign_Left)
+			[
+				SNew(SButton)
+				.OnClicked_Lambda([&bClickedOk, &OutAssetData, WeakDialog = TWeakPtr<SWindow>(Dialog)]()
+				{
+					bClickedOk = true;
+					if (TSharedPtr<SWindow> PinnedDialog = WeakDialog.Pin())
+					{
+						PinnedDialog->RequestDestroyWindow();
+					}
+
+					OutAssetData = FAssetData();
+					return FReply::Handled();
+				})
+				.Text(LOCTEXT("InitializeFromEmptyTemplateButton", "Create empty graph"))
+			]
+			+ SHorizontalBox::Slot()
+			[
+				SNew(SSpacer)
+			]
+			+ SHorizontalBox::Slot()
+			.AutoWidth()
+			[
+				SNew(SPrimaryButton)
+				.OnClicked_Lambda([&bClickedOk, WeakDialog = TWeakPtr<SWindow>(Dialog)]()
+				{
+					bClickedOk = true;
+					if (TSharedPtr<SWindow> PinnedDialog = WeakDialog.Pin())
+					{
+						PinnedDialog->RequestDestroyWindow();
+					}
+
+					return FReply::Handled();
+				})
+				.IsEnabled_Lambda([&OutAssetData]()
+				{
+					return OutAssetData.IsValid();
+				})
+				.Text(LOCTEXT("InitializeFromTemplateButton", "Initialize From Template"))
+			]
+			+SHorizontalBox::Slot()
+			.AutoWidth()
+			[
+				SNew(SButton)
+				.OnClicked_Lambda([WeakDialog = TWeakPtr<SWindow>(Dialog)]()
+				{
+					if (TSharedPtr<SWindow> PinnedDialog = WeakDialog.Pin())
+					{
+						PinnedDialog->RequestDestroyWindow();
+					}
+
+					return FReply::Handled();
+				})
+				.Text(LOCTEXT("CancelButton", "Cancel"))
+			]
+		]);
+
+	FSlateApplication::Get().AddModalWindow(Dialog.ToSharedRef(), FGlobalTabmanager::Get()->GetRootWindow());
+	return bClickedOk;
+}
+
+#undef LOCTEXT_NAMESPACE
\ No newline at end of file
diff --git a/Engine/Plugins/PCG/Source/PCGEditor/Private/PCGEditorUtils.h b/Engine/Plugins/PCG/Source/PCGEditor/Private/PCGEditorUtils.h
--- a/Engine/Plugins/PCG/Source/PCGEditor/Private/PCGEditorUtils.h
+++ b/Engine/Plugins/PCG/Source/PCGEditor/Private/PCGEditorUtils.h
@@ -2,13 +2,13 @@
 
 #pragma once
 
+#include "Internationalization/Text.h"
 
 template <typename FuncType> class TFunctionRef;
-
 struct FAssetData;
 struct FARFilter;
-class UObject;
 class FString;
+class UObject;
 
 namespace PCGEditorUtils
 {
@@ -31,4 +31,7 @@ namespace PCGEditorUtils
 
 	/** Asset deprecation methods */
 	void ForcePCGBlueprintVariableVisibility();
+
+	/** Opens a template picker dialog. Returns true if user selected 'ok'. */
+	bool PickGraphTemplate(FAssetData& OutData, const FText& InTitleOverride = FText());
 }
diff --git a/Engine/Plugins/PCG/Source/PCGEditor/Private/PCGGraphFactory.cpp b/Engine/Plugins/PCG/Source/PCGEditor/Private/PCGGraphFactory.cpp
--- a/Engine/Plugins/PCG/Source/PCGEditor/Private/PCGGraphFactory.cpp
+++ b/Engine/Plugins/PCG/Source/PCGEditor/Private/PCGGraphFactory.cpp
@@ -1,12 +1,19 @@
 // Copyright Epic Games, Inc. All Rights Reserved.
 
 #include "PCGGraphFactory.h"
+
+#include "PCGEditorSettings.h"
+#include "PCGEditorUtils.h"
 #include "PCGGraph.h"
 
+#include "AssetRegistry/AssetData.h"
+
 /////////////////////////
 // PCGGraph
 /////////////////////////
 
+#define LOCTEXT_NAMESPACE "PCGGraphFactory"
+
 UPCGGraphFactory::UPCGGraphFactory(const FObjectInitializer& ObjectInitializer)
 	: Super(ObjectInitializer)
 {
@@ -17,7 +24,22 @@ UPCGGraphFactory::UPCGGraphFactory(const FObjectInitializer& ObjectInitializer)
 
 UObject* UPCGGraphFactory::FactoryCreateNew(UClass* InClass, UObject* InParent, FName InName, EObjectFlags Flags, UObject* Context, FFeedbackContext* Warn)
 {
-	return NewObject<UPCGGraph>(InParent, InClass, InName, Flags);
+	UPCGGraph* NewGraph = nullptr;
+	
+	if (TemplateGraph)
+	{
+		NewGraph = Cast<UPCGGraph>(StaticDuplicateObject(TemplateGraph, InParent, InName, Flags, InClass));
+		if (NewGraph)
+		{
+			NewGraph->bIsTemplate = false;
+		}
+	}
+	else
+	{
+		NewGraph = NewObject<UPCGGraph>(InParent, InClass, InName, Flags);
+	}
+
+	return NewGraph;
 }
 
 bool UPCGGraphFactory::ShouldShowInNewMenu() const
@@ -25,6 +47,49 @@ bool UPCGGraphFactory::ShouldShowInNewMenu() const
 	return true;
 }
 
+bool UPCGGraphFactory::ConfigureProperties()
+{
+	// Early out if the editor config says we don't want to have a dialog for it.
+	const UPCGEditorSettings* EditorSettings = GetDefault<UPCGEditorSettings>();
+	if (EditorSettings && !EditorSettings->bShowTemplatePickerOnNewGraph)
+	{
+		return true;
+	}
+
+	// Check if templates actually exist
+	bool bTemplatesExist = false;
+	PCGEditorUtils::ForEachPCGGraphAssetData([&bTemplatesExist](const FAssetData& AssetData)
+	{
+		if (AssetData.IsInstanceOf<UPCGGraph>())
+		{
+			const bool bIsTemplate = AssetData.GetTagValueRef<bool>(GET_MEMBER_NAME_CHECKED(UPCGGraph, bIsTemplate));
+			if (bIsTemplate)
+			{
+				bTemplatesExist = true;
+				return false;
+			}
+		}
+
+		return true;
+	});
+
+	if (!bTemplatesExist)
+	{
+		return true;
+	}
+
+	// Select template
+	FAssetData SelectedTemplate;
+	const bool bPicked = PCGEditorUtils::PickGraphTemplate(SelectedTemplate, LOCTEXT("TemplatePickerTitle", "Create Graph From Template..."));
+
+	if (bPicked && SelectedTemplate.IsValid())
+	{
+		TemplateGraph = Cast<UPCGGraph>(SelectedTemplate.GetAsset());
+	}
+
+	return bPicked;
+}
+
 /////////////////////////
 // PCGGraphInstance
 /////////////////////////
@@ -52,4 +117,6 @@ UObject* UPCGGraphInstanceFactory::FactoryCreateNew(UClass* InClass, UObject* In
 bool UPCGGraphInstanceFactory::ShouldShowInNewMenu() const
 {
 	return true;
-}
\ No newline at end of file
+}
+
+#undef LOCTEXT_NAMESPACE
\ No newline at end of file
diff --git a/Engine/Plugins/PCG/Source/PCGEditor/Private/PCGGraphFactory.h b/Engine/Plugins/PCG/Source/PCGEditor/Private/PCGGraphFactory.h
--- a/Engine/Plugins/PCG/Source/PCGEditor/Private/PCGGraphFactory.h
+++ b/Engine/Plugins/PCG/Source/PCGEditor/Private/PCGGraphFactory.h
@@ -6,6 +6,7 @@
 
 #include "PCGGraphFactory.generated.h"
 
+class UPCGGraph;
 class UPCGGraphInterface;
 
 UCLASS(hidecategories=Object)
@@ -16,8 +17,12 @@ class UPCGGraphFactory : public UFactory
 public:
 	UPCGGraphFactory(const FObjectInitializer& ObjectInitializer);
 
+	UPROPERTY()
+	TObjectPtr<UPCGGraph> TemplateGraph;
+
 	//~UFactory interface
 	virtual UObject* FactoryCreateNew(UClass* InClass, UObject* InParent, FName InName, EObjectFlags Flags, UObject* Context, FFeedbackContext* Warn) override;
+	virtual bool ConfigureProperties() override;
 	virtual bool ShouldShowInNewMenu() const override;
 };
 
diff --git a/Engine/Plugins/PCG/Source/PCGEditor/Public/PCGEditorSettings.h b/Engine/Plugins/PCG/Source/PCGEditor/Public/PCGEditorSettings.h
--- a/Engine/Plugins/PCG/Source/PCGEditor/Public/PCGEditorSettings.h
+++ b/Engine/Plugins/PCG/Source/PCGEditor/Public/PCGEditorSettings.h
@@ -234,4 +234,8 @@ public:
 	/** Target number of points when trying to debug volume data or landscape data, to avoid exploding the number of points if the volume/landscape is big, and to have a good representation if the volume/landscape is small. */
 	UPROPERTY(EditAnywhere, config, Category = "Editor Performance")
 	int32 TargetNumPointsForDebug = 64000;
+
+	/** Controls whether the template picker dialog will be shown when creating a new PCG graph, both in the contextual menu and in the PCG graph pickers. */
+	UPROPERTY(EditAnywhere, config, Category = "Content Browser")
+	bool bShowTemplatePickerOnNewGraph = true;
 };
