[PCG] Fix bug with dependency pin being linked unintentionally with all InputPins in Dynamic Pin Nodes

#rnx

diff --git a/Engine/Plugins/PCG/Source/PCG/Private/Elements/PCGMergeAttributes.cpp b/Engine/Plugins/PCG/Source/PCG/Private/Elements/PCGMergeAttributes.cpp
--- a/Engine/Plugins/PCG/Source/PCG/Private/Elements/PCGMergeAttributes.cpp
+++ b/Engine/Plugins/PCG/Source/PCG/Private/Elements/PCGMergeAttributes.cpp
@@ -72,8 +72,10 @@ bool FPCGMergeAttributesElement::ExecuteInternal(FPCGContext* Context) const
 	for (const UPCGPin* InputPin : Context->Node->GetInputPins())
 	{
 		check(InputPin);
-
-		Sources.Append(Context->InputData.GetInputsByPin(InputPin->Properties.Label));
+		if (!InputPin->IsDatalessPin())
+		{
+			Sources.Append(Context->InputData.GetInputsByPin(InputPin->Properties.Label));
+		}
 	}
 
 	for (const FPCGTaggedData& Source : Sources)
diff --git a/Engine/Plugins/PCG/Source/PCG/Private/Elements/PCGMergeElement.cpp b/Engine/Plugins/PCG/Source/PCG/Private/Elements/PCGMergeElement.cpp
--- a/Engine/Plugins/PCG/Source/PCG/Private/Elements/PCGMergeElement.cpp
+++ b/Engine/Plugins/PCG/Source/PCG/Private/Elements/PCGMergeElement.cpp
@@ -61,7 +61,10 @@ bool FPCGMergeElement::ExecuteInternal(FPCGContext* Context) const
 	{
 		check(InputPin);
 
-		Sources.Append(Context->InputData.GetInputsByPin(InputPin->Properties.Label));
+		if (!InputPin->IsDatalessPin())
+		{
+			Sources.Append(Context->InputData.GetInputsByPin(InputPin->Properties.Label));
+		}
 	}
 
 	TArray<FPCGTaggedData>& Outputs = Context->OutputData.TaggedData;
diff --git a/Engine/Plugins/PCG/Source/PCG/Private/Elements/PCGOuterIntersectionElement.cpp b/Engine/Plugins/PCG/Source/PCG/Private/Elements/PCGOuterIntersectionElement.cpp
--- a/Engine/Plugins/PCG/Source/PCG/Private/Elements/PCGOuterIntersectionElement.cpp
+++ b/Engine/Plugins/PCG/Source/PCG/Private/Elements/PCGOuterIntersectionElement.cpp
@@ -92,9 +92,10 @@ bool FPCGOuterIntersectionElement::ExecuteInternal(FPCGContext* Context) const
 		{
 			check(InputPin);
 			// Only checking secondary source pins expecting spatial data
-			if (InputPin->Properties.Label != PCGIntersectionConstants::PrimaryLabel &&
-				!!(InputPin->Properties.AllowedTypes & EPCGDataType::Spatial) &&
-				Context->InputData.GetSpatialInputCountByPin(InputPin->Properties.Label) == 0)
+			if (InputPin->Properties.Label != PCGIntersectionConstants::PrimaryLabel
+				&& !InputPin->IsDatalessPin()
+				&& !!(InputPin->Properties.AllowedTypes & EPCGDataType::Spatial)
+				&& Context->InputData.GetSpatialInputCountByPin(InputPin->Properties.Label) == 0)
 			{
 				UE_LOG(LogPCG, Verbose, TEXT("Intersection resulted in empty output due to no spatial data on a dynamic input pin."));
 				return true;
@@ -109,7 +110,7 @@ bool FPCGOuterIntersectionElement::ExecuteInternal(FPCGContext* Context) const
 		const bool bIsPrimaryPin = InputPin->Properties.Label == PCGIntersectionConstants::PrimaryLabel;
 		
 		// Only checking pins expecting spatial data, not params, settings, etc
-		if (bIsPrimaryPin || !(InputPin->Properties.AllowedTypes & EPCGDataType::Spatial))
+		if (bIsPrimaryPin || InputPin->IsDatalessPin() || !(InputPin->Properties.AllowedTypes & EPCGDataType::Spatial))
 		{
 			continue;
 		}
diff --git a/Engine/Plugins/PCG/Source/PCG/Public/PCGPin.h b/Engine/Plugins/PCG/Source/PCG/Public/PCGPin.h
--- a/Engine/Plugins/PCG/Source/PCG/Public/PCGPin.h
+++ b/Engine/Plugins/PCG/Source/PCG/Public/PCGPin.h
@@ -284,6 +284,8 @@ public:
 	UFUNCTION(BlueprintCallable, Category = Settings)
 	bool IsOutputPin() const;
 
+	bool IsDatalessPin() const { return Properties.Usage == EPCGPinUsage::DependencyOnly; }
+
 	int32 EdgeCount() const;
 
 	/** Returns the current pin types, which can either be the static types from the pin properties, or a dynamic type based on connected edges. */
