[PCG] Re-arrange the PCG Node right click context menu priority

+ Editor User Preferences to enable/disable displaying some node menus on the right click context

diff --git a/Engine/Plugins/PCG/Source/PCGEditor/Private/PCGEditorCommands.cpp b/Engine/Plugins/PCG/Source/PCGEditor/Private/PCGEditorCommands.cpp
--- a/Engine/Plugins/PCG/Source/PCGEditor/Private/PCGEditorCommands.cpp
+++ b/Engine/Plugins/PCG/Source/PCGEditor/Private/PCGEditorCommands.cpp
@@ -24,7 +24,7 @@ void FPCGEditorCommands::RegisterCommands()
 {
 	UI_COMMAND(CollapseNodes, "Collapse into Subgraph", "Collapse selected nodes into a separate PCGGraph asset.", EUserInterfaceActionType::Button, FInputChord(EModifierKey::Alt, EKeys::J));
 	UI_COMMAND(ExportNodes, "Export nodes to PCGSettings", "Exports selected nodes to separate and reusable PCGSettings assets.", EUserInterfaceActionType::Button, FInputChord());
-	UI_COMMAND(ConvertToStandaloneNodes, "Convert to standalone Nodes", "Converts instanced nodes to standalone nodes.", EUserInterfaceActionType::Button, FInputChord());
+	UI_COMMAND(ConvertToStandaloneNodes, "Convert to Standalone Nodes", "Converts instanced nodes to standalone nodes.", EUserInterfaceActionType::Button, FInputChord());
 	UI_COMMAND(Find, "Find", "Finds PCG nodes and comments in the current graph.", EUserInterfaceActionType::Button, FInputChord(EModifierKey::Control, EKeys::F));
 	UI_COMMAND(ShowSelectedDetails, "Show Node Details", "Opens a details panel for the selected nodes.", EUserInterfaceActionType::Button, FInputChord(EKeys::F4));
 	UI_COMMAND(PauseAutoRegeneration, "Pause Regen", "Pause automatic regeneration of the current graph.", EUserInterfaceActionType::ToggleButton, FInputChord(EModifierKey::Alt, EKeys::R));
diff --git a/Engine/Plugins/PCG/Source/PCGEditor/Private/PCGEditorGraphNode.cpp b/Engine/Plugins/PCG/Source/PCGEditor/Private/PCGEditorGraphNode.cpp
--- a/Engine/Plugins/PCG/Source/PCGEditor/Private/PCGEditorGraphNode.cpp
+++ b/Engine/Plugins/PCG/Source/PCGEditor/Private/PCGEditorGraphNode.cpp
@@ -2,12 +2,14 @@
 
 #include "PCGEditorGraphNode.h"
 
+#include "GraphEditorActions.h"
 #include "PCGNode.h"
 #include "PCGPin.h"
 #include "PCGSettings.h"
 #include "PCGSubsystem.h"
 
 #include "PCGEditorModule.h"
+#include "PCGEditorSettings.h"
 
 #include "ToolMenu.h"
 #include "ToolMenuSection.h"
@@ -45,13 +47,17 @@ void UPCGEditorGraphNode::GetNodeContextMenuActions(UToolMenu* Menu, class UGrap
 		return;
 	}
 
-	FToolMenuSection& Section = Menu->AddSection("EdGraphSchemaGeneral", LOCTEXT("GeneralHeader", "General"));
-	Section.AddMenuEntry(FGenericCommands::Get().Delete);
-	Section.AddMenuEntry(FGenericCommands::Get().Cut);
-	Section.AddMenuEntry(FGenericCommands::Get().Copy);
-	Section.AddMenuEntry(FGenericCommands::Get().Duplicate);
-
 	Super::GetNodeContextMenuActions(Menu, Context);
+
+	// General node actions
+	if (GetDefault<UPCGEditorSettings>()->bShowNodeGeneralActionsRightClickContextMenu)
+	{
+		FToolMenuSection& Section = Menu->AddSection("EdGraphSchemaGeneral", LOCTEXT("GeneralHeader", "General"));
+		Section.AddMenuEntry(FGenericCommands::Get().Delete);
+		Section.AddMenuEntry(FGenericCommands::Get().Cut);
+		Section.AddMenuEntry(FGenericCommands::Get().Copy);
+		Section.AddMenuEntry(FGenericCommands::Get().Duplicate);
+	}
 }
 
 void UPCGEditorGraphNode::AllocateDefaultPins()
diff --git a/Engine/Plugins/PCG/Source/PCGEditor/Private/PCGEditorGraphNodeBase.cpp b/Engine/Plugins/PCG/Source/PCGEditor/Private/PCGEditorGraphNodeBase.cpp
--- a/Engine/Plugins/PCG/Source/PCGEditor/Private/PCGEditorGraphNodeBase.cpp
+++ b/Engine/Plugins/PCG/Source/PCGEditor/Private/PCGEditorGraphNodeBase.cpp
@@ -116,6 +116,19 @@ namespace PCGEditorGraphNodeBase
 
 		return true;
 	}
+
+	namespace Constants
+	{
+		FName NodeConversionActionName = TEXT("NodeConversion");
+
+		FText NodeConversionActionLabel = LOCTEXT("NodeConversion", "Convert Node");
+		FText NodeConversionActionTooltip = LOCTEXT("NodeConversionTooltip", "Convert a single node into a different node, or otherwise compatible output.");
+		FTextFormat ConvertToFormat = LOCTEXT("ConvertTo", "Convert to {0}");
+
+		FText ConversionHeaderLabel = LOCTEXT("ConvertHeader", "Convert To");
+		FText OrganizationHeaderLabel = LOCTEXT("OrganizationHeader", "Organization");
+		FText DeterminismHeaderLabel = LOCTEXT("DeterminismHeader", "Determinism");
+	}
 }
 
 void UPCGEditorGraphNodeBase::Construct(UPCGNode* InPCGNode)
@@ -174,11 +187,14 @@ void UPCGEditorGraphNodeBase::PostTransacted(const FTransactionObjectEvent& Tran
 
 void UPCGEditorGraphNodeBase::GetNodeContextMenuActions(UToolMenu* Menu, class UGraphNodeContextMenuContext* Context) const
 {
+	using namespace PCGEditorGraphNodeBase::Constants;
+
 	if (!Context->Node)
 	{
 		return;
 	}
 
+	// Local pin special actions should come first.
 	{
 		FToolMenuSection& Section = Menu->AddSection("EdGraphSchemaPinActions", LOCTEXT("PinActionsMenuHeader", "Pin Actions"));
 		if (Context->Pin && bCanUserAddRemoveSourcePins)
@@ -200,18 +216,9 @@ void UPCGEditorGraphNodeBase::GetNodeContextMenuActions(UToolMenu* Menu, class U
 		}
 	}
 
+	// Node special actions should come after pin actions
 	{
 		FToolMenuSection& Section = Menu->AddSection("EdGraphSchemaNodeActions", LOCTEXT("NodeActionsHeader", "Node Actions"));
-		Section.AddMenuEntry(FPCGEditorCommands::Get().ToggleEnabled, LOCTEXT("ToggleEnabledLabel", "Enable"));
-		Section.AddMenuEntry(FPCGEditorCommands::Get().ToggleDebug, LOCTEXT("ToggleDebugLabel", "Debug"));
-		Section.AddMenuEntry(FPCGEditorCommands::Get().DebugOnlySelected);
-		Section.AddMenuEntry(FPCGEditorCommands::Get().DisableDebugOnAllNodes);
-		Section.AddMenuEntry(FPCGEditorCommands::Get().ToggleInspect, LOCTEXT("ToggleinspectionLabel", "Inspect"));
-		Section.AddMenuEntry(FGraphEditorCommands::Get().BreakNodeLinks);
-		Section.AddMenuEntry(FPCGEditorCommands::Get().CollapseNodes);
-		Section.AddMenuEntry(FPCGEditorCommands::Get().ExportNodes, TAttribute<FText>(), TAttribute<FText>(), FSlateIcon(FPCGEditorStyle::Get().GetStyleSetName(), "ClassIcon.PCGSettings"));
-		Section.AddMenuEntry(FPCGEditorCommands::Get().ConvertToStandaloneNodes);
-		Section.AddMenuEntry(FPCGEditorCommands::Get().RenameNode, LOCTEXT("RenameNode", "Rename"));
 
 		const UPCGSettings* Settings = PCGNode->GetSettings();
 
@@ -231,35 +238,47 @@ void UPCGEditorGraphNodeBase::GetNodeContextMenuActions(UToolMenu* Menu, class U
 			const TArray<FPCGPreconfiguredInfo> ConversionInfo = Settings->GetConversionInfo();
 			if (!ConversionInfo.IsEmpty())
 			{
-				Section.AddSubMenu(
-					TEXT("NodeConversion"),
-					LOCTEXT("NodeConversion", "Convert Node"),
-					LOCTEXT("NodeConversionTooltip", "Convert a single node into a different node, or otherwise compatible output."),
-					FNewToolMenuDelegate::CreateLambda([this, ConversionInfo = std::move(ConversionInfo)](UToolMenu* AlignmentMenu)
+				auto AddEntries = [this, &Section, ConversionInfo = std::move(ConversionInfo)](UToolMenu* AlignmentMenu)
+				{
+					FToolMenuSection& SubSection = (ConversionInfo.Num() > 1) ? AlignmentMenu->AddSection("EdGraphSchemaConversion", ConversionHeaderLabel) : Section;
+
+					for (const FPCGPreconfiguredInfo& Conversion : ConversionInfo)
 					{
-						FToolMenuSection& SubSection = AlignmentMenu->AddSection("EdGraphSchemaConversion", LOCTEXT("ConvertHeader", "Convert To"));
+						FText Label = (ConversionInfo.Num() > 1) ? Conversion.Label : FText::Format(ConvertToFormat, Conversion.Label);
+						FUIAction Action(FExecuteAction::CreateLambda([this, Conversion = std::move(Conversion)] { const_cast<UPCGEditorGraphNodeBase*>(this)->OnConvertNode(Conversion); }));
+						// TODO: UX feedback (read-only or greyed out font, etc) if the conversion can't happen.
+						SubSection.AddMenuEntry(FName(Conversion.Label.ToString()), Label, Conversion.Tooltip, FSlateIcon(), std::move(Action));
+					}
+				};
 
-						for (const FPCGPreconfiguredInfo& Conversion : ConversionInfo)
-						{
-							// TODO: UX feedback (read-only or greyed out font, etc) if the conversion can't happen.
-							SubSection.AddMenuEntry(
-									FName(Conversion.Label.ToString()),
-									Conversion.Label,
-									Conversion.Tooltip,
-									FSlateIcon(),
-									// Capture Conversion by value for delayed use.
-									FUIAction(FExecuteAction::CreateLambda([this, Conversion]
-									{
-										const_cast<UPCGEditorGraphNodeBase*>(this)->OnConvertNode(Conversion);
-									})));
-						}
-					}));
+				if (ConversionInfo.Num() == 1) // Single conversion, just add it to the action list
+				{
+					AddEntries(Menu);
+				}
+				else // Add the entries to a submenu
+				{
+					Section.AddSubMenu(NodeConversionActionName, NodeConversionActionLabel, NodeConversionActionTooltip, FNewToolMenuDelegate::CreateLambda(AddEntries));
+				}
 			}
 		}
+
+		// General PCG node actions
+		Section.AddMenuEntry(FPCGEditorCommands::Get().ToggleEnabled, LOCTEXT("ToggleEnabledLabel", "Enable"));
+		Section.AddMenuEntry(FPCGEditorCommands::Get().RenameNode, LOCTEXT("RenameNode", "Rename"));
+		Section.AddMenuEntry(FPCGEditorCommands::Get().ToggleDebug, LOCTEXT("ToggleDebugLabel", "Debug"));
+		Section.AddMenuEntry(FPCGEditorCommands::Get().DebugOnlySelected);
+		Section.AddMenuEntry(FPCGEditorCommands::Get().DisableDebugOnAllNodes);
+		Section.AddMenuEntry(FPCGEditorCommands::Get().ToggleInspect, LOCTEXT("ToggleinspectionLabel", "Inspect"));
+		Section.AddMenuEntry(FGraphEditorCommands::Get().BreakNodeLinks);
+		Section.AddMenuEntry(FPCGEditorCommands::Get().ExportNodes, TAttribute<FText>(), TAttribute<FText>(), FSlateIcon(FPCGEditorStyle::Get().GetStyleSetName(), "ClassIcon.PCGSettings"));
+		Section.AddMenuEntry(FPCGEditorCommands::Get().CollapseNodes);
+		Section.AddMenuEntry(FPCGEditorCommands::Get().ConvertToStandaloneNodes);
 	}
 
+	// Organizational actions
+	if (GetDefault<UPCGEditorSettings>()->bShowNodeOrganizationalActionsRightClickContextMenu)
 	{
-		FToolMenuSection& Section = Menu->AddSection("EdGraphSchemaOrganization", LOCTEXT("OrganizationHeader", "Organization"));
+		FToolMenuSection& Section = Menu->AddSection("EdGraphSchemaOrganization", OrganizationHeaderLabel);
 		Section.AddMenuEntry(
 			"PCGNode_SetColor",
 			LOCTEXT("PCGNode_SetColor", "Set Node Color"),
@@ -289,19 +308,22 @@ void UPCGEditorGraphNodeBase::GetNodeContextMenuActions(UToolMenu* Menu, class U
 		}));
 	}
 
+	// PCG Determinism actions
+	if (GetDefault<UPCGEditorSettings>()->bShowNodeDeterminismActionsRightClickContext)
+	{
+		FToolMenuSection& Section = Menu->AddSection("EdGraphSchemaDeterminism", DeterminismHeaderLabel);
+		Section.AddMenuEntry(FPCGEditorCommands::Get().RunDeterminismNodeTest,
+			LOCTEXT("Determinism_RunTest", "Validate Determinism on Selection"),
+			LOCTEXT("Determinism_RunTestToolTip", "Run a test to validate the selected nodes for determinism."));
+	}
+
+	// Comment Group is the final section
 	{
 		FToolMenuSection& Section = Menu->AddSection("EdGraphSchemaCommentGroup", LOCTEXT("CommentGroupHeader", "Comment Group"));
 		Section.AddMenuEntry(FGraphEditorCommands::Get().CreateComment,
 			LOCTEXT("MultiCommentDesc", "Create Comment from Selection"),
 			LOCTEXT("CommentToolTip", "Create a resizable comment box around selection."));
 	}
-
-	{
-		FToolMenuSection& Section = Menu->AddSection("EdGraphSchemaDeterminism", LOCTEXT("DeterminismHeader", "Determinism"));
-		Section.AddMenuEntry(FPCGEditorCommands::Get().RunDeterminismNodeTest,
-			LOCTEXT("Determinism_RunTest", "Validate Determinism on Selection"),
-			LOCTEXT("Determinism_RunTestToolTip", "Run a test to validate the selected nodes for determinism."));
-	}
 }
 
 void UPCGEditorGraphNodeBase::AutowireNewNode(UEdGraphPin* FromPin)
diff --git a/Engine/Plugins/PCG/Source/PCGEditor/Public/PCGEditorSettings.h b/Engine/Plugins/PCG/Source/PCGEditor/Public/PCGEditorSettings.h
--- a/Engine/Plugins/PCG/Source/PCGEditor/Public/PCGEditorSettings.h
+++ b/Engine/Plugins/PCG/Source/PCGEditor/Public/PCGEditorSettings.h
@@ -238,4 +238,16 @@ public:
 	/** Controls whether the template picker dialog will be shown when creating a new PCG graph, both in the contextual menu and in the PCG graph pickers. */
 	UPROPERTY(EditAnywhere, config, Category = "Content Browser")
 	bool bShowTemplatePickerOnNewGraph = true;
+
+	/** Display node 'General' actions on the right click context menu. */
+	UPROPERTY(EditAnywhere, config, Category = "Graph UI")
+	bool bShowNodeGeneralActionsRightClickContextMenu = true;
+
+	/** Display node 'Organizational' actions on the right click context menu. */
+	UPROPERTY(EditAnywhere, config, Category = "Graph UI")
+	bool bShowNodeOrganizationalActionsRightClickContextMenu = true;
+
+	/** Display node 'Determinism' actions on the right click context menu. */
+	UPROPERTY(EditAnywhere, config, Category = "Graph UI")
+	bool bShowNodeDeterminismActionsRightClickContext = true;
 };
