[PCG] Optimization of FPCGGraphCache::CleanFromCache/GetGraphCacheEntryCount
- Add IPCGElement to FPCGCacheEntryKey(s) acceleration table

diff --git a/Engine/Plugins/PCG/Source/PCG/Private/Graph/PCGGraphCache.cpp b/Engine/Plugins/PCG/Source/PCG/Private/Graph/PCGGraphCache.cpp
--- a/Engine/Plugins/PCG/Source/PCG/Private/Graph/PCGGraphCache.cpp
+++ b/Engine/Plugins/PCG/Source/PCG/Private/Graph/PCGGraphCache.cpp
@@ -34,6 +34,11 @@ static TAutoConsoleVariable<bool> CVarCacheMemoryBudgetEnabled(
 	true,
 	TEXT("Whether memory budget is enforced (items purged from cache to respect pcg.Cache.MemoryBudgetMB."));
 
+static TAutoConsoleVariable<bool> CVarValidateElementToCacheEntryKeys(
+	TEXT("pcg.Cache.Debug.ValidateElementToCacheEntryKeys"),
+	false,
+	TEXT("Validate ElementToCacheEntryKeys acceleration table (debug)."));
+
 // Initial max number of entries graph cache
 static const int32 GPCGGraphCacheInitialCapacity = 65536;
 
@@ -126,10 +131,8 @@ void FPCGGraphCache::StoreInCache(const FPCGStoreInCacheParams& Params, const FP
 			GrowCache_Unsafe();
 		}
 
-		FPCGCacheEntryKey CacheKey(InElement, InDependenciesCrc);
-		CacheData.Add(CacheKey, InOutput);
-		
-		AddDataToAccountedMemory(InOutput);
+		const FPCGCacheEntryKey CacheKey(InElement, InDependenciesCrc);
+		AddToCacheInternal(CacheKey, InOutput, /*bAddToMemory=*/true);
 	}
 }
 
@@ -137,11 +140,8 @@ void FPCGGraphCache::ClearCache()
 {
 	UE::TScopeLock ScopedLock(CacheLock);
 
-	MemoryRecords.Empty();
-	TotalMemoryUsed = 0;
-
 	// Remove all entries
-	CacheData.Empty(CacheData.Max());
+	ClearCacheInternal(CacheData.Max(), /*bClearMemory=*/true);
 }
 
 bool FPCGGraphCache::EnforceMemoryBudget()
@@ -170,9 +170,9 @@ bool FPCGGraphCache::EnforceMemoryBudget()
 
 		while (TotalMemoryUsed > TargetCacheMemoryUsage && CacheData.Num() > 0)
 		{
-			FPCGDataCollection RemovedData = CacheData.RemoveLeastRecent();
-			RemoveFromMemoryTotal(RemovedData);
+			RemoveFromCacheInternal(CacheData.GetLeastRecentKey());
 		}
+		ValidateElementToCacheEntryKeys();
 	}
 
 	return true;
@@ -194,22 +194,17 @@ void FPCGGraphCache::CleanFromCache(const IPCGElement* InElement, const UPCGSett
 	{
 		UE::TScopeLock ScopedLock(CacheLock);
 
-		TArray<FPCGCacheEntryKey> Keys;
-		CacheData.GetKeys(Keys);
+		ValidateElementToCacheEntryKeys();
 
-		for (const FPCGCacheEntryKey& Key : Keys)
-		{
-			if (Key.GetElement() != InElement)
-			{
-				continue;
-			}
+		TSet<FPCGCacheEntryKey> ElementCacheEntryKeys;
+		ElementToCacheEntryKeys.RemoveAndCopyValue(InElement, ElementCacheEntryKeys);
 
-			if (const FPCGDataCollection* Data = CacheData.Find(Key))
-			{
-				RemoveFromMemoryTotal(*Data);
-				CacheData.Remove(Key);
-			}
+		for (const FPCGCacheEntryKey& Key : ElementCacheEntryKeys)
+		{
+			RemoveFromCacheInternal(Key);
 		}
+
+		ValidateElementToCacheEntryKeys();
 	}
 }
 
@@ -218,20 +213,12 @@ uint32 FPCGGraphCache::GetGraphCacheEntryCount(IPCGElement* InElement) const
 	TRACE_CPUPROFILER_EVENT_SCOPE(FPCGGraphCache::GetFromCache);
 	UE::TScopeLock ScopedLock(CacheLock);
 
-	uint32 Count = 0;
-
-	TArray<FPCGCacheEntryKey> Keys;
-	CacheData.GetKeys(Keys);
-
-	for (const FPCGCacheEntryKey& Key : Keys)
+	if (const TSet<FPCGCacheEntryKey>* ElementCacheEntryKeys = ElementToCacheEntryKeys.Find(InElement))
 	{
-		if (Key.GetElement() == InElement)
-		{
-			Count++;
-		}
+		return ElementCacheEntryKeys->Num();
 	}
 
-	return Count;
+	return 0;
 }
 #endif // WITH_EDITOR
 
@@ -246,6 +233,70 @@ void FPCGGraphCache::AddReferencedObjects(FReferenceCollector& Collector)
 	}
 }
 
+void FPCGGraphCache::ValidateElementToCacheEntryKeys() const
+{
+	if (CVarValidateElementToCacheEntryKeys.GetValueOnAnyThread())
+	{
+		int32 CacheKeyCount = 0;
+		for (const auto& Kvp : ElementToCacheEntryKeys)
+		{
+			CacheKeyCount += Kvp.Value.Num();
+		}
+
+		check(CacheKeyCount == CacheData.Num());
+	}
+}
+
+void FPCGGraphCache::ClearCacheInternal(int32 InMaxEntries, bool bClearMemory)
+{
+	if (bClearMemory)
+	{
+		MemoryRecords.Empty();
+		TotalMemoryUsed = 0;
+	}
+
+	CacheData.Empty(InMaxEntries);
+	ElementToCacheEntryKeys.Empty();
+}
+
+void FPCGGraphCache::AddToCacheInternal(const FPCGCacheEntryKey& InKey, const FPCGDataCollection& InCollection, bool bAddToMemory)
+{
+	// We currently grow the cache before calling add so this shouldn't be needed but if 
+	// the rules change we need to make sure we keep ElementToCacheEntryKeys in sync
+	if (CacheData.Num() == CacheData.Max())
+	{
+		RemoveFromCacheInternal(CacheData.GetLeastRecentKey());
+	}
+
+	CacheData.Add(InKey, InCollection);
+	ElementToCacheEntryKeys.FindOrAdd(InKey.GetElement()).Add(InKey);
+
+	if (bAddToMemory)
+	{
+		AddDataToAccountedMemory(InCollection);
+	}
+
+	ValidateElementToCacheEntryKeys();
+}
+
+void FPCGGraphCache::RemoveFromCacheInternal(const FPCGCacheEntryKey& InKey)
+{
+	if (TSet<FPCGCacheEntryKey>* ElementCacheEntryKeys = ElementToCacheEntryKeys.Find(InKey.GetElement()))
+	{
+		ElementCacheEntryKeys->Remove(InKey);
+		if (ElementCacheEntryKeys->IsEmpty())
+		{
+			ElementToCacheEntryKeys.Remove(InKey.GetElement());
+		}
+	}
+
+	if (const FPCGDataCollection* RemovedData = CacheData.Find(InKey))
+	{
+		RemoveFromMemoryTotal(*RemovedData);
+		CacheData.Remove(InKey);
+	}
+}
+
 void FPCGGraphCache::GrowCache_Unsafe()
 {
 	TRACE_CPUPROFILER_EVENT_SCOPE(FPCGGraphCache::GrowCache_Unsafe);
@@ -259,13 +310,14 @@ void FPCGGraphCache::GrowCache_Unsafe()
 	}
 
 	// Resize and flush
-	CacheData.Empty(CacheData.Num() * 2);
+	ClearCacheInternal(CacheData.Num() * 2, /*bClearMemory=*/false);
+
 	UE_LOG(LogPCG, Log, TEXT("Graph cache doubled in capacity to %d entries."), CacheData.Max());
 
 	// Copy back. Restore temporal order.
 	for (TLruCache<FPCGCacheEntryKey, FPCGDataCollection>::TIterator It(CacheDataCopy); It; ++It)
 	{
-		CacheData.Add(It.Key(), It.Value());
+		AddToCacheInternal(It.Key(), It.Value(), /*bAddToMemory=*/false);
 	}
 }
 
diff --git a/Engine/Plugins/PCG/Source/PCG/Private/Graph/PCGGraphCache.h b/Engine/Plugins/PCG/Source/PCG/Private/Graph/PCGGraphCache.h
--- a/Engine/Plugins/PCG/Source/PCG/Private/Graph/PCGGraphCache.h
+++ b/Engine/Plugins/PCG/Source/PCG/Private/Graph/PCGGraphCache.h
@@ -107,9 +107,15 @@ private:
 	/** For each data element in collection, update records to reflect removal and update total memory figure if all instances removed. */
 	void RemoveFromMemoryTotal(const FPCGDataCollection& InCollection);
 
+	void ClearCacheInternal(int32 InMaxEntries, bool bClearMemory);
+	void AddToCacheInternal(const FPCGCacheEntryKey& InKey, const FPCGDataCollection& InCollection, bool bAddToMemory);
+	void RemoveFromCacheInternal(const FPCGCacheEntryKey& InKey);
+	void ValidateElementToCacheEntryKeys() const;
+
 	// Note: we are not going to serialize this as-is, since the pointers will change
 	// We will have to serialize on a node id basis most likely
 	TLruCache<FPCGCacheEntryKey, FPCGDataCollection> CacheData;
+	TMap<const IPCGElement*, TSet<FPCGCacheEntryKey>> ElementToCacheEntryKeys;
 
 	/** Map from data UIDs to records. Provides ref counting and caches memory size. */
 	TMap<uint64, FCachedMemoryRecord> MemoryRecords;
