[PCG] Fix static analysis error related to moved object

#rnx

diff --git a/Engine/Plugins/PCG/Source/PCGEditor/Private/PCGEditorGraphNodeBase.cpp b/Engine/Plugins/PCG/Source/PCGEditor/Private/PCGEditorGraphNodeBase.cpp
--- a/Engine/Plugins/PCG/Source/PCGEditor/Private/PCGEditorGraphNodeBase.cpp
+++ b/Engine/Plugins/PCG/Source/PCGEditor/Private/PCGEditorGraphNodeBase.cpp
@@ -238,20 +238,23 @@ void UPCGEditorGraphNodeBase::GetNodeContextMenuActions(UToolMenu* Menu, class U
 			const TArray<FPCGPreconfiguredInfo> ConversionInfo = Settings->GetConversionInfo();
 			if (!ConversionInfo.IsEmpty())
 			{
+				const int32 NumConversions = ConversionInfo.Num();
+
 				auto AddEntries = [this, &Section, ConversionInfo = std::move(ConversionInfo)](UToolMenu* AlignmentMenu)
 				{
 					FToolMenuSection& SubSection = (ConversionInfo.Num() > 1) ? AlignmentMenu->AddSection("EdGraphSchemaConversion", ConversionHeaderLabel) : Section;
 
 					for (const FPCGPreconfiguredInfo& Conversion : ConversionInfo)
 					{
-						FText Label = (ConversionInfo.Num() > 1) ? Conversion.Label : FText::Format(ConvertToFormat, Conversion.Label);
+						const FText Label = (ConversionInfo.Num() > 1) ? Conversion.Label : FText::Format(ConvertToFormat, Conversion.Label);
+						const FText Tooltip = Conversion.Tooltip;
 						FUIAction Action(FExecuteAction::CreateLambda([this, Conversion = std::move(Conversion)] { const_cast<UPCGEditorGraphNodeBase*>(this)->OnConvertNode(Conversion); }));
 						// TODO: UX feedback (read-only or greyed out font, etc) if the conversion can't happen.
-						SubSection.AddMenuEntry(FName(Conversion.Label.ToString()), Label, Conversion.Tooltip, FSlateIcon(), std::move(Action));
+						SubSection.AddMenuEntry(FName(Label.ToString()), Label, Tooltip, FSlateIcon(), std::move(Action));
 					}
 				};
 
-				if (ConversionInfo.Num() == 1) // Single conversion, just add it to the action list
+				if (NumConversions == 1) // Single conversion, just add it to the action list
 				{
 					AddEntries(Menu);
 				}
