[PCG] Fix bug where copying and pasting dynamic pin nodes loses the Add Pin button

#rnx

diff --git a/Engine/Plugins/PCG/Source/PCGEditor/Private/PCGEditorGraphNodeBase.cpp b/Engine/Plugins/PCG/Source/PCGEditor/Private/PCGEditorGraphNodeBase.cpp
--- a/Engine/Plugins/PCG/Source/PCGEditor/Private/PCGEditorGraphNodeBase.cpp
+++ b/Engine/Plugins/PCG/Source/PCGEditor/Private/PCGEditorGraphNodeBase.cpp
@@ -148,8 +148,6 @@ void UPCGEditorGraphNodeBase::Construct(UPCGNode* InPCGNode)
 	{
 		const ENodeEnabledState NewEnabledState = !PCGSettingsInterface->bEnabled ? ENodeEnabledState::Disabled : ENodeEnabledState::Enabled;
 		SetEnabledState(NewEnabledState);
-
-		bCanUserAddRemoveSourcePins = (InPCGNode->GetSettings() && InPCGNode->GetSettings()->IsA<UPCGSettingsWithDynamicInputs>());
 	}
 
 	// Update to current graph/inspection state.
@@ -197,7 +195,7 @@ void UPCGEditorGraphNodeBase::GetNodeContextMenuActions(UToolMenu* Menu, class U
 	// Local pin special actions should come first.
 	{
 		FToolMenuSection& Section = Menu->AddSection("EdGraphSchemaPinActions", LOCTEXT("PinActionsMenuHeader", "Pin Actions"));
-		if (Context->Pin && bCanUserAddRemoveSourcePins)
+		if (Context->Pin && CanUserAddRemoveDynamicInputPins())
 		{
 			Section.AddMenuEntry(FPCGEditorCommands::Get().AddSourcePin);
 			Section.AddMenuEntry("RemovePin",
@@ -250,7 +248,7 @@ void UPCGEditorGraphNodeBase::GetNodeContextMenuActions(UToolMenu* Menu, class U
 						const FText Tooltip = Conversion.Tooltip;
 						FUIAction Action(FExecuteAction::CreateLambda([this, Conversion = std::move(Conversion)] { const_cast<UPCGEditorGraphNodeBase*>(this)->OnConvertNode(Conversion); }));
 						// TODO: UX feedback (read-only or greyed out font, etc) if the conversion can't happen.
-						SubSection.AddMenuEntry(FName(Label.ToString()), Label, Tooltip, FSlateIcon(), std::move(Action));
+						SubSection.AddMenuEntry(FName(Label.ToString()), Label, std::move(Tooltip), FSlateIcon(), std::move(Action));
 					}
 				};
 
@@ -695,6 +693,12 @@ bool UPCGEditorGraphNodeBase::GetCompactNodeIcon(FName& OutCompactNodeIcon) cons
 	return Settings && Settings->GetCompactNodeIcon(OutCompactNodeIcon);
 }
 
+bool UPCGEditorGraphNodeBase::CanUserAddRemoveDynamicInputPins() const
+{
+	const UPCGSettings* Settings = PCGNode ? PCGNode->GetSettings() : nullptr;
+	return Settings && Settings->IsA<UPCGSettingsWithDynamicInputs>();
+}
+
 bool UPCGEditorGraphNodeBase::HasFlippedTitleLines() const
 {
 	return PCGNode ? PCGNode->HasFlippedTitleLines() : false;
@@ -1167,7 +1171,7 @@ void UPCGEditorGraphNodeBase::OnUserAddDynamicInputPin()
 	}
 }
 
-bool UPCGEditorGraphNodeBase::CanUserRemoveDynamicInputPin(UEdGraphPin* InPinToRemove)
+bool UPCGEditorGraphNodeBase::CanUserRemoveDynamicInputPin(UEdGraphPin* InPinToRemove) const
 {
 	check(PCGNode && InPinToRemove);
 
diff --git a/Engine/Plugins/PCG/Source/PCGEditor/Private/PCGEditorGraphNodeBase.h b/Engine/Plugins/PCG/Source/PCGEditor/Private/PCGEditorGraphNodeBase.h
--- a/Engine/Plugins/PCG/Source/PCGEditor/Private/PCGEditorGraphNodeBase.h
+++ b/Engine/Plugins/PCG/Source/PCGEditor/Private/PCGEditorGraphNodeBase.h
@@ -51,7 +51,7 @@ public:
 	// ~End UEdGraphNode interface
 
 	void OnUserAddDynamicInputPin();
-	bool CanUserRemoveDynamicInputPin(UEdGraphPin* InPinToRemove);
+	bool CanUserRemoveDynamicInputPin(UEdGraphPin* InPinToRemove) const;
 	void OnUserRemoveDynamicInputPin(UEdGraphPin* InRemovedPin);
 
 	void OnConvertNode(const FPCGPreconfiguredInfo& ConversionInfo);
@@ -86,7 +86,7 @@ public:
 	/** Returns custom compact node icon if available */
 	bool GetCompactNodeIcon(FName& OutCompactNodeIcon) const;
 
-	bool CanUserAddRemoveDynamicInputPins() const { return bCanUserAddRemoveSourcePins; }
+	bool CanUserAddRemoveDynamicInputPins() const;
 
 	/** Whether to flip the order of the title lines - display generated title first and authored second. */
 	bool HasFlippedTitleLines() const;
@@ -157,7 +157,6 @@ protected:
 	bool bDeferredReconstruct = false;
 	bool bDisableReconstructFromNode = false;
 	bool bIsInspected = false;
-	bool bCanUserAddRemoveSourcePins = false;
 
 	/** Whether this node was culled in the last execution. */
 	bool bIsCulledFromExecution = false;
