[PCG] Add new string operations

+ToUpper
+ToLower
+TrimStart
+TrimEnd
+TrimStartAndEnd

diff --git a/Engine/Plugins/PCG/Source/PCG/Private/Elements/Metadata/PCGMetadataStringOpElement.cpp b/Engine/Plugins/PCG/Source/PCG/Private/Elements/Metadata/PCGMetadataStringOpElement.cpp
--- a/Engine/Plugins/PCG/Source/PCG/Private/Elements/Metadata/PCGMetadataStringOpElement.cpp
+++ b/Engine/Plugins/PCG/Source/PCG/Private/Elements/Metadata/PCGMetadataStringOpElement.cpp
@@ -16,6 +16,15 @@ FName UPCGMetadataStringOpSettings::GetInputPinLabel(uint32 Index) const
 {
 	switch (Operation)
 	{
+	case EPCGMetadataStringOperation::ToLower: // fall-through
+	case EPCGMetadataStringOperation::ToUpper: // fall-through
+	case EPCGMetadataStringOperation::TrimStart: // fall-through
+	case EPCGMetadataStringOperation::TrimEnd: // fall-through
+	case EPCGMetadataStringOperation::TrimStartAndEnd:
+	{
+		return PCGPinConstants::DefaultInputLabel;
+	}
+
 	case EPCGMetadataStringOperation::Append:
 		switch (Index)
 		{
@@ -43,12 +52,18 @@ uint32 UPCGMetadataStringOpSettings::GetOperandNum() const
 {
 	switch (Operation)
 	{
-	case EPCGMetadataStringOperation::Append:
-		return 2;
-	case EPCGMetadataStringOperation::Replace:
-		return 3;
-	default:
-		return 0;
+		case EPCGMetadataStringOperation::ToLower: // fall-through
+		case EPCGMetadataStringOperation::ToUpper: // fall-through
+		case EPCGMetadataStringOperation::TrimStart: // fall-through
+		case EPCGMetadataStringOperation::TrimEnd: // fall-through
+		case EPCGMetadataStringOperation::TrimStartAndEnd:
+			return 1;
+		case EPCGMetadataStringOperation::Append:
+			return 2;
+		case EPCGMetadataStringOperation::Replace:
+			return 3;
+		default:
+			return 0;
 	}
 }
 
@@ -125,7 +140,27 @@ bool FPCGMetadataStringOpElement::DoOperation(PCGMetadataOps::FOperationData& Op
 
 	const UPCGMetadataStringOpSettings* Settings = CastChecked<UPCGMetadataStringOpSettings>(OperationData.Settings);
 
-	if (Settings->Operation == EPCGMetadataStringOperation::Append)
+	if (Settings->Operation == EPCGMetadataStringOperation::ToUpper)
+	{
+		return DoUnaryOp<FString>(OperationData, [](const FString& Value) -> FString { return Value.ToUpper(); });
+	}
+	else if (Settings->Operation == EPCGMetadataStringOperation::ToLower)
+	{
+		return DoUnaryOp<FString>(OperationData, [](const FString& Value) -> FString { return Value.ToLower(); });
+	}
+	else if (Settings->Operation == EPCGMetadataStringOperation::TrimStart)
+	{
+		return DoUnaryOp<FString>(OperationData, [](const FString& Value) -> FString { return Value.TrimStart(); });
+	}
+	else if (Settings->Operation == EPCGMetadataStringOperation::TrimEnd)
+	{
+		return DoUnaryOp<FString>(OperationData, [](const FString& Value) -> FString { return Value.TrimEnd(); });
+	}
+	else if (Settings->Operation == EPCGMetadataStringOperation::TrimStartAndEnd)
+	{
+		return DoUnaryOp<FString>(OperationData, [](const FString& Value) -> FString { return Value.TrimStartAndEnd(); });
+	}
+	else if (Settings->Operation == EPCGMetadataStringOperation::Append)
 	{
 		return DoBinaryOp<FString, FString>(OperationData, [](const FString& Value1, const FString& Value2) -> FString { return Value1 + Value2; });
 	}
diff --git a/Engine/Plugins/PCG/Source/PCG/Public/Elements/Metadata/PCGMetadataStringOpElement.h b/Engine/Plugins/PCG/Source/PCG/Public/Elements/Metadata/PCGMetadataStringOpElement.h
--- a/Engine/Plugins/PCG/Source/PCG/Public/Elements/Metadata/PCGMetadataStringOpElement.h
+++ b/Engine/Plugins/PCG/Source/PCG/Public/Elements/Metadata/PCGMetadataStringOpElement.h
@@ -10,7 +10,12 @@ UENUM()
 enum class EPCGMetadataStringOperation : uint16
 {
 	Append UMETA(DisplayName="Append String"),
-	Replace UMETA(DisplayName="Replace String")
+	Replace UMETA(DisplayName="Replace String"),
+	ToUpper UMETA(DisplayName="ToUpper (String)"),
+	ToLower UMETA(DisplayName="ToLower (String)"),
+	TrimStart UMETA(DisplayName="Trim Start (String)"),
+	TrimEnd UMETA(DisplayName="Trim End (String)"),
+	TrimStartAndEnd UMETA(DisplayName="Trim Start and End (String)")
 };
 
 UCLASS(MinimalAPI, BlueprintType, ClassGroup = (Procedural))
