[PCG] Small adjustments to string operations added

#rnx

diff --git a/Engine/Plugins/PCG/Source/PCG/Public/Elements/Metadata/PCGMetadataStringOpElement.h b/Engine/Plugins/PCG/Source/PCG/Public/Elements/Metadata/PCGMetadataStringOpElement.h
--- a/Engine/Plugins/PCG/Source/PCG/Public/Elements/Metadata/PCGMetadataStringOpElement.h
+++ b/Engine/Plugins/PCG/Source/PCG/Public/Elements/Metadata/PCGMetadataStringOpElement.h
@@ -11,11 +11,11 @@ enum class EPCGMetadataStringOperation : uint16
 {
 	Append UMETA(DisplayName="Append String"),
 	Replace UMETA(DisplayName="Replace String"),
-	ToUpper UMETA(DisplayName="ToUpper (String)"),
-	ToLower UMETA(DisplayName="ToLower (String)"),
-	TrimStart UMETA(DisplayName="Trim Start (String)"),
-	TrimEnd UMETA(DisplayName="Trim End (String)"),
-	TrimStartAndEnd UMETA(DisplayName="Trim Start and End (String)")
+	ToUpper UMETA(DisplayName="To Upper", Tooltip="Convert all characters to upper case."),
+	ToLower UMETA(DisplayName="To Lower", Tooltip="Convert all characters to lower case."),
+	TrimStart UMETA(DisplayName="Trim Start", Tooltip="Trim whitespace from the beginning of the string."),
+	TrimEnd UMETA(DisplayName="Trim End", Tooltip="Trim whitespace from the end of the string."),
+	TrimStartAndEnd UMETA(DisplayName="Trim Start and End", Tooltip="Trim whitespace from the beginning and end of the string.")
 };
 
 UCLASS(MinimalAPI, BlueprintType, ClassGroup = (Procedural))
