[PCG] Fix autowiring edges to the execution dependency pin of dropped nodes with no primary pins

#rnx

diff --git a/Engine/Plugins/PCG/Source/PCG/Private/Elements/PCGMergeAttributes.cpp b/Engine/Plugins/PCG/Source/PCG/Private/Elements/PCGMergeAttributes.cpp
--- a/Engine/Plugins/PCG/Source/PCG/Private/Elements/PCGMergeAttributes.cpp
+++ b/Engine/Plugins/PCG/Source/PCG/Private/Elements/PCGMergeAttributes.cpp
@@ -72,7 +72,7 @@ bool FPCGMergeAttributesElement::ExecuteInternal(FPCGContext* Context) const
 	for (const UPCGPin* InputPin : Context->Node->GetInputPins())
 	{
 		check(InputPin);
-		if (!InputPin->IsDatalessPin())
+		if (!InputPin->Properties.IsDatalessPin())
 		{
 			Sources.Append(Context->InputData.GetInputsByPin(InputPin->Properties.Label));
 		}
diff --git a/Engine/Plugins/PCG/Source/PCG/Private/Elements/PCGMergeElement.cpp b/Engine/Plugins/PCG/Source/PCG/Private/Elements/PCGMergeElement.cpp
--- a/Engine/Plugins/PCG/Source/PCG/Private/Elements/PCGMergeElement.cpp
+++ b/Engine/Plugins/PCG/Source/PCG/Private/Elements/PCGMergeElement.cpp
@@ -61,7 +61,7 @@ bool FPCGMergeElement::ExecuteInternal(FPCGContext* Context) const
 	{
 		check(InputPin);
 
-		if (!InputPin->IsDatalessPin())
+		if (!InputPin->Properties.IsDatalessPin())
 		{
 			Sources.Append(Context->InputData.GetInputsByPin(InputPin->Properties.Label));
 		}
diff --git a/Engine/Plugins/PCG/Source/PCG/Private/Elements/PCGOuterIntersectionElement.cpp b/Engine/Plugins/PCG/Source/PCG/Private/Elements/PCGOuterIntersectionElement.cpp
--- a/Engine/Plugins/PCG/Source/PCG/Private/Elements/PCGOuterIntersectionElement.cpp
+++ b/Engine/Plugins/PCG/Source/PCG/Private/Elements/PCGOuterIntersectionElement.cpp
@@ -93,7 +93,7 @@ bool FPCGOuterIntersectionElement::ExecuteInternal(FPCGContext* Context) const
 			check(InputPin);
 			// Only checking secondary source pins expecting spatial data
 			if (InputPin->Properties.Label != PCGIntersectionConstants::PrimaryLabel
-				&& !InputPin->IsDatalessPin()
+				&& !InputPin->Properties.IsDatalessPin()
 				&& !!(InputPin->Properties.AllowedTypes & EPCGDataType::Spatial)
 				&& Context->InputData.GetSpatialInputCountByPin(InputPin->Properties.Label) == 0)
 			{
@@ -110,7 +110,7 @@ bool FPCGOuterIntersectionElement::ExecuteInternal(FPCGContext* Context) const
 		const bool bIsPrimaryPin = InputPin->Properties.Label == PCGIntersectionConstants::PrimaryLabel;
 		
 		// Only checking pins expecting spatial data, not params, settings, etc
-		if (bIsPrimaryPin || InputPin->IsDatalessPin() || !(InputPin->Properties.AllowedTypes & EPCGDataType::Spatial))
+		if (bIsPrimaryPin || InputPin->Properties.IsDatalessPin() || !(InputPin->Properties.AllowedTypes & EPCGDataType::Spatial))
 		{
 			continue;
 		}
diff --git a/Engine/Plugins/PCG/Source/PCG/Public/PCGPin.h b/Engine/Plugins/PCG/Source/PCG/Public/PCGPin.h
--- a/Engine/Plugins/PCG/Source/PCG/Public/PCGPin.h
+++ b/Engine/Plugins/PCG/Source/PCG/Public/PCGPin.h
@@ -88,6 +88,8 @@ struct PCG_API FPCGPinProperties
 	bool IsNormalPin() const { return PinStatus == EPCGPinStatus::Normal; }
 	void SetNormalPin() { PinStatus = EPCGPinStatus::Normal; }
 
+	bool IsDatalessPin() const { return Usage == EPCGPinUsage::DependencyOnly; }
+
 	// Convert the bIsAdvanced boolean to PinStatus for deprecation purposes.
 	void PostSerialize(const FArchive& Ar);
 
@@ -284,8 +286,6 @@ public:
 	UFUNCTION(BlueprintCallable, Category = Settings)
 	bool IsOutputPin() const;
 
-	bool IsDatalessPin() const { return Properties.Usage == EPCGPinUsage::DependencyOnly; }
-
 	int32 EdgeCount() const;
 
 	/** Returns the current pin types, which can either be the static types from the pin properties, or a dynamic type based on connected edges. */
diff --git a/Engine/Plugins/PCG/Source/PCGEditor/Private/PCGEditorGraphNodeBase.cpp b/Engine/Plugins/PCG/Source/PCGEditor/Private/PCGEditorGraphNodeBase.cpp
--- a/Engine/Plugins/PCG/Source/PCGEditor/Private/PCGEditorGraphNodeBase.cpp
+++ b/Engine/Plugins/PCG/Source/PCGEditor/Private/PCGEditorGraphNodeBase.cpp
@@ -342,6 +342,12 @@ void UPCGEditorGraphNodeBase::AutowireNewNode(UEdGraphPin* FromPin)
 	{
 		check(OtherPin);
 
+		// TODO: Allow autoconnecting output dependency pins to input dependency pins.
+		if (OtherPin->Properties.IsAdvancedPin() || OtherPin->Properties.IsDatalessPin())
+		{
+			continue;
+		}
+
 		const FName& OtherPinName = OtherPin->Properties.Label;
 		UEdGraphPin* ToPin = FindPinChecked(OtherPinName, bFromPinIsInput ? EEdGraphPinDirection::EGPD_Output : EEdGraphPinDirection::EGPD_Input);
 		if (ToPin && GetSchema()->TryCreateConnection(FromPin, ToPin))
