[PCG] Add Dynamic Pins to Gather Node

* Fix source iteration on Merge nodes
* Assign Single Connection Multi Data to all dynamic pins on all dynamic pin nodes

diff --git a/Engine/Plugins/PCG/Source/PCG/Private/Elements/PCGGather.cpp b/Engine/Plugins/PCG/Source/PCG/Private/Elements/PCGGather.cpp
--- a/Engine/Plugins/PCG/Source/PCG/Private/Elements/PCGGather.cpp
+++ b/Engine/Plugins/PCG/Source/PCG/Private/Elements/PCGGather.cpp
@@ -18,15 +18,6 @@ PRAGMA_ENABLE_DEPRECATION_WARNINGS
 }
 #endif // WITH_EDITOR
 
-TArray<FPCGPinProperties> UPCGGatherSettings::InputPinProperties() const
-{
-	TArray<FPCGPinProperties> PinProperties;
-	FPCGPinProperties& InputPinProperty = PinProperties.Emplace_GetRef(PCGPinConstants::DefaultInputLabel, EPCGDataType::Any);
-	InputPinProperty.SetRequiredPin();
-
-	return PinProperties;
-}
-
 TArray<FPCGPinProperties> UPCGGatherSettings::OutputPinProperties() const
 {
 	TArray<FPCGPinProperties> PinProperties;
@@ -40,6 +31,22 @@ FPCGElementPtr UPCGGatherSettings::CreateElement() const
 	return MakeShared<FPCGGatherElement>();
 }
 
+FName UPCGGatherSettings::GetDynamicInputPinsBaseLabel() const
+{
+	return PCGPinConstants::DefaultInputLabel;
+}
+
+#if WITH_EDITOR
+void UPCGGatherSettings::AddDefaultDynamicInputPin()
+{
+	FPCGPinProperties SecondaryPinProperties(
+		FName(GetDynamicInputPinsBaseLabel().ToString() + FString::FromInt(DynamicInputPinProperties.Num() + 2)),
+		EPCGDataType::Any,
+		/*bInAllowMultipleConnections=*/false);
+	AddDynamicInputPin(std::move(SecondaryPinProperties));
+}
+#endif // WITH_EDITOR
+
 namespace PCGGather
 {
 	FPCGDataCollection GatherDataForPin(const FPCGDataCollection& InputData, const FName InputLabel, const FName OutputLabel)
@@ -61,7 +68,7 @@ namespace PCGGather
 			Output.TaggedData = MoveTemp(GatheredData);
 		}
 
-		for(FPCGTaggedData& TaggedData : Output.TaggedData)
+		for (FPCGTaggedData& TaggedData : Output.TaggedData)
 		{
 			TaggedData.Pin = OutputLabel;
 		}
@@ -74,7 +81,16 @@ bool FPCGGatherElement::ExecuteInternal(FPCGContext* Context) const
 {
 	TRACE_CPUPROFILER_EVENT_SCOPE(FPCGGatherElement::Execute);
 
-	Context->OutputData = PCGGather::GatherDataForPin(Context->InputData);
+	FPCGDataCollection GatheredData;
+	const UPCGGatherSettings* Settings = Context->GetInputSettings<UPCGGatherSettings>();
+	check(Settings);
+
+	for (FName PinLabel : Settings->GetNodeDefinedPinLabels())
+	{
+		GatheredData += PCGGather::GatherDataForPin(Context->InputData, PinLabel);
+	}
+
+	Context->OutputData = GatheredData;
 
 	return true;
 }
diff --git a/Engine/Plugins/PCG/Source/PCG/Private/Elements/PCGMergeAttributes.cpp b/Engine/Plugins/PCG/Source/PCG/Private/Elements/PCGMergeAttributes.cpp
--- a/Engine/Plugins/PCG/Source/PCG/Private/Elements/PCGMergeAttributes.cpp
+++ b/Engine/Plugins/PCG/Source/PCG/Private/Elements/PCGMergeAttributes.cpp
@@ -53,7 +53,10 @@ TArray<FPCGPinProperties> UPCGMergeAttributesSettings::StaticInputPinProperties(
 #if WITH_EDITOR
 void UPCGMergeAttributesSettings::AddDefaultDynamicInputPin()
 {
-	FPCGPinProperties SecondaryPinProperties(FName(GetDynamicInputPinsBaseLabel().ToString() + FString::FromInt(DynamicInputPinProperties.Num() + 2)), EPCGDataType::Param);
+	FPCGPinProperties SecondaryPinProperties(
+		FName(GetDynamicInputPinsBaseLabel().ToString() + FString::FromInt(DynamicInputPinProperties.Num() + 2)),
+		EPCGDataType::Param,
+		/*bInAllowMultipleConnections=*/false);
 	AddDynamicInputPin(std::move(SecondaryPinProperties));
 }
 #endif // WITH_EDITOR
@@ -68,14 +71,13 @@ bool FPCGMergeAttributesElement::ExecuteInternal(FPCGContext* Context) const
 	FPCGTaggedData* MergedOutput = nullptr;
 	UPCGParamData* MergedAttributeSet = nullptr;
 
+	const UPCGMergeAttributesSettings* Settings = Context->GetInputSettings<UPCGMergeAttributesSettings>();
+	check(Settings);
+
 	TArray<FPCGTaggedData> Sources;
-	for (const UPCGPin* InputPin : Context->Node->GetInputPins())
+	for (FName PinLabel : Settings->GetNodeDefinedPinLabels())
 	{
-		check(InputPin);
-		if (!InputPin->Properties.IsDatalessPin())
-		{
-			Sources.Append(Context->InputData.GetInputsByPin(InputPin->Properties.Label));
-		}
+		Sources.Append(Context->InputData.GetInputsByPin(PinLabel));
 	}
 
 	for (const FPCGTaggedData& Source : Sources)
diff --git a/Engine/Plugins/PCG/Source/PCG/Private/Elements/PCGMergeElement.cpp b/Engine/Plugins/PCG/Source/PCG/Private/Elements/PCGMergeElement.cpp
--- a/Engine/Plugins/PCG/Source/PCG/Private/Elements/PCGMergeElement.cpp
+++ b/Engine/Plugins/PCG/Source/PCG/Private/Elements/PCGMergeElement.cpp
@@ -43,7 +43,10 @@ TArray<FPCGPinProperties> UPCGMergeSettings::StaticInputPinProperties() const
 #if WITH_EDITOR
 void UPCGMergeSettings::AddDefaultDynamicInputPin()
 {
-	FPCGPinProperties SecondaryPinProperties(FName(GetDynamicInputPinsBaseLabel().ToString() + FString::FromInt(DynamicInputPinProperties.Num() + 2)), EPCGDataType::Point);
+	FPCGPinProperties SecondaryPinProperties(
+		FName(GetDynamicInputPinsBaseLabel().ToString() + FString::FromInt(DynamicInputPinProperties.Num() + 2)),
+		EPCGDataType::Point,
+		/*bInAllowMultipleConnections=*/false);
 	AddDynamicInputPin(std::move(SecondaryPinProperties));
 }
 #endif // WITH_EDITOR
@@ -57,14 +60,9 @@ bool FPCGMergeElement::ExecuteInternal(FPCGContext* Context) const
 	const bool bMergeMetadata = !Settings || Settings->bMergeMetadata;
 
 	TArray<FPCGTaggedData> Sources;
-	for (const UPCGPin* InputPin : Context->Node->GetInputPins())
+	for (FName PinLabel : Settings->GetNodeDefinedPinLabels())
 	{
-		check(InputPin);
-
-		if (!InputPin->Properties.IsDatalessPin())
-		{
-			Sources.Append(Context->InputData.GetInputsByPin(InputPin->Properties.Label));
-		}
+		Sources.Append(Context->InputData.GetInputsByPin(PinLabel));
 	}
 
 	TArray<FPCGTaggedData>& Outputs = Context->OutputData.TaggedData;
diff --git a/Engine/Plugins/PCG/Source/PCG/Private/PCGData.cpp b/Engine/Plugins/PCG/Source/PCG/Private/PCGData.cpp
--- a/Engine/Plugins/PCG/Source/PCG/Private/PCGData.cpp
+++ b/Engine/Plugins/PCG/Source/PCG/Private/PCGData.cpp
@@ -403,6 +403,12 @@ bool FPCGDataCollection::operator!=(const FPCGDataCollection& Other) const
 	return !operator==(Other);
 }
 
+FPCGDataCollection& FPCGDataCollection::operator+=(const FPCGDataCollection& Other)
+{
+	AddData(Other.TaggedData, Other.DataCrcs);
+	return *this;
+}
+
 void FPCGDataCollection::AddReferences(FReferenceCollector& Collector)
 {
 	for (FPCGTaggedData& Data : TaggedData)
diff --git a/Engine/Plugins/PCG/Source/PCG/Private/PCGSettingsWithDynamicInputs.cpp b/Engine/Plugins/PCG/Source/PCG/Private/PCGSettingsWithDynamicInputs.cpp
--- a/Engine/Plugins/PCG/Source/PCG/Private/PCGSettingsWithDynamicInputs.cpp
+++ b/Engine/Plugins/PCG/Source/PCG/Private/PCGSettingsWithDynamicInputs.cpp
@@ -58,6 +58,16 @@ int32 UPCGSettingsWithDynamicInputs::GetStaticInputPinNum() const
 	return StaticInputPinProperties().Num();
 }
 
+TArray<FName> UPCGSettingsWithDynamicInputs::GetNodeDefinedPinLabels() const
+{
+	const TArray<FPCGPinProperties> PinProperties = InputPinProperties();
+
+	TArray<FName> Labels;
+	Labels.Reserve(PinProperties.Num());
+	Algo::Transform(PinProperties, Labels, [](const FPCGPinProperties& Properties) { return Properties.Label; });
+
+	return Labels;
+}
 
 void UPCGSettingsWithDynamicInputs::AddDynamicInputPin(FPCGPinProperties&& CustomProperties)
 {
diff --git a/Engine/Plugins/PCG/Source/PCG/Public/Elements/PCGGather.h b/Engine/Plugins/PCG/Source/PCG/Public/Elements/PCGGather.h
--- a/Engine/Plugins/PCG/Source/PCG/Public/Elements/PCGGather.h
+++ b/Engine/Plugins/PCG/Source/PCG/Public/Elements/PCGGather.h
@@ -1,7 +1,7 @@
 // Copyright Epic Games, Inc. All Rights Reserved.
 #pragma once
 
-#include "PCGSettings.h"
+#include "PCGSettingsWithDynamicInputs.h"
 
 #include "PCGGather.generated.h"
 
@@ -9,7 +9,7 @@
  * Used to wrangle multiple input wires into one output wire for organizational purposes.
  */
 UCLASS(BlueprintType, ClassGroup = (Procedural))
-class UPCGGatherSettings : public UPCGSettings
+class UPCGGatherSettings : public UPCGSettingsWithDynamicInputs
 {
 	GENERATED_BODY()
 
@@ -26,10 +26,18 @@ public:
 	virtual bool HasDynamicPins() const override { return true; }
 
 protected:
-	virtual TArray<FPCGPinProperties> InputPinProperties() const override;
 	virtual TArray<FPCGPinProperties> OutputPinProperties() const override;
 	virtual FPCGElementPtr CreateElement() const override;
 	//~End UPCGSettings interface
+
+	//~Begin PCGSettingsWithDynamicInputs interface
+	FName GetDynamicInputPinsBaseLabel() const override;
+
+public:
+#if WITH_EDITOR
+	void AddDefaultDynamicInputPin() override;
+#endif // WITH_EDITOR
+	//~End PCGSettingsWithDynamicInputs interface
 };
 
 class FPCGGatherElement : public IPCGElement
diff --git a/Engine/Plugins/PCG/Source/PCG/Public/PCGData.h b/Engine/Plugins/PCG/Source/PCG/Public/PCGData.h
--- a/Engine/Plugins/PCG/Source/PCG/Public/PCGData.h
+++ b/Engine/Plugins/PCG/Source/PCG/Public/PCGData.h
@@ -216,6 +216,7 @@ struct PCG_API FPCGDataCollection
 
 	bool operator==(const FPCGDataCollection& Other) const;
 	bool operator!=(const FPCGDataCollection& Other) const;
+	FPCGDataCollection& operator+=(const FPCGDataCollection& Other);
 	void AddReferences(FReferenceCollector& Collector);
 
 	/** Computes CRCs for all data items. */
diff --git a/Engine/Plugins/PCG/Source/PCG/Public/PCGSettingsWithDynamicInputs.h b/Engine/Plugins/PCG/Source/PCG/Public/PCGSettingsWithDynamicInputs.h
--- a/Engine/Plugins/PCG/Source/PCG/Public/PCGSettingsWithDynamicInputs.h
+++ b/Engine/Plugins/PCG/Source/PCG/Public/PCGSettingsWithDynamicInputs.h
@@ -10,7 +10,7 @@
 * UPCGSettings subclass with functionality to dynamically add/remove input pins
 */
 UCLASS(Abstract, BlueprintType, ClassGroup = (Procedural))
-class PCG_API UPCGSettingsWithDynamicInputs: public UPCGSettings
+class PCG_API UPCGSettingsWithDynamicInputs : public UPCGSettings
 {
 	GENERATED_BODY()
 
@@ -36,6 +36,9 @@ public:
 	int32 GetDynamicInputPinNum() const { return DynamicInputPinProperties.Num(); }
 #endif // WITH_EDITOR
 
+	/** Get an array of pin labels to the pins defined by the settings. I.e. Not overrides, advanced pins, etc. */
+	TArray<FName> GetNodeDefinedPinLabels() const;
+
 protected:
 #if WITH_EDITOR
 	/** Add a new dynamic source pin with the specified properties */
