[PCG] Add dynamic inputs to the Union node for more control over ordering

diff --git a/Engine/Plugins/PCG/Source/PCG/Private/Elements/PCGUnionElement.cpp b/Engine/Plugins/PCG/Source/PCG/Private/Elements/PCGUnionElement.cpp
--- a/Engine/Plugins/PCG/Source/PCG/Private/Elements/PCGUnionElement.cpp
+++ b/Engine/Plugins/PCG/Source/PCG/Private/Elements/PCGUnionElement.cpp
@@ -7,14 +7,7 @@
 
 #include UE_INLINE_GENERATED_CPP_BY_NAME(PCGUnionElement)
 
-TArray<FPCGPinProperties> UPCGUnionSettings::InputPinProperties() const
-{
-	TArray<FPCGPinProperties> PinProperties;
-	FPCGPinProperties& InputPinProperty = PinProperties.Emplace_GetRef(PCGPinConstants::DefaultInputLabel, EPCGDataType::Spatial);
-	InputPinProperty.SetRequiredPin();
-
-	return PinProperties;
-}
+#define LOCTEXT_NAMESPACE "PCGUnionSettings"
 
 TArray<FPCGPinProperties> UPCGUnionSettings::OutputPinProperties() const
 {
@@ -29,6 +22,30 @@ FPCGElementPtr UPCGUnionSettings::CreateElement() const
 	return MakeShared<FPCGUnionElement>();
 }
 
+FName UPCGUnionSettings::GetDynamicInputPinsBaseLabel() const
+{
+	return PCGPinConstants::DefaultInputLabel;
+}
+
+TArray<FPCGPinProperties> UPCGUnionSettings::StaticInputPinProperties() const
+{
+	TArray<FPCGPinProperties> Properties;
+	Properties.Emplace(PCGPinConstants::DefaultInputLabel, EPCGDataType::Spatial);
+	return Properties;
+}
+
+#if WITH_EDITOR
+void UPCGUnionSettings::AddDefaultDynamicInputPin()
+{
+	FPCGPinProperties SecondaryPinProperties(
+		FName(GetDynamicInputPinsBaseLabel().ToString() + FString::FromInt(DynamicInputPinProperties.Num() + 2)),
+		EPCGDataType::Spatial,
+		/*bInAllowMultipleConnections=*/false);
+	SecondaryPinProperties.Tooltip = LOCTEXT("DynamicPinPropertyTooltip", "Dynamic pins, such as this one, will be unioned together in order.");
+	AddDynamicInputPin(std::move(SecondaryPinProperties));
+}
+#endif // WITH_EDITOR
+
 bool FPCGUnionElement::ExecuteInternal(FPCGContext* Context) const
 {
 	TRACE_CPUPROFILER_EVENT_SCOPE(FPCGUnionElement::Execute);
@@ -36,8 +53,6 @@ bool FPCGUnionElement::ExecuteInternal(FPCGContext* Context) const
 	const UPCGUnionSettings* Settings = Context->GetInputSettings<UPCGUnionSettings>();
 	check(Settings);
 
-	TArray<FPCGTaggedData> Inputs = Context->InputData.GetInputs();
-
 	const EPCGUnionType Type = Settings->Type;
 	const EPCGUnionDensityFunction DensityFunction = Settings->DensityFunction;
 
@@ -47,14 +62,20 @@ bool FPCGUnionElement::ExecuteInternal(FPCGContext* Context) const
 	UPCGUnionData* UnionData = nullptr;
 	int32 UnionTaggedDataIndex = -1;
 
-	for (FPCGTaggedData& Input : Inputs)
+	TArray<FPCGTaggedData, TInlineAllocator<8>> Sources;
+	for (FName PinLabel : Settings->GetNodeDefinedPinLabels())
 	{
-		const UPCGSpatialData* SpatialData = Cast<UPCGSpatialData>(Input.Data);
+		Sources.Append(Context->InputData.GetInputsByPin(PinLabel));
+	}
+
+	for (const FPCGTaggedData& Source : Sources)
+	{
+		const UPCGSpatialData* SpatialData = Cast<UPCGSpatialData>(Source.Data);
 
 		// Non-spatial data we're not going to touch
 		if (!SpatialData)
 		{
-			Outputs.Add(Input);
+			Outputs.Add(Source);
 			continue;
 		}
 
@@ -62,7 +83,7 @@ bool FPCGUnionElement::ExecuteInternal(FPCGContext* Context) const
 		{
 			FirstSpatialData = SpatialData;
 			UnionTaggedDataIndex = Outputs.Num();
-			Outputs.Add(Input);
+			Outputs.Add(Source);
 			continue;
 		}
 
@@ -80,7 +101,7 @@ bool FPCGUnionElement::ExecuteInternal(FPCGContext* Context) const
 		else
 		{
 			UnionData->AddData(SpatialData);
-			UnionTaggedData.Tags.Append(Input.Tags);
+			UnionTaggedData.Tags.Append(Source.Tags);
 		}
 		
 		UnionTaggedData.Data = UnionData;
@@ -88,3 +109,5 @@ bool FPCGUnionElement::ExecuteInternal(FPCGContext* Context) const
 
 	return true;
 }
+
+#undef LOCTEXT_NAMESPACE
\ No newline at end of file
diff --git a/Engine/Plugins/PCG/Source/PCG/Public/Elements/PCGUnionElement.h b/Engine/Plugins/PCG/Source/PCG/Public/Elements/PCGUnionElement.h
--- a/Engine/Plugins/PCG/Source/PCG/Public/Elements/PCGUnionElement.h
+++ b/Engine/Plugins/PCG/Source/PCG/Public/Elements/PCGUnionElement.h
@@ -2,13 +2,13 @@
 
 #pragma once
 
-#include "PCGSettings.h"
+#include "PCGSettingsWithDynamicInputs.h"
 #include "Data/PCGUnionData.h"
 
 #include "PCGUnionElement.generated.h"
 
 UCLASS(MinimalAPI, BlueprintType, ClassGroup = (Procedural))
-class UPCGUnionSettings : public UPCGSettings
+class UPCGUnionSettings : public UPCGSettingsWithDynamicInputs
 {
 	GENERATED_BODY()
 
@@ -17,17 +17,25 @@ public:
 #if WITH_EDITOR
 	virtual FName GetDefaultNodeName() const override { return FName(TEXT("Union")); }
 	virtual FText GetDefaultNodeTitle() const override { return NSLOCTEXT("PCGUnionSettings", "NodeTitle", "Union"); }
+	FText GetNodeTooltipText() const override { return NSLOCTEXT("PCGUnionSettings", "NodeTooltip", "Combine spatial data into a union of all inputs. Order of inputs is respected, beginning with the dynamic pin inputs."); }
 	virtual EPCGSettingsType GetType() const override { return EPCGSettingsType::Spatial; }
 #endif
-	
 
 protected:
-	virtual TArray<FPCGPinProperties> InputPinProperties() const override;
 	virtual TArray<FPCGPinProperties> OutputPinProperties() const override;
 	virtual FPCGElementPtr CreateElement() const override;
 	//~End UPCGSettings interface
 
+	//~Begin UPCGSettingsWithDynamicInputs interface
+	FName GetDynamicInputPinsBaseLabel() const override;
+	TArray<FPCGPinProperties> StaticInputPinProperties() const override;
+
 public:
+#if WITH_EDITOR
+	void AddDefaultDynamicInputPin() override;
+#endif // WITH_EDITOR
+	//~End UPCGSettingsWithDynamicInputs interface
+
 	UPROPERTY(BlueprintReadWrite, EditAnywhere, Category = Settings, meta = (PCG_Overridable))
 	EPCGUnionType Type = EPCGUnionType::LeftToRightPriority;
 
