[PCG] Prevent PCGWorldActor from being Editor only

#rnx

diff --git a/Engine/Plugins/PCG/Source/PCG/Private/PCGWorldActor.cpp b/Engine/Plugins/PCG/Source/PCG/Private/PCGWorldActor.cpp
--- a/Engine/Plugins/PCG/Source/PCG/Private/PCGWorldActor.cpp
+++ b/Engine/Plugins/PCG/Source/PCG/Private/PCGWorldActor.cpp
@@ -51,6 +51,37 @@ APCGWorldActor::APCGWorldActor(const FObjectInitializer& ObjectInitializer)
 }
 
 #if WITH_EDITOR
+bool APCGWorldActor::IsEditorOnly() const
+{
+	// We still have runtime dependencies on this actor in PCGActorAndComponentMapping::RegisterPartitionActor
+	// @todo_pcg: eventually remove this dependency
+	return false;
+}
+
+void APCGWorldActor::PostLoad()
+{
+	Super::PostLoad();
+
+	// Even if value is not used by APCGWorldActor::IsEditorOnly, this value can show as read-only in the actor's details panel
+	// to avoid confusion we set it to false here.
+	bIsEditorOnlyActor = false;
+}
+
+bool APCGWorldActor::CanEditChange(const FProperty* InProperty) const
+{
+	if (!Super::CanEditChange(InProperty))
+	{
+		return false;
+	}
+
+	if (InProperty && InProperty->GetFName() == GET_MEMBER_NAME_CHECKED(AActor, bIsEditorOnlyActor))
+	{
+		return false;
+	}
+
+	return true;
+}
+
 void APCGWorldActor::BeginCacheForCookedPlatformData(const ITargetPlatform* TargetPlatform)
 {
 	Super::BeginCacheForCookedPlatformData(TargetPlatform);
diff --git a/Engine/Plugins/PCG/Source/PCG/Public/PCGWorldActor.h b/Engine/Plugins/PCG/Source/PCG/Public/PCGWorldActor.h
--- a/Engine/Plugins/PCG/Source/PCG/Public/PCGWorldActor.h
+++ b/Engine/Plugins/PCG/Source/PCG/Public/PCGWorldActor.h
@@ -33,6 +33,9 @@ public:
 	virtual bool ShouldExport() override { return false; }
 	virtual bool ShouldImport(FStringView ActorPropString, bool IsMovingLevel) override { return false; }
 	virtual void BeginCacheForCookedPlatformData(const ITargetPlatform* TargetPlatform) override;
+	virtual bool IsEditorOnly() const override;
+	virtual bool CanEditChange(const FProperty* InProperty) const override;
+	virtual void PostLoad() override;
 	//~End AActor Interface
 #endif
 
