[PCG] Improvements on checking if data is used multiple times

* Enabled by default behind a CVar, add a trace profiler if we ever see it is an issue we can disable it. It more foolproof to always have it one
* Mark data crossing the grid size boundaries and in loops as used mutliple times, as they will be used by multiple PCG Components/Loop executions.
* Input Forwarding element is forced to not verify if a data is used multiple times, as it is just outputing a stored data collection.

#rnx

diff --git a/Engine/Plugins/PCG/Source/PCG/Private/Elements/PCGHiGenGridSize.cpp b/Engine/Plugins/PCG/Source/PCG/Private/Elements/PCGHiGenGridSize.cpp
--- a/Engine/Plugins/PCG/Source/PCG/Private/Elements/PCGHiGenGridSize.cpp
+++ b/Engine/Plugins/PCG/Source/PCG/Private/Elements/PCGHiGenGridSize.cpp
@@ -152,6 +152,8 @@ bool FPCGHiGenGridSizeElement::ExecuteInternal(FPCGContext* Context) const
 	for (FPCGTaggedData& Data : Context->OutputData.TaggedData)
 	{
 		Data.Pin = PCGPinConstants::DefaultOutputLabel;
+		// Since data will downstream be used by potentially multiple partitions, we have to make sure that the data is flagged as used multiple times.
+		Data.bIsUsedMultipleTimes = true;
 	}
 
 	UPCGVolumeData* VolumeData = FPCGContext::NewObject_AnyThread<UPCGVolumeData>(Context);
diff --git a/Engine/Plugins/PCG/Source/PCG/Private/Elements/PCGLoopElement.cpp b/Engine/Plugins/PCG/Source/PCG/Private/Elements/PCGLoopElement.cpp
--- a/Engine/Plugins/PCG/Source/PCG/Private/Elements/PCGLoopElement.cpp
+++ b/Engine/Plugins/PCG/Source/PCG/Private/Elements/PCGLoopElement.cpp
@@ -229,10 +229,12 @@ bool FPCGLoopElement::ExecuteInternal(FPCGContext* InContext) const
 			return true;
 		}
 
+		const int32 NumberOfIterations = LoopDataCollection[0].TaggedData.Num();
+
 		// Early out if the loop collection do not contain the same amount of entries
 		for (int LoopDataIndex = 1; LoopDataIndex < LoopDataCollection.Num(); ++LoopDataIndex)
 		{
-			if (LoopDataCollection[LoopDataIndex].TaggedData.Num() != LoopDataCollection[0].TaggedData.Num())
+			if (LoopDataCollection[LoopDataIndex].TaggedData.Num() != NumberOfIterations)
 			{
 				PCGE_LOG(Error, GraphAndLog, LOCTEXT("MismatchLoopCollections", "Data to loop on does not have the same number of entries for all pins!"));
 				return true;
@@ -248,12 +250,21 @@ bool FPCGLoopElement::ExecuteInternal(FPCGContext* InContext) const
 		FPCGDataCollection FixedDataCollection;
 		PrepareSubgraphData(Settings, Context, FixedInputDataCollection, FixedDataCollection);
 
+		// If we execute more than once, it's the same as if we were using the data multiple times for fixed input.
+		if (NumberOfIterations > 1)
+		{
+			for (FPCGTaggedData& TaggedData : FixedDataCollection.TaggedData)
+			{
+				TaggedData.bIsUsedMultipleTimes = true;
+			}
+		}
+
 		FPCGTaskId PreviousTaskId = InvalidPCGTaskId;
 		const bool bIsFeedbackLoop = !(FeedbackPinNames.IsEmpty());
 
 		// Dispatch the subgraph for each loop entries we have.
 		// Implementation note: even if the execution gets cancelled, these tasks would get cancelled because they are associated to the current source component
-		for (int EntryIndex = 0; EntryIndex < LoopDataCollection[0].TaggedData.Num(); ++EntryIndex)
+		for (int EntryIndex = 0; EntryIndex < NumberOfIterations; ++EntryIndex)
 		{
 			// The final input for a given loop iteration is going to be the fixed data + the entry at the current loop index in all the gathered loop collections
 			FPCGDataCollection InputDataCollection = FixedDataCollection;
diff --git a/Engine/Plugins/PCG/Source/PCG/Private/PCGElement.cpp b/Engine/Plugins/PCG/Source/PCG/Private/PCGElement.cpp
--- a/Engine/Plugins/PCG/Source/PCG/Private/PCGElement.cpp
+++ b/Engine/Plugins/PCG/Source/PCG/Private/PCGElement.cpp
@@ -30,6 +30,11 @@ static TAutoConsoleVariable<bool> CVarPCGAllowPerDataCaching(
 	false,
 	TEXT("Controls whether we test & split down inputs to check caching per input on primary loop nodes."));
 
+static TAutoConsoleVariable<bool> CVarPCGShouldVerifyIfOutputsAreUsedMultipleTimes(
+	TEXT("pcg.ShouldVerifyIfOutputsAreUsedMultipleTimes"),
+	true,
+	TEXT("Add small computation at the end of each node to detect if the data is used multiple times. Necessary for data stealing."));
+
 #if WITH_EDITOR
 #define PCG_ELEMENT_EXECUTION_BREAKPOINT() \
 	if (Context && Context->GetInputSettingsInterface() && Context->GetInputSettingsInterface()->bBreakDebugger) \
@@ -484,10 +489,25 @@ void IPCGElement::PostExecute(FPCGContext* Context) const
 	// Analyze if the output data is used multiple times, if the element requires it.
 	if (ShouldVerifyIfOutputsAreUsedMultipleTimes(Settings))
 	{
+		TRACE_CPUPROFILER_EVENT_SCOPE(IPCGElement::PostExecute::ShouldVerifyIfOutputsAreUsedMultipleTimes);
+		
+		TSet<const UPCGData*> InputData;
+		TMap<const UPCGData*, int> OutputTaggedDataMap;
+		OutputTaggedDataMap.Reserve(Context->OutputData.TaggedData.Num());
 		for (FPCGTaggedData& OutputData : Context->OutputData.TaggedData)
 		{
-			// Enforce that pinless data is always used multiple times, or if the debug mode is enabled.
-			if (OutputData.bPinlessData || (Settings && Settings->CanBeDebugged() && Settings->bDebug))
+			OutputTaggedDataMap.FindOrAdd(OutputData.Data, 0) += 1;
+		}
+
+		for (const FPCGTaggedData& InputTaggedData: Context->InputData.TaggedData)
+		{
+			InputData.Add(InputTaggedData.Data);
+		}
+		
+		for (FPCGTaggedData& OutputData : Context->OutputData.TaggedData)
+		{
+			// Enforce that pinless data is always used multiple times, or if the debug mode is enabled, or was detected multiple times in the output.
+			if (OutputTaggedDataMap[OutputData.Data] > 1 || OutputData.bPinlessData || (Settings && Settings->CanBeDebugged() && Settings->bDebug))
 			{
 				OutputData.bIsUsedMultipleTimes = true;
 				continue;
@@ -496,25 +516,13 @@ void IPCGElement::PostExecute(FPCGContext* Context) const
 			// For data that are marked to be used multiple times, they are potentially not used multiple times if they are not passthrough
 			// (hence if they are not in the input). So set them back to false in that case. It will be set to true again by the executor
 			// if it is actually used in multiple places.
-			auto PassthroughPredicate = [&OutputData](const FPCGTaggedData& InputData) { return InputData.Data == OutputData.Data; };
-			if (OutputData.bIsUsedMultipleTimes && !Context->InputData.TaggedData.ContainsByPredicate(PassthroughPredicate))
+			if (OutputData.bIsUsedMultipleTimes && !InputData.Contains(OutputData.Data))
 			{
 #if !UE_BUILD_SHIPPING
 				OutputData.OriginatingNode = Context->Node;
 #endif //!UE_BUILD_SHIPPING
 				OutputData.bIsUsedMultipleTimes = false;
 			}
-			
-			// We also need to verify that the data is not used in other outputs.
-			auto SameDataDifferentTaggedData = [&OutputData](const FPCGTaggedData& OtherOutputData)
-			{
-				return (&OtherOutputData != &OutputData) && (OtherOutputData.Data == OutputData.Data);
-			};
-			
-			if (!OutputData.bIsUsedMultipleTimes && Context->OutputData.TaggedData.ContainsByPredicate(SameDataDifferentTaggedData))
-			{
-				OutputData.bIsUsedMultipleTimes = true;
-			}
 		}
 	}
 	
@@ -958,6 +966,11 @@ bool IPCGElement::IsCacheableInstance(const UPCGSettingsInterface* InSettingsInt
 	}
 }
 
+bool IPCGElement::ShouldVerifyIfOutputsAreUsedMultipleTimes(const UPCGSettings* InSettings) const
+{
+	return CVarPCGShouldVerifyIfOutputsAreUsedMultipleTimes.GetValueOnAnyThread();
+}
+
 void IPCGElement::GetDependenciesCrc(const FPCGDataCollection& InInput, const UPCGSettings* InSettings, UPCGComponent* InComponent, FPCGCrc& OutCrc) const
 {
 	TRACE_CPUPROFILER_EVENT_SCOPE_TEXT(*FString::Printf(TEXT("IPCGElement::GetDependenciesCrc (%s)"), InSettings ? *InSettings->GetFName().ToString() : TEXT("")));
diff --git a/Engine/Plugins/PCG/Source/PCG/Public/Elements/PCGGather.h b/Engine/Plugins/PCG/Source/PCG/Public/Elements/PCGGather.h
--- a/Engine/Plugins/PCG/Source/PCG/Public/Elements/PCGGather.h
+++ b/Engine/Plugins/PCG/Source/PCG/Public/Elements/PCGGather.h
@@ -44,6 +44,7 @@ class FPCGGatherElement : public IPCGElement
 {
 protected:
 	virtual bool ExecuteInternal(FPCGContext* Context) const override;
+	virtual bool ShouldVerifyIfOutputsAreUsedMultipleTimes(const UPCGSettings* InSettings) const override { return true; }
 	virtual bool SupportsGPUResidentData(FPCGContext* InContext) const override { return true; }
 };
 
diff --git a/Engine/Plugins/PCG/Source/PCG/Public/PCGElement.h b/Engine/Plugins/PCG/Source/PCG/Public/PCGElement.h
--- a/Engine/Plugins/PCG/Source/PCG/Public/PCGElement.h
+++ b/Engine/Plugins/PCG/Source/PCG/Public/PCGElement.h
@@ -111,7 +111,7 @@ public:
 	virtual bool IsCacheable(const UPCGSettings* InSettings) const { return true; }
 	
 	/** Returns true if the node outputs requires to update the output tagged data to detect if the data is used multiple times. */
-	virtual bool ShouldVerifyIfOutputsAreUsedMultipleTimes(const UPCGSettings* InSettings) const { return false; }
+	virtual bool ShouldVerifyIfOutputsAreUsedMultipleTimes(const UPCGSettings* InSettings) const;
 
 	/** Whether to do a 'deep' fine-grained CRC of the output data to pass to downstream nodes. Can be expensive so should be used sparingly. */
 	virtual bool ShouldComputeFullOutputDataCrc(FPCGContext* Context) const { return false; }
diff --git a/Engine/Plugins/PCG/Source/PCG/Public/PCGSubgraph.h b/Engine/Plugins/PCG/Source/PCG/Public/PCGSubgraph.h
--- a/Engine/Plugins/PCG/Source/PCG/Public/PCGSubgraph.h
+++ b/Engine/Plugins/PCG/Source/PCG/Public/PCGSubgraph.h
@@ -204,6 +204,9 @@ public:
 	// Since this class is stateful because it owns a FPCGDataCollection it can't be cached unless it implements a proper GetDependenciesCrc.
 	// For now since it only forwards its input, we just disable the caching.
 	virtual bool IsCacheable(const UPCGSettings* InSettings) const override { return false; }
+
+	// Should not verify if outputs are used multiples times, as we are just outputting the input collection we have.
+	virtual bool ShouldVerifyIfOutputsAreUsedMultipleTimes(const UPCGSettings* InSettings) const override { return false; }
 protected:
 	virtual bool ExecuteInternal(FPCGContext* Context) const override;
 	virtual bool IsPassthrough(const UPCGSettings* InSettings) const override { return true; }
