[PCG] Support grass map generation in world partition levels

#rnx

diff --git a/Engine/Plugins/PCG/Source/PCG/Private/Data/PCGTextureData.cpp b/Engine/Plugins/PCG/Source/PCG/Private/Data/PCGTextureData.cpp
--- a/Engine/Plugins/PCG/Source/PCG/Private/Data/PCGTextureData.cpp
+++ b/Engine/Plugins/PCG/Source/PCG/Private/Data/PCGTextureData.cpp
@@ -9,6 +9,7 @@
 #include "Helpers/PCGHelpers.h"
 
 #include "RHIStaticStates.h"
+#include "RenderCaptureInterface.h"
 #include "TextureResource.h"
 #include "Engine/Texture2D.h"
 #include "Engine/Texture2DArray.h"
@@ -17,6 +18,12 @@
 
 namespace PCGTextureSamplingHelpers
 {
+	static int32 GTriggerReadbackCaptureDispatches = 0;
+	static FAutoConsoleVariableRef CVarTriggerReadbackCapture(
+		TEXT("pcg.Texture.TriggerReadbackCaptures"),
+		GTriggerReadbackCaptureDispatches,
+		TEXT("Trigger GPU readback captures for this many of the subsequent texture data initializations."));
+
 	TOptional<bool> IsTextureCPUAccessible(UTexture2D* Texture)
 	{
 		if (!Texture)
@@ -748,6 +755,9 @@ bool UPCGTextureData::InitializeFromGPUTexture()
 		TWeakObjectPtr<UPCGTextureData> ThisWeakPtr(this);
 		const FString ResourceName = *TextureRHI->GetName().ToString();
 
+		RenderCaptureInterface::FScopedCapture RenderCapture(PCGTextureSamplingHelpers::GTriggerReadbackCaptureDispatches > 0, TEXT("PCGTextureReadbackCapture"));
+		PCGTextureSamplingHelpers::GTriggerReadbackCaptureDispatches = FMath::Max(PCGTextureSamplingHelpers::GTriggerReadbackCaptureDispatches - 1, 0);
+
 		FPCGTextureReadbackInterface::Dispatch(Params, [ThisWeakPtr, TextureSize, ResourceName](void* OutBuffer, int32 ReadbackWidth, int32 ReadbackHeight)
 		{
 			TRACE_CPUPROFILER_EVENT_SCOPE(UPCGTextureData::Initialize::DispatchCallback);
@@ -758,26 +768,37 @@ bool UPCGTextureData::InitializeFromGPUTexture()
 				return;
 			}
 
-			const int32 PixelCount = ReadbackWidth * ReadbackHeight;
+			// Texture readbacks can require memory alignment, e.g. a 127x127 texture can be readback as 128x128.
+			// So when initializing the CPU data, we should ignore the additional pixels.
+			const int32 ReadbackPaddingWidth = ReadbackWidth - TextureSize.X;
+			const int32 ReadbackPaddingHeight = ReadbackHeight - TextureSize.Y;
 
-			if (TextureSize.X * TextureSize.Y != PixelCount)
+			if (ReadbackPaddingWidth < 0 || ReadbackPaddingHeight < 0)
 			{
 				UE_LOG(LogPCG, Error,
-					TEXT("PCGTextureData readback has different dimensions than the source texture '%s'. Expected (%d, %d), received (%d, %d)."),
+					TEXT("PCGTextureData readback has smaller dimensions than the source texture '%s'. Expected greater than or equal to (%d, %d), received (%d, %d)."),
 					*ResourceName,
 					TextureSize.X, TextureSize.Y,
 					ReadbackWidth, ReadbackHeight);
+				return;
 			}
 
 			if (const FColor* FormattedImageData = reinterpret_cast<const FColor*>(OutBuffer))
 			{
-				This->Width = ReadbackWidth;
-				This->Height = ReadbackHeight;
-				This->ColorData.SetNum(PixelCount);
+				This->Width = TextureSize.X;
+				This->Height = TextureSize.Y;
+				This->ColorData.SetNum(TextureSize.X * TextureSize.Y);
 
-				for (int32 D = 0; D < PixelCount; ++D)
+				int32 ActualTexelIndex = 0;
+
+				for (int32 TexelX = 0; TexelX < TextureSize.X; ++TexelX)
 				{
-					This->ColorData[D] = FormattedImageData[D].ReinterpretAsLinear();
+					for (int32 TexelY = 0; TexelY < TextureSize.Y; ++TexelY)
+					{
+						const int32 ReadbackTexelIndex = TexelY * ReadbackWidth + TexelX;
+
+						This->ColorData[ActualTexelIndex++] = FormattedImageData[ReadbackTexelIndex].ReinterpretAsLinear();
+					}
 				}
 			}
 			else
diff --git a/Engine/Plugins/PCG/Source/PCG/Private/Elements/PCGGenerateGrassMaps.cpp b/Engine/Plugins/PCG/Source/PCG/Private/Elements/PCGGenerateGrassMaps.cpp
--- a/Engine/Plugins/PCG/Source/PCG/Private/Elements/PCGGenerateGrassMaps.cpp
+++ b/Engine/Plugins/PCG/Source/PCG/Private/Elements/PCGGenerateGrassMaps.cpp
@@ -144,19 +144,28 @@ bool FPCGGenerateGrassMapsElement::ExecuteInternal(FPCGContext* InContext) const
 
 		for (TSoftObjectPtr<ALandscapeProxy> LandscapeProxyPtr : LandscapeData->Landscapes)
 		{
-			if (LandscapeProxyPtr.Get())
+			ALandscapeProxy* LandscapeProxy = LandscapeProxyPtr.Get();
+
+			if (!LandscapeProxy)
 			{
-				Context->LandscapeProxy = LandscapeProxyPtr.Get();
-				break;
+				continue;
 			}
-		}
 
-		if (ALandscape* Landscape = Context->LandscapeProxy.Get() ? Context->LandscapeProxy->GetLandscapeActor() : nullptr)
-		{
-			TArray<ULandscapeComponent*> AllLandscapeComponents;
-			Landscape->GetComponents<ULandscapeComponent>(AllLandscapeComponents);
+			if (!Context->LandscapeProxy.Get())
+			{
+				Context->LandscapeProxy = LandscapeProxy;
+			}
+			else
+			{
+				// TODO: In order to support multiple-material we need to do a dispatch of one landscape grass exporter per LS proxy.
+				if (Context->LandscapeProxy->LandscapeMaterial != LandscapeProxy->LandscapeMaterial)
+				{
+					PCGLog::LogErrorOnGraph(LOCTEXT("LandscapeMaterialMismatch", "Encountered multiple landscape materials on the same landscape. Only one landscape material per landscape is supported at this time."), Context);
+					return true;
+				}
+			}
 
-			for (ULandscapeComponent* LandscapeComponent : AllLandscapeComponents)
+			for (ULandscapeComponent* LandscapeComponent : LandscapeProxy->LandscapeComponents)
 			{
 				if (!ensure(LandscapeComponent))
 				{
@@ -348,8 +357,6 @@ bool FPCGGenerateGrassMapsElement::ExecuteInternal(FPCGContext* InContext) const
 			const uint32 NumTilesX = FMath::RoundToInt(GrassMapExtent.X / LandscapeComponentExtent);
 			const uint32 NumTilesY = FMath::RoundToInt(GrassMapExtent.Y / LandscapeComponentExtent);
 
-			check(NumTilesX * NumTilesY == Context->LandscapeComponents.Num());
-
 			// Each corner of a quad in the landscape corresponds to one texel in the grass map.
 			const uint32 LandscapeComponentResolution = Context->LandscapeProxy->ComponentSizeQuads + 1;
 			const FIntPoint GrassMapResolution(LandscapeComponentResolution * NumTilesX, LandscapeComponentResolution * NumTilesY);
