[PCG] Add support for accessing data by 'data label' in CustomHLSL nodes
* Data labels are an alternative to using data indices when writing HLSL
* For example, InTexture_SampleWorldPos('MyGrassTexture', WorldPos);
* Add DataLabelResolver DI for remapping a data label/data id to a data index
* GenerateGrassMaps and GetTextureData nodes automatically tag outgoing textures with a data label

diff --git a/Engine/Plugins/PCG/Source/PCG/Internal/Compute/PCGComputeCommon.h b/Engine/Plugins/PCG/Source/PCG/Internal/Compute/PCGComputeCommon.h
--- a/Engine/Plugins/PCG/Source/PCG/Internal/Compute/PCGComputeCommon.h
+++ b/Engine/Plugins/PCG/Source/PCG/Internal/Compute/PCGComputeCommon.h
@@ -82,6 +82,8 @@ namespace PCGComputeConstants
 	// Used to represent invalid/removed points. We use a value slightly less than max float,
 	// as not all platforms support float infinity in shaders.
 	constexpr float INVALID_DENSITY = 3.402823e+38f;
+
+	constexpr TCHAR DataLabelTagPrefix[] = { TEXT("PCG_DATA_LABEL") };
 }
 
 namespace PCGComputeHelpers
@@ -116,4 +118,10 @@ namespace PCGComputeHelpers
 
 	int32 GetAttributeIdFromMetadataAttributeIndex(int32 InAttributeIndex);
 	int32 GetMetadataAttributeIndexFromAttributeId(int32 InAttributeId);
+
+	/** Produces the data label prefixed with PCGComputeConstants::DataLabelTagPrefix. */
+	FString GetPrefixedDataLabel(const FString& InLabel);
+
+	/** Produces the data interface name of a data label resolver. */
+	FString GetDataLabelResolverName(FName InPinLabel);
 }
diff --git a/Engine/Plugins/PCG/Source/PCG/Private/Compute/DataInterfaces/PCGDataLabelResolverDataInterface.cpp b/Engine/Plugins/PCG/Source/PCG/Private/Compute/DataInterfaces/PCGDataLabelResolverDataInterface.cpp
new file mode 100644
--- /dev/null
+++ b/Engine/Plugins/PCG/Source/PCG/Private/Compute/DataInterfaces/PCGDataLabelResolverDataInterface.cpp
@@ -0,0 +1,158 @@
+// Copyright Epic Games, Inc. All Rights Reserved.
+
+#include "Compute/DataInterfaces/PCGDataLabelResolverDataInterface.h"
+
+#include "PCGModule.h"
+#include "PCGSettings.h"
+#include "Compute/PCGComputeGraph.h"
+
+#include "RenderGraphBuilder.h"
+#include "RenderGraphResources.h"
+#include "ShaderCompilerCore.h"
+#include "ShaderParameterMetadataBuilder.h"
+#include "ComputeFramework/ComputeMetadataBuilder.h"
+#include "ComputeFramework/ShaderParamTypeDefinition.h"
+
+#include UE_INLINE_GENERATED_CPP_BY_NAME(PCGDataLabelResolverDataInterface)
+
+void UPCGDataLabelResolverDataInterface::GetSupportedInputs(TArray<FShaderFunctionDefinition>& OutFunctions) const
+{
+	OutFunctions.AddDefaulted_GetRef()
+		.SetName(TEXT("GetDataIndexFromIdInternal"))
+		.AddReturnType(EShaderFundamentalType::Uint) // Index of data in input data collection
+		.AddParam(EShaderFundamentalType::Uint); // InDataId
+}
+
+void UPCGDataLabelResolverDataInterface::GetSupportedOutputs(TArray<FShaderFunctionDefinition>& OutFunctions) const
+{
+	OutFunctions.AddDefaulted_GetRef()
+		.SetName(TEXT("GetDataIndexFromIdInternal"))
+		.AddReturnType(EShaderFundamentalType::Uint) // Index of data in input data collection
+		.AddParam(EShaderFundamentalType::Uint); // InDataId
+}
+
+BEGIN_SHADER_PARAMETER_STRUCT(FPCGDataLabelResolverDataInterfaceParameters, )
+	SHADER_PARAMETER(uint32, InNumData)
+	SHADER_PARAMETER_RDG_BUFFER_SRV(StructuredBuffer<uint32>, InDataIdToDataIndexMap)
+END_SHADER_PARAMETER_STRUCT()
+
+void UPCGDataLabelResolverDataInterface::GetShaderParameters(TCHAR const* UID, FShaderParametersMetadataBuilder& InOutBuilder, FShaderParametersMetadataAllocations& InOutAllocations) const
+{
+	InOutBuilder.AddNestedStruct<FPCGDataLabelResolverDataInterfaceParameters>(UID);
+}
+
+void UPCGDataLabelResolverDataInterface::GetHLSL(FString& OutHLSL, FString const& InDataInterfaceName) const
+{
+	const TMap<FString, FStringFormatArg> TemplateArgs =
+	{
+		{ TEXT("DataInterfaceName"), InDataInterfaceName },
+	};
+
+	OutHLSL += FString::Format(TEXT(
+		"uint {DataInterfaceName}_InNumData;\n"
+		"StructuredBuffer<int> {DataInterfaceName}_InDataIdToDataIndexMap;\n"
+		"\n"
+		"uint GetDataIndexFromIdInternal_{DataInterfaceName}(int InDataId)\n"
+		"{\n"
+		"	if (InDataId < 0 || InDataId >= {DataInterfaceName}_InNumData)\n"
+		"	{\n"
+		"		return 0;\n"
+		"	}\n"
+		"\n"
+		"	// Data index could be -1 if the label was not resolved, so fall-back to data index 0 if necessary.\n"
+		"	return max((uint){DataInterfaceName}_InDataIdToDataIndexMap[InDataId], 0);\n"
+		"}\n"
+		"\n"),
+		TemplateArgs);
+}
+
+UComputeDataProvider* UPCGDataLabelResolverDataInterface::CreateDataProvider(TObjectPtr<UObject> InBinding, uint64 InInputMask, uint64 InOutputMask) const
+{
+	TRACE_CPUPROFILER_EVENT_SCOPE(UPCGDataLabelResolverDataInterface::CreateDataProvider);
+	check(ConsumerSettings);
+
+	UPCGDataBinding* Binding = CastChecked<UPCGDataBinding>(InBinding);
+	const UPCGComputeGraph* ComputeGraph = Binding ? Binding->Graph.Get() : nullptr;
+	check(ComputeGraph);
+
+	UPCGDataLabelResolverDataProvider* Provider = NewObject<UPCGDataLabelResolverDataProvider>();
+
+	const FPCGPinDataLabels* FoundPinDataLabels = ComputeGraph->GetDataLabelsStaticTable().Find(ConsumerSettings);
+	if (const FPCGDataLabels* FoundLabels = FoundPinDataLabels ? FoundPinDataLabels->PinToDataLabels.Find(InputPinLabel) : nullptr)
+	{
+		TArray<int32> DataIdToDataIndexMap;
+		DataIdToDataIndexMap.SetNumUninitialized(FoundLabels->Labels.Num());
+
+		for (int32 LabelIndex = 0; LabelIndex < FoundLabels->Labels.Num(); ++LabelIndex)
+		{
+			const FString& Label = FoundLabels->Labels[LabelIndex];
+			const FString PrefixedLabel = PCGComputeHelpers::GetPrefixedDataLabel(Label);
+
+			// Initialize with invalid indices. The data IDs are not ordered.
+			DataIdToDataIndexMap[LabelIndex] = INDEX_NONE;
+
+			// TODO: In order to support data labels on GPU <-> GPU edges, we have to get the tags from data descriptions
+			// instead of from the CPU tagged data.
+			for (int32 InputDataIndex = 0; InputDataIndex < Binding->InputDataCollection.TaggedData.Num(); ++InputDataIndex)
+			{
+				if (Binding->InputDataCollection.TaggedData[InputDataIndex].Tags.Contains(PrefixedLabel))
+				{
+					DataIdToDataIndexMap[LabelIndex] = InputDataIndex;
+					break;
+				}
+			}
+
+			if (DataIdToDataIndexMap[LabelIndex] == INDEX_NONE)
+			{
+				UE_LOG(LogPCG, Warning, TEXT("DataLabelResolver found data label '%s', but did not match any incoming data."), *Label);
+			}
+		}
+
+		Provider->DataIdToDataIndexMap = MoveTemp(DataIdToDataIndexMap);
+	}
+
+	return Provider;
+}
+
+FComputeDataProviderRenderProxy* UPCGDataLabelResolverDataProvider::GetRenderProxy()
+{
+	return new FPCGDataLabelResolverDataProviderProxy(DataIdToDataIndexMap);
+}
+
+bool FPCGDataLabelResolverDataProviderProxy::IsValid(FValidationData const& InValidationData) const
+{
+	return InValidationData.ParameterStructSize == sizeof(FParameters);
+}
+
+void FPCGDataLabelResolverDataProviderProxy::GatherDispatchData(FDispatchData const& InDispatchData)
+{
+	const TStridedView<FParameters> ParameterArray = MakeStridedParameterView<FParameters>(InDispatchData);
+	for (int32 InvocationIndex = 0; InvocationIndex < InDispatchData.NumInvocations; ++InvocationIndex)
+	{
+		FParameters& Parameters = ParameterArray[InvocationIndex];
+		Parameters.InNumData = DataIdToDataIndexMap.Num();
+		Parameters.InDataIdToDataIndexMap = DataIdToDataIndexBufferSRV;
+	}
+}
+
+void FPCGDataLabelResolverDataProviderProxy::AllocateResources(FRDGBuilder& GraphBuilder, FAllocationData const& InAllocationData)
+{
+	const uint32 SizeBytes = FMath::Max(DataIdToDataIndexMap.Num(), 1) * DataIdToDataIndexMap.GetTypeSize();
+	const FRDGBufferDesc Desc = FRDGBufferDesc::CreateByteAddressDesc(SizeBytes);
+
+	if (DataIdToDataIndexMap.IsEmpty())
+	{
+		FRDGBufferRef NullBuffer = GraphBuilder.CreateBuffer(Desc, TEXT("PCGDataLabelResolver_DataIdToDataIndexMap_Empty"));
+		GraphBuilder.QueueBufferUpload(NullBuffer, 0, sizeof(int32));
+
+		DataIdToDataIndexBufferSRV = GraphBuilder.CreateSRV(NullBuffer);
+	}
+	else
+	{
+		FRDGBufferRef Buffer = GraphBuilder.CreateBuffer(Desc, TEXT("PCGDataLabelResolver_DataIdToDataIndexMap"));
+		GraphBuilder.QueueBufferUpload(Buffer, DataIdToDataIndexMap.GetData(), SizeBytes);
+
+		DataIdToDataIndexBufferSRV = GraphBuilder.CreateSRV(Buffer);
+	}
+}
+
diff --git a/Engine/Plugins/PCG/Source/PCG/Private/Compute/DataInterfaces/PCGDataLabelResolverDataInterface.h b/Engine/Plugins/PCG/Source/PCG/Private/Compute/DataInterfaces/PCGDataLabelResolverDataInterface.h
new file mode 100644
--- /dev/null
+++ b/Engine/Plugins/PCG/Source/PCG/Private/Compute/DataInterfaces/PCGDataLabelResolverDataInterface.h
@@ -0,0 +1,70 @@
+// Copyright Epic Games, Inc. All Rights Reserved.
+
+#pragma once
+
+#include "Compute/PCGDataBinding.h"
+
+#include "ComputeFramework/ComputeDataInterface.h"
+#include "ComputeFramework/ComputeDataProvider.h"
+
+#include "PCGDataLabelResolverDataInterface.generated.h"
+
+class FPCGDataLabelResolverDataInterfaceParameters;
+
+/** Data interface for mapping from data label to data index. */
+UCLASS(ClassGroup = (Procedural))
+class PCG_API UPCGDataLabelResolverDataInterface : public UComputeDataInterface
+{
+	GENERATED_BODY()
+
+public:
+	//~ Begin UComputeDataInterface Interface
+	TCHAR const* GetClassName() const override { return TEXT("PCGDataLabelResolver"); }
+	bool CanSupportUnifiedDispatch() const override { return true; }
+	void GetSupportedInputs(TArray<FShaderFunctionDefinition>& OutFunctions) const override;
+	void GetSupportedOutputs(TArray<FShaderFunctionDefinition>& OutFunctions) const override;
+	void GetShaderParameters(TCHAR const* UID, FShaderParametersMetadataBuilder& InOutBuilder, FShaderParametersMetadataAllocations& InOutAllocations) const override;
+	void GetHLSL(FString& OutHLSL, FString const& InDataInterfaceName) const override;
+	UComputeDataProvider* CreateDataProvider(TObjectPtr<UObject> InBinding, uint64 InInputMask, uint64 InOutputMask) const override;
+	//~ End UComputeDataInterface Interface
+
+	UPROPERTY()
+	TObjectPtr<const UPCGSettings> ConsumerSettings = nullptr;
+
+	UPROPERTY()
+	FName InputPinLabel;
+};
+
+UCLASS()
+class UPCGDataLabelResolverDataProvider : public UComputeDataProvider
+{
+	GENERATED_BODY()
+
+public:
+	//~ Begin UComputeDataProvider Interface
+	FComputeDataProviderRenderProxy* GetRenderProxy() override;
+	//~ End UComputeDataProvider Interface
+
+	/** Map from data ID to data index for accessing the input data collection. */
+	TArray<int32> DataIdToDataIndexMap;
+};
+
+class FPCGDataLabelResolverDataProviderProxy : public FComputeDataProviderRenderProxy
+{
+public:
+	FPCGDataLabelResolverDataProviderProxy(const TArray<int32>& InDataIdToDataIndexMap)
+		: DataIdToDataIndexMap(InDataIdToDataIndexMap)
+	{}
+
+	//~ Begin FComputeDataProviderRenderProxy Interface
+	bool IsValid(FValidationData const& InValidationData) const override;
+	void GatherDispatchData(FDispatchData const& InDispatchData) override;
+	void AllocateResources(FRDGBuilder& GraphBuilder, FAllocationData const& InAllocationData) override;
+	//~ End FComputeDataProviderRenderProxy Interface
+
+protected:
+	using FParameters = FPCGDataLabelResolverDataInterfaceParameters;
+
+	TArray<int32> DataIdToDataIndexMap;
+	FRDGBufferSRVRef DataIdToDataIndexBufferSRV = nullptr;
+};
diff --git a/Engine/Plugins/PCG/Source/PCG/Private/Compute/Elements/PCGCustomHLSL.cpp b/Engine/Plugins/PCG/Source/PCG/Private/Compute/Elements/PCGCustomHLSL.cpp
--- a/Engine/Plugins/PCG/Source/PCG/Private/Compute/Elements/PCGCustomHLSL.cpp
+++ b/Engine/Plugins/PCG/Source/PCG/Private/Compute/Elements/PCGCustomHLSL.cpp
@@ -648,6 +648,22 @@ bool UPCGCustomHLSLSettings::IsShaderTextReadOnly() const
 {
 	return false;
 }
+
+void UPCGCustomHLSLSettings::GetDataLabels(FName InPinLabel, TArray<FString>& OutDataLabels) const
+{
+	// Matches against {PinName}_AnyFunction('{DataLabel}'...
+	const FString Pattern = FString::Format(TEXT("{0}_.*?[\\s]*?\\([\\s]*?'([a-zA-Z0-9_].*?)'"), { *InPinLabel.ToString() });
+	const FString Source = ShaderFunctions + TEXT("\n") + ShaderSource;
+
+	// First capture: Data label (supports a - z, A - Z, 0 - 9, and underscores).
+	constexpr int LabelCaptureGroup = 1;
+
+	FRegexMatcher ModuleMatcher(FRegexPattern(Pattern), Source);
+	while (ModuleMatcher.FindNext())
+	{
+		OutDataLabels.AddUnique(ModuleMatcher.GetCaptureGroup(LabelCaptureGroup));
+	}
+}
 #endif //WITH_EDITOR
 
 const FPCGPinPropertiesGPU* UPCGCustomHLSLSettings::GetOutputPinPropertiesGPU(const FName& InPinLabel) const
@@ -665,6 +681,8 @@ void UPCGCustomHLSLSettings::UpdateDeclarations()
 	UpdateInputDeclarations();
 	UpdateOutputDeclarations();
 	UpdateHelperDeclarations();
+
+	// TODO: Should data labels be explained/exemplified in the declarations?
 }
 
 void UPCGCustomHLSLSettings::UpdateInputDeclarations()
@@ -1797,8 +1815,10 @@ bool UPCGCustomHLSLSettings::ProcessShaderSource(FString* InOutSource, FText* Ou
 	return true;
 }
 
-FString UPCGCustomHLSLSettings::GetCookedKernelSource(const FPCGKernelAttributeTable& InAttributeStaticTable) const
+FString UPCGCustomHLSLSettings::GetCookedKernelSource(const UPCGComputeGraph* InComputeGraph) const
 {
+	check(InComputeGraph);
+
 	const FIntVector GroupSize(PCGComputeConstants::THREAD_GROUP_SIZE, 1, 1);
 
 	FString Source = ShaderSource;
@@ -1819,12 +1839,51 @@ FString UPCGCustomHLSLSettings::GetCookedKernelSource(const FPCGKernelAttributeT
 	for (const FPCGKernelAttributeKey& AttributeKey : KernelAttributeKeys)
 	{
 		const FString SourceDefinition = PCGHLSLElement::GetKernelAttributeKeyAsString(AttributeKey);
-		const FString AttributeIndexAsString = FString::FromInt(InAttributeStaticTable.GetAttributeId(AttributeKey));
+		const FString AttributeIndexAsString = FString::FromInt(InComputeGraph->GetAttributeStaticTable().GetAttributeId(AttributeKey));
 
 		Source.ReplaceInline(*SourceDefinition, *AttributeIndexAsString);
 		Functions.ReplaceInline(*SourceDefinition, *AttributeIndexAsString);
 	}
 
+	// We inject data IDs directly into the source. They will get remapped to data indices using a label resolver data interface.
+	if (const FPCGPinDataLabels* PinToDataLabels = InComputeGraph->GetDataLabelsStaticTable().Find(this))
+	{
+		for (const TPair<FName, FPCGDataLabels>& Pair : PinToDataLabels->PinToDataLabels)
+		{
+			const FName PinLabel = Pair.Key;
+			const TArray<FString>& DataLabels = Pair.Value.Labels;
+
+			for (int DataId = 0; DataId < DataLabels.Num(); ++DataId)
+			{
+				const FString& DataLabel = DataLabels[DataId];
+				const FString ReplacementStr = FString::Format(TEXT("{0}_GetDataIndexFromIdInternal(/*DataId=*/{1})"),
+					{ 
+						PCGComputeHelpers::GetDataLabelResolverName(PinLabel),
+						FString::FromInt(DataId)
+					});
+
+				// Matches against {PinName}_AnyFunction('{DataLabel}'...
+				// First capture group is the data label, so that we can find & replace it by index.
+				const FString Pattern = FString::Format(TEXT("{0}_.*?[\\s]*?\\([\\s]*?('{1}')"), { *PinLabel.ToString(), DataLabel });
+
+				auto ReplaceDataLabels = [&Pattern, &ReplacementStr](FString& InOutSource)
+				{
+					FRegexMatcher ModuleMatcher(FRegexPattern(Pattern), InOutSource);
+					while (ModuleMatcher.FindNext())
+					{
+						const int32 DataLabelCharIndexStart = ModuleMatcher.GetCaptureGroupBeginning(1);
+						const int32 DataLabelCharIndexEnd = ModuleMatcher.GetCaptureGroupEnding(1);
+
+						InOutSource = InOutSource.Left(DataLabelCharIndexStart) + ReplacementStr + InOutSource.RightChop(DataLabelCharIndexEnd);
+					}
+				};
+
+				ReplaceDataLabels(Source);
+				ReplaceDataLabels(Functions);
+			}
+		}
+	}
+
 	FString ShaderPathName = GetPathName();
 	PCGHLSLElement::ConvertObjectPathToShaderFilePath(ShaderPathName);
 
diff --git a/Engine/Plugins/PCG/Source/PCG/Private/Compute/Elements/PCGCustomHLSL.h b/Engine/Plugins/PCG/Source/PCG/Private/Compute/Elements/PCGCustomHLSL.h
--- a/Engine/Plugins/PCG/Source/PCG/Private/Compute/Elements/PCGCustomHLSL.h
+++ b/Engine/Plugins/PCG/Source/PCG/Private/Compute/Elements/PCGCustomHLSL.h
@@ -65,11 +65,15 @@ public:
 	virtual void ApplyPreconfiguredSettings(const FPCGPreConfiguredSettingsInfo& PreconfigureInfo) override;
 
 	virtual bool IsKernelValid(FPCGContext* InContext = nullptr, bool bQuiet = true) const override;
-	virtual FString GetCookedKernelSource(const FPCGKernelAttributeTable& InAttributeStaticTable) const override;
+	virtual FString GetCookedKernelSource(const UPCGComputeGraph* InComputeGraph) const override;
 	virtual void GetKernelAttributeKeys(TArray<FPCGKernelAttributeKey>& OutKeys) const override { OutKeys.Append(KernelAttributeKeys); }
 	virtual int ComputeKernelThreadCount(const UPCGDataBinding* Binding) const override;
 	virtual bool ComputeOutputPinDataDesc(const UPCGPin* OutputPin, const UPCGDataBinding* Binding, FPCGDataCollectionDesc& OutDesc) const override;
 
+#if WITH_EDITOR
+	virtual void GetDataLabels(FName InPinLabel, TArray<FString>& OutDataLabels) const override;
+#endif
+
 protected:
 #if WITH_EDITOR
 	virtual void PreEditChange(FProperty* PropertyAboutToChange) override;
diff --git a/Engine/Plugins/PCG/Source/PCG/Private/Compute/PCGComputeCommon.cpp b/Engine/Plugins/PCG/Source/PCG/Private/Compute/PCGComputeCommon.cpp
--- a/Engine/Plugins/PCG/Source/PCG/Private/Compute/PCGComputeCommon.cpp
+++ b/Engine/Plugins/PCG/Source/PCG/Private/Compute/PCGComputeCommon.cpp
@@ -374,4 +374,14 @@ namespace PCGComputeHelpers
 	{
 		return InAttributeId >= PCGComputeConstants::NUM_RESERVED_ATTRS ? (InAttributeId - PCGComputeConstants::NUM_RESERVED_ATTRS) : INDEX_NONE;
 	}
+
+	FString GetPrefixedDataLabel(const FString& InLabel)
+	{
+		return FString::Format(TEXT("{0}:{1}"), { PCGComputeConstants::DataLabelTagPrefix, InLabel });
+	}
+
+	FString GetDataLabelResolverName(FName InPinLabel)
+	{
+		return FString::Format(TEXT("{0}_DataResolver"), { *InPinLabel.ToString() });
+	}
 }
diff --git a/Engine/Plugins/PCG/Source/PCG/Private/Compute/PCGComputeGraph.h b/Engine/Plugins/PCG/Source/PCG/Private/Compute/PCGComputeGraph.h
--- a/Engine/Plugins/PCG/Source/PCG/Private/Compute/PCGComputeGraph.h
+++ b/Engine/Plugins/PCG/Source/PCG/Private/Compute/PCGComputeGraph.h
@@ -14,6 +14,26 @@
 class UPCGNode;
 class UPCGPin;
 
+/** Helper struct for serializing data labels. */
+USTRUCT()
+struct FPCGDataLabels
+{
+	GENERATED_BODY()
+
+	UPROPERTY()
+	TArray<FString> Labels;
+};
+
+/** Helper struct for serializing map of pin name to data labels. */
+USTRUCT()
+struct FPCGPinDataLabels
+{
+	GENERATED_BODY()
+
+	UPROPERTY()
+	TMap</*PinLabel*/FName, FPCGDataLabels> PinToDataLabels;
+};
+
 UCLASS()
 class UPCGComputeGraph : public UComputeGraph
 {
@@ -27,6 +47,9 @@ public:
 	/** Get the attribute table collated statically during compilation. */
 	const FPCGKernelAttributeTable& GetAttributeStaticTable() const { return AttributeStaticTable; }
 
+	/** Get the data labels collated statically during compilation. */
+	const TMap<TSoftObjectPtr<const UPCGSettings>, FPCGPinDataLabels>& GetDataLabelsStaticTable() const { return DataLabelsStaticTable; }
+
 	const TArray<FString>& GetStringTable() const { return StringTable; }
 
 public:
@@ -60,6 +83,10 @@ protected:
 	UPROPERTY()
 	FPCGKernelAttributeTable AttributeStaticTable;
 
+	/** Global data label information collated during compilation. */
+	UPROPERTY()
+	TMap<TSoftObjectPtr<const UPCGSettings>, FPCGPinDataLabels> DataLabelsStaticTable;
+
 	/** String table collated during compilation. */
 	UPROPERTY()
 	TArray<FString> StringTable;
diff --git a/Engine/Plugins/PCG/Source/PCG/Private/Elements/Metadata/PCGMetadataPartition.cpp b/Engine/Plugins/PCG/Source/PCG/Private/Elements/Metadata/PCGMetadataPartition.cpp
--- a/Engine/Plugins/PCG/Source/PCG/Private/Elements/Metadata/PCGMetadataPartition.cpp
+++ b/Engine/Plugins/PCG/Source/PCG/Private/Elements/Metadata/PCGMetadataPartition.cpp
@@ -130,7 +130,7 @@ bool UPCGMetadataPartitionSettings::IsKernelValid(FPCGContext* InContext, bool b
 	return true;
 }
 
-FString UPCGMetadataPartitionSettings::GetCookedKernelSource(const FPCGKernelAttributeTable& InAttributeStaticTable) const
+FString UPCGMetadataPartitionSettings::GetCookedKernelSource(const UPCGComputeGraph* InComputeGraph) const
 {
 	FString TemplateFile;
 	ensure(LoadShaderSourceFile(PCGMetadataPartitionSettings::TemplateFilePath, EShaderPlatform::SP_PCD3D_SM5, &TemplateFile, nullptr));
diff --git a/Engine/Plugins/PCG/Source/PCG/Private/Elements/PCGCopyPoints.cpp b/Engine/Plugins/PCG/Source/PCG/Private/Elements/PCGCopyPoints.cpp
--- a/Engine/Plugins/PCG/Source/PCG/Private/Elements/PCGCopyPoints.cpp
+++ b/Engine/Plugins/PCG/Source/PCG/Private/Elements/PCGCopyPoints.cpp
@@ -27,7 +27,7 @@ FText UPCGCopyPointsSettings::GetNodeTooltipText() const
 }
 #endif
 
-FString UPCGCopyPointsSettings::GetCookedKernelSource(const FPCGKernelAttributeTable& InAttributeStaticTable) const
+FString UPCGCopyPointsSettings::GetCookedKernelSource(const UPCGComputeGraph* InComputeGraph) const
 {
 	FString TemplateFile;
 	ensure(LoadShaderSourceFile(TEXT("/Plugin/PCG/Private/Elements/PCGCopyPoints.usf"), EShaderPlatform::SP_PCD3D_SM5, &TemplateFile, nullptr));
diff --git a/Engine/Plugins/PCG/Source/PCG/Private/Elements/PCGGenerateGrassMaps.cpp b/Engine/Plugins/PCG/Source/PCG/Private/Elements/PCGGenerateGrassMaps.cpp
--- a/Engine/Plugins/PCG/Source/PCG/Private/Elements/PCGGenerateGrassMaps.cpp
+++ b/Engine/Plugins/PCG/Source/PCG/Private/Elements/PCGGenerateGrassMaps.cpp
@@ -5,6 +5,7 @@
 #include "PCGComponent.h"
 #include "PCGGrassMapUnpackerCS.h"
 #include "PCGSubsystem.h"
+#include "Compute/PCGComputeCommon.h"
 #include "Data/PCGLandscapeData.h"
 #include "Data/PCGTextureData.h"
 #include "Helpers/PCGHelpers.h"
@@ -473,18 +474,37 @@ bool FPCGGenerateGrassMapsElement::ExecuteInternal(FPCGContext* InContext) const
 
 	// 5. Emit texture data objects.
 
-	for (UPCGTextureData* TextureData : Context->TextureDatas)
+	for (int DataIndex = 0; DataIndex < Context->GrassTypes.Num(); ++DataIndex)
 	{
+		check(Context->TextureDatas.IsValidIndex(DataIndex));
+		UPCGTextureData* TextureData = Context->TextureDatas[DataIndex];
+		check(TextureData);
+
 		if (!TextureData->IsSuccessfullyInitialized())
 		{
 			PCGE_LOG(Warning, LogOnly, LOCTEXT("TextureInitFailed", "Data could not be retrieved for this texture, initialization failed."));
 			continue;
 		}
 
-		Context->OutputData.TaggedData.Emplace_GetRef().Data = TextureData;
+		FPCGTaggedData& OutTaggedData = Context->OutputData.TaggedData.Emplace_GetRef();
+		OutTaggedData.Data = TextureData;
+		
+		FString GrassTypeName = TEXT("");
 
-		// TODO: Tag data with grass type name, if we can get it.
-		// TODO: Also output attribute set for grass type varieties object paths?
+		if (ULandscapeGrassType* GrassType = Context->GrassTypes[DataIndex].Get())
+		{
+			GrassType->GetName(GrassTypeName);
+		}
+
+		if (!GrassTypeName.IsEmpty())
+		{
+			OutTaggedData.Tags.Add(PCGComputeHelpers::GetPrefixedDataLabel(GrassTypeName));
+		}
+		else
+		{
+			PCGE_LOG(Warning, LogOnly, LOCTEXT("MissingGrassName",
+				"Grass type name was missing, data could not be labeled. Make sure all GrassTypes in your landscape material have an asset associated."));
+		}
 	}
 
 	return true;
diff --git a/Engine/Plugins/PCG/Source/PCG/Private/Elements/PCGSkinnedMeshSpawner.cpp b/Engine/Plugins/PCG/Source/PCG/Private/Elements/PCGSkinnedMeshSpawner.cpp
--- a/Engine/Plugins/PCG/Source/PCG/Private/Elements/PCGSkinnedMeshSpawner.cpp
+++ b/Engine/Plugins/PCG/Source/PCG/Private/Elements/PCGSkinnedMeshSpawner.cpp
@@ -84,7 +84,7 @@ bool UPCGSkinnedMeshSpawnerSettings::IsKernelValid(FPCGContext* InContext, bool
 	return true;
 }
 
-FString UPCGSkinnedMeshSpawnerSettings::GetCookedKernelSource(const FPCGKernelAttributeTable& InAttributeStaticTable) const
+FString UPCGSkinnedMeshSpawnerSettings::GetCookedKernelSource(const UPCGComputeGraph* InComputeGraph) const
 {
 	FString TemplateFile;
 	ensure(LoadShaderSourceFile(TemplateFilePath, EShaderPlatform::SP_PCD3D_SM5, &TemplateFile, nullptr));
diff --git a/Engine/Plugins/PCG/Source/PCG/Private/Elements/PCGStaticMeshSpawner.cpp b/Engine/Plugins/PCG/Source/PCG/Private/Elements/PCGStaticMeshSpawner.cpp
--- a/Engine/Plugins/PCG/Source/PCG/Private/Elements/PCGStaticMeshSpawner.cpp
+++ b/Engine/Plugins/PCG/Source/PCG/Private/Elements/PCGStaticMeshSpawner.cpp
@@ -93,7 +93,7 @@ bool UPCGStaticMeshSpawnerSettings::IsKernelValid(FPCGContext* InContext, bool b
 	return true;
 }
 
-FString UPCGStaticMeshSpawnerSettings::GetCookedKernelSource(const FPCGKernelAttributeTable& InAttributeStaticTable) const
+FString UPCGStaticMeshSpawnerSettings::GetCookedKernelSource(const UPCGComputeGraph* InComputeGraph) const
 {
 	FString TemplateFile;
 	ensure(LoadShaderSourceFile(TemplateFilePath, EShaderPlatform::SP_PCD3D_SM5, &TemplateFile, nullptr));
diff --git a/Engine/Plugins/PCG/Source/PCG/Private/Elements/PCGTextureSampler.cpp b/Engine/Plugins/PCG/Source/PCG/Private/Elements/PCGTextureSampler.cpp
--- a/Engine/Plugins/PCG/Source/PCG/Private/Elements/PCGTextureSampler.cpp
+++ b/Engine/Plugins/PCG/Source/PCG/Private/Elements/PCGTextureSampler.cpp
@@ -6,6 +6,7 @@
 #include "PCGContext.h"
 #include "PCGCrc.h"
 #include "PCGSubsystem.h"
+#include "Compute/PCGComputeCommon.h"
 #include "Data/PCGRenderTargetData.h"
 #include "Helpers/PCGBlueprintHelpers.h"
 #include "Helpers/PCGDynamicTrackingHelpers.h"
@@ -174,9 +175,11 @@ bool FPCGTextureSamplerElement::ExecuteInternal(FPCGContext* InContext) const
 		return true;
 	}
 
+	UTexture2DArray* TextureArray = Cast<UTexture2DArray>(Texture);
 	uint32 TextureArrayIndex = 0;
+	const bool bIsTextureArray = TextureArray != nullptr;
 
-	if (UTexture2DArray* TextureArray = Cast<UTexture2DArray>(Texture))
+	if (bIsTextureArray)
 	{
 #if WITH_EDITOR
 		const int32 ArraySize = TextureArray->SourceTextures.Num();
@@ -287,7 +290,12 @@ bool FPCGTextureSamplerElement::ExecuteInternal(FPCGContext* InContext) const
 	check(BaseTextureData);
 
 	// Commit to adding texture data.
-	Context->OutputData.TaggedData.Emplace_GetRef().Data = BaseTextureData;
+	FPCGTaggedData& OutTaggedData = Context->OutputData.TaggedData.Emplace_GetRef();
+	OutTaggedData.Data = BaseTextureData;
+
+	// Tag the data with a label (based on the asset name) in order to reference this texture data downstream via name rather than data index.
+	const FString DataLabel = bIsTextureArray ? FString::Format(TEXT("{0}_{1}"), { Settings->Texture.GetAssetName(), TextureArrayIndex }) : Settings->Texture.GetAssetName();
+	OutTaggedData.Tags.Add(PCGComputeHelpers::GetPrefixedDataLabel(DataLabel));
 
 	BaseTextureData->bUseDensitySourceChannel = Settings->bUseDensitySourceChannel;
 	BaseTextureData->ColorChannel = Settings->ColorChannel;
diff --git a/Engine/Plugins/PCG/Source/PCG/Private/Graph/PCGGraphCompilerGPU.cpp b/Engine/Plugins/PCG/Source/PCG/Private/Graph/PCGGraphCompilerGPU.cpp
--- a/Engine/Plugins/PCG/Source/PCG/Private/Graph/PCGGraphCompilerGPU.cpp
+++ b/Engine/Plugins/PCG/Source/PCG/Private/Graph/PCGGraphCompilerGPU.cpp
@@ -7,9 +7,11 @@
 #include "PCGPin.h"
 #include "Compute/PCGComputeCommon.h"
 #include "Compute/PCGComputeKernelSource.h"
+#include "Compute/DataInterfaces/PCGComputeDataInterface.h"
 #include "Compute/DataInterfaces/PCGCustomKernelDataInterface.h"
 #include "Compute/DataInterfaces/PCGDataCollectionDataInterface.h"
 #include "Compute/DataInterfaces/PCGDataCollectionUploadDataInterface.h"
+#include "Compute/DataInterfaces/PCGDataLabelResolverDataInterface.h"
 #include "Compute/DataInterfaces/PCGDebugDataInterface.h"
 #include "Compute/DataInterfaces/PCGLandscapeDataInterface.h"
 #include "Compute/DataInterfaces/PCGTextureDataInterface.h"
@@ -700,6 +702,28 @@ void FPCGGraphCompilerGPU::CreateDataInterfaces(
 			// Record strings.
 			Settings->AddStaticCreatedStrings(InOutComputeGraph->StringTable);
 
+			// Create all the data index remaps for input pins.
+			auto GetDataLabels = [Settings, InOutComputeGraph](const TArray<FPCGPinProperties>& PinProperties)
+			{
+				for (const FPCGPinProperties& PinProps : PinProperties)
+				{
+					TArray<FString> DataLabels;
+					Settings->GetDataLabels(PinProps.Label, DataLabels);
+
+					if (!DataLabels.IsEmpty())
+					{
+						FPCGPinDataLabels& PinDataLabels = InOutComputeGraph->DataLabelsStaticTable.FindOrAdd(Settings);
+						ensure(!PinDataLabels.PinToDataLabels.Contains(PinProps.Label));
+
+						FPCGDataLabels& Labels = PinDataLabels.PinToDataLabels.Add(PinProps.Label);
+						Labels.Labels = MoveTemp(DataLabels);
+					}
+				}
+			};
+
+			GetDataLabels(Settings->AllInputPinProperties());
+			GetDataLabels(Settings->AllOutputPinProperties());
+
 			// Create all the output data interfaces.
 			for (const FPCGPinProperties& OutputPinProperties : Settings->AllOutputPinProperties())
 			{
@@ -899,12 +923,29 @@ void FPCGGraphCompilerGPU::CompileComputeGraph(
 			InOutComputeGraph->StaticMeshSpawners.Add(Settings);
 		}
 
-		// One data interface may be connected to multiple downstream pins. Because each entry in these arrays produces its own bindings and graph edges,
-		// we also need to specify the pin label to uniquely describe each edge.
-		TArray<TPair</*DataInterfaceIndex=*/int, /*PinLabel=*/FName>> InputDataInterfaceIndexAndPin;
-		TArray<TPair</*DataInterfaceIndex=*/int, /*PinLabel=*/FName>> OutputDataInterfaceIndexAndPin;
-		InputDataInterfaceIndexAndPin.Reserve(Settings->InputPinProperties().Num());
-		OutputDataInterfaceIndexAndPin.Reserve(Settings->OutputPinProperties().Num());
+		// A given data interface may be connected to multiple downstream pins. Because each entry in these arrays produces its own bindings and graph edges,
+		// we also need to specify the prefix to uniquely describe each edge.
+		TArray<TPair</*DataInterfaceIndex=*/int, /*DataInterfacePrefix=*/FName>> InputDataInterfaceIndexAndPrefix;
+		TArray<TPair</*DataInterfaceIndex=*/int, /*DataInterfacePrefix=*/FName>> OutputDataInterfaceIndexAndPrefix;
+		InputDataInterfaceIndexAndPrefix.Reserve(Settings->InputPinProperties().Num());
+		OutputDataInterfaceIndexAndPrefix.Reserve(Settings->OutputPinProperties().Num());
+
+		auto CreateDataResolverDataInterface = [&InOutContext, Settings, InOutComputeGraph](TArray<TPair</*DataInterfaceIndex=*/int, /*DataInterfacePrefix=*/FName>>& InOutDataInterfaceIndexAndPrefix, FName PinLabel)
+		{
+			// Add data label resolver DI if this pin uses any data labels.
+			const FPCGPinDataLabels* PinDataLabels = InOutComputeGraph->GetDataLabelsStaticTable().Find(Settings);
+			const FPCGDataLabels* DataLabels = PinDataLabels ? PinDataLabels->PinToDataLabels.Find(PinLabel) : nullptr;
+
+			if (DataLabels && !DataLabels->Labels.IsEmpty())
+			{
+				UPCGDataLabelResolverDataInterface* Resolver = InOutContext.NewObject_AnyThread<UPCGDataLabelResolverDataInterface>(InOutComputeGraph);
+				Resolver->InputPinLabel = PinLabel;
+				Resolver->ConsumerSettings = Settings;
+
+				const int ResolverIndex = InOutComputeGraph->DataInterfaces.Add(Resolver);
+				InOutDataInterfaceIndexAndPrefix.Emplace(ResolverIndex, PCGComputeHelpers::GetDataLabelResolverName(PinLabel));
+			}
+		};
 
 		// Setup input pin DIs.
 		for (const FPCGGraphTaskInput& Input : InCompiledTasks[TaskId].Inputs)
@@ -939,7 +980,8 @@ void FPCGGraphCompilerGPU::CompileComputeGraph(
 				continue;
 			}
 
-			InputDataInterfaceIndexAndPin.Emplace(Index, Input.DownstreamPin->Label);
+			InputDataInterfaceIndexAndPrefix.Emplace(Index, Input.DownstreamPin->Label);
+			CreateDataResolverDataInterface(InputDataInterfaceIndexAndPrefix, Input.DownstreamPin->Label);
 		}
 
 		// Setup output pin DIs. Always bind a DI to every output pin, so kernel always has something to write to.
@@ -964,7 +1006,8 @@ void FPCGGraphCompilerGPU::CompileComputeGraph(
 				continue;
 			}
 
-			OutputDataInterfaceIndexAndPin.Emplace(Index, OutputPinProperties.Label);
+			OutputDataInterfaceIndexAndPrefix.Emplace(Index, OutputPinProperties.Label);
+			CreateDataResolverDataInterface(OutputDataInterfaceIndexAndPrefix, OutputPinProperties.Label);
 		}
 
 		// Kernel data interface.
@@ -974,7 +1017,9 @@ void FPCGGraphCompilerGPU::CompileComputeGraph(
 
 			const int32 KernelDIIndex = InOutComputeGraph->DataInterfaces.Num();
 			InOutComputeGraph->DataInterfaces.Add(KernelDI);
-			InputDataInterfaceIndexAndPin.Emplace(KernelDIIndex, NAME_None);
+
+			// Kernel DI functions don't need a prefix.
+			InputDataInterfaceIndexAndPrefix.Emplace(KernelDIIndex, NAME_None);
 		}
 
 		// Additional DIs created by settings.
@@ -998,21 +1043,23 @@ void FPCGGraphCompilerGPU::CompileComputeGraph(
 
 			for (TObjectPtr<UComputeDataInterface> DataInterface : AdditionalInputDIs)
 			{
+				// TODO: This could also produce a prefix?
 				const int32 DataInterfaceIndex = AddAdditionalDataInterface(DataInterface);
 
 				if (DataInterfaceIndex != INDEX_NONE)
 				{
-					InputDataInterfaceIndexAndPin.Emplace(DataInterfaceIndex, NAME_None);
+					InputDataInterfaceIndexAndPrefix.Emplace(DataInterfaceIndex, NAME_None);
 				}
 			}
 
 			for (TObjectPtr<UComputeDataInterface> DataInterface : AdditionalOutputDIs)
 			{
+				// TODO: This could also produce a prefix?
 				const int32 DataInterfaceIndex = AddAdditionalDataInterface(DataInterface);
 
 				if (DataInterfaceIndex != INDEX_NONE)
 				{
-					OutputDataInterfaceIndexAndPin.Emplace(DataInterfaceIndex, NAME_None);
+					OutputDataInterfaceIndexAndPrefix.Emplace(DataInterfaceIndex, NAME_None);
 				}
 			}
 		}
@@ -1026,7 +1073,9 @@ void FPCGGraphCompilerGPU::CompileComputeGraph(
 
 			const int32 DebugDIIndex = InOutComputeGraph->DataInterfaces.Num();
 			InOutComputeGraph->DataInterfaces.Add(DebugDI);
-			OutputDataInterfaceIndexAndPin.Emplace(DebugDIIndex, NAME_None);
+
+			// TODO: Maybe debug DI should be prefixed with 'Debug', e.g. 'Debug_WriteValue()'.
+			OutputDataInterfaceIndexAndPrefix.Emplace(DebugDIIndex, NAME_None);
 		}
 
 		InOutComputeGraph->bLogDataDescriptions |= Settings->bDumpDataDescriptions;
@@ -1038,7 +1087,7 @@ void FPCGGraphCompilerGPU::CompileComputeGraph(
 		{
 			const UComputeDataInterface* DataInterface = nullptr;
 			int32 DataInterfaceBindingIndex = INDEX_NONE;
-			FName PinLabel = NAME_None; // Required to uniquely identify this binding if this data interface is connected to multiple pins.
+			FName Prefix = NAME_None;
 			FString BindingFunctionName = TEXT("");
 		};
 
@@ -1066,7 +1115,7 @@ void FPCGGraphCompilerGPU::CompileComputeGraph(
 			{
 				FInterfaceBinding& Binding = KernelWithBindings.InputDataBindings.Emplace_GetRef();
 				Binding.DataInterface = DataInterface;
-				Binding.PinLabel = InDataInterfaceIndex.Value;
+				Binding.Prefix = InDataInterfaceIndex.Value;
 				Binding.BindingFunctionName = Functions[FuncIndex].Name;
 				Binding.DataInterfaceBindingIndex = FuncIndex;
 			}
@@ -1082,19 +1131,19 @@ void FPCGGraphCompilerGPU::CompileComputeGraph(
 			{
 				FInterfaceBinding& Binding = KernelWithBindings.OutputDataBindings.Emplace_GetRef();
 				Binding.DataInterface = DataInterface;
-				Binding.PinLabel = InDataInterfaceIndex.Value;
+				Binding.Prefix = InDataInterfaceIndex.Value;
 				Binding.BindingFunctionName = Functions[FuncIndex].Name;
 				Binding.DataInterfaceBindingIndex = FuncIndex;
 			}
 		};
 
 		// Bind data interfaces.
-		for (const TPair<int, FName>& InputDataInterfaceIndex : InputDataInterfaceIndexAndPin)
+		for (const TPair<int, FName>& InputDataInterfaceIndex : InputDataInterfaceIndexAndPrefix)
 		{
 			SetupAllInputBindings(InputDataInterfaceIndex);
 		}
 
-		for (const TPair<int, FName>& OutputDataInterfaceIndex : OutputDataInterfaceIndexAndPin)
+		for (const TPair<int, FName>& OutputDataInterfaceIndex : OutputDataInterfaceIndexAndPrefix)
 		{
 			SetupAllOutputBindings(OutputDataInterfaceIndex);
 		}
@@ -1107,7 +1156,8 @@ void FPCGGraphCompilerGPU::CompileComputeGraph(
 			KernelSource->EntryPoint = TEXT("Main");
 			KernelSource->GroupSize = FIntVector(PCGComputeConstants::THREAD_GROUP_SIZE, 1, 1);
 
-			KernelSource->SetSource(Settings->GetCookedKernelSource(InOutComputeGraph->AttributeStaticTable));
+			// TODO: Instead of passing a half-baked compute graph into this virtual call, we should pass in a compilation context per compute graph.
+			KernelSource->SetSource(Settings->GetCookedKernelSource(InOutComputeGraph));
 
 			if (Settings->bDumpCookedHLSL)
 			{
@@ -1160,8 +1210,7 @@ void FPCGGraphCompilerGPU::CompileComputeGraph(
 		{
 			TArray<FInterfaceBinding>& Bindings = bInEdgesAreInputs ? KernelWithBindings.InputDataBindings : KernelWithBindings.OutputDataBindings;
 
-			// Add all graph edges for bindings. This is somewhat odd but likely what we'll do vs the more finegrained interface
-			// definitions in optimus, but we need to see.
+			// Add all graph edges for bindings, which means include all functions that the data interfaces expose.
 			for (int BindingIndex = 0; BindingIndex < Bindings.Num(); ++BindingIndex)
 			{
 				FInterfaceBinding& Binding = Bindings[BindingIndex];
@@ -1177,7 +1226,7 @@ void FPCGGraphCompilerGPU::CompileComputeGraph(
 				UComputeDataInterface* DataInterface = InOutComputeGraph->DataInterfaces[Edge.DataInterfaceIndex];
 				check(DataInterface);
 
-				if (Binding.PinLabel != NAME_None)
+				if (Binding.Prefix != NAME_None)
 				{
 					TArray<FShaderFunctionDefinition> DataInterfaceFunctions;
 					if (bInEdgesAreInputs)
@@ -1191,7 +1240,7 @@ void FPCGGraphCompilerGPU::CompileComputeGraph(
 
 					Edge.BindingFunctionNameOverride = FString::Format(
 						TEXT("{0}_{1}"),
-						{ Binding.PinLabel.ToString(), DataInterfaceFunctions[Edge.DataInterfaceBindingIndex].Name }
+						{ Binding.Prefix.ToString(), DataInterfaceFunctions[Edge.DataInterfaceBindingIndex].Name }
 					);
 				}
 			}
diff --git a/Engine/Plugins/PCG/Source/PCG/Public/Elements/Metadata/PCGMetadataPartition.h b/Engine/Plugins/PCG/Source/PCG/Public/Elements/Metadata/PCGMetadataPartition.h
--- a/Engine/Plugins/PCG/Source/PCG/Public/Elements/Metadata/PCGMetadataPartition.h
+++ b/Engine/Plugins/PCG/Source/PCG/Public/Elements/Metadata/PCGMetadataPartition.h
@@ -19,7 +19,7 @@ public:
 
 	//~Begin UPCGSettings interface
 	virtual bool IsKernelValid(FPCGContext* InContext = nullptr, bool bQuiet = true) const;
-	virtual FString GetCookedKernelSource(const FPCGKernelAttributeTable& InAttributeStaticTable) const override;
+	virtual FString GetCookedKernelSource(const UPCGComputeGraph* InComputeGraph) const override;
 	virtual void GetKernelAttributeKeys(TArray<FPCGKernelAttributeKey>& OutKeys) const override;
 	virtual int ComputeKernelThreadCount(const UPCGDataBinding* Binding) const override;
 	virtual bool ComputeOutputPinDataDesc(const UPCGPin* OutputPin, const UPCGDataBinding* Binding, FPCGDataCollectionDesc& OutDesc) const override;
diff --git a/Engine/Plugins/PCG/Source/PCG/Public/Elements/PCGCopyPoints.h b/Engine/Plugins/PCG/Source/PCG/Public/Elements/PCGCopyPoints.h
--- a/Engine/Plugins/PCG/Source/PCG/Public/Elements/PCGCopyPoints.h
+++ b/Engine/Plugins/PCG/Source/PCG/Public/Elements/PCGCopyPoints.h
@@ -54,7 +54,7 @@ public:
 	virtual void CreateAdditionalInputDataInterfaces(FPCGGPUCompilationContext& InOutContext, UObject* InObjectOuter, TArray<TObjectPtr<UComputeDataInterface>>& OutDataInterfaces) const override;
 #endif
 
-	virtual FString GetCookedKernelSource(const FPCGKernelAttributeTable& InAttributeStaticTable) const override;
+	virtual FString GetCookedKernelSource(const UPCGComputeGraph* InComputeGraph) const override;
 	virtual int ComputeKernelThreadCount(const UPCGDataBinding* Binding) const override;
 	virtual bool ComputeOutputPinDataDesc(const UPCGPin* OutputPin, const UPCGDataBinding* Binding, FPCGDataCollectionDesc& OutDesc) const override;
 
diff --git a/Engine/Plugins/PCG/Source/PCG/Public/Elements/PCGSkinnedMeshSpawner.h b/Engine/Plugins/PCG/Source/PCG/Public/Elements/PCGSkinnedMeshSpawner.h
--- a/Engine/Plugins/PCG/Source/PCG/Public/Elements/PCGSkinnedMeshSpawner.h
+++ b/Engine/Plugins/PCG/Source/PCG/Public/Elements/PCGSkinnedMeshSpawner.h
@@ -27,7 +27,7 @@ public:
 
 	// ~Begin UPCGSettings interface
 	virtual bool IsKernelValid(FPCGContext* InContext = nullptr, bool bQuiet = true) const;
-	virtual FString GetCookedKernelSource(const FPCGKernelAttributeTable& InAttributeStaticTable) const override;
+	virtual FString GetCookedKernelSource(const UPCGComputeGraph* InComputeGraph) const override;
 	virtual void GetKernelAttributeKeys(TArray<FPCGKernelAttributeKey>& OutKeys) const override;
 	virtual int ComputeKernelThreadCount(const UPCGDataBinding* Binding) const override;
 	virtual bool ComputeOutputPinDataDesc(const UPCGPin* OutputPin, const UPCGDataBinding* Binding, FPCGDataCollectionDesc& OutDesc) const override;
diff --git a/Engine/Plugins/PCG/Source/PCG/Public/Elements/PCGStaticMeshSpawner.h b/Engine/Plugins/PCG/Source/PCG/Public/Elements/PCGStaticMeshSpawner.h
--- a/Engine/Plugins/PCG/Source/PCG/Public/Elements/PCGStaticMeshSpawner.h
+++ b/Engine/Plugins/PCG/Source/PCG/Public/Elements/PCGStaticMeshSpawner.h
@@ -30,7 +30,7 @@ public:
 
 	// ~Begin UPCGSettings interface
 	virtual bool IsKernelValid(FPCGContext* InContext = nullptr, bool bQuiet = true) const;
-	virtual FString GetCookedKernelSource(const FPCGKernelAttributeTable& InAttributeStaticTable) const override;
+	virtual FString GetCookedKernelSource(const UPCGComputeGraph* InComputeGraph) const override;
 	virtual void GetKernelAttributeKeys(TArray<FPCGKernelAttributeKey>& OutKeys) const override;
 	virtual void AddStaticCreatedStrings(TArray<FString>& InOutStringTable) const override;
 	virtual int ComputeKernelThreadCount(const UPCGDataBinding* Binding) const override;
diff --git a/Engine/Plugins/PCG/Source/PCG/Public/PCGSettings.h b/Engine/Plugins/PCG/Source/PCG/Public/PCGSettings.h
--- a/Engine/Plugins/PCG/Source/PCG/Public/PCGSettings.h
+++ b/Engine/Plugins/PCG/Source/PCG/Public/PCGSettings.h
@@ -514,7 +514,7 @@ public:
 	/** [EXPERIMENTAL] Produces the node specific portion of kernel shader source text, including the main entry point.
 	* Note: This API function is subject to change without deprecation.
 	*/
-	virtual FString GetCookedKernelSource(const FPCGKernelAttributeTable& InAttributeStaticTable) const { return TEXT(""); }
+	virtual FString GetCookedKernelSource(const UPCGComputeGraph* InComputeGraph) const { return TEXT(""); }
 
 	/** [EXPERIMENTAL] Get all the attributes read or written by this node for which we know the name and type statically. In cases that only the
 	* attribute name used by the node is known (example: user specifies SM Spawner instance data attributes by name only, not type), the attribute
@@ -560,6 +560,11 @@ public:
 	virtual void CreateAdditionalOutputDataInterfaces(FPCGGPUCompilationContext& InOutContext, UObject* InObjectOuter, TArray<TObjectPtr<UComputeDataInterface>>& OutDataInterfaces) const {}
 	virtual void CreateAdditionalOutputDataInterfaces(TArray<TObjectPtr<UComputeDataInterface>>& OutDataInterfaces) const {}
 
+	/** [EXPERIMENTAL] Get any data labels that are used statically by this node.
+	* Note: This API function is subject to change without deprecation.
+	*/
+	virtual void GetDataLabels(FName InPinLabel, TArray<FString>& OutDataLabels) const {}
+
 	/** [EXPERIMENTAL] Whether compute graphs should be split at the output of this node. Useful for reading back runtime statistics or diagnostic info before continuing.
 	* Note: This API function is subject to change without deprecation.
 	*/
