Added graph customization (at a single level) that can control visible nodes (category-based) and graphs (based on list/hierarchy)

diff --git a/Engine/Plugins/PCG/Source/PCG/Private/Editor/PCGGraphCustomization.cpp b/Engine/Plugins/PCG/Source/PCG/Private/Editor/PCGGraphCustomization.cpp
new file mode 100644
--- /dev/null
+++ b/Engine/Plugins/PCG/Source/PCG/Private/Editor/PCGGraphCustomization.cpp
@@ -0,0 +1,76 @@
+// Copyright Epic Games, Inc. All Rights Reserved.
+
+#include "Editor/PCGGraphCustomization.h"
+#include "PCGGraph.h"
+
+#include UE_INLINE_GENERATED_CPP_BY_NAME(PCGGraphCustomization)
+
+const FPCGGraphEditorCustomization* FPCGGraphEditorCustomization::GetParent() const
+{
+	return nullptr;
+}
+
+bool FPCGGraphEditorCustomization::Accepts(const FText& InCategory) const
+{
+	if (bFilterNodesByCategory)
+	{
+		FString Category = InCategory.ToString();
+
+		for (const FString& Filter : FilteredCategories)
+		{
+			if (Category.StartsWith(Filter))
+			{
+				return NodeFilterType == EPCGGraphEditorFiltering::Allow;
+			}
+		}
+	}
+
+	// "Default" behavior -> delegate to parent
+	if (const FPCGGraphEditorCustomization* Parent = GetParent())
+	{
+		return Parent->Accepts(InCategory);
+	}
+	else
+	{
+		return !bFilterNodesByCategory || NodeFilterType != EPCGGraphEditorFiltering::Allow;
+	}
+}
+
+bool FPCGGraphEditorCustomization::FiltersSubgraphs() const
+{
+	if (bFilterSubgraphs)
+	{
+		return true;
+	}
+	else if (const FPCGGraphEditorCustomization* Parent = GetParent())
+	{
+		return Parent->FiltersSubgraphs();
+	}
+	else
+	{
+		return false;
+	}
+}
+
+bool FPCGGraphEditorCustomization::FilterSubgraph(const FSoftObjectPath& InSubgraphPath) const
+{
+	if (bFilterSubgraphs)
+	{
+		const TSoftObjectPtr<UPCGGraph> SubgraphSoftPtr = TSoftObjectPtr<UPCGGraph>(InSubgraphPath);
+
+		if (FilteredSubgraphTypes.Contains(SubgraphSoftPtr))
+		{
+			return SubgraphFilterType != EPCGGraphEditorFiltering::Allow;
+		}
+	}
+
+	// Otherwise, delegate filtering to parent if any
+	if (const FPCGGraphEditorCustomization* Parent = GetParent())
+	{
+		return Parent->FilterSubgraph(InSubgraphPath);
+	}
+	else
+	{
+		return bFilterSubgraphs && NodeFilterType == EPCGGraphEditorFiltering::Allow;
+	} 
+}
\ No newline at end of file
diff --git a/Engine/Plugins/PCG/Source/PCG/Private/PCGGraph.cpp b/Engine/Plugins/PCG/Source/PCG/Private/PCGGraph.cpp
--- a/Engine/Plugins/PCG/Source/PCG/Private/PCGGraph.cpp
+++ b/Engine/Plugins/PCG/Source/PCG/Private/PCGGraph.cpp
@@ -19,6 +19,8 @@
 #include "Graph/PCGGraphCompiler.h"
 #include "Graph/PCGGraphExecutor.h"
 
+#include "AssetRegistry/AssetData.h"
+
 #if WITH_EDITOR
 #include "CoreGlobals.h"
 #include "Editor.h"
@@ -271,6 +273,22 @@ bool UPCGGraphInterface::VerifyAndUpdateIfGraphParameterValueChanged(const FProp
 		return false;
 	}
 }
+
+bool UPCGGraphInterface::VerifyIfGraphCustomizationChanged(const FPropertyChangedChainEvent& PropertyChangedEvent)
+{
+	const TDoubleLinkedList<FProperty*>::TDoubleLinkedListNode* PropertyNode = PropertyChangedEvent.PropertyChain.GetHead();
+	while (PropertyNode)
+	{
+		if (PropertyNode->GetValue() && PropertyNode->GetValue()->GetFName() == GET_MEMBER_NAME_CHECKED(UPCGGraph, GraphCustomization))
+		{
+			return true;
+		}
+
+		PropertyNode = PropertyNode->GetNextNode();
+	}
+
+	return false;
+}
 #endif // WITH_EDITOR
 
 EPCGChangeType UPCGGraphInterface::GetChangeTypeForGraphParameterChange(EPCGGraphParameterEvent InChangeType, FName InChangedPropertyName)
@@ -372,6 +390,10 @@ UPCGGraph::UPCGGraph(const FObjectInitializer& ObjectInitializer)
 	// Force the user parameters to have an empty property bag. It is necessary to catch the first
 	// add property into the undo/redo history.
 	UserParameters.MigrateToNewBagStruct(UPropertyBag::GetOrCreateFromDescs({}));
+
+#if WITH_EDITOR
+	GraphCustomization.Graph = this;
+#endif
 }
 
 void UPCGGraph::PostLoad()
@@ -468,6 +490,8 @@ void UPCGGraph::PostLoad()
 	{
 		FixInvalidEdges();
 	}
+
+	GraphCustomization.Graph = this;
 #endif // WITH_EDITOR
 }
 
@@ -1687,6 +1711,10 @@ void UPCGGraph::PostEditChangeChainProperty(FPropertyChangedChainEvent& Property
 	{
 		NotifyGraphChanged(EPCGChangeType::Cosmetic);
 	}
+	else if (VerifyIfGraphCustomizationChanged(PropertyChangedEvent))
+	{
+		NotifyGraphChanged(EPCGChangeType::GraphCustomization);
+	}
 	else if (VerifyAndUpdateIfGraphParameterValueChanged(PropertyChangedEvent))
 	{
 		// Handled in the function
@@ -2017,6 +2045,10 @@ void UPCGGraphInstance::PostEditChangeChainProperty(FPropertyChangedChainEvent&
 	{
 		// Changes to this property will be caught by the update on override status, and "OnGraphParametersChanged" will be called accordingly.
 	}
+	else if (VerifyIfGraphCustomizationChanged(PropertyChangedEvent))
+	{
+		OnGraphChangedDelegate.Broadcast(this, EPCGChangeType::GraphCustomization);
+	}
 	else if (VerifyAndUpdateIfGraphParameterValueChanged(PropertyChangedEvent))
 	{
 		// Handled in the function
@@ -2564,4 +2596,21 @@ void FPCGOverrideInstancedPropertyBag::MigrateToNewBagInstance(const FInstancedP
 	}
 }
 
+bool UPCGGraphInstance::GraphAssetFilter(const FAssetData& AssetData) const
+{
+#if WITH_EDITOR
+	if (UPCGSubgraphSettings* SubgraphSettings = Cast<UPCGSubgraphSettings>(GetOuter()))
+	{
+		return SubgraphSettings->SubgraphAssetFilter(AssetData);
+	}
+	// TODO : add filtering on PCG components?
+	else
+	{
+		return false;
+	}
+#else
+	return false;
+#endif
+}
+
 #undef LOCTEXT_NAMESPACE
diff --git a/Engine/Plugins/PCG/Source/PCG/Private/PCGSubgraph.cpp b/Engine/Plugins/PCG/Source/PCG/Private/PCGSubgraph.cpp
--- a/Engine/Plugins/PCG/Source/PCG/Private/PCGSubgraph.cpp
+++ b/Engine/Plugins/PCG/Source/PCG/Private/PCGSubgraph.cpp
@@ -11,6 +11,8 @@
 #include "Helpers/PCGSettingsHelpers.h"
 
 #include "Algo/Find.h"
+#include "AssetRegistry/AssetData.h"
+#include "AssetRegistry/AssetRegistryModule.h"
 
 #include UE_INLINE_GENERATED_CPP_BY_NAME(PCGSubgraph)
 
@@ -331,6 +333,57 @@ void UPCGSubgraphSettings::SetSubgraphInternal(UPCGGraphInterface* InGraph)
 	SubgraphInstance->SetGraph(InGraph);
 }
 
+bool UPCGSubgraphSettings::SubgraphAssetFilter(const FAssetData& AssetData) const
+{
+#if WITH_EDITOR
+	// Need access to this graph to get the customization
+	UPCGGraph* Graph = GetOuter() ? Cast<UPCGGraph>(GetOuter()->GetOuter()) : nullptr;
+
+	if (!Graph || !Graph->GraphCustomization.FiltersSubgraphs())
+	{
+		return false;
+	}
+
+	FAssetRegistryModule& AssetRegistryModule = FModuleManager::Get().LoadModuleChecked<FAssetRegistryModule>(TEXT("AssetRegistry"));
+
+	auto AssetTest = [Graph, &AssetRegistryModule](const FAssetData& AssetData) -> bool
+	{
+		FSoftObjectPath PathToTest;
+		if (AssetData.IsInstanceOf<UPCGGraphInstance>())
+		{
+			PathToTest = FSoftObjectPath(AssetData.GetTagValueRef<FString>(GET_MEMBER_NAME_CHECKED(UPCGGraphInstance, Graph)));
+		}
+		else
+		{
+			PathToTest = AssetData.GetSoftObjectPath();
+		}
+
+		// Check if this is accepted by the filter, or recurse a level otherwise
+		if (Graph->GraphCustomization.FilterSubgraph(PathToTest))
+		{
+			// Get parent path and recurse
+			FAssetData CurrentAsset = AssetData;
+			if (AssetData.IsInstanceOf<UPCGGraphInstance>())
+			{
+				CurrentAsset = AssetRegistryModule.Get().GetAssetByObjectPath(PathToTest);
+			}
+
+			// Refused by the filter
+			return true;
+		}
+		else
+		{
+			// Accepted by the filter
+			return false;
+		}
+	};
+
+	return AssetTest(AssetData);
+#else
+	return false; // always accept.
+#endif
+}
+
 void UPCGSubgraphSettings::PostLoad()
 {
 #if WITH_EDITOR
diff --git a/Engine/Plugins/PCG/Source/PCG/Public/Editor/PCGGraphCustomization.h b/Engine/Plugins/PCG/Source/PCG/Public/Editor/PCGGraphCustomization.h
new file mode 100644
--- /dev/null
+++ b/Engine/Plugins/PCG/Source/PCG/Public/Editor/PCGGraphCustomization.h
@@ -0,0 +1,58 @@
+// Copyright Epic Games, Inc. All Rights Reserved.
+
+#pragma once
+
+#include "PCGCommon.h"
+#include "PCGNode.h"
+#include "PCGSettings.h"
+
+#include "PCGGraphCustomization.generated.h"
+
+UENUM()
+enum class EPCGGraphEditorFiltering
+{
+	Allow,
+	Disallow
+};
+
+/** 
+* Container struct for editor customization of a specific graph. 
+* This will grow/evolve as we expose more options to adapt the PCG editor to specific use cases.
+* This can be used to restrict choices and to make the usage more streamlined when it calls for it.
+* It will be extended to support hierarchy of customization (for things that make sense).
+*/
+USTRUCT()
+struct PCG_API FPCGGraphEditorCustomization
+{
+	GENERATED_BODY()
+
+	// Pointer to graph so we can have hierarchical behavior
+	UPROPERTY(Transient)
+	TObjectPtr<const UPCGGraph> Graph;
+
+	UPROPERTY(EditAnywhere, Category = Filtering)
+	bool bFilterNodesByCategory = false;
+
+	UPROPERTY(EditAnywhere, Category = Filtering, meta = (EditCondition = "bFilterNodesByCategory"))
+	EPCGGraphEditorFiltering NodeFilterType = EPCGGraphEditorFiltering::Allow;
+
+	UPROPERTY(EditAnywhere, Category = Filtering, meta = (EditCondition = "bFilterNodesByCategory"))
+	TSet<FString> FilteredCategories;
+
+	UPROPERTY(EditAnywhere, Category = Filtering)
+	bool bFilterSubgraphs = false;
+
+	UPROPERTY(EditAnywhere, Category = Filtering, meta = (EditCondition = "bFilterSubgraphs"))
+	EPCGGraphEditorFiltering SubgraphFilterType = EPCGGraphEditorFiltering::Allow;
+
+	UPROPERTY(EditAnywhere, Category = Filtering, meta = (EditCondition = "bFilterSubgraphs"))
+	TSet<TSoftObjectPtr<UPCGGraph>> FilteredSubgraphTypes;
+
+	bool Accepts(const FText& InCategory) const;
+	bool FiltersSubgraphs() const;
+	// Returns false if path is accepted by the customization
+	bool FilterSubgraph(const FSoftObjectPath& InSubgraphPath) const;
+
+private:
+	const FPCGGraphEditorCustomization* GetParent() const;
+};
\ No newline at end of file
diff --git a/Engine/Plugins/PCG/Source/PCG/Public/PCGCommon.h b/Engine/Plugins/PCG/Source/PCG/Public/PCGCommon.h
--- a/Engine/Plugins/PCG/Source/PCG/Public/PCGCommon.h
+++ b/Engine/Plugins/PCG/Source/PCG/Public/PCGCommon.h
@@ -44,7 +44,7 @@ namespace PCGPinIdHelpers
 }
 
 UENUM(meta = (Bitflags))
-enum class EPCGChangeType : uint8
+enum class EPCGChangeType : uint32
 {
 	None = 0,
 	Cosmetic = 1 << 0,
@@ -57,6 +57,8 @@ enum class EPCGChangeType : uint8
 	GenerationGrid = 1 << 6,
 	/** Change to any shader source code. */
 	ShaderSource = 1 << 7,
+	/** Changes in the graph customization that will impact the editor. */
+	GraphCustomization = 1 << 8,
 };
 ENUM_CLASS_FLAGS(EPCGChangeType);
 
@@ -100,7 +102,7 @@ enum class EPCGDataType : uint32
 	Param = 1 << 27 UMETA(DisplayName = "Attribute Set"),
 
 	// Combination of Param and Point, necessary for named-based serialization of enums.
-	PointOrParam = Point | Param,
+	PointOrParam = Point | Param UMETA(DisplayName = "Point or Attribute Set"),
 
 	Settings = 1 << 28 UMETA(Hidden),
 	Other = 1 << 29,
@@ -131,7 +133,7 @@ enum class EPCGExclusiveDataType : uint8
 	Settings UMETA(Hidden),
 	Other,
 	Any,
-	PointOrParam,
+	PointOrParam UMETA(DisplayName = "Point or Attribute Set"),
 	DynamicMesh
 };
 
diff --git a/Engine/Plugins/PCG/Source/PCG/Public/PCGGraph.h b/Engine/Plugins/PCG/Source/PCG/Public/PCGGraph.h
--- a/Engine/Plugins/PCG/Source/PCG/Public/PCGGraph.h
+++ b/Engine/Plugins/PCG/Source/PCG/Public/PCGGraph.h
@@ -6,6 +6,7 @@
 #include "PCGNode.h"
 #include "PCGSettings.h"
 #include "Compute/PCGCompilerDiagnostic.h"
+#include "Editor/PCGGraphCustomization.h"
 #include "Graph/PCGStackContext.h"
 #include "Helpers/PCGGraphParameterExtension.h"
 
@@ -14,6 +15,7 @@
 
 #include "PCGGraph.generated.h"
 
+struct FAssetData;
 class UPCGGraphInterface;
 class UPCGGraphCompilationData;
 #if WITH_EDITOR
@@ -133,11 +135,15 @@ protected:
 	virtual bool IsTemplatePropertyEnabled() const { return false; }
 
 	bool VerifyAndUpdateIfGraphParameterValueChanged(const FPropertyChangedChainEvent& PropertyChangedEvent);
+	bool VerifyIfGraphCustomizationChanged(const FPropertyChangedChainEvent& PropertyChangedEvent);
 #endif // WITH_EDITOR
 
 public:
 
 #if WITH_EDITORONLY_DATA
+	UPROPERTY(EditAnywhere, Category = Customization)
+	FPCGGraphEditorCustomization GraphCustomization;
+
 	UPROPERTY(BlueprintReadWrite, EditAnywhere, Category = AssetInfo, AssetRegistrySearchable, meta = (EditCondition = "IsExportToLibraryEnabled", EditConditionHides))
 	bool bExposeToLibrary = false;
 
@@ -235,6 +241,7 @@ class PCG_API UPCGGraph : public UPCGGraphInterface
 	friend class FPCGEditor;
 	friend class FPCGSubgraphHelpers;
 #endif // WITH_EDITOR
+	friend class UPCGSubgraphSettings;
 
 	GENERATED_BODY()
 
@@ -594,7 +601,11 @@ public:
 	bool IsPropertyOverridden(const FProperty* InProperty) const { return ParametersOverrides.IsPropertyOverridden(InProperty); }
 	bool IsPropertyOverriddenAndNotDefault(const FProperty* InProperty) const;
 
-	UPROPERTY(EditAnywhere, BlueprintReadOnly, Category = Instance, AssetRegistrySearchable)
+	/** Because subgraph nodes hold graph instances, delegation of the selection can't be done elsewhere, so this function gets called when the dialog is opened and will forward to the graph customization if filtering is applied. */
+	UFUNCTION()
+	bool GraphAssetFilter(const FAssetData& AssetData) const;
+
+	UPROPERTY(EditAnywhere, BlueprintReadOnly, Category = Instance, AssetRegistrySearchable, meta = (GetAssetFilter = "GraphAssetFilter"))
 	TObjectPtr<UPCGGraphInterface> Graph;
 
 	UPROPERTY(EditAnywhere, Category = Instance, meta = (NoResetToDefault))
diff --git a/Engine/Plugins/PCG/Source/PCG/Public/PCGNode.h b/Engine/Plugins/PCG/Source/PCG/Public/PCGNode.h
--- a/Engine/Plugins/PCG/Source/PCG/Public/PCGNode.h
+++ b/Engine/Plugins/PCG/Source/PCG/Public/PCGNode.h
@@ -6,7 +6,7 @@
 
 class UPCGNode;
 class UPCGPin;
-enum class EPCGChangeType : uint8;
+enum class EPCGChangeType : uint32;
 enum class EPCGNodeTitleType : uint8;
 struct FPCGPinProperties;
 
diff --git a/Engine/Plugins/PCG/Source/PCG/Public/PCGSubgraph.h b/Engine/Plugins/PCG/Source/PCG/Public/PCGSubgraph.h
--- a/Engine/Plugins/PCG/Source/PCG/Public/PCGSubgraph.h
+++ b/Engine/Plugins/PCG/Source/PCG/Public/PCGSubgraph.h
@@ -123,6 +123,11 @@ protected:
 public:
 	virtual UPCGGraphInterface* GetSubgraphInterface() const override;
 	virtual bool IsDynamicGraph() const override;
+
+	/** Used to filter the subgraph list based on the graph configuration. */
+	UFUNCTION()
+	bool SubgraphAssetFilter(const FAssetData& AssetData) const;
+
 protected:
 	virtual void SetSubgraphInternal(UPCGGraphInterface* InGraph) override;
 	//~End UPCGBaseSubgraphSettings interface
diff --git a/Engine/Plugins/PCG/Source/PCGEditor/Private/PCGEditor.cpp b/Engine/Plugins/PCG/Source/PCGEditor/Private/PCGEditor.cpp
--- a/Engine/Plugins/PCG/Source/PCGEditor/Private/PCGEditor.cpp
+++ b/Engine/Plugins/PCG/Source/PCGEditor/Private/PCGEditor.cpp
@@ -2724,7 +2724,7 @@ void FPCGEditor::InitToolMenuContext(FToolMenuContext& MenuContext)
 
 TSharedRef<SPCGEditorGraphNodePalette> FPCGEditor::CreatePaletteWidget()
 {
-	return SNew(SPCGEditorGraphNodePalette);
+	return SNew(SPCGEditorGraphNodePalette, SharedThis(this));
 }
 
 TSharedRef<SPCGEditorGraphDebugObjectTree> FPCGEditor::CreateDebugObjectTreeWidget()
@@ -2952,6 +2952,14 @@ void FPCGEditor::OnGraphChanged(UPCGGraphInterface* InGraph, EPCGChangeType Chan
 		// We could make the user do this manually, but that makes iterating on data interfaces really painful.
 		FlushShaderFileCache();
 	}
+
+	if (!!(ChangeType & EPCGChangeType::GraphCustomization))
+	{
+		if (PaletteWidget)
+		{
+			PaletteWidget->RequestRefresh();
+		}		
+	}
 }
 
 void FPCGEditor::OnNodeSourceCompiled(const UPCGNode* InNode, const FPCGCompilerDiagnostics& InDiagnostics)
diff --git a/Engine/Plugins/PCG/Source/PCGEditor/Private/PCGEditorGraphNodeBase.h b/Engine/Plugins/PCG/Source/PCGEditor/Private/PCGEditorGraphNodeBase.h
--- a/Engine/Plugins/PCG/Source/PCGEditor/Private/PCGEditorGraphNodeBase.h
+++ b/Engine/Plugins/PCG/Source/PCGEditor/Private/PCGEditorGraphNodeBase.h
@@ -11,7 +11,7 @@
 #include "PCGEditorGraphNodeBase.generated.h"
 
 struct FPCGPreconfiguredInfo;
-enum class EPCGChangeType : uint8;
+enum class EPCGChangeType : uint32;
 
 class UPCGComponent;
 class UPCGGraph;
diff --git a/Engine/Plugins/PCG/Source/PCGEditor/Private/PCGEditorGraphSchema.cpp b/Engine/Plugins/PCG/Source/PCGEditor/Private/PCGEditorGraphSchema.cpp
--- a/Engine/Plugins/PCG/Source/PCGEditor/Private/PCGEditorGraphSchema.cpp
+++ b/Engine/Plugins/PCG/Source/PCGEditor/Private/PCGEditorGraphSchema.cpp
@@ -36,46 +36,61 @@
 
 #define LOCTEXT_NAMESPACE "PCGEditorGraphSchema"
 
-void UPCGEditorGraphSchema::GetPaletteActions(FGraphActionMenuBuilder& ActionMenuBuilder, const EPCGElementType InPCGElementTypeFilter) const
+FPCGActionsFilter::FPCGActionsFilter(const UEdGraph* InEdGraph, EPCGElementType InElementFilterType)
+	: FilterType(InElementFilterType)
+	, Graph(Cast<UPCGEditorGraph>(InEdGraph))
 {
-	if (!!(InPCGElementTypeFilter & EPCGElementType::Native))
+	if (Graph && Graph->GetPCGGraph())
 	{
-		GetNativeElementActions(ActionMenuBuilder);
+		Customization = Graph->GetPCGGraph()->GraphCustomization;
 	}
-	if (!!(InPCGElementTypeFilter & EPCGElementType::Subgraph))
+}
+
+bool FPCGActionsFilter::Accepts(const FText& InCategory) const
+{
+	return Customization.Accepts(InCategory);
+}
+
+void UPCGEditorGraphSchema::GetPaletteActions(FGraphActionMenuBuilder& ActionMenuBuilder, const FPCGActionsFilter& ActionFilter) const
+{
+	GetGraphActions(ActionMenuBuilder, ActionFilter, /*bIsContextual=*/false);
+}
+
+void UPCGEditorGraphSchema::GetGraphContextActions(FGraphContextMenuBuilder& ContextMenuBuilder) const
+{
+	Super::GetGraphContextActions(ContextMenuBuilder);
+	GetGraphActions(ContextMenuBuilder, FPCGActionsFilter(ContextMenuBuilder.CurrentGraph), /*bIsContextual=*/true);
+}
+
+void UPCGEditorGraphSchema::GetGraphActions(FGraphActionMenuBuilder& ActionMenuBuilder, const FPCGActionsFilter& ActionFilter, bool bIsContextual) const
+{
+	EPCGElementType ElementTypeFilter = ActionFilter.FilterType;
+
+	if (!!(ElementTypeFilter & EPCGElementType::Native))
 	{
-		GetSubgraphElementActions(ActionMenuBuilder);
+		GetNativeElementActions(ActionMenuBuilder, ActionFilter);
 	}
-	if (!!(InPCGElementTypeFilter & EPCGElementType::Blueprint))
+	if (!!(ElementTypeFilter & EPCGElementType::Subgraph))
 	{
-		GetBlueprintElementActions(ActionMenuBuilder);
+		GetSubgraphElementActions(ActionMenuBuilder, ActionFilter);
 	}
-	if (!!(InPCGElementTypeFilter & EPCGElementType::Settings))
+	if (!!(ElementTypeFilter & EPCGElementType::Blueprint))
 	{
-		GetSettingsElementActions(ActionMenuBuilder, /*bIsContextual=*/false);
+		GetBlueprintElementActions(ActionMenuBuilder, ActionFilter);
 	}
-	if (!!(InPCGElementTypeFilter & EPCGElementType::Asset))
+	if (!!(ElementTypeFilter & EPCGElementType::Settings))
 	{
-		GetDataAssetActions(ActionMenuBuilder);
+		GetSettingsElementActions(ActionMenuBuilder, ActionFilter, bIsContextual);
 	}
-	if (!!(InPCGElementTypeFilter & EPCGElementType::Other))
+	if (!!(ElementTypeFilter & EPCGElementType::Asset))
 	{
-		GetNamedRerouteUsageActions(ActionMenuBuilder);
-		GetExtraElementActions(ActionMenuBuilder);
+		GetDataAssetActions(ActionMenuBuilder, ActionFilter);
+	}
+	if (!!(ElementTypeFilter & EPCGElementType::Other))
+	{
+		GetNamedRerouteUsageActions(ActionMenuBuilder, ActionFilter);
+		GetExtraElementActions(ActionMenuBuilder, ActionFilter);
 	}
-}
-
-void UPCGEditorGraphSchema::GetGraphContextActions(FGraphContextMenuBuilder& ContextMenuBuilder) const
-{
-	Super::GetGraphContextActions(ContextMenuBuilder);
-
-	GetNativeElementActions(ContextMenuBuilder, ContextMenuBuilder.CurrentGraph);
-	GetSubgraphElementActions(ContextMenuBuilder);
-	GetBlueprintElementActions(ContextMenuBuilder);
-	GetSettingsElementActions(ContextMenuBuilder, /*bIsContextual=*/true);
-	GetDataAssetActions(ContextMenuBuilder);
-	GetNamedRerouteUsageActions(ContextMenuBuilder, ContextMenuBuilder.CurrentGraph);
-	GetExtraElementActions(ContextMenuBuilder);
 }
 
 FLinearColor UPCGEditorGraphSchema::GetPinTypeColor(const FEdGraphPinType& PinType) const
@@ -368,7 +383,7 @@ void UPCGEditorGraphSchema::BreakSinglePinLink(UEdGraphPin* SourcePin, UEdGraphP
 	PCGGraph->RemoveEdge(SourcePCGNode, SourcePin->PinName, TargetPCGNode, TargetPin->PinName);
 }
 
-void UPCGEditorGraphSchema::GetNativeElementActions(FGraphActionMenuBuilder& ActionMenuBuilder, const UEdGraph* CurrentGraph) const
+void UPCGEditorGraphSchema::GetNativeElementActions(FGraphActionMenuBuilder& ActionMenuBuilder, const FPCGActionsFilter& ActionFilter) const
 {
 	TArray<UClass*> SettingsClasses;
 	for (TObjectIterator<UClass> It; It; ++It)
@@ -392,6 +407,11 @@ void UPCGEditorGraphSchema::GetNativeElementActions(FGraphActionMenuBuilder& Act
 				const FText Category = StaticEnum<EPCGSettingsType>()->GetDisplayNameTextByValue(static_cast<__underlying_type(EPCGSettingsType)>(PCGSettings->GetType()));
 				const FText Description = PCGSettings->GetNodeTooltipText();
 
+				if (!ActionFilter.Accepts(Category))
+				{
+					continue;
+				}
+
 				TArray<FPCGPreConfiguredSettingsInfo> AllPreconfiguredInfo = PCGSettings->GetPreconfiguredInfo();
 
 				if (AllPreconfiguredInfo.IsEmpty() || !PCGSettings->OnlyExposePreconfiguredSettings())
@@ -412,6 +432,11 @@ void UPCGEditorGraphSchema::GetNativeElementActions(FGraphActionMenuBuilder& Act
 				// Also add preconfigured settings
 				const FText NewCategory = PCGSettings->GroupPreconfiguredSettings() ? FText::Format(LOCTEXT("PreconfiguredSettingsCategory", "{0}|{1}"), Category, MenuDesc) : Category;
 
+				if (!ActionFilter.Accepts(NewCategory))
+				{
+					continue;
+				}
+
 				for (FPCGPreConfiguredSettingsInfo PreconfiguredInfo : AllPreconfiguredInfo)
 				{
 					TSharedPtr<FPCGEditorGraphSchemaAction_NewNativeElement> NewPreconfiguredAction(new FPCGEditorGraphSchemaAction_NewNativeElement(NewCategory, PreconfiguredInfo.Label, PreconfiguredInfo.Tooltip.IsEmpty() ? Description : PreconfiguredInfo.Tooltip, 0));
@@ -423,36 +448,33 @@ void UPCGEditorGraphSchema::GetNativeElementActions(FGraphActionMenuBuilder& Act
 		}
 	}
 
-	if (const UPCGEditorGraph* Graph = Cast<UPCGEditorGraph>(CurrentGraph))
+	const FText UserParameterCategory = LOCTEXT("UserParametersCategoryName", "Graph Parameters");
+
+	if (ActionFilter.Accepts(UserParameterCategory))
 	{
-		if (const UPCGGraph* PCGGraph = const_cast<UPCGEditorGraph*>(Graph)->GetPCGGraph())
+		const UPCGEditorGraph* Graph = ActionFilter.Graph;
+		const UPCGGraph* PCGGraph = Graph ? const_cast<UPCGEditorGraph*>(Graph)->GetPCGGraph() : nullptr;
+		const FInstancedPropertyBag* UserParameters = PCGGraph ? PCGGraph->GetUserParametersStruct() : nullptr;
+		if (const UPropertyBag* BagStruct = UserParameters ? UserParameters->GetPropertyBagStruct() : nullptr)
 		{
-			if (const FInstancedPropertyBag* UserParameters = PCGGraph->GetUserParametersStruct())
+			for (const FPropertyBagPropertyDesc& PropertyDesc : BagStruct->GetPropertyDescs())
 			{
-				if (const UPropertyBag* BagStruct = UserParameters->GetPropertyBagStruct())
-				{
-					const FText Category = LOCTEXT("UserParametersCategoryName", "Graph Parameters");
-
-					for (const FPropertyBagPropertyDesc& PropertyDesc : BagStruct->GetPropertyDescs())
-					{
-						const FText MenuDesc = FText::Format(LOCTEXT("GetterNodeName", "Get {0}"), FText::FromName(PropertyDesc.Name));
-						const FText Description = FText::Format(LOCTEXT("NodeTooltip", "Get the value from '{0}' parameter, can be overridden by the graph instance."), FText::FromName(PropertyDesc.Name));
-
-						TSharedPtr<FPCGEditorGraphSchemaAction_NewGetParameterElement> NewAction(new FPCGEditorGraphSchemaAction_NewGetParameterElement(Category, MenuDesc, Description, 0));
-						NewAction->SettingsClass = UPCGUserParameterGetSettings::StaticClass();
-						NewAction->PropertyName = PropertyDesc.Name;
-						NewAction->PropertyGuid = PropertyDesc.ID;
-						ActionMenuBuilder.AddAction(NewAction);
-					}
-				}
+				const FText MenuDesc = FText::Format(LOCTEXT("GetterNodeName", "Get {0}"), FText::FromName(PropertyDesc.Name));
+				const FText Description = FText::Format(LOCTEXT("NodeTooltip", "Get the value from '{0}' parameter, can be overridden by the graph instance."), FText::FromName(PropertyDesc.Name));
+
+				TSharedPtr<FPCGEditorGraphSchemaAction_NewGetParameterElement> NewAction(new FPCGEditorGraphSchemaAction_NewGetParameterElement(UserParameterCategory, MenuDesc, Description, 0));
+				NewAction->SettingsClass = UPCGUserParameterGetSettings::StaticClass();
+				NewAction->PropertyName = PropertyDesc.Name;
+				NewAction->PropertyGuid = PropertyDesc.ID;
+				ActionMenuBuilder.AddAction(NewAction);
 			}
 		}
 	}
 }
 
-void UPCGEditorGraphSchema::GetBlueprintElementActions(FGraphActionMenuBuilder& ActionMenuBuilder) const
+void UPCGEditorGraphSchema::GetBlueprintElementActions(FGraphActionMenuBuilder& ActionMenuBuilder, const FPCGActionsFilter& ActionFilter) const
 {
-	PCGEditorUtils::ForEachPCGBlueprintAssetData([&ActionMenuBuilder](const FAssetData& AssetData)
+	PCGEditorUtils::ForEachPCGBlueprintAssetData([&ActionMenuBuilder, &ActionFilter](const FAssetData& AssetData)
 	{
 		const bool bExposeToLibrary = AssetData.GetTagValueRef<bool>(GET_MEMBER_NAME_CHECKED(UPCGBlueprintElement, bExposeToLibrary));
 		const bool bOnlyExposePreconfiguredSettings = AssetData.GetTagValueRef<bool>(GET_MEMBER_NAME_CHECKED(UPCGBlueprintElement, bOnlyExposePreconfiguredSettings));
@@ -463,6 +485,11 @@ void UPCGEditorGraphSchema::GetBlueprintElementActions(FGraphActionMenuBuilder&
 			const FText Category = AssetData.GetTagValueRef<FText>(GET_MEMBER_NAME_CHECKED(UPCGBlueprintElement, Category));
 			const FText Description = AssetData.GetTagValueRef<FText>(GET_MEMBER_NAME_CHECKED(UPCGBlueprintElement, Description));
 
+			if (!ActionFilter.Accepts(Category))
+			{
+				return true;
+			}
+
 			const FSoftClassPath GeneratedClass = FSoftClassPath(AssetData.GetTagValueRef<FString>(FBlueprintTags::GeneratedClassPath));
 
 			// Only load the class if we have enabled preconfigured settings.
@@ -485,6 +512,11 @@ void UPCGEditorGraphSchema::GetBlueprintElementActions(FGraphActionMenuBuilder&
 			// Also add preconfigured settings
 			const FText NewCategory = FText::Format(LOCTEXT("PreconfiguredSettingsCategory", "{0}|{1}"), Category, MenuDesc);
 
+			if (!ActionFilter.Accepts(NewCategory))
+			{
+				return true;
+			}
+
 			for (FPCGPreConfiguredSettingsInfo PreconfiguredInfo : AllPreconfiguredInfo)
 			{
 				TSharedPtr<FPCGEditorGraphSchemaAction_NewBlueprintElement> NewPreconfiguredAction(new FPCGEditorGraphSchemaAction_NewBlueprintElement(NewCategory, PreconfiguredInfo.Label, PreconfiguredInfo.Tooltip.IsEmpty() ? Description : PreconfiguredInfo.Tooltip, 0));
@@ -498,27 +530,26 @@ void UPCGEditorGraphSchema::GetBlueprintElementActions(FGraphActionMenuBuilder&
 	});
 }
 
-void UPCGEditorGraphSchema::GetSettingsElementActions(FGraphActionMenuBuilder& ActionMenuBuilder, bool bIsContextual) const
+void UPCGEditorGraphSchema::GetSettingsElementActions(FGraphActionMenuBuilder& ActionMenuBuilder, const FPCGActionsFilter& ActionFilter, bool bIsContextual) const
 {
-	PCGEditorUtils::ForEachPCGSettingsAssetData([&ActionMenuBuilder, bIsContextual](const FAssetData& AssetData)
+	PCGEditorUtils::ForEachPCGSettingsAssetData([&ActionMenuBuilder, &ActionFilter, bIsContextual](const FAssetData& AssetData)
 	{
 		const bool bExposeToLibrary = AssetData.GetTagValueRef<bool>(TEXT("bExposeToLibrary"));
 		if (bExposeToLibrary)
 		{
 			const FText MenuDesc = FText::FromString(FName::NameToDisplayString(AssetData.AssetName.ToString(), false));
 			const FText Category = AssetData.GetTagValueRef<FText>(TEXT("Category"));
+			const FText MenuAndSubCategory = FText::Join(LOCTEXT("MenuDelimiter", "|"), Category, MenuDesc);
 			const FText Description = AssetData.GetTagValueRef<FText>(TEXT("Description"));
 
-			if (!bIsContextual)
+			if (!bIsContextual && ActionFilter.Accepts(Category))
 			{
 				TSharedPtr<FPCGEditorGraphSchemaAction_NewSettingsElement> NewSettingsAction(new FPCGEditorGraphSchemaAction_NewSettingsElement(Category, MenuDesc, Description, 0));
 				NewSettingsAction->SettingsObjectPath = AssetData.GetSoftObjectPath();
 				ActionMenuBuilder.AddAction(NewSettingsAction);
 			}
-			else
+			else if(bIsContextual && ActionFilter.Accepts(MenuAndSubCategory))
 			{
-				const FText MenuAndSubCategory = FText::Join(LOCTEXT("MenuDelimiter", "|"), Category, MenuDesc);
-
 				TSharedPtr<FPCGEditorGraphSchemaAction_NewSettingsElement> NewSettingsActionCopy(new FPCGEditorGraphSchemaAction_NewSettingsElement(MenuAndSubCategory, LOCTEXT("ContextMenuCopySettings", "Copy"), Description, 0));
 				NewSettingsActionCopy->SettingsObjectPath = AssetData.GetSoftObjectPath();
 				NewSettingsActionCopy->Behavior = EPCGEditorNewSettingsBehavior::ForceCopy;
@@ -535,11 +566,11 @@ void UPCGEditorGraphSchema::GetSettingsElementActions(FGraphActionMenuBuilder& A
 	});
 }
 
-void UPCGEditorGraphSchema::GetSubgraphElementActions(FGraphActionMenuBuilder& ActionMenuBuilder) const
+void UPCGEditorGraphSchema::GetSubgraphElementActions(FGraphActionMenuBuilder& ActionMenuBuilder, const FPCGActionsFilter& ActionFilter) const
 {
 	FAssetRegistryModule& AssetRegistryModule = FModuleManager::Get().LoadModuleChecked<FAssetRegistryModule>(TEXT("AssetRegistry"));
 
-	PCGEditorUtils::ForEachPCGGraphAssetData([&ActionMenuBuilder, &AssetRegistryModule](const FAssetData& AssetData) -> bool
+	PCGEditorUtils::ForEachPCGGraphAssetData([&ActionMenuBuilder, &ActionFilter, &AssetRegistryModule](const FAssetData& AssetData) -> bool
 	{
 		const bool bExposeToLibrary = AssetData.GetTagValueRef<bool>(GET_MEMBER_NAME_CHECKED(UPCGGraphInterface, bExposeToLibrary));
 		if (bExposeToLibrary)
@@ -594,7 +625,7 @@ void UPCGEditorGraphSchema::GetSubgraphElementActions(FGraphActionMenuBuilder& A
 
 			FText Category, Description;
 
-			if (GetRecursiveTextsAndIsValid(AssetData, Category, Description, GetRecursiveTextsAndIsValid))
+			if (GetRecursiveTextsAndIsValid(AssetData, Category, Description, GetRecursiveTextsAndIsValid) && ActionFilter.Accepts(Category))
 			{
 				// As stated above, we either have an override and we take it, or we use the asset name, to differentiate all possible instances of the same graph.
 				const FText MenuDesc = AssetData.GetTagValueRef<bool>(GET_MEMBER_NAME_CHECKED(UPCGGraphInterface, bOverrideTitle))
@@ -611,7 +642,7 @@ void UPCGEditorGraphSchema::GetSubgraphElementActions(FGraphActionMenuBuilder& A
 	});
 }
 
-void UPCGEditorGraphSchema::GetExtraElementActions(FGraphActionMenuBuilder& ActionMenuBuilder) const
+void UPCGEditorGraphSchema::GetExtraElementActions(FGraphActionMenuBuilder& ActionMenuBuilder, const FPCGActionsFilter& ActionFilter) const
 {
 	const FText NoCategory;
 	// Comment action
@@ -636,8 +667,9 @@ void UPCGEditorGraphSchema::GetExtraElementActions(FGraphActionMenuBuilder& Acti
 	ActionMenuBuilder.AddAction(NewNamedRerouteDeclarationAction);
 }
 
-void UPCGEditorGraphSchema::GetNamedRerouteUsageActions(FGraphActionMenuBuilder& ActionMenuBuilder, const UEdGraph* CurrentGraph) const
+void UPCGEditorGraphSchema::GetNamedRerouteUsageActions(FGraphActionMenuBuilder& ActionMenuBuilder, const FPCGActionsFilter& ActionFilter) const
 {
+	const UEdGraph* CurrentGraph = ActionFilter.Graph;
 	const UPCGEditorGraph* Graph = Cast<UPCGEditorGraph>(CurrentGraph);
 
 	if (!Graph)
@@ -666,9 +698,9 @@ void UPCGEditorGraphSchema::GetNamedRerouteUsageActions(FGraphActionMenuBuilder&
 	}
 }
 
-void UPCGEditorGraphSchema::GetDataAssetActions(FGraphActionMenuBuilder& ActionMenuBuilder) const
+void UPCGEditorGraphSchema::GetDataAssetActions(FGraphActionMenuBuilder& ActionMenuBuilder, const FPCGActionsFilter& ActionFilter) const
 {
-	PCGEditorUtils::ForEachPCGAssetData([&ActionMenuBuilder](const FAssetData& AssetData)
+	PCGEditorUtils::ForEachPCGAssetData([&ActionMenuBuilder, &ActionFilter](const FAssetData& AssetData)
 	{
 		const bool bExposeToLibrary = AssetData.GetTagValueRef<bool>(GET_MEMBER_NAME_CHECKED(UPCGDataAsset, bExposeToLibrary));
 		if (bExposeToLibrary)
@@ -678,6 +710,11 @@ void UPCGEditorGraphSchema::GetDataAssetActions(FGraphActionMenuBuilder& ActionM
 			const FText Description = AssetData.GetTagValueRef<FText>(GET_MEMBER_NAME_CHECKED(UPCGDataAsset, Description));
 			const FString AssetName = AssetData.GetTagValueRef<FString>(GET_MEMBER_NAME_CHECKED(UPCGDataAsset, Name));
 
+			if (!ActionFilter.Accepts(Category))
+			{
+				return true;
+			}
+
 			const FSoftClassPath SettingsClassPath = FSoftClassPath(AssetData.GetTagValueRef<FString>(GET_MEMBER_NAME_CHECKED(UPCGDataAsset, SettingsClass)));
 			TSubclassOf<UPCGSettings> SettingsClass = SettingsClassPath.TryLoadClass<UPCGSettings>();
 
diff --git a/Engine/Plugins/PCG/Source/PCGEditor/Private/PCGEditorGraphSchema.h b/Engine/Plugins/PCG/Source/PCGEditor/Private/PCGEditorGraphSchema.h
--- a/Engine/Plugins/PCG/Source/PCGEditor/Private/PCGEditorGraphSchema.h
+++ b/Engine/Plugins/PCG/Source/PCGEditor/Private/PCGEditorGraphSchema.h
@@ -3,6 +3,10 @@
 #pragma once
 
 #include "EdGraph/EdGraphSchema.h"
+
+#include "PCGEditorCommon.h"
+#include "Editor/PCGGraphCustomization.h"
+
 #include "ConnectionDrawingPolicy.h"
 
 #include "PCGEditorGraphSchema.generated.h"
@@ -11,13 +15,24 @@ enum class EPCGElementType : uint8;
 
 class UPCGEditorGraph;
 
+struct FPCGActionsFilter
+{
+	explicit FPCGActionsFilter(const UEdGraph* InEdGraph, EPCGElementType InElementFilterType = EPCGElementType::All);
+
+	bool Accepts(const FText& InCategory) const;
+
+	EPCGElementType FilterType = EPCGElementType::All;
+	const UPCGEditorGraph* Graph = nullptr;
+	FPCGGraphEditorCustomization Customization;
+};
+
 UCLASS()
 class UPCGEditorGraphSchema : public UEdGraphSchema
 {
 	GENERATED_BODY()
 
 public:
-	void GetPaletteActions(FGraphActionMenuBuilder& ActionMenuBuilder, const EPCGElementType InPCGElementTypeFilter) const;
+	void GetPaletteActions(FGraphActionMenuBuilder& ActionMenuBuilder, const FPCGActionsFilter& ActionFilter) const;
 
 	//~ Begin EdGraphSchema Interface
 	virtual void GetGraphContextActions(FGraphContextMenuBuilder& ContextMenuBuilder) const override;
@@ -33,13 +48,14 @@ public:
 	//~ End EdGraphSchema Interface
 
 private:
-	void GetNativeElementActions(FGraphActionMenuBuilder& ActionMenuBuilder, const UEdGraph* CurrentGraph = nullptr) const;
-	void GetBlueprintElementActions(FGraphActionMenuBuilder& ActionMenuBuilder) const;
-	void GetSubgraphElementActions(FGraphActionMenuBuilder& ActionMenuBuilder) const;
-	void GetSettingsElementActions(FGraphActionMenuBuilder& ActionMenuBuilder, bool bIsContextual) const;
-	void GetExtraElementActions(FGraphActionMenuBuilder& ActionMenuBuilder) const;
-	void GetNamedRerouteUsageActions(FGraphActionMenuBuilder& ActionMenuBuilder, const UEdGraph* CurrentGraph = nullptr) const;
-	void GetDataAssetActions(FGraphActionMenuBuilder& ActionMenuBuilder) const;
+	void GetGraphActions(FGraphActionMenuBuilder& ActionMenuBuilder, const FPCGActionsFilter& ActionFilter, bool bIsContextual) const;
+	void GetNativeElementActions(FGraphActionMenuBuilder& ActionMenuBuilder, const FPCGActionsFilter& ActionFilter) const;
+	void GetBlueprintElementActions(FGraphActionMenuBuilder& ActionMenuBuilder, const FPCGActionsFilter& ActionFilter) const;
+	void GetSubgraphElementActions(FGraphActionMenuBuilder& ActionMenuBuilder, const FPCGActionsFilter& ActionFilter) const;
+	void GetSettingsElementActions(FGraphActionMenuBuilder& ActionMenuBuilder, const FPCGActionsFilter& ActionFilter, bool bIsContextual) const;
+	void GetExtraElementActions(FGraphActionMenuBuilder& ActionMenuBuilder, const FPCGActionsFilter& ActionFilter) const;
+	void GetNamedRerouteUsageActions(FGraphActionMenuBuilder& ActionMenuBuilder, const FPCGActionsFilter& ActionFilter) const;
+	void GetDataAssetActions(FGraphActionMenuBuilder& ActionMenuBuilder, const FPCGActionsFilter& ActionFilter) const;
 
 	virtual bool TryCreateConnectionInternal(UEdGraphPin* A, UEdGraphPin* B, bool bAddConversionNodeIfNeeded) const;
 };
diff --git a/Engine/Plugins/PCG/Source/PCGEditor/Private/SPCGEditorGraphNodePalette.cpp b/Engine/Plugins/PCG/Source/PCGEditor/Private/SPCGEditorGraphNodePalette.cpp
--- a/Engine/Plugins/PCG/Source/PCGEditor/Private/SPCGEditorGraphNodePalette.cpp
+++ b/Engine/Plugins/PCG/Source/PCGEditor/Private/SPCGEditorGraphNodePalette.cpp
@@ -5,6 +5,8 @@
 #include "AssetRegistry/AssetData.h"
 #include "PCGGraph.h"
 
+#include "PCGEditor.h"
+#include "PCGEditorGraph.h"
 #include "PCGEditorGraphSchema.h"
 #include "PCGEditorGraphSchemaActions.h"
 #include "PCGEditorUtils.h"
@@ -37,9 +39,10 @@ FText SPCGEditorGraphNodePaletteItem::GetItemTooltip() const
 	return ActionPtr.Pin()->GetTooltipDescription();
 }
 
-void SPCGEditorGraphNodePalette::Construct(const FArguments& InArgs)
+void SPCGEditorGraphNodePalette::Construct(const FArguments& InArgs, TSharedPtr<FPCGEditor> InPCGEditor)
 {
 	const UEnum* PCGElementTypeEnum = StaticEnum<EPCGElementType>();
+	PCGEditor = InPCGEditor;
 	
 	this->ChildSlot
 	[
@@ -99,6 +102,11 @@ SPCGEditorGraphNodePalette::~SPCGEditorGraphNodePalette()
 	}
 }
 
+void SPCGEditorGraphNodePalette::RequestRefresh()
+{
+	bNeedsRefresh = true;
+}
+
 TSharedRef<SWidget> SPCGEditorGraphNodePalette::OnCreateWidgetForAction(FCreateWidgetForActionData* const InCreateData)
 {
 	return SNew(SPCGEditorGraphNodePaletteItem, InCreateData);
@@ -124,7 +132,7 @@ void SPCGEditorGraphNodePalette::CollectAllActions(FGraphActionListBuilderBase&
 	const UPCGEditorGraphSchema* PCGSchema = GetDefault<UPCGEditorGraphSchema>();
 
 	FGraphActionMenuBuilder ActionMenuBuilder;
-	PCGSchema->GetPaletteActions(ActionMenuBuilder, ElementType);
+	PCGSchema->GetPaletteActions(ActionMenuBuilder, FPCGActionsFilter(PCGEditor.Pin()->GetPCGEditorGraph(), ElementType));
 	OutAllActions.Append(ActionMenuBuilder);
 }
 
diff --git a/Engine/Plugins/PCG/Source/PCGEditor/Private/SPCGEditorGraphNodePalette.h b/Engine/Plugins/PCG/Source/PCGEditor/Private/SPCGEditorGraphNodePalette.h
--- a/Engine/Plugins/PCG/Source/PCGEditor/Private/SPCGEditorGraphNodePalette.h
+++ b/Engine/Plugins/PCG/Source/PCGEditor/Private/SPCGEditorGraphNodePalette.h
@@ -6,6 +6,8 @@
 
 #include "SGraphPalette.h"
 
+class FPCGEditor;
+
 class SPCGEditorGraphNodePaletteItem : public SGraphPaletteItem
 {
 public:
@@ -27,9 +29,11 @@ public:
 	SLATE_BEGIN_ARGS(SPCGEditorGraphNodePalette) {};
 	SLATE_END_ARGS()
 
-	void Construct(const FArguments& InArgs);
+	void Construct(const FArguments& InArgs, TSharedPtr<FPCGEditor> InPCGEditor);
 	virtual ~SPCGEditorGraphNodePalette();
 
+	void RequestRefresh();
+
 protected:
 	//~ Begin SGraphPalette Interface
 	virtual TSharedRef<SWidget> OnCreateWidgetForAction(FCreateWidgetForActionData* const InCreateData) override;
@@ -48,6 +52,7 @@ private:
 	
 	int32 GetTypeValue() const;
 	
+	TWeakPtr<FPCGEditor> PCGEditor;
 	EPCGElementType ElementType = EPCGElementType::All;
 	bool bNeedsRefresh = false;
 };
