[PCG] Fix ensure happening when PCG Partition Actors GenerateOnLoad without a loaded OriginalComponent

diff --git a/Engine/Plugins/PCG/Source/PCG/Private/PCGComponent.cpp b/Engine/Plugins/PCG/Source/PCG/Private/PCGComponent.cpp
--- a/Engine/Plugins/PCG/Source/PCG/Private/PCGComponent.cpp
+++ b/Engine/Plugins/PCG/Source/PCG/Private/PCGComponent.cpp
@@ -1661,7 +1661,10 @@ void UPCGComponent::BeginPlay()
 	// Procedural instances are never persisted so always require generation.
 	const bool bAlreadyGenerated = bGenerated & !bProceduralInstancesInUse;
 
-	if (bActivated && !bAlreadyGenerated && GenerationTrigger == EPCGComponentGenerationTrigger::GenerateOnLoad)
+	// Either this is the original component and it is non-null or this is a local component and we need the original component to be loaded to trigger a generation
+	UPCGComponent* OriginalComponent = GetOriginalComponent();
+
+	if (OriginalComponent && bActivated && !bAlreadyGenerated && GenerationTrigger == EPCGComponentGenerationTrigger::GenerateOnLoad)
 	{
 		GenerateInternal(/*bForce=*/false, EPCGHiGenGrid::Uninitialized, EPCGComponentGenerationTrigger::GenerateOnLoad, {});
 		bRuntimeGenerated = true;
