[PCG] Multi-Layer Metadata

* First pass on it, everything should use the default layer for now, more CLs will come after those to add support for different layers. This is the ground base work.
* Add a new class `FPCGMetadataLayer` that fit in place of the previous `UPCGMetdata` API
* `UPCGMetadata` forward most of its calls to `FPCGMetadataLayer`.
* Add new `FPCGMetadataLayerID` and `FPCGAttributeIdentifier`, to identify different metadata layers, and different attributes on a given layer.
* `FPCGAttributeIdentifier` is constructible from a FName to give an easy path to upgrade, while maintaining compatibility with previous code (aka only one Metadata, the default one).
* Each data is responsible to setup its layers in advance and to know which layers are available
* There is a concept of "Default Layer" for retro-compatibility, all data have a default layer of "Data" (meaning global metadata), except Points and Attributes, which is "Elements"
* Deprecate all methods that directly manipulated a `UPCGMetadata` for easy path to upgrade (Attributes and Accessors).

diff --git a/Engine/Plugins/PCG/Config/BasePCG.ini b/Engine/Plugins/PCG/Config/BasePCG.ini
--- a/Engine/Plugins/PCG/Config/BasePCG.ini
+++ b/Engine/Plugins/PCG/Config/BasePCG.ini
@@ -23,6 +23,9 @@
 +FunctionRedirects=(OldName="PCGDifferenceData.AddDifference",NewName="PCGDifferenceData.K2_AddDifference")
 +FunctionRedirects=(OldName="PCGParamData.FilterParamsByName",NewName="PCGParamData.K2_FilterParamsByName")
 +FunctionRedirects=(OldName="PCGParamData.FilterParamsByKey",NewName="PCGParamData.K2_FilterParamsByKey")
++FunctionRedirects=(OldName="PCGMetadata.HasAttribute",NewName="PCGMetadata.BP_HasAttribute")
++FunctionRedirects=(OldName="PCGMetadata.AddAttribute",NewName="PCGMetadata.BP_AddAttribute")
++FunctionRedirects=(OldName="PCGMetadata.CopyAttribute",NewName="PCGMetadata.BP_CopyAttribute")
 
 +PropertyRedirects=(OldName="PCGSplineSamplerSettings.Params",NewName="PCGSplineSamplerSettings.SamplerParams")
 +PropertyRedirects=(OldName="PCGProjectionSettings.Params",NewName="PCGProjectionSettings.ProjectionParams")
diff --git a/Engine/Plugins/PCG/Source/PCG/Private/Data/PCGBasePointData.cpp b/Engine/Plugins/PCG/Source/PCG/Private/Data/PCGBasePointData.cpp
--- a/Engine/Plugins/PCG/Source/PCG/Private/Data/PCGBasePointData.cpp
+++ b/Engine/Plugins/PCG/Source/PCG/Private/Data/PCGBasePointData.cpp
@@ -2,3 +2,9 @@
 
 #include "Data/PCGBasePointData.h"
 
+UPCGBasePointData::UPCGBasePointData(const FObjectInitializer& ObjectInitializer)
+	: Super(ObjectInitializer)
+{
+	check(Metadata);
+	Metadata->SetupLayer(PCGMetadataLayerID::Elements, /*bIsDefault=*/true);
+}
\ No newline at end of file
diff --git a/Engine/Plugins/PCG/Source/PCG/Private/Data/PCGSpatialData.cpp b/Engine/Plugins/PCG/Source/PCG/Private/Data/PCGSpatialData.cpp
--- a/Engine/Plugins/PCG/Source/PCG/Private/Data/PCGSpatialData.cpp
+++ b/Engine/Plugins/PCG/Source/PCG/Private/Data/PCGSpatialData.cpp
@@ -18,6 +18,7 @@ UPCGSpatialData::UPCGSpatialData(const FObjectInitializer& ObjectInitializer)
 	: Super(ObjectInitializer)
 {
 	Metadata = ObjectInitializer.CreateDefaultSubobject<UPCGMetadata>(this, TEXT("Metadata"));
+	Metadata->SetupLayer(PCGMetadataLayerID::Data, /*bIsDefault=*/true);
 }
 
 void UPCGSpatialDataWithPointCache::GetResourceSizeEx(FResourceSizeEx& CumulativeResourceSize)
diff --git a/Engine/Plugins/PCG/Source/PCG/Private/Helpers/PCGMetadataHelpers.cpp b/Engine/Plugins/PCG/Source/PCG/Private/Helpers/PCGMetadataHelpers.cpp
--- a/Engine/Plugins/PCG/Source/PCG/Private/Helpers/PCGMetadataHelpers.cpp
+++ b/Engine/Plugins/PCG/Source/PCG/Private/Helpers/PCGMetadataHelpers.cpp
@@ -18,6 +18,16 @@ namespace PCGMetadataHelpers
 	{
 		return Metadata1 && Metadata2 && Metadata1->GetRoot() == Metadata2->GetRoot();
 	}
+	
+	bool HasSameRoot(const FPCGMetadataLayer* Metadata1, const FPCGMetadataLayer* Metadata2)
+	{
+		return Metadata1 && Metadata2 && Metadata1->GetRoot() == Metadata2->GetRoot();
+	}
+	
+	FPCGMetadataLayer* GetDefaultMetadataLayer(UPCGMetadata* InMetadata)
+	{
+		return InMetadata ? InMetadata->GetDefaultMetadataLayer() : nullptr;
+	}
 
 	const UPCGMetadata* GetParentMetadata(const UPCGMetadata* Metadata)
 	{
@@ -30,6 +40,12 @@ namespace PCGMetadataHelpers
 		return Parent.Get();
 	}
 
+	const FPCGMetadataLayer* GetParentMetadata(const FPCGMetadataLayer* Metadata)
+	{
+		check(Metadata);
+		return Metadata->GetParent();
+	}
+
 	const UPCGMetadata* GetConstMetadata(const UPCGData* InData)
 	{
 		return InData ? InData->ConstMetadata() : nullptr;
@@ -379,6 +395,49 @@ namespace PCGMetadataHelpers
 
 		return CopyAttributes(TargetData, SourceData, AttributeSelectors, /*bSameOrigin=*/false, OptionalContext);
 	}
+
+	void ComputePointWeightedAttribute(FPCGMetadataLayer* InOutMetadata, FPCGPoint& OutPoint, const TArrayView<TPair<const FPCGPoint*, float>>& InWeightedPoints, const FPCGMetadataLayer* InMetadata)
+	{
+		check(InOutMetadata);
+		TArray<TPair<PCGMetadataEntryKey, float>, TInlineAllocator<4>> InWeightedKeys;
+		InWeightedKeys.Reserve(InWeightedPoints.Num());
+
+		for (const TPair<const FPCGPoint*, float>& WeightedPoint : InWeightedPoints)
+		{
+			InWeightedKeys.Emplace(WeightedPoint.Key->MetadataEntry, WeightedPoint.Value);
+		}
+
+		InOutMetadata->ComputeWeightedAttribute(OutPoint.MetadataEntry, MakeArrayView(InWeightedKeys), InMetadata);
+	}
+
+	void SetPointAttributes(FPCGMetadataLayer* InOutMetadata, const TArrayView<const FPCGPoint>& InPoints, const FPCGMetadataLayer* InMetadata, const TArrayView<FPCGPoint>& OutPoints, FPCGContext* OptionalContext)
+	{
+		check(InOutMetadata);
+		if (!InMetadata || InMetadata->GetAttributeCount() == 0 || InOutMetadata->GetAttributeCount() == 0)
+		{
+			return;
+		}
+
+		TRACE_CPUPROFILER_EVENT_SCOPE(PCGMetadataHelpers::SetPointAttributes);
+
+		check(InPoints.Num() == OutPoints.Num());
+
+		// Extract the metadata entry keys from the in & out points
+		TArray<PCGMetadataEntryKey, TInlineAllocator<256>> InKeys;
+		TArray<PCGMetadataEntryKey, TInlineAllocator<256>> OutKeys;
+
+		Algo::Transform(InPoints, InKeys, [](const FPCGPoint& Point) { return Point.MetadataEntry; });
+		Algo::Transform(OutPoints, OutKeys, [](const FPCGPoint& Point) { return Point.MetadataEntry; });
+
+		InOutMetadata->SetAttributes(InKeys, InMetadata, OutKeys, OptionalContext);
+
+		// Write back the keys on the points
+		for (int KeyIndex = 0; KeyIndex < OutKeys.Num(); ++KeyIndex)
+		{
+			OutPoints[KeyIndex].MetadataEntry = OutKeys[KeyIndex];
+		}
+	}
+
 }
 
 #undef LOCTEXT_NAMESPACE
\ No newline at end of file
diff --git a/Engine/Plugins/PCG/Source/PCG/Private/Metadata/Accessors/PCGAttributeAccessorHelpers.cpp b/Engine/Plugins/PCG/Source/PCG/Private/Metadata/Accessors/PCGAttributeAccessorHelpers.cpp
--- a/Engine/Plugins/PCG/Source/PCG/Private/Metadata/Accessors/PCGAttributeAccessorHelpers.cpp
+++ b/Engine/Plugins/PCG/Source/PCG/Private/Metadata/Accessors/PCGAttributeAccessorHelpers.cpp
@@ -466,6 +466,18 @@ TUniquePtr<const IPCGAttributeAccessor> PCGAttributeAccessorHelpers::CreateConst
 }
 
 TUniquePtr<const IPCGAttributeAccessor> PCGAttributeAccessorHelpers::CreateConstAccessor(const FPCGMetadataAttributeBase* InAttribute, const UPCGMetadata* InMetadata, bool bQuiet)
+{
+	if (!InMetadata || !InAttribute)
+	{
+		return {};
+	}
+	else
+	{
+		return CreateConstAccessor(InAttribute, InMetadata->GetConstDefaultMetadataLayer(), bQuiet);
+	}
+}
+
+TUniquePtr<const IPCGAttributeAccessor> PCGAttributeAccessorHelpers::CreateConstAccessor(const FPCGMetadataAttributeBase* InAttribute, const FPCGMetadataLayer* InMetadata, bool bQuiet)
 {
 	if (!InMetadata || !InAttribute)
 	{
@@ -531,6 +543,18 @@ TUniquePtr<IPCGAttributeAccessor> PCGAttributeAccessorHelpers::CreateAccessor(UP
 }
 
 TUniquePtr<IPCGAttributeAccessor> PCGAttributeAccessorHelpers::CreateAccessor(FPCGMetadataAttributeBase* InAttribute, UPCGMetadata* InMetadata, bool bQuiet)
+{
+	if (!InMetadata || !InAttribute)
+	{
+		return {};
+	}
+	else
+	{
+		return CreateAccessor(InAttribute, InMetadata->GetDefaultMetadataLayer());
+	}
+}
+
+TUniquePtr<IPCGAttributeAccessor> PCGAttributeAccessorHelpers::CreateAccessor(FPCGMetadataAttributeBase* InAttribute, FPCGMetadataLayer* InMetadata, bool bQuiet)
 {
 	if (!InMetadata || !InAttribute)
 	{
@@ -597,7 +621,7 @@ TUniquePtr<IPCGAttributeAccessor> PCGAttributeAccessorHelpers::CreateAccessorWit
 			InMatchingAccessor->Get<AttributeType>(DefaultValue, FPCGAttributeAccessorKeysEntries(PCGInvalidEntryKey));
 			FPCGMetadataAttribute<AttributeType>* Attribute = Metadata->CreateAttribute<AttributeType>(AttributeName, DefaultValue, /*bAllowInterpolation=*/true, /*bOverrideParent=*/false);
 
-			return Attribute ? MakeUnique<FPCGAttributeAccessor<AttributeType>>(Attribute, Metadata) : nullptr;
+			return Attribute ? MakeUnique<FPCGAttributeAccessor<AttributeType>>(Attribute, Attribute->GetMetadataLayer()) : nullptr;
 		};
 
 		Result = PCGMetadataAttribute::CallbackWithRightType(InMatchingAccessor->GetUnderlyingType(), std::move(CreateAttributeAndAccessor));
diff --git a/Engine/Plugins/PCG/Source/PCG/Private/Metadata/Accessors/PCGAttributeAccessorKeys.cpp b/Engine/Plugins/PCG/Source/PCG/Private/Metadata/Accessors/PCGAttributeAccessorKeys.cpp
--- a/Engine/Plugins/PCG/Source/PCG/Private/Metadata/Accessors/PCGAttributeAccessorKeys.cpp
+++ b/Engine/Plugins/PCG/Source/PCG/Private/Metadata/Accessors/PCGAttributeAccessorKeys.cpp
@@ -21,17 +21,24 @@ FPCGAttributeAccessorKeysEntries::FPCGAttributeAccessorKeysEntries(PCGMetadataEn
 }
 
 FPCGAttributeAccessorKeysEntries::FPCGAttributeAccessorKeysEntries(const UPCGMetadata* Metadata, bool bAddDefaultValueIfEmpty)
+	: FPCGAttributeAccessorKeysEntries(Metadata ? Metadata->GetConstDefaultMetadataLayer() : nullptr, bAddDefaultValueIfEmpty)
+{}
+
+FPCGAttributeAccessorKeysEntries::FPCGAttributeAccessorKeysEntries(UPCGMetadata* Metadata, bool bAddDefaultValueIfEmpty)
+	: FPCGAttributeAccessorKeysEntries(Metadata ? Metadata->GetDefaultMetadataLayer() : nullptr, bAddDefaultValueIfEmpty)
+{}
+
+FPCGAttributeAccessorKeysEntries::FPCGAttributeAccessorKeysEntries(const FPCGMetadataLayer* Metadata, bool bAddDefaultValueIfEmpty)
 	: IPCGAttributeAccessorKeys(/*bInReadOnly=*/ true)
 {
 	InitializeFromMetadata(Metadata, bAddDefaultValueIfEmpty);
 	Entries = TArrayView<PCGMetadataEntryKey>(ExtractedEntries);
 }
 
-FPCGAttributeAccessorKeysEntries::FPCGAttributeAccessorKeysEntries(UPCGMetadata* Metadata, bool bAddDefaultValueIfEmpty)
+FPCGAttributeAccessorKeysEntries::FPCGAttributeAccessorKeysEntries(FPCGMetadataLayer* Metadata, bool bAddDefaultValueIfEmpty)
 	: IPCGAttributeAccessorKeys(/*bInReadOnly=*/ false)
 {
 	InitializeFromMetadata(Metadata, bAddDefaultValueIfEmpty);
-
 	Entries = TArrayView<PCGMetadataEntryKey>(ExtractedEntries);
 }
 
@@ -49,7 +56,7 @@ FPCGAttributeAccessorKeysEntries::FPCGAttributeAccessorKeysEntries(const TArrayV
 
 }
 
-void FPCGAttributeAccessorKeysEntries::InitializeFromMetadata(const UPCGMetadata* Metadata, bool bAddDefaultValueIfEmpty)
+void FPCGAttributeAccessorKeysEntries::InitializeFromMetadata(const FPCGMetadataLayer* Metadata, bool bAddDefaultValueIfEmpty)
 {
 	if (!Metadata)
 	{
diff --git a/Engine/Plugins/PCG/Source/PCG/Private/Metadata/PCGMetadata.cpp b/Engine/Plugins/PCG/Source/PCG/Private/Metadata/PCGMetadata.cpp
--- a/Engine/Plugins/PCG/Source/PCG/Private/Metadata/PCGMetadata.cpp
+++ b/Engine/Plugins/PCG/Source/PCG/Private/Metadata/PCGMetadata.cpp
@@ -3,145 +3,322 @@
 #include "Metadata/PCGMetadata.h"
 
 #include "PCGContext.h"
+#include "PCGCustomVersion.h"
 #include "PCGData.h"
 #include "PCGPoint.h"
-#include "Elements/Metadata/PCGMetadataElementCommon.h"
-#include "Helpers/PCGPropertyHelpers.h"
-#include "Metadata/PCGAttributePropertySelector.h"
-#include "Metadata/Accessors/PCGAttributeAccessorHelpers.h"
 
-#include "Algo/AnyOf.h"
 #include "Algo/Transform.h"
-#include "Async/ParallelFor.h"
 #include "Serialization/ArchiveCrc32.h"
 
 #include UE_INLINE_GENERATED_CPP_BY_NAME(PCGMetadata)
 
-void UPCGMetadata::Serialize(FArchive& InArchive)
-{
-	Super::Serialize(InArchive);
-
-	int32 NumAttributes = (InArchive.IsLoading() ? 0 : Attributes.Num());
-	// We need to keep track of the max attribute Id, since it won't necessary be equal to the number of attributes + 1.
-	int64 MaxAttributeId = -1;
+////////////////////////////
+// UPCGMetadata
+////////////////////////////
 
-	InArchive << NumAttributes;
-
-	if (InArchive.IsLoading())
+template <typename Func, typename ...Args>
+decltype(auto) UPCGMetadata::WithMetadataLayer(const FPCGMetadataLayerID& InMetadataLayerID, Func InFunc, Args&& ...InArgs)
+{
+	if (FPCGMetadataLayer* FoundMetadataLayer = GetMetadataLayer(InMetadataLayerID))
+	{
+		return Invoke(InFunc, *FoundMetadataLayer, std::forward<Args>(InArgs)...);
+	}
+	else
 	{
-		for (int32 AttributeIndex = 0; AttributeIndex < NumAttributes; ++AttributeIndex)
+		UE_LOG(LogPCG, Error, TEXT("Failed to find MetadataLayer with id %s"), *InMetadataLayerID.DebugName.ToString());
+		
+		using ReturnType = std::invoke_result_t<Func, FPCGMetadataLayer&, Args...>;
+		if constexpr (!std::is_same_v<ReturnType, void>)
 		{
-			FName AttributeName = NAME_None;
-			InArchive << AttributeName;
-
-			int32 AttributeTypeId = 0;
-			InArchive << AttributeTypeId;
-
-			FPCGMetadataAttributeBase* SerializedAttribute = PCGMetadataAttribute::AllocateEmptyAttributeFromType(static_cast<int16>(AttributeTypeId));
-			if (ensure(SerializedAttribute))
-			{
-				SerializedAttribute->Name = AttributeName;
-				SerializedAttribute->Serialize(this, InArchive);
-				Attributes.Add(AttributeName, SerializedAttribute);
-				MaxAttributeId = FMath::Max(SerializedAttribute->AttributeId, MaxAttributeId);
-			}
+			return ReturnType{};
 		}
 	}
+}
+
+template <typename Func>
+decltype(auto) UPCGMetadata::WithMetadataLayer_Lambda(const FPCGMetadataLayerID& InMetadataLayerID, Func InFunc)
+{
+	if (FPCGMetadataLayer* FoundMetadataLayer = GetMetadataLayer(InMetadataLayerID))
+	{
+		return InFunc(FoundMetadataLayer);
+	}
 	else
 	{
-		for (auto& AttributePair : Attributes)
+		UE_LOG(LogPCG, Error, TEXT("Failed to find layer with id %s"), *InMetadataLayerID.DebugName.ToString());
+		
+		using ReturnType = std::invoke_result_t<Func, FPCGMetadataLayer*>;
+		if constexpr (!std::is_same_v<ReturnType, void>)
 		{
-			InArchive << AttributePair.Key;
-			
-			int32 AttributeTypeId = AttributePair.Value->GetTypeId();
-			InArchive << AttributeTypeId;
+			return ReturnType{};
+		}
+	}
+}
 
-			AttributePair.Value->Serialize(this, InArchive);
+template <typename Func, typename ...Args>
+decltype(auto) UPCGMetadata::WithConstMetadataLayer(const FPCGMetadataLayerID& InMetadataLayerID, Func InFunc, Args&& ...InArgs) const
+{
+	if (const FPCGMetadataLayer* FoundMetadataLayer = GetConstMetadataLayer(InMetadataLayerID))
+	{
+		return Invoke(InFunc, *FoundMetadataLayer, std::forward<Args>(InArgs)...);
+	}
+	else
+	{
+		UE_LOG(LogPCG, Error, TEXT("Failed to find layer with id %s"), *InMetadataLayerID.DebugName.ToString());
+		
+		using ReturnType = typename std::invoke_result_t<Func, FPCGMetadataLayer&, Args...>;
+		if constexpr (!std::is_same_v<ReturnType, void>)
+		{
+			return ReturnType{};
 		}
 	}
+}
 
-	InArchive << ParentKeys;
+template <typename Func>
+void UPCGMetadata::ForEachValidUniqueConstMetadataLayer(Func InFunc) const
+{
+	auto Apply = [this, &InFunc](const FPCGMetadataLayerID& LayerID)
+	{
+		if (LayerID == DefaultLayer && !LayerID.IsDefault())
+		{
+			return;
+		}
+		
+		const FPCGMetadataLayer* CurrentMetadataLayer = MetadataLayers[LayerID].Get();
+		if (!CurrentMetadataLayer)
+		{
+			return;
+		}
+		
+		InFunc(LayerID, *CurrentMetadataLayer);
+	};
 
-	// Finally, initialize non-serialized members
-	if (InArchive.IsLoading())
+	// Always check the default first
+	Apply(PCGMetadataLayerID::Default);
+	
+	for (const auto& It : MetadataLayers)
 	{
-		// The next attribute id need to be bigger than the max attribute id of all attributes (or we could have collisions).
-		// Therefore by construction, it should never be less than the number of attributes (but can be greater).
-		NextAttributeId = MaxAttributeId + 1;
-		check(NextAttributeId >= Attributes.Num());
-		ItemKeyOffset = (Parent ? Parent->GetItemCountForChild() : 0);
+		if (It.Key.IsDefault())
+		{
+			continue;
+		}
+
+		Apply(It.Key);
 	}
 }
 
-void UPCGMetadata::AddToCrc(FArchiveCrc32& Ar, bool bFullDataCrc) const
+template <typename Func>
+void UPCGMetadata::ForEachValidUniqueMetadataLayer(Func InFunc)
+{	
+	auto Apply = [this, &InFunc](const FPCGMetadataLayerID& LayerID)
+	{
+		if (LayerID == DefaultLayer && !LayerID.IsDefault())
+		{
+			return;
+		}
+		
+		FPCGMetadataLayer* CurrentMetadataLayer = MetadataLayers[LayerID].Get();
+		if (!CurrentMetadataLayer)
+		{
+			return;
+		}
+		
+		InFunc(LayerID, *CurrentMetadataLayer);
+	};
+	
+	// Always check the default first
+	Apply(PCGMetadataLayerID::Default);
+	
+	for (const auto& It : MetadataLayers)
+	{
+		if (It.Key.IsDefault())
+		{
+			continue;
+		}
+
+		Apply(It.Key);
+	}
+}
+
+FPCGMetadataLayer* UPCGMetadata::GetMetadataLayer(const FPCGMetadataLayerID& InMetadataLayerID)
 {
-	const UPCGData* Data = Cast<UPCGData>(GetOuter());
-	check(Data);
+	return FindOrCreateMetadataLayer(InMetadataLayerID);
+}
 
-	// Sort attributes so we have a consistent processing path
-	TArray<FName> AttributeNames;
+const FPCGMetadataLayer* UPCGMetadata::GetConstMetadataLayer(const FPCGMetadataLayerID& InMetadataLayerID) const
+{
+	const TSharedPtr<FPCGMetadataLayer>* FoundMetadataLayer = MetadataLayers.Find(InMetadataLayerID);
+	return FoundMetadataLayer ? FoundMetadataLayer->Get() : nullptr;
+}
 
+FPCGMetadataLayer* UPCGMetadata::FindOrCreateMetadataLayer(const FPCGMetadataLayerID& InMetadataLayerID)
+{
+	TSharedPtr<FPCGMetadataLayer>* FoundMetadataLayer = MetadataLayers.Find(InMetadataLayerID);
+	if(!ensure(FoundMetadataLayer))
 	{
-		AttributeLock.ReadLock();
-		AttributeNames.Reserve(Attributes.Num());
+		return nullptr;
+	}
 
-		for (const TPair<FName, FPCGMetadataAttributeBase*>& Attribute : Attributes)
+	if (!FoundMetadataLayer->IsValid())
+	{
+		UE::TScopeLock<UE::FSpinLock> Lock(MetadataLayersSpinLock);
+		
+		FoundMetadataLayer = MetadataLayers.Find(InMetadataLayerID);
+		if (!FoundMetadataLayer->IsValid())
 		{
-			AttributeNames.Add(Attribute.Key);
+			CreateMetadataLayer_Unsafe(InMetadataLayerID);
 		}
-		AttributeLock.ReadUnlock();
+		
+	}
+
+	return FoundMetadataLayer->Get();
+}
+
+FPCGMetadataLayer* UPCGMetadata::CreateMetadataLayer_Unsafe(const FPCGMetadataLayerID& InMetadataLayerID)
+{
+	check(MetadataLayers.Contains(InMetadataLayerID) && !MetadataLayers[InMetadataLayerID].IsValid());
+	if (!InMetadataLayerID.IsDefault() && InMetadataLayerID == DefaultLayer)
+	{
+		// Nothing to do
+	}
+	else if (!MetadataLayers[InMetadataLayerID].IsValid())
+	{
+		MetadataLayers[InMetadataLayerID] = MakeShared<FPCGMetadataLayer>(this, InMetadataLayerID);
+	}
+	
+	return MetadataLayers[InMetadataLayerID].Get();
+}
+
+UPCGMetadata::UPCGMetadata(const FObjectInitializer& ObjectInitializer)
+	: Super(ObjectInitializer)
+{
+	// Always initialize the default level
+	SetupLayer(PCGMetadataLayerID::Default, /*bIsDefault=*/true);
+	CreateMetadataLayer_Unsafe(PCGMetadataLayerID::Default);
+}
 
-		Algo::Sort(AttributeNames, [this](const FName& A, const FName& B) { return A.LexicalLess(B); });
+void UPCGMetadata::SetupLayer(FPCGMetadataLayerID LayerID, bool bIsDefault)
+{
+	check(!MetadataLayers.Contains(LayerID));
+	if (bIsDefault && !LayerID.IsDefault())
+	{
+		check(MetadataLayers.Contains(PCGMetadataLayerID::Default) && MetadataLayers[PCGMetadataLayerID::Default].IsValid());
+		// If the default level is not Default, we need to change it
+		if (!DefaultLayer.IsDefault())
+		{
+			check(MetadataLayers.Contains(DefaultLayer) && MetadataLayers[DefaultLayer] == MetadataLayers[PCGMetadataLayerID::Default]);
+			MetadataLayers[DefaultLayer].Reset();
+		}
+		
+		DefaultLayer = LayerID;
+		// Map the default layer to the DefaultID layer and update layer to the not-default layerID.
+		MetadataLayers.Emplace(LayerID, MetadataLayers[PCGMetadataLayerID::Default])->LayerID = LayerID;
+	}
+	else
+	{
+		MetadataLayers.Emplace(LayerID);
 	}
+}
 
-	// Add attributes to CRC
-	for (FName AttributeName : AttributeNames)
+void UPCGMetadata::SetupLayersFromPCGDataType(const TSubclassOf<UPCGData>& PCGDataType)
+{
+	const UPCGData* CDO = *PCGDataType ? CastChecked<const UPCGData>(PCGDataType->GetDefaultObject()) : nullptr;
+	if (CDO)
 	{
-		Ar << AttributeName;
+		SetupLayersFromOtherMetadataIfNeeded(CDO->ConstMetadata());
+	}
+}
 
-		FPCGAttributePropertyInputSelector InputSource;
-		InputSource.SetAttributeName(AttributeName);
+void UPCGMetadata::Serialize(FArchive& InArchive)
+{
+	Super::Serialize(InArchive);
 
-		TUniquePtr<const IPCGAttributeAccessor> InputAccessor = PCGAttributeAccessorHelpers::CreateConstAccessor(Data, InputSource);
-		TUniquePtr<const IPCGAttributeAccessorKeys> InputKeys = PCGAttributeAccessorHelpers::CreateConstKeys(Data, InputSource);
+	InArchive.UsingCustomVersion(FPCGCustomVersion::GUID);
 
-		auto Callback = [&InputAccessor, &InputKeys, &Ar](auto&& Dummy)
+	TArray<FPCGMetadataLayerID> LayerIDs;
+	
+	if (InArchive.IsLoading())
+	{
+		if (InArchive.CustomVer(FPCGCustomVersion::GUID) < FPCGCustomVersion::MultiLevelMetadata)
+		{
+			const UPCGData* Outer = Cast<const UPCGData>(GetOuter());
+			FPCGMetadataLayer* DefaultMetadataLayer = GetDefaultMetadataLayer();
+			check(DefaultMetadataLayer);
+			DefaultMetadataLayer->Serialize(InArchive);
+		}
+		else
 		{
-			using AttributeType = std::decay_t<decltype(Dummy)>;
-			TArray<AttributeType> Values;
-			if constexpr (std::is_trivially_copyable_v<AttributeType>)
+			FPCGMetadataLayerID ArchiveDefaultLevel;
+			InArchive << ArchiveDefaultLevel;
+
+			if (ArchiveDefaultLevel != DefaultLayer)
 			{
-				Values.SetNumUninitialized(InputKeys->GetNum());
+				UE_LOG(LogPCG, Warning, TEXT("Mismatch between default metadata levels while loading. "
+								 "You should make sure to update the serialized metadata with the new default. "
+								 "Serialized level: %s ; Current level: %s"),
+								 *ArchiveDefaultLevel.DebugName.ToString(), *DefaultLayer.DebugName.ToString());
 			}
-			else
+			
+			InArchive << LayerIDs;
+			
+			for (const FPCGMetadataLayerID& LayerID : LayerIDs)
 			{
-				Values.SetNum(InputKeys->GetNum());
+				bool bIsValid = false;
+				InArchive << bIsValid;
+				
+				// If the layer is the default layer, nothing to do.
+				if (!LayerID.IsDefault())
+				{
+					FPCGMetadataLayer* NewMetadataLayer = FindOrCreateMetadataLayer(LayerID);
+					check(NewMetadataLayer);
+					if (bIsValid)
+					{
+						NewMetadataLayer->Serialize(InArchive);
+					}
+				}
 			}
-			
-			InputAccessor->GetRange<AttributeType>(Values, 0, *InputKeys);
+		}
+	}
+	else
+	{
+		MetadataLayers.GetKeys(LayerIDs);
+		// Sort the layers to always serialize in the same order.
+		LayerIDs.Sort();
+
+		InArchive << DefaultLayer;
+		InArchive << LayerIDs;
+		for (const FPCGMetadataLayerID& LayerID : LayerIDs)
+		{
+			bool bIsValid = MetadataLayers[LayerID].IsValid();
+			InArchive << bIsValid;
 
-			for (AttributeType& Value : Values)
+			// Don't serialize the default level
+			if (bIsValid &&  !LayerID.IsDefault())
 			{
-				// Add value to Crc
-				PCG::Private::Serialize(Ar, Value);
+				MetadataLayers[LayerID]->Serialize(InArchive);
 			}
-		};
-
-		PCGMetadataAttribute::CallbackWithRightType(InputAccessor->GetUnderlyingType(), Callback);
+		}
 	}
 }
 
-void UPCGMetadata::BeginDestroy()
+void UPCGMetadata::AddToCrc(FArchiveCrc32& Ar, bool bFullDataCrc) const
 {
-	FWriteScopeLock ScopeLock(AttributeLock);
-	for (TPair<FName, FPCGMetadataAttributeBase*>& AttributeEntry : Attributes)
+	const UPCGData* Data = Cast<UPCGData>(GetOuter());
+	check(Data);
+	
+	TArray<FPCGMetadataLayerID> LayerIDs;
+	MetadataLayers.GetKeys(LayerIDs);
+	LayerIDs.Sort();
+
+	for (const FPCGMetadataLayerID& LayerID : LayerIDs)
 	{
-		delete AttributeEntry.Value;
-		AttributeEntry.Value = nullptr;
-	}
-	Attributes.Reset();
+		bool bIsValid = MetadataLayers[LayerID].IsValid();
+		Ar << bIsValid;
 
-	Super::BeginDestroy();
+		if (bIsValid)
+		{
+			MetadataLayers[LayerID]->AddToCrc(Ar, Data, bFullDataCrc);
+		}
+	}
 }
 
 void UPCGMetadata::Initialize(const UPCGMetadata* InParent, bool bAddAttributesFromParent)
@@ -154,20 +331,30 @@ void UPCGMetadata::Initialize(const UPCGMetadata* InParent, bool bAddAttributesF
 
 void UPCGMetadata::InitializeWithAttributeFilter(const UPCGMetadata* InParent, const TSet<FName>& InFilteredAttributes, EPCGMetadataFilterMode InFilterMode, EPCGStringMatchingOperator InMatchOperator)
 {
-	if (Parent || Attributes.Num() != 0)
+	return InitializeWithAttributeFilter(InParent, FPCGAttributeIdentifier::TransformNameSet(InFilteredAttributes), InFilterMode, InMatchOperator);
+}
+
+void UPCGMetadata::InitializeWithAttributeFilter(const UPCGMetadata* InParent, const TSet<FPCGAttributeIdentifier>& InFilteredAttributes, EPCGMetadataFilterMode InFilterMode, EPCGStringMatchingOperator InMatchOperator)
+{
+	if (Parent)
 	{
-		// Already initialized; note that while that might be construed as a warning, there are legit cases where this is correct
+		// Already initialized; note that while that might be constructed as a warning, there are legit cases where this is correct
 		return;
 	}
 
-	Parent = ((InParent != this) ? InParent : nullptr);
-	ItemKeyOffset = Parent ? Parent->GetItemCountForChild() : 0;
-
-	// If we have been given an include list which is empty, then don't bother adding any attributes
-	const bool bSkipAddingAttributesFromParent = (InFilterMode == EPCGMetadataFilterMode::IncludeAttributes) && (InFilteredAttributes.Num() == 0);
-	if (!bSkipAddingAttributesFromParent)
+	if (InParent != this && InParent)
 	{
-		AddAttributesFiltered(InParent, InFilteredAttributes, InFilterMode, InMatchOperator);
+		SetupLayersFromOtherMetadataIfNeeded(InParent);
+		Parent = InParent;
+		
+		TMap<FPCGMetadataLayerID, TSet<FName>> FilteredAttributeSplit = FPCGAttributeIdentifier::SplitByLayer(InFilteredAttributes);
+		
+		InParent->ForEachValidUniqueConstMetadataLayer([this, &FilteredAttributeSplit, InFilterMode, InMatchOperator](const FPCGMetadataLayerID& LayerID, const FPCGMetadataLayer& OtherMetadataLayer)
+		{
+			TSet<FName>* FilteredAttributesPtr = FilteredAttributeSplit.Find(LayerID);
+			FPCGMetadataLayer* CurrentMetadataLayer = FindOrCreateMetadataLayer(LayerID);
+			CurrentMetadataLayer->InitializeWithAttributeFilter(&OtherMetadataLayer, (FilteredAttributesPtr ? *FilteredAttributesPtr : TSet<FName>{}), InFilterMode, InMatchOperator);
+		});
 	}
 }
 
@@ -177,147 +364,71 @@ void UPCGMetadata::InitializeAsCopy(const UPCGMetadata* InMetadataToCopy, const
 }
 
 void UPCGMetadata::InitializeAsCopyWithAttributeFilter(const UPCGMetadata* InMetadataToCopy, const TSet<FName>& InFilteredAttributes, EPCGMetadataFilterMode InFilterMode, const TArray<PCGMetadataEntryKey>* EntriesToCopy)
+{
+	InitializeAsCopyWithAttributeFilter(InMetadataToCopy, FPCGAttributeIdentifier::TransformNameSet(InFilteredAttributes), InFilterMode, EntriesToCopy);
+}
+
+void UPCGMetadata::InitializeAsCopyWithAttributeFilter(const UPCGMetadata* InMetadataToCopy, const TSet<FPCGAttributeIdentifier>& InFilteredAttributes, EPCGMetadataFilterMode InFilterMode, const TArray<PCGMetadataEntryKey>* EntriesToCopy)
 {
 	if (!InMetadataToCopy)
 	{
 		return;
 	}
-
+	
 	check(InMetadataToCopy);
-	if (Parent || Attributes.Num() != 0)
+	
+	if (Parent)
 	{
-		UE_LOG(LogPCG, Error, TEXT("Metadata has already been initialized or already contains attributes"));
+		UE_LOG(LogPCG, Error, TEXT("Metadata has already been initialized."));
 		return;
 	}
 
-	const bool bSkipAttributesInFilterList = (InFilterMode == EPCGMetadataFilterMode::ExcludeAttributes);
-
-	// If we have a partial copy, it will flatten the metadata, so we don't need a parent.
-	// Otherwise, we keep the parent hierarchy.
-	const bool bPartialCopy = EntriesToCopy && EntriesToCopy->Num() <= InMetadataToCopy->GetItemCountForChild();
-	TArray<PCGMetadataEntryKey> NewEntryKeys;
-	TArray<PCGMetadataValueKey> NewValueKeys;
-	if (bPartialCopy)
-	{
-		const int32 Count = EntriesToCopy->Num();
-		NewEntryKeys.SetNumUninitialized(Count);
-		NewValueKeys.SetNumUninitialized(Count);
-		ParentKeys.SetNumUninitialized(Count);
-		for (int32 j = 0; j < Count; ++j)
-		{
-			NewEntryKeys[j] = PCGMetadataEntryKey(j);
-			ParentKeys[j] = -1;
-		}
-
-		ItemKeyOffset = 0;
-	}
-	else
-	{
-		ParentKeys = InMetadataToCopy->ParentKeys;
-		ItemKeyOffset = InMetadataToCopy->ItemKeyOffset;
-		Parent = InMetadataToCopy->Parent;
-		OtherParents = InMetadataToCopy->OtherParents;
-	}
-
-	// Copy attributes
-	for (const TPair<FName, FPCGMetadataAttributeBase*>& OtherAttribute : InMetadataToCopy->Attributes)
+	SetupLayersFromOtherMetadataIfNeeded(InMetadataToCopy);
+		
+	TMap<FPCGMetadataLayerID, TSet<FName>> FilteredAttributeSplit = FPCGAttributeIdentifier::SplitByLayer(InFilteredAttributes);
+	InMetadataToCopy->ForEachValidUniqueConstMetadataLayer([this, &FilteredAttributeSplit, InFilterMode, &EntriesToCopy](const FPCGMetadataLayerID& LayerID, const FPCGMetadataLayer& OtherMetadataLayer)
 	{
-		const bool bAttributeInFilterList = InFilteredAttributes.Contains(OtherAttribute.Key);
-		const bool bSkipThisAttribute = (bSkipAttributesInFilterList == bAttributeInFilterList);
-
-		if (!bSkipThisAttribute)
-		{
-			// Don't copy entries if we have a partial copy, we will set them all after.
-			FPCGMetadataAttributeBase* Attribute = CopyAttribute(OtherAttribute.Value, OtherAttribute.Key, /*bKeepParent=*/false, /*bCopyEntries=*/!bPartialCopy, /*bCopyValues=*/true);
-
-			if (bPartialCopy && OtherAttribute.Value && Attribute)
-			{
-				OtherAttribute.Value->GetValueKeys(*EntriesToCopy, NewValueKeys);
-				Attribute->SetValuesFromValueKeys(NewEntryKeys, NewValueKeys);
-			}
-		}
-	}
+		TSet<FName>* FilteredAttributesPtr = FilteredAttributeSplit.Find(LayerID);
+		FPCGMetadataLayer* CurrentMetadataLayer = FindOrCreateMetadataLayer(LayerID);
+		CurrentMetadataLayer->InitializeAsCopyWithAttributeFilter(&OtherMetadataLayer, (FilteredAttributesPtr ? *FilteredAttributesPtr : TSet<FName>{}), InFilterMode, EntriesToCopy);
+	});
 }
 
 void UPCGMetadata::AddAttributesFiltered(const UPCGMetadata* InOther, const TSet<FName>& InFilteredAttributes, EPCGMetadataFilterMode InFilterMode, EPCGStringMatchingOperator InMatchOperator)
+{
+	AddAttributesFiltered(InOther, FPCGAttributeIdentifier::TransformNameSet(InFilteredAttributes), InFilterMode, InMatchOperator);
+}
+
+void UPCGMetadata::AddAttributesFiltered(const UPCGMetadata* InOther, const TSet<FPCGAttributeIdentifier>& InFilteredAttributes, EPCGMetadataFilterMode InFilterMode, EPCGStringMatchingOperator InMatchOperator)
 {
 	if (!InOther)
 	{
 		return;
 	}
+	
+	TMap<FPCGMetadataLayerID, TSet<FName>> FilteredAttributeSplit = FPCGAttributeIdentifier::SplitByLayer(InFilteredAttributes);
 
-	TArray<FString> InFilteredAttributesStrings;
-	if (InMatchOperator != EPCGStringMatchingOperator::Equal)
-	{
-		Algo::Transform(InFilteredAttributes, InFilteredAttributesStrings, [](const FName& InFilteredAttribute) {return InFilteredAttribute.ToString(); });
-	}
-
-	auto IsAttributeInFilterList = [&InFilteredAttributes, &InFilteredAttributesStrings, InMatchOperator](FName OtherAttributeName) -> bool
-	{
-		if (InMatchOperator == EPCGStringMatchingOperator::Equal)
-		{
-			return InFilteredAttributes.Contains(OtherAttributeName);
-		}
-		else
-		{
-			const FString OtherAttributeString = OtherAttributeName.ToString();
-			if (InMatchOperator == EPCGStringMatchingOperator::Substring)
-			{
-				return Algo::AnyOf(InFilteredAttributesStrings, [&OtherAttributeString](const FString& InAttribute) { return OtherAttributeString.Contains(InAttribute); });
-			}
-			else if (InMatchOperator == EPCGStringMatchingOperator::Matches)
-			{
-				return Algo::AnyOf(InFilteredAttributesStrings, [&OtherAttributeString](const FString& InAttribute) { return OtherAttributeString.MatchesWildcard(InAttribute); });
-			}
-			else
-			{
-				checkNoEntry();
-				return false;
-			}
-		}
-	};
-
-	bool bAttributeAdded = false;
-
-	for (const TPair<FName, FPCGMetadataAttributeBase*>& OtherAttribute : InOther->Attributes)
+	InOther->ForEachValidUniqueConstMetadataLayer([this, &FilteredAttributeSplit, InFilterMode, InMatchOperator, InOther](const FPCGMetadataLayerID& LayerID, const FPCGMetadataLayer& OtherMetadataLayer)
 	{
-		// Skip this attribute if it is in an exclude list, or if it is not in an include list
-		const bool bAttributeInFilterList = IsAttributeInFilterList(OtherAttribute.Key);
-		const bool bSkipAttributesInFilterList = InFilterMode == EPCGMetadataFilterMode::ExcludeAttributes;
-		const bool bSkipThisAttribute = bSkipAttributesInFilterList == bAttributeInFilterList;
-
-		if(bSkipThisAttribute || !OtherAttribute.Value)
-		{
-			continue;
-		}
-		else if (HasAttribute(OtherAttribute.Key))
+		TSet<FName>* FilteredAttributesPtr = FilteredAttributeSplit.Find(LayerID);
+		FPCGMetadataLayer* CurrentMetadataLayer = FindOrCreateMetadataLayer(LayerID);
+				
+		if (CurrentMetadataLayer->AddAttributesFiltered(&OtherMetadataLayer, (FilteredAttributesPtr ? *FilteredAttributesPtr : TSet<FName>{}), InFilterMode, InMatchOperator))
 		{
-			// If both the current attribute and the other attribute have the same type - nothing to do
-			// If the current attribute can be broadcasted to the other but not the other way around - change the type
-			// If none of this is true - do nothing
-			const FPCGMetadataAttributeBase* Attribute = GetConstAttribute(OtherAttribute.Key);
-			check(Attribute);
-
-			if(Attribute->GetTypeId() != OtherAttribute.Value->GetTypeId() && 
-				!PCG::Private::IsBroadcastable(OtherAttribute.Value->GetTypeId(), Attribute->GetTypeId()) &&
-				PCG::Private::IsBroadcastable(Attribute->GetTypeId(), OtherAttribute.Value->GetTypeId()))
+			if (InOther != Parent)
 			{
-				ChangeAttributeType(OtherAttribute.Key, OtherAttribute.Value->GetTypeId());
+				OtherParents.Add(InOther);
 			}
 		}
-		else if (CopyAttribute(OtherAttribute.Value, OtherAttribute.Key, /*bKeepParent=*/InOther == Parent, /*bCopyEntries=*/false, /*bCopyValues=*/false))
-		{
-			bAttributeAdded = true;
-		}
-	}
+	});
+}
 
-	if (InOther != Parent && bAttributeAdded)
-	{
-		OtherParents.Add(InOther);
-	}
+void UPCGMetadata::BP_AddAttribute(const UPCGMetadata* InOther, FName AttributeName)
+{
+	AddAttribute(InOther, AttributeName);
 }
 
-void UPCGMetadata::AddAttribute(const UPCGMetadata* InOther, FName AttributeName)
+void UPCGMetadata::AddAttribute(const UPCGMetadata* InOther, FPCGAttributeIdentifier AttributeName)
 {
 	if (!InOther || !InOther->HasAttribute(AttributeName) || HasAttribute(AttributeName))
 	{
@@ -345,20 +456,19 @@ void UPCGMetadata::CopyAttributes(const UPCGMetadata* InOther)
 		return;
 	}
 
-	for (const TPair<FName, FPCGMetadataAttributeBase*> OtherAttribute : InOther->Attributes)
+	InOther->ForEachValidUniqueConstMetadataLayer([this](const FPCGMetadataLayerID& LayerID, const FPCGMetadataLayer& OtherMetadataLayer)
 	{
-		if (HasAttribute(OtherAttribute.Key))
-		{
-			continue;
-		}
-		else
-		{
-			CopyAttribute(OtherAttribute.Value, OtherAttribute.Key, /*bKeepParent=*/false, /*bCopyEntries=*/true, /*bCopyValues=*/true);
-		}
-	}
+		FPCGMetadataLayer* CurrentMetadataLayer = FindOrCreateMetadataLayer(LayerID);
+		CurrentMetadataLayer->CopyAttributes(&OtherMetadataLayer);
+	});
+}
+
+void UPCGMetadata::BP_CopyAttribute(const UPCGMetadata* InOther, FName AttributeToCopy, FName NewAttributeName)
+{
+	CopyAttribute(InOther, AttributeToCopy, NewAttributeName); 
 }
 
-void UPCGMetadata::CopyAttribute(const UPCGMetadata* InOther, FName AttributeToCopy, FName NewAttributeName)
+void UPCGMetadata::CopyAttribute(const UPCGMetadata* InOther, FPCGAttributeIdentifier AttributeToCopy, FName NewAttributeName)
 {
 	if (!InOther)
 	{
@@ -379,7 +489,7 @@ void UPCGMetadata::CopyAttribute(const UPCGMetadata* InOther, FName AttributeToC
 		UE_LOG(LogPCG, Error, TEXT("Mismatch in copy attributes since the entries do not match"));
 		return;
 	}
-
+	
 	CopyAttribute(InOther->GetConstAttribute(AttributeToCopy), NewAttributeName, /*bKeepParent=*/false, /*bCopyEntries=*/true, /*bCopyValues=*/true);
 }
 
@@ -429,163 +539,69 @@ void UPCGMetadata::FlattenImpl()
 	TRACE_CPUPROFILER_EVENT_SCOPE(UPCGMetadata::FlattenImpl);
 	Modify();
 
-	const int32 NumEntries = GetItemCountForChild();
-
-	AttributeLock.WriteLock();
-	for (auto& AttributePair : Attributes)
-	{
-		FPCGMetadataAttributeBase* Attribute = AttributePair.Value;
-		check(Attribute);
-
-		// For all stored entries (from the root), we need to make sure that entries that should have a concrete value have it
-		// Optimization notes:
-		// - we could skip entries that existed prior to attribute existence, etc.
-		// - we could skip entries that have no parent, but that would require checking against the parent entries in the parent hierarchy
-		for (int64 EntryKey = 0; EntryKey < NumEntries; ++EntryKey)
-		{
-			// Get value using value inheritance as expected
-			PCGMetadataValueKey ValueKey = Attribute->GetValueKey(EntryKey);
-			if (ValueKey != PCGDefaultValueKey)
-			{
-				// Set concrete non-default value
-				Attribute->SetValueFromValueKey(EntryKey, ValueKey);
-			}
-		}
-
-		// Finally, flatten values
-		Attribute->Flatten();
-	}
-	AttributeLock.WriteUnlock();
+	ForEachValidUniqueMetadataLayer([](const FPCGMetadataLayerID& LayerID, FPCGMetadataLayer& MetadataLayer) { MetadataLayer.FlattenImpl();});
 
 	Parent = nullptr;
-	ParentKeys.Reset();
-	ParentKeys.Init(PCGInvalidEntryKey, NumEntries);
-	ItemKeyOffset = 0;
 }
 
 bool UPCGMetadata::FlattenAndCompress(const TArray<PCGMetadataEntryKey>& InEntryKeysToKeep)
 {
 	TRACE_CPUPROFILER_EVENT_SCOPE(UPCGMetadata::FlattenAndCompress);
-
-	// No keys or no parents, nothing to do
-	if (Attributes.IsEmpty())
-	{
-		return false;
-	}
-
-	Modify();
-
-	AttributeLock.WriteLock();
-	for (auto& AttributePair : Attributes)
-	{
-		FPCGMetadataAttributeBase* Attribute = AttributePair.Value;
-		check(Attribute);
-
-		Attribute->FlattenAndCompress(InEntryKeysToKeep);
-	}
-	AttributeLock.WriteUnlock();
+	
+	ForEachValidUniqueMetadataLayer([&InEntryKeysToKeep](const FPCGMetadataLayerID& LayerID, FPCGMetadataLayer& MetadataLayer) { MetadataLayer.FlattenAndCompress(InEntryKeysToKeep);});
 
 	Parent = nullptr;
-	ParentKeys.Reset();
-	ParentKeys.Init(PCGInvalidEntryKey, InEntryKeysToKeep.Num());
-	ItemKeyOffset = 0;
-
 	return true;
 }
 
-void UPCGMetadata::AddAttributeInternal(FName AttributeName, FPCGMetadataAttributeBase* Attribute)
+FPCGMetadataAttributeBase* UPCGMetadata::GetMutableAttribute(FPCGAttributeIdentifier AttributeName)
 {
-	// This call assumes we have a write lock on the attribute map.
-	Attributes.Add(AttributeName, Attribute);
+	return WithMetadataLayer(AttributeName.MetadataLayer, &FPCGMetadataLayer::GetMutableAttribute, AttributeName.Name);
 }
 
-void UPCGMetadata::RemoveAttributeInternal(FName AttributeName)
+const FPCGMetadataAttributeBase* UPCGMetadata::GetConstAttribute(FPCGAttributeIdentifier AttributeName) const
 {
-	Attributes.Remove(AttributeName);
+	return WithConstMetadataLayer(AttributeName.MetadataLayer, &FPCGMetadataLayer::GetConstAttribute, AttributeName.Name);
 }
 
-FPCGMetadataAttributeBase* UPCGMetadata::GetMutableAttribute(FName AttributeName)
+bool UPCGMetadata::BP_HasAttribute(FName AttributeName) const
 {
-	FPCGMetadataAttributeBase* Attribute = nullptr;
-
-	AttributeLock.ReadLock();
-	if (FPCGMetadataAttributeBase** FoundAttribute = Attributes.Find(AttributeName))
-	{
-		Attribute = *FoundAttribute;
-
-		// Also when accessing an attribute, notify the PCG Data owner that the latest attribute manipulated is this one.
-		SetLastCachedSelectorOnOwner(AttributeName);
-	}
-	AttributeLock.ReadUnlock();
-
-	return Attribute;
+	return HasAttribute(AttributeName);
 }
 
-const FPCGMetadataAttributeBase* UPCGMetadata::GetConstAttribute(FName AttributeName) const
+bool UPCGMetadata::HasAttribute(FPCGAttributeIdentifier AttributeName) const
 {
-	const FPCGMetadataAttributeBase* Attribute = nullptr;
-
-	AttributeLock.ReadLock();
-	if (const FPCGMetadataAttributeBase* const* FoundAttribute = Attributes.Find(AttributeName))
-	{
-		Attribute = *FoundAttribute;
-	}
-	AttributeLock.ReadUnlock();
-
-	return Attribute;
+	return WithConstMetadataLayer(AttributeName.MetadataLayer, &FPCGMetadataLayer::HasAttribute, AttributeName.Name);
 }
 
-const FPCGMetadataAttributeBase* UPCGMetadata::GetConstAttributeById(int32 InAttributeId) const
-{
-	const FPCGMetadataAttributeBase* Attribute = nullptr;
-
-	AttributeLock.ReadLock();
-	for (const auto& AttributePair : Attributes)
-	{
-		if (AttributePair.Value && AttributePair.Value->AttributeId == InAttributeId)
-		{
-			Attribute = AttributePair.Value;
-			break;
-		}
-	}
-	AttributeLock.ReadUnlock();
-
-	return Attribute;
-}
-
-bool UPCGMetadata::HasAttribute(FName AttributeName) const
-{
-	FReadScopeLock ScopeLock(AttributeLock);
-	return Attributes.Contains(AttributeName);
-}
-
-bool UPCGMetadata::HasCommonAttributes(const UPCGMetadata* InMetadata) const
+bool UPCGMetadata::HasCommonAttributes(const UPCGMetadata* InMetadata) const
 {
 	if (!InMetadata)
 	{
 		return false;
 	}
-
+	
 	bool bHasCommonAttribute = false;
-
-	AttributeLock.ReadLock();
-	for (const TPair<FName, FPCGMetadataAttributeBase*>& AttributePair : Attributes)
+	ForEachValidUniqueConstMetadataLayer([InMetadata, &bHasCommonAttribute](const FPCGMetadataLayerID& LayerID, const FPCGMetadataLayer& MetadataLayer)
 	{
-		if (InMetadata->HasAttribute(AttributePair.Key))
+		if (bHasCommonAttribute)
 		{
-			bHasCommonAttribute = true;
-			break;
+			return;
 		}
-	}
-	AttributeLock.ReadUnlock();
+		
+		const FPCGMetadataLayer* OtherMetadataLayer = InMetadata->GetConstMetadataLayer(LayerID);
+		bHasCommonAttribute = OtherMetadataLayer && MetadataLayer.HasCommonAttributes(OtherMetadataLayer);
+	});
 
 	return bHasCommonAttribute;
 }
 
 int32 UPCGMetadata::GetAttributeCount() const
 {
-	FReadScopeLock ScopeLock(AttributeLock);
-	return Attributes.Num();
+	int32 Count = 0;
+	ForEachValidUniqueConstMetadataLayer([&Count](const FPCGMetadataLayerID& LayerID, const FPCGMetadataLayer& MetadataLayer) { Count += MetadataLayer.GetAttributeCount(); });
+	
+	return Count;
 }
 
 void UPCGMetadata::GetAttributes(TArray<FName>& AttributeNames, TArray<EPCGMetadataTypes>& AttributeTypes) const
@@ -593,550 +609,218 @@ void UPCGMetadata::GetAttributes(TArray<FName>& AttributeNames, TArray<EPCGMetad
 	AttributeNames.Reset();
 	AttributeTypes.Reset();
 
-	AttributeLock.ReadLock();
-	for (const TPair<FName, FPCGMetadataAttributeBase*>& Attribute : Attributes)
-	{
-		check(Attribute.Value && Attribute.Value->Name == Attribute.Key);
-		AttributeNames.Add(Attribute.Key);
-
-		if (Attribute.Value->GetTypeId() < static_cast<uint16>(EPCGMetadataTypes::Unknown))
-		{
-			AttributeTypes.Add(static_cast<EPCGMetadataTypes>(Attribute.Value->GetTypeId()));
-		}
-		else
-		{
-			AttributeTypes.Add(EPCGMetadataTypes::Unknown);
-		}
-	}
-	AttributeLock.ReadUnlock();
-}
-
-FName UPCGMetadata::GetLatestAttributeNameOrNone() const
-{
-	FName LatestAttributeName = NAME_None;
-	int64 MaxAttributeId = -1;
-	
-	AttributeLock.ReadLock();
-	for (const TPair<FName, FPCGMetadataAttributeBase*>& It : Attributes)
-	{
-		if (It.Value && (It.Value->AttributeId > MaxAttributeId))
-		{
-			MaxAttributeId = It.Value->AttributeId;
-			LatestAttributeName = It.Key;
-		}
-	}
-	AttributeLock.ReadUnlock();
-
-	return LatestAttributeName;
+	ForEachValidUniqueConstMetadataLayer([&AttributeNames, &AttributeTypes](const FPCGMetadataLayerID& LayerID, const FPCGMetadataLayer& MetadataLayer) { MetadataLayer.GetAttributes(AttributeNames, AttributeTypes); });
 }
 
-bool UPCGMetadata::ParentHasAttribute(FName AttributeName) const
+void UPCGMetadata::GetAllAttributes(TArray<FPCGAttributeIdentifier>& AttributeNames, TArray<EPCGMetadataTypes>& AttributeTypes) const
 {
-	return Parent && Parent->HasAttribute(AttributeName);
-}
+	AttributeNames.Reset();
+	AttributeTypes.Reset();
 
-UPCGMetadata* UPCGMetadata::CreateInteger32Attribute(FName AttributeName, int32 DefaultValue, bool bAllowsInterpolation, bool bOverrideParent)
-{
-	CreateAttribute<int32>(AttributeName, DefaultValue, bAllowsInterpolation, bOverrideParent);
-	return this;
-}
+	TArray<FName> SubAttributeNames;
 
-UPCGMetadata* UPCGMetadata::CreateInteger64Attribute(FName AttributeName, int64 DefaultValue, bool bAllowsInterpolation, bool bOverrideParent)
-{
-	CreateAttribute<int64>(AttributeName, DefaultValue, bAllowsInterpolation, bOverrideParent);
-	return this;
+	ForEachValidUniqueConstMetadataLayer([&AttributeNames, &AttributeTypes, &SubAttributeNames](const FPCGMetadataLayerID& LayerID, const FPCGMetadataLayer& MetadataLayer)
+	{
+		SubAttributeNames.Empty();
+		MetadataLayer.GetAttributes(SubAttributeNames, AttributeTypes);
+		Algo::Transform(SubAttributeNames, AttributeNames, [&LayerID](const FName Name) { return FPCGAttributeIdentifier(Name, LayerID);});
+	});
 }
 
-UPCGMetadata* UPCGMetadata::CreateFloatAttribute(FName AttributeName, float DefaultValue, bool bAllowsInterpolation, bool bOverrideParent)
+FName UPCGMetadata::GetLatestAttributeNameOrNone() const
 {
-	CreateAttribute<float>(AttributeName, DefaultValue, bAllowsInterpolation, bOverrideParent);
-	return this;
+	const FPCGMetadataLayer* DefaultMetadataLayer = GetConstDefaultMetadataLayer();
+	check(DefaultMetadataLayer);
+	return DefaultMetadataLayer->GetLatestAttributeNameOrNone();
 }
 
-UPCGMetadata* UPCGMetadata::CreateDoubleAttribute(FName AttributeName, double DefaultValue, bool bAllowsInterpolation, bool bOverrideParent)
+bool UPCGMetadata::ParentHasAttribute(FPCGAttributeIdentifier AttributeName) const
 {
-	CreateAttribute<double>(AttributeName, DefaultValue, bAllowsInterpolation, bOverrideParent);
-	return this;
+	return Parent && Parent->HasAttribute(AttributeName);
 }
 
-UPCGMetadata* UPCGMetadata::CreateVectorAttribute(FName AttributeName, FVector DefaultValue, bool bAllowsInterpolation, bool bOverrideParent)
-{
-	CreateAttribute<FVector>(AttributeName, DefaultValue, bAllowsInterpolation, bOverrideParent);
-	return this;
+#define PCG_IMPL_CREATE_TYPED_ATTRIBUTE(FuncName, ArgType) \
+UPCGMetadata* UPCGMetadata::FuncName(FName AttributeName, ArgType DefaultValue, bool bAllowsInterpolation, bool bOverrideParent) \
+{ \
+	CreateAttribute<std::remove_const_t<std::remove_reference_t<ArgType>>>(AttributeName, DefaultValue, bAllowsInterpolation, bOverrideParent); \
+	return this; \
 }
 
-UPCGMetadata* UPCGMetadata::CreateVector4Attribute(FName AttributeName, FVector4 DefaultValue, bool bAllowsInterpolation, bool bOverrideParent)
-{
-	CreateAttribute<FVector4>(AttributeName, DefaultValue, bAllowsInterpolation, bOverrideParent);
-	return this;
-}
+PCG_IMPL_CREATE_TYPED_ATTRIBUTE(CreateInteger32Attribute, int32)
+PCG_IMPL_CREATE_TYPED_ATTRIBUTE(CreateInteger64Attribute, int64)
+PCG_IMPL_CREATE_TYPED_ATTRIBUTE(CreateFloatAttribute, float)
+PCG_IMPL_CREATE_TYPED_ATTRIBUTE(CreateDoubleAttribute, double)
+PCG_IMPL_CREATE_TYPED_ATTRIBUTE(CreateVectorAttribute, FVector)
+PCG_IMPL_CREATE_TYPED_ATTRIBUTE(CreateVector4Attribute, FVector4)
+PCG_IMPL_CREATE_TYPED_ATTRIBUTE(CreateVector2Attribute, FVector2D)
+PCG_IMPL_CREATE_TYPED_ATTRIBUTE(CreateRotatorAttribute, FRotator)
+PCG_IMPL_CREATE_TYPED_ATTRIBUTE(CreateQuatAttribute, FQuat)
+PCG_IMPL_CREATE_TYPED_ATTRIBUTE(CreateTransformAttribute, FTransform)
+PCG_IMPL_CREATE_TYPED_ATTRIBUTE(CreateStringAttribute, FString)
+PCG_IMPL_CREATE_TYPED_ATTRIBUTE(CreateNameAttribute, FName)
+PCG_IMPL_CREATE_TYPED_ATTRIBUTE(CreateBoolAttribute, bool)
+PCG_IMPL_CREATE_TYPED_ATTRIBUTE(CreateSoftObjectPathAttribute, const FSoftObjectPath&)
+PCG_IMPL_CREATE_TYPED_ATTRIBUTE(CreateSoftClassPathAttribute, const FSoftClassPath&)
 
-UPCGMetadata* UPCGMetadata::CreateVector2Attribute(FName AttributeName, FVector2D DefaultValue, bool bAllowsInterpolation, bool bOverrideParent)
-{
-	CreateAttribute<FVector2D>(AttributeName, DefaultValue, bAllowsInterpolation, bOverrideParent);
-	return this;
-}
+#undef PCG_IMPL_CREATE_TYPED_ATTRIBUTE
 
-UPCGMetadata* UPCGMetadata::CreateRotatorAttribute(FName AttributeName, FRotator DefaultValue, bool bAllowsInterpolation, bool bOverrideParent)
+bool UPCGMetadata::CreateAttributeFromProperty(FPCGAttributeIdentifier AttributeName, const UObject* Object, const FProperty* InProperty)
 {
-	CreateAttribute<FRotator>(AttributeName, DefaultValue, bAllowsInterpolation, bOverrideParent);
-	return this;
+	FPCGMetadataLayer* FoundMetadataLayer = FindOrCreateMetadataLayer(AttributeName.MetadataLayer);
+	return FoundMetadataLayer ? FoundMetadataLayer->CreateAttributeFromProperty(AttributeName.Name, Object, InProperty) : false;
 }
 
-UPCGMetadata* UPCGMetadata::CreateQuatAttribute(FName AttributeName, FQuat DefaultValue, bool bAllowsInterpolation, bool bOverrideParent)
+bool UPCGMetadata::CreateAttributeFromDataProperty(FPCGAttributeIdentifier AttributeName, const void* Data, const FProperty* InProperty)
 {
-	CreateAttribute<FQuat>(AttributeName, DefaultValue, bAllowsInterpolation, bOverrideParent);
-	return this;
+	FPCGMetadataLayer* FoundMetadataLayer = FindOrCreateMetadataLayer(AttributeName.MetadataLayer);
+	return FoundMetadataLayer ? FoundMetadataLayer->CreateAttributeFromDataProperty(AttributeName.Name, Data, InProperty) : false;
 }
 
-UPCGMetadata* UPCGMetadata::CreateTransformAttribute(FName AttributeName, FTransform DefaultValue, bool bAllowsInterpolation, bool bOverrideParent)
+bool UPCGMetadata::SetAttributeFromProperty(FPCGAttributeIdentifier AttributeName, PCGMetadataEntryKey& EntryKey, const UObject* Object, const FProperty* InProperty, bool bCreate)
 {
-	CreateAttribute<FTransform>(AttributeName, DefaultValue, bAllowsInterpolation, bOverrideParent);
-	return this;
+	FPCGMetadataLayer* FoundMetadataLayer = FindOrCreateMetadataLayer(AttributeName.MetadataLayer);
+	return FoundMetadataLayer ? FoundMetadataLayer->SetAttributeFromProperty(AttributeName.Name, EntryKey, Object, InProperty, bCreate) : false;
 }
 
-UPCGMetadata* UPCGMetadata::CreateStringAttribute(FName AttributeName, FString DefaultValue, bool bAllowsInterpolation, bool bOverrideParent)
+bool UPCGMetadata::SetAttributeFromDataProperty(FPCGAttributeIdentifier AttributeName, PCGMetadataEntryKey& EntryKey, const void* Data, const FProperty* InProperty, bool bCreate)
 {
-	CreateAttribute<FString>(AttributeName, DefaultValue, bAllowsInterpolation, bOverrideParent);
-	return this;
+	FPCGMetadataLayer* FoundMetadataLayer = FindOrCreateMetadataLayer(AttributeName.MetadataLayer);
+	return FoundMetadataLayer ? FoundMetadataLayer->SetAttributeFromDataProperty(AttributeName.Name, EntryKey, Data, InProperty, bCreate) : false;
 }
 
-UPCGMetadata* UPCGMetadata::CreateNameAttribute(FName AttributeName, FName DefaultValue, bool bAllowsInterpolation, bool bOverrideParent)
+bool UPCGMetadata::BP_CopyExistingAttribute(FName AttributeToCopy, FName NewAttributeName, bool bKeepParent)
 {
-	CreateAttribute<FName>(AttributeName, DefaultValue, bAllowsInterpolation, bOverrideParent);
-	return this;
+	return CopyExistingAttribute(AttributeToCopy, NewAttributeName, bKeepParent);
 }
 
-UPCGMetadata* UPCGMetadata::CreateBoolAttribute(FName AttributeName, bool DefaultValue, bool bAllowsInterpolation, bool bOverrideParent)
+bool UPCGMetadata::CopyExistingAttribute(FPCGAttributeIdentifier AttributeToCopy, FName NewAttributeName, bool bKeepParent)
 {
-	CreateAttribute<bool>(AttributeName, DefaultValue, bAllowsInterpolation, bOverrideParent);
-	return this;
+	return CopyAttribute(AttributeToCopy, NewAttributeName, bKeepParent, /*bCopyEntries=*/true, /*bCopyValues=*/true) != nullptr;
 }
 
-UPCGMetadata* UPCGMetadata::CreateSoftObjectPathAttribute(FName AttributeName, const FSoftObjectPath& DefaultValue, bool bAllowsInterpolation, bool bOverrideParent)
+FPCGMetadataAttributeBase* UPCGMetadata::CopyAttribute(FPCGAttributeIdentifier AttributeToCopy, FName NewAttributeName, bool bKeepParent, bool bCopyEntries, bool bCopyValues)
 {
-	CreateAttribute<FSoftObjectPath>(AttributeName, DefaultValue, bAllowsInterpolation, bOverrideParent);
-	return this;
+	return WithMetadataLayer_Lambda(AttributeToCopy.MetadataLayer, [&](FPCGMetadataLayer* MetadataLayer) { return MetadataLayer->CopyAttribute(AttributeToCopy.Name, NewAttributeName, bKeepParent, bCopyEntries, bCopyValues); });
 }
 
-UPCGMetadata* UPCGMetadata::CreateSoftClassPathAttribute(FName AttributeName, const FSoftClassPath& DefaultValue, bool bAllowsInterpolation, bool bOverrideParent)
+FPCGMetadataAttributeBase* UPCGMetadata::CopyAttribute(const FPCGMetadataAttributeBase* OriginalAttribute, FPCGAttributeIdentifier NewAttributeName, bool bKeepParent, bool bCopyEntries, bool bCopyValues)
 {
-	CreateAttribute<FSoftClassPath>(AttributeName, DefaultValue, bAllowsInterpolation, bOverrideParent);
-	return this;
-}
+	check(OriginalAttribute);
+	const FPCGMetadataLayer* OriginalMetadataLayer = OriginalAttribute->GetMetadataLayer();
+	check(OriginalMetadataLayer);
+	const UPCGMetadata* RootMetadata = GetRoot();
+	check((RootMetadata && RootMetadata->MetadataLayers.Contains(OriginalMetadataLayer->LayerID) && RootMetadata->MetadataLayers[OriginalMetadataLayer->LayerID].Get() == OriginalMetadataLayer->GetRoot()) || !bKeepParent);
 
-namespace PCGMetadata
-{
-	template<typename DataType>
-	bool CreateAttributeFromPropertyHelper(UPCGMetadata* Metadata, FName AttributeName, const DataType* DataPtr, const FProperty* InProperty)
+	FPCGMetadataAttributeBase* Attribute = WithMetadataLayer_Lambda(NewAttributeName.MetadataLayer, [&](FPCGMetadataLayer* MetadataLayer) -> FPCGMetadataAttributeBase*
 	{
-		if (!InProperty || !DataPtr || !Metadata)
+		// TODO: Maybe the operation needs to be handled differently? 
+		if (OriginalMetadataLayer->LayerID != MetadataLayer->LayerID)
 		{
-			return false;
+			// UE_LOG(LogPCG, Error, TEXT("[Metadata - CopyAttribute] Can't copy across levels."));
+			// return nullptr;
 		}
 
-		auto CreateAttribute = [AttributeName, Metadata](auto&& PropertyValue) -> bool
-		{
-			using PropertyType = std::decay_t<decltype(PropertyValue)>;
-			FPCGMetadataAttributeBase* BaseAttribute = Metadata->FindOrCreateAttribute<PropertyType>(AttributeName, PropertyValue, /*bAllowsInterpolation=*/false, /*bOverrideParent=*/false, /*bOverwriteIfTypeMismatch=*/true);
-			
-			return (BaseAttribute != nullptr);
-		};
-
-		return PCGPropertyHelpers::GetPropertyValueWithCallback(DataPtr, InProperty, CreateAttribute);
-	}
-
-	template<typename DataType>
-	bool SetAttributeFromPropertyHelper(UPCGMetadata* Metadata, FName AttributeName, PCGMetadataEntryKey& EntryKey, const DataType* DataPtr, const FProperty* InProperty, bool bCreate)
+		return MetadataLayer->CopyAttribute(OriginalAttribute, NewAttributeName.Name, bKeepParent, bCopyEntries, bCopyValues);
+	});
+	
+	if (!Attribute)
 	{
-		if (!InProperty || !DataPtr || !Metadata)
-		{
-			return false;
-		}
-
-		// Check if an attribute already exists or not if we ask to create a new one
-		if (!bCreate && !Metadata->HasAttribute(AttributeName))
-		{
-			return false;
-		}
-
-		auto CreateAttributeAndSet = [AttributeName, Metadata, bCreate, &EntryKey](auto&& PropertyValue) -> bool
-		{
-			using PropertyType = std::remove_const_t<std::remove_reference_t<decltype(PropertyValue)>>;
-
-			FPCGMetadataAttributeBase* BaseAttribute = Metadata->GetMutableAttribute(AttributeName);
-
-			if (!BaseAttribute && bCreate)
-			{
-				// Interpolation is disabled and no parent override.
-				BaseAttribute = Metadata->CreateAttribute<PropertyType>(AttributeName, PropertyValue, false, false);
-			}
-
-			if (!BaseAttribute)
-			{
-				return false;
-			}
-
-			// Allow to set the value if both type matches or if we can construct AttributeType from PropertyType.
-			return PCGMetadataAttribute::CallbackWithRightType(BaseAttribute->GetTypeId(), [&EntryKey, &PropertyValue, BaseAttribute, Metadata](auto AttributeValue) -> bool
-				{
-					using AttributeType = decltype(AttributeValue);
-					FPCGMetadataAttribute<AttributeType>* Attribute = static_cast<FPCGMetadataAttribute<AttributeType>*>(BaseAttribute);
-
-					// Special cased because FSoftObjectPath currently has a deprecated constructor from FName which generates compile warnings.
-					constexpr bool bAssigningNameToSoftObjectPath = std::is_same_v<AttributeType, FSoftObjectPath> && std::is_same_v<PropertyType, FName>;
-
-					if constexpr (std::is_same_v<AttributeType, PropertyType>)
-					{
-						Metadata->InitializeOnSet(EntryKey);
-						Attribute->SetValue(EntryKey, PropertyValue);
-						return true;
-					}
-					else if constexpr (std::is_constructible_v<AttributeType, PropertyType> && !bAssigningNameToSoftObjectPath)
-					{
-						Metadata->InitializeOnSet(EntryKey);
-						Attribute->SetValue(EntryKey, AttributeType(PropertyValue));
-						return true;
-					}
-					else
-					{
-						return false;
-					}
-				});
-		};
-
-		return PCGPropertyHelpers::GetPropertyValueWithCallback(DataPtr, InProperty, CreateAttributeAndSet);
+		UE_LOG(LogPCG, Error, TEXT("[Metadata - CopyAttribute] Metadata level not existing in current metadata or copy failed."));
 	}
-}
 
-bool UPCGMetadata::CreateAttributeFromProperty(FName AttributeName, const UObject* Object, const FProperty* InProperty)
-{
-	return PCGMetadata::CreateAttributeFromPropertyHelper<UObject>(this, AttributeName, Object, InProperty);
-}
-
-bool UPCGMetadata::CreateAttributeFromDataProperty(FName AttributeName, const void* Data, const FProperty* InProperty)
-{
-	return PCGMetadata::CreateAttributeFromPropertyHelper<void>(this, AttributeName, Data, InProperty);
-}
-
-bool UPCGMetadata::SetAttributeFromProperty(FName AttributeName, PCGMetadataEntryKey& EntryKey, const UObject* Object, const FProperty* InProperty, bool bCreate)
-{
-	return PCGMetadata::SetAttributeFromPropertyHelper<UObject>(this, AttributeName, EntryKey, Object, InProperty, bCreate);
-}
-
-bool UPCGMetadata::SetAttributeFromDataProperty(FName AttributeName, PCGMetadataEntryKey& EntryKey, const void* Data, const FProperty* InProperty, bool bCreate)
-{
-	return PCGMetadata::SetAttributeFromPropertyHelper<void>(this, AttributeName, EntryKey, Data, InProperty, bCreate);
+	return Attribute;
 }
 
-bool UPCGMetadata::CopyExistingAttribute(FName AttributeToCopy, FName NewAttributeName, bool bKeepParent)
+bool UPCGMetadata::BP_RenameAttribute(FName AttributeToRename, FName NewAttributeName)
 {
-	return CopyAttribute(AttributeToCopy, NewAttributeName, bKeepParent, /*bCopyEntries=*/true, /*bCopyValues=*/true) != nullptr;
+	return RenameAttribute(AttributeToRename, NewAttributeName);
 }
 
-FPCGMetadataAttributeBase* UPCGMetadata::CopyAttribute(FName AttributeToCopy, FName NewAttributeName, bool bKeepParent, bool bCopyEntries, bool bCopyValues)
+bool UPCGMetadata::RenameAttribute(FPCGAttributeIdentifier AttributeToRename, FName NewAttributeName)
 {
-	const FPCGMetadataAttributeBase* OriginalAttribute = nullptr;
-
-	AttributeLock.ReadLock();
-	if (FPCGMetadataAttributeBase** AttributeFound = Attributes.Find(AttributeToCopy))
+	if (FPCGMetadataLayer* FoundMetadataLayer = GetMetadataLayer(AttributeToRename.MetadataLayer))
 	{
-		OriginalAttribute = *AttributeFound;
+		return FoundMetadataLayer->RenameAttribute(AttributeToRename.Name, NewAttributeName);
 	}
-	AttributeLock.ReadUnlock();
-
-	if (!OriginalAttribute && Parent)
-	{
-		OriginalAttribute = Parent->GetConstAttribute(AttributeToCopy);
-	}
-
-	if (!OriginalAttribute)
+	else
 	{
-		UE_LOG(LogPCG, Warning, TEXT("Attribute %s does not exist, therefore cannot be copied"), *AttributeToCopy.ToString());
-		return nullptr;
+		return false;
 	}
-
-	return CopyAttribute(OriginalAttribute, NewAttributeName, bKeepParent, bCopyEntries, bCopyValues);
 }
 
-FPCGMetadataAttributeBase* UPCGMetadata::CopyAttribute(const FPCGMetadataAttributeBase* OriginalAttribute, FName NewAttributeName, bool bKeepParent, bool bCopyEntries, bool bCopyValues)
+void UPCGMetadata::BP_ClearAttribute(FName AttributeToClear)
 {
-	check(OriginalAttribute);
-	check(OriginalAttribute->GetMetadata()->GetRoot() == GetRoot() || !bKeepParent);
-	FPCGMetadataAttributeBase* NewAttribute = OriginalAttribute->Copy(NewAttributeName, this, bKeepParent, bCopyEntries, bCopyValues);
-
-	if (NewAttribute)
-	{
-		AttributeLock.WriteLock();
-		NewAttribute->AttributeId = NextAttributeId++;
-		AddAttributeInternal(NewAttributeName, NewAttribute);
-
-		// Also when creating an attribute, notify the PCG Data owner that the latest attribute manipulated is this one.
-		SetLastCachedSelectorOnOwner(NewAttributeName);
-
-		AttributeLock.WriteUnlock();
-	}
-
-	return NewAttribute;
+	ClearAttribute(AttributeToClear);
 }
 
-bool UPCGMetadata::RenameAttribute(FName AttributeToRename, FName NewAttributeName)
+void UPCGMetadata::ClearAttribute(FPCGAttributeIdentifier AttributeToClear)
 {
-	if (!FPCGMetadataAttributeBase::IsValidName(NewAttributeName))
-	{
-		UE_LOG(LogPCG, Error, TEXT("New attribute name %s is not valid"), *NewAttributeName.ToString());
-		return false;
-	}
-
-	bool bRenamed = false;
-	AttributeLock.WriteLock();
-	if (FPCGMetadataAttributeBase** AttributeFound = Attributes.Find(AttributeToRename))
-	{
-		FPCGMetadataAttributeBase* Attribute = *AttributeFound;
-		RemoveAttributeInternal(AttributeToRename);
-		Attribute->Name = NewAttributeName;
-		AddAttributeInternal(NewAttributeName, Attribute);
-
-		// Also when renaming an attribute, notify the PCG Data owner that the latest attribute manipulated is this one.
-		SetLastCachedSelectorOnOwner(NewAttributeName);
-		
-		bRenamed = true;
-	}
-	AttributeLock.WriteUnlock();
-
-	if (!bRenamed)
-	{
-		UE_LOG(LogPCG, Warning, TEXT("Attribute %s does not exist and therefore cannot be renamed"), *AttributeToRename.ToString());
-	}
-
-	return bRenamed;
+	return WithMetadataLayer(AttributeToClear.MetadataLayer, &FPCGMetadataLayer::ClearAttribute, AttributeToClear.Name);
 }
 
-void UPCGMetadata::ClearAttribute(FName AttributeToClear)
+void UPCGMetadata::BP_DeleteAttribute(FName AttributeToDelete)
 {
-	FPCGMetadataAttributeBase* Attribute = nullptr;
-
-	AttributeLock.ReadLock();
-	if (FPCGMetadataAttributeBase** AttributeFound = Attributes.Find(AttributeToClear))
-	{
-		Attribute = *AttributeFound;
-	}
-	AttributeLock.ReadUnlock();
-
-	// If the attribute exists, then we can lose all the entries
-	// If it doesn't but it exists in the parent hierarchy, then we must create a new attribute.
-	if (Attribute)
-	{
-		Attribute->ClearEntries();
-	}
+	DeleteAttribute(AttributeToDelete);
 }
 
-void UPCGMetadata::DeleteAttribute(FName AttributeToDelete)
+void UPCGMetadata::DeleteAttribute(FPCGAttributeIdentifier AttributeToDelete)
 {
-	FPCGMetadataAttributeBase* Attribute = nullptr;
-
-	// If it's a local attribute, then just delete it
-	AttributeLock.WriteLock();
-	if (FPCGMetadataAttributeBase** AttributeFound = Attributes.Find(AttributeToDelete))
-	{
-		Attribute = *AttributeFound;
-		RemoveAttributeInternal(AttributeToDelete);
-	}
-	AttributeLock.WriteUnlock();
-
-	if (Attribute)
-	{
-		delete Attribute;
-	}
-	else
-	{
-		UE_LOG(LogPCG, Verbose, TEXT("Attribute %s does not exist and therefore cannot be deleted"), *AttributeToDelete.ToString());
-	}
+	return WithMetadataLayer(AttributeToDelete.MetadataLayer, &FPCGMetadataLayer::DeleteAttribute, AttributeToDelete.Name);
 }
 
-bool UPCGMetadata::ChangeAttributeType(FName AttributeName, int16 AttributeNewType)
+bool UPCGMetadata::ChangeAttributeType(FPCGAttributeIdentifier AttributeName, int16 AttributeNewType)
 {
-	FPCGMetadataAttributeBase* Attribute = GetMutableAttribute(AttributeName);
-
-	if (!Attribute)
-	{
-		UE_LOG(LogPCG, Error, TEXT("Attribute '%s' does not exist and therefore cannot change its type"), *AttributeName.ToString());
-		return false;
-	}
-
-	if (Attribute->GetTypeId() == AttributeNewType)
-	{
-		// Nothing to do, attribute is already the type we want
-		return true;
-	}
-
-	if (FPCGMetadataAttributeBase* NewAttribute = Attribute->CopyToAnotherType(AttributeNewType))
-	{
-		NewAttribute->AttributeId = Attribute->AttributeId;
-
-		AttributeLock.WriteLock();
-		RemoveAttributeInternal(AttributeName);
-		AddAttributeInternal(AttributeName, NewAttribute);
-		AttributeLock.WriteUnlock();
-
-		delete Attribute;
-		Attribute = nullptr;
-	}
-
-	return true;
+	return WithMetadataLayer(AttributeName.MetadataLayer, &FPCGMetadataLayer::ChangeAttributeType, AttributeName.Name, AttributeNewType);
 }
 
 int64 UPCGMetadata::GetItemCountForChild() const
 {
-	FReadScopeLock ScopeLock(ItemLock);
-	return ParentKeys.Num() + ItemKeyOffset;
+	return WithConstMetadataLayer(PCGMetadataLayerID::Default, &FPCGMetadataLayer::GetItemCountForChild);
 }
 
 int64 UPCGMetadata::GetLocalItemCount() const
 {
-	FReadScopeLock ScopeLock(ItemLock);
-	return ParentKeys.Num();
+	return WithConstMetadataLayer(PCGMetadataLayerID::Default, &FPCGMetadataLayer::GetLocalItemCount);
 }
 
 int64 UPCGMetadata::AddEntry(int64 ParentEntry)
 {
-	FWriteScopeLock ScopeLock(ItemLock);
-	return ParentKeys.Add(ParentEntry) + ItemKeyOffset;
+	return WithMetadataLayer(PCGMetadataLayerID::Default, &FPCGMetadataLayer::AddEntry, ParentEntry);
 }
 
 TArray<int64> UPCGMetadata::AddEntries(TArrayView<const int64> ParentEntryKeys)
 {
-	TArray<int64> Result;
-	Result.Reserve(ParentEntryKeys.Num());
-
-	FWriteScopeLock ScopeLock(ItemLock);
-	ParentKeys.Reserve(ParentKeys.Num() + ParentEntryKeys.Num());
-	for (const int64 ParentEntry : ParentEntryKeys)
-	{
-		Result.Add(ParentKeys.Add(ParentEntry) + ItemKeyOffset);
-	}
-
-	return Result;
+	return WithMetadataLayer(PCGMetadataLayerID::Default, &FPCGMetadataLayer::AddEntries, ParentEntryKeys);
 }
 
 void UPCGMetadata::AddEntriesInPlace(TArrayView<int64*> ParentEntryKeys)
 {
-	FWriteScopeLock ScopeLock(ItemLock);
-	ParentKeys.Reserve(ParentKeys.Num() + ParentEntryKeys.Num());
-	for (int64* ParentEntry : ParentEntryKeys)
-	{
-		*ParentEntry = ParentKeys.Add(*ParentEntry) + ItemKeyOffset;
-	}
+	return WithMetadataLayer(PCGMetadataLayerID::Default, &FPCGMetadataLayer::AddEntriesInPlace, ParentEntryKeys);
 }
 
 int64 UPCGMetadata::AddEntryPlaceholder()
 {
-	FReadScopeLock ScopeLock(ItemLock);
-	return ParentKeys.Num() + DelayedEntriesIndex.IncrementExchange() + ItemKeyOffset;
+	return WithMetadataLayer(PCGMetadataLayerID::Default, &FPCGMetadataLayer::AddEntryPlaceholder);
 }
 
 void UPCGMetadata::AddDelayedEntries(const TArray<TTuple<int64, int64>>& AllEntries)
 {
-	FWriteScopeLock ScopeLock(ItemLock);
-	ParentKeys.AddUninitialized(AllEntries.Num());
-	for (const TTuple<int64, int64>& Entry : AllEntries)
-	{
-		int64 Index = Entry.Get<0>() - ItemKeyOffset;
-		check(Index < ParentKeys.Num());
-		ParentKeys[Index] = Entry.Get<1>();
-	}
-
-	DelayedEntriesIndex.Exchange(0);
+	return WithMetadataLayer(PCGMetadataLayerID::Default, &FPCGMetadataLayer::AddDelayedEntries, AllEntries);
 }
 
 bool UPCGMetadata::InitializeOnSet(PCGMetadataEntryKey& InOutKey, PCGMetadataEntryKey InParentKeyA, const UPCGMetadata* InParentMetadataA, PCGMetadataEntryKey InParentKeyB, const UPCGMetadata* InParentMetadataB)
 {
-	if (InOutKey == PCGInvalidEntryKey)
-	{
-		if (InParentKeyA != PCGInvalidEntryKey && Parent == InParentMetadataA)
-		{
-			InOutKey = AddEntry(InParentKeyA);
-			return true;
-		}
-		else if (InParentKeyB != PCGInvalidEntryKey && Parent == InParentMetadataB)
-		{
-			InOutKey = AddEntry(InParentKeyB);
-			return true;
-		}
-		else
-		{
-			InOutKey = AddEntry();
-			return false;
-		}
-	}
-	else if(InOutKey < ItemKeyOffset)
-	{
-		InOutKey = AddEntry(InOutKey);
-		return false;
-	}
-	else
-	{
-		return false;
-	}
+	return WithMetadataLayer(PCGMetadataLayerID::Default, &FPCGMetadataLayer::InitializeOnSet, InOutKey, InParentKeyA, InParentMetadataA ? InParentMetadataA->GetConstMetadataLayer(PCGMetadataLayerID::Default) : nullptr, InParentKeyB, InParentMetadataB ? InParentMetadataB->GetConstMetadataLayer(PCGMetadataLayerID::Default) : nullptr);
 }
 
 PCGMetadataEntryKey UPCGMetadata::GetParentKey(PCGMetadataEntryKey LocalItemKey) const
 {
-	if (LocalItemKey < ItemKeyOffset)
-	{
-		// Key is already in parent referential
-		return LocalItemKey;
-	}
-	else
-	{
-		FReadScopeLock ScopeLock(ItemLock);
-		if (LocalItemKey - ItemKeyOffset < ParentKeys.Num())
-		{
-			return ParentKeys[LocalItemKey - ItemKeyOffset];
-		}
-		else
-		{
-			UE_LOG(LogPCG, Warning, TEXT("Invalid metadata key - check for entry key not properly initialized"));
-			return PCGInvalidEntryKey;
-		}
-	}
+	return WithConstMetadataLayer(PCGMetadataLayerID::Default, &FPCGMetadataLayer::GetParentKey, LocalItemKey);
 }
 
 void UPCGMetadata::GetParentKeys(TArrayView<PCGMetadataEntryKey> LocalItemKeys, const TBitArray<>* Mask) const
 {
-	auto GetParentKey_Unsafe = [this](PCGMetadataEntryKey& LocalItemKey) -> void
-	{
-		if (LocalItemKey < ItemKeyOffset)
-		{
-			// Key is already in parent referential
-			return;
-		}
-		else if (LocalItemKey - ItemKeyOffset < ParentKeys.Num())
-		{
-			LocalItemKey = ParentKeys[LocalItemKey - ItemKeyOffset];
-		}
-		else
-		{
-			UE_LOG(LogPCG, Warning, TEXT("Invalid metadata key - check for entry key not properly initialized"));
-			LocalItemKey = PCGInvalidEntryKey;
-		}
-	};
-
-	FReadScopeLock ScopeLock(ItemLock);
-	if (Mask && ensure(LocalItemKeys.Num() == Mask->Num()))
-	{
-		for (TConstSetBitIterator<> It(*Mask); It; ++It)
-		{
-			GetParentKey_Unsafe(LocalItemKeys[It.GetIndex()]);
-		}
-	}
-	else
-	{
-		for (PCGMetadataEntryKey& LocalItemKey : LocalItemKeys)
-		{
-			GetParentKey_Unsafe(LocalItemKey);
-		}
-	}
+	return WithConstMetadataLayer(PCGMetadataLayerID::Default, &FPCGMetadataLayer::GetParentKeys, LocalItemKeys, Mask);
 }
 
 void UPCGMetadata::MergePointAttributes(const FPCGPoint& InPointA, const FPCGPoint& InPointB, FPCGPoint& OutPoint, EPCGMetadataOp Op)
@@ -1156,122 +840,33 @@ void UPCGMetadata::MergeAttributes(PCGMetadataEntryKey InKeyA, const UPCGMetadat
 
 void UPCGMetadata::MergeAttributesSubset(PCGMetadataEntryKey InKeyA, const UPCGMetadata* InMetadataA, const UPCGMetadata* InMetadataSubsetA, PCGMetadataEntryKey InKeyB, const UPCGMetadata* InMetadataB, const UPCGMetadata* InMetadataSubsetB, PCGMetadataEntryKey& OutKey, EPCGMetadataOp Op)
 {
-	//TRACE_CPUPROFILER_EVENT_SCOPE(UPCGMetadata::MergeAttributesSubset);
-	// Early out: nothing to do if both input metadata are null / points have no assigned metadata
-	if (!InMetadataA && !InMetadataB)
-	{
-		return;
-	}
-
-	// For each attribute in the current metadata, query the values from point A & B, apply operation on the result and finally store in the out point.
-	InitializeOnSet(OutKey, InKeyA, InMetadataA, InKeyB, InMetadataB);
-
-	AttributeLock.ReadLock();
-	for(const TPair<FName, FPCGMetadataAttributeBase*>& AttributePair : Attributes)
-	{
-		const FName& AttributeName = AttributePair.Key;
-		FPCGMetadataAttributeBase* Attribute = AttributePair.Value;
-
-		// Get attribute from A
-		const FPCGMetadataAttributeBase* AttributeA = nullptr;
-		if (InMetadataA && ((InMetadataA == InMetadataSubsetA) || (InMetadataSubsetA && InMetadataSubsetA->HasAttribute(AttributeName))))
-		{
-			AttributeA = InMetadataA->GetConstAttribute(AttributeName);
-		}
-
-		if (AttributeA && AttributeA->GetTypeId() != Attribute->GetTypeId())
-		{
-			UE_LOG(LogPCG, Error, TEXT("Metadata type mismatch with attribute %s"), *AttributeName.ToString());
-			AttributeA = nullptr;
-		}
+	// TODO Support more levels?
+	const FPCGMetadataLayer* InMetadataLayerA = InMetadataA->GetConstMetadataLayer(PCGMetadataLayerID::Default);
+	const FPCGMetadataLayer* InMetadataLayerSubsetA = InMetadataSubsetA->GetConstMetadataLayer(PCGMetadataLayerID::Default);
+	const FPCGMetadataLayer* InMetadataLayerB = InMetadataB->GetConstMetadataLayer(PCGMetadataLayerID::Default);
+	const FPCGMetadataLayer* InMetadataLayerSubsetB = InMetadataSubsetB->GetConstMetadataLayer(PCGMetadataLayerID::Default);
 
-		// Get attribute from B
-		const FPCGMetadataAttributeBase* AttributeB = nullptr;
-		if (InMetadataB && ((InMetadataB == InMetadataSubsetB) || (InMetadataSubsetB && InMetadataSubsetB->HasAttribute(AttributeName))))
-		{
-			AttributeB = InMetadataB->GetConstAttribute(AttributeName);
-		}
-
-		if (AttributeB && AttributeB->GetTypeId() != Attribute->GetTypeId())
-		{
-			UE_LOG(LogPCG, Error, TEXT("Metadata type mismatch with attribute %s"), *AttributeName.ToString());
-			AttributeB = nullptr;
-		}
-
-		if (AttributeA || AttributeB)
-		{
-			Attribute->SetValue(OutKey, AttributeA, InKeyA, AttributeB, InKeyB, Op);
-		}
-	}
-	AttributeLock.ReadUnlock();
+	return WithMetadataLayer(PCGMetadataLayerID::Default, &FPCGMetadataLayer::MergeAttributesSubset, InKeyA, InMetadataLayerA, InMetadataLayerSubsetA, InKeyB, InMetadataLayerB, InMetadataLayerSubsetB, OutKey, Op); 
 }
 
 void UPCGMetadata::ResetWeightedAttributes(PCGMetadataEntryKey& OutKey)
 {
-	InitializeOnSet(OutKey);
-
-	AttributeLock.ReadLock();
-	for(const TPair<FName, FPCGMetadataAttributeBase*>& AttributePair : Attributes)
-	{
-		FPCGMetadataAttributeBase* Attribute = AttributePair.Value;
-		
-		if (Attribute && Attribute->AllowsInterpolation())
-		{
-			Attribute->SetZeroValue(OutKey);
-		}
-	}
-	AttributeLock.ReadUnlock();
+	return WithMetadataLayer(PCGMetadataLayerID::Default, &FPCGMetadataLayer::ResetWeightedAttributes, OutKey);
 }
 
 void UPCGMetadata::AccumulateWeightedAttributes(PCGMetadataEntryKey InKey, const UPCGMetadata* InMetadata, float Weight, bool bSetNonInterpolableAttributes, PCGMetadataEntryKey& OutKey)
 {
-	if (!InMetadata)
-	{
-		return;
-	}
-
-	bool bHasSetParent = InitializeOnSet(OutKey, InKey, InMetadata);
-
-	const bool bShouldSetNonInterpolableAttributes = bSetNonInterpolableAttributes && !bHasSetParent;
-
-	AttributeLock.ReadLock();
-	for(const TPair<FName, FPCGMetadataAttributeBase*>& AttributePair : Attributes)
-	{
-		const FName& AttributeName = AttributePair.Key;
-		FPCGMetadataAttributeBase* Attribute = AttributePair.Value;
-
-		if (const FPCGMetadataAttributeBase* OtherAttribute = InMetadata->GetConstAttribute(AttributeName))
-		{
-			if (OtherAttribute->GetTypeId() != Attribute->GetTypeId())
-			{
-				UE_LOG(LogPCG, Error, TEXT("Metadata type mismatch with attribute %s"), *AttributeName.ToString());
-				continue;
-			}
-
-			if (Attribute->AllowsInterpolation())
-			{
-				Attribute->AccumulateValue(OutKey, OtherAttribute, InKey, Weight);
-			}
-			else if (bShouldSetNonInterpolableAttributes)
-			{
-				Attribute->SetValue(OutKey, OtherAttribute, InKey);
-			}
-		}
-	}
-	AttributeLock.ReadUnlock();
+	// TODO: Support other levels?
+	return WithMetadataLayer(PCGMetadataLayerID::Default, &FPCGMetadataLayer::AccumulateWeightedAttributes, InKey, InMetadata ? InMetadata->GetConstMetadataLayer(PCGMetadataLayerID::Default) : nullptr, Weight, bSetNonInterpolableAttributes, OutKey);
 }
 
 void UPCGMetadata::ComputePointWeightedAttribute(FPCGPoint& OutPoint, const TArrayView<TPair<const FPCGPoint*, float>>& InWeightedPoints, const UPCGMetadata* InMetadata)
 {
-	TArray<TPair<PCGMetadataEntryKey, float>, TInlineAllocator<4>> InWeightedKeys;
-	InWeightedKeys.Reserve(InWeightedPoints.Num());
-
-	for (const TPair<const FPCGPoint*, float>& WeightedPoint : InWeightedPoints)
+	// TODO: Support other levels?
+	return WithMetadataLayer_Lambda(PCGMetadataLayerID::Default, [&](FPCGMetadataLayer* MetadataLayer)
 	{
-		InWeightedKeys.Emplace(WeightedPoint.Key->MetadataEntry, WeightedPoint.Value);
-	}
-
-	ComputeWeightedAttribute(OutPoint.MetadataEntry, MakeArrayView(InWeightedKeys), InMetadata);
+		PCGMetadataHelpers::ComputePointWeightedAttribute(MetadataLayer, OutPoint, InWeightedPoints, InMetadata->GetConstMetadataLayer(PCGMetadataLayerID::Default));
+	});
 }
 
 void UPCGMetadata::ComputeWeightedAttribute(PCGMetadataEntryKey& OutKey, const TArrayView<TPair<PCGMetadataEntryKey, float>>& InWeightedKeys, const UPCGMetadata* InMetadata)
@@ -1281,35 +876,13 @@ void UPCGMetadata::ComputeWeightedAttribute(PCGMetadataEntryKey& OutKey, const T
 		return;
 	}
 
-	// Could ensure that InitializeOnSet returns false...
-	AttributeLock.ReadLock();
-	for (const TPair<FName, FPCGMetadataAttributeBase*>& AttributePair : Attributes)
-	{
-		const FName& AttributeName = AttributePair.Key;
-		FPCGMetadataAttributeBase* Attribute = AttributePair.Value;
-
-		if (!Attribute->AllowsInterpolation())
-		{
-			continue;
-		}
-
-		if (const FPCGMetadataAttributeBase* OtherAttribute = InMetadata->GetConstAttribute(AttributeName))
-		{
-			if (OtherAttribute->GetTypeId() != Attribute->GetTypeId())
-			{
-				UE_LOG(LogPCG, Error, TEXT("Metadata type mismatch with attribute %s"), *AttributeName.ToString());
-				continue;
-			}
-
-			Attribute->SetWeightedValue(OutKey, OtherAttribute, InWeightedKeys);
-		}
-	}
-	AttributeLock.ReadUnlock();
+	// TODO: Support other levels?
+	return WithMetadataLayer(PCGMetadataLayerID::Default, &FPCGMetadataLayer::ComputeWeightedAttribute, OutKey, InWeightedKeys, InMetadata->GetConstMetadataLayer(PCGMetadataLayerID::Default));
 }
 
 int64 UPCGMetadata::GetItemKeyCountForParent() const
 {
-	return ItemKeyOffset;
+	return WithConstMetadataLayer(PCGMetadataLayerID::Default, &FPCGMetadataLayer::GetItemKeyCountForParent);
 }
 
 void UPCGMetadata::SetAttributes(PCGMetadataEntryKey InKey, const UPCGMetadata* InMetadata, PCGMetadataEntryKey& OutKey)
@@ -1319,30 +892,11 @@ void UPCGMetadata::SetAttributes(PCGMetadataEntryKey InKey, const UPCGMetadata*
 		return;
 	}
 
-	if (InitializeOnSet(OutKey, InKey, InMetadata))
-	{
-		// Early out; we don't need to do anything else at this point
-		return;
-	}
-
-	AttributeLock.ReadLock();
-	for(const TPair<FName, FPCGMetadataAttributeBase*>& AttributePair : Attributes)
+	// TODO: Support other levels?
+	return WithMetadataLayer_Lambda(PCGMetadataLayerID::Default, [&InKey, InMetadataLayer = InMetadata->GetConstMetadataLayer(PCGMetadataLayerID::Default), &OutKey](FPCGMetadataLayer* MetadataLayer)
 	{
-		const FName& AttributeName = AttributePair.Key;
-		FPCGMetadataAttributeBase* Attribute = AttributePair.Value;
-
-		if (const FPCGMetadataAttributeBase* OtherAttribute = InMetadata->GetConstAttribute(AttributeName))
-		{
-			if (OtherAttribute->GetTypeId() != Attribute->GetTypeId())
-			{
-				UE_LOG(LogPCG, Error, TEXT("Metadata type mismatch with attribute %s"), *AttributeName.ToString());
-				continue;
-			}
-
-			Attribute->SetValue(OutKey, OtherAttribute, InKey);
-		}
-	}
-	AttributeLock.ReadUnlock();
+		return MetadataLayer->SetAttributes(InKey, InMetadataLayer, OutKey);
+	});
 }
 
 void UPCGMetadata::SetPointAttributes(const TArrayView<const FPCGPoint>& InPoints, const UPCGMetadata* InMetadata, const TArrayView<FPCGPoint>& OutPoints, FPCGContext* OptionalContext)
@@ -1352,24 +906,11 @@ void UPCGMetadata::SetPointAttributes(const TArrayView<const FPCGPoint>& InPoint
 		return;
 	}
 
-	TRACE_CPUPROFILER_EVENT_SCOPE(UPCGMetadata::SetPointAttributes);
-
-	check(InPoints.Num() == OutPoints.Num());
-
-	// Extract the metadata entry keys from the in & out points
-	TArray<PCGMetadataEntryKey, TInlineAllocator<256>> InKeys;
-	TArray<PCGMetadataEntryKey, TInlineAllocator<256>> OutKeys;
-
-	Algo::Transform(InPoints, InKeys, [](const FPCGPoint& Point) { return Point.MetadataEntry; });
-	Algo::Transform(OutPoints, OutKeys, [](const FPCGPoint& Point) { return Point.MetadataEntry; });
-
-	SetAttributes(InKeys, InMetadata, OutKeys, OptionalContext);
-
-	// Write back the keys on the points
-	for (int KeyIndex = 0; KeyIndex < OutKeys.Num(); ++KeyIndex)
+	// TODO: Support other levels?
+	return WithMetadataLayer_Lambda(PCGMetadataLayerID::Default, [&](FPCGMetadataLayer* MetadataLayer)
 	{
-		OutPoints[KeyIndex].MetadataEntry = OutKeys[KeyIndex];
-	}
+		PCGMetadataHelpers::SetPointAttributes(MetadataLayer, InPoints, InMetadata->GetConstMetadataLayer(PCGMetadataLayerID::Default), OutPoints, OptionalContext);
+	});
 }
 
 void UPCGMetadata::SetAttributes(const TArrayView<const PCGMetadataEntryKey>& InOriginalKeys, const UPCGMetadata* InMetadata, const TArrayView<PCGMetadataEntryKey>* InOutOptionalKeys, FPCGContext* OptionalContext)
@@ -1379,132 +920,11 @@ void UPCGMetadata::SetAttributes(const TArrayView<const PCGMetadataEntryKey>& In
 		return;
 	}
 
-	TRACE_CPUPROFILER_EVENT_SCOPE(UPCGMetadata::SetAttributes);
-
-	check(!InOutOptionalKeys || (InOriginalKeys.Num() == InOutOptionalKeys->Num()));
-
-	// There are a few things we can do to optimize here -
-	// basically, we don't need to set attributes more than once for a given <in, out> pair
-	TArray<PCGMetadataEntryKey, TInlineAllocator<256>> InKeys;
-	TArray<PCGMetadataEntryKey, TInlineAllocator<256>> OutKeys;
-
-	if (InOutOptionalKeys)
-	{
-		TRACE_CPUPROFILER_EVENT_SCOPE(UPCGMetadata::SetAttributes::CreateDeduplicatedKeys);
-		TMap<TPair<PCGMetadataEntryKey, PCGMetadataEntryKey>, int> PairMapping;
-
-		for (int KeyIndex = 0; KeyIndex < InOriginalKeys.Num(); ++KeyIndex)
-		{
-			PCGMetadataEntryKey InKey = InOriginalKeys[KeyIndex];
-			PCGMetadataEntryKey& OutKey = (*InOutOptionalKeys)[KeyIndex];
-
-			if (int* MatchingPairIndex = PairMapping.Find(TPair<PCGMetadataEntryKey, PCGMetadataEntryKey>(InKey, OutKey)))
-			{
-				OutKey = *MatchingPairIndex;
-			}
-			else
-			{
-				int NewIndex = InKeys.Add(InKey);
-
-				PairMapping.Emplace(TPair<PCGMetadataEntryKey, PCGMetadataEntryKey>(InKey, OutKey), NewIndex);
-				OutKeys.Add(OutKey);
-				OutKey = NewIndex;
-			}
-		}
-	}
-	else
-	{
-		InKeys.Append(InOriginalKeys);
-		OutKeys.Init(PCGInvalidEntryKey, InOriginalKeys.Num());
-	}
-
-	{
-		TRACE_CPUPROFILER_EVENT_SCOPE(UPCGMetadata::SetAttributes::InitializeOnSet);
-
-		for (int32 KeyIndex = 0; KeyIndex < InKeys.Num(); ++KeyIndex)
-		{
-			InitializeOnSet(OutKeys[KeyIndex], InKeys[KeyIndex], InMetadata);
-		}
-	}
-
-	AttributeLock.ReadLock();
-	int32 AttributeOffset = 0;
-	const int32 DefaultAttributesPerDispatch = 64;
-	int32 AttributesPerDispatch = OptionalContext ? DefaultAttributesPerDispatch : 1;
-	if (OptionalContext && OptionalContext->AsyncState.NumAvailableTasks > 0)
+	// TODO: Support other levels?
+	return WithMetadataLayer_Lambda(PCGMetadataLayerID::Default, [&](FPCGMetadataLayer* MetadataLayer)
 	{
-		AttributesPerDispatch = FMath::Min(OptionalContext->AsyncState.NumAvailableTasks, AttributesPerDispatch);
-	}
-
-	while (AttributeOffset < Attributes.Num())
-	{
-		TArray<FName> AttributeNames;
-		TArray<EPCGMetadataTypes> AttributeTypes;
-		GetAttributes(AttributeNames, AttributeTypes);
-
-		const int32 AttributeCountInCurrentDispatch = FMath::Min(AttributesPerDispatch, Attributes.Num() - AttributeOffset);
-		ParallelFor(AttributeCountInCurrentDispatch, [this, OptionalContext, AttributeOffset, InMetadata, &AttributeNames, &InKeys, &OutKeys](int32 WorkerIndex)
-		{
-			const FName AttributeName = AttributeNames[AttributeOffset + WorkerIndex];
-			FPCGMetadataAttributeBase* Attribute = Attributes[AttributeName];
-
-			if (const FPCGMetadataAttributeBase* OtherAttribute = InMetadata->GetConstAttribute(AttributeName))
-			{
-				if (!PCG::Private::IsBroadcastableOrConstructible(OtherAttribute->GetTypeId(), Attribute->GetTypeId()))
-				{
-					PCGE_LOG_C(Error, GraphAndLog, OptionalContext, FText::Format(NSLOCTEXT("PCGMetadata", "TypeMismatch", "Metadata type mismatch with attribute '{0}'"), FText::FromName(AttributeName)));
-					return;
-				}
-
-				if (Attribute == OtherAttribute)
-				{
-					TArray<PCGMetadataValueKey> ValueKeys;
-					Attribute->GetValueKeys(TArrayView<const PCGMetadataEntryKey>(InKeys), ValueKeys);
-					Attribute->SetValuesFromValueKeys(OutKeys, ValueKeys);
-				}
-				else
-				{
-					// Create accessor for the other attribute
-					TUniquePtr<const IPCGAttributeAccessor> OtherAttributeAccessor = PCGAttributeAccessorHelpers::CreateConstAccessor(OtherAttribute, InMetadata);
-
-					TArrayView<PCGMetadataEntryKey> InKeysView(InKeys);
-					FPCGAttributeAccessorKeysEntries OtherAttributeKeys(InKeysView);
-
-					if (!OtherAttributeAccessor)
-					{
-						return;
-					}
-
-					auto GetAndSetValues = [Attribute, &OutKeys, &OtherAttributeAccessor, &OtherAttributeKeys](auto Dummy) -> bool
-					{
-						using Type = decltype(Dummy);
-
-						auto SetValues = [Attribute, &OutKeys](const TArrayView<Type>& View, const int32 Start, const int32 Range)
-						{
-							TArrayView<PCGMetadataEntryKey> Keys(OutKeys.GetData() + Start, Range);
-							static_cast<FPCGMetadataAttribute<Type>*>(Attribute)->SetValues(Keys, View);
-						};
-
-						return PCGMetadataElementCommon::ApplyOnAccessorRange<Type>(OtherAttributeKeys, *OtherAttributeAccessor, SetValues, EPCGAttributeAccessorFlags::AllowBroadcast | EPCGAttributeAccessorFlags::AllowConstructible);
-					};
-
-					PCGMetadataAttribute::CallbackWithRightType(Attribute->GetTypeId(), GetAndSetValues);
-				}
-			}
-		});
-
-		AttributeOffset += AttributeCountInCurrentDispatch;
-	}
-	AttributeLock.ReadUnlock();
-
-	if (InOutOptionalKeys)
-	{
-		// Finally, copy back the actual out keys to the original out keys
-		for (PCGMetadataEntryKey& OutKey : *InOutOptionalKeys)
-		{
-			OutKey = OutKeys[OutKey];
-		}
-	}
+		return MetadataLayer->SetAttributes(InOriginalKeys, InMetadata->GetConstMetadataLayer(PCGMetadataLayerID::Default), InOutOptionalKeys, OptionalContext);
+	});
 }
 
 void UPCGMetadata::SetAttributes(const TArrayView<const PCGMetadataEntryKey>& InKeys, const UPCGMetadata* InMetadata, const TArrayView<PCGMetadataEntryKey>& OutKeys, FPCGContext* OptionalContext)
@@ -1565,3 +985,23 @@ void UPCGMetadata::SetLastCachedSelectorOnOwner(FName AttributeName)
 		OwnerData->SetLastSelector(Selector);
 	}
 }
+
+void UPCGMetadata::SetupLayersFromOtherMetadataIfNeeded(const UPCGMetadata* OtherMetadata)
+{
+	// Only do this if our outer is not a PCG data and we are not setuped.
+	// To be used with caution, only with floating metadata.
+	if (!OtherMetadata || OtherMetadata->DefaultLayer.IsDefault() || !DefaultLayer.IsDefault() || (GetOuter() && GetOuter()->IsA<UPCGData>()))
+	{
+		return;
+	}
+
+	for (const auto& It : OtherMetadata->MetadataLayers)
+	{
+		if (MetadataLayers.Contains(It.Key))
+		{
+			continue;
+		}
+		
+		SetupLayer(It.Key, It.Key == OtherMetadata->DefaultLayer);
+	}
+}
diff --git a/Engine/Plugins/PCG/Source/PCG/Private/Metadata/PCGMetadataAttribute.cpp b/Engine/Plugins/PCG/Source/PCG/Private/Metadata/PCGMetadataAttribute.cpp
--- a/Engine/Plugins/PCG/Source/PCG/Private/Metadata/PCGMetadataAttribute.cpp
+++ b/Engine/Plugins/PCG/Source/PCG/Private/Metadata/PCGMetadataAttribute.cpp
@@ -26,7 +26,7 @@ namespace PCGMetadataAttributeBase
 	}
 }
 
-FPCGMetadataAttributeBase::FPCGMetadataAttributeBase(UPCGMetadata* InMetadata, FName InName, const FPCGMetadataAttributeBase* InParent, bool bInAllowsInterpolation)
+FPCGMetadataAttributeBase::FPCGMetadataAttributeBase(FPCGMetadataLayer* InMetadata, FName InName, const FPCGMetadataAttributeBase* InParent, bool bInAllowsInterpolation)
 	: Metadata(InMetadata)
 	, Parent(InParent)
 	, bAllowsInterpolation(bInAllowsInterpolation)
@@ -34,7 +34,24 @@ FPCGMetadataAttributeBase::FPCGMetadataAttributeBase(UPCGMetadata* InMetadata, F
 {
 }
 
+PRAGMA_DISABLE_DEPRECATION_WARNINGS
+FPCGMetadataAttributeBase::FPCGMetadataAttributeBase(UPCGMetadata* InMetadata, FName InName, const FPCGMetadataAttributeBase* InParent, bool bInAllowsInterpolation)
+	: FPCGMetadataAttributeBase(InMetadata ? InMetadata->GetDefaultMetadataLayer() : nullptr, InName, InParent, bInAllowsInterpolation)
+{}
+
 void FPCGMetadataAttributeBase::Serialize(UPCGMetadata* InMetadata, FArchive& InArchive)
+{
+	Serialize(InMetadata ? InMetadata->GetDefaultMetadataLayer() : nullptr, InArchive);
+}
+
+const UPCGMetadata* FPCGMetadataAttributeBase::GetMetadata() const
+{
+	return Metadata ? Metadata->GetTopMetadata() : nullptr;
+}
+
+PRAGMA_ENABLE_DEPRECATION_WARNINGS
+
+void FPCGMetadataAttributeBase::Serialize(FPCGMetadataLayer* InMetadata, FArchive& InArchive)
 {
 	InArchive << EntryToValueKeyMap;
 	Metadata = InMetadata;
diff --git a/Engine/Plugins/PCG/Source/PCG/Private/Metadata/PCGMetadataLayer.cpp b/Engine/Plugins/PCG/Source/PCG/Private/Metadata/PCGMetadataLayer.cpp
new file mode 100644
--- /dev/null
+++ b/Engine/Plugins/PCG/Source/PCG/Private/Metadata/PCGMetadataLayer.cpp
@@ -0,0 +1,1349 @@
+// Copyright Epic Games, Inc. All Rights Reserved.
+
+#include "Metadata/PCGMetadataLayer.h"
+
+#include "PCGContext.h"
+#include "PCGElement.h"
+#include "Elements/Metadata/PCGMetadataElementCommon.h"
+#include "Metadata/PCGMetadata.h"
+#include "Helpers/PCGPropertyHelpers.h"
+
+#include "Algo/AnyOf.h"
+#include "Async/ParallelFor.h"
+#include "Serialization/ArchiveCrc32.h"
+
+namespace PCGMetadata
+{
+	template<typename DataType>
+	bool CreateAttributeFromPropertyHelper(FPCGMetadataLayer* Metadata, FName AttributeName, const DataType* DataPtr, const FProperty* InProperty)
+	{
+		if (!InProperty || !DataPtr || !Metadata)
+		{
+			return false;
+		}
+
+		auto CreateAttribute = [AttributeName, Metadata]<typename T>(T PropertyValue) -> bool
+		{
+			FPCGMetadataAttributeBase* BaseAttribute = Metadata->FindOrCreateAttribute<T>(AttributeName, PropertyValue, /*bAllowsInterpolation=*/false, /*bOverrideParent=*/false, /*bOverwriteIfTypeMismatch=*/true);
+			
+			return (BaseAttribute != nullptr);
+		};
+
+		return PCGPropertyHelpers::GetPropertyValueWithCallback(DataPtr, InProperty, CreateAttribute);
+	}
+
+	template<typename DataType>
+	bool SetAttributeFromPropertyHelper(FPCGMetadataLayer* Metadata, FName AttributeName, PCGMetadataEntryKey& EntryKey, const DataType* DataPtr, const FProperty* InProperty, bool bCreate)
+	{
+		if (!InProperty || !DataPtr || !Metadata)
+		{
+			return false;
+		}
+
+		// Check if an attribute already exists or not if we ask to create a new one
+		if (!bCreate && !Metadata->HasAttribute(AttributeName))
+		{
+			return false;
+		}
+
+		auto CreateAttributeAndSet = [AttributeName, Metadata, bCreate, &EntryKey]<typename PropertyType>(PropertyType PropertyValue) -> bool
+		{
+			FPCGMetadataAttributeBase* BaseAttribute = Metadata->GetMutableAttribute(AttributeName);
+
+			if (!BaseAttribute && bCreate)
+			{
+				// Interpolation is disabled and no parent override.
+				BaseAttribute = Metadata->CreateAttribute<PropertyType>(AttributeName, PropertyValue, false, false);
+			}
+
+			if (!BaseAttribute)
+			{
+				return false;
+			}
+
+			// Allow to set the value if both type matches or if we can construct AttributeType from PropertyType.
+			return PCGMetadataAttribute::CallbackWithRightType(BaseAttribute->GetTypeId(), [&EntryKey, &PropertyValue, BaseAttribute, Metadata]<typename AttributeType>(AttributeType AttributeValue) -> bool
+				{
+					FPCGMetadataAttribute<AttributeType>* Attribute = static_cast<FPCGMetadataAttribute<AttributeType>*>(BaseAttribute);
+
+					// Special cased because FSoftObjectPath currently has a deprecated constructor from FName which generates compile warnings.
+					constexpr bool bAssigningNameToSoftObjectPath = std::is_same_v<AttributeType, FSoftObjectPath> && std::is_same_v<PropertyType, FName>;
+
+					if constexpr (std::is_same_v<AttributeType, PropertyType>)
+					{
+						Metadata->InitializeOnSet(EntryKey);
+						Attribute->SetValue(EntryKey, PropertyValue);
+						return true;
+					}
+					else if constexpr (std::is_constructible_v<AttributeType, PropertyType> && !bAssigningNameToSoftObjectPath)
+					{
+						Metadata->InitializeOnSet(EntryKey);
+						Attribute->SetValue(EntryKey, AttributeType(PropertyValue));
+						return true;
+					}
+					else
+					{
+						return false;
+					}
+				});
+		};
+
+		return PCGPropertyHelpers::GetPropertyValueWithCallback(DataPtr, InProperty, CreateAttributeAndSet);
+	}
+}
+
+
+//////////////////////////
+/// FPCGMetadataLayer
+//////////////////////////
+
+
+void FPCGMetadataLayer::Serialize(FArchive& InArchive)
+{
+	int32 NumAttributes = (InArchive.IsLoading() ? 0 : Attributes.Num());
+	// We need to keep track of the max attribute Id, since it won't necessary be equal to the number of attributes + 1.
+	int64 MaxAttributeId = -1;
+
+	InArchive << NumAttributes;
+
+	if (InArchive.IsLoading())
+	{
+		for (int32 AttributeIndex = 0; AttributeIndex < NumAttributes; ++AttributeIndex)
+		{
+			FName AttributeName = NAME_None;
+			InArchive << AttributeName;
+
+			int32 AttributeTypeId = 0;
+			InArchive << AttributeTypeId;
+
+			FPCGMetadataAttributeBase* SerializedAttribute = PCGMetadataAttribute::AllocateEmptyAttributeFromType(static_cast<int16>(AttributeTypeId));
+			if (ensure(SerializedAttribute))
+			{
+				SerializedAttribute->Name = AttributeName;
+				SerializedAttribute->Serialize(this, InArchive);
+				Attributes.Add(AttributeName, SerializedAttribute);
+				MaxAttributeId = FMath::Max(SerializedAttribute->AttributeId, MaxAttributeId);
+			}
+		}
+	}
+	else
+	{
+		for (auto& AttributePair : Attributes)
+		{
+			InArchive << AttributePair.Key;
+			
+			int32 AttributeTypeId = AttributePair.Value->GetTypeId();
+			InArchive << AttributeTypeId;
+
+			AttributePair.Value->Serialize(this, InArchive);
+		}
+	}
+
+	InArchive << ParentKeys;
+
+	// Finally, initialize non-serialized members
+	if (InArchive.IsLoading())
+	{
+		// The next attribute id need to be bigger than the max attribute id of all attributes (or we could have collisions).
+		// Therefore by construction, it should never be less than the number of attributes (but can be greater).
+		NextAttributeId = MaxAttributeId + 1;
+		check(NextAttributeId >= Attributes.Num());
+		ItemKeyOffset = (Parent ? Parent->GetItemCountForChild() : 0);
+	}
+}
+
+void FPCGMetadataLayer::AddToCrc(FArchiveCrc32& Ar, const UPCGData* Data, bool bFullDataCrc) const
+{
+	check(Data);
+
+	// Sort attributes so we have a consistent processing path
+	TArray<FName> AttributeNames;
+
+	{
+		AttributeLock.ReadLock();
+		AttributeNames.Reserve(Attributes.Num());
+
+		for (const TPair<FName, FPCGMetadataAttributeBase*>& Attribute : Attributes)
+		{
+			AttributeNames.Add(Attribute.Key);
+		}
+		AttributeLock.ReadUnlock();
+
+		Algo::Sort(AttributeNames, [this](const FName& A, const FName& B) { return A.LexicalLess(B); });
+	}
+
+	// Add attributes to CRC
+	for (FName AttributeName : AttributeNames)
+	{
+		Ar << AttributeName;
+
+		FPCGAttributePropertyInputSelector InputSource;
+		InputSource.SetAttributeName(AttributeName);
+
+		TUniquePtr<const IPCGAttributeAccessor> InputAccessor = PCGAttributeAccessorHelpers::CreateConstAccessor(Data, InputSource);
+		TUniquePtr<const IPCGAttributeAccessorKeys> InputKeys = PCGAttributeAccessorHelpers::CreateConstKeys(Data, InputSource);
+
+		auto Callback = [&InputAccessor, &InputKeys, &Ar](auto&& Dummy)
+		{
+			using AttributeType = std::decay_t<decltype(Dummy)>;
+			TArray<AttributeType> Values;
+			if constexpr (std::is_trivially_copyable_v<AttributeType>)
+			{
+				Values.SetNumUninitialized(InputKeys->GetNum());
+			}
+			else
+			{
+				Values.SetNum(InputKeys->GetNum());
+			}
+			
+			InputAccessor->GetRange<AttributeType>(Values, 0, *InputKeys);
+
+			for (AttributeType& Value : Values)
+			{
+				// Add value to Crc
+				PCG::Private::Serialize(Ar, Value);
+			}
+		};
+
+		PCGMetadataAttribute::CallbackWithRightType(InputAccessor->GetUnderlyingType(), Callback);
+	}
+}
+
+FPCGMetadataLayer::FPCGMetadataLayer(UPCGMetadata* InTopMetadata, FPCGMetadataLayerID InMetadataLayerID)
+	: TopMetadata(InTopMetadata)
+	, LayerID(InMetadataLayerID)
+{
+	check(InTopMetadata);
+}
+
+FPCGMetadataLayer::~FPCGMetadataLayer()
+{
+	FWriteScopeLock ScopeLock(AttributeLock);
+	for (TPair<FName, FPCGMetadataAttributeBase*>& AttributeEntry : Attributes)
+	{
+		delete AttributeEntry.Value;
+		AttributeEntry.Value = nullptr;
+	}
+	Attributes.Reset();
+}
+
+void FPCGMetadataLayer::Initialize(const FPCGMetadataLayer* InParent, bool bAddAttributesFromParent)
+{
+	// If we are adding attributes from parent, then we use exclude filter with empty list so
+	// that all parameters added. Otherwise use include filter with empty list so none are added.
+	const EPCGMetadataFilterMode bFilter = bAddAttributesFromParent ? EPCGMetadataFilterMode::ExcludeAttributes : EPCGMetadataFilterMode::IncludeAttributes;
+	InitializeWithAttributeFilter(InParent, TSet<FName>(), bFilter);
+}
+
+void FPCGMetadataLayer::InitializeWithAttributeFilter(const FPCGMetadataLayer* InParent, const TSet<FName>& InFilteredAttributes, EPCGMetadataFilterMode InFilterMode, EPCGStringMatchingOperator InMatchOperator)
+{
+	if (Parent || Attributes.Num() != 0)
+	{
+		// Already initialized; note that while that might be construed as a warning, there are legit cases where this is correct
+		return;
+	}
+
+	Parent = ((InParent != this) ? InParent : nullptr);
+	ItemKeyOffset = Parent ? Parent->GetItemCountForChild() : 0;
+
+	// If we have been given an include list which is empty, then don't bother adding any attributes
+	const bool bSkipAddingAttributesFromParent = (InFilterMode == EPCGMetadataFilterMode::IncludeAttributes) && (InFilteredAttributes.Num() == 0);
+	if (!bSkipAddingAttributesFromParent)
+	{
+		AddAttributesFiltered(InParent, InFilteredAttributes, InFilterMode, InMatchOperator);
+	}
+}
+
+void FPCGMetadataLayer::InitializeAsCopy(const FPCGMetadataLayer* InMetadataToCopy, const TArray<PCGMetadataEntryKey>* EntriesToCopy)
+{
+	InitializeAsCopyWithAttributeFilter(InMetadataToCopy, TSet<FName>(), EPCGMetadataFilterMode::ExcludeAttributes, EntriesToCopy);
+}
+
+void FPCGMetadataLayer::InitializeAsCopyWithAttributeFilter(const FPCGMetadataLayer* InMetadataToCopy, const TSet<FName>& InFilteredAttributes, EPCGMetadataFilterMode InFilterMode, const TArray<PCGMetadataEntryKey>* EntriesToCopy)
+{
+	if (!InMetadataToCopy)
+	{
+		return;
+	}
+
+	check(InMetadataToCopy);
+	if (Parent || Attributes.Num() != 0)
+	{
+		UE_LOG(LogPCG, Error, TEXT("Metadata has already been initialized or already contains attributes"));
+		return;
+	}
+
+	const bool bSkipAttributesInFilterList = (InFilterMode == EPCGMetadataFilterMode::ExcludeAttributes);
+
+	// If we have a partial copy, it will flatten the metadata, so we don't need a parent.
+	// Otherwise, we keep the parent hierarchy.
+	const bool bPartialCopy = EntriesToCopy && EntriesToCopy->Num() <= InMetadataToCopy->GetItemCountForChild();
+	TArray<PCGMetadataEntryKey> NewEntryKeys;
+	TArray<PCGMetadataValueKey> NewValueKeys;
+	if (bPartialCopy)
+	{
+		const int32 Count = EntriesToCopy->Num();
+		NewEntryKeys.SetNumUninitialized(Count);
+		NewValueKeys.SetNumUninitialized(Count);
+		ParentKeys.SetNumUninitialized(Count);
+		for (int32 j = 0; j < Count; ++j)
+		{
+			NewEntryKeys[j] = PCGMetadataEntryKey(j);
+			ParentKeys[j] = -1;
+		}
+
+		ItemKeyOffset = 0;
+	}
+	else
+	{
+		ParentKeys = InMetadataToCopy->ParentKeys;
+		ItemKeyOffset = InMetadataToCopy->ItemKeyOffset;
+		Parent = InMetadataToCopy->Parent;
+	}
+
+	// Copy attributes
+	for (const TPair<FName, FPCGMetadataAttributeBase*>& OtherAttribute : InMetadataToCopy->Attributes)
+	{
+		const bool bAttributeInFilterList = InFilteredAttributes.Contains(OtherAttribute.Key);
+		const bool bSkipThisAttribute = (bSkipAttributesInFilterList == bAttributeInFilterList);
+
+		if (!bSkipThisAttribute)
+		{
+			// Don't copy entries if we have a partial copy, we will set them all after.
+			FPCGMetadataAttributeBase* Attribute = CopyAttribute(OtherAttribute.Value, OtherAttribute.Key, /*bKeepParent=*/false, /*bCopyEntries=*/!bPartialCopy, /*bCopyValues=*/true);
+
+			if (bPartialCopy && OtherAttribute.Value && Attribute)
+			{
+				OtherAttribute.Value->GetValueKeys(*EntriesToCopy, NewValueKeys);
+				Attribute->SetValuesFromValueKeys(NewEntryKeys, NewValueKeys);
+			}
+		}
+	}
+}
+
+bool FPCGMetadataLayer::AddAttributesFiltered(const FPCGMetadataLayer* InOther, const TSet<FName>& InFilteredAttributes, EPCGMetadataFilterMode InFilterMode, EPCGStringMatchingOperator InMatchOperator)
+{
+	if (!InOther)
+	{
+		return false;
+	}
+
+	TArray<FString> InFilteredAttributesStrings;
+	if (InMatchOperator != EPCGStringMatchingOperator::Equal)
+	{
+		Algo::Transform(InFilteredAttributes, InFilteredAttributesStrings, [](const FName& InFilteredAttribute) {return InFilteredAttribute.ToString(); });
+	}
+
+	auto IsAttributeInFilterList = [&InFilteredAttributes, &InFilteredAttributesStrings, InMatchOperator](FName OtherAttributeName) -> bool
+	{
+		if (InMatchOperator == EPCGStringMatchingOperator::Equal)
+		{
+			return InFilteredAttributes.Contains(OtherAttributeName);
+		}
+		else
+		{
+			const FString OtherAttributeString = OtherAttributeName.ToString();
+			if (InMatchOperator == EPCGStringMatchingOperator::Substring)
+			{
+				return Algo::AnyOf(InFilteredAttributesStrings, [&OtherAttributeString](const FString& InAttribute) { return OtherAttributeString.Contains(InAttribute); });
+			}
+			else if (InMatchOperator == EPCGStringMatchingOperator::Matches)
+			{
+				return Algo::AnyOf(InFilteredAttributesStrings, [&OtherAttributeString](const FString& InAttribute) { return OtherAttributeString.MatchesWildcard(InAttribute); });
+			}
+			else
+			{
+				checkNoEntry();
+				return false;
+			}
+		}
+	};
+
+	bool bAttributeAdded = false;
+
+	for (const TPair<FName, FPCGMetadataAttributeBase*>& OtherAttribute : InOther->Attributes)
+	{
+		// Skip this attribute if it is in an exclude list, or if it is not in an include list
+		const bool bAttributeInFilterList = IsAttributeInFilterList(OtherAttribute.Key);
+		const bool bSkipAttributesInFilterList = InFilterMode == EPCGMetadataFilterMode::ExcludeAttributes;
+		const bool bSkipThisAttribute = bSkipAttributesInFilterList == bAttributeInFilterList;
+
+		if(bSkipThisAttribute || !OtherAttribute.Value)
+		{
+			continue;
+		}
+		else if (HasAttribute(OtherAttribute.Key))
+		{
+			// If both the current attribute and the other attribute have the same type - nothing to do
+			// If the current attribute can be broadcasted to the other but not the other way around - change the type
+			// If none of this is true - do nothing
+			const FPCGMetadataAttributeBase* Attribute = GetConstAttribute(OtherAttribute.Key);
+			check(Attribute);
+
+			if(Attribute->GetTypeId() != OtherAttribute.Value->GetTypeId() && 
+				!PCG::Private::IsBroadcastable(OtherAttribute.Value->GetTypeId(), Attribute->GetTypeId()) &&
+				PCG::Private::IsBroadcastable(Attribute->GetTypeId(), OtherAttribute.Value->GetTypeId()))
+			{
+				ChangeAttributeType(OtherAttribute.Key, OtherAttribute.Value->GetTypeId());
+			}
+		}
+		else if (CopyAttribute(OtherAttribute.Value, OtherAttribute.Key, /*bKeepParent=*/InOther == Parent, /*bCopyEntries=*/false, /*bCopyValues=*/false))
+		{
+			bAttributeAdded = true;
+		}
+	}
+
+	return bAttributeAdded;
+}
+
+bool FPCGMetadataLayer::AddAttribute(const FPCGMetadataLayer* InOther, FName AttributeName)
+{
+	if (!InOther || !InOther->HasAttribute(AttributeName) || HasAttribute(AttributeName))
+	{
+		return false;
+	}
+
+	return CopyAttribute(InOther->GetConstAttribute(AttributeName), AttributeName, /*bKeepParent=*/InOther == Parent, /*bCopyEntries=*/false, /*bCopyValues=*/false) != nullptr;
+}
+
+void FPCGMetadataLayer::CopyAttributes(const FPCGMetadataLayer* InOther)
+{
+	if (!InOther || InOther == Parent)
+	{
+		return;
+	}
+
+	if (GetItemCountForChild() != InOther->GetItemCountForChild())
+	{
+		UE_LOG(LogPCG, Error, TEXT("Mismatch in copy attributes since the entries do not match"));
+		return;
+	}
+
+	for (const TPair<FName, FPCGMetadataAttributeBase*> OtherAttribute : InOther->Attributes)
+	{
+		if (HasAttribute(OtherAttribute.Key))
+		{
+			continue;
+		}
+		else
+		{
+			CopyAttribute(OtherAttribute.Value, OtherAttribute.Key, /*bKeepParent=*/false, /*bCopyEntries=*/true, /*bCopyValues=*/true);
+		}
+	}
+}
+
+void FPCGMetadataLayer::CopyAttribute(const FPCGMetadataLayer* InOther, FName AttributeToCopy, FName NewAttributeName)
+{
+	if (!InOther)
+	{
+		return;
+	}
+	else if (HasAttribute(NewAttributeName) || !InOther->HasAttribute(AttributeToCopy))
+	{
+		return;
+	}
+	else if (InOther == Parent)
+	{
+		CopyExistingAttribute(AttributeToCopy, NewAttributeName);
+		return;
+	}
+
+	if (GetItemCountForChild() != InOther->GetItemCountForChild())
+	{
+		UE_LOG(LogPCG, Error, TEXT("Mismatch in copy attributes since the entries do not match"));
+		return;
+	}
+
+	CopyAttribute(InOther->GetConstAttribute(AttributeToCopy), NewAttributeName, /*bKeepParent=*/false, /*bCopyEntries=*/true, /*bCopyValues=*/true);
+}
+
+const FPCGMetadataLayer* FPCGMetadataLayer::GetRoot() const
+{
+	if (Parent)
+	{
+		return Parent->GetRoot();
+	}
+	else
+	{
+		return this;
+	}
+}
+
+bool FPCGMetadataLayer::HasParent(const FPCGMetadataLayer* InTentativeParent) const
+{
+	if (!InTentativeParent)
+	{
+		return false;
+	}
+
+	const FPCGMetadataLayer* HierarchicalParent = Parent;
+	while (HierarchicalParent && HierarchicalParent != InTentativeParent)
+	{
+		HierarchicalParent = HierarchicalParent->Parent;
+	}
+
+	return HierarchicalParent == InTentativeParent;
+}
+
+void FPCGMetadataLayer::FlattenImpl()
+{
+	TRACE_CPUPROFILER_EVENT_SCOPE(FPCGSubMetadata::FlattenImpl);
+
+	const int32 NumEntries = GetItemCountForChild();
+
+	AttributeLock.WriteLock();
+	for (auto& AttributePair : Attributes)
+	{
+		FPCGMetadataAttributeBase* Attribute = AttributePair.Value;
+		check(Attribute);
+
+		// For all stored entries (from the root), we need to make sure that entries that should have a concrete value have it
+		// Optimization notes:
+		// - we could skip entries that existed prior to attribute existence, etc.
+		// - we could skip entries that have no parent, but that would require checking against the parent entries in the parent hierarchy
+		for (int64 EntryKey = 0; EntryKey < NumEntries; ++EntryKey)
+		{
+			// Get value using value inheritance as expected
+			PCGMetadataValueKey ValueKey = Attribute->GetValueKey(EntryKey);
+			if (ValueKey != PCGDefaultValueKey)
+			{
+				// Set concrete non-default value
+				Attribute->SetValueFromValueKey(EntryKey, ValueKey);
+			}
+		}
+
+		// Finally, flatten values
+		Attribute->Flatten();
+	}
+	AttributeLock.WriteUnlock();
+
+	Parent = nullptr;
+	ParentKeys.Reset();
+	ParentKeys.Init(PCGInvalidEntryKey, NumEntries);
+	ItemKeyOffset = 0;
+}
+
+bool FPCGMetadataLayer::FlattenAndCompress(const TArray<PCGMetadataEntryKey>& InEntryKeysToKeep)
+{
+	TRACE_CPUPROFILER_EVENT_SCOPE(FPCGSubMetadata::FlattenAndCompress);
+
+	// No keys or no parents, nothing to do
+	if (Attributes.IsEmpty())
+	{
+		return false;
+	}
+
+	AttributeLock.WriteLock();
+	for (auto& AttributePair : Attributes)
+	{
+		FPCGMetadataAttributeBase* Attribute = AttributePair.Value;
+		check(Attribute);
+
+		Attribute->FlattenAndCompress(InEntryKeysToKeep);
+	}
+	AttributeLock.WriteUnlock();
+
+	Parent = nullptr;
+	ParentKeys.Reset();
+	ParentKeys.Init(PCGInvalidEntryKey, InEntryKeysToKeep.Num());
+	ItemKeyOffset = 0;
+
+	return true;
+}
+
+void FPCGMetadataLayer::AddAttributeInternal(FName AttributeName, FPCGMetadataAttributeBase* Attribute)
+{
+	// This call assumes we have a write lock on the attribute map.
+	Attributes.Add(AttributeName, Attribute);
+}
+
+void FPCGMetadataLayer::RemoveAttributeInternal(FName AttributeName)
+{
+	Attributes.Remove(AttributeName);
+}
+
+void FPCGMetadataLayer::SetLastCachedSelectorOnOwner(FName AttributeName)
+{
+	TopMetadata->SetLastCachedSelectorOnOwner(AttributeName);
+}
+
+FPCGMetadataAttributeBase* FPCGMetadataLayer::GetMutableAttribute(FName AttributeName)
+{
+	FPCGMetadataAttributeBase* Attribute = nullptr;
+
+	AttributeLock.ReadLock();
+	if (FPCGMetadataAttributeBase** FoundAttribute = Attributes.Find(AttributeName))
+	{
+		Attribute = *FoundAttribute;
+
+		// Also when accessing an attribute, notify the PCG Data owner that the latest attribute manipulated is this one.
+		SetLastCachedSelectorOnOwner(AttributeName);
+	}
+	AttributeLock.ReadUnlock();
+
+	return Attribute;
+}
+
+const FPCGMetadataAttributeBase* FPCGMetadataLayer::GetConstAttribute(FName AttributeName) const
+{
+	const FPCGMetadataAttributeBase* Attribute = nullptr;
+
+	AttributeLock.ReadLock();
+	if (const FPCGMetadataAttributeBase* const* FoundAttribute = Attributes.Find(AttributeName))
+	{
+		Attribute = *FoundAttribute;
+	}
+	AttributeLock.ReadUnlock();
+
+	return Attribute;
+}
+
+const FPCGMetadataAttributeBase* FPCGMetadataLayer::GetConstAttributeById(int32 InAttributeId) const
+{
+	const FPCGMetadataAttributeBase* Attribute = nullptr;
+
+	AttributeLock.ReadLock();
+	for (const auto& AttributePair : Attributes)
+	{
+		if (AttributePair.Value && AttributePair.Value->AttributeId == InAttributeId)
+		{
+			Attribute = AttributePair.Value;
+			break;
+		}
+	}
+	AttributeLock.ReadUnlock();
+
+	return Attribute;
+}
+
+bool FPCGMetadataLayer::HasAttribute(FName AttributeName) const
+{
+	FReadScopeLock ScopeLock(AttributeLock);
+	return Attributes.Contains(AttributeName);
+}
+
+bool FPCGMetadataLayer::HasCommonAttributes(const FPCGMetadataLayer* InMetadata) const
+{
+	if (!InMetadata)
+	{
+		return false;
+	}
+
+	bool bHasCommonAttribute = false;
+
+	AttributeLock.ReadLock();
+	for (const TPair<FName, FPCGMetadataAttributeBase*>& AttributePair : Attributes)
+	{
+		if (InMetadata->HasAttribute(AttributePair.Key))
+		{
+			bHasCommonAttribute = true;
+			break;
+		}
+	}
+	AttributeLock.ReadUnlock();
+
+	return bHasCommonAttribute;
+}
+
+int32 FPCGMetadataLayer::GetAttributeCount() const
+{
+	FReadScopeLock ScopeLock(AttributeLock);
+	return Attributes.Num();
+}
+
+void FPCGMetadataLayer::GetAttributes(TArray<FName>& AttributeNames, TArray<EPCGMetadataTypes>& AttributeTypes) const
+{
+	AttributeLock.ReadLock();
+	for (const TPair<FName, FPCGMetadataAttributeBase*>& Attribute : Attributes)
+	{
+		check(Attribute.Value && Attribute.Value->Name == Attribute.Key);
+		AttributeNames.Add(Attribute.Key);
+
+		if (Attribute.Value->GetTypeId() < static_cast<uint16>(EPCGMetadataTypes::Unknown))
+		{
+			AttributeTypes.Add(static_cast<EPCGMetadataTypes>(Attribute.Value->GetTypeId()));
+		}
+		else
+		{
+			AttributeTypes.Add(EPCGMetadataTypes::Unknown);
+		}
+	}
+	AttributeLock.ReadUnlock();
+}
+
+FName FPCGMetadataLayer::GetLatestAttributeNameOrNone() const
+{
+	FName LatestAttributeName = NAME_None;
+	int64 MaxAttributeId = -1;
+	
+	AttributeLock.ReadLock();
+	for (const TPair<FName, FPCGMetadataAttributeBase*>& It : Attributes)
+	{
+		if (It.Value && (It.Value->AttributeId > MaxAttributeId))
+		{
+			MaxAttributeId = It.Value->AttributeId;
+			LatestAttributeName = It.Key;
+		}
+	}
+	AttributeLock.ReadUnlock();
+
+	return LatestAttributeName;
+}
+
+bool FPCGMetadataLayer::ParentHasAttribute(FName AttributeName) const
+{
+	return Parent && Parent->HasAttribute(AttributeName);
+}
+
+bool FPCGMetadataLayer::CreateAttributeFromProperty(FName AttributeName, const UObject* Object, const FProperty* InProperty)
+{
+	return PCGMetadata::CreateAttributeFromPropertyHelper<UObject>(this, AttributeName, Object, InProperty);
+}
+
+bool FPCGMetadataLayer::CreateAttributeFromDataProperty(FName AttributeName, const void* Data, const FProperty* InProperty)
+{
+	return PCGMetadata::CreateAttributeFromPropertyHelper<void>(this, AttributeName, Data, InProperty);
+}
+
+bool FPCGMetadataLayer::SetAttributeFromProperty(FName AttributeName, PCGMetadataEntryKey& EntryKey, const UObject* Object, const FProperty* InProperty, bool bCreate)
+{
+	return PCGMetadata::SetAttributeFromPropertyHelper<UObject>(this, AttributeName, EntryKey, Object, InProperty, bCreate);
+}
+
+bool FPCGMetadataLayer::SetAttributeFromDataProperty(FName AttributeName, PCGMetadataEntryKey& EntryKey, const void* Data, const FProperty* InProperty, bool bCreate)
+{
+	return PCGMetadata::SetAttributeFromPropertyHelper<void>(this, AttributeName, EntryKey, Data, InProperty, bCreate);
+}
+
+bool FPCGMetadataLayer::CopyExistingAttribute(FName AttributeToCopy, FName NewAttributeName, bool bKeepParent)
+{
+	return CopyAttribute(AttributeToCopy, NewAttributeName, bKeepParent, /*bCopyEntries=*/true, /*bCopyValues=*/true) != nullptr;
+}
+
+FPCGMetadataAttributeBase* FPCGMetadataLayer::CopyAttribute(FName AttributeToCopy, FName NewAttributeName, bool bKeepParent, bool bCopyEntries, bool bCopyValues)
+{
+	const FPCGMetadataAttributeBase* OriginalAttribute = nullptr;
+
+	AttributeLock.ReadLock();
+	if (FPCGMetadataAttributeBase** AttributeFound = Attributes.Find(AttributeToCopy))
+	{
+		OriginalAttribute = *AttributeFound;
+	}
+	AttributeLock.ReadUnlock();
+
+	if (!OriginalAttribute && Parent)
+	{
+		OriginalAttribute = Parent->GetConstAttribute(AttributeToCopy);
+	}
+
+	if (!OriginalAttribute)
+	{
+		UE_LOG(LogPCG, Warning, TEXT("Attribute %s does not exist, therefore cannot be copied"), *AttributeToCopy.ToString());
+		return nullptr;
+	}
+
+	return CopyAttribute(OriginalAttribute, NewAttributeName, bKeepParent, bCopyEntries, bCopyValues);
+}
+
+FPCGMetadataAttributeBase* FPCGMetadataLayer::CopyAttribute(const FPCGMetadataAttributeBase* OriginalAttribute, FName NewAttributeName, bool bKeepParent, bool bCopyEntries, bool bCopyValues)
+{
+	check(OriginalAttribute);
+	check(OriginalAttribute->GetMetadataLayer()->GetRoot() == GetRoot() || !bKeepParent);
+	FPCGMetadataAttributeBase* NewAttribute = OriginalAttribute->Copy(NewAttributeName, this, bKeepParent, bCopyEntries, bCopyValues);
+
+	if (NewAttribute)
+	{
+		AttributeLock.WriteLock();
+		NewAttribute->AttributeId = NextAttributeId++;
+		AddAttributeInternal(NewAttributeName, NewAttribute);
+
+		// Also when creating an attribute, notify the PCG Data owner that the latest attribute manipulated is this one.
+		SetLastCachedSelectorOnOwner(NewAttributeName);
+
+		AttributeLock.WriteUnlock();
+	}
+
+	return NewAttribute;
+}
+
+bool FPCGMetadataLayer::RenameAttribute(FName AttributeToRename, FName NewAttributeName)
+{
+	if (!FPCGMetadataAttributeBase::IsValidName(NewAttributeName))
+	{
+		UE_LOG(LogPCG, Error, TEXT("New attribute name %s is not valid"), *NewAttributeName.ToString());
+		return false;
+	}
+
+	bool bRenamed = false;
+	AttributeLock.WriteLock();
+	if (FPCGMetadataAttributeBase** AttributeFound = Attributes.Find(AttributeToRename))
+	{
+		FPCGMetadataAttributeBase* Attribute = *AttributeFound;
+		RemoveAttributeInternal(AttributeToRename);
+		Attribute->Name = NewAttributeName;
+		AddAttributeInternal(NewAttributeName, Attribute);
+
+		// Also when renaming an attribute, notify the PCG Data owner that the latest attribute manipulated is this one.
+		SetLastCachedSelectorOnOwner(NewAttributeName);
+		
+		bRenamed = true;
+	}
+	AttributeLock.WriteUnlock();
+
+	if (!bRenamed)
+	{
+		UE_LOG(LogPCG, Warning, TEXT("Attribute %s does not exist and therefore cannot be renamed"), *AttributeToRename.ToString());
+	}
+
+	return bRenamed;
+}
+
+void FPCGMetadataLayer::ClearAttribute(FName AttributeToClear)
+{
+	FPCGMetadataAttributeBase* Attribute = nullptr;
+
+	AttributeLock.ReadLock();
+	if (FPCGMetadataAttributeBase** AttributeFound = Attributes.Find(AttributeToClear))
+	{
+		Attribute = *AttributeFound;
+	}
+	AttributeLock.ReadUnlock();
+
+	// If the attribute exists, then we can lose all the entries
+	// If it doesn't but it exists in the parent hierarchy, then we must create a new attribute.
+	if (Attribute)
+	{
+		Attribute->ClearEntries();
+	}
+}
+
+void FPCGMetadataLayer::DeleteAttribute(FName AttributeToDelete)
+{
+	FPCGMetadataAttributeBase* Attribute = nullptr;
+
+	// If it's a local attribute, then just delete it
+	AttributeLock.WriteLock();
+	if (FPCGMetadataAttributeBase** AttributeFound = Attributes.Find(AttributeToDelete))
+	{
+		Attribute = *AttributeFound;
+		RemoveAttributeInternal(AttributeToDelete);
+	}
+	AttributeLock.WriteUnlock();
+
+	if (Attribute)
+	{
+		delete Attribute;
+	}
+	else
+	{
+		UE_LOG(LogPCG, Verbose, TEXT("Attribute %s does not exist and therefore cannot be deleted"), *AttributeToDelete.ToString());
+	}
+}
+
+bool FPCGMetadataLayer::ChangeAttributeType(FName AttributeName, int16 AttributeNewType)
+{
+	FPCGMetadataAttributeBase* Attribute = GetMutableAttribute(AttributeName);
+
+	if (!Attribute)
+	{
+		UE_LOG(LogPCG, Error, TEXT("Attribute '%s' does not exist and therefore cannot change its type"), *AttributeName.ToString());
+		return false;
+	}
+
+	if (Attribute->GetTypeId() == AttributeNewType)
+	{
+		// Nothing to do, attribute is already the type we want
+		return true;
+	}
+
+	if (FPCGMetadataAttributeBase* NewAttribute = Attribute->CopyToAnotherType(AttributeNewType))
+	{
+		NewAttribute->AttributeId = Attribute->AttributeId;
+
+		AttributeLock.WriteLock();
+		RemoveAttributeInternal(AttributeName);
+		AddAttributeInternal(AttributeName, NewAttribute);
+		AttributeLock.WriteUnlock();
+
+		delete Attribute;
+		Attribute = nullptr;
+	}
+
+	return true;
+}
+
+int64 FPCGMetadataLayer::GetItemCountForChild() const
+{
+	FReadScopeLock ScopeLock(ItemLock);
+	return ParentKeys.Num() + ItemKeyOffset;
+}
+
+int64 FPCGMetadataLayer::GetLocalItemCount() const
+{
+	FReadScopeLock ScopeLock(ItemLock);
+	return ParentKeys.Num();
+}
+
+int64 FPCGMetadataLayer::AddEntry(int64 ParentEntry)
+{
+	FWriteScopeLock ScopeLock(ItemLock);
+	return ParentKeys.Add(ParentEntry) + ItemKeyOffset;
+}
+
+TArray<int64> FPCGMetadataLayer::AddEntries(TArrayView<const int64> ParentEntryKeys)
+{
+	TArray<int64> Result;
+	Result.Reserve(ParentEntryKeys.Num());
+
+	FWriteScopeLock ScopeLock(ItemLock);
+	ParentKeys.Reserve(ParentKeys.Num() + ParentEntryKeys.Num());
+	for (const int64 ParentEntry : ParentEntryKeys)
+	{
+		Result.Add(ParentKeys.Add(ParentEntry) + ItemKeyOffset);
+	}
+
+	return Result;
+}
+
+void FPCGMetadataLayer::AddEntriesInPlace(TArrayView<int64*> ParentEntryKeys)
+{
+	FWriteScopeLock ScopeLock(ItemLock);
+	ParentKeys.Reserve(ParentKeys.Num() + ParentEntryKeys.Num());
+	for (int64* ParentEntry : ParentEntryKeys)
+	{
+		*ParentEntry = ParentKeys.Add(*ParentEntry) + ItemKeyOffset;
+	}
+}
+
+int64 FPCGMetadataLayer::AddEntryPlaceholder()
+{
+	FReadScopeLock ScopeLock(ItemLock);
+	return ParentKeys.Num() + DelayedEntriesIndex.IncrementExchange() + ItemKeyOffset;
+}
+
+void FPCGMetadataLayer::AddDelayedEntries(const TArray<TTuple<int64, int64>>& AllEntries)
+{
+	FWriteScopeLock ScopeLock(ItemLock);
+	ParentKeys.AddUninitialized(AllEntries.Num());
+	for (const TTuple<int64, int64>& Entry : AllEntries)
+	{
+		int64 Index = Entry.Get<0>() - ItemKeyOffset;
+		check(Index < ParentKeys.Num());
+		ParentKeys[Index] = Entry.Get<1>();
+	}
+
+	DelayedEntriesIndex.Exchange(0);
+}
+
+bool FPCGMetadataLayer::InitializeOnSet(PCGMetadataEntryKey& InOutKey, PCGMetadataEntryKey InParentKeyA, const FPCGMetadataLayer* InParentMetadataA, PCGMetadataEntryKey InParentKeyB, const FPCGMetadataLayer* InParentMetadataB)
+{
+	if (InOutKey == PCGInvalidEntryKey)
+	{
+		if (InParentKeyA != PCGInvalidEntryKey && Parent == InParentMetadataA)
+		{
+			InOutKey = AddEntry(InParentKeyA);
+			return true;
+		}
+		else if (InParentKeyB != PCGInvalidEntryKey && Parent == InParentMetadataB)
+		{
+			InOutKey = AddEntry(InParentKeyB);
+			return true;
+		}
+		else
+		{
+			InOutKey = AddEntry();
+			return false;
+		}
+	}
+	else if(InOutKey < ItemKeyOffset)
+	{
+		InOutKey = AddEntry(InOutKey);
+		return false;
+	}
+	else
+	{
+		return false;
+	}
+}
+
+PCGMetadataEntryKey FPCGMetadataLayer::GetParentKey(PCGMetadataEntryKey LocalItemKey) const
+{
+	if (LocalItemKey < ItemKeyOffset)
+	{
+		// Key is already in parent referential
+		return LocalItemKey;
+	}
+	else
+	{
+		FReadScopeLock ScopeLock(ItemLock);
+		if (LocalItemKey - ItemKeyOffset < ParentKeys.Num())
+		{
+			return ParentKeys[LocalItemKey - ItemKeyOffset];
+		}
+		else
+		{
+			UE_LOG(LogPCG, Warning, TEXT("Invalid metadata key - check for entry key not properly initialized"));
+			return PCGInvalidEntryKey;
+		}
+	}
+}
+
+void FPCGMetadataLayer::GetParentKeys(TArrayView<PCGMetadataEntryKey> LocalItemKeys, const TBitArray<>* Mask) const
+{
+	auto GetParentKey_Unsafe = [this](PCGMetadataEntryKey& LocalItemKey) -> void
+	{
+		if (LocalItemKey < ItemKeyOffset)
+		{
+			// Key is already in parent referential
+			return;
+		}
+		else if (LocalItemKey - ItemKeyOffset < ParentKeys.Num())
+		{
+			LocalItemKey = ParentKeys[LocalItemKey - ItemKeyOffset];
+		}
+		else
+		{
+			UE_LOG(LogPCG, Warning, TEXT("Invalid metadata key - check for entry key not properly initialized"));
+			LocalItemKey = PCGInvalidEntryKey;
+		}
+	};
+
+	FReadScopeLock ScopeLock(ItemLock);
+	if (Mask && ensure(LocalItemKeys.Num() == Mask->Num()))
+	{
+		for (TConstSetBitIterator<> It(*Mask); It; ++It)
+		{
+			GetParentKey_Unsafe(LocalItemKeys[It.GetIndex()]);
+		}
+	}
+	else
+	{
+		for (PCGMetadataEntryKey& LocalItemKey : LocalItemKeys)
+		{
+			GetParentKey_Unsafe(LocalItemKey);
+		}
+	}
+}
+
+void FPCGMetadataLayer::MergeAttributes(PCGMetadataEntryKey InKeyA, const FPCGMetadataLayer* InMetadataA, PCGMetadataEntryKey InKeyB, const FPCGMetadataLayer* InMetadataB, PCGMetadataEntryKey& OutKey, EPCGMetadataOp Op)
+{
+	MergeAttributesSubset(InKeyA, InMetadataA, InMetadataA, InKeyB, InMetadataB, InMetadataB, OutKey, Op);
+}
+
+void FPCGMetadataLayer::MergeAttributesSubset(PCGMetadataEntryKey InKeyA, const FPCGMetadataLayer* InMetadataA, const FPCGMetadataLayer* InMetadataSubsetA, PCGMetadataEntryKey InKeyB, const FPCGMetadataLayer* InMetadataB, const FPCGMetadataLayer* InMetadataSubsetB, PCGMetadataEntryKey& OutKey, EPCGMetadataOp Op)
+{
+	//TRACE_CPUPROFILER_EVENT_SCOPE(FPCGMetadataLayer::MergeAttributesSubset);
+	// Early out: nothing to do if both input metadata are null / points have no assigned metadata
+	if (!InMetadataA && !InMetadataB)
+	{
+		return;
+	}
+
+	// For each attribute in the current metadata, query the values from point A & B, apply operation on the result and finally store in the out point.
+	InitializeOnSet(OutKey, InKeyA, InMetadataA, InKeyB, InMetadataB);
+
+	AttributeLock.ReadLock();
+	for(const TPair<FName, FPCGMetadataAttributeBase*>& AttributePair : Attributes)
+	{
+		const FName& AttributeName = AttributePair.Key;
+		FPCGMetadataAttributeBase* Attribute = AttributePair.Value;
+
+		// Get attribute from A
+		const FPCGMetadataAttributeBase* AttributeA = nullptr;
+		if (InMetadataA && ((InMetadataA == InMetadataSubsetA) || (InMetadataSubsetA && InMetadataSubsetA->HasAttribute(AttributeName))))
+		{
+			AttributeA = InMetadataA->GetConstAttribute(AttributeName);
+		}
+
+		if (AttributeA && AttributeA->GetTypeId() != Attribute->GetTypeId())
+		{
+			UE_LOG(LogPCG, Error, TEXT("Metadata type mismatch with attribute %s"), *AttributeName.ToString());
+			AttributeA = nullptr;
+		}
+
+		// Get attribute from B
+		const FPCGMetadataAttributeBase* AttributeB = nullptr;
+		if (InMetadataB && ((InMetadataB == InMetadataSubsetB) || (InMetadataSubsetB && InMetadataSubsetB->HasAttribute(AttributeName))))
+		{
+			AttributeB = InMetadataB->GetConstAttribute(AttributeName);
+		}
+
+		if (AttributeB && AttributeB->GetTypeId() != Attribute->GetTypeId())
+		{
+			UE_LOG(LogPCG, Error, TEXT("Metadata type mismatch with attribute %s"), *AttributeName.ToString());
+			AttributeB = nullptr;
+		}
+
+		if (AttributeA || AttributeB)
+		{
+			Attribute->SetValue(OutKey, AttributeA, InKeyA, AttributeB, InKeyB, Op);
+		}
+	}
+	AttributeLock.ReadUnlock();
+}
+
+void FPCGMetadataLayer::ResetWeightedAttributes(PCGMetadataEntryKey& OutKey)
+{
+	InitializeOnSet(OutKey);
+
+	AttributeLock.ReadLock();
+	for(const TPair<FName, FPCGMetadataAttributeBase*>& AttributePair : Attributes)
+	{
+		FPCGMetadataAttributeBase* Attribute = AttributePair.Value;
+		
+		if (Attribute && Attribute->AllowsInterpolation())
+		{
+			Attribute->SetZeroValue(OutKey);
+		}
+	}
+	AttributeLock.ReadUnlock();
+}
+
+void FPCGMetadataLayer::AccumulateWeightedAttributes(PCGMetadataEntryKey InKey, const FPCGMetadataLayer* InMetadata, float Weight, bool bSetNonInterpolableAttributes, PCGMetadataEntryKey& OutKey)
+{
+	if (!InMetadata)
+	{
+		return;
+	}
+
+	bool bHasSetParent = InitializeOnSet(OutKey, InKey, InMetadata);
+
+	const bool bShouldSetNonInterpolableAttributes = bSetNonInterpolableAttributes && !bHasSetParent;
+
+	AttributeLock.ReadLock();
+	for(const TPair<FName, FPCGMetadataAttributeBase*>& AttributePair : Attributes)
+	{
+		const FName& AttributeName = AttributePair.Key;
+		FPCGMetadataAttributeBase* Attribute = AttributePair.Value;
+
+		if (const FPCGMetadataAttributeBase* OtherAttribute = InMetadata->GetConstAttribute(AttributeName))
+		{
+			if (OtherAttribute->GetTypeId() != Attribute->GetTypeId())
+			{
+				UE_LOG(LogPCG, Error, TEXT("Metadata type mismatch with attribute %s"), *AttributeName.ToString());
+				continue;
+			}
+
+			if (Attribute->AllowsInterpolation())
+			{
+				Attribute->AccumulateValue(OutKey, OtherAttribute, InKey, Weight);
+			}
+			else if (bShouldSetNonInterpolableAttributes)
+			{
+				Attribute->SetValue(OutKey, OtherAttribute, InKey);
+			}
+		}
+	}
+	AttributeLock.ReadUnlock();
+}
+
+void FPCGMetadataLayer::ComputeWeightedAttribute(PCGMetadataEntryKey& OutKey, const TArrayView<TPair<PCGMetadataEntryKey, float>>& InWeightedKeys, const FPCGMetadataLayer* InMetadata)
+{
+	if (!InMetadata || InWeightedKeys.IsEmpty())
+	{
+		return;
+	}
+
+	// Could ensure that InitializeOnSet returns false...
+	AttributeLock.ReadLock();
+	for (const TPair<FName, FPCGMetadataAttributeBase*>& AttributePair : Attributes)
+	{
+		const FName& AttributeName = AttributePair.Key;
+		FPCGMetadataAttributeBase* Attribute = AttributePair.Value;
+
+		if (!Attribute->AllowsInterpolation())
+		{
+			continue;
+		}
+
+		if (const FPCGMetadataAttributeBase* OtherAttribute = InMetadata->GetConstAttribute(AttributeName))
+		{
+			if (OtherAttribute->GetTypeId() != Attribute->GetTypeId())
+			{
+				UE_LOG(LogPCG, Error, TEXT("Metadata type mismatch with attribute %s"), *AttributeName.ToString());
+				continue;
+			}
+
+			Attribute->SetWeightedValue(OutKey, OtherAttribute, InWeightedKeys);
+		}
+	}
+	AttributeLock.ReadUnlock();
+}
+
+int64 FPCGMetadataLayer::GetItemKeyCountForParent() const
+{
+	return ItemKeyOffset;
+}
+
+void FPCGMetadataLayer::SetAttributes(PCGMetadataEntryKey InKey, const FPCGMetadataLayer* InMetadata, PCGMetadataEntryKey& OutKey)
+{
+	if (!InMetadata)
+	{
+		return;
+	}
+
+	if (InitializeOnSet(OutKey, InKey, InMetadata))
+	{
+		// Early out; we don't need to do anything else at this point
+		return;
+	}
+
+	AttributeLock.ReadLock();
+	for(const TPair<FName, FPCGMetadataAttributeBase*>& AttributePair : Attributes)
+	{
+		const FName& AttributeName = AttributePair.Key;
+		FPCGMetadataAttributeBase* Attribute = AttributePair.Value;
+
+		if (const FPCGMetadataAttributeBase* OtherAttribute = InMetadata->GetConstAttribute(AttributeName))
+		{
+			if (OtherAttribute->GetTypeId() != Attribute->GetTypeId())
+			{
+				UE_LOG(LogPCG, Error, TEXT("Metadata type mismatch with attribute %s"), *AttributeName.ToString());
+				continue;
+			}
+
+			Attribute->SetValue(OutKey, OtherAttribute, InKey);
+		}
+	}
+	AttributeLock.ReadUnlock();
+}
+
+void FPCGMetadataLayer::SetAttributes(const TArrayView<const PCGMetadataEntryKey>& InOriginalKeys, const FPCGMetadataLayer* InMetadata, const TArrayView<PCGMetadataEntryKey>* InOutOptionalKeys, FPCGContext* OptionalContext)
+{
+	if (!InMetadata || InMetadata->GetAttributeCount() == 0 || GetAttributeCount() == 0 || InOriginalKeys.IsEmpty())
+	{
+		return;
+	}
+
+	TRACE_CPUPROFILER_EVENT_SCOPE(FPCGSubMetadata::SetAttributes);
+
+	check(!InOutOptionalKeys || (InOriginalKeys.Num() == InOutOptionalKeys->Num()));
+
+	// There are a few things we can do to optimize here -
+	// basically, we don't need to set attributes more than once for a given <in, out> pair
+	TArray<PCGMetadataEntryKey, TInlineAllocator<256>> InKeys;
+	TArray<PCGMetadataEntryKey, TInlineAllocator<256>> OutKeys;
+
+	if (InOutOptionalKeys)
+	{
+		TRACE_CPUPROFILER_EVENT_SCOPE(FPCGSubMetadata::SetAttributes::CreateDeduplicatedKeys);
+		TMap<TPair<PCGMetadataEntryKey, PCGMetadataEntryKey>, int> PairMapping;
+
+		for (int KeyIndex = 0; KeyIndex < InOriginalKeys.Num(); ++KeyIndex)
+		{
+			PCGMetadataEntryKey InKey = InOriginalKeys[KeyIndex];
+			PCGMetadataEntryKey& OutKey = (*InOutOptionalKeys)[KeyIndex];
+
+			if (int* MatchingPairIndex = PairMapping.Find(TPair<PCGMetadataEntryKey, PCGMetadataEntryKey>(InKey, OutKey)))
+			{
+				OutKey = *MatchingPairIndex;
+			}
+			else
+			{
+				int NewIndex = InKeys.Add(InKey);
+
+				PairMapping.Emplace(TPair<PCGMetadataEntryKey, PCGMetadataEntryKey>(InKey, OutKey), NewIndex);
+				OutKeys.Add(OutKey);
+				OutKey = NewIndex;
+			}
+		}
+	}
+	else
+	{
+		InKeys.Append(InOriginalKeys);
+		OutKeys.Init(PCGInvalidEntryKey, InOriginalKeys.Num());
+	}
+
+	{
+		TRACE_CPUPROFILER_EVENT_SCOPE(FPCGSubMetadata::SetAttributes::InitializeOnSet);
+
+		for (int32 KeyIndex = 0; KeyIndex < InKeys.Num(); ++KeyIndex)
+		{
+			InitializeOnSet(OutKeys[KeyIndex], InKeys[KeyIndex], InMetadata);
+		}
+	}
+
+	AttributeLock.ReadLock();
+	int32 AttributeOffset = 0;
+	const int32 DefaultAttributesPerDispatch = 64;
+	int32 AttributesPerDispatch = OptionalContext ? DefaultAttributesPerDispatch : 1;
+	if (OptionalContext && OptionalContext->AsyncState.NumAvailableTasks > 0)
+	{
+		AttributesPerDispatch = FMath::Min(OptionalContext->AsyncState.NumAvailableTasks, AttributesPerDispatch);
+	}
+
+	while (AttributeOffset < Attributes.Num())
+	{
+		TArray<FName> AttributeNames;
+		TArray<EPCGMetadataTypes> AttributeTypes;
+		GetAttributes(AttributeNames, AttributeTypes);
+
+		const int32 AttributeCountInCurrentDispatch = FMath::Min(AttributesPerDispatch, Attributes.Num() - AttributeOffset);
+		ParallelFor(AttributeCountInCurrentDispatch, [this, OptionalContext, AttributeOffset, InMetadata, &AttributeNames, &InKeys, &OutKeys](int32 WorkerIndex)
+		{
+			const FName AttributeName = AttributeNames[AttributeOffset + WorkerIndex];
+			FPCGMetadataAttributeBase* Attribute = Attributes[AttributeName];
+
+			if (const FPCGMetadataAttributeBase* OtherAttribute = InMetadata->GetConstAttribute(AttributeName))
+			{
+				if (!PCG::Private::IsBroadcastableOrConstructible(OtherAttribute->GetTypeId(), Attribute->GetTypeId()))
+				{
+					PCGE_LOG_C(Error, GraphAndLog, OptionalContext, FText::Format(NSLOCTEXT("PCGMetadata", "TypeMismatch", "Metadata type mismatch with attribute '{0}'"), FText::FromName(AttributeName)));
+					return;
+				}
+
+				if (Attribute == OtherAttribute)
+				{
+					TArray<PCGMetadataValueKey> ValueKeys;
+					Attribute->GetValueKeys(TArrayView<const PCGMetadataEntryKey>(InKeys), ValueKeys);
+					Attribute->SetValuesFromValueKeys(OutKeys, ValueKeys);
+				}
+				else
+				{
+					// Create accessor for the other attribute
+					TUniquePtr<const IPCGAttributeAccessor> OtherAttributeAccessor = PCGAttributeAccessorHelpers::CreateConstAccessor(OtherAttribute, InMetadata);
+
+					TArrayView<PCGMetadataEntryKey> InKeysView(InKeys);
+					FPCGAttributeAccessorKeysEntries OtherAttributeKeys(InKeysView);
+
+					if (!OtherAttributeAccessor)
+					{
+						return;
+					}
+
+					auto GetAndSetValues = [Attribute, &OutKeys, &OtherAttributeAccessor, &OtherAttributeKeys](auto Dummy) -> bool
+					{
+						using Type = decltype(Dummy);
+
+						auto SetValues = [Attribute, &OutKeys](const TArrayView<Type>& View, const int32 Start, const int32 Range)
+						{
+							TArrayView<PCGMetadataEntryKey> Keys(OutKeys.GetData() + Start, Range);
+							static_cast<FPCGMetadataAttribute<Type>*>(Attribute)->SetValues(Keys, View);
+						};
+
+						return PCGMetadataElementCommon::ApplyOnAccessorRange<Type>(OtherAttributeKeys, *OtherAttributeAccessor, SetValues, EPCGAttributeAccessorFlags::AllowBroadcast | EPCGAttributeAccessorFlags::AllowConstructible);
+					};
+
+					PCGMetadataAttribute::CallbackWithRightType(Attribute->GetTypeId(), GetAndSetValues);
+				}
+			}
+		});
+
+		AttributeOffset += AttributeCountInCurrentDispatch;
+	}
+	AttributeLock.ReadUnlock();
+
+	if (InOutOptionalKeys)
+	{
+		// Finally, copy back the actual out keys to the original out keys
+		for (PCGMetadataEntryKey& OutKey : *InOutOptionalKeys)
+		{
+			OutKey = OutKeys[OutKey];
+		}
+	}
+}
+
+void FPCGMetadataLayer::SetAttributes(const TArrayView<const PCGMetadataEntryKey>& InKeys, const FPCGMetadataLayer* InMetadata, const TArrayView<PCGMetadataEntryKey>& OutKeys, FPCGContext* OptionalContext)
+{
+	SetAttributes(InKeys, InMetadata, &OutKeys, OptionalContext);
+}
diff --git a/Engine/Plugins/PCG/Source/PCG/Private/Metadata/PCGMetadataPartitionCommon.cpp b/Engine/Plugins/PCG/Source/PCG/Private/Metadata/PCGMetadataPartitionCommon.cpp
--- a/Engine/Plugins/PCG/Source/PCG/Private/Metadata/PCGMetadataPartitionCommon.cpp
+++ b/Engine/Plugins/PCG/Source/PCG/Private/Metadata/PCGMetadataPartitionCommon.cpp
@@ -6,6 +6,7 @@
 #include "PCGModule.h"
 #include "PCGParamData.h"
 #include "Data/PCGPointData.h"
+#include "Data/PCGSpatialData.h"
 #include "Elements/Metadata/PCGMetadataElementCommon.h"
 #include "Metadata/PCGMetadata.h"
 #include "Metadata/PCGMetadataCommon.h"
@@ -15,6 +16,8 @@
 #include "Metadata/Accessors/PCGAttributeAccessorKeys.h"
 
 #include "Algo/IndexOf.h"
+#include "UObject/Package.h"
+#include "UObject/UObjectGlobals.h"
 
 #define LOCTEXT_NAMESPACE "PCGMetadataPartitionCommon"
 
@@ -523,7 +526,7 @@ namespace PCGMetadataPartitionCommon
 
 			if (InSpatialData)
 			{
-				UPCGSpatialData* NewData = FPCGContext::NewObject_AnyThread<UPCGSpatialData>(InOptionalContext);
+				UPCGSpatialData* NewData = FPCGContext::NewObject_AnyThread<UPCGSpatialData>(InOptionalContext, GetTransientPackage(), InSpatialData->GetClass());
 				NewData->InitializeFromData(InSpatialData);
 				NewMetadata = NewData->Metadata;
 				PartitionedData.Add(NewData);
@@ -600,7 +603,7 @@ namespace PCGMetadataPartitionCommon
 			{
 				if (InSpatialData)
 				{
-					UPCGSpatialData* NewData = FPCGContext::NewObject_AnyThread<UPCGSpatialData>(InOptionalContext);
+					UPCGSpatialData* NewData = FPCGContext::NewObject_AnyThread<UPCGSpatialData>(InOptionalContext, GetTransientPackage(), InSpatialData->GetClass());
 					NewData->InitializeFromData(InSpatialData);
 					OutputData = NewData;
 					OutMetadata = NewData->Metadata;
diff --git a/Engine/Plugins/PCG/Source/PCG/Private/PCGData.cpp b/Engine/Plugins/PCG/Source/PCG/Private/PCGData.cpp
--- a/Engine/Plugins/PCG/Source/PCG/Private/PCGData.cpp
+++ b/Engine/Plugins/PCG/Source/PCG/Private/PCGData.cpp
@@ -82,9 +82,9 @@ void UPCGData::PostEditImport()
 
 void UPCGData::Flatten()
 {
-	if (UPCGMetadata* Metadata = MutableMetadata())
+	if (UPCGMetadata* ThisMetadata = MutableMetadata())
 	{
-		Metadata->FlattenImpl();
+		ThisMetadata->FlattenImpl();
 	}
 }
 
diff --git a/Engine/Plugins/PCG/Source/PCG/Private/PCGParamData.cpp b/Engine/Plugins/PCG/Source/PCG/Private/PCGParamData.cpp
--- a/Engine/Plugins/PCG/Source/PCG/Private/PCGParamData.cpp
+++ b/Engine/Plugins/PCG/Source/PCG/Private/PCGParamData.cpp
@@ -17,6 +17,10 @@ UPCGParamData::UPCGParamData(const FObjectInitializer& ObjectInitializer)
 	: Super(ObjectInitializer)
 {
 	Metadata = ObjectInitializer.CreateDefaultSubobject<UPCGMetadata>(this, TEXT("Metadata"));
+	for (const FPCGMetadataLayerID& MetadataLevel : UPCGParamData::GetAllSupportedMetadataLayerIDs())
+	{
+		Metadata->SetupLayer(MetadataLevel, MetadataLevel == UPCGParamData::GetDefaultMetadataLayerID());
+	}
 }
 
 void UPCGParamData::AddToCrc(FArchiveCrc32& Ar, bool bFullDataCrc) const
diff --git a/Engine/Plugins/PCG/Source/PCG/Private/Tests/Accessors/PCGAttributeAccessorTest.cpp b/Engine/Plugins/PCG/Source/PCG/Private/Tests/Accessors/PCGAttributeAccessorTest.cpp
--- a/Engine/Plugins/PCG/Source/PCG/Private/Tests/Accessors/PCGAttributeAccessorTest.cpp
+++ b/Engine/Plugins/PCG/Source/PCG/Private/Tests/Accessors/PCGAttributeAccessorTest.cpp
@@ -5,6 +5,7 @@
 #include "Data/PCGPointData.h"
 #include "Metadata/Accessors/IPCGAttributeAccessorTpl.h"
 #include "Metadata/Accessors/PCGAttributeAccessor.h"
+#include "Metadata/Accessors/PCGAttributeAccessorHelpers.h"
 
 #include "UObject/Package.h"
 
@@ -22,9 +23,9 @@ namespace FPCGAttributeAccessorTestHelpers
 	{
 		AttributeData()
 			: PointData(NewObject<UPCGPointData>(GetTransientPackage(), NAME_None, RF_Transient))
-			, FloatAttribute(PointData->Metadata, TEXT("FloatAttribute"), nullptr, 0.0f, true)
-			, VectorAttribute(PointData->Metadata, TEXT("VectorAttribute"), nullptr, FVector::ZeroVector, true)
-			, StringAttribute(PointData->Metadata, TEXT("StringAttribute"), nullptr, FString(), true)
+			, FloatAttribute(PointData->Metadata->CreateAttribute<float>(TEXT("FloatAttribute"), 0.0f, true, false))
+			, VectorAttribute(PointData->Metadata->CreateAttribute<FVector>(TEXT("VectorAttribute"), FVector::ZeroVector, false, false))
+			, StringAttribute(PointData->Metadata->CreateAttribute<FString>(TEXT("StringAttribute"), FString(), false, false))
 		{
 			// All points will point to a given entry, in reverse order, except the last one, that will point to DefaultValue
 			TArray<FPCGPoint>& Points = PointData->GetMutablePoints();
@@ -33,15 +34,15 @@ namespace FPCGAttributeAccessorTestHelpers
 			for (int32 i = 0; i < ValuesSize; ++i)
 			{
 				PCGMetadataEntryKey EntryKey = PointData->Metadata->AddEntry();
-				FloatAttribute.SetValue(EntryKey, FloatValues[i]);
-				VectorAttribute.SetValue(EntryKey, VectorValues[i]);
-				StringAttribute.SetValue(EntryKey, StringValues[i]);
+				FloatAttribute->SetValue(EntryKey, FloatValues[i]);
+				VectorAttribute->SetValue(EntryKey, VectorValues[i]);
+				StringAttribute->SetValue(EntryKey, StringValues[i]);
 				Points[ValuesSize - i - 1].MetadataEntry = EntryKey;
 			}
 
-			FloatAccessor = MakeUnique<FPCGAttributeAccessor<float>>(&FloatAttribute, PointData->Metadata);
-			VectorAccessor = MakeUnique<FPCGAttributeAccessor<FVector>>(&VectorAttribute, PointData->Metadata);
-			StringAccessor = MakeUnique<FPCGAttributeAccessor<FString>>(&StringAttribute, PointData->Metadata);
+			FloatAccessor = PCGAttributeAccessorHelpers::CreateAccessor(FloatAttribute, PointData->Metadata);
+			VectorAccessor = PCGAttributeAccessorHelpers::CreateAccessor(VectorAttribute, PointData->Metadata);
+			StringAccessor = PCGAttributeAccessorHelpers::CreateAccessor(StringAttribute, PointData->Metadata);
 		}
 
 		void ResetTempFloats()
@@ -52,9 +53,9 @@ namespace FPCGAttributeAccessorTestHelpers
 		}
 
 		TObjectPtr<UPCGPointData> PointData;
-		FPCGMetadataAttribute<float> FloatAttribute;
-		FPCGMetadataAttribute<FVector> VectorAttribute;
-		FPCGMetadataAttribute<FString> StringAttribute;
+		FPCGMetadataAttribute<float>* FloatAttribute;
+		FPCGMetadataAttribute<FVector>* VectorAttribute;
+		FPCGMetadataAttribute<FString>* StringAttribute;
 
 		TUniquePtr<IPCGAttributeAccessor> FloatAccessor;
 		TUniquePtr<IPCGAttributeAccessor> VectorAccessor;
@@ -141,7 +142,7 @@ bool FPCGAttributeAccessorSingleGetDefaultTest::RunTest(const FString& Parameter
 bool FPCGAttributeAccessorSingleGetTest::RunTest(const FString& Parameters)
 {
 	FPCGAttributeAccessorTestHelpers::AttributeData AttributeData;
-	FPCGAttributeAccessorKeysEntries Keys(AttributeData.FloatAttribute.GetMetadata());
+	FPCGAttributeAccessorKeysEntries Keys(AttributeData.FloatAttribute->GetMetadata());
 
 	if (!TestEqual("Right number of entries", Keys.GetNum(), FPCGAttributeAccessorTestHelpers::ValuesSize))
 	{
@@ -181,7 +182,7 @@ bool FPCGAttributeAccessorSingleGetTest::RunTest(const FString& Parameters)
 bool FPCGAttributeAccessorSingleGetOutsideRangeTest::RunTest(const FString& Parameters)
 {
 	FPCGAttributeAccessorTestHelpers::AttributeData AttributeData;
-	FPCGAttributeAccessorKeysEntries Keys(AttributeData.FloatAttribute.GetMetadata());
+	FPCGAttributeAccessorKeysEntries Keys(AttributeData.FloatAttribute->GetMetadata());
 
 	float DefaultFloat = 1.0f;
 
@@ -255,7 +256,7 @@ bool FPCGAttributeAccessorGetRangeDefaultTest::RunTest(const FString& Parameters
 bool FPCGAttributeAccessorGetRangeTest::RunTest(const FString& Parameters)
 {
 	FPCGAttributeAccessorTestHelpers::AttributeData AttributeData;
-	FPCGAttributeAccessorKeysEntries Keys(AttributeData.FloatAttribute.GetMetadata());
+	FPCGAttributeAccessorKeysEntries Keys(AttributeData.FloatAttribute->GetMetadata());
 
 	int32 Index = 0;
 
@@ -309,7 +310,7 @@ bool FPCGAttributeAccessorGetRangeTest::RunTest(const FString& Parameters)
 bool FPCGAttributeAccessorGetRangeOutsideRangeTest::RunTest(const FString& Parameters)
 {
 	FPCGAttributeAccessorTestHelpers::AttributeData AttributeData;
-	FPCGAttributeAccessorKeysEntries Keys(AttributeData.FloatAttribute.GetMetadata());
+	FPCGAttributeAccessorKeysEntries Keys(AttributeData.FloatAttribute->GetMetadata());
 
 	int32 OutsideRangeIndex = FPCGAttributeAccessorTestHelpers::ValuesSize;
 
@@ -369,19 +370,19 @@ bool FPCGAttributeAccessorSingleSetDefaultTest::RunTest(const FString& Parameter
 	FString NewDefaultString(TEXT("AAA"));
 
 	if (!TestTrue("Set new float default value, with flag", AttributeData.FloatAccessor->Set(NewDefaultFloat, DefaultEntry, EPCGAttributeAccessorFlags::AllowSetDefaultValue)
-		|| !TestEqual("New default float", AttributeData.FloatAttribute.GetValue(PCGDefaultValueKey), NewDefaultFloat)))
+		|| !TestEqual("New default float", AttributeData.FloatAttribute->GetValue(PCGDefaultValueKey), NewDefaultFloat)))
 	{
 		return false;
 	}
 
 	if (!TestTrue("Set new vector default value, with flag", AttributeData.VectorAccessor->Set(NewDefaultVector, DefaultEntry, EPCGAttributeAccessorFlags::AllowSetDefaultValue)
-		|| !TestEqual("New default vector", AttributeData.VectorAttribute.GetValue(PCGDefaultValueKey), NewDefaultVector)))
+		|| !TestEqual("New default vector", AttributeData.VectorAttribute->GetValue(PCGDefaultValueKey), NewDefaultVector)))
 	{
 		return false;
 	}
 
 	if (!TestTrue("Set new string default value, with flag", AttributeData.StringAccessor->Set(NewDefaultString, DefaultEntry, EPCGAttributeAccessorFlags::AllowSetDefaultValue)
-		|| !TestEqual("New default string", AttributeData.StringAttribute.GetValue(PCGDefaultValueKey), NewDefaultString)))
+		|| !TestEqual("New default string", AttributeData.StringAttribute->GetValue(PCGDefaultValueKey), NewDefaultString)))
 	{
 		return false;
 	}
@@ -392,22 +393,22 @@ bool FPCGAttributeAccessorSingleSetDefaultTest::RunTest(const FString& Parameter
 bool FPCGAttributeAccessorSingleSetTest::RunTest(const FString& Parameters)
 {
 	FPCGAttributeAccessorTestHelpers::AttributeData AttributeData;
-	FPCGAttributeAccessorKeysEntries Keys(const_cast<UPCGMetadata*>(AttributeData.FloatAttribute.GetMetadata()));
+	FPCGAttributeAccessorKeysEntries Keys(const_cast<FPCGMetadataLayer*>(AttributeData.FloatAttribute->GetMetadataLayer()));
 
 	if (!TestTrue("Set float attribute", AttributeData.FloatAccessor->Set(*AttributeData.TempFloats, Keys))
-		|| !TestEqual("float value 0", AttributeData.FloatAttribute.GetValueFromItemKey(0), *AttributeData.TempFloats))
+		|| !TestEqual("float value 0", AttributeData.FloatAttribute->GetValueFromItemKey(0), *AttributeData.TempFloats))
 	{
 		return false;
 	}
 
 	if (!TestTrue("Set vector attribute", AttributeData.VectorAccessor->Set(*AttributeData.TempVectors, /*Index=*/ 1, Keys))
-		|| !TestEqual("vector value 1", AttributeData.VectorAttribute.GetValueFromItemKey(1), *AttributeData.TempVectors))
+		|| !TestEqual("vector value 1", AttributeData.VectorAttribute->GetValueFromItemKey(1), *AttributeData.TempVectors))
 	{
 		return false;
 	}
 
 	if (!TestTrue("Set string attribute", AttributeData.StringAccessor->Set(*AttributeData.TempStrings, /*Index=*/ 2, Keys))
-		|| !TestEqual("string value 2", AttributeData.StringAttribute.GetValueFromItemKey(2), *AttributeData.TempStrings))
+		|| !TestEqual("string value 2", AttributeData.StringAttribute->GetValueFromItemKey(2), *AttributeData.TempStrings))
 	{
 		return false;
 	}
@@ -430,9 +431,9 @@ bool FPCGAttributeAccessorSingleInvalidKeySetTest::RunTest(const FString& Parame
 	float NewDefaultFloat = 0.5f;
 
 	if (!TestTrue("Set new float default value, without flag", AttributeData.FloatAccessor->Set(NewDefaultFloat, DefaultEntry)
-		|| !TestEqual("Default float still default", AttributeData.FloatAttribute.GetValue(PCGDefaultValueKey), 0.0f))
+		|| !TestEqual("Default float still default", AttributeData.FloatAttribute->GetValue(PCGDefaultValueKey), 0.0f))
 		|| !TestEqual("New entry added", AttributeData.PointData->Metadata->GetItemCountForChild(), FPCGAttributeAccessorTestHelpers::ValuesSize + 1)
-		|| !TestEqual("New float value", AttributeData.FloatAttribute.GetValueFromItemKey(FPCGAttributeAccessorTestHelpers::ValuesSize), NewDefaultFloat))
+		|| !TestEqual("New float value", AttributeData.FloatAttribute->GetValueFromItemKey(FPCGAttributeAccessorTestHelpers::ValuesSize), NewDefaultFloat))
 	{
 		return false;
 	}
@@ -443,30 +444,30 @@ bool FPCGAttributeAccessorSingleInvalidKeySetTest::RunTest(const FString& Parame
 bool FPCGAttributeAccessorSetRangeTest::RunTest(const FString& Parameters)
 {
 	FPCGAttributeAccessorTestHelpers::AttributeData AttributeData;
-	FPCGAttributeAccessorKeysEntries Keys(const_cast<UPCGMetadata*>(AttributeData.FloatAttribute.GetMetadata()));
+	FPCGAttributeAccessorKeysEntries Keys(const_cast<FPCGMetadataLayer*>(AttributeData.FloatAttribute->GetMetadataLayer()));
 
 	const int32 Index = 0;
 
 	if (!TestTrue("Set float attribute", AttributeData.FloatAccessor->SetRange(FPCGAttributeAccessorTestHelpers::MakeConstArrayView(AttributeData.TempFloats), Index, Keys))
-		|| !TestEqual("float value 0", AttributeData.FloatAttribute.GetValueFromItemKey(0), AttributeData.TempFloats[0])
-		|| !TestEqual("float value 1", AttributeData.FloatAttribute.GetValueFromItemKey(1), AttributeData.TempFloats[1])
-		|| !TestEqual("float value 2", AttributeData.FloatAttribute.GetValueFromItemKey(2), AttributeData.TempFloats[2]))
+		|| !TestEqual("float value 0", AttributeData.FloatAttribute->GetValueFromItemKey(0), AttributeData.TempFloats[0])
+		|| !TestEqual("float value 1", AttributeData.FloatAttribute->GetValueFromItemKey(1), AttributeData.TempFloats[1])
+		|| !TestEqual("float value 2", AttributeData.FloatAttribute->GetValueFromItemKey(2), AttributeData.TempFloats[2]))
 	{
 		return false;
 	}
 
 	if (!TestTrue("Set vector attribute", AttributeData.VectorAccessor->SetRange(FPCGAttributeAccessorTestHelpers::MakeConstArrayView(AttributeData.TempVectors), Index, Keys))
-		|| !TestEqual("vector value 0", AttributeData.VectorAttribute.GetValueFromItemKey(0), AttributeData.TempVectors[0])
-		|| !TestEqual("vector value 1", AttributeData.VectorAttribute.GetValueFromItemKey(1), AttributeData.TempVectors[1])
-		|| !TestEqual("vector value 2", AttributeData.VectorAttribute.GetValueFromItemKey(2), AttributeData.TempVectors[2]))
+		|| !TestEqual("vector value 0", AttributeData.VectorAttribute->GetValueFromItemKey(0), AttributeData.TempVectors[0])
+		|| !TestEqual("vector value 1", AttributeData.VectorAttribute->GetValueFromItemKey(1), AttributeData.TempVectors[1])
+		|| !TestEqual("vector value 2", AttributeData.VectorAttribute->GetValueFromItemKey(2), AttributeData.TempVectors[2]))
 	{
 		return false;
 	}
 
 	if (!TestTrue("Set string attribute", AttributeData.StringAccessor->SetRange(FPCGAttributeAccessorTestHelpers::MakeConstArrayView(AttributeData.TempStrings), Index, Keys))
-		|| !TestEqual("string value 0", AttributeData.StringAttribute.GetValueFromItemKey(0), AttributeData.TempStrings[0])
-		|| !TestEqual("string value 1", AttributeData.StringAttribute.GetValueFromItemKey(1), AttributeData.TempStrings[1])
-		|| !TestEqual("string value 2", AttributeData.StringAttribute.GetValueFromItemKey(2), AttributeData.TempStrings[2]))
+		|| !TestEqual("string value 0", AttributeData.StringAttribute->GetValueFromItemKey(0), AttributeData.TempStrings[0])
+		|| !TestEqual("string value 1", AttributeData.StringAttribute->GetValueFromItemKey(1), AttributeData.TempStrings[1])
+		|| !TestEqual("string value 2", AttributeData.StringAttribute->GetValueFromItemKey(2), AttributeData.TempStrings[2]))
 	{
 		return false;
 	}
@@ -477,7 +478,7 @@ bool FPCGAttributeAccessorSetRangeTest::RunTest(const FString& Parameters)
 bool FPCGAttributeAccessorSetRangeOutsideRangeTest::RunTest(const FString& Parameters)
 {
 	FPCGAttributeAccessorTestHelpers::AttributeData AttributeData;
-	FPCGAttributeAccessorKeysEntries Keys(AttributeData.FloatAttribute.GetMetadata());
+	FPCGAttributeAccessorKeysEntries Keys(AttributeData.FloatAttribute->GetMetadata());
 
 	const int32 Index = 0;
 
@@ -522,7 +523,7 @@ bool FPCGAttributeAccessorSingleGetPointsTest::RunTest(const FString& Parameters
 	for (int32 i = 0; i < 4; ++i)
 	{
 		if (!TestTrue(FString::Printf(TEXT("Get Float value %d"), i), AttributeData.FloatAccessor->Get(*AttributeData.TempFloats, i, Keys))
-			|| !TestEqual(FString::Printf(TEXT("Float value %d"), i), *AttributeData.TempFloats, AttributeData.FloatAttribute.GetValueFromItemKey(Points[i].MetadataEntry)))
+			|| !TestEqual(FString::Printf(TEXT("Float value %d"), i), *AttributeData.TempFloats, AttributeData.FloatAttribute->GetValueFromItemKey(Points[i].MetadataEntry)))
 		{
 			return false;
 		}
@@ -580,7 +581,7 @@ bool FPCGAttributeAccessorSingleSetPointsTest::RunTest(const FString& Parameters
 	for (int32 i = 0; i < 4; ++i)
 	{
 		if (!TestTrue(FString::Printf(TEXT("Set Float value %d"), i), AttributeData.FloatAccessor->Set(TempFloats[i], i, Keys))
-			|| !TestEqual(FString::Printf(TEXT("Float value %d"), i), AttributeData.FloatAttribute.GetValueFromItemKey(Points[i].MetadataEntry), TempFloats[i]))
+			|| !TestEqual(FString::Printf(TEXT("Float value %d"), i), AttributeData.FloatAttribute->GetValueFromItemKey(Points[i].MetadataEntry), TempFloats[i]))
 		{
 			return false;
 		}
@@ -618,10 +619,10 @@ bool FPCGAttributeAccessorSetRangePointsTest::RunTest(const FString& Parameters)
 	const int32 Index = 0;
 
 	if (!TestTrue("Set Float attribute full range", AttributeData.FloatAccessor->SetRange(FPCGAttributeAccessorTestHelpers::MakeConstArrayView(TempFloats), Index, Keys))
-		|| !TestEqual("Float value 0", AttributeData.FloatAttribute.GetValueFromItemKey(Points[0].MetadataEntry), TempFloats[0])
-		|| !TestEqual("Float value 1", AttributeData.FloatAttribute.GetValueFromItemKey(Points[1].MetadataEntry), TempFloats[1])
-		|| !TestEqual("Float value 2", AttributeData.FloatAttribute.GetValueFromItemKey(Points[2].MetadataEntry), TempFloats[2])
-		|| !TestEqual("Float value 3", AttributeData.FloatAttribute.GetValueFromItemKey(Points[3].MetadataEntry), TempFloats[3]))
+		|| !TestEqual("Float value 0", AttributeData.FloatAttribute->GetValueFromItemKey(Points[0].MetadataEntry), TempFloats[0])
+		|| !TestEqual("Float value 1", AttributeData.FloatAttribute->GetValueFromItemKey(Points[1].MetadataEntry), TempFloats[1])
+		|| !TestEqual("Float value 2", AttributeData.FloatAttribute->GetValueFromItemKey(Points[2].MetadataEntry), TempFloats[2])
+		|| !TestEqual("Float value 3", AttributeData.FloatAttribute->GetValueFromItemKey(Points[3].MetadataEntry), TempFloats[3]))
 	{
 		return false;
 	}
diff --git a/Engine/Plugins/PCG/Source/PCG/Public/Data/PCGBasePointData.h b/Engine/Plugins/PCG/Source/PCG/Public/Data/PCGBasePointData.h
--- a/Engine/Plugins/PCG/Source/PCG/Public/Data/PCGBasePointData.h
+++ b/Engine/Plugins/PCG/Source/PCG/Public/Data/PCGBasePointData.h
@@ -15,6 +15,8 @@ class PCG_API UPCGBasePointData : public UPCGSpatialData
 	GENERATED_BODY()
 
 public:
+	UPCGBasePointData(const FObjectInitializer& ObjectInitializer);
+	
 	virtual bool IsValidRef(const PCGPointOctree::FPointRef& InPointRef) const PURE_VIRTUAL(UPCGBasePointData::IsValidRef, return false;);
 
 	virtual bool IsPointOctreeDirty() const PURE_VIRTUAL(UPCGBasePointData::IsPointOctreeDirty, return false;);
@@ -41,4 +43,8 @@ public:
 	virtual FVector GetScaledExtents(int32 InPointIndex) const PURE_VIRTUAL(UPCGBasePointData::GetScaledExtents, return FVector::ZeroVector;);
 	virtual FVector GetLocalSize(int32 InPointIndex) const PURE_VIRTUAL(UPCGBasePointData::GetLocalSize, return FVector::ZeroVector;);
 	virtual FVector GetScaledLocalSize(int32 InPointIndex) const PURE_VIRTUAL(UPCGBasePointData::GetScaledLocalSize, return FVector::ZeroVector;);
+
+	/** Metadata specific */
+	virtual FPCGMetadataLayerID GetDefaultMetadataLayerID() const override { return PCGMetadataLayerID::Elements; }
+	virtual TArray<FPCGMetadataLayerID> GetAllSupportedMetadataLayerIDs() const override { return {PCGMetadataLayerID::Data, PCGMetadataLayerID::Elements}; }
 };
\ No newline at end of file
diff --git a/Engine/Plugins/PCG/Source/PCG/Public/Data/PCGSpatialData.h b/Engine/Plugins/PCG/Source/PCG/Public/Data/PCGSpatialData.h
--- a/Engine/Plugins/PCG/Source/PCG/Public/Data/PCGSpatialData.h
+++ b/Engine/Plugins/PCG/Source/PCG/Public/Data/PCGSpatialData.h
@@ -154,12 +154,6 @@ public:
 
 	virtual UPCGDifferenceData* Subtract(FPCGContext* InContext, const UPCGSpatialData* InOther) const;
 
-	UFUNCTION(BlueprintCallable, Category = Metadata)
-	virtual const UPCGMetadata* ConstMetadata() const override { return Metadata; }
-
-	UFUNCTION(BlueprintCallable, Category = Metadata)
-	virtual UPCGMetadata* MutableMetadata() override { return Metadata; }
-
 	UFUNCTION(BlueprintCallable, Category = Metadata, meta=(DeprecatedFunction, DeprecationMessage = "The Create Empty Metadata function is not needed anymore - it can safely be removed"))
 	UPCGMetadata* CreateEmptyMetadata();
 
@@ -179,10 +173,6 @@ public:
 	UPROPERTY(BlueprintReadWrite, EditAnywhere, Category = Data)
 	bool bKeepZeroDensityPoints = false;
 
-	// Not accessible through blueprint to make sure the constness is preserved
-	UPROPERTY(VisibleAnywhere, Category = Metadata)
-	TObjectPtr<UPCGMetadata> Metadata = nullptr;
-
 protected:
 	UE_DEPRECATED(5.5, "Call/Implement version with FPCGContext parameter")
 	virtual UPCGSpatialData* CopyInternal() const { return nullptr; }
diff --git a/Engine/Plugins/PCG/Source/PCG/Public/Helpers/PCGMetadataHelpers.h b/Engine/Plugins/PCG/Source/PCG/Public/Helpers/PCGMetadataHelpers.h
--- a/Engine/Plugins/PCG/Source/PCG/Public/Helpers/PCGMetadataHelpers.h
+++ b/Engine/Plugins/PCG/Source/PCG/Public/Helpers/PCGMetadataHelpers.h
@@ -4,10 +4,12 @@
 
 #include "Metadata/PCGAttributePropertySelector.h"
 #include "Metadata/PCGMetadataAttributeTraits.h"
+#include "Metadata/PCGMetadataCommon.h"
 
 #include "HAL/Platform.h"
 
 class FPCGMetadataAttributeBase;
+class FPCGMetadataLayer;
 class UPCGData;
 class UPCGMetadata;
 struct FPCGContext;
@@ -17,7 +19,12 @@ template <typename FuncType> class TFunction;
 namespace PCGMetadataHelpers
 {
 	PCG_API bool HasSameRoot(const UPCGMetadata* Metadata1, const UPCGMetadata* Metadata2);
+	PCG_API bool HasSameRoot(const FPCGMetadataLayer* Metadata1, const FPCGMetadataLayer* Metadata2);
 	PCG_API const UPCGMetadata* GetParentMetadata(const UPCGMetadata* Metadata);
+	PCG_API const FPCGMetadataLayer* GetParentMetadata(const FPCGMetadataLayer* Metadata);
+
+	// Utility function to break circular dependency for UPCGMetadata deprecation in PCGMetadataAttributeTpl.
+	PCG_API FPCGMetadataLayer* GetDefaultMetadataLayer(UPCGMetadata* InMetadata);
 
 	// Helpers functions to cast in spatial/param and return metadata. Nullptr if data doesn't have metadata
 	PCG_API const UPCGMetadata* GetConstMetadata(const UPCGData* InData);
@@ -58,4 +65,8 @@ namespace PCGMetadataHelpers
 
 	/** Copy all the attributes coming from Source Data into Target Data. */
 	PCG_API bool CopyAllAttributes(const UPCGData* SourceData, UPCGData* TargetData, FPCGContext* OptionalContext = nullptr);
+
+	PCG_API void SetPointAttributes(FPCGMetadataLayer* InOutMetadata, const TArrayView<const FPCGPoint>& InPoints, const FPCGMetadataLayer* InMetadata, const TArrayView<FPCGPoint>& OutPoints, FPCGContext* OptionalContext = nullptr);
+
+	PCG_API void ComputePointWeightedAttribute(FPCGMetadataLayer* InOutMetadata, FPCGPoint& OutPoint, const TArrayView<TPair<const FPCGPoint*, float>>& InWeightedPoints, const FPCGMetadataLayer* InMetadata);
 }
diff --git a/Engine/Plugins/PCG/Source/PCG/Public/Metadata/Accessors/PCGAttributeAccessor.h b/Engine/Plugins/PCG/Source/PCG/Public/Metadata/Accessors/PCGAttributeAccessor.h
--- a/Engine/Plugins/PCG/Source/PCG/Public/Metadata/Accessors/PCGAttributeAccessor.h
+++ b/Engine/Plugins/PCG/Source/PCG/Public/Metadata/Accessors/PCGAttributeAccessor.h
@@ -5,9 +5,8 @@
 #include "IPCGAttributeAccessorTpl.h"
 
 #include "Metadata/PCGMetadata.h"
-
-template <typename T> class FPCGMetadataAttribute;
-
+#include "Metadata/PCGMetadataAttributeTpl.h"
+#include "Metadata/PCGMetadataLayer.h"
 
 /**
 * Templated accessor class for attributes. Will wrap around a typed attribute.
@@ -21,22 +20,32 @@ public:
 	using Super = IPCGAttributeAccessorT<FPCGAttributeAccessor<T>>;
 
 	// Can't write if metadata is null
-	FPCGAttributeAccessor(FPCGMetadataAttribute<T>* InAttribute, UPCGMetadata* InMetadata, bool bForceReadOnly = false)
-		: Super(/*bInReadOnly=*/ InMetadata == nullptr || bForceReadOnly)
+	FPCGAttributeAccessor(FPCGMetadataAttribute<T>* InAttribute, FPCGMetadataLayer* InMetadataLayer, bool bForceReadOnly = false)
+		: Super(/*bInReadOnly=*/ InMetadataLayer == nullptr || bForceReadOnly)
 		, Attribute(InAttribute)
-		, Metadata(InMetadata)
+		, MetadataLayer(InMetadataLayer)
 	{
-		check(InAttribute);
+		check(InAttribute && (!InMetadataLayer || InAttribute->GetMetadataLayer() == InMetadataLayer));
 	}
 
-	FPCGAttributeAccessor(const FPCGMetadataAttribute<T>* InAttribute, const UPCGMetadata* InMetadata, bool bForceReadOnly = false)
+	FPCGAttributeAccessor(const FPCGMetadataAttribute<T>* InAttribute, const FPCGMetadataLayer* InMetadataLayer, bool bForceReadOnly = false)
 		: Super(/*bInReadOnly=*/ true)
 		, Attribute(const_cast<FPCGMetadataAttribute<T>*>(InAttribute))
-		, Metadata(const_cast<UPCGMetadata*>(InMetadata))
+		, MetadataLayer(const_cast<FPCGMetadataLayer*>(InMetadataLayer))
 	{
-		check(InAttribute);
+		check(InAttribute && (!InMetadataLayer || InAttribute->GetMetadataLayer() == InMetadataLayer));
 	}
 
+	UE_DEPRECATED(5.6, "Use the FPCGMetadataLayer version")
+	FPCGAttributeAccessor(FPCGMetadataAttribute<T>* InAttribute, UPCGMetadata* InMetadata, bool bForceReadOnly = false)
+		: FPCGAttributeAccessor(InAttribute, InMetadata->GetDefaultMetadataLayer(), bForceReadOnly)
+	{}
+
+	UE_DEPRECATED(5.6, "Use the FPCGMetadataLayer version")
+	FPCGAttributeAccessor(FPCGMetadataAttribute<T>* InAttribute, const UPCGMetadata* InMetadata, bool bForceReadOnly = false)
+		: FPCGAttributeAccessor(InAttribute, InMetadata->GetConstDefaultMetadataLayer(), bForceReadOnly)
+	{}
+
 	bool GetRangeImpl(TArrayView<T> OutValues, int32 Index, const IPCGAttributeAccessorKeys& Keys) const
 	{
 		TArray<const PCGMetadataEntryKey*, TInlineAllocator<256>> EntryKeyPtrs;
@@ -131,7 +140,7 @@ public:
 		for (int EntryIndex = 0; EntryIndex < EntryKeysView.Num(); ++EntryIndex)
 		{
 			PCGMetadataEntryKey& EntryKey = *EntryKeysView[EntryIndex];
-			if (EntryKey == PCGInvalidEntryKey || (EntryKey < Metadata->GetItemKeyCountForParent() && !bCanReuseEntryKeys))
+			if (EntryKey == PCGInvalidEntryKey || (EntryKey < MetadataLayer->GetItemKeyCountForParent() && !bCanReuseEntryKeys))
 			{
 				EntriesToSet.Add(&EntryKey);
 			}
@@ -140,7 +149,7 @@ public:
 		if (!EntriesToSet.IsEmpty())
 		{
 			TRACE_CPUPROFILER_EVENT_SCOPE(FPCGAttributeAccessor::Prepare::AddEntriesInPlace);
-			Metadata->AddEntriesInPlace(EntriesToSet);
+			MetadataLayer->AddEntriesInPlace(EntriesToSet);
 		}
 
 		if (bPreallocateValues && ensure(!bWasPrepared))
@@ -168,6 +177,6 @@ public:
 
 private:
 	FPCGMetadataAttribute<T>* Attribute = nullptr;
-	UPCGMetadata* Metadata = nullptr;
+	FPCGMetadataLayer* MetadataLayer = nullptr;
 	bool bWasPrepared = false;
 };
diff --git a/Engine/Plugins/PCG/Source/PCG/Public/Metadata/Accessors/PCGAttributeAccessorHelpers.h b/Engine/Plugins/PCG/Source/PCG/Public/Metadata/Accessors/PCGAttributeAccessorHelpers.h
--- a/Engine/Plugins/PCG/Source/PCG/Public/Metadata/Accessors/PCGAttributeAccessorHelpers.h
+++ b/Engine/Plugins/PCG/Source/PCG/Public/Metadata/Accessors/PCGAttributeAccessorHelpers.h
@@ -68,6 +68,7 @@ namespace PCGAttributeAccessorHelpers
 	* Creates a const accessor to an attribute without requiring a selector.
 	*/
 	PCG_API TUniquePtr<const IPCGAttributeAccessor> CreateConstAccessor(const FPCGMetadataAttributeBase* InAttribute, const UPCGMetadata* InMetadata, bool bQuiet = false);
+	PCG_API TUniquePtr<const IPCGAttributeAccessor> CreateConstAccessor(const FPCGMetadataAttributeBase* InAttribute, const FPCGMetadataLayer* InMetadata, bool bQuiet = false);
 
 	/**
 	* Creates an accessor to the property or attribute pointed at by the InSelector.
@@ -90,6 +91,7 @@ namespace PCGAttributeAccessorHelpers
 	* Creates an accessor to an attribute without requiring a selector.
 	*/
 	PCG_API TUniquePtr<IPCGAttributeAccessor> CreateAccessor(FPCGMetadataAttributeBase* InAttribute, UPCGMetadata* InMetadata, bool bQuiet = false);
+	PCG_API TUniquePtr<IPCGAttributeAccessor> CreateAccessor(FPCGMetadataAttributeBase* InAttribute, FPCGMetadataLayer* InMetadata, bool bQuiet = false);
 
 	PCG_API TUniquePtr<const IPCGAttributeAccessorKeys> CreateConstKeys(const UPCGData* InData, const FPCGAttributePropertySelector& InSelector);
 	PCG_API TUniquePtr<IPCGAttributeAccessorKeys> CreateKeys(UPCGData* InData, const FPCGAttributePropertySelector& InSelector);
diff --git a/Engine/Plugins/PCG/Source/PCG/Public/Metadata/Accessors/PCGAttributeAccessorKeys.h b/Engine/Plugins/PCG/Source/PCG/Public/Metadata/Accessors/PCGAttributeAccessorKeys.h
--- a/Engine/Plugins/PCG/Source/PCG/Public/Metadata/Accessors/PCGAttributeAccessorKeys.h
+++ b/Engine/Plugins/PCG/Source/PCG/Public/Metadata/Accessors/PCGAttributeAccessorKeys.h
@@ -5,6 +5,7 @@
 #include "Metadata/PCGMetadataCommon.h"
 
 
+class FPCGMetadataLayer;
 struct FPCGPoint;
 
 class FPCGMetadataAttributeBase;
@@ -149,13 +150,16 @@ public:
 	explicit FPCGAttributeAccessorKeysEntries(const UPCGMetadata* Metadata, bool bAddDefaultValueIfEmpty = false);
 	explicit FPCGAttributeAccessorKeysEntries(UPCGMetadata* Metadata, bool bAddDefaultValueIfEmpty = true);
 
+	explicit FPCGAttributeAccessorKeysEntries(const FPCGMetadataLayer* Metadata, bool bAddDefaultValueIfEmpty = false);
+	explicit FPCGAttributeAccessorKeysEntries(FPCGMetadataLayer* Metadata, bool bAddDefaultValueIfEmpty = true);
+
 	virtual int32 GetNum() const override { return Entries.Num(); }
 
 protected:
 	virtual bool GetMetadataEntryKeys(int32 InStart, TArrayView<PCGMetadataEntryKey*> OutEntryKeys) override;
 	virtual bool GetMetadataEntryKeys(int32 InStart, TArrayView<const PCGMetadataEntryKey*> OutEntryKeys) const override;
 
-	void InitializeFromMetadata(const UPCGMetadata* Metadata, bool bAddDefaultValueIfEmpty);
+	void InitializeFromMetadata(const FPCGMetadataLayer* Metadata, bool bAddDefaultValueIfEmpty);
 
 	TArrayView<PCGMetadataEntryKey> Entries;
 	TArray<PCGMetadataEntryKey> ExtractedEntries;
diff --git a/Engine/Plugins/PCG/Source/PCG/Public/Metadata/PCGMetadata.h b/Engine/Plugins/PCG/Source/PCG/Public/Metadata/PCGMetadata.h
--- a/Engine/Plugins/PCG/Source/PCG/Public/Metadata/PCGMetadata.h
+++ b/Engine/Plugins/PCG/Source/PCG/Public/Metadata/PCGMetadata.h
@@ -3,9 +3,12 @@
 #pragma once
 
 #include "PCGCommon.h"
-#include "PCGMetadataCommon.h" // IWYU pragma: keep
 #include "PCGMetadataAttributeTraits.h"
+#include "PCGMetadataCommon.h" // IWYU pragma: keep
 #include "Metadata/PCGMetadataAttributeTpl.h"
+#include "Metadata/PCGMetadataLayer.h"
+
+#include "Misc/SpinLock.h"
 
 #include "PCGMetadata.generated.h"
 
@@ -15,18 +18,29 @@ struct FPCGContext;
 UCLASS(BlueprintType)
 class PCG_API UPCGMetadata : public UObject
 {
+	friend FPCGMetadataLayer;
+	
 	GENERATED_BODY()
 
 public:
-
+	UPCGMetadata(const FObjectInitializer& ObjectInitializer);
+	
 	//~ Begin UObject interface
 	virtual void Serialize(FArchive& InArchive) override;
-	virtual void BeginDestroy() override;
 	//~ End UObject interface
 
+	/** To be called by the outer pcg data to initialize the layers supported for this data. Can mark this one as default, will override previously set default. */
+	void SetupLayer(FPCGMetadataLayerID LayerID, bool bIsDefault);
+
+	/** If we have floating metadata, setup it to match a given data type. */
+	template <typename DataType, std::enable_if_t<std::is_base_of_v<UPCGData, DataType>, int> = 0>
+	void SetupLayersFromPCGDataType() { SetupFromPCGDataType(DataType::StaticClass()); }
+
+	void SetupLayersFromPCGDataType(const TSubclassOf<UPCGData>& PCGDataType);
+
 	/** Initializes the metadata from a parent metadata, if any (can be null). Copies attributes and values. */
 	UFUNCTION(BlueprintCallable, Category = "PCG|Metadata")
-	void Initialize(const UPCGMetadata* InParent) { InitializeWithAttributeFilter(InParent, TSet<FName>()); }
+	void Initialize(const UPCGMetadata* InParent) { InitializeWithAttributeFilter(InParent, TSet<FPCGAttributeIdentifier>()); }
 
 	/** Initializes the metadata from a parent metadata, if any (can be null) with the option to not add attributes from the parent. */
 	void Initialize(const UPCGMetadata* InParent, bool bAddAttributesFromParent);
@@ -39,6 +53,8 @@ public:
 	 */
 	UFUNCTION(BlueprintCallable, Category = "PCG|Metadata")
 	void InitializeWithAttributeFilter(const UPCGMetadata* InParent, const TSet<FName>& InFilteredAttributes, EPCGMetadataFilterMode InFilterMode = EPCGMetadataFilterMode::ExcludeAttributes, EPCGStringMatchingOperator InMatchOperator = EPCGStringMatchingOperator::Equal);
+	
+	void InitializeWithAttributeFilter(const UPCGMetadata* InParent, const TSet<FPCGAttributeIdentifier>& InFilteredAttributes, EPCGMetadataFilterMode InFilterMode = EPCGMetadataFilterMode::ExcludeAttributes, EPCGStringMatchingOperator InMatchOperator = EPCGStringMatchingOperator::Equal);
 
 	/** Initializes the metadata from a parent metadata by copying all attributes to it.
 	* @param InMetadataToCopy Metadata to copy from
@@ -53,10 +69,11 @@ public:
 	* @param InOptionalEntriesToCopy Optional array that contains the keys to copy over. This array order will be respected, so it can also be used to re-order entries. Can be null.
 	*/
 	void InitializeAsCopyWithAttributeFilter(const UPCGMetadata* InMetadataToCopy, const TSet<FName>& InFilteredAttributes, EPCGMetadataFilterMode InFilterMode = EPCGMetadataFilterMode::ExcludeAttributes, const TArray<PCGMetadataEntryKey>* InOptionalEntriesToCopy = nullptr);
+	void InitializeAsCopyWithAttributeFilter(const UPCGMetadata* InMetadataToCopy, const TSet<FPCGAttributeIdentifier>& InFilteredAttributes, EPCGMetadataFilterMode InFilterMode = EPCGMetadataFilterMode::ExcludeAttributes, const TArray<PCGMetadataEntryKey>* InOptionalEntriesToCopy = nullptr);
 
 	/** Creates missing attributes from another metadata if they are not currently present - note that this does not copy values */
 	UFUNCTION(BlueprintCallable, Category = "PCG|Metadata")
-	void AddAttributes(const UPCGMetadata* InOther) { AddAttributesFiltered(InOther, TSet<FName>()); }
+	void AddAttributes(const UPCGMetadata* InOther) { AddAttributesFiltered(InOther, TSet<FPCGAttributeIdentifier>()); }
 
 	/**
 	 * Creates missing attributes from another metadata if they are not currently present - note that this does not copy values.
@@ -67,20 +84,26 @@ public:
 	UFUNCTION(BlueprintCallable, Category = "PCG|Metadata")
 	void AddAttributesFiltered(const UPCGMetadata* InOther, const TSet<FName>& InFilteredAttributes, EPCGMetadataFilterMode InFilterMode = EPCGMetadataFilterMode::ExcludeAttributes, EPCGStringMatchingOperator InMatchOperator = EPCGStringMatchingOperator::Equal);
 
+	void AddAttributesFiltered(const UPCGMetadata* InOther, const TSet<FPCGAttributeIdentifier>& InFilteredAttributes, EPCGMetadataFilterMode InFilterMode = EPCGMetadataFilterMode::ExcludeAttributes, EPCGStringMatchingOperator InMatchOperator = EPCGStringMatchingOperator::Equal);
+	
 	/** Creates missing attribute from another metadata if it is not currently present - note that this does not copy values */
-	UFUNCTION(BlueprintCallable, Category = "PCG|Metadata")
-	void AddAttribute(const UPCGMetadata* InOther, FName AttributeName);
+	UFUNCTION(BlueprintCallable, Category = "PCG|Metadata", meta = (DisplayName = "Add Attribute"))
+	void BP_AddAttribute(const UPCGMetadata* InOther, FName AttributeName);
+	
+	void AddAttribute(const UPCGMetadata* InOther, FPCGAttributeIdentifier AttributeName);
 
 	/** Copies attributes from another metadata, including entries & values. Warning: this is intended when dealing with the same data set */
 	UFUNCTION(BlueprintCallable, Category = "PCG|Metadata|Advanced")
 	void CopyAttributes(const UPCGMetadata* InOther);
 
 	/** Copies an attribute from another metadata, including entries & values. Warning: this is intended when dealing with the same data set */
-	UFUNCTION(BlueprintCallable, Category = "PCG|Metadata|Advanced")
-	void CopyAttribute(const UPCGMetadata* InOther, FName AttributeToCopy, FName NewAttributeName);
+	UFUNCTION(BlueprintCallable, Category = "PCG|Metadata|Advanced", meta = (DisplayName = "Copy Attribute"))
+	void BP_CopyAttribute(const UPCGMetadata* InOther, FName AttributeToCopy, FName NewAttributeName);
 
+	void CopyAttribute(const UPCGMetadata* InOther, FPCGAttributeIdentifier AttributeToCopy, FName NewAttributeName);
+	
 	/** Copies another attribute, with options to keep its parent and copy entries/values */
-	FPCGMetadataAttributeBase* CopyAttribute(const FPCGMetadataAttributeBase* OriginalAttribute, FName NewAttributeName, bool bKeepParent, bool bCopyEntries, bool bCopyValues);
+	FPCGMetadataAttributeBase* CopyAttribute(const FPCGMetadataAttributeBase* OriginalAttribute, FPCGAttributeIdentifier NewAttributeName, bool bKeepParent, bool bCopyEntries, bool bCopyValues);
 
 	/** Returns this metadata's parent */
 	TWeakObjectPtr<const UPCGMetadata> GetParentPtr() const { return Parent; }
@@ -98,6 +121,7 @@ public:
 	/** Unparents current metadata, flatten attribute and only keep the entries specified. Return true if something has changed and keys needs be updated. */
 	bool FlattenAndCompress(const TArray<PCGMetadataEntryKey>& InEntryKeysToKeep);
 
+	// TODO: Will be moved at the end of the class before commit, as they are functions planned to be deprecated.
 	UFUNCTION(BlueprintCallable, Category = "PCG|Metadata")
 	UPARAM(DisplayName = "Metadata") UPCGMetadata* CreateInteger32Attribute(FName AttributeName, int32 DefaultValue, bool bAllowsInterpolation, bool bOverrideParent = true);
 
@@ -149,7 +173,7 @@ public:
 	* @param Property: The property to set from
 	* @returns true if the attribute creation succeeded
 	*/
-	bool CreateAttributeFromProperty(FName AttributeName, const UObject* Object, const FProperty* Property);
+	bool CreateAttributeFromProperty(FPCGAttributeIdentifier AttributeName, const UObject* Object, const FProperty* Property);
 
 	/** Creates an attribute given a property.
 	* @param AttributeName: Target attribute to create
@@ -157,7 +181,7 @@ public:
 	* @param Property: The property to set from
 	* @returns true if the attribute creation succeeded
 	*/
-	bool CreateAttributeFromDataProperty(FName AttributeName, const void* Data, const FProperty* Property);
+	bool CreateAttributeFromDataProperty(FPCGAttributeIdentifier AttributeName, const void* Data, const FProperty* Property);
 
 	/** Set an attribute given a property and its value.
 	* @param AttributeName: Target attribute to set the property's value to
@@ -167,7 +191,7 @@ public:
 	* @param bCreate: If true and the attribute doesn't exists, it will create an attribute based on the property type
 	* @returns true if the attribute creation (if required) and the value set succeeded
 	*/
-	bool SetAttributeFromProperty(FName AttributeName, PCGMetadataEntryKey& EntryKey, const UObject* Object, const FProperty* Property, bool bCreate);
+	bool SetAttributeFromProperty(FPCGAttributeIdentifier AttributeName, PCGMetadataEntryKey& EntryKey, const UObject* Object, const FProperty* Property, bool bCreate);
 
 	/** Set an attribute given a property and its value.
 	* @param AttributeName: Target attribute to set the property's value to
@@ -177,15 +201,16 @@ public:
 	* @param bCreate: If true and the attribute doesn't exists, it will create an attribute based on the property type
 	* @returns true if the attribute creation (if required) and the value set succeeded
 	*/
-	bool SetAttributeFromDataProperty(FName AttributeName, PCGMetadataEntryKey& EntryKey, const void* Data, const FProperty* Property, bool bCreate);
+	bool SetAttributeFromDataProperty(FPCGAttributeIdentifier AttributeName, PCGMetadataEntryKey& EntryKey, const void* Data, const FProperty* Property, bool bCreate);
+	// End of move
 
 	/** Get attributes */
-	FPCGMetadataAttributeBase* GetMutableAttribute(FName AttributeName);
-	const FPCGMetadataAttributeBase* GetConstAttribute(FName AttributeName) const;
-	const FPCGMetadataAttributeBase* GetConstAttributeById(int32 AttributeId) const;
+	FPCGMetadataAttributeBase* GetMutableAttribute(FPCGAttributeIdentifier AttributeName);
+	const FPCGMetadataAttributeBase* GetConstAttribute(FPCGAttributeIdentifier AttributeName) const;
 
-	UFUNCTION(BlueprintCallable, Category = "PCG|Metadata")
-	bool HasAttribute(FName AttributeName) const;
+	UFUNCTION(BlueprintCallable, Category = "PCG|Metadata", meta = (DisplayName = "Has Attribute"))
+	bool BP_HasAttribute(FName AttributeName) const;
+	bool HasAttribute(FPCGAttributeIdentifier AttributeName) const;
 
 	UFUNCTION(BlueprintCallable, Category = "PCG|Metadata")
 	bool HasCommonAttributes(const UPCGMetadata* InMetadata) const;
@@ -195,39 +220,49 @@ public:
 	int32 GetAttributeCount() const;
 
 	template <typename T>
-	FPCGMetadataAttribute<T>* GetMutableTypedAttribute(FName AttributeName);
+	FPCGMetadataAttribute<T>* GetMutableTypedAttribute(FPCGAttributeIdentifier AttributeName);
 
 	template <typename T>
-	FPCGMetadataAttribute<T>* GetMutableTypedAttribute_Unsafe(FName AttributeName);
+	FPCGMetadataAttribute<T>* GetMutableTypedAttribute_Unsafe(FPCGAttributeIdentifier AttributeName);
 
 	template <typename T>
-	const FPCGMetadataAttribute<T>* GetConstTypedAttribute(FName AttributeName) const;
+	const FPCGMetadataAttribute<T>* GetConstTypedAttribute(FPCGAttributeIdentifier AttributeName) const;
 
 	UFUNCTION(BlueprintCallable, Category = "PCG|Metadata")
 	void GetAttributes(TArray<FName>& AttributeNames, TArray<EPCGMetadataTypes>& AttributeTypes) const;
 
+	void GetAllAttributes(TArray<FPCGAttributeIdentifier>& AttributeNames, TArray<EPCGMetadataTypes>& AttributeTypes) const;
+
 	/** Returns name of the most recently created attribute, or none if no attributes are present. */
 	FName GetLatestAttributeNameOrNone() const;
 
 	/** Delete/Hide attribute */
 	// Due to stream inheriting, we might want to consider "hiding" parent stream and deleting local streams only
-	UFUNCTION(BlueprintCallable, Category = "PCG|Metadata")
-	void DeleteAttribute(FName AttributeName);
+	UFUNCTION(BlueprintCallable, Category = "PCG|Metadata", meta = (DisplayName = "DeleteAttribute"))
+	void BP_DeleteAttribute(FName AttributeName);
 
+	void DeleteAttribute(FPCGAttributeIdentifier AttributeName);
+	
 	/** Copy attribute */
-	UFUNCTION(BlueprintCallable, Category = "PCG|Metadata")
-	bool CopyExistingAttribute(FName AttributeToCopy, FName NewAttributeName, bool bKeepParent = true);
+	UFUNCTION(BlueprintCallable, Category = "PCG|Metadata", meta = (DisplayName = "CopyExistingAttribute"))
+	bool BP_CopyExistingAttribute(FName AttributeToCopy, FName NewAttributeName, bool bKeepParent = true);
+	
+	bool CopyExistingAttribute(FPCGAttributeIdentifier AttributeToCopy, FName NewAttributeName, bool bKeepParent = true);
 
 	/** Rename attribute */
-	UFUNCTION(BlueprintCallable, Category = "PCG|Metadata")
-	bool RenameAttribute(FName AttributeToRename, FName NewAttributeName);
+	UFUNCTION(BlueprintCallable, Category = "PCG|Metadata", meta = (DisplayName = "RenameAttribute"))
+	bool BP_RenameAttribute(FName AttributeToRename, FName NewAttributeName);
 
+	bool RenameAttribute(FPCGAttributeIdentifier AttributeToRename, FName NewAttributeName);
+	
 	/** Clear/Reinit attribute */
-	UFUNCTION(BlueprintCallable, Category = "PCG|Metadata")
-	void ClearAttribute(FName AttributeToClear);
+	UFUNCTION(BlueprintCallable, Category = "PCG|Metadata", meta = (DisplayName = "ClearAttributeW"))
+	void BP_ClearAttribute(FName AttributeToClear);
+	
+	void ClearAttribute(FPCGAttributeIdentifier AttributeToClear);
 
 	/** Change type of an attribute */
-	bool ChangeAttributeType(FName AttributeName, int16 AttributeNewType);
+	bool ChangeAttributeType(FPCGAttributeIdentifier AttributeName, int16 AttributeNewType);
 
 	/** Adds a unique entry key to the metadata */
 	UFUNCTION(BlueprintCallable, Category = "PCG|Metadata")
@@ -273,6 +308,7 @@ public:
 	void SetAttributes(const TArrayView<const PCGMetadataEntryKey>& InOriginalKeys, const UPCGMetadata* InMetadata, const TArrayView<PCGMetadataEntryKey>* InOutOptionalKeys = nullptr, FPCGContext* OptionalContext = nullptr);
 	void SetAttributes(const TArrayView<const PCGMetadataEntryKey>& InKeys, const UPCGMetadata* InMetadata, const TArrayView<PCGMetadataEntryKey>& OutKeys, FPCGContext* OptionalContext = nullptr);
 
+	// TODO: Will be moved at the end of the class before commit, as they are functions planned to be deprecated.
 	/** Attributes operations - shorthand for points */
 	void MergePointAttributes(const FPCGPoint& InPointA, const FPCGPoint& InPointB, FPCGPoint& OutPoint, EPCGMetadataOp Op);
 	void MergePointAttributesSubset(const FPCGPoint& InPointA, const UPCGMetadata* InMetadataA, const UPCGMetadata* InMetadataSubetA, const FPCGPoint& InPointB, const UPCGMetadata* InMetadataB, const UPCGMetadata* InMetadataSubsetb, FPCGPoint& OutPoint, EPCGMetadataOp Op);
@@ -305,6 +341,7 @@ public:
 
 	void ComputePointWeightedAttribute(FPCGPoint& OutPoint, const TArrayView<TPair<const FPCGPoint*, float>>& InWeightedPoints, const UPCGMetadata* InMetadata);
 	void ComputeWeightedAttribute(PCGMetadataEntryKey& OutKey, const TArrayView<TPair<PCGMetadataEntryKey, float>>& InWeightedKeys, const UPCGMetadata* InMetadata);
+	 // End of move
 
 	int64 GetItemKeyCountForParent() const;
 	int64 GetLocalItemCount() const;
@@ -319,7 +356,7 @@ public:
 	* Return a typed attribute pointer, of the requested type T.
 	*/
 	template<typename T>
-	FPCGMetadataAttribute<T>* CreateAttribute(FName AttributeName, const T& DefaultValue, bool bAllowsInterpolation, bool bOverrideParent);
+	FPCGMetadataAttribute<T>* CreateAttribute(FPCGAttributeIdentifier AttributeName, const T& DefaultValue, bool bAllowsInterpolation, bool bOverrideParent);
 
 	/**
 	* Find or create an attribute. Follows CreateAttribute signature.
@@ -328,7 +365,7 @@ public:
 	* Return a typed attribute pointer, of the requested type T.
 	*/
 	template<typename T>
-	FPCGMetadataAttribute<T>* FindOrCreateAttribute(FName AttributeName, const T& DefaultValue = T{}, bool bAllowsInterpolation = true, bool bOverrideParent = true, bool bOverwriteIfTypeMismatch = true);
+	FPCGMetadataAttribute<T>* FindOrCreateAttribute(FPCGAttributeIdentifier AttributeName, const T& DefaultValue = T{}, bool bAllowsInterpolation = true, bool bOverrideParent = true, bool bOverwriteIfTypeMismatch = true);
 
 	// Need this gymnastic for those 2 functions, because blueprints doesn't support default arguments for Arrays, and we don't want to force c++ user to provide OptionalNewEntriesOrder.
 
@@ -348,22 +385,51 @@ public:
 	UFUNCTION(BlueprintCallable, Category = "PCG|Metadata", meta = (DisplayName = "Initialize As Copy With Attribute Filter", AutoCreateRefTerm = "InFilteredAttributes,InOptionalEntriesToCopy"))
 	void K2_InitializeAsCopyWithAttributeFilter(const UPCGMetadata* InMetadataToCopy, const TSet<FName>& InFilteredAttributes, const TArray<int64>& InOptionalEntriesToCopy, EPCGMetadataFilterMode InFilterMode = EPCGMetadataFilterMode::ExcludeAttributes)
 	{
-		return InitializeAsCopyWithAttributeFilter(InMetadataToCopy, InFilteredAttributes, InFilterMode, !InOptionalEntriesToCopy.IsEmpty() ? &InOptionalEntriesToCopy : nullptr);
+		return InitializeAsCopyWithAttributeFilter(InMetadataToCopy, FPCGAttributeIdentifier::TransformNameSet(InFilteredAttributes), InFilterMode, !InOptionalEntriesToCopy.IsEmpty() ? &InOptionalEntriesToCopy : nullptr);
 	}
 
 	/** Computes Crc from all attributes & keys from outer's data. */
 	void AddToCrc(FArchiveCrc32& Ar, bool bFullDataCrc) const;
 
-protected:
-	FPCGMetadataAttributeBase* CopyAttribute(FName AttributeToCopy, FName NewAttributeName, bool bKeepParent, bool bCopyEntries, bool bCopyValues);
+	/** Expected to be used when metadata layers are already setup. */
+	FPCGMetadataLayer* GetMetadataLayer(const FPCGMetadataLayerID& InMetadataLayerID);
+	
+	/** Expected to be used when metadata layer are already setup. */
+	const FPCGMetadataLayer* GetConstMetadataLayer(const FPCGMetadataLayerID& InMetadataLayerID) const;
+
+	FPCGMetadataLayer* GetDefaultMetadataLayer() { return GetMetadataLayer(PCGMetadataLayerID::Default); }
+	const FPCGMetadataLayer* GetConstDefaultMetadataLayer() const { return GetConstMetadataLayer(PCGMetadataLayerID::Default); }
+
+private:
+	// Utility functions to forward calls to FPCGMetadataLayer.
+	template <typename Func, typename ...Args>
+	decltype(auto) WithMetadataLayer(const FPCGMetadataLayerID& InMetadataLayerID, Func InFunc, Args&& ...InArgs);
+	
+	template <typename Func>
+	decltype(auto) WithMetadataLayer_Lambda(const FPCGMetadataLayerID& InMetadataLayerID, Func InFunc);
+
+	template <typename Func, typename ...Args>
+	decltype(auto) WithConstMetadataLayer(const FPCGMetadataLayerID& InMetadataLayerID, Func InFunc, Args&& ...InArgs) const;
 
-	bool ParentHasAttribute(FName AttributeName) const;
+	template <typename Func>
+	void ForEachValidUniqueConstMetadataLayer(Func InFunc) const;
+
+	template <typename Func>
+	void ForEachValidUniqueMetadataLayer(Func InFunc);
+
+protected:
+	FPCGMetadataLayer* FindOrCreateMetadataLayer(const FPCGMetadataLayerID& InMetadataLayerID);
+	
+	FPCGMetadataLayer* CreateMetadataLayer_Unsafe(const FPCGMetadataLayerID& InMetadataLayerID);
+	
+	FPCGMetadataAttributeBase* CopyAttribute(FPCGAttributeIdentifier AttributeToCopy, FName NewAttributeName, bool bKeepParent, bool bCopyEntries, bool bCopyValues);
 
-	void AddAttributeInternal(FName AttributeName, FPCGMetadataAttributeBase* Attribute);
-	void RemoveAttributeInternal(FName AttributeName);
+	bool ParentHasAttribute(FPCGAttributeIdentifier AttributeName) const;
 
 	void SetLastCachedSelectorOnOwner(FName AttributeName);
 
+	void SetupLayersFromOtherMetadataIfNeeded(const UPCGMetadata* OtherMetadata);
+
 	UPROPERTY()
 	TObjectPtr<const UPCGMetadata> Parent;
 
@@ -372,141 +438,36 @@ protected:
 	UPROPERTY()
 	TSet<TWeakObjectPtr<const UPCGMetadata>> OtherParents;
 
-	TMap<FName, FPCGMetadataAttributeBase*> Attributes;
-	PCGMetadataAttributeKey NextAttributeId = 0;
-
-	TArray<PCGMetadataEntryKey> ParentKeys;
-	int64 ItemKeyOffset = 0;
-
-	mutable FRWLock AttributeLock;
-	mutable FRWLock ItemLock;
-
-	TAtomic<int64> DelayedEntriesIndex = 0;
+	FPCGMetadataLayerID DefaultLayer = PCGMetadataLayerID::Default;
+	TMap<FPCGMetadataLayerID, TSharedPtr<FPCGMetadataLayer>> MetadataLayers;
+	UE::FSpinLock MetadataLayersSpinLock;
 };
 
 template<typename T>
-FPCGMetadataAttribute<T>* UPCGMetadata::CreateAttribute(FName AttributeName, const T& DefaultValue, bool bAllowsInterpolation, bool bOverrideParent)
+FPCGMetadataAttribute<T>* UPCGMetadata::CreateAttribute(FPCGAttributeIdentifier AttributeName, const T& DefaultValue, bool bAllowsInterpolation, bool bOverrideParent)
 {
-	if (!FPCGMetadataAttributeBase::IsValidName(AttributeName))
-	{
-		UE_LOG(LogPCG, Error, TEXT("Attribute name '%s' is invalid"), *AttributeName.ToString());
-		return nullptr;
-	}
-
-	const FPCGMetadataAttributeBase* ParentAttribute = nullptr;
-
-	if (bOverrideParent && Parent)
-	{
-		ParentAttribute = Parent->GetConstAttribute(AttributeName);
-	}
-
-	if (ParentAttribute && (ParentAttribute->GetTypeId() != PCG::Private::MetadataTypes<T>::Id))
-	{
-		// Can't parent if the types do not match
-		ParentAttribute = nullptr;
-	}
-
-	FPCGMetadataAttribute<T>* NewAttribute = new FPCGMetadataAttribute<T>(this, AttributeName, ParentAttribute, DefaultValue, bAllowsInterpolation);
-
-	{
-		FWriteScopeLock WriteLock(AttributeLock);
-
-		if (FPCGMetadataAttributeBase** ExistingAttribute = Attributes.Find(AttributeName))
-		{
-			delete NewAttribute;
-			if ((*ExistingAttribute)->GetTypeId() != PCG::Private::MetadataTypes<T>::Id)
-			{
-				UE_LOG(LogPCG, Error, TEXT("Attribute %s already exists but is not the right type. Abort."), *AttributeName.ToString());
-				return nullptr;
-			}
-			else
-			{
-				UE_LOG(LogPCG, Warning, TEXT("Attribute %s already exists"), *AttributeName.ToString());
-				NewAttribute = static_cast<FPCGMetadataAttribute<T>*>(*ExistingAttribute);
-			}
-		}
-		else
-		{
-			NewAttribute->AttributeId = NextAttributeId++;
-			AddAttributeInternal(AttributeName, NewAttribute);
-
-			// Also when creating an attribute, notify the PCG Data owner that the latest attribute manipulated is this one.
-			SetLastCachedSelectorOnOwner(AttributeName);
-		}
-	}
-
-	return NewAttribute;
+	FPCGMetadataLayer* FoundSubMetadata = FindOrCreateMetadataLayer(AttributeName.MetadataLayer);
+	check(FoundSubMetadata);
+	return FoundSubMetadata->CreateAttribute<T>(AttributeName.Name, DefaultValue, bAllowsInterpolation, bOverrideParent);
 }
 
 template<typename T>
-FPCGMetadataAttribute<T>* UPCGMetadata::GetMutableTypedAttribute_Unsafe(FName AttributeName)
+FPCGMetadataAttribute<T>* UPCGMetadata::GetMutableTypedAttribute_Unsafe(FPCGAttributeIdentifier AttributeName)
 {
-	FPCGMetadataAttribute<T>* Attribute = nullptr;
-
-	if (FPCGMetadataAttributeBase** FoundAttribute = Attributes.Find(AttributeName))
-	{
-		FPCGMetadataAttributeBase* BaseAttribute = *FoundAttribute;
-		if (BaseAttribute && BaseAttribute->GetTypeId() == PCG::Private::MetadataTypes<T>::Id)
-		{
-			Attribute = static_cast<FPCGMetadataAttribute<T>*>(BaseAttribute);
-		}
-	}
-
-	return Attribute;
+	FPCGMetadataLayer* FoundSubMetadata = GetMetadataLayer(AttributeName.MetadataLayer);
+	return FoundSubMetadata ? FoundSubMetadata->GetMutableTypedAttribute_Unsafe<T>(AttributeName.Name) : nullptr;
 }
 
 template<typename T>
-FPCGMetadataAttribute<T>* UPCGMetadata::FindOrCreateAttribute(FName AttributeName, const T& DefaultValue, bool bAllowsInterpolation, bool bOverrideParent, bool bOverwriteIfTypeMismatch)
+FPCGMetadataAttribute<T>* UPCGMetadata::FindOrCreateAttribute(FPCGAttributeIdentifier AttributeName, const T& DefaultValue, bool bAllowsInterpolation, bool bOverrideParent, bool bOverwriteIfTypeMismatch)
 {
-	{
-		FReadScopeLock ScopeLock(AttributeLock);
-		if (FPCGMetadataAttribute<T>* Attribute = GetMutableTypedAttribute_Unsafe<T>(AttributeName))
-		{
-			return Attribute;
-		}
-	}
-
-	FWriteScopeLock ScopeLock(AttributeLock);
-	if (FPCGMetadataAttribute<T>* Attribute = GetMutableTypedAttribute_Unsafe<T>(AttributeName))
-	{
-		return Attribute;
-	}
-
-	// If an attribute with this name exists here, there is a type mismatch.
-	if (FPCGMetadataAttributeBase** FoundAttribute = Attributes.Find(AttributeName))
-	{
-		if (bOverwriteIfTypeMismatch)
-		{
-			delete *FoundAttribute;
-			RemoveAttributeInternal(AttributeName);
-		}
-		else
-		{
-			return nullptr;
-		}
-	}
-
-	// A new attribute will be created.
-	if (!FPCGMetadataAttributeBase::IsValidName(AttributeName))
-	{
-		UE_LOG(LogPCG, Error, TEXT("Attribute name '%s' is invalid"), *AttributeName.ToString());
-		return nullptr;
-	}
-
-	// Parent is const and therefore should be safe
-	const FPCGMetadataAttributeBase* ParentAttribute = (bOverrideParent && Parent) ? Parent->GetConstTypedAttribute<T>(AttributeName) : nullptr;
-	FPCGMetadataAttribute<T>* NewAttribute = new FPCGMetadataAttribute<T>(this, AttributeName, ParentAttribute, DefaultValue, bAllowsInterpolation);
-	NewAttribute->AttributeId = NextAttributeId++;
-	AddAttributeInternal(AttributeName, NewAttribute);
-
-	// Also when creating an attribute, notify the PCG Data owner that the latest attribute manipulated is this one.
-	SetLastCachedSelectorOnOwner(AttributeName);
-
-	return NewAttribute;
+	FPCGMetadataLayer* FoundSubMetadata = FindOrCreateMetadataLayer(AttributeName.MetadataLayer);
+	check(FoundSubMetadata);
+	return FoundSubMetadata->FindOrCreateAttribute<T>(AttributeName.Name, DefaultValue, bAllowsInterpolation, bOverrideParent, bOverwriteIfTypeMismatch);
 }
 
 template <typename T>
-FPCGMetadataAttribute<T>* UPCGMetadata::GetMutableTypedAttribute(FName AttributeName)
+FPCGMetadataAttribute<T>* UPCGMetadata::GetMutableTypedAttribute(FPCGAttributeIdentifier AttributeName)
 {
 	FPCGMetadataAttributeBase* BaseAttribute = GetMutableAttribute(AttributeName);
 	return (BaseAttribute && (BaseAttribute->GetTypeId() == PCG::Private::MetadataTypes<T>::Id))
@@ -515,7 +476,7 @@ FPCGMetadataAttribute<T>* UPCGMetadata::GetMutableTypedAttribute(FName Attribute
 }
 
 template <typename T>
-const FPCGMetadataAttribute<T>* UPCGMetadata::GetConstTypedAttribute(FName AttributeName) const
+const FPCGMetadataAttribute<T>* UPCGMetadata::GetConstTypedAttribute(FPCGAttributeIdentifier AttributeName) const
 {
 	const FPCGMetadataAttributeBase* BaseAttribute = GetConstAttribute(AttributeName);
 	return (BaseAttribute && (BaseAttribute->GetTypeId() == PCG::Private::MetadataTypes<T>::Id))
diff --git a/Engine/Plugins/PCG/Source/PCG/Public/Metadata/PCGMetadataAttribute.h b/Engine/Plugins/PCG/Source/PCG/Public/Metadata/PCGMetadataAttribute.h
--- a/Engine/Plugins/PCG/Source/PCG/Public/Metadata/PCGMetadataAttribute.h
+++ b/Engine/Plugins/PCG/Source/PCG/Public/Metadata/PCGMetadataAttribute.h
@@ -6,6 +6,7 @@
 
 #include "UObject/ObjectPtr.h" // IWYU pragma: keep
 
+class FPCGMetadataLayer;
 class UPCGMetadata;
 
 namespace PCGMetadataAttributeConstants
@@ -20,9 +21,16 @@ class PCG_API FPCGMetadataAttributeBase
 {
 public:
 	FPCGMetadataAttributeBase() = default;
+	FPCGMetadataAttributeBase(FPCGMetadataLayer* InMetadata, FName InName, const FPCGMetadataAttributeBase* InParent, bool bInAllowsInterpolation);
+
+	UE_DEPRECATED(5.6, "Use the version with FPCGMetadataLayer")
 	FPCGMetadataAttributeBase(UPCGMetadata* InMetadata, FName InName, const FPCGMetadataAttributeBase* InParent, bool bInAllowsInterpolation);
 
 	virtual ~FPCGMetadataAttributeBase() = default;
+	
+	virtual void Serialize(FPCGMetadataLayer* InMetadata, FArchive& InArchive);
+
+	UE_DEPRECATED(5.6, "Use the version with FPCGMetadataLayer")
 	virtual void Serialize(UPCGMetadata* InMetadata, FArchive& InArchive);
 
 	/** Unparents current attribute by flattening the values, entries, etc. */
@@ -32,11 +40,17 @@ public:
 
 	/** Remove all entries, values and parenting. */
 	virtual void Reset() = 0;
-
-	const UPCGMetadata* GetMetadata() const { return Metadata; }
+	
+	const UPCGMetadata* GetMetadata() const;
+	const FPCGMetadataLayer* GetMetadataLayer() const { return Metadata; }
+	
 	int16 GetTypeId() const { return TypeId; }
 
+	virtual FPCGMetadataAttributeBase* Copy(FName NewName, FPCGMetadataLayer* InMetadata, bool bKeepParent, bool bCopyEntries = true, bool bCopyValues = true) const = 0;
+
+	UE_DEPRECATED(5.6, "Use the version with FPCGMetadataLayer")
 	virtual FPCGMetadataAttributeBase* Copy(FName NewName, UPCGMetadata* InMetadata, bool bKeepParent, bool bCopyEntries = true, bool bCopyValues = true) const = 0;
+	
 	virtual FPCGMetadataAttributeBase* CopyToAnotherType(int16 Type) const = 0;
 
 	virtual PCGMetadataValueKey GetValueKeyOffsetForChild() const = 0;
@@ -102,7 +116,7 @@ protected:
 	TMap<PCGMetadataEntryKey, PCGMetadataValueKey> EntryToValueKeyMap;
 	mutable FRWLock EntryMapLock;
 
-	TObjectPtr<UPCGMetadata> Metadata = nullptr;
+	FPCGMetadataLayer* Metadata = nullptr;
 	const FPCGMetadataAttributeBase* Parent = nullptr;
 	int16 TypeId = 0;
 	bool bAllowsInterpolation = false;
diff --git a/Engine/Plugins/PCG/Source/PCG/Public/Metadata/PCGMetadataAttributeTpl.h b/Engine/Plugins/PCG/Source/PCG/Public/Metadata/PCGMetadataAttributeTpl.h
--- a/Engine/Plugins/PCG/Source/PCG/Public/Metadata/PCGMetadataAttributeTpl.h
+++ b/Engine/Plugins/PCG/Source/PCG/Public/Metadata/PCGMetadataAttributeTpl.h
@@ -11,6 +11,7 @@
 #include "Metadata/PCGMetadataAttributeTraits.h"
 #include "Misc/ScopeRWLock.h"
 
+class FPCGMetadataLayer;
 class UPCGMetadata;
 
 template<typename T>
@@ -20,7 +21,8 @@ class FPCGMetadataAttribute : public FPCGMetadataAttributeBase
 	friend class FPCGMetadataAttribute;
 
 public:
-	FPCGMetadataAttribute(UPCGMetadata* InMetadata, FName InName, const FPCGMetadataAttributeBase* InParent, const T& InDefaultValue, bool bInAllowsInterpolation)
+	/** TODO: All deprecated versions taking a UPCGMetadata will be added at the end of this declaration. For the review, it will be in place for "clarity" */
+	FPCGMetadataAttribute(FPCGMetadataLayer* InMetadata, FName InName, const FPCGMetadataAttributeBase* InParent, const T& InDefaultValue, bool bInAllowsInterpolation)
 		: FPCGMetadataAttributeBase(InMetadata, InName, InParent, bInAllowsInterpolation)
 		, DefaultValue(InDefaultValue)
 	{
@@ -35,13 +37,18 @@ public:
 		}
 	}
 
+	UE_DEPRECATED(5.6, "Use the version with FPCGMetadataLayer")
+	FPCGMetadataAttribute(UPCGMetadata* InMetadata, FName InName, const FPCGMetadataAttributeBase* InParent, const T& InDefaultValue, bool bInAllowsInterpolation)
+		: FPCGMetadataAttribute(PCGMetadataHelpers::GetDefaultMetadataLayer(InMetadata), InName, InParent, InDefaultValue, bAllowsInterpolation)
+	{}
+
 	// This constructor is used only during serialization
 	FPCGMetadataAttribute()
 	{
 		TypeId = PCG::Private::MetadataTypes<T>::Id;
 	}
 
-	virtual void Serialize(UPCGMetadata* InMetadata, FArchive& InArchive) override
+	virtual void Serialize(FPCGMetadataLayer* InMetadata, FArchive& InArchive) override
 	{
 		FPCGMetadataAttributeBase::Serialize(InMetadata, InArchive);
 
@@ -187,21 +194,42 @@ public:
 
 	const FPCGMetadataAttribute* GetParent() const { return static_cast<const FPCGMetadataAttribute*>(Parent); }
 
+	UE_DEPRECATED(5.6, "Use the version with FPCGMetadataLayer")
 	FPCGMetadataAttribute* TypedCopy(FName NewName, UPCGMetadata* InMetadata, bool bKeepParent, bool bCopyEntries = true, bool bCopyValues = true)
+	{
+		return static_cast<FPCGMetadataAttribute*>(TypedCopy(NewName, PCGMetadataHelpers::GetDefaultMetadataLayer(InMetadata), bKeepParent, bCopyEntries, bCopyValues));
+	}
+	
+	FPCGMetadataAttribute* TypedCopy(FName NewName, FPCGMetadataLayer* InMetadata, bool bKeepParent, bool bCopyEntries = true, bool bCopyValues = true)
 	{
 		return static_cast<FPCGMetadataAttribute*>(Copy(NewName, InMetadata, bKeepParent, bCopyEntries, bCopyValues));
 	}
 
-	virtual FPCGMetadataAttributeBase* Copy(FName NewName, UPCGMetadata* InMetadata, bool bKeepParent, bool bCopyEntries = true, bool bCopyValues = true) const override
+	UE_DEPRECATED(5.6, "Use the version with FPCGMetadataLayer")
+	virtual FPCGMetadataAttribute* Copy(FName NewName, UPCGMetadata* InMetadata, bool bKeepParent, bool bCopyEntries = true, bool bCopyValues = true) const override
+	{
+		return static_cast<FPCGMetadataAttribute*>(Copy(NewName, PCGMetadataHelpers::GetDefaultMetadataLayer(InMetadata), bKeepParent, bCopyEntries, bCopyValues));
+	}
+	
+	virtual FPCGMetadataAttributeBase* Copy(FName NewName, FPCGMetadataLayer* InMetadata, bool bKeepParent, bool bCopyEntries = true, bool bCopyValues = true) const override
 	{
 		return CopyInternal<T>(NewName, InMetadata, bKeepParent, bCopyEntries, bCopyValues);
 	}
 
 	virtual FPCGMetadataAttributeBase* CopyToAnotherType(int16 TargetType) const override;
 
-	// TODO: add enable if only on compatible types, but this has some repercussion on using metadata on types that aren't normally supported.
 	template<typename U>
+	UE_DEPRECATED(5.6, "Do not use anymore, use CopyToAnotherType")
 	FPCGMetadataAttributeBase* CopyInternal(FName NewName, UPCGMetadata* InMetadata, bool bKeepParent, bool bCopyEntries, bool bCopyValues) const
+	{
+		return CopyInternal<U>(NewName, PCGMetadataHelpers::GetDefaultMetadataLayer(InMetadata), bKeepParent, bCopyEntries, bCopyValues);
+	}
+
+protected:
+	// TODO: add enable if only on compatible types, but this has some repercussion on using metadata on types that aren't normally supported.
+	// TODO: The UPCGMetadata version was exposed, but should have not, mark the FPCGMetadataLayer version protected. Will be moved at the end, before commit.
+	template<typename U>
+	FPCGMetadataAttributeBase* CopyInternal(FName NewName, FPCGMetadataLayer* InMetadata, bool bKeepParent, bool bCopyEntries, bool bCopyValues) const
 	{
 		// If we copy an attribute where we don't want to keep the parent, while copying entries and/or values, we'll lose data.
 		// In that case, we will copy all the data from this attribute and all its ancestors.
@@ -227,10 +255,10 @@ public:
 		TArray<const FPCGMetadataAttribute<T>*, TInlineAllocator<2>> Parents = { this };
 		if (!bKeepParent && (bCopyEntries || bCopyValues))
 		{
-			const UPCGMetadata* CurrentMetadata = Metadata.Get();
+			const FPCGMetadataLayer* CurrentMetadata = Metadata;
 			const FPCGMetadataAttribute<T>* Current = this;
 
-			const UPCGMetadata* ParentMetadata = PCGMetadataHelpers::GetParentMetadata(CurrentMetadata);
+			const FPCGMetadataLayer* ParentMetadata = PCGMetadataHelpers::GetParentMetadata(CurrentMetadata);
 			while(ParentMetadata && Current->Parent)
 			{
 				CurrentMetadata = ParentMetadata;
@@ -293,6 +321,7 @@ public:
 		return AttributeCopy;
 	}
 
+public:
 	virtual void SetValue(PCGMetadataEntryKey ItemKey, const FPCGMetadataAttributeBase* InAttribute, PCGMetadataEntryKey InEntryKey) override
 	{
 		check(ItemKey != PCGInvalidEntryKey);
diff --git a/Engine/Plugins/PCG/Source/PCG/Public/Metadata/PCGMetadataCommon.h b/Engine/Plugins/PCG/Source/PCG/Public/Metadata/PCGMetadataCommon.h
--- a/Engine/Plugins/PCG/Source/PCG/Public/Metadata/PCGMetadataCommon.h
+++ b/Engine/Plugins/PCG/Source/PCG/Public/Metadata/PCGMetadataCommon.h
@@ -42,3 +42,156 @@ enum class EPCGMetadataFilterMode : uint8
 	/** Only the listed attributes will be changed by the projection or added from the target data. */
 	IncludeAttributes,
 };
+
+UENUM()
+enum class EPCGMetadataLayerFlag : uint8
+{
+	/** Depends on the data. Should map to the same concept before multi-level metadata. */
+	Default = 0,
+	
+	/** Metadata data at the data level. */
+	Data = 1,
+
+	/** Metadata on elements like points on point data and entries on param data. */
+	Elements = 2,
+
+	/** For data that can have more levels. */
+	Custom = 255
+};
+
+USTRUCT(BlueprintType)
+struct FPCGMetadataLayerID
+{
+	GENERATED_BODY();
+
+	FPCGMetadataLayerID() = default;
+	
+	FPCGMetadataLayerID(EPCGMetadataLayerFlag InFlag, int32 InCustomFlag = -1, FName InDebugName = NAME_None)
+		: Flag(InFlag)
+		, CustomFlag(InCustomFlag)
+		, DebugName(InDebugName)
+	{
+		check(CustomFlag == -1 || InFlag == EPCGMetadataLayerFlag::Custom);
+	}
+
+	bool operator==(const FPCGMetadataLayerID& Other) const
+	{
+		return Flag == Other.Flag && CustomFlag == Other.CustomFlag;
+	}
+
+	bool operator<(const FPCGMetadataLayerID& Other) const
+	{
+		return Flag < Other.Flag || (Flag == Other.Flag && CustomFlag < Other.CustomFlag);
+	}
+
+	friend uint32 GetTypeHash(const FPCGMetadataLayerID& Item)
+	{
+		return HashCombine(static_cast<uint32>(Item.Flag), static_cast<uint32>(Item.CustomFlag));
+	}
+
+	friend FArchive& operator<<(FArchive& Ar, FPCGMetadataLayerID& Item)
+	{
+		Ar << Item.Flag;
+		Ar << Item.CustomFlag;
+		return Ar;
+	}
+
+	bool IsDefault() const { return Flag == EPCGMetadataLayerFlag::Default; }
+
+	UPROPERTY(EditAnywhere, BlueprintReadWrite, Category = "")
+	EPCGMetadataLayerFlag Flag = EPCGMetadataLayerFlag::Default;
+	
+	UPROPERTY(EditAnywhere, BlueprintReadWrite, Category = "")
+	int32 CustomFlag = -1;
+
+	UPROPERTY(EditAnywhere, BlueprintReadWrite, Category = "")
+	FName DebugName = NAME_None;
+};
+
+namespace PCGMetadataLayerID
+{
+	static const FPCGMetadataLayerID Default{EPCGMetadataLayerFlag::Default, -1, TEXT("Default")};
+	static const FPCGMetadataLayerID Elements{EPCGMetadataLayerFlag::Elements, -1, TEXT("Elements")};
+	static const FPCGMetadataLayerID Data{EPCGMetadataLayerFlag::Data, -1, TEXT("Data")};
+}
+
+USTRUCT(BlueprintType)
+struct FPCGAttributeIdentifier
+{
+	GENERATED_BODY();
+
+	FPCGAttributeIdentifier() = default;
+
+	// Needs to be backward compatible with FNames (and everything that can be constructed into a FName).
+	template <typename T = FName, std::enable_if_t<std::is_constructible_v<FName, T>, bool> = true>
+	FPCGAttributeIdentifier(const T& InName, FPCGMetadataLayerID InMetadataLayerID = PCGMetadataLayerID::Default)
+		: Name(InName)
+		, MetadataLayer(InMetadataLayerID)
+	{}
+
+	template <typename T = FName, std::enable_if_t<std::is_constructible_v<FName, T>, bool> = true>
+	FPCGAttributeIdentifier& operator=(const T& InName)
+	{
+		Name = FName(InName);
+		return *this;
+	}
+
+	UE_DEPRECATED(5.6, "Explicitly use the Name field.")
+	operator FName() const
+	{
+		return Name;
+	}
+
+	bool operator==(const FPCGAttributeIdentifier& Other) const
+	{
+		return Name == Other.Name && MetadataLayer == Other.MetadataLayer;
+	}
+
+	friend uint32 GetTypeHash(const FPCGAttributeIdentifier& Item)
+	{
+		return HashCombine(GetTypeHash(Item.Name), GetTypeHash(Item.MetadataLayer));
+	}
+	
+	friend FArchive& operator<<(FArchive& Ar, FPCGAttributeIdentifier& Item)
+	{
+		Ar << Item.Name;
+		Ar << Item.MetadataLayer;
+		return Ar;
+	}
+	
+	static TSet<FPCGAttributeIdentifier> TransformNameSet(const TSet<FName>& InContainer)
+	{
+		TSet<FPCGAttributeIdentifier> OutContainer;
+		OutContainer.Reserve(InContainer.Num());
+		Algo::Transform(InContainer, OutContainer, [](const FName& InAttributeName) { return InAttributeName;});
+		return OutContainer;
+	}
+
+	template <typename AllocatorType>
+	static TArray<FPCGAttributeIdentifier, AllocatorType> TransformNameArray(const TArray<FName, AllocatorType>& InContainer)
+	{
+		TArray<FPCGAttributeIdentifier, AllocatorType> OutContainer;
+		OutContainer.Reserve(InContainer.Num());
+		Algo::Transform(InContainer, OutContainer, [](const FName& InAttributeName) { return InAttributeName;});
+		return OutContainer;
+	}
+
+	template <typename Container>
+	static TMap<FPCGMetadataLayerID, TSet<FName>> SplitByLayer(const Container& InContainer)
+	{
+		TMap<FPCGMetadataLayerID, TSet<FName>> OutMap;
+		for (const FPCGAttributeIdentifier& It : InContainer)
+		{
+			OutMap.FindOrAdd(It.MetadataLayer).Add(It.Name);
+		}
+
+		return OutMap;
+	}
+
+	UPROPERTY(EditAnywhere, BlueprintReadWrite, Category = "")
+	FName Name = NAME_None;
+
+	UPROPERTY(EditAnywhere, BlueprintReadWrite, Category = "", meta = (ShowOnlyInnerProperties))
+	FPCGMetadataLayerID MetadataLayer;
+};
+
diff --git a/Engine/Plugins/PCG/Source/PCG/Public/Metadata/PCGMetadataLayer.h b/Engine/Plugins/PCG/Source/PCG/Public/Metadata/PCGMetadataLayer.h
new file mode 100644
--- /dev/null
+++ b/Engine/Plugins/PCG/Source/PCG/Public/Metadata/PCGMetadataLayer.h
@@ -0,0 +1,396 @@
+// Copyright Epic Games, Inc. All Rights Reserved.
+
+#pragma once
+
+#include "PCGCommon.h"
+#include "Metadata/PCGMetadataCommon.h"
+#include "Metadata/PCGMetadataAttributeTpl.h"
+
+class UPCGMetadata;
+
+class PCG_API FPCGMetadataLayer
+{
+	friend UPCGMetadata;
+
+public:
+	FPCGMetadataLayer(UPCGMetadata* InTopMetadata, FPCGMetadataLayerID InMetadataLayerID);
+	virtual ~FPCGMetadataLayer();
+	
+	virtual void Serialize(FArchive& InArchive);
+	
+	/** Initializes the sub metadata from a parent sub metadata, if any (can be null). Copies attributes and values. */
+	void Initialize(const FPCGMetadataLayer* InParent) { InitializeWithAttributeFilter(InParent, TSet<FName>()); }
+
+	/** Initializes the sub metadata from a parent sub metadata, if any (can be null) with the option to not add attributes from the parent. */
+	void Initialize(const FPCGMetadataLayer* InParent, bool bAddAttributesFromParent);
+
+	/**
+	 * Initializes the sub metadata from a parent sub metadata. Copies attributes and values.
+	 * @param InParent The parent metadata to use as a template, if any (can be null).
+	 * @param InFilteredAttributes Optional list of attributes to exclude or include when adding the attributes from the parent.
+	 * @param InFilterMode Defines attribute filter operation.
+	 */
+	void InitializeWithAttributeFilter(const FPCGMetadataLayer* InParent, const TSet<FName>& InFilteredAttributes, EPCGMetadataFilterMode InFilterMode = EPCGMetadataFilterMode::ExcludeAttributes, EPCGStringMatchingOperator InMatchOperator = EPCGStringMatchingOperator::Equal);
+
+	/** Initializes the sub metadata from a parent sub metadata by copying all attributes to it.
+	* @param InMetadataToCopy Metadata to copy from
+	* @param InOptionalEntriesToCopy Optional array that contains the keys to copy over. This array order will be respected, so it can also be used to re-order entries. Can be null.
+	*/
+	void InitializeAsCopy(const FPCGMetadataLayer* InMetadataToCopy, const TArray<PCGMetadataEntryKey>* InOptionalEntriesToCopy = nullptr);
+
+	/** Initializes the metadata from a parent metadata by copy filtered attributes only to it
+	* @param InMetadataToCopy Metadata to copy from
+	* @param InFilteredAttributes Attributes to keep/exclude, can be empty.
+	* @param InFilterMode Filter to know if we should keep or exclude InFilteredAttributes.
+	* @param InOptionalEntriesToCopy Optional array that contains the keys to copy over. This array order will be respected, so it can also be used to re-order entries. Can be null.
+	*/
+	void InitializeAsCopyWithAttributeFilter(const FPCGMetadataLayer* InMetadataToCopy, const TSet<FName>& InFilteredAttributes, EPCGMetadataFilterMode InFilterMode = EPCGMetadataFilterMode::ExcludeAttributes, const TArray<PCGMetadataEntryKey>* InOptionalEntriesToCopy = nullptr);
+
+	/** Creates missing attributes from another metadata if they are not currently present - note that this does not copy values */
+	bool AddAttributes(const FPCGMetadataLayer* InOther) { return AddAttributesFiltered(InOther, TSet<FName>()); }
+
+	/**
+	 * Creates missing attributes from another metadata if they are not currently present - note that this does not copy values.
+	 * @param InOther The other metadata to obtain a list of attributes from.
+	 * @param InFilteredAttributes Optional list of attributes to exclude or include when adding the attributes.
+	 * @param InFilterMode Defines attribute filter operation.
+	 */
+	bool AddAttributesFiltered(const FPCGMetadataLayer* InOther, const TSet<FName>& InFilteredAttributes, EPCGMetadataFilterMode InFilterMode = EPCGMetadataFilterMode::ExcludeAttributes, EPCGStringMatchingOperator InMatchOperator = EPCGStringMatchingOperator::Equal);
+
+	/** Creates missing attribute from another metadata if it is not currently present - note that this does not copy values */
+	bool AddAttribute(const FPCGMetadataLayer* InOther, FName AttributeName);
+
+	/** Copies attributes from another metadata, including entries & values. Warning: this is intended when dealing with the same data set */
+	void CopyAttributes(const FPCGMetadataLayer* InOther);
+
+	/** Copies an attribute from another metadata, including entries & values. Warning: this is intended when dealing with the same data set */
+	void CopyAttribute(const FPCGMetadataLayer* InOther, FName AttributeToCopy, FName NewAttributeName);
+
+	/** Copies another attribute, with options to keep its parent and copy entries/values */
+	FPCGMetadataAttributeBase* CopyAttribute(const FPCGMetadataAttributeBase* OriginalAttribute, FName NewAttributeName, bool bKeepParent, bool bCopyEntries, bool bCopyValues);
+
+	/** Returns this metadata's parent */
+	const FPCGMetadataLayer* GetParent() const { return Parent; }
+	const FPCGMetadataLayer* GetRoot() const;
+	const UPCGMetadata* GetTopMetadata() const { return TopMetadata; }
+	bool HasParent(const FPCGMetadataLayer* InTentativeParent) const;
+
+	/** Unparents current metadata by flattening the attributes (values, entries, etc.) */
+	void FlattenImpl();
+
+	/** Unparents current metadata, flatten attribute and only keep the entries specified. Return true if something has changed and keys needs be updated. */
+	bool FlattenAndCompress(const TArray<PCGMetadataEntryKey>& InEntryKeysToKeep);
+
+	/** Creates an attribute given a property.
+	* @param AttributeName: Target attribute to create
+	* @param Object: Object to get the property value from
+	* @param Property: The property to set from
+	* @returns true if the attribute creation succeeded
+	*/
+	bool CreateAttributeFromProperty(FName AttributeName, const UObject* Object, const FProperty* Property);
+
+	/** Creates an attribute given a property.
+	* @param AttributeName: Target attribute to create
+	* @param Data: Data pointer to get the property value from
+	* @param Property: The property to set from
+	* @returns true if the attribute creation succeeded
+	*/
+	bool CreateAttributeFromDataProperty(FName AttributeName, const void* Data, const FProperty* Property);
+
+	/** Set an attribute given a property and its value.
+	* @param AttributeName: Target attribute to set the property's value to
+	* @param EntryKey: Metadata entry key to set the value to
+	* @param Object: Object to get the property value from
+	* @param Property: The property to set from
+	* @param bCreate: If true and the attribute doesn't exists, it will create an attribute based on the property type
+	* @returns true if the attribute creation (if required) and the value set succeeded
+	*/
+	bool SetAttributeFromProperty(FName AttributeName, PCGMetadataEntryKey& EntryKey, const UObject* Object, const FProperty* Property, bool bCreate);
+
+	/** Set an attribute given a property and its value.
+	* @param AttributeName: Target attribute to set the property's value to
+	* @param EntryKey: Metadata entry key to set the value to
+	* @param Data: Data pointer to get the property value from
+	* @param Property: The property to set from
+	* @param bCreate: If true and the attribute doesn't exists, it will create an attribute based on the property type
+	* @returns true if the attribute creation (if required) and the value set succeeded
+	*/
+	bool SetAttributeFromDataProperty(FName AttributeName, PCGMetadataEntryKey& EntryKey, const void* Data, const FProperty* Property, bool bCreate);
+
+	/** Get attributes */
+	FPCGMetadataAttributeBase* GetMutableAttribute(FName AttributeName);
+	const FPCGMetadataAttributeBase* GetConstAttribute(FName AttributeName) const;
+	const FPCGMetadataAttributeBase* GetConstAttributeById(int32 AttributeId) const;
+	
+	bool HasAttribute(FName AttributeName) const;
+	bool HasCommonAttributes(const FPCGMetadataLayer* InMetadata) const;
+
+	/** Return the number of attributes in this metadata. */
+	int32 GetAttributeCount() const;
+
+	template <typename T>
+	FPCGMetadataAttribute<T>* GetMutableTypedAttribute(FName AttributeName);
+
+	template <typename T>
+	FPCGMetadataAttribute<T>* GetMutableTypedAttribute_Unsafe(FName AttributeName);
+
+	template <typename T>
+	const FPCGMetadataAttribute<T>* GetConstTypedAttribute(FName AttributeName) const;
+	
+	void GetAttributes(TArray<FName>& AttributeNames, TArray<EPCGMetadataTypes>& AttributeTypes) const;
+
+	/** Returns name of the most recently created attribute, or none if no attributes are present. */
+	FName GetLatestAttributeNameOrNone() const;
+
+	/** Delete/Hide attribute */
+	// Due to stream inheriting, we might want to consider "hiding" parent stream and deleting local streams only
+	void DeleteAttribute(FName AttributeName);
+
+	/** Copy attribute */
+	bool CopyExistingAttribute(FName AttributeToCopy, FName NewAttributeName, bool bKeepParent = true);
+
+	/** Rename attribute */
+	bool RenameAttribute(FName AttributeToRename, FName NewAttributeName);
+
+	/** Clear/Reinit attribute */
+	void ClearAttribute(FName AttributeToClear);
+
+	/** Change type of an attribute */
+	bool ChangeAttributeType(FName AttributeName, int16 AttributeNewType);
+
+	/** Adds a unique entry key to the metadata */
+	int64 AddEntry(int64 ParentEntryKey = -1);
+
+	/** Adds a unique entry key to the metadata for all the parent entry keys. */
+	TArray<int64> AddEntries(TArrayView<const int64> ParentEntryKeys);
+
+	/** Adds a unique entry key to the metadata for all the parent entry keys, in place. */
+	void AddEntriesInPlace(TArrayView<int64*> ParentEntryKeys);
+
+	/** Advanced method.
+	*   In a MT context, we might not want to add the entry directly (because of write lock). Call this to generate an unique index in the MT context
+	*   And call AddDelayedEntries at the end when you want to add all the entries.
+	*   Make sure to not call AddEntry in a different thread or even in the same thread if you use this function, or it will change all the indexes.
+	*/
+	int64 AddEntryPlaceholder();
+
+	/** Advanced method.
+	*   If you used AddEntryPlaceholder, call this function at the end of your MT processing to add all the entries in one shot.
+	*   Make sure to call this one at the end!!
+	*   @param AllEntries Array of pairs. First item is the EntryIndex generated by AddEntryPlaceholder, the second the ParentEntryKey (cf AddEntry)
+	*/
+	void AddDelayedEntries(const TArray<TTuple<int64, int64>>& AllEntries);
+
+	/** Initializes the metadata entry key. Returns true if key set from either parent */
+	bool InitializeOnSet(PCGMetadataEntryKey& InOutKey, PCGMetadataEntryKey InParentKeyA = PCGInvalidEntryKey, const FPCGMetadataLayer* InParentMetadataA = nullptr, PCGMetadataEntryKey InParentKeyB = PCGInvalidEntryKey, const FPCGMetadataLayer* InParentMetadataB = nullptr);
+
+	/** Metadata chaining mechanism */
+	PCGMetadataEntryKey GetParentKey(PCGMetadataEntryKey LocalItemKey) const;
+
+	/** Metadata chaining mechanism for bulk version. Can provide a mask to only update only a subset of the passed keys. */
+	void GetParentKeys(TArrayView<PCGMetadataEntryKey> LocalItemKeys, const TBitArray<>* Mask = nullptr) const;
+
+	/** Attributes operations */
+	void MergeAttributes(PCGMetadataEntryKey InKeyA, const FPCGMetadataLayer* InMetadataA, PCGMetadataEntryKey InKeyB, const FPCGMetadataLayer* InMetadataB, PCGMetadataEntryKey& OutKey, EPCGMetadataOp Op);
+	void MergeAttributesSubset(PCGMetadataEntryKey InKeyA, const FPCGMetadataLayer* InMetadataA, const FPCGMetadataLayer* InMetadataSubetA, PCGMetadataEntryKey InKeyB, const FPCGMetadataLayer* InMetadataB, const FPCGMetadataLayer* InMetadataSubsetb, PCGMetadataEntryKey& OutKey, EPCGMetadataOp Op);
+
+	void ResetWeightedAttributes(PCGMetadataEntryKey& OutKey);
+	void AccumulateWeightedAttributes(PCGMetadataEntryKey InKey, const FPCGMetadataLayer* InMetadata, float Weight, bool bSetNonInterpolableAttributes, PCGMetadataEntryKey& OutKey);
+
+	void SetAttributes(PCGMetadataEntryKey InKey, const FPCGMetadataLayer* InMetadata, PCGMetadataEntryKey& OutKey);
+	void SetAttributes(const TArrayView<const PCGMetadataEntryKey>& InOriginalKeys, const FPCGMetadataLayer* InMetadata, const TArrayView<PCGMetadataEntryKey>* InOutOptionalKeys = nullptr, FPCGContext* OptionalContext = nullptr);
+	void SetAttributes(const TArrayView<const PCGMetadataEntryKey>& InKeys, const FPCGMetadataLayer* InMetadata, const TArrayView<PCGMetadataEntryKey>& OutKeys, FPCGContext* OptionalContext = nullptr);
+
+	void ComputeWeightedAttribute(PCGMetadataEntryKey& OutKey, const TArrayView<TPair<PCGMetadataEntryKey, float>>& InWeightedKeys, const FPCGMetadataLayer* InMetadata);
+
+	int64 GetItemKeyCountForParent() const;
+	int64 GetLocalItemCount() const;
+
+	/** Return the number of entries in metadata including the parent entries. */
+	int64 GetItemCountForChild() const;
+
+	/**
+	* Create a new attribute. If the attribute already exists, it will raise a warning (use FindOrCreateAttribute if this usecase can arise)
+	* If the attribute already exists but is of the wrong type, it will fail and return nullptr. Same if the name is invalid.
+	* Return a typed attribute pointer, of the requested type T.
+	*/
+	template<typename T>
+	FPCGMetadataAttribute<T>* CreateAttribute(FName AttributeName, const T& DefaultValue, bool bAllowsInterpolation, bool bOverrideParent);
+
+	/**
+	* Find or create an attribute. Follows CreateAttribute signature.
+	* Extra boolean bOverwriteIfTypeMismatch allows to overwrite an existing attribute if the type mismatch.
+	* Same as CreateAttribute, it will return nullptr if the attribute name is invalid.
+	* Return a typed attribute pointer, of the requested type T.
+	*/
+	template<typename T>
+	FPCGMetadataAttribute<T>* FindOrCreateAttribute(FName AttributeName, const T& DefaultValue = T{}, bool bAllowsInterpolation = true, bool bOverrideParent = true, bool bOverwriteIfTypeMismatch = true);
+
+	/** Computes Crc from all attributes & keys from outer's data. */
+	void AddToCrc(FArchiveCrc32& Ar, const UPCGData* Data, bool bFullDataCrc) const;
+
+protected:
+	FPCGMetadataAttributeBase* CopyAttribute(FName AttributeToCopy, FName NewAttributeName, bool bKeepParent, bool bCopyEntries, bool bCopyValues);
+
+	bool ParentHasAttribute(FName AttributeName) const;
+
+	void AddAttributeInternal(FName AttributeName, FPCGMetadataAttributeBase* Attribute);
+	void RemoveAttributeInternal(FName AttributeName);
+
+	void SetLastCachedSelectorOnOwner(FName AttributeName);
+
+	UPCGMetadata* TopMetadata = nullptr;
+	FPCGMetadataLayerID LayerID;
+	const FPCGMetadataLayer* Parent = nullptr;
+
+	TMap<FName, FPCGMetadataAttributeBase*> Attributes;
+	PCGMetadataAttributeKey NextAttributeId = 0;
+
+	TArray<PCGMetadataEntryKey> ParentKeys;
+	int64 ItemKeyOffset = 0;
+
+	mutable FRWLock AttributeLock;
+	mutable FRWLock ItemLock;
+
+	TAtomic<int64> DelayedEntriesIndex = 0;
+};
+
+
+template<typename T>
+FPCGMetadataAttribute<T>* FPCGMetadataLayer::CreateAttribute(FName AttributeName, const T& DefaultValue, bool bAllowsInterpolation, bool bOverrideParent)
+{
+	if (!FPCGMetadataAttributeBase::IsValidName(AttributeName))
+	{
+		UE_LOG(LogPCG, Error, TEXT("Attribute name '%s' is invalid"), *AttributeName.ToString());
+		return nullptr;
+	}
+
+	const FPCGMetadataAttributeBase* ParentAttribute = nullptr;
+
+	if (bOverrideParent && Parent)
+	{
+		ParentAttribute = Parent->GetConstAttribute(AttributeName);
+	}
+
+	if (ParentAttribute && (ParentAttribute->GetTypeId() != PCG::Private::MetadataTypes<T>::Id))
+	{
+		// Can't parent if the types do not match
+		ParentAttribute = nullptr;
+	}
+
+	FPCGMetadataAttribute<T>* NewAttribute = new FPCGMetadataAttribute<T>(this, AttributeName, ParentAttribute, DefaultValue, bAllowsInterpolation);
+
+	{
+		FWriteScopeLock WriteLock(AttributeLock);
+
+		if (FPCGMetadataAttributeBase** ExistingAttribute = Attributes.Find(AttributeName))
+		{
+			delete NewAttribute;
+			if ((*ExistingAttribute)->GetTypeId() != PCG::Private::MetadataTypes<T>::Id)
+			{
+				UE_LOG(LogPCG, Error, TEXT("Attribute %s already exists but is not the right type. Abort."), *AttributeName.ToString());
+				return nullptr;
+			}
+			else
+			{
+				UE_LOG(LogPCG, Warning, TEXT("Attribute %s already exists"), *AttributeName.ToString());
+				NewAttribute = static_cast<FPCGMetadataAttribute<T>*>(*ExistingAttribute);
+			}
+		}
+		else
+		{
+			NewAttribute->AttributeId = NextAttributeId++;
+			AddAttributeInternal(AttributeName, NewAttribute);
+
+			// Also when creating an attribute, notify the PCG Data owner that the latest attribute manipulated is this one.
+			SetLastCachedSelectorOnOwner(AttributeName);
+		}
+	}
+
+	return NewAttribute;
+}
+
+template<typename T>
+FPCGMetadataAttribute<T>* FPCGMetadataLayer::GetMutableTypedAttribute_Unsafe(FName AttributeName)
+{
+	FPCGMetadataAttribute<T>* Attribute = nullptr;
+
+	if (FPCGMetadataAttributeBase** FoundAttribute = Attributes.Find(AttributeName))
+	{
+		FPCGMetadataAttributeBase* BaseAttribute = *FoundAttribute;
+		if (BaseAttribute && BaseAttribute->GetTypeId() == PCG::Private::MetadataTypes<T>::Id)
+		{
+			Attribute = static_cast<FPCGMetadataAttribute<T>*>(BaseAttribute);
+		}
+	}
+
+	return Attribute;
+}
+
+template<typename T>
+FPCGMetadataAttribute<T>* FPCGMetadataLayer::FindOrCreateAttribute(FName AttributeName, const T& DefaultValue, bool bAllowsInterpolation, bool bOverrideParent, bool bOverwriteIfTypeMismatch)
+{
+	{
+		FReadScopeLock ScopeLock(AttributeLock);
+		if (FPCGMetadataAttribute<T>* Attribute = GetMutableTypedAttribute_Unsafe<T>(AttributeName))
+		{
+			return Attribute;
+		}
+	}
+
+	FWriteScopeLock ScopeLock(AttributeLock);
+	if (FPCGMetadataAttribute<T>* Attribute = GetMutableTypedAttribute_Unsafe<T>(AttributeName))
+	{
+		return Attribute;
+	}
+
+	// If an attribute with this name exists here, there is a type mismatch.
+	if (FPCGMetadataAttributeBase** FoundAttribute = Attributes.Find(AttributeName))
+	{
+		if (bOverwriteIfTypeMismatch)
+		{
+			delete *FoundAttribute;
+			RemoveAttributeInternal(AttributeName);
+		}
+		else
+		{
+			return nullptr;
+		}
+	}
+
+	// A new attribute will be created.
+	if (!FPCGMetadataAttributeBase::IsValidName(AttributeName))
+	{
+		UE_LOG(LogPCG, Error, TEXT("Attribute name '%s' is invalid"), *AttributeName.ToString());
+		return nullptr;
+	}
+
+	// Parent is const and therefore should be safe
+	const FPCGMetadataAttributeBase* ParentAttribute = (bOverrideParent && Parent) ? Parent->GetConstTypedAttribute<T>(AttributeName) : nullptr;
+	FPCGMetadataAttribute<T>* NewAttribute = new FPCGMetadataAttribute<T>(this, AttributeName, ParentAttribute, DefaultValue, bAllowsInterpolation);
+	NewAttribute->AttributeId = NextAttributeId++;
+	AddAttributeInternal(AttributeName, NewAttribute);
+
+	// Also when creating an attribute, notify the PCG Data owner that the latest attribute manipulated is this one.
+	SetLastCachedSelectorOnOwner(AttributeName);
+
+	return NewAttribute;
+}
+
+template <typename T>
+FPCGMetadataAttribute<T>* FPCGMetadataLayer::GetMutableTypedAttribute(FName AttributeName)
+{
+	FPCGMetadataAttributeBase* BaseAttribute = GetMutableAttribute(AttributeName);
+	return (BaseAttribute && (BaseAttribute->GetTypeId() == PCG::Private::MetadataTypes<T>::Id))
+		? static_cast<FPCGMetadataAttribute<T>*>(BaseAttribute)
+		: nullptr;
+}
+
+template <typename T>
+const FPCGMetadataAttribute<T>* FPCGMetadataLayer::GetConstTypedAttribute(FName AttributeName) const
+{
+	const FPCGMetadataAttributeBase* BaseAttribute = GetConstAttribute(AttributeName);
+	return (BaseAttribute && (BaseAttribute->GetTypeId() == PCG::Private::MetadataTypes<T>::Id))
+		? static_cast<const FPCGMetadataAttribute<T>*>(BaseAttribute)
+		: nullptr;
+}
\ No newline at end of file
diff --git a/Engine/Plugins/PCG/Source/PCG/Public/PCGCustomVersion.h b/Engine/Plugins/PCG/Source/PCG/Public/PCGCustomVersion.h
--- a/Engine/Plugins/PCG/Source/PCG/Public/PCGCustomVersion.h
+++ b/Engine/Plugins/PCG/Source/PCG/Public/PCGCustomVersion.h
@@ -80,6 +80,9 @@ struct FPCGCustomVersion
 		// Attributes and tags can now contain spaces and will no longer be parsed by spaces.
 		AttributesAndTagsCanContainSpaces = 23,
 
+		// MultiLevelMetadata
+		MultiLevelMetadata = 24,
+
 		// -----<new versions can be added above this line>-------------------------------------------------
 		VersionPlusOne,
 		LatestVersion = VersionPlusOne - 1
diff --git a/Engine/Plugins/PCG/Source/PCG/Public/PCGData.h b/Engine/Plugins/PCG/Source/PCG/Public/PCGData.h
--- a/Engine/Plugins/PCG/Source/PCG/Public/PCGData.h
+++ b/Engine/Plugins/PCG/Source/PCG/Public/PCGData.h
@@ -5,6 +5,7 @@
 #include "PCGCommon.h"
 #include "PCGCrc.h"
 #include "Metadata/PCGAttributePropertySelector.h"
+#include "Metadata/PCGMetadataCommon.h"
 
 #include "Kismet/BlueprintFunctionLibrary.h"
 #include "Templates/SubclassOf.h"
@@ -72,9 +73,19 @@ public:
 	// ~End UObject interface
 
 	// Metadata ops, to be implemented if data supports Metadata
-	virtual UPCGMetadata* MutableMetadata() { return nullptr; }
-	virtual const UPCGMetadata* ConstMetadata() const { return nullptr; }
+	UFUNCTION(BlueprintCallable, Category = Metadata)
+	virtual const UPCGMetadata* ConstMetadata() const { return Metadata; }
+
+	UFUNCTION(BlueprintCallable, Category = Metadata)
+	virtual UPCGMetadata* MutableMetadata() { return Metadata; }
+
 	virtual void Flatten();
+	virtual FPCGMetadataLayerID GetDefaultMetadataLayerID() const { return PCGMetadataLayerID::Data; }
+	virtual TArray<FPCGMetadataLayerID> GetAllSupportedMetadataLayerIDs() const { return {PCGMetadataLayerID::Data}; }
+
+	// Not accessible through blueprint to make sure the constness is preserved.
+	UPROPERTY(VisibleAnywhere, Category = Metadata)
+	TObjectPtr<UPCGMetadata> Metadata = nullptr;
 
 protected:
 	/** Computes Crc for this and any connected data. */
diff --git a/Engine/Plugins/PCG/Source/PCG/Public/PCGParamData.h b/Engine/Plugins/PCG/Source/PCG/Public/PCGParamData.h
--- a/Engine/Plugins/PCG/Source/PCG/Public/PCGParamData.h
+++ b/Engine/Plugins/PCG/Source/PCG/Public/PCGParamData.h
@@ -29,11 +29,8 @@ public:
 	virtual UPCGParamData* DuplicateData(FPCGContext* Context, bool bInitializeMetadata = true) const override;
 	// ~End UPCGData interface
 
-	UFUNCTION(BlueprintCallable, Category = Metadata)
-	virtual const UPCGMetadata* ConstMetadata() const override{ return Metadata; }
-
-	UFUNCTION(BlueprintCallable, Category = Metadata)
-	virtual UPCGMetadata* MutableMetadata() override { return Metadata; }
+	virtual FPCGMetadataLayerID GetDefaultMetadataLayerID() const override { return PCGMetadataLayerID::Elements; }
+	virtual TArray<FPCGMetadataLayerID> GetAllSupportedMetadataLayerIDs() const override { return {PCGMetadataLayerID::Data, PCGMetadataLayerID::Elements}; }
 
 	/** Returns the entry for the given name */
 	UFUNCTION(BlueprintCallable, Category = Params)
@@ -60,10 +57,6 @@ public:
 
 	UPCGParamData* FilterParamsByKey(FPCGContext* Context, int64 InKey) const;
 
-	// Not accessible through blueprint to make sure the constness is preserved
-	UPROPERTY(VisibleAnywhere, Category = Metadata)
-	TObjectPtr<UPCGMetadata> Metadata = nullptr;
-
 protected:
 	UPROPERTY()
 	TMap<FName, int64> NameMap;
