[PCG] Change Set to Array for Managed Actors for SpawnActor Actor Reference order when re-use managed resources

* GeneratedActors is no more used, you need to update your code to use the new getter/setter functions on UPCGManagedActors.
* Also enable full CRC if we re-use actors. It's expected to be a negligible cost compare to actually spawn an actor.

diff --git a/Engine/Plugins/PCG/Source/PCG/Private/Elements/PCGCreateTargetActor.cpp b/Engine/Plugins/PCG/Source/PCG/Private/Elements/PCGCreateTargetActor.cpp
--- a/Engine/Plugins/PCG/Source/PCG/Private/Elements/PCGCreateTargetActor.cpp
+++ b/Engine/Plugins/PCG/Source/PCG/Private/Elements/PCGCreateTargetActor.cpp
@@ -428,7 +428,7 @@ bool FPCGCreateTargetActorElement::ExecuteInternal(FPCGContext* Context) const
 
 				if (UPCGManagedActors* Resource = Cast<UPCGManagedActors>(InResource))
 				{
-					if (Resource->GetCrc().IsValid() && Resource->GetCrc() == ResourceCrc && Resource->GeneratedActors.Num() == 1 && Resource->GeneratedActors.Array()[0].IsValid())
+					if (Resource->GetCrc().IsValid() && Resource->GetCrc() == ResourceCrc && Resource->GetConstGeneratedActors().Num() == 1 && Resource->GetConstGeneratedActors()[0].IsValid())
 					{
 						ReusableResource = Resource;
 					}
@@ -441,7 +441,7 @@ bool FPCGCreateTargetActorElement::ExecuteInternal(FPCGContext* Context) const
 	if (ReusableResource)
 	{
 		ReusableResource->MarkAsReused();
-		GeneratedActor = ReusableResource->GeneratedActors.Array()[0].Get();
+		GeneratedActor = ReusableResource->GetConstGeneratedActors()[0].Get();
 		ensure(GeneratedActor->Tags.Contains(PCGHelpers::DefaultPCGActorTag));
 	}
 	else
@@ -484,7 +484,7 @@ bool FPCGCreateTargetActorElement::ExecuteInternal(FPCGContext* Context) const
 			{
 				ManagedActors->SetCrc(ResourceCrc);
 			}
-			ManagedActors->GeneratedActors.Add(GeneratedActor);
+			ManagedActors->GetMutableGeneratedActors().AddUnique(GeneratedActor);
 			ManagedActors->bSupportsReset = !Settings->bDeleteActorsBeforeGeneration;
 
 			SourceComponent->AddToManagedResources(ManagedActors);
diff --git a/Engine/Plugins/PCG/Source/PCG/Private/Elements/PCGSpawnActor.cpp b/Engine/Plugins/PCG/Source/PCG/Private/Elements/PCGSpawnActor.cpp
--- a/Engine/Plugins/PCG/Source/PCG/Private/Elements/PCGSpawnActor.cpp
+++ b/Engine/Plugins/PCG/Source/PCG/Private/Elements/PCGSpawnActor.cpp
@@ -857,14 +857,15 @@ void FPCGSpawnActorElement::SpawnActors(FPCGSubgraphContext* Context, AActor* Ta
 	UPCGSubsystem* Subsystem = (Context->SourceComponent.Get() ? Context->SourceComponent->GetSubsystem() : nullptr);
 
 	// Try to reuse actors if they are preexisting
-	TArray<UPCGManagedActors*> ReusedManagedActorsResources;
+	UPCGManagedActors* ReusedManagedActorsResource = nullptr;
 	FPCGCrc InputDependenciesCrc;
 	if (CVarAllowActorReuse.GetValueOnAnyThread())
 	{
 		FPCGDataCollection SingleInputCollection;
-		SingleInputCollection.TaggedData.Add(Output);
-		// TODO: review this, it might make more sense to do a full data crc here
-		SingleInputCollection.ComputeCrcs(/*bFullDataCrc=*/false);
+		SingleInputCollection.TaggedData.Emplace_GetRef().Data = PointData;
+		// Need to do a full CRC here as the PointData might not be the original input (if there was some partitioning because of spawning by attribute). 
+		// Since it is spawning by attribute, all point data will be different.
+		SingleInputCollection.ComputeCrcs(/*bFullDataCrc=*/true);
 
 		GetDependenciesCrc(SingleInputCollection, Settings, Context->SourceComponent.Get(), InputDependenciesCrc);
 
@@ -877,13 +878,19 @@ void FPCGSpawnActorElement::SpawnActors(FPCGSubgraphContext* Context, AActor* Ta
 
 		if (InputDependenciesCrc.IsValid())
 		{
-			Context->SourceComponent->ForEachManagedResource([&ReusedManagedActorsResources, &InputDependenciesCrc, &Context](UPCGManagedResource* InResource)
+			Context->SourceComponent->ForEachManagedResource([&ReusedManagedActorsResource, &InputDependenciesCrc, &Context, NumPoints = Points.Num()](UPCGManagedResource* InResource)
 			{
+				if (ReusedManagedActorsResource)
+				{
+					return;
+				}
+
 				if (UPCGManagedActors* Resource = Cast<UPCGManagedActors>(InResource))
 				{
-					if (Resource->GetCrc().IsValid() && Resource->GetCrc() == InputDependenciesCrc)
+					// We can only re-use the resource if it matches the number of points (if actor failed to spawned for whatever reason, we won't know which point is associated with the fail)
+					if (Resource->GetCrc().IsValid() && Resource->GetCrc() == InputDependenciesCrc && Resource->GetConstGeneratedActors().Num() == NumPoints)
 					{
-						ReusedManagedActorsResources.Add(Resource);
+						ReusedManagedActorsResource = Resource;
 					}
 				}
 			});
@@ -893,24 +900,29 @@ void FPCGSpawnActorElement::SpawnActors(FPCGSubgraphContext* Context, AActor* Ta
 	TArray<AActor*> ProcessedActors;
 	const bool bActorsHavePCGComponents = (UPCGSpawnActorSettings::GetGraphInterfaceFromActorSubclass(InTemplateActorClass) != nullptr);
 
-	if (!ReusedManagedActorsResources.IsEmpty())
+	if (ReusedManagedActorsResource)
 	{
 		// If the actors are fully independent, we might need to make sure to call Generate if the underlying graph has changed - e.g. if the actor is dirty
-		for (UPCGManagedActors* ManagedActors : ReusedManagedActorsResources)
+		ReusedManagedActorsResource->MarkAsReused();
+
+		// If we're in the no-merge case, keep track of these actors to generate.
+		// Also set to the output data the actor reference.
+		if (Settings->Option == EPCGSpawnActorOption::NoMerging)
 		{
-			check(ManagedActors);
-			if (!ManagedActors->IsMarkedUnused() && Settings->bWarnOnIdenticalSpawn)
+			const TArray<TSoftObjectPtr<AActor>>& GeneratedActors = ReusedManagedActorsResource->GetConstGeneratedActors();
+			for (int32 i = 0; i < GeneratedActors.Num(); ++i)
 			{
-				// TODO: Revisit if the stack is added to the managed actors at creation
-				PCGLog::LogWarningOnGraph(LOCTEXT("IdenticalActorSpawn", "Identical actor spawn occurred. It may be beneficial to re-check graph logic for identical spawn conditions (same actor at same location, etc) or repeated nodes."), Context);
-			}
+				const TSoftObjectPtr<AActor>& ManagedActorPtr = GeneratedActors[i];
 
-			ManagedActors->MarkAsReused();
+				// Write to out data the actor reference
+				if (OutPointData && ActorReferenceAttribute)
+				{
+					FPCGPoint& OutPoint = OutPointData->GetMutablePoints()[i + OutPointOffset];
+					OutPointData->Metadata->InitializeOnSet(OutPoint.MetadataEntry);
+					ActorReferenceAttribute->SetValue(OutPoint.MetadataEntry, ManagedActorPtr.ToSoftObjectPath());
+				}
 
-			// There's no setup to be done, just generation if we're in the no-merge case, so keep track of these actors only in this case
-			if (bActorsHavePCGComponents && Settings->Option == EPCGSpawnActorOption::NoMerging)
-			{
-				for (TSoftObjectPtr<AActor>& ManagedActorPtr : ManagedActors->GeneratedActors)
+				if (bActorsHavePCGComponents)
 				{
 					if (AActor* ManagedActor = ManagedActorPtr.Get())
 					{
@@ -979,7 +991,7 @@ void FPCGSpawnActorElement::SpawnActors(FPCGSubgraphContext* Context, AActor* Ta
 				}
 			}
 
-			ManagedActors->GeneratedActors.Add(GeneratedActor);
+			ManagedActors->GetMutableGeneratedActors().AddUnique(GeneratedActor);
 
 			if (bActorsHavePCGComponents)
 			{
diff --git a/Engine/Plugins/PCG/Source/PCG/Private/PCGComponent.cpp b/Engine/Plugins/PCG/Source/PCG/Private/PCGComponent.cpp
--- a/Engine/Plugins/PCG/Source/PCG/Private/PCGComponent.cpp
+++ b/Engine/Plugins/PCG/Source/PCG/Private/PCGComponent.cpp
@@ -94,7 +94,7 @@ namespace PCGComponent
 				{
 					if (UPCGManagedActors* ManagedActors = Cast<UPCGManagedActors>(ManagedResource))
 					{
-						for (const TSoftObjectPtr<AActor>& ManagedActor : ManagedActors->GeneratedActors)
+						for (const TSoftObjectPtr<AActor>& ManagedActor : ManagedActors->GetConstGeneratedActors())
 						{
 							if (const FWorldPartitionActorDescInstance* ActorDescInstance = WorldPartition->GetActorDescInstanceByPath(ManagedActor.ToSoftObjectPath()))
 							{
@@ -309,7 +309,7 @@ void UPCGComponent::AddComponentsToManagedResources(const TArray<UActorComponent
 
 void UPCGComponent::AddActorsToManagedResources(const TArray<AActor*>& InActors)
 {
-	TSet<TSoftObjectPtr<AActor>> ValidActors;
+	TArray<TSoftObjectPtr<AActor>> ValidActors;
 
 	for (AActor* Actor : InActors)
 	{
@@ -321,7 +321,7 @@ void UPCGComponent::AddActorsToManagedResources(const TArray<AActor*>& InActors)
 				Actor->Tags.Add(PCGHelpers::DefaultPCGActorTag);
 			}
 
-			ValidActors.Add(Actor);
+			ValidActors.AddUnique(Actor);
 		}
 	}
 
@@ -331,7 +331,7 @@ void UPCGComponent::AddActorsToManagedResources(const TArray<AActor*>& InActors)
 	}
 
 	UPCGManagedActors* ManagedResource = NewObject<UPCGManagedActors>(this);
-	ManagedResource->GeneratedActors = ValidActors;
+	ManagedResource->GetMutableGeneratedActors() = std::move(ValidActors);
 
 	AddToManagedResources(ManagedResource);
 }
@@ -979,7 +979,7 @@ void UPCGComponent::PurgeUnlinkedResources(const AActor* InActor)
 			}
 			else if (UPCGManagedActors* ManagedActors = Cast<UPCGManagedActors>(ManagedResource))
 			{
-				for (const TSoftObjectPtr<AActor>& GeneratedActor : ManagedActors->GeneratedActors)
+				for (const TSoftObjectPtr<AActor>& GeneratedActor : ManagedActors->GetConstGeneratedActors())
 				{
 					ActorsToDelete.Remove(GeneratedActor);
 				}
@@ -1850,7 +1850,7 @@ void UPCGComponent::PostLoad()
 		if (GeneratedActors_DEPRECATED.Num() > 0)
 		{
 			UPCGManagedActors* ManagedActors = NewObject<UPCGManagedActors>(this);
-			ManagedActors->GeneratedActors = GeneratedActors_DEPRECATED;
+			ManagedActors->GetMutableGeneratedActors() = GeneratedActors_DEPRECATED.Array();
 			GeneratedResources.Add(ManagedActors);
 			GeneratedActors_DEPRECATED.Reset();
 		}
@@ -3800,7 +3800,7 @@ TArray<TSoftObjectPtr<AActor>> UPCGComponent::GetManagedActorPaths(AActor* InAct
 		{
 			if (UPCGManagedActors* ManagedActors = Cast<UPCGManagedActors>(ManagedResource))
 			{
-				ManagedActorPaths.Append(ManagedActors->GeneratedActors);
+				ManagedActorPaths.Append(ManagedActors->GetConstGeneratedActors());
 			}
 		}
 	});
diff --git a/Engine/Plugins/PCG/Source/PCG/Private/PCGManagedResource.cpp b/Engine/Plugins/PCG/Source/PCG/Private/PCGManagedResource.cpp
--- a/Engine/Plugins/PCG/Source/PCG/Private/PCGManagedResource.cpp
+++ b/Engine/Plugins/PCG/Source/PCG/Private/PCGManagedResource.cpp
@@ -73,7 +73,22 @@ void UPCGManagedActors::PostEditImport()
 	// In this case, the managed actors won't be copied along the actor/component,
 	// So we just have to "forget" the actors, leaving the ownership to the original actor only.
 	Super::PostEditImport();
-	GeneratedActors.Reset();
+	GeneratedActorsArray.Reset();
+}
+
+void UPCGManagedActors::PostLoad()
+{
+	Super::PostLoad();
+
+#if WITH_EDITOR
+PRAGMA_DISABLE_DEPRECATION_WARNINGS
+	if (!GeneratedActors.IsEmpty())
+	{
+		GeneratedActorsArray = GeneratedActors.Array();
+		GeneratedActors.Reset();
+	}
+PRAGMA_ENABLE_DEPRECATION_WARNINGS
+#endif // WITH_EDITOR
 }
 
 bool UPCGManagedActors::Release(bool bHardRelease, TSet<TSoftObjectPtr<AActor>>& OutActorsToDelete)
@@ -84,10 +99,10 @@ bool UPCGManagedActors::Release(bool bHardRelease, TSet<TSoftObjectPtr<AActor>>&
 
 	if (!Super::Release(bShouldDeleteActors, OutActorsToDelete))
 	{
-		PCGGeneratedResourcesLogging::LogManagedActorsRelease(this, GeneratedActors, bHardRelease, /*bOnlyMarkedForCleanup=*/true);
+		PCGGeneratedResourcesLogging::LogManagedActorsRelease(this, GeneratedActorsArray, bHardRelease, /*bOnlyMarkedForCleanup=*/true);
 
 		// Mark actors as potentially-to-be-cleaned-up
-		for (TSoftObjectPtr<AActor> GeneratedActor : GeneratedActors)
+		for (TSoftObjectPtr<AActor> GeneratedActor : GeneratedActorsArray)
 		{
 			if (GeneratedActor.IsValid())
 			{
@@ -102,7 +117,7 @@ bool UPCGManagedActors::Release(bool bHardRelease, TSet<TSoftObjectPtr<AActor>>&
 	if (bMarkedTransientOnLoad)
 	{
 		// Here, instead of adding the actors to be deleted (which has the side effect of potentially emptying the package, which leads to its deletion, we will hide the actors instead
-		for (TSoftObjectPtr<AActor> GeneratedActor : GeneratedActors)
+		for (TSoftObjectPtr<AActor> GeneratedActor : GeneratedActorsArray)
 		{
 			// In rare cases where the actor wouldn't be currently loaded (in WP) we must make sure it is done here
 			if (AActor* Actor = GeneratedActor.LoadSynchronous())
@@ -117,15 +132,15 @@ bool UPCGManagedActors::Release(bool bHardRelease, TSet<TSoftObjectPtr<AActor>>&
 	else
 #endif // WITH_EDITOR
 	{
-		OutActorsToDelete.Append(GeneratedActors);
+		OutActorsToDelete.Append(GeneratedActorsArray);
 	}
 
-	PCGGeneratedResourcesLogging::LogManagedActorsRelease(this, GeneratedActors, bHardRelease, /*bOnlyMarkedForCleanup=*/false);
+	PCGGeneratedResourcesLogging::LogManagedActorsRelease(this, GeneratedActorsArray, bHardRelease, /*bOnlyMarkedForCleanup=*/false);
 
 	// Cleanup recursively
 	TInlineComponentArray<UPCGComponent*, 1> ComponentsToCleanup;
 
-	for (TSoftObjectPtr<AActor> GeneratedActor : GeneratedActors)
+	for (TSoftObjectPtr<AActor> GeneratedActor : GeneratedActorsArray)
 	{
 		if (GeneratedActor.IsValid())
 		{
@@ -146,7 +161,7 @@ bool UPCGManagedActors::Release(bool bHardRelease, TSet<TSoftObjectPtr<AActor>>&
 	if (!bMarkedTransientOnLoad)
 #endif
 	{
-		GeneratedActors.Reset();
+		GeneratedActorsArray.Reset();
 	}
 	
 	return true;
@@ -154,14 +169,14 @@ bool UPCGManagedActors::Release(bool bHardRelease, TSet<TSoftObjectPtr<AActor>>&
 
 bool UPCGManagedActors::ReleaseIfUnused(TSet<TSoftObjectPtr<AActor>>& OutActorsToDelete)
 {
-	return Super::ReleaseIfUnused(OutActorsToDelete) || GeneratedActors.IsEmpty();
+	return Super::ReleaseIfUnused(OutActorsToDelete) || GeneratedActorsArray.IsEmpty();
 }
 
 bool UPCGManagedActors::MoveResourceToNewActor(AActor* NewActor)
 {
 	check(NewActor);
 
-	for (TSoftObjectPtr<AActor>& Actor : GeneratedActors)
+	for (TSoftObjectPtr<AActor>& Actor : GeneratedActorsArray)
 	{
 		if (!Actor.IsValid())
 		{
@@ -178,7 +193,7 @@ bool UPCGManagedActors::MoveResourceToNewActor(AActor* NewActor)
 		}
 	}
 
-	GeneratedActors.Empty();
+	GeneratedActorsArray.Empty();
 
 	return true;
 }
@@ -194,7 +209,7 @@ void UPCGManagedActors::MarkAsReused()
 {
 	Super::MarkAsReused();
 
-	for (TSoftObjectPtr<AActor> GeneratedActor : GeneratedActors)
+	for (TSoftObjectPtr<AActor> GeneratedActor : GeneratedActorsArray)
 	{
 		if (GeneratedActor.IsValid())
 		{
@@ -210,7 +225,7 @@ bool UPCGManagedActors::IsManaging(const UObject* InObject) const
 		return false;
 	}
 
-	return Algo::AnyOf(GeneratedActors, [InObject](const TSoftObjectPtr<AActor>& SoftGeneratedActor)
+	return Algo::AnyOf(GeneratedActorsArray, [InObject](const TSoftObjectPtr<AActor>& SoftGeneratedActor)
 		{
 			const AActor* GeneratedActor = SoftGeneratedActor.Get();
 			return GeneratedActor && GeneratedActor == InObject;
@@ -222,7 +237,7 @@ void UPCGManagedActors::ChangeTransientState(EPCGEditorDirtyMode NewEditingMode)
 {
 	const bool bNowTransient = (NewEditingMode == EPCGEditorDirtyMode::Preview);
 
-	for (TSoftObjectPtr<AActor> GeneratedActor : GeneratedActors)
+	for (TSoftObjectPtr<AActor> GeneratedActor : GeneratedActorsArray)
 	{
 		// Make sure to load if needed because we need to affect the actors regardless of the current WP state
 		if (GeneratedActor.LoadSynchronous() != nullptr)
diff --git a/Engine/Plugins/PCG/Source/PCG/Private/Utils/PCGGeneratedResourcesLogging.cpp b/Engine/Plugins/PCG/Source/PCG/Private/Utils/PCGGeneratedResourcesLogging.cpp
--- a/Engine/Plugins/PCG/Source/PCG/Private/Utils/PCGGeneratedResourcesLogging.cpp
+++ b/Engine/Plugins/PCG/Source/PCG/Private/Utils/PCGGeneratedResourcesLogging.cpp
@@ -41,11 +41,11 @@ namespace PCGGeneratedResourcesLogging
 			UE_LOG(LogPCG, Log, TEXT("[PCGMANAGEDRESOURCES] [%s/%s]         Managed actors (%d):"),
 				(InComponent && InComponent->GetOwner()) ? *InComponent->GetOwner()->GetName() : TEXT("MISSINGCOMPONENT"),
 				(InComponent && InComponent->GetGraph()) ? *InComponent->GetGraph()->GetName() : TEXT("MISSINGGRAPH"),
-				ManagedActors->GeneratedActors.Num());
+				ManagedActors->GetConstGeneratedActors().Num());
 
 			const int32 MaxPrint = CVarGenResourcesLoggingMaxPrintCount.GetValueOnAnyThread();
 			int32 Count = 0;
-			for (const TSoftObjectPtr<AActor>& Actor : ManagedActors->GeneratedActors)
+			for (const TSoftObjectPtr<AActor>& Actor : ManagedActors->GetConstGeneratedActors())
 			{
 				if (Count++ >= MaxPrint)
 				{
@@ -58,14 +58,14 @@ namespace PCGGeneratedResourcesLogging
 						(InComponent && InComponent->GetOwner()) ? *InComponent->GetOwner()->GetName() : TEXT("MISSINGCOMPONENT"),
 						(InComponent && InComponent->GetGraph()) ? *InComponent->GetGraph()->GetName() : TEXT("MISSINGGRAPH"),
 						*Actor->GetFName().ToString(),
-						Count, ManagedActors->GeneratedActors.Num());
+						Count, ManagedActors->GetConstGeneratedActors().Num());
 				}
 				else
 				{
 					UE_LOG(LogPCG, Log, TEXT("[PCGMANAGEDRESOURCES] [%s/%s]                 |- NULL (%d/%d)"),
 						(InComponent && InComponent->GetOwner()) ? *InComponent->GetOwner()->GetName() : TEXT("MISSINGCOMPONENT"),
 						(InComponent && InComponent->GetGraph()) ? *InComponent->GetGraph()->GetName() : TEXT("MISSINGGRAPH"),
-						Count, ManagedActors->GeneratedActors.Num());
+						Count, ManagedActors->GetConstGeneratedActors().Num());
 				}
 			}
 		}
@@ -338,7 +338,7 @@ namespace PCGGeneratedResourcesLogging
 #endif
 	}
 
-	void LogManagedActorsRelease(const UPCGManagedResource* InResource, const TSet<TSoftObjectPtr<AActor>>& ActorsToDelete, bool bHardRelease, bool bOnlyMarkedForCleanup)
+	void LogManagedActorsRelease(const UPCGManagedResource* InResource, const TArray<TSoftObjectPtr<AActor>>& ActorsToDelete, bool bHardRelease, bool bOnlyMarkedForCleanup)
 	{
 #if WITH_EDITOR
 		if (!LogEnabled())
diff --git a/Engine/Plugins/PCG/Source/PCG/Private/Utils/PCGGeneratedResourcesLogging.h b/Engine/Plugins/PCG/Source/PCG/Private/Utils/PCGGeneratedResourcesLogging.h
--- a/Engine/Plugins/PCG/Source/PCG/Private/Utils/PCGGeneratedResourcesLogging.h
+++ b/Engine/Plugins/PCG/Source/PCG/Private/Utils/PCGGeneratedResourcesLogging.h
@@ -29,7 +29,7 @@ namespace PCGGeneratedResourcesLogging
 	void LogCleanupUnusedManagedResourcesResource(const UPCGComponent* InComponent, UPCGManagedResource* Resource);
 	void LogCleanupUnusedManagedResourcesFinished(const UPCGComponent* InComponent, const TArray<UPCGManagedResource*>& GeneratedResources);
 
-	void LogManagedActorsRelease(const UPCGManagedResource* InResource, const TSet<TSoftObjectPtr<AActor>>& ActorsToDelete, bool bHardRelease, bool bOnlyMarkedForCleanup);
+	void LogManagedActorsRelease(const UPCGManagedResource* InResource, const TArray<TSoftObjectPtr<AActor>>& ActorsToDelete, bool bHardRelease, bool bOnlyMarkedForCleanup);
 
 	void LogManagedResourceSoftRelease(UPCGManagedResource* InResource);
 	void LogManagedResourceHardRelease(UPCGManagedResource* InResource);
diff --git a/Engine/Plugins/PCG/Source/PCG/Public/PCGManagedResource.h b/Engine/Plugins/PCG/Source/PCG/Public/PCGManagedResource.h
--- a/Engine/Plugins/PCG/Source/PCG/Public/PCGManagedResource.h
+++ b/Engine/Plugins/PCG/Source/PCG/Public/PCGManagedResource.h
@@ -90,6 +90,7 @@ class PCG_API UPCGManagedActors : public UPCGManagedResource
 public:
 	//~Begin UObject interface
 	virtual void PostEditImport() override;
+	virtual void PostLoad() override;
 	//~End UObject interface
 
 	//~Begin UPCGManagedResource interface
@@ -106,8 +107,8 @@ public:
 #endif
 	//~End UPCGManagedResource interface
 
-	UPROPERTY(BlueprintReadWrite, EditAnywhere, Category = GeneratedData)
-	TSet<TSoftObjectPtr<AActor>> GeneratedActors;
+	const TArray<TSoftObjectPtr<AActor>>& GetConstGeneratedActors() const { return GeneratedActorsArray; }
+	TArray<TSoftObjectPtr<AActor>>& GetMutableGeneratedActors() { return GeneratedActorsArray; }
 
 	/** Controls whether the resource will be removed at the beginning of the generation instead of being kept until the end, in the eventuality it is reused. 
 	* In practice, this means that the actors will be deleted as part of the initial cleanup in the generation instead of being removed at the end.
@@ -115,6 +116,14 @@ public:
 	*/
 	UPROPERTY(BlueprintReadWrite, EditAnywhere, Category = GeneratedData)
 	bool bSupportsReset = true;
+
+	UE_DEPRECATED(5.6, "Generated Actors is now an array, you must use the Getter/Setter, as this set won't be used anymore.")
+	UPROPERTY(meta = (DeprecatedProperty))
+	TSet<TSoftObjectPtr<AActor>> GeneratedActors;
+	
+private:
+	UPROPERTY(BlueprintReadWrite, EditAnywhere, Category = GeneratedData, meta = (AllowPrivateAccess = "true", DisplayName = "GeneratedActors"))
+	TArray<TSoftObjectPtr<AActor>> GeneratedActorsArray;
 };
 
 UCLASS(Abstract)
