[PCG] Missing redirectors for Metadata

#rnx

diff --git a/Engine/Plugins/PCG/Config/BasePCG.ini b/Engine/Plugins/PCG/Config/BasePCG.ini
--- a/Engine/Plugins/PCG/Config/BasePCG.ini
+++ b/Engine/Plugins/PCG/Config/BasePCG.ini
@@ -26,6 +26,10 @@
 +FunctionRedirects=(OldName="PCGMetadata.HasAttribute",NewName="PCGMetadata.BP_HasAttribute")
 +FunctionRedirects=(OldName="PCGMetadata.AddAttribute",NewName="PCGMetadata.BP_AddAttribute")
 +FunctionRedirects=(OldName="PCGMetadata.CopyAttribute",NewName="PCGMetadata.BP_CopyAttribute")
++FunctionRedirects=(OldName="PCGMetadata.DeleteAttribute",NewName="PCGMetadata.BP_DeleteAttribute")
++FunctionRedirects=(OldName="PCGMetadata.CopyExistingAttribute",NewName="PCGMetadata.BP_CopyExistingAttribute")
++FunctionRedirects=(OldName="PCGMetadata.RenameAttribute",NewName="PCGMetadata.BP_RenameAttribute")
++FunctionRedirects=(OldName="PCGMetadata.ClearAttribute",NewName="PCGMetadata.BP_ClearAttribute")
 
 +PropertyRedirects=(OldName="PCGSplineSamplerSettings.Params",NewName="PCGSplineSamplerSettings.SamplerParams")
 +PropertyRedirects=(OldName="PCGProjectionSettings.Params",NewName="PCGProjectionSettings.ProjectionParams")
