Added some options to remove scaling/rotation from position transformation in the copy points as needed.

diff --git a/Engine/Plugins/PCG/Shaders/Private/Elements/PCGCopyPoints.usf b/Engine/Plugins/PCG/Shaders/Private/Elements/PCGCopyPoints.usf
--- a/Engine/Plugins/PCG/Shaders/Private/Elements/PCGCopyPoints.usf
+++ b/Engine/Plugins/PCG/Shaders/Private/Elements/PCGCopyPoints.usf
@@ -134,22 +134,26 @@ void Main(uint3 GroupId : SV_GroupID, uint GroupIndex : SV_GroupIndex)
 	float3 TargetScale = Target_GetScale(TargetDataIndex, TargetPointIndexInData);
 	float4 TargetColor = Target_GetColor(TargetDataIndex, TargetPointIndexInData);
 	uint TargetSeed = Target_GetSeed(TargetDataIndex, TargetPointIndexInData);
-	
-	float3 OutPosition = TargetPosition + (QuatRotateVector(TargetRotation, (SourcePosition * TargetScale)));
-	float4 OutRotation = float4(0, 0, 0, 1);
-	float3 OutScale = (float3)0;
-	float4 OutColor = (float4)0;
-	uint OutSeed = 0;
 
 	const uint RotationInheritance = GetRotationInheritance();
+	const uint ApplyTargetRotationToPositions = GetApplyTargetRotationToPositions();
 	const uint ScaleInheritance = GetScaleInheritance();
+	const uint ApplyTargetScaleToPositions = GetApplyTargetScaleToPositions();
 	const uint ColorInheritance = GetColorInheritance();
 	const uint SeedInheritance = GetSeedInheritance();
 	const uint AttributeInheritance = GetAttributeInheritance();
 
+	float3 ScaleForPosition = ApplyTargetScaleToPositions > 0 ? TargetScale : float3(1, 1, 1);
+
+	float3 OutPosition = TargetPosition + (ApplyTargetRotationToPositions > 0 ? QuatRotateVector(TargetRotation, (SourcePosition * ScaleForPosition)) : (SourcePosition * ScaleForPosition));
+	float4 OutRotation = float4(0, 0, 0, 1);
+	float3 OutScale = (float3)0;
+	float4 OutColor = (float4)0;
+	uint OutSeed = 0;
+
 	if (RotationInheritance == COPY_POINTS_INHERITANCE_MODE_Relative)
 	{
-	  OutRotation = QuatMultiply(TargetRotation, SourceRotation);
+		OutRotation = QuatMultiply(TargetRotation, SourceRotation);
 	}
 	else if (RotationInheritance == COPY_POINTS_INHERITANCE_MODE_Source)
 	{
diff --git a/Engine/Plugins/PCG/Source/PCG/Private/Compute/DataInterfaces/Elements/PCGCopyPointsDataInterface.cpp b/Engine/Plugins/PCG/Source/PCG/Private/Compute/DataInterfaces/Elements/PCGCopyPointsDataInterface.cpp
--- a/Engine/Plugins/PCG/Source/PCG/Private/Compute/DataInterfaces/Elements/PCGCopyPointsDataInterface.cpp
+++ b/Engine/Plugins/PCG/Source/PCG/Private/Compute/DataInterfaces/Elements/PCGCopyPointsDataInterface.cpp
@@ -19,10 +19,18 @@ void UPCGCopyPointsDataInterface::GetSupportedInputs(TArray<FShaderFunctionDefin
 		.SetName(TEXT("GetRotationInheritance"))
 		.AddReturnType(FShaderValueType::Get(EShaderFundamentalType::Uint));
 
+	OutFunctions.AddDefaulted_GetRef()
+		.SetName(TEXT("GetApplyTargetRotationToPositions"))
+		.AddReturnType(FShaderValueType::Get(EShaderFundamentalType::Uint));
+
 	OutFunctions.AddDefaulted_GetRef()
 		.SetName(TEXT("GetScaleInheritance"))
 		.AddReturnType(FShaderValueType::Get(EShaderFundamentalType::Uint));
 
+	OutFunctions.AddDefaulted_GetRef()
+		.SetName(TEXT("GetApplyTargetScaleToPositions"))
+		.AddReturnType(FShaderValueType::Get(EShaderFundamentalType::Uint));
+
 	OutFunctions.AddDefaulted_GetRef()
 		.SetName(TEXT("GetColorInheritance"))
 		.AddReturnType(FShaderValueType::Get(EShaderFundamentalType::Uint));
@@ -42,7 +50,9 @@ void UPCGCopyPointsDataInterface::GetSupportedInputs(TArray<FShaderFunctionDefin
 
 BEGIN_SHADER_PARAMETER_STRUCT(FPCGCopyPointsDataInterfaceParameters,)
 	SHADER_PARAMETER(uint32, RotationInheritance)
+	SHADER_PARAMETER(uint32, ApplyTargetRotationToPositions)
 	SHADER_PARAMETER(uint32, ScaleInheritance)
+	SHADER_PARAMETER(uint32, ApplyTargetScaleToPositions)
 	SHADER_PARAMETER(uint32, ColorInheritance)
 	SHADER_PARAMETER(uint32, SeedInheritance)
 	SHADER_PARAMETER(uint32, AttributeInheritance)
@@ -63,14 +73,18 @@ void UPCGCopyPointsDataInterface::GetHLSL(FString& OutHLSL, FString const& InDat
 
 	OutHLSL += FString::Format(TEXT(
 		"uint {DataInterfaceName}_RotationInheritance;\n"
+		"uint {DataInterfaceName}_ApplyTargetRotationToPositions;\n"
 		"uint {DataInterfaceName}_ScaleInheritance;\n"
+		"uint {DataInterfaceName}_ApplyTargetScaleToPositions;\n"
 		"uint {DataInterfaceName}_ColorInheritance;\n"
 		"uint {DataInterfaceName}_SeedInheritance;\n"
 		"uint {DataInterfaceName}_AttributeInheritance;\n"
 		"uint {DataInterfaceName}_bCopyEachSourceOnEveryTarget;\n"
 		"\n"
 		"uint GetRotationInheritance_{DataInterfaceName}() { return {DataInterfaceName}_RotationInheritance; }\n"
+		"uint GetApplyTargetRotationToPositions_{DataInterfaceName}() { return {DataInterfaceName}_ApplyTargetRotationToPositions; }\n"
 		"uint GetScaleInheritance_{DataInterfaceName}() { return {DataInterfaceName}_ScaleInheritance; }\n"
+		"uint GetApplyTargetScaleToPositions_{DataInterfaceName}() { return {DataInterfaceName}_ApplyTargetScaleToPositions; }\n"
 		"uint GetColorInheritance_{DataInterfaceName}() { return {DataInterfaceName}_ColorInheritance; }\n"
 		"uint GetSeedInheritance_{DataInterfaceName}() { return {DataInterfaceName}_SeedInheritance; }\n"
 		"uint GetAttributeInheritance_{DataInterfaceName}() { return {DataInterfaceName}_AttributeInheritance; }\n"
@@ -90,7 +104,9 @@ FComputeDataProviderRenderProxy* UPCGCopyPointsDataProvider::GetRenderProxy()
 {
 	return new FPCGCopyPointsDataProviderProxy(
 		static_cast<uint32>(Settings->RotationInheritance),
+		static_cast<uint32>((Settings->RotationInheritance != EPCGCopyPointsInheritanceMode::Source || Settings->bApplyTargetRotationToPositions) ? 1 : 0),
 		static_cast<uint32>(Settings->ScaleInheritance),
+		static_cast<uint32>((Settings->ScaleInheritance != EPCGCopyPointsInheritanceMode::Source || Settings->bApplyTargetScaleToPositions) ? 1 : 0),
 		static_cast<uint32>(Settings->ColorInheritance),
 		static_cast<uint32>(Settings->SeedInheritance),
 		static_cast<uint32>(Settings->AttributeInheritance),
@@ -115,7 +131,9 @@ void FPCGCopyPointsDataProviderProxy::GatherDispatchData(FDispatchData const& In
 	{
 		FParameters& Parameters = ParameterArray[InvocationIndex];
 		Parameters.RotationInheritance = RotationInheritance;
+		Parameters.ApplyTargetRotationToPositions = ApplyTargetRotationToPositions;
 		Parameters.ScaleInheritance = ScaleInheritance;
+		Parameters.ApplyTargetScaleToPositions = ApplyTargetScaleToPositions;
 		Parameters.ColorInheritance = ColorInheritance;
 		Parameters.SeedInheritance = SeedInheritance;
 		Parameters.AttributeInheritance = AttributeInheritance;
diff --git a/Engine/Plugins/PCG/Source/PCG/Private/Compute/DataInterfaces/Elements/PCGCopyPointsDataInterface.h b/Engine/Plugins/PCG/Source/PCG/Private/Compute/DataInterfaces/Elements/PCGCopyPointsDataInterface.h
--- a/Engine/Plugins/PCG/Source/PCG/Private/Compute/DataInterfaces/Elements/PCGCopyPointsDataInterface.h
+++ b/Engine/Plugins/PCG/Source/PCG/Private/Compute/DataInterfaces/Elements/PCGCopyPointsDataInterface.h
@@ -47,9 +47,11 @@ public:
 class FPCGCopyPointsDataProviderProxy : public FComputeDataProviderRenderProxy
 {
 public:
-	FPCGCopyPointsDataProviderProxy(uint32 InRotationInheritance, uint32 InScaleInheritance, uint32 InColorInheritance, uint32 InSeedInheritance, uint32 InAttributeInheritance, uint32 bInCopyEachSourceOnEveryTarget)
+	FPCGCopyPointsDataProviderProxy(uint32 InRotationInheritance, uint32 InApplyTargetRotationToPositions, uint32 InScaleInheritance, uint32 InApplyTargetScaleToPositions, uint32 InColorInheritance, uint32 InSeedInheritance, uint32 InAttributeInheritance, uint32 bInCopyEachSourceOnEveryTarget)
 		: RotationInheritance(InRotationInheritance)
+		, ApplyTargetRotationToPositions(InApplyTargetRotationToPositions)
 		, ScaleInheritance(InScaleInheritance)
+		, ApplyTargetScaleToPositions(InApplyTargetScaleToPositions)
 		, ColorInheritance(InColorInheritance)
 		, SeedInheritance(InSeedInheritance)
 		, AttributeInheritance(InAttributeInheritance)
@@ -65,7 +67,9 @@ protected:
 	using FParameters = FPCGCopyPointsDataInterfaceParameters;
 
 	uint32 RotationInheritance = 0;
+	uint32 ApplyTargetRotationToPositions = 0;
 	uint32 ScaleInheritance = 0;
+	uint32 ApplyTargetScaleToPositions = 0;
 	uint32 ColorInheritance = 0;
 	uint32 SeedInheritance = 0;
 	uint32 AttributeInheritance = 0;
diff --git a/Engine/Plugins/PCG/Source/PCG/Private/Elements/PCGCopyPoints.cpp b/Engine/Plugins/PCG/Source/PCG/Private/Elements/PCGCopyPoints.cpp
--- a/Engine/Plugins/PCG/Source/PCG/Private/Elements/PCGCopyPoints.cpp
+++ b/Engine/Plugins/PCG/Source/PCG/Private/Elements/PCGCopyPoints.cpp
@@ -169,6 +169,8 @@ bool FPCGCopyPointsElement::ExecuteInternal(FPCGContext* Context) const
 	const EPCGCopyPointsInheritanceMode ColorInheritance = Settings->ColorInheritance;
 	const EPCGCopyPointsInheritanceMode SeedInheritance = Settings->SeedInheritance;
 	const EPCGCopyPointsMetadataInheritanceMode AttributeInheritance = Settings->AttributeInheritance;
+	const bool bApplyTargetRotationToPositions = Settings->bApplyTargetRotationToPositions;
+	const bool bApplyTargetScaleToPositions = Settings->bApplyTargetScaleToPositions;
 
 	const TArray<FPCGTaggedData> Sources = Context->InputData.GetInputsByPin(PCGCopyPointsConstants::SourcePointsLabel);
 	const TArray<FPCGTaggedData> Targets = Context->InputData.GetInputsByPin(PCGCopyPointsConstants::TargetPointsLabel);
@@ -363,26 +365,40 @@ bool FPCGCopyPointsElement::ExecuteInternal(FPCGContext* Context) const
 
 			OutPoint = SourcePoint;
 
-			// Set the position, rotation and scale as relative by default.
-			OutPoint.Transform = SourcePoint.Transform * TargetPoint.Transform;
+			FTransform SourceTransform(SourcePoint.Transform.GetLocation());
+			FTransform TargetTransform(TargetPoint.Transform.GetLocation());
 
 			// Set Rotation, Scale, and Color based on inheritance mode
-			if (RotationInheritance == EPCGCopyPointsInheritanceMode::Source)
+			if (RotationInheritance != EPCGCopyPointsInheritanceMode::Target)
 			{
-				OutPoint.Transform.SetRotation(SourcePoint.Transform.GetRotation());
+				SourceTransform.SetRotation(SourcePoint.Transform.GetRotation());
 			}
-			else if (RotationInheritance == EPCGCopyPointsInheritanceMode::Target)
+
+			if (RotationInheritance != EPCGCopyPointsInheritanceMode::Source || bApplyTargetRotationToPositions)
 			{
-				OutPoint.Transform.SetRotation(TargetPoint.Transform.GetRotation());
+				TargetTransform.SetRotation(TargetPoint.Transform.GetRotation());
 			}
 
-			if (ScaleInheritance == EPCGCopyPointsInheritanceMode::Source)
+			if (ScaleInheritance != EPCGCopyPointsInheritanceMode::Target)
 			{ 
-				OutPoint.Transform.SetScale3D(SourcePoint.Transform.GetScale3D());
+				SourceTransform.SetScale3D(SourcePoint.Transform.GetScale3D());
+			}
+			
+			if (ScaleInheritance != EPCGCopyPointsInheritanceMode::Source || bApplyTargetScaleToPositions)
+			{
+				TargetTransform.SetScale3D(TargetPoint.Transform.GetScale3D());
+			}
+
+			OutPoint.Transform = SourceTransform * TargetTransform;
+
+			if (RotationInheritance == EPCGCopyPointsInheritanceMode::Source && bApplyTargetRotationToPositions)
+			{
+				OutPoint.Transform.SetRotation(SourcePoint.Transform.GetRotation());
 			}
-			else if (ScaleInheritance == EPCGCopyPointsInheritanceMode::Target)
+
+			if (ScaleInheritance == EPCGCopyPointsInheritanceMode::Source && bApplyTargetScaleToPositions)
 			{
-				OutPoint.Transform.SetScale3D(TargetPoint.Transform.GetScale3D());
+				OutPoint.Transform.SetScale3D(SourcePoint.Transform.GetScale3D());
 			}
 
 			if (ColorInheritance == EPCGCopyPointsInheritanceMode::Relative)
diff --git a/Engine/Plugins/PCG/Source/PCG/Public/Elements/PCGCopyPoints.h b/Engine/Plugins/PCG/Source/PCG/Public/Elements/PCGCopyPoints.h
--- a/Engine/Plugins/PCG/Source/PCG/Public/Elements/PCGCopyPoints.h
+++ b/Engine/Plugins/PCG/Source/PCG/Public/Elements/PCGCopyPoints.h
@@ -69,10 +69,18 @@ public:
 	UPROPERTY(BlueprintReadWrite, EditAnywhere, Category = Settings, meta = (PCG_Overridable))
 	EPCGCopyPointsInheritanceMode RotationInheritance = EPCGCopyPointsInheritanceMode::Relative;
 
+	/** If this option is set, points will have their source position transformed using the target transform with rotation */
+	UPROPERTY(BlueprintReadWrite, EditAnywhere, Category = Settings, meta = (PCG_Overridable, EditCondition = "RotationInheritance==EPCGCopyPointsInheritanceMode::Source", EditConditionHides))
+	bool bApplyTargetRotationToPositions = true;
+
 	/** The method used to determine output point scale */
 	UPROPERTY(BlueprintReadWrite, EditAnywhere, Category = Settings, meta = (PCG_Overridable))
 	EPCGCopyPointsInheritanceMode ScaleInheritance = EPCGCopyPointsInheritanceMode::Relative;
 
+	/** If this option is set, points will have their source position transformed using the target transform with scale */
+	UPROPERTY(BlueprintReadWrite, EditAnywhere, Category = Settings, meta = (PCG_Overridable, EditCondition="ScaleInheritance==EPCGCopyPointsInheritanceMode::Source", EditConditionHides))
+	bool bApplyTargetScaleToPositions = true;
+
 	/** The method used to determine output point color */
 	UPROPERTY(BlueprintReadWrite, EditAnywhere, Category = Settings, meta = (PCG_Overridable))
 	EPCGCopyPointsInheritanceMode ColorInheritance = EPCGCopyPointsInheritanceMode::Relative;
