[PCG] Fix flipped X and Y in texture data readback

#rnx

diff --git a/Engine/Plugins/PCG/Source/PCG/Private/Data/PCGTextureData.cpp b/Engine/Plugins/PCG/Source/PCG/Private/Data/PCGTextureData.cpp
--- a/Engine/Plugins/PCG/Source/PCG/Private/Data/PCGTextureData.cpp
+++ b/Engine/Plugins/PCG/Source/PCG/Private/Data/PCGTextureData.cpp
@@ -791,9 +791,9 @@ bool UPCGTextureData::InitializeFromGPUTexture()
 
 				int32 ActualTexelIndex = 0;
 
-				for (int32 TexelX = 0; TexelX < TextureSize.X; ++TexelX)
+				for (int32 TexelY = 0; TexelY < TextureSize.Y; ++TexelY)
 				{
-					for (int32 TexelY = 0; TexelY < TextureSize.Y; ++TexelY)
+					for (int32 TexelX = 0; TexelX < TextureSize.X; ++TexelX)
 					{
 						const int32 ReadbackTexelIndex = TexelY * ReadbackWidth + TexelX;
 
