[PCG] Removed the requirement of the first pin on the Merge/Gather nodes

Shouldn't explicitly mark the static input pin on any of these as required, as data on any input pin should prevent culling.

#rnx

diff --git a/Engine/Plugins/PCG/Source/PCG/Private/Elements/PCGGather.cpp b/Engine/Plugins/PCG/Source/PCG/Private/Elements/PCGGather.cpp
--- a/Engine/Plugins/PCG/Source/PCG/Private/Elements/PCGGather.cpp
+++ b/Engine/Plugins/PCG/Source/PCG/Private/Elements/PCGGather.cpp
@@ -18,6 +18,14 @@ PRAGMA_ENABLE_DEPRECATION_WARNINGS
 }
 #endif // WITH_EDITOR
 
+TArray<FPCGPinProperties> UPCGGatherSettings::StaticInputPinProperties() const
+{
+	TArray<FPCGPinProperties> PinProperties;
+	// Do not explicitly mark the static input pin as required, as data on any input pin should prevent culling.
+	PinProperties.Emplace(PCGPinConstants::DefaultInputLabel, EPCGDataType::Any);
+	return PinProperties;
+}
+
 TArray<FPCGPinProperties> UPCGGatherSettings::OutputPinProperties() const
 {
 	TArray<FPCGPinProperties> PinProperties;
diff --git a/Engine/Plugins/PCG/Source/PCG/Private/Elements/PCGMergeAttributes.cpp b/Engine/Plugins/PCG/Source/PCG/Private/Elements/PCGMergeAttributes.cpp
--- a/Engine/Plugins/PCG/Source/PCG/Private/Elements/PCGMergeAttributes.cpp
+++ b/Engine/Plugins/PCG/Source/PCG/Private/Elements/PCGMergeAttributes.cpp
@@ -44,9 +44,8 @@ FName UPCGMergeAttributesSettings::GetDynamicInputPinsBaseLabel() const
 TArray<FPCGPinProperties> UPCGMergeAttributesSettings::StaticInputPinProperties() const
 {
 	TArray<FPCGPinProperties> PinProperties;
-	FPCGPinProperties& InputPinProperty = PinProperties.Emplace_GetRef(PCGPinConstants::DefaultInputLabel, EPCGDataType::Param, /*bAllowMultipleConnections=*/true);
-	InputPinProperty.SetRequiredPin();
-
+	// Do not explicitly mark the static input pin as required, as data on any input pin should prevent culling.
+	PinProperties.Emplace(PCGPinConstants::DefaultInputLabel, EPCGDataType::Param);
 	return PinProperties;
 }
 
diff --git a/Engine/Plugins/PCG/Source/PCG/Private/Elements/PCGMergeElement.cpp b/Engine/Plugins/PCG/Source/PCG/Private/Elements/PCGMergeElement.cpp
--- a/Engine/Plugins/PCG/Source/PCG/Private/Elements/PCGMergeElement.cpp
+++ b/Engine/Plugins/PCG/Source/PCG/Private/Elements/PCGMergeElement.cpp
@@ -37,7 +37,10 @@ FName UPCGMergeSettings::GetDynamicInputPinsBaseLabel() const
 
 TArray<FPCGPinProperties> UPCGMergeSettings::StaticInputPinProperties() const
 {
-	return DefaultPointInputPinProperties();
+	TArray<FPCGPinProperties> Properties;
+	// Do not explicitly mark the static input pin as required, as data on any input pin should prevent culling.
+	Properties.Emplace(PCGPinConstants::DefaultInputLabel, EPCGDataType::Point);
+	return Properties;
 }
 
 #if WITH_EDITOR
diff --git a/Engine/Plugins/PCG/Source/PCG/Public/Elements/PCGGather.h b/Engine/Plugins/PCG/Source/PCG/Public/Elements/PCGGather.h
--- a/Engine/Plugins/PCG/Source/PCG/Public/Elements/PCGGather.h
+++ b/Engine/Plugins/PCG/Source/PCG/Public/Elements/PCGGather.h
@@ -26,6 +26,7 @@ public:
 	virtual bool HasDynamicPins() const override { return true; }
 
 protected:
+	TArray<FPCGPinProperties> StaticInputPinProperties() const override;
 	virtual TArray<FPCGPinProperties> OutputPinProperties() const override;
 	virtual FPCGElementPtr CreateElement() const override;
 	//~End UPCGSettings interface
