[PCG] Implement PCGPointArrayData (Struct of Array data structure)
- Move some implementation of PCGPointData into PCGBasePointData
- Add concept of TPCGValueRange/TConstPCGValueRange to abstract reading / writing to different UPCGBasePointData class implementations
- Implement UPCGBasePointData::GetPointAccessorMethods() (currently UPCGPointData still implements its own version but after more testing we can remove them)
- ValueRange unittests
- Replace PCGTestCommon CreatePointData methods with UPCGBasePointData compatible templated versions
- Other examples of UPCGBasePointData implementations: CreatePointGrid, PCGLevelToAsset
- Update some code from the previous index octree changelist to be more efficient using ValueRanges

#rnx

diff --git a/Engine/Plugins/PCG/Source/PCG/Private/Data/PCGBasePointData.cpp b/Engine/Plugins/PCG/Source/PCG/Private/Data/PCGBasePointData.cpp
--- a/Engine/Plugins/PCG/Source/PCG/Private/Data/PCGBasePointData.cpp
+++ b/Engine/Plugins/PCG/Source/PCG/Private/Data/PCGBasePointData.cpp
@@ -2,9 +2,1446 @@
 
 #include "Data/PCGBasePointData.h"
 
+#include "Helpers/PCGHelpers.h"
+#include "Helpers/PCGPointHelpers.h"
+#include "Helpers/PCGTagHelpers.h"
+#include "Metadata/PCGMetadataAccessor.h"
+#include "Metadata/Accessors/IPCGAttributeAccessorTpl.h"
+#include "Metadata/Accessors/PCGAttributeAccessorHelpers.h"
+#include "Metadata/Accessors/PCGAttributeAccessorKeys.h"
+
+#include "GameFramework/Actor.h"
+#include "Serialization/ArchiveCrc32.h"
+
+static TAutoConsoleVariable<bool> CVarCacheFullPointDataCrc(
+	TEXT("pcg.Cache.FullPointDataCrc"),
+	true,
+	TEXT("Enable fine-grained CRC of point data for change tracking on elements that request it, rather than using data UID."));
+
+namespace PCGPointHelpers
+{
+	bool GetDistanceRatios(const FTransform& InTransform, const FVector& InBoundsMin, const FVector& InBoundsMax, const float InSteepness, const FVector& InPosition, FVector& OutRatios)
+	{
+		FVector LocalPosition = InTransform.InverseTransformPosition(InPosition);
+		LocalPosition -= (InBoundsMax + InBoundsMin) / 2;
+		LocalPosition /= PCGPointHelpers::GetExtents(InBoundsMin, InBoundsMax);
+
+		// ]-2+s, 2-s] is the valid range of values
+		const FVector::FReal LowerBound = InSteepness - 2;
+		const FVector::FReal HigherBound = 2 - InSteepness;
+
+		if (LocalPosition.X <= LowerBound || LocalPosition.X > HigherBound ||
+			LocalPosition.Y <= LowerBound || LocalPosition.Y > HigherBound ||
+			LocalPosition.Z <= LowerBound || LocalPosition.Z > HigherBound)
+		{
+			return false;
+		}
+
+		// [-s, +s] is the range where the density is 1 on that axis
+		const FVector::FReal XDist = FMath::Max(0, FMath::Abs(LocalPosition.X) - InSteepness);
+		const FVector::FReal YDist = FMath::Max(0, FMath::Abs(LocalPosition.Y) - InSteepness);
+		const FVector::FReal ZDist = FMath::Max(0, FMath::Abs(LocalPosition.Z) - InSteepness);
+
+		const FVector::FReal DistanceScale = FMath::Max(2 - 2 * InSteepness, KINDA_SMALL_NUMBER);
+
+		OutRatios.X = XDist / DistanceScale;
+		OutRatios.Y = YDist / DistanceScale;
+		OutRatios.Z = ZDist / DistanceScale;
+		return true;
+	}
+
+	FVector::FReal ManhattanDensity(const FTransform& InTransform, const FVector& InBoundsMin, const FVector& InBoundsMax, const float InSteepness, const float InDensity, const FVector& InPosition)
+	{
+		FVector Ratios;
+		if (GetDistanceRatios(InTransform, InBoundsMin, InBoundsMax, InSteepness, InPosition, Ratios))
+		{
+			return InDensity * (1 - Ratios.X) * (1 - Ratios.Y) * (1 - Ratios.Z);
+		}
+		else
+		{
+			return 0;
+		}
+	}
+
+	FVector::FReal InverseEuclidianDistance(const FTransform& InTransform, const FVector& InBoundsMin, const FVector& InBoundsMax, const float InSteepness, const FVector& InPosition)
+	{
+		FVector Ratios;
+		if (GetDistanceRatios(InTransform, InBoundsMin, InBoundsMax, InSteepness, InPosition, Ratios))
+		{
+			return 1.0 - Ratios.Length();
+		}
+		else
+		{
+			return 0;
+		}
+	}
+
+	/** Computes reasonable overlap ratio for point, 1d, 2d and volume overlaps, to be used as weights.
+	* Note that this assumes that either data set is homogeneous in its points dimension (either 0d, 1d, 2d, 3d)
+	* Otherwise there will be some artifacts from our assumption here (namely using a 1.0 value for the additional coordinates).
+	*/
+	FVector::FReal ComputeOverlapRatio(const FBox& Numerator, const FBox& Denominator)
+	{
+		const FVector NumeratorExtent = Numerator.GetExtent();
+		const FVector DenominatorExtent = Denominator.GetExtent();
+
+		return (FVector::FReal)((DenominatorExtent.X > 0 ? NumeratorExtent.X / DenominatorExtent.X : 1.0) *
+			(DenominatorExtent.Y > 0 ? NumeratorExtent.Y / DenominatorExtent.Y : 1.0) *
+			(DenominatorExtent.Z > 0 ? NumeratorExtent.Z / DenominatorExtent.Z : 1.0));
+	}
+
+	FVector::FReal VolumeOverlap(const FTransform& InTransform, const FVector& InBoundsMin, const FVector& InBoundsMax, const float InSteepness, const FBox& InBounds, const FMatrix& InInverseTransform)
+	{
+		// This is similar in idea to SAT considering we have two boxes - since we will test all 6 axes.
+		// However, there is some uncertainty due to rotation, and using the overlap value as-is is an overestimation, which might not be critical in this case
+		// TODO: investigate if we should do a 8-pt test instead (would be more precise, but significantly more costly).
+		// Implementation note: we are using FMatrix here because we want to support non-uniform scales
+		const FBox PointBounds = PCGPointHelpers::GetLocalDensityBounds(InSteepness, InBoundsMin, InBoundsMax);
+
+		FMatrix PointTransformToInTransform = InTransform.ToMatrixWithScale() * InInverseTransform;
+		const FBox PointBoundsTransformed = PointBounds.TransformBy(PointTransformToInTransform);
+
+		const FBox FirstOverlap = InBounds.Overlap(PointBoundsTransformed);
+		if (!FirstOverlap.IsValid)
+		{
+			return 0;
+		}
+
+		FMatrix InTransformToPointTransform = PointTransformToInTransform.Inverse();
+		const FBox InBoundsTransformed = InBounds.TransformBy(InTransformToPointTransform);
+
+		const FBox SecondOverlap = InBoundsTransformed.Overlap(PointBounds);
+		if (!SecondOverlap.IsValid)
+		{
+			return 0;
+		}
+
+		return FMath::Min(ComputeOverlapRatio(FirstOverlap, InBounds), ComputeOverlapRatio(SecondOverlap, InBoundsTransformed));
+	}
+
+	/** Helper function for additive blending of quaternions (copied from ControlRig) */
+	FQuat AddQuatWithWeight(const FQuat& Q, const FQuat& V, float Weight)
+	{
+		FQuat BlendQuat = V * Weight;
+
+		if ((Q | BlendQuat) >= 0.0f)
+			return Q + BlendQuat;
+		else
+			return Q - BlendQuat;
+	}
+}
+
+namespace PCGBasePointData
+{
+	class FPCGAttributeAccessorKeysPointIndices : public IPCGAttributeAccessorKeys
+	{
+	public:
+		FPCGAttributeAccessorKeysPointIndices(UPCGBasePointData* InPointData)
+			: IPCGAttributeAccessorKeys(/*bInReadOnly=*/false)
+			, PointData(InPointData)
+		{}
+
+		FPCGAttributeAccessorKeysPointIndices(const UPCGBasePointData* InPointData)
+			: IPCGAttributeAccessorKeys(/*bInReadOnly=*/true)
+			, PointData(const_cast<UPCGBasePointData*>(InPointData))
+		{}
+
+		virtual bool GetKeyIndices(int32 InStart, int32 InCount, TArray<int32>& OutKeyIndices, bool& bOutContiguous) const override
+		{
+			if (PointData->GetNumPoints() == 0)
+			{
+				return false;
+			}
+
+			// Optimization (avoid allocating indices memory)
+			if (InStart + InCount <= PointData->GetNumPoints())
+			{
+				bOutContiguous = true;
+				return true;
+			}
+
+			OutKeyIndices.SetNumUninitialized(InCount);
+			for (int32 Index = 0; Index < InCount; ++Index)
+			{
+				OutKeyIndices[Index] = (InStart + Index) % PointData->GetNumPoints();
+			}
+
+			return true;
+		}
+
+		virtual int32 GetNum() const override
+		{
+			return PointData ? PointData->GetNumPoints() : 0;
+		}
+
+	protected:
+		virtual bool GetMetadataEntryKeys(int32 InStart, TArrayView<PCGMetadataEntryKey*> OutEntryKeys) override
+		{
+			if (!PointData || PointData->GetNumPoints() == 0)
+			{
+				return false;
+			}
+
+			TPCGValueRange<int64> MetadataEntryKeys = PointData->GetMetadataEntryValueRange();
+
+			for (int32 Index = 0; Index < OutEntryKeys.Num(); ++Index)
+			{
+				int32 EntryIndex = (InStart + Index) % MetadataEntryKeys.Num();
+				OutEntryKeys[Index] = &MetadataEntryKeys[EntryIndex];
+			}
+
+			return true;
+		}
+
+		virtual bool GetMetadataEntryKeys(int32 InStart, TArrayView<const PCGMetadataEntryKey*> OutEntryKeys) const override
+		{
+			if (!PointData || PointData->GetNumPoints() == 0)
+			{
+				return false;
+			}
+
+			const TConstPCGValueRange<int64> MetadataEntryKeys = PointData->GetConstMetadataEntryValueRange();
+
+			for (int32 Index = 0; Index < OutEntryKeys.Num(); ++Index)
+			{
+				int32 EntryIndex = (InStart + Index) % MetadataEntryKeys.Num();
+				OutEntryKeys[Index] = &MetadataEntryKeys[EntryIndex];
+			}
+
+			return true;
+		}
+
+	private:
+		UPCGBasePointData* PointData = nullptr;
+	};
+
+	template <typename T, typename R = T>
+	class FPCGNativePointPropertyAccessor : public IPCGAttributeAccessorT<FPCGNativePointPropertyAccessor<T,R>>
+	{
+	public:
+		using Type = T;
+		using Super = IPCGAttributeAccessorT<FPCGNativePointPropertyAccessor<T,R>>;
+
+		FPCGNativePointPropertyAccessor(const UPCGBasePointData* InPointData, EPCGPointNativeProperties InNativeProperty)
+			: Super(/*bReadOnly=*/true)
+		{
+			ValueRange = const_cast<UPCGBasePointData*>(InPointData)->GetValueRange<R>(InNativeProperty);
+		}
+
+		FPCGNativePointPropertyAccessor(UPCGBasePointData* InPointData, EPCGPointNativeProperties InNativeProperty)
+			: Super(/*bReadOnly=*/false)
+		{
+			ValueRange = InPointData->GetValueRange<R>(InNativeProperty, /*bAllocate=*/true);
+		}
+
+		bool GetRangeImpl(TArrayView<T> OutValues, int32 StartIndex, const IPCGAttributeAccessorKeys& Keys) const
+		{
+			TArray<int32> KeyIndices;
+			bool bContiguous = false;
+			if (!Keys.GetKeyIndices(StartIndex, OutValues.Num(), KeyIndices, bContiguous))
+			{
+				return false;
+			}
+						
+			if (bContiguous)
+			{
+				for (int32 OutIndex = 0; OutIndex < OutValues.Num(); ++OutIndex)
+				{
+					OutValues[OutIndex] = (T)ValueRange[StartIndex + OutIndex];
+				}
+			}
+			else
+			{
+				for (int32 OutIndex = 0; OutIndex < OutValues.Num(); ++OutIndex)
+				{
+					OutValues[OutIndex] = (T)ValueRange[KeyIndices[OutIndex]];
+				}
+			}
+
+			return true;
+		}
+
+		bool SetRangeImpl(TArrayView<const T> InValues, int32 StartIndex, IPCGAttributeAccessorKeys& Keys, EPCGAttributeAccessorFlags Flags)
+		{
+			if (Keys.IsReadOnly())
+			{
+				return false;
+			}
+
+			TArray<int32> KeyIndices;
+			bool bContiguous = false;
+			if (!Keys.GetKeyIndices(StartIndex, InValues.Num(), KeyIndices, bContiguous))
+			{
+				return false;
+			}
+
+			if (bContiguous)
+			{
+				for (int32 InIndex = 0; InIndex < InValues.Num(); ++InIndex)
+				{
+					ValueRange[StartIndex + InIndex] = (R)InValues[InIndex];
+				}
+			}
+			else
+			{
+				for (int32 InIndex = 0; InIndex < InValues.Num(); ++InIndex)
+				{
+					ValueRange[KeyIndices[InIndex]] = (R)InValues[InIndex];
+				}
+			}
+
+			return true;
+		}
+
+	private:
+		TPCGValueRange<R> ValueRange;
+	};
+
+	template <typename T, typename... Args>
+	using FPointCustomPropertyGetter = TFunction<bool(int32, T&, Args...)>;
+
+	template <typename T, typename... Args>
+	using FPointCustomPropertySetter = TFunction<bool(int32, const T&, Args...)>;
+
+	template <typename T, typename... Args>
+	class FPCGCustomPointPropertyAccessor : public IPCGAttributeAccessorT<FPCGCustomPointPropertyAccessor<T, Args...>>
+	{
+	public:
+		using Type = T;
+		using Super = IPCGAttributeAccessorT<FPCGCustomPointPropertyAccessor<T, Args...>>;
+
+		FPCGCustomPointPropertyAccessor(const UPCGBasePointData* InPointData, const FPointCustomPropertyGetter<T, Args...>& InGetter, const FPointCustomPropertySetter<T, Args...>& InSetter, Args&& ...InValueRanges)
+			: Super(/*bInReadOnly=*/ false)
+			, PointData(const_cast<UPCGBasePointData*>(InPointData))
+			, Getter(InGetter)
+			, Setter(InSetter)
+			, ValueRanges(std::forward<Args>(InValueRanges)...)
+		{
+		}
+
+		FPCGCustomPointPropertyAccessor(const UPCGBasePointData* InPointData, const FPointCustomPropertyGetter<T, Args...>& InGetter, Args&& ...InValueRanges)
+			: Super(/*bInReadOnly=*/ true)
+			, PointData(const_cast<UPCGBasePointData*>(InPointData))
+			, Getter(InGetter)
+			, ValueRanges(std::forward<Args>(InValueRanges)...)
+		{}
+
+		bool GetRangeImpl(TArrayView<T> OutValues, int32 StartIndex, const IPCGAttributeAccessorKeys& Keys) const
+		{
+			TArray<int32> KeyIndices;
+			bool bContiguous = false;
+			if (!Keys.GetKeyIndices(StartIndex, OutValues.Num(), KeyIndices, bContiguous))
+			{
+				return false;
+			}
+
+			if (bContiguous)
+			{
+				for (int32 OutIndex = 0; OutIndex < OutValues.Num(); ++OutIndex)
+				{
+					ValueRanges.ApplyAfter(Getter, StartIndex + OutIndex, OutValues[OutIndex]);
+				}
+			}
+			else
+			{
+				for (int32 OutIndex = 0; OutIndex < OutValues.Num(); ++OutIndex)
+				{
+					ValueRanges.ApplyAfter(Getter, KeyIndices[OutIndex], OutValues[OutIndex]);
+				}
+			}
+
+			return true;
+		}
+
+		bool SetRangeImpl(TArrayView<const T> InValues, int32 StartIndex, IPCGAttributeAccessorKeys& Keys, EPCGAttributeAccessorFlags Flags)
+		{
+			if (Keys.IsReadOnly())
+			{
+				return false;
+			}
+
+			TArray<int32> KeyIndices;
+			bool bContiguous = false;
+			if (!Keys.GetKeyIndices(StartIndex, InValues.Num(), KeyIndices, bContiguous))
+			{
+				return false;
+			}
+
+			if (bContiguous)
+			{
+				for (int32 InIndex = 0; InIndex < InValues.Num(); ++InIndex)
+				{
+					ValueRanges.ApplyAfter(Setter, StartIndex + InIndex, InValues[InIndex]);
+				}
+			}
+			else
+			{
+				for (int32 InIndex = 0; InIndex < InValues.Num(); ++InIndex)
+				{
+					ValueRanges.ApplyAfter(Setter, KeyIndices[InIndex], InValues[InIndex]);
+				}
+			}
+
+			return true;
+		}
+
+	private:
+		UPCGBasePointData* PointData = nullptr;
+		FPointCustomPropertyGetter<T, Args...> Getter;
+		FPointCustomPropertySetter<T, Args...> Setter;
+		TTuple<Args...> ValueRanges;
+	};
+
+	TUniquePtr<IPCGAttributeAccessor> CreatePropertyAccessor(UPCGBasePointData* PointData, EPCGPointNativeProperties NativeProperty, bool bQuiet)
+	{
+		switch (NativeProperty)
+		{
+		case EPCGPointNativeProperties::Transform:
+			return MakeUnique<FPCGNativePointPropertyAccessor<FTransform>>(PointData, NativeProperty);
+		case EPCGPointNativeProperties::Density:
+		case EPCGPointNativeProperties::Steepness:
+			// The values are floats but we want to remain compatible with existing UPCGPointData Steepnees/Density property accessors which are of type double 
+			return MakeUnique<FPCGNativePointPropertyAccessor<double, float>>(PointData, NativeProperty);
+		case EPCGPointNativeProperties::BoundsMin:
+		case EPCGPointNativeProperties::BoundsMax:
+			return MakeUnique<FPCGNativePointPropertyAccessor<FVector>>(PointData, NativeProperty);
+		case EPCGPointNativeProperties::Color:
+			return MakeUnique<FPCGNativePointPropertyAccessor<FVector4>>(PointData, NativeProperty);
+		case EPCGPointNativeProperties::Seed:
+			return MakeUnique<FPCGNativePointPropertyAccessor<int32>>(PointData, NativeProperty);
+		default:
+			UE_LOG(LogPCG, Error, TEXT("EPCGPointNativeProperty value '%d' does not exist."), NativeProperty);
+			return TUniquePtr<IPCGAttributeAccessor>();
+		}
+	}
+
+	TUniquePtr<IPCGAttributeAccessor> CreateConstPropertyAccessor(const UPCGBasePointData* PointData, EPCGPointNativeProperties NativeProperty, bool bQuiet)
+	{
+		switch (NativeProperty)
+		{
+		case EPCGPointNativeProperties::Transform:
+			return MakeUnique<FPCGNativePointPropertyAccessor<FTransform>>(PointData, NativeProperty);
+		case EPCGPointNativeProperties::Density:
+		case EPCGPointNativeProperties::Steepness:
+			// The values are floats but we want to remain compatible with existing UPCGPointData Steepnees/Density property accessors which are of type double 
+			return MakeUnique<FPCGNativePointPropertyAccessor<double, float>>(PointData, NativeProperty);
+		case EPCGPointNativeProperties::BoundsMin:
+		case EPCGPointNativeProperties::BoundsMax:
+			return MakeUnique<FPCGNativePointPropertyAccessor<FVector>>(PointData, NativeProperty);
+		case EPCGPointNativeProperties::Color:
+			return MakeUnique<FPCGNativePointPropertyAccessor<FVector4>>(PointData, NativeProperty);
+		case EPCGPointNativeProperties::Seed:
+			return MakeUnique<FPCGNativePointPropertyAccessor<int32>>(PointData, NativeProperty);
+		default:
+			UE_LOG(LogPCG, Error, TEXT("EPCGPointNativeProperty value '%d' does not exist."), NativeProperty);
+			return TUniquePtr<IPCGAttributeAccessor>();
+		}
+	}
+
+	TUniquePtr<IPCGAttributeAccessor> CreateCustomPropertyAccessor(UPCGBasePointData* PointData, FName Name)
+	{
+		if (Name == PCGPointCustomPropertyNames::ExtentsName)
+		{
+			return MakeUnique<FPCGCustomPointPropertyAccessor<FVector, TPCGValueRange<FVector>, TPCGValueRange<FVector>>>(
+				PointData,
+				[](int32 Index, FVector& OutValue, const TPCGValueRange<FVector>& BoundsMin, const TPCGValueRange<FVector>& BoundsMax)
+				{
+					OutValue = PCGPointHelpers::GetExtents(BoundsMin[Index], BoundsMax[Index]);
+					return true;
+				},
+				[](int32 Index, const FVector& InValue, TPCGValueRange<FVector> BoundsMin, TPCGValueRange<FVector> BoundsMax)
+				{
+					PCGPointHelpers::SetExtents(InValue, BoundsMin[Index], BoundsMax[Index]);
+					return true;
+				},
+				PointData->GetBoundsMinValueRange(),
+				PointData->GetBoundsMaxValueRange()
+			);
+		}
+		else if (Name == PCGPointCustomPropertyNames::LocalCenterName)
+		{
+			return MakeUnique<FPCGCustomPointPropertyAccessor<FVector, TPCGValueRange<FVector>, TPCGValueRange<FVector>>>(
+				PointData,
+				[](int32 Index, FVector& OutValue, const TPCGValueRange<FVector>& BoundsMin, const TPCGValueRange<FVector>& BoundsMax)
+				{
+					OutValue = PCGPointHelpers::GetLocalCenter(BoundsMin[Index], BoundsMax[Index]);
+					return true;
+				},
+				[](int32 Index, const FVector& InValue, TPCGValueRange<FVector> BoundsMin, TPCGValueRange<FVector> BoundsMax)
+				{
+					PCGPointHelpers::SetLocalCenter(InValue, BoundsMin[Index], BoundsMax[Index]);
+					return true;
+				},
+				PointData->GetBoundsMinValueRange(),
+				PointData->GetBoundsMaxValueRange()
+			);
+		}
+		else if (Name == PCGPointCustomPropertyNames::PositionName)
+		{
+			return MakeUnique<FPCGCustomPointPropertyAccessor<FVector, TPCGValueRange<FTransform>>>(
+				PointData,
+				[](int32 Index, FVector& OutValue, const TPCGValueRange<FTransform>& Transform)
+				{
+					OutValue = Transform[Index].GetLocation();
+					return true;
+				},
+				[](int32 Index, const FVector& InValue, TPCGValueRange<FTransform> Transform)
+				{
+					Transform[Index].SetLocation(InValue);
+					return true;
+				},
+				PointData->GetTransformValueRange()
+			);
+		}
+		else if (Name == PCGPointCustomPropertyNames::RotationName)
+		{
+			return MakeUnique<FPCGCustomPointPropertyAccessor<FQuat, TPCGValueRange<FTransform>>>(
+				PointData,
+				[](int32 Index, FQuat& OutValue, const TPCGValueRange<FTransform>& Transform)
+				{
+					OutValue = Transform[Index].GetRotation();
+					return true;
+				},
+				[](int32 Index, const FQuat& InValue, TPCGValueRange<FTransform> Transform)
+				{
+					Transform[Index].SetRotation(InValue);
+					return true;
+				},
+				PointData->GetTransformValueRange()
+			);
+		}
+		else if (Name == PCGPointCustomPropertyNames::ScaleName)
+		{
+			return MakeUnique<FPCGCustomPointPropertyAccessor<FVector, TPCGValueRange<FTransform>>>(
+				PointData,
+				[](int32 Index, FVector& OutValue, const TPCGValueRange<FTransform>& Transform)
+				{
+					OutValue = Transform[Index].GetScale3D();
+					return true;
+				},
+				[](int32 Index, const FVector& InValue, TPCGValueRange<FTransform> Transform)
+				{
+					Transform[Index].SetScale3D(InValue);
+					return true;
+				},
+				PointData->GetTransformValueRange()
+			);
+		}
+
+		return TUniquePtr<IPCGAttributeAccessor>();
+	}
+
+	TUniquePtr<IPCGAttributeAccessor> CreateConstCustomPropertyAccessor(const UPCGBasePointData* PointData, FName Name)
+	{
+		if (Name == PCGPointCustomPropertyNames::ExtentsName)
+		{
+			return MakeUnique<FPCGCustomPointPropertyAccessor<FVector, TConstPCGValueRange<FVector>, TConstPCGValueRange<FVector>>>(
+				PointData,
+				[](int32 Index, FVector& OutValue, const TConstPCGValueRange<FVector>& BoundsMin, const TConstPCGValueRange<FVector>& BoundsMax)
+				{
+					OutValue = PCGPointHelpers::GetExtents(BoundsMin[Index], BoundsMax[Index]);
+					return true;
+				},
+				PointData->GetConstBoundsMinValueRange(),
+				PointData->GetConstBoundsMaxValueRange()
+			);
+		}
+		else if (Name == PCGPointCustomPropertyNames::LocalCenterName)
+		{
+			return MakeUnique<FPCGCustomPointPropertyAccessor<FVector, TConstPCGValueRange<FVector>, TConstPCGValueRange<FVector>>>(
+				PointData,
+				[](int32 Index, FVector& OutValue, const TConstPCGValueRange<FVector>& BoundsMin, const TConstPCGValueRange<FVector>& BoundsMax)
+				{
+					OutValue = PCGPointHelpers::GetLocalCenter(BoundsMin[Index], BoundsMax[Index]);
+					return true;
+				},
+				PointData->GetConstBoundsMinValueRange(),
+				PointData->GetConstBoundsMaxValueRange()
+			);
+		}
+		else if (Name == PCGPointCustomPropertyNames::PositionName)
+		{
+			return MakeUnique<FPCGCustomPointPropertyAccessor<FVector, TConstPCGValueRange<FTransform>>>(
+				PointData,
+				[](int32 Index, FVector& OutValue, const TConstPCGValueRange<FTransform>& Transform)
+				{
+					OutValue = Transform[Index].GetLocation();
+					return true;
+				},
+				PointData->GetConstTransformValueRange()
+			);
+		}
+		else if (Name == PCGPointCustomPropertyNames::RotationName)
+		{
+			return MakeUnique<FPCGCustomPointPropertyAccessor<FQuat, TConstPCGValueRange<FTransform>>>(
+				PointData,
+				[](int32 Index, FQuat& OutValue, const TConstPCGValueRange<FTransform>& Transform)
+				{
+					OutValue = Transform[Index].GetRotation();
+					return true;
+				},
+				PointData->GetConstTransformValueRange()
+			);
+		}
+		else if (Name == PCGPointCustomPropertyNames::ScaleName)
+		{
+			return MakeUnique<FPCGCustomPointPropertyAccessor<FVector, TConstPCGValueRange<FTransform>>>(
+				PointData,
+				[](int32 Index, FVector& OutValue, const TConstPCGValueRange<FTransform>& Transform)
+				{
+					OutValue = Transform[Index].GetScale3D();
+					return true;
+				},
+				PointData->GetConstTransformValueRange()
+			);
+		}
+		else if (Name == PCGPointCustomPropertyNames::LocalSizeName)
+		{
+			return MakeUnique<FPCGCustomPointPropertyAccessor<FVector, TConstPCGValueRange<FVector>, TConstPCGValueRange<FVector>>>(
+				PointData,
+				[](int32 Index, FVector& OutValue, const TConstPCGValueRange<FVector>& BoundsMin, const TConstPCGValueRange<FVector>& BoundsMax)
+				{
+					OutValue = PCGPointHelpers::GetLocalSize(BoundsMin[Index], BoundsMax[Index]);
+					return true;
+				},
+				PointData->GetConstBoundsMinValueRange(),
+				PointData->GetConstBoundsMaxValueRange()
+			);
+		}
+		else if (Name == PCGPointCustomPropertyNames::ScaledLocalSizeName)
+		{
+			return MakeUnique<FPCGCustomPointPropertyAccessor<FVector, TConstPCGValueRange<FTransform>, TConstPCGValueRange<FVector>, TConstPCGValueRange<FVector>>>(
+				PointData,
+				[](int32 Index, FVector& OutValue, const TConstPCGValueRange<FTransform>& Transform, const TConstPCGValueRange<FVector>& BoundsMin, const TConstPCGValueRange<FVector>& BoundsMax)
+				{
+					OutValue = PCGPointHelpers::GetScaledLocalSize(Transform[Index], BoundsMin[Index], BoundsMax[Index]);
+					return true;
+				},
+				PointData->GetConstTransformValueRange(),
+				PointData->GetConstBoundsMinValueRange(),
+				PointData->GetConstBoundsMaxValueRange()
+			);
+		}
+
+		return TUniquePtr<IPCGAttributeAccessor>();
+	}
+
+	TUniquePtr<IPCGAttributeAccessor> CreateStaticAccessor(const UPCGBasePointData* PointData, const FPCGAttributePropertySelector& InSelector, bool bQuiet, bool bConst)
+	{
+		TUniquePtr<IPCGAttributeAccessor> Accessor;
+
+		if (InSelector.GetSelection() == EPCGAttributePropertySelection::PointProperty)
+		{
+			const FName PropertyName = InSelector.GetName();
+
+			if (const int64 EnumValue = StaticEnum<EPCGPointNativeProperties>()->GetValueByName(PropertyName); EnumValue != INDEX_NONE)
+			{
+				const EPCGPointNativeProperties NativeProperty = (EPCGPointNativeProperties)EnumValue;
+				if (bConst)
+				{
+					Accessor = CreateConstPropertyAccessor(PointData, NativeProperty, bQuiet);
+				}
+				else
+				{
+					Accessor = CreatePropertyAccessor(const_cast<UPCGBasePointData*>(PointData), NativeProperty, bQuiet);
+				}
+			}
+			else if (PCGPointCustomPropertyNames::IsCustomPropertyName(PropertyName))
+			{
+				if (bConst)
+				{
+					Accessor = CreateConstCustomPropertyAccessor(PointData, PropertyName);
+				}
+				else
+				{
+					Accessor = CreateCustomPropertyAccessor(const_cast<UPCGBasePointData*>(PointData), PropertyName);
+				}
+			}
+		}
+		else if (InSelector.GetSelection() == EPCGAttributePropertySelection::Attribute && !bConst)
+		{
+			// Let parent factory create the accessor but allocate the metadata entry memory
+			const_cast<UPCGBasePointData*>(PointData)->AllocateProperties(EPCGPointNativeProperties::MetadataEntry);
+		}
+
+		return Accessor;
+	}
+}
+
 UPCGBasePointData::UPCGBasePointData(const FObjectInitializer& ObjectInitializer)
 	: Super(ObjectInitializer)
 {
 	check(Metadata);
 	Metadata->SetupLayer(PCGMetadataLayerID::Elements, /*bIsDefault=*/true);
+}
+
+FPCGAttributeAccessorMethods UPCGBasePointData::GetPointAccessorMethods()
+{
+	auto CreateAccessorFunc = [](UPCGData* InData, const FPCGAttributePropertySelector& InSelector, bool bQuiet)
+	{
+		return PCGBasePointData::CreateStaticAccessor(CastChecked<UPCGBasePointData>(InData), InSelector, bQuiet, false);
+	};
+
+	auto CreateConstAccessorFunc = [](const UPCGData* InData, const FPCGAttributePropertySelector& InSelector, bool bQuiet) -> TUniquePtr<const IPCGAttributeAccessor>
+	{
+		return PCGBasePointData::CreateStaticAccessor(CastChecked<UPCGBasePointData>(InData), InSelector, bQuiet, true);
+	};
+
+	auto CreateAccessorKeysFunc = [](UPCGData* InData, const FPCGAttributePropertySelector& InSelector, bool bQuiet)
+	{
+		UPCGBasePointData* PointData = CastChecked<UPCGBasePointData>(InData);
+		return MakeUnique<PCGBasePointData::FPCGAttributeAccessorKeysPointIndices>(PointData);
+	};
+
+	auto CreateConstAccessorKeysFunc = [](const UPCGData* InData, const FPCGAttributePropertySelector& InSelector, bool bQuiet) -> TUniquePtr<const IPCGAttributeAccessorKeys>
+	{
+		const UPCGBasePointData* PointData = CastChecked<UPCGBasePointData>(InData);
+		return MakeUnique<const PCGBasePointData::FPCGAttributeAccessorKeysPointIndices>(PointData);
+	};
+
+	FPCGAttributeAccessorMethods Methods
+	{
+		.CreateAccessorFunc = CreateAccessorFunc,
+		.CreateConstAccessorFunc = CreateConstAccessorFunc,
+		.CreateAccessorKeysFunc = CreateAccessorKeysFunc,
+		.CreateConstAccessorKeysFunc = CreateConstAccessorKeysFunc
+	};
+
+	return Methods;
+}
+
+FBox UPCGBasePointData::GetBounds() const
+{
+	RecomputeBoundsIfNeeded();
+
+	return Bounds;
+}
+
+void UPCGBasePointData::CopyPointsFrom(const UPCGBasePointData* InData, const TArray<int32>& InDataIndices)
+{
+	return CopyPoints(InData, this, InDataIndices, false);
+}
+
+void UPCGBasePointData::CopyPoints(const UPCGBasePointData* From, UPCGBasePointData* To, const TArray<int32>& InDataIndices, bool bCopyAll) const
+{
+	TRACE_CPUPROFILER_EVENT_SCOPE(UPCGBasePointData::CopyPoints);
+	check(From);
+
+	const int32 NumPoints = bCopyAll ? From->GetNumPoints() : InDataIndices.Num();
+	To->SetNumPoints(NumPoints);
+
+	const TConstPCGValueRange<FTransform> FromTransformRange = From->GetConstTransformValueRange();
+	const TConstPCGValueRange<float> FromSteepnessRange = From->GetConstSteepnessValueRange();
+	const TConstPCGValueRange<float> FromDensityRange = From->GetConstDensityValueRange();
+	const TConstPCGValueRange<FVector> FromBoundsMinRange = From->GetConstBoundsMinValueRange();
+	const TConstPCGValueRange<FVector> FromBoundsMaxRange = From->GetConstBoundsMaxValueRange();
+	const TConstPCGValueRange<FVector4> FromColorRange = From->GetConstColorValueRange();
+	const TConstPCGValueRange<int64> FromMetadataEntryRange = From->GetConstMetadataEntryValueRange();
+	const TConstPCGValueRange<int32> FromSeedRange = From->GetConstSeedValueRange();
+
+	TPCGValueRange<FTransform> ToTransformRange = To->GetTransformValueRange();
+	TPCGValueRange<float> ToSteepnessRange = To->GetSteepnessValueRange();
+	TPCGValueRange<float> ToDensityRange = To->GetDensityValueRange();
+	TPCGValueRange<FVector> ToBoundsMinRange = To->GetBoundsMinValueRange();
+	TPCGValueRange<FVector> ToBoundsMaxRange = To->GetBoundsMaxValueRange();
+	TPCGValueRange<FVector4> ToColorRange = To->GetColorValueRange();
+	TPCGValueRange<int64> ToMetadataEntryRange = To->GetMetadataEntryValueRange();
+	TPCGValueRange<int32> ToSeedRange = To->GetSeedValueRange();
+
+	for (int32 Index = 0; Index < NumPoints; ++Index)
+	{
+		const int32 PointIndex = bCopyAll ? Index : InDataIndices[Index];
+
+		ToTransformRange[PointIndex] = FromTransformRange[PointIndex];
+		ToSteepnessRange[PointIndex] = FromSteepnessRange[PointIndex];
+		ToDensityRange[PointIndex] = FromDensityRange[PointIndex];
+		ToBoundsMinRange[PointIndex] = FromBoundsMinRange[PointIndex];
+		ToBoundsMaxRange[PointIndex] = FromBoundsMaxRange[PointIndex];
+		ToColorRange[PointIndex] = FromColorRange[PointIndex];
+		ToMetadataEntryRange[PointIndex] = FromMetadataEntryRange[PointIndex];
+		ToSeedRange[PointIndex] = FromSeedRange[PointIndex];
+	}
+
+	To->DirtyCache();
+}
+
+void UPCGBasePointData::RecomputeBounds() const
+{
+	FScopeLock Lock(&CachedDataLock);
+	if (!bBoundsAreDirty)
+	{
+		return;
+	}
+
+	const TConstPCGValueRange<FTransform> TransformRange = GetConstTransformValueRange();
+	const TConstPCGValueRange<float> SteepnessRange = GetConstSteepnessValueRange();
+	const TConstPCGValueRange<FVector> BoundsMinRange = GetConstBoundsMinValueRange();
+	const TConstPCGValueRange<FVector> BoundsMaxRange = GetConstBoundsMaxValueRange();
+
+	FBox NewBounds(EForceInit::ForceInit);
+	for (int32 PointIndex = 0; PointIndex < GetNumPoints(); ++PointIndex)
+	{
+		FBoxSphereBounds PointBounds = PCGPointHelpers::GetDensityBounds(TransformRange[PointIndex], SteepnessRange[PointIndex], BoundsMinRange[PointIndex], BoundsMaxRange[PointIndex]);
+		NewBounds += FBox::BuildAABB(PointBounds.Origin, PointBounds.BoxExtent);
+	}
+
+	Bounds = NewBounds;
+	bBoundsAreDirty = false;
+}
+
+const PCGPointOctree::FPointOctree& UPCGBasePointData::GetPointOctree() const
+{
+	RebuildOctreeIfNeeded();
+	
+	return PCGPointOctree;
+}
+
+void UPCGBasePointData::RebuildOctree() const
+{
+	FScopeLock Lock(&CachedDataLock);
+	if (!bOctreeIsDirty)
+	{
+		return;
+	}
+
+	TRACE_CPUPROFILER_EVENT_SCOPE(UPCGPointData::RebuildOctree)
+	
+	const FBox PointBounds = GetBounds();
+	PCGPointOctree::FPointOctree NewOctree(PointBounds.GetCenter(), PointBounds.GetExtent().Length());
+
+	const TConstPCGValueRange<FTransform> TransformRange = GetConstTransformValueRange();
+	const TConstPCGValueRange<float> SteepnessRange = GetConstSteepnessValueRange();
+	const TConstPCGValueRange<FVector> BoundsMinRange = GetConstBoundsMinValueRange();
+	const TConstPCGValueRange<FVector> BoundsMaxRange = GetConstBoundsMaxValueRange();
+
+	for (int32 PointIndex = 0; PointIndex < GetNumPoints(); ++PointIndex)
+	{
+		NewOctree.AddElement(PCGPointOctree::FPointRef(PointIndex, PCGPointHelpers::GetDensityBounds(TransformRange[PointIndex], SteepnessRange[PointIndex], BoundsMinRange[PointIndex], BoundsMaxRange[PointIndex])));
+	}
+
+	PCGPointOctree = MoveTemp(NewOctree);
+	bOctreeIsDirty = false;
+}
+
+void UPCGBasePointData::SetTransform(const FTransform& InTransform)
+{
+	FreeProperties(EPCGPointNativeProperties::Transform);
+	TPCGValueRange<FTransform> ValueRange = GetTransformValueRange(/*bAllocate=*/false);
+	ValueRange.Set(InTransform);
+}
+
+void UPCGBasePointData::SetDensity(float InDensity)
+{
+	FreeProperties(EPCGPointNativeProperties::Density);
+	TPCGValueRange<float> ValueRange = GetDensityValueRange(/*bAllocate=*/false);
+	ValueRange.Set(InDensity);
+}
+
+void UPCGBasePointData::SetBoundsMin(const FVector& InBoundsMin)
+{
+	FreeProperties(EPCGPointNativeProperties::BoundsMin);
+	TPCGValueRange<FVector> ValueRange = GetBoundsMinValueRange(/*bAllocate=*/false);
+	ValueRange.Set(InBoundsMin);
+}
+
+void UPCGBasePointData::SetBoundsMax(const FVector& InBoundsMax)
+{
+	FreeProperties(EPCGPointNativeProperties::BoundsMax);
+	TPCGValueRange<FVector> ValueRange = GetBoundsMaxValueRange(/*bAllocate=*/false);
+	ValueRange.Set(InBoundsMax);
+}
+
+void UPCGBasePointData::SetColor(const FVector4& InColor)
+{
+	FreeProperties(EPCGPointNativeProperties::Color);
+	TPCGValueRange<FVector4> ValueRange = GetColorValueRange(/*bAllocate=*/false);
+	ValueRange.Set(InColor);
+}
+
+void UPCGBasePointData::SetSteepness(float InSteepness)
+{
+	FreeProperties(EPCGPointNativeProperties::Steepness);
+	TPCGValueRange<float> ValueRange = GetSteepnessValueRange(/*bAllocate=*/false);
+	ValueRange.Set(InSteepness);
+}
+
+void UPCGBasePointData::SetSeed(int32 InSeed)
+{
+	TPCGValueRange<int32> ValueRange = GetSeedValueRange(/*bAllocate=*/false);
+	ValueRange.Set(InSeed);
+}
+
+void UPCGBasePointData::SetMetadataEntry(int64 InMetadataEntry)
+{
+	FreeProperties(EPCGPointNativeProperties::MetadataEntry);
+	TPCGValueRange<int64> ValueRange = GetMetadataEntryValueRange(/*bAllocate=*/false);
+	ValueRange.Set(InMetadataEntry);
+}
+
+void UPCGBasePointData::SetExtents(const FVector& InExtents)
+{
+	// Allocate if needed
+	{
+		const TConstPCGValueRange<FVector> BoundsMinRange = GetConstBoundsMinValueRange();
+		const TConstPCGValueRange<FVector> BoundsMaxRange = GetConstBoundsMaxValueRange();
+
+		check(BoundsMinRange.Num() == BoundsMaxRange.Num());
+
+		if (BoundsMinRange.ViewNum() != BoundsMaxRange.ViewNum())
+		{
+			AllocateProperties(EPCGPointNativeProperties::BoundsMin | EPCGPointNativeProperties::BoundsMax);
+		}
+	}
+
+	// If Allocation needed it was already done
+	TPCGValueRange<FVector> BoundsMinRange = GetBoundsMinValueRange(/*bAllocate=*/false);
+	TPCGValueRange<FVector> BoundsMaxRange = GetBoundsMaxValueRange(/*bAllocate=*/false);
+
+	for (int32 PointIndex = 0; PointIndex < BoundsMinRange.ViewNum(); ++PointIndex)
+	{
+		PCGPointHelpers::SetExtents(InExtents, BoundsMinRange[PointIndex], BoundsMaxRange[PointIndex]);
+	}
+}
+
+void UPCGBasePointData::SetLocalCenter(const FVector& InLocalCenter)
+{
+	// Allocate if needed
+	{
+		const TConstPCGValueRange<FVector> BoundsMinRange = GetConstBoundsMinValueRange();
+		const TConstPCGValueRange<FVector> BoundsMaxRange = GetConstBoundsMaxValueRange();
+
+		check(BoundsMinRange.Num() == BoundsMaxRange.Num());
+
+		if (BoundsMinRange.ViewNum() != BoundsMaxRange.ViewNum())
+		{
+			AllocateProperties(EPCGPointNativeProperties::BoundsMin | EPCGPointNativeProperties::BoundsMax);
+		}
+	}
+
+	// If Allocation needed it was already done
+	TPCGValueRange<FVector> BoundsMinRange = GetBoundsMinValueRange(/*bAllocate=*/false);
+	TPCGValueRange<FVector> BoundsMaxRange = GetBoundsMaxValueRange(/*bAllocate=*/false);
+
+	for (int32 PointIndex = 0; PointIndex < BoundsMinRange.ViewNum(); ++PointIndex)
+	{
+		PCGPointHelpers::SetLocalCenter(InLocalCenter, BoundsMinRange[PointIndex], BoundsMaxRange[PointIndex]);
+	}
+}
+
+const FTransform& UPCGBasePointData::GetTransform(int32 InPointIndex) const
+{
+	const TConstPCGValueRange<FTransform> ValueRange = GetConstTransformValueRange();
+	return ValueRange[InPointIndex];
+}
+
+float UPCGBasePointData::GetDensity(int32 InPointIndex) const
+{
+	const TConstPCGValueRange<float> ValueRange = GetConstDensityValueRange();
+	return ValueRange[InPointIndex];
+}
+
+const FVector& UPCGBasePointData::GetBoundsMin(int32 InPointIndex) const
+{
+	const TConstPCGValueRange<FVector> ValueRange = GetConstBoundsMinValueRange();
+	return ValueRange[InPointIndex];
+}
+
+const FVector& UPCGBasePointData::GetBoundsMax(int32 InPointIndex) const
+{
+	const TConstPCGValueRange<FVector> ValueRange = GetConstBoundsMaxValueRange();
+	return ValueRange[InPointIndex];
+}
+
+const FVector4& UPCGBasePointData::GetColor(int32 InPointIndex) const
+{
+	const TConstPCGValueRange<FVector4> ValueRange = GetConstColorValueRange();
+	return ValueRange[InPointIndex];
+}
+
+float UPCGBasePointData::GetSteepness(int32 InPointIndex) const
+{
+	const TConstPCGValueRange<float> ValueRange = GetConstSteepnessValueRange();
+	return ValueRange[InPointIndex];
+}
+
+int32 UPCGBasePointData::GetSeed(int32 InPointIndex) const
+{
+	const TConstPCGValueRange<int32> ValueRange = GetConstSeedValueRange();
+	return ValueRange[InPointIndex];
+}
+
+int64 UPCGBasePointData::GetMetadataEntry(int32 InPointIndex) const
+{
+	const TConstPCGValueRange<int64> ValueRange = GetConstMetadataEntryValueRange();
+	return ValueRange[InPointIndex];
+}
+
+FBoxSphereBounds UPCGBasePointData::GetDensityBounds(int32 InPointIndex) const
+{
+	const TConstPCGValueRange<FTransform> TransformRange = GetConstTransformValueRange();
+	const TConstPCGValueRange<float> SteepnessRange = GetConstSteepnessValueRange();
+	const TConstPCGValueRange<FVector> BoundsMinRange = GetConstBoundsMinValueRange();
+	const TConstPCGValueRange<FVector> BoundsMaxRange = GetConstBoundsMaxValueRange();
+
+	return PCGPointHelpers::GetDensityBounds(TransformRange[InPointIndex], SteepnessRange[InPointIndex], BoundsMinRange[InPointIndex], BoundsMaxRange[InPointIndex]);
+}
+
+FBox UPCGBasePointData::GetLocalDensityBounds(int32 InPointIndex) const
+{
+	const TConstPCGValueRange<float> SteepnessRange = GetConstSteepnessValueRange();
+	const TConstPCGValueRange<FVector> BoundsMinRange = GetConstBoundsMinValueRange();
+	const TConstPCGValueRange<FVector> BoundsMaxRange = GetConstBoundsMaxValueRange();
+
+	return PCGPointHelpers::GetLocalDensityBounds(SteepnessRange[InPointIndex], BoundsMinRange[InPointIndex], BoundsMaxRange[InPointIndex]);
+}
+
+FBox UPCGBasePointData::GetLocalBounds(int32 InPointIndex) const
+{
+	const TConstPCGValueRange<FVector> BoundsMinRange = GetConstBoundsMinValueRange();
+	const TConstPCGValueRange<FVector> BoundsMaxRange = GetConstBoundsMaxValueRange();
+
+	return PCGPointHelpers::GetLocalBounds(BoundsMinRange[InPointIndex], BoundsMaxRange[InPointIndex]);
+}
+
+FVector UPCGBasePointData::GetLocalCenter(int32 InPointIndex) const
+{
+	const TConstPCGValueRange<FVector> BoundsMinRange = GetConstBoundsMinValueRange();
+	const TConstPCGValueRange<FVector> BoundsMaxRange = GetConstBoundsMaxValueRange();
+
+	return PCGPointHelpers::GetLocalCenter(BoundsMinRange[InPointIndex], BoundsMaxRange[InPointIndex]);
+}
+
+FVector UPCGBasePointData::GetExtents(int32 InPointIndex) const
+{
+	const TConstPCGValueRange<FVector> BoundsMinRange = GetConstBoundsMinValueRange();
+	const TConstPCGValueRange<FVector> BoundsMaxRange = GetConstBoundsMaxValueRange();
+
+	return PCGPointHelpers::GetExtents(BoundsMinRange[InPointIndex], BoundsMaxRange[InPointIndex]);
+}
+
+FVector UPCGBasePointData::GetScaledExtents(int32 InPointIndex) const
+{
+	const TConstPCGValueRange<FTransform> TransformRange = GetConstTransformValueRange();
+	const TConstPCGValueRange<FVector> BoundsMinRange = GetConstBoundsMinValueRange();
+	const TConstPCGValueRange<FVector> BoundsMaxRange = GetConstBoundsMaxValueRange();
+
+	return PCGPointHelpers::GetScaledExtents(TransformRange[InPointIndex], BoundsMinRange[InPointIndex], BoundsMaxRange[InPointIndex]);
+}
+
+FVector UPCGBasePointData::GetLocalSize(int32 InPointIndex) const
+{
+	const TConstPCGValueRange<FVector> BoundsMinRange = GetConstBoundsMinValueRange();
+	const TConstPCGValueRange<FVector> BoundsMaxRange = GetConstBoundsMaxValueRange();
+
+	return PCGPointHelpers::GetLocalSize(BoundsMinRange[InPointIndex], BoundsMaxRange[InPointIndex]);
+}
+
+FVector UPCGBasePointData::GetScaledLocalSize(int32 InPointIndex) const
+{
+	const TConstPCGValueRange<FTransform> TransformRange = GetConstTransformValueRange();
+	const TConstPCGValueRange<FVector> BoundsMinRange = GetConstBoundsMinValueRange();
+	const TConstPCGValueRange<FVector> BoundsMaxRange = GetConstBoundsMaxValueRange();
+
+	return PCGPointHelpers::GetScaledLocalSize(TransformRange[InPointIndex], BoundsMinRange[InPointIndex], BoundsMaxRange[InPointIndex]);
+}
+
+void UPCGBasePointData::GetResourceSizeEx(FResourceSizeEx& CumulativeResourceSize)
+{
+	Super::GetResourceSizeEx(CumulativeResourceSize);
+
+	CumulativeResourceSize.AddDedicatedSystemMemoryBytes(PCGPointOctree.GetSizeBytes() + sizeof(Bounds));
+}
+
+void UPCGBasePointData::InitializeFromActor(AActor* InActor, bool* bOutOptionalSanitizedTagAttributeName)
+{
+	check(InActor);
+	check(Metadata && Metadata->GetAttributeCount() == 0);
+
+	AddSinglePointFromActor(InActor, bOutOptionalSanitizedTagAttributeName);
+}
+
+void UPCGBasePointData::AddSinglePointFromActor(AActor* InActor, bool* bOutOptionalSanitizedTagAttributeName)
+{
+	check(InActor);
+
+	const int32 PointIndex = GetNumPoints();
+	SetNumPoints(PointIndex + 1);
+
+	EPCGPointNativeProperties PropertiesToAllocate = EPCGPointNativeProperties::None;
+	
+	// Values to assign
+	const float PointSteepness = 1.0f;
+	const FTransform PointTransform = InActor->GetActorTransform();
+	const FVector& Position = PointTransform.GetLocation();
+	const int32 PointSeed = PCGHelpers::ComputeSeed((int)Position.X, (int)Position.Y, (int)Position.Z);
+	const FBox LocalBounds = PCGHelpers::GetActorLocalBounds(InActor);
+	const FVector PointBoundsMin = LocalBounds.Min;
+	const FVector PointBoundsMax = LocalBounds.Max;
+	const int64 PointMetadataEntry = Metadata->AddEntry();
+
+	// SteepnessRange - Initialize and Allocate if needed 
+	const TConstPCGValueRange<float> ConstSteepnessRange = GetConstSteepnessValueRange();
+	TOptional<const float> SteepnessSingleValue = ConstSteepnessRange.GetSingleValue();
+	const bool bAllocateSteepness = SteepnessSingleValue.IsSet() && SteepnessSingleValue.GetValue() != PointSteepness;
+	TPCGValueRange<float> SteepnessRange = GetSteepnessValueRange(bAllocateSteepness);
+	
+	// TransformRange - Initialize and Allocate if needed 
+	const TConstPCGValueRange<FTransform> ConstTransformRange = GetConstTransformValueRange();
+	TOptional<const FTransform> TransformSingleValue = ConstTransformRange.GetSingleValue();
+	const bool bAllocateTransform = TransformSingleValue.IsSet() && !TransformSingleValue.GetValue().Equals(PointTransform);
+	TPCGValueRange<FTransform> TransformRange = GetTransformValueRange(bAllocateTransform);
+
+	// SeedRange - Initialize and Allocate if needed 
+	const TConstPCGValueRange<int32> ConstSeedRange = GetConstSeedValueRange();
+	TOptional<const int32> SeedSingleValue = ConstSeedRange.GetSingleValue();
+	const bool bAllocateSeed = SeedSingleValue.IsSet() && SeedSingleValue.GetValue() != PointSeed;
+	TPCGValueRange<int32> SeedRange = GetSeedValueRange(bAllocateSeed);
+
+	// BoundsMinRange - Initialize and Allocate if needed 
+	const TConstPCGValueRange<FVector> ConstBoundsMinRange = GetConstBoundsMinValueRange();
+	TOptional<const FVector> BoundsMinSingleValue = ConstBoundsMinRange.GetSingleValue();
+	const bool bAllocateBoundsMin = BoundsMinSingleValue.IsSet() && BoundsMinSingleValue.GetValue() != PointBoundsMin;
+	TPCGValueRange<FVector> BoundsMinRange = GetBoundsMinValueRange(bAllocateBoundsMin);
+		
+	// BoundsMaxRange - Initialize and Allocate if needed 
+	const TConstPCGValueRange<FVector> ConstBoundsMaxRange = GetConstBoundsMaxValueRange();
+	TOptional<const FVector> BoundsMaxSingleValue = ConstBoundsMaxRange.GetSingleValue();
+	const bool bAllocateBoundsMax = BoundsMaxSingleValue.IsSet() && BoundsMaxSingleValue.GetValue() != PointBoundsMax;
+	TPCGValueRange<FVector> BoundsMaxRange = GetBoundsMaxValueRange(bAllocateBoundsMax);
+
+	// MetadataEntryRange - Initialize and Allocate if needed 
+	const TConstPCGValueRange<int64> ConstMetadataEntryRange = GetConstMetadataEntryValueRange();
+	TOptional<const int64> MetadataEntrySingleValue = ConstMetadataEntryRange.GetSingleValue();
+	const bool bAllocateMetadataEntry = MetadataEntrySingleValue.IsSet() && MetadataEntrySingleValue.GetValue() != PointMetadataEntry;
+	TPCGValueRange<int64> MetadataEntryRange = GetMetadataEntryValueRange(bAllocateMetadataEntry);
+
+	// Assign values
+	SteepnessRange[PointIndex] = PointSteepness;
+	TransformRange[PointIndex] = PointTransform;
+	SeedRange[PointIndex] = PointSeed;
+	BoundsMinRange[PointIndex] = PointBoundsMin;
+	BoundsMaxRange[PointIndex] = PointBoundsMax;
+	MetadataEntryRange[PointIndex] = PointMetadataEntry;
+
+	FPCGMetadataAttribute<FSoftObjectPath>* ActorReferenceAttribute = Metadata->FindOrCreateAttribute(PCGPointDataConstants::ActorReferenceAttribute, FSoftObjectPath(), /*bAllowsInterpolation=*/false, /*bOverrideParent=*/false, /*bOverwriteIfTypeMismatch=*/false);
+	if (ActorReferenceAttribute)
+	{
+		ActorReferenceAttribute->SetValue(PointMetadataEntry, FSoftObjectPath(InActor));
+	}
+
+	bool bSanitizedAttributeNames = false;
+
+	// Parse tags as well
+	for (FName Tag : InActor->Tags)
+	{
+		PCG::Private::FParseTagResult TagData(Tag);
+		if (PCG::Private::SetAttributeFromTag(TagData, Metadata, PointMetadataEntry, /*bCanCreateAttribute=*/true))
+		{
+			bSanitizedAttributeNames |= TagData.HasBeenSanitized();
+		}
+	}
+
+	if (bOutOptionalSanitizedTagAttributeName)
+	{
+		*bOutOptionalSanitizedTagAttributeName = bSanitizedAttributeNames;
+	}
+}
+
+bool UPCGBasePointData::SamplePoint(const FTransform& InTransform, const FBox& InBounds, FPCGPoint& OutPoint, UPCGMetadata* OutMetadata) const
+{
+	// Run a projection but don't change the point transform. There is a large overlap in code/functionality so this shares one code path.
+	FPCGProjectionParams Params{};
+	Params.bProjectPositions = Params.bProjectRotations = Params.bProjectScales = false;
+	Params.ColorBlendMode = EPCGProjectionColorBlendMode::SourceValue;
+
+	// The ProjectPoint implementation in this class returns true if the query point is overlapping the point data, which is what SamplePoint should return, so forward the return value.
+	return ProjectPoint(InTransform, InBounds, Params, OutPoint, OutMetadata);
+}
+
+bool UPCGBasePointData::ProjectPoint(const FTransform& InTransform, const FBox& InBounds, const FPCGProjectionParams& InParams, FPCGPoint& OutPoint, UPCGMetadata* OutMetadata) const
+{
+	return ProjectPoint(InTransform, InBounds, InParams, OutPoint, OutMetadata, true);
+}
+
+bool UPCGBasePointData::ProjectPoint(const FTransform& InTransform, const FBox& InBounds, const FPCGProjectionParams& InParams, FPCGPoint& OutPoint, UPCGMetadata* OutMetadata, bool bUseBounds) const
+{
+	//TRACE_CPUPROFILER_EVENT_SCOPE(UPCGPointData::ProjectPoint);
+	RebuildOctreeIfNeeded();
+
+	TArray<TPair<int32, FVector::FReal>, TInlineAllocator<4>> Contributions;
+	const bool bSampleInVolume = (InBounds.GetExtent() != FVector::ZeroVector);
+
+	const TConstPCGValueRange<FTransform> TransformRange = GetConstTransformValueRange();
+	const TConstPCGValueRange<float> SteepnessRange = GetConstSteepnessValueRange();
+	const TConstPCGValueRange<float> DensityRange = GetConstDensityValueRange();
+	const TConstPCGValueRange<FVector> BoundsMinRange = GetConstBoundsMinValueRange();
+	const TConstPCGValueRange<FVector> BoundsMaxRange = GetConstBoundsMaxValueRange();
+	const TConstPCGValueRange<FVector4> ColorRange = GetConstColorValueRange();
+	const TConstPCGValueRange<int64> MetadataEntryRange = GetConstMetadataEntryValueRange();
+
+	if (!bSampleInVolume)
+	{
+		const FVector InPosition = InTransform.GetLocation();
+		PCGPointOctree.FindElementsWithBoundsTest(FBoxCenterAndExtent(InPosition, FVector::Zero()), [this, &InPosition, &Contributions, &TransformRange, &BoundsMinRange, &BoundsMaxRange, &SteepnessRange](const PCGPointOctree::FPointRef& InPointRef)
+		{
+			Contributions.Emplace(InPointRef.Index, PCGPointHelpers::InverseEuclidianDistance(TransformRange[InPointRef.Index], BoundsMinRange[InPointRef.Index], BoundsMaxRange[InPointRef.Index], SteepnessRange[InPointRef.Index], InPosition));
+		});
+	}
+	else
+	{
+		FBox TransformedBounds = InBounds.TransformBy(InTransform);
+		FMatrix InTransformInverseMatrix = InTransform.ToMatrixWithScale().Inverse();
+
+		PCGPointOctree.FindElementsWithBoundsTest(FBoxCenterAndExtent(TransformedBounds.GetCenter(), TransformedBounds.GetExtent()), 
+			[this, bUseBounds, &InBounds, &InTransformInverseMatrix, &Contributions, &TransformRange, &BoundsMinRange, &BoundsMaxRange, &SteepnessRange](const PCGPointOctree::FPointRef& InPointRef)
+			{
+				const FVector::FReal Contribution = bUseBounds ? PCGPointHelpers::VolumeOverlap(TransformRange[InPointRef.Index], BoundsMinRange[InPointRef.Index], BoundsMaxRange[InPointRef.Index], SteepnessRange[InPointRef.Index], InBounds, InTransformInverseMatrix) : 1.0;
+				if (Contribution > 0)
+				{
+					Contributions.Emplace(InPointRef.Index, Contribution);
+				}
+			});
+	}
+
+	FVector::FReal SumContributions = 0;
+	FVector::FReal MaxContribution = 0;
+	int32 MaxContributor = INDEX_NONE;
+
+	for (const TPair<int32, FVector::FReal>& Contribution : Contributions)
+	{
+		SumContributions += Contribution.Value;
+
+		if (Contribution.Value > MaxContribution)
+		{
+			MaxContribution = Contribution.Value;
+			MaxContributor = Contribution.Key;
+		}
+	}
+
+	if (SumContributions <= 0)
+	{
+		return false;
+	}
+
+	// Rationale: 
+	// When doing volume-to-volume intersection, we want the final density to reflect the amount of overlap
+	// if any - hence the volume overlap computation before.
+	// But, considering that some points may/will overlap (incl. due to steepness), we want to make sure we do not
+	// sum up to more than the total volume. 
+	// Note that this might create some artifacts on the edges in some instances, but we will revisit this once we have a
+	// better and sufficiently efficient solution.
+	const FVector::FReal DensityNormalizationFactor = ((SumContributions > 1.0) ? (1.0 / SumContributions) : 1.0);
+
+	TArray<TPair<PCGMetadataEntryKey, float>, TInlineAllocator<4>> ContributionsForMetadata;
+
+	// Computed weighted average of spatial properties
+	FVector WeightedPosition = FVector::ZeroVector;
+	FQuat WeightedQuat = FQuat(0.0, 0.0, 0.0, 0.0);
+	FVector WeightedScale = FVector::ZeroVector;
+	FVector::FReal WeightedDensity = 0;
+	FVector WeightedBoundsMin = FVector::ZeroVector;
+	FVector WeightedBoundsMax = FVector::ZeroVector;
+	FVector4 WeightedColor = FVector4::Zero();
+	float WeightedSteepness = 0;
+
+	TArray<int64> MetadataEntries;
+
+	for (const TPair<int32, FVector::FReal>& Contribution : Contributions)
+	{
+		const int32 SourcePointIndex = Contribution.Key;
+		const FVector::FReal Weight = Contribution.Value / SumContributions;
+
+		const FTransform& SourcePointTransform = TransformRange[SourcePointIndex];
+		const float& SourcePointSteepness = SteepnessRange[SourcePointIndex];
+		const float& SourcePointDensity = DensityRange[SourcePointIndex];
+		const FVector& SourcePointBoundsMin = BoundsMinRange[SourcePointIndex];
+		const FVector& SourcePointBoundsMax = BoundsMaxRange[SourcePointIndex];
+		const FVector4& SourcePointColor = ColorRange[SourcePointIndex];
+		const int64& SourcePointMetadataEntry = MetadataEntryRange[SourcePointIndex];
+
+		WeightedPosition += SourcePointTransform.GetLocation() * Weight;
+		WeightedQuat = PCGPointHelpers::AddQuatWithWeight(WeightedQuat, SourcePointTransform.GetRotation(), Weight);
+		WeightedScale += SourcePointTransform.GetScale3D() * Weight;
+
+		if (!bSampleInVolume)
+		{
+			WeightedDensity += PCGPointHelpers::ManhattanDensity(SourcePointTransform, SourcePointBoundsMin, SourcePointBoundsMax, SourcePointSteepness, SourcePointDensity, InTransform.GetLocation());
+		}
+		else
+		{
+			WeightedDensity += SourcePointDensity * (bUseBounds ? (Contribution.Value * DensityNormalizationFactor) : Weight);
+		}
+
+		WeightedBoundsMin += SourcePointBoundsMin * Weight;
+		WeightedBoundsMax += SourcePointBoundsMax * Weight;
+		WeightedColor += SourcePointColor * Weight;
+		WeightedSteepness += SourcePointSteepness * Weight;
+
+		ContributionsForMetadata.Emplace(SourcePointMetadataEntry, static_cast<float>(Weight));
+	}
+
+	// Finally, apply changes to point, based on the projection settings
+	if (InParams.bProjectPositions)
+	{
+		OutPoint.Transform.SetLocation(bSampleInVolume ? WeightedPosition : InTransform.GetLocation());
+	}
+	else
+	{
+		OutPoint.Transform.SetLocation(InTransform.GetLocation());
+	}
+
+	if (InParams.bProjectRotations)
+	{
+		WeightedQuat.Normalize();
+		OutPoint.Transform.SetRotation(WeightedQuat);
+	}
+	else
+	{
+		OutPoint.Transform.SetRotation(InTransform.GetRotation());
+	}
+
+	if (InParams.bProjectScales)
+	{
+		OutPoint.Transform.SetScale3D(WeightedScale);
+	}
+	else
+	{
+		OutPoint.Transform.SetScale3D(InTransform.GetScale3D());
+	}
+
+	OutPoint.Density = static_cast<float>(WeightedDensity);
+	OutPoint.BoundsMin = WeightedBoundsMin;
+	OutPoint.BoundsMax = WeightedBoundsMax;
+	OutPoint.Color = WeightedColor;
+	OutPoint.Steepness = WeightedSteepness;
+
+	if (OutMetadata)
+	{
+		//TRACE_CPUPROFILER_EVENT_SCOPE(UPCGPointData::SamplePoint::SetupMetadata);
+		// Initialise metadata entry for this temporary point
+		PCGMetadataEntryKey Key = OutMetadata ? (OutMetadata->HasParent(Metadata) ? OutMetadata->AddEntry(MetadataEntryRange[MaxContributor]) : OutMetadata->AddEntry()) : PCGInvalidEntryKey;
+		
+		if (ContributionsForMetadata.Num() > 1)
+		{
+			OutMetadata->ComputeWeightedAttribute(Key, MakeArrayView(ContributionsForMetadata), Metadata);
+		}
+	}
+
+	return true;
+}
+
+void UPCGBasePointData::Flatten()
+{
+	if (!Metadata)
+	{
+		return;
+	}
+		
+	// If there is no more attributes, reset all keys from points to invalid
+	if (Metadata->GetAttributeCount() == 0)
+	{
+		TConstPCGValueRange<int64> ConstMetadataEntryRange = GetConstMetadataEntryValueRange();
+		TOptional<const int64> MetadataEntrySingleValue = ConstMetadataEntryRange.GetSingleValue();
+		
+		// Range contains multiple values or the only value in range isn't default
+		if (!MetadataEntrySingleValue.IsSet() || MetadataEntrySingleValue.GetValue() != PCGInvalidEntryKey)
+		{
+			Modify();
+		}
+
+		// Release metadata memory if needed
+		FreeProperties(EPCGPointNativeProperties::MetadataEntry);
+		
+		// Set all values to invalid
+		TPCGValueRange<int64> MetadataEntryRange = GetMetadataEntryValueRange(/*bAllocate=*/false);
+		for (int32 Index = 0; Index < MetadataEntryRange.ViewNum(); ++Index)
+		{
+			MetadataEntryRange[Index] = PCGInvalidEntryKey;
+		}
+		
+		return;
+	}
+
+	// Gather all the keys that are not invalid
+	const TConstPCGValueRange<int64> ConstMetadataEntryRange = GetConstMetadataEntryValueRange();
+	TArray<PCGMetadataEntryKey> EntryKeys;
+	EntryKeys.Reserve(GetNumPoints());
+	for (const int64& MetadataEntry : ConstMetadataEntryRange)
+	{
+		if (MetadataEntry != PCGInvalidEntryKey)
+		{
+			EntryKeys.Add(MetadataEntry);
+		}
+	}
+
+	// Then flatten and compress the Metadata for all invalid entry keys. Return true if something changed.
+	if (Metadata->FlattenAndCompress(EntryKeys))
+	{
+		Modify();
+				
+		// Go over all the points and assign all a new entry key for all points that has a valid entry key in the first place.
+		TPCGValueRange<int64> MetadataEntryRange = GetMetadataEntryValueRange();
+		PCGMetadataEntryKey CurrentEntryKey = 0;
+		for (int64& MetadataEntry : MetadataEntryRange)
+		{
+			if (MetadataEntry != PCGInvalidEntryKey)
+			{
+				MetadataEntry = CurrentEntryKey++;
+			}
+		}
+	}
+}
+
+void UPCGBasePointData::AddToCrc(FArchiveCrc32& Ar, bool bFullDataCrc) const
+{
+	Super::AddToCrc(Ar, bFullDataCrc);
+
+	// The code below has non-trivial cost, and can be disabled from console.
+	if (!bFullDataCrc || !CVarCacheFullPointDataCrc.GetValueOnAnyThread())
+	{
+		// Fallback to UID
+		AddUIDToCrc(Ar);
+		return;
+	}
+
+	TRACE_CPUPROFILER_EVENT_SCOPE(UPCGBasePointData::AddToCrc);
+
+	FString ClassName = StaticClass()->GetPathName();
+	Ar << ClassName;
+
+	int32 NumPoints = GetNumPoints();
+	if (NumPoints == 0)
+	{
+		return;
+	}
+
+	Ar << NumPoints;
+
+	// Crc point data.
+	{
+		auto CrcRange = [] <class T> (FArchiveCrc32 & Ar, const TPCGValueRange<T>&ValueRange)
+		{
+			for (int32 Index = 0; Index < ValueRange.ViewNum(); ++Index)
+			{
+				Ar << ValueRange[Index];
+			}
+		};
+
+		UPCGBasePointData* NonConstThis = const_cast<UPCGBasePointData*>(this);
+
+		// Skip Metadata entry keys
+		CrcRange(Ar, NonConstThis->GetTransformValueRange(false));
+		CrcRange(Ar, NonConstThis->GetDensityValueRange(false));
+		CrcRange(Ar, NonConstThis->GetBoundsMinValueRange(false));
+		CrcRange(Ar, NonConstThis->GetBoundsMaxValueRange(false));
+		CrcRange(Ar, NonConstThis->GetSteepnessValueRange(false));
+		CrcRange(Ar, NonConstThis->GetSeedValueRange(false));
+		CrcRange(Ar, NonConstThis->GetColorValueRange(false));
+	}
+
+	// Crc metadata.
+	if (const UPCGMetadata* PCGMetadata = ConstMetadata())
+	{
+		PCGMetadata->AddToCrc(Ar, bFullDataCrc);
+	}
 }
\ No newline at end of file
diff --git a/Engine/Plugins/PCG/Source/PCG/Private/Data/PCGCollisionWrapperData.cpp b/Engine/Plugins/PCG/Source/PCG/Private/Data/PCGCollisionWrapperData.cpp
--- a/Engine/Plugins/PCG/Source/PCG/Private/Data/PCGCollisionWrapperData.cpp
+++ b/Engine/Plugins/PCG/Source/PCG/Private/Data/PCGCollisionWrapperData.cpp
@@ -321,13 +321,15 @@ bool UPCGCollisionWrapperData::SamplePoint(const FTransform& InTransform, const
 	// E.g. get the point collision attribute value, match to our internal list, get the shape (might require loading), then test against the matching shapes
 	// note that we need to place the stuff in the right referential too, since our shapes will all be at the origin
 	const PCGPointOctree::FPointOctree& Octree = GetPointData()->GetPointOctree();
+	const TConstPCGValueRange<FTransform> TransformRange = GetPointData()->GetConstTransformValueRange();
+	const TConstPCGValueRange<float> DensityRange = GetPointData()->GetConstDensityValueRange();
 
 	float Density = 0;
 	FBox TransformedBounds = InBounds.TransformBy(InTransform);
 
-	Octree.FindElementsWithBoundsTest(FBoxCenterAndExtent(TransformedBounds.GetCenter(), TransformedBounds.GetExtent()), [this, &InTransform, &InBounds, &Density](const PCGPointOctree::FPointRef& InPointRef)
+	Octree.FindElementsWithBoundsTest(FBoxCenterAndExtent(TransformedBounds.GetCenter(), TransformedBounds.GetExtent()), [this, &TransformRange, &DensityRange, &InTransform, &InBounds, &Density](const PCGPointOctree::FPointRef& InPointRef)
 	{
-		const float PointDensity = GetPointData()->GetDensity(InPointRef.Index);
+		const float PointDensity = DensityRange[InPointRef.Index];
 
 		if (PointDensity < Density)
 		{
@@ -336,7 +338,7 @@ bool UPCGCollisionWrapperData::SamplePoint(const FTransform& InTransform, const
 
 		bool bHasOverlap = false;
 
-		FTransform RelativeTransform = InTransform.GetRelativeTransform(GetPointData()->GetTransform(InPointRef.Index));
+		FTransform RelativeTransform = InTransform.GetRelativeTransform(TransformRange[InPointRef.Index]);
 		FCollisionShape CollisionShape;
 		CollisionShape.SetBox(FVector3f(InBounds.GetExtent() * RelativeTransform.GetScale3D()));
 		FPhysicsShapeAdapter CollAdapter(RelativeTransform.GetRotation(), CollisionShape);
diff --git a/Engine/Plugins/PCG/Source/PCG/Private/Data/PCGPointArrayData.cpp b/Engine/Plugins/PCG/Source/PCG/Private/Data/PCGPointArrayData.cpp
new file mode 100644
--- /dev/null
+++ b/Engine/Plugins/PCG/Source/PCG/Private/Data/PCGPointArrayData.cpp
@@ -0,0 +1,33 @@
+// Copyright Epic Games, Inc. All Rights Reserved.
+
+#include "Data/PCGPointArrayData.h"
+
+#include "PCGContext.h"
+
+#include "Data/PCGPointData.h"
+
+void UPCGPointArrayData::GetResourceSizeEx(FResourceSizeEx& CumulativeResourceSize)
+{
+	Super::GetResourceSizeEx(CumulativeResourceSize);
+
+	CumulativeResourceSize.AddDedicatedSystemMemoryBytes(PointArray.GetSizeBytes());
+}
+
+const UPCGPointData* UPCGPointArrayData::ToPointData(FPCGContext* Context, const FBox& InBounds) const
+{
+	UPCGPointData* PointData = FPCGContext::NewObject_AnyThread<UPCGPointData>(Context);
+	PointData->InitializeFromData(this);
+
+	CopyPoints(this, PointData, {}, /*bCopyAll=*/true);
+
+	return PointData;
+}
+
+UPCGSpatialData* UPCGPointArrayData::CopyInternal(FPCGContext* Context) const
+{
+	UPCGPointArrayData* NewPointData = FPCGContext::NewObject_AnyThread<UPCGPointArrayData>(Context);
+	
+	NewPointData->PointArray = PointArray;
+
+	return NewPointData;
+}
diff --git a/Engine/Plugins/PCG/Source/PCG/Private/Data/PCGPointData.cpp b/Engine/Plugins/PCG/Source/PCG/Private/Data/PCGPointData.cpp
--- a/Engine/Plugins/PCG/Source/PCG/Private/Data/PCGPointData.cpp
+++ b/Engine/Plugins/PCG/Source/PCG/Private/Data/PCGPointData.cpp
@@ -10,129 +10,19 @@
 #include "Metadata/Accessors/PCGAttributeAccessorKeys.h"
 
 #include "Algo/Transform.h"
+#include "Containers/StridedView.h"
 #include "GameFramework/Actor.h"
 #include "HAL/IConsoleManager.h"
-#include "Serialization/ArchiveCrc32.h"
 
 #include UE_INLINE_GENERATED_CPP_BY_NAME(PCGPointData)
 
-static TAutoConsoleVariable<bool> CVarCacheFullPointDataCrc(
-	TEXT("pcg.Cache.FullPointDataCrc"),
-	true,
-	TEXT("Enable fine-grained CRC of point data for change tracking on elements that request it, rather than using data UID."));
+#define PCG_POINTDATA_MAKERANGE_BASE(RangeType, MakeViewType, Type, Field) RangeType<Type>(MakeViewType(sizeof(FPCGPoint), Points.IsEmpty() ? nullptr : &Points[0].Field, Points.Num()));
+
+#define PCG_POINTDATA_MAKERANGE(Type, Field) PCG_POINTDATA_MAKERANGE_BASE(TPCGValueRange, MakeStridedView, Type, Field)
+#define PCG_POINTDATA_MAKECONSTRANGE(Type, Field) PCG_POINTDATA_MAKERANGE_BASE(TConstPCGValueRange, MakeConstStridedView, Type, Field)
 
 namespace PCGPointHelpers
 {
-	bool GetDistanceRatios(const FPCGPoint& InPoint, const FVector& InPosition, FVector& OutRatios)
-	{
-		FVector LocalPosition = InPoint.Transform.InverseTransformPosition(InPosition);
-		LocalPosition -= (InPoint.BoundsMax + InPoint.BoundsMin) / 2;
-		LocalPosition /= InPoint.GetExtents();
-
-		// ]-2+s, 2-s] is the valid range of values
-		const FVector::FReal LowerBound = InPoint.Steepness - 2;
-		const FVector::FReal HigherBound = 2 - InPoint.Steepness;
-
-		if (LocalPosition.X <= LowerBound || LocalPosition.X > HigherBound ||
-			LocalPosition.Y <= LowerBound || LocalPosition.Y > HigherBound ||
-			LocalPosition.Z <= LowerBound || LocalPosition.Z > HigherBound)
-		{
-			return false;
-		}
-
-		// [-s, +s] is the range where the density is 1 on that axis
-		const FVector::FReal XDist = FMath::Max(0, FMath::Abs(LocalPosition.X) - InPoint.Steepness);
-		const FVector::FReal YDist = FMath::Max(0, FMath::Abs(LocalPosition.Y) - InPoint.Steepness);
-		const FVector::FReal ZDist = FMath::Max(0, FMath::Abs(LocalPosition.Z) - InPoint.Steepness);
-
-		const FVector::FReal DistanceScale = FMath::Max(2 - 2 * InPoint.Steepness, KINDA_SMALL_NUMBER);
-
-		OutRatios.X = XDist / DistanceScale;
-		OutRatios.Y = YDist / DistanceScale;
-		OutRatios.Z = ZDist / DistanceScale;
-		return true;
-	}
-
-	FVector::FReal ManhattanDensity(const FPCGPoint& InPoint, const FVector& InPosition)
-	{
-		FVector Ratios;
-		if (GetDistanceRatios(InPoint, InPosition, Ratios))
-		{
-			return InPoint.Density * (1 - Ratios.X) * (1 - Ratios.Y) * (1 - Ratios.Z);
-		}
-		else
-		{
-			return 0;
-		}
-	}
-
-	FVector::FReal InverseEuclidianDistance(const FPCGPoint& InPoint, const FVector& InPosition)
-	{
-		FVector Ratios;
-		if (GetDistanceRatios(InPoint, InPosition, Ratios))
-		{
-			return 1.0 - Ratios.Length();
-		}
-		else
-		{
-			return 0;
-		}
-	}
-
-	/** Computes reasonable overlap ratio for point, 1d, 2d and volume overlaps, to be used as weights.
-	* Note that this assumes that either data set is homogeneous in its points dimension (either 0d, 1d, 2d, 3d)
-	* Otherwise there will be some artifacts from our assumption here (namely using a 1.0 value for the additional coordinates).
-	*/
-	FVector::FReal ComputeOverlapRatio(const FBox& Numerator, const FBox& Denominator)
-	{
-		const FVector NumeratorExtent = Numerator.GetExtent();
-		const FVector DenominatorExtent = Denominator.GetExtent();
-
-		return (FVector::FReal)((DenominatorExtent.X > 0 ? NumeratorExtent.X / DenominatorExtent.X : 1.0) *
-			(DenominatorExtent.Y > 0 ? NumeratorExtent.Y / DenominatorExtent.Y : 1.0) *
-			(DenominatorExtent.Z > 0 ? NumeratorExtent.Z / DenominatorExtent.Z : 1.0));
-	}
-
-	FVector::FReal VolumeOverlap(const FPCGPoint& InPoint, const FBox& InBounds, const FMatrix& InInverseTransform)
-	{
-		// This is similar in idea to SAT considering we have two boxes - since we will test all 6 axes.
-		// However, there is some uncertainty due to rotation, and using the overlap value as-is is an overestimation, which might not be critical in this case
-		// TODO: investigate if we should do a 8-pt test instead (would be more precise, but significantly more costly).
-		// Implementation note: we are using FMatrix here because we want to support non-uniform scales
-		const FBox PointBounds = InPoint.GetLocalDensityBounds();
-
-		FMatrix PointTransformToInTransform = InPoint.Transform.ToMatrixWithScale() * InInverseTransform;
-		const FBox PointBoundsTransformed = PointBounds.TransformBy(PointTransformToInTransform);
-
-		const FBox FirstOverlap = InBounds.Overlap(PointBoundsTransformed);
-		if (!FirstOverlap.IsValid)
-		{
-			return 0;
-		}
-
-		FMatrix InTransformToPointTransform = PointTransformToInTransform.Inverse();
-		const FBox InBoundsTransformed = InBounds.TransformBy(InTransformToPointTransform);
-
-		const FBox SecondOverlap = InBoundsTransformed.Overlap(PointBounds);
-		if (!SecondOverlap.IsValid)
-		{
-			return 0;
-		}
-
-		return FMath::Min(ComputeOverlapRatio(FirstOverlap, InBounds), ComputeOverlapRatio(SecondOverlap, InBoundsTransformed));
-	}
-
-	/** Helper function for additive blending of quaternions (copied from ControlRig) */
-	FQuat AddQuatWithWeight(const FQuat& Q, const FQuat& V, float Weight)
-	{
-		FQuat BlendQuat = V * Weight;
-
-		if ((Q | BlendQuat) >= 0.0f)
-			return Q + BlendQuat;
-		else
-			return Q - BlendQuat;
-	}
-
 	void Lerp(const FPCGPoint& A, const FPCGPoint& B, float Ratio, const UPCGMetadata* SourceMetadata, FPCGPoint& OutPoint, UPCGMetadata* OutMetadata)
 	{
 		check(Ratio >= 0 && Ratio <= 1.0f);
@@ -254,14 +144,12 @@ void UPCGPointData::GetResourceSizeEx(FResourceSizeEx& CumulativeResourceSize)
 {
 	Super::GetResourceSizeEx(CumulativeResourceSize);
 
-	CumulativeResourceSize.AddDedicatedSystemMemoryBytes(Points.GetAllocatedSize() + Octree.GetSizeBytes() + PCGPointOctree.GetSizeBytes() + sizeof(Bounds));
+	CumulativeResourceSize.AddDedicatedSystemMemoryBytes(Points.GetAllocatedSize() + Octree.GetSizeBytes());
 }
 
 TArray<FPCGPoint>& UPCGPointData::GetMutablePoints()
 {
-	bOctreeIsDirty = true;
-	bOctreeOldIsDirty = true;
-	bBoundsAreDirty = true;
+	DirtyCache();
 	return Points;
 }
 
@@ -277,56 +165,6 @@ const UPCGPointData::PointOctree& UPCGPointData::GetOctree() const
 }
 PRAGMA_ENABLE_DEPRECATION_WARNINGS
 
-const PCGPointOctree::FPointOctree& UPCGPointData::GetPointOctree() const
-{
-	if (bOctreeIsDirty)
-	{
-		RebuildOctree();
-	}
-
-	return PCGPointOctree;
-}
-
-void UPCGPointData::AddToCrc(FArchiveCrc32& Ar, bool bFullDataCrc) const
-{
-	Super::AddToCrc(Ar, bFullDataCrc);
-
-	// The code below has non-trivial cost, and can be disabled from console.
-	if (!bFullDataCrc || !CVarCacheFullPointDataCrc.GetValueOnAnyThread())
-	{
-		// Fallback to UID
-		AddUIDToCrc(Ar);
-		return;
-	}
-
-	TRACE_CPUPROFILER_EVENT_SCOPE(UPCGPointData::AddToCrc);
-
-	FString ClassName = StaticClass()->GetPathName();
-	Ar << ClassName;
-
-	if (Points.Num() == 0)
-	{
-		return;
-	}
-
-	// Crc point data.
-	{
-		// Create copy so we can zero-out the metadata keys which are non-deterministic.
-		for (const FPCGPoint& OriginalPoint : Points)
-		{
-			FPCGPoint Point = OriginalPoint;
-			Point.MetadataEntry = 0;
-			Ar << Point;
-		} 
-	}
-
-	// Crc metadata.
-	if (const UPCGMetadata* PCGMetadata = ConstMetadata())
-	{
-		PCGMetadata->AddToCrc(Ar, bFullDataCrc);
-	}
-}
-
 TUniquePtr<IPCGAttributeAccessor> UPCGPointData::CreateStaticAccessor(const FPCGAttributePropertySelector& InSelector, bool bQuiet)
 {
 	TUniquePtr<IPCGAttributeAccessor> Accessor;
@@ -391,131 +229,84 @@ TArray<FTransform> UPCGPointData::GetTransformsCopy() const
 	return Transforms;
 }
 
-const FTransform& UPCGPointData::GetTransform(int32 InPointIndex) const
-{
-	return Points[InPointIndex].Transform;
-}
-
-float UPCGPointData::GetDensity(int32 InPointIndex) const
+TPCGValueRange<FTransform> UPCGPointData::GetTransformValueRange(bool bAllocate)
 {
-	return Points[InPointIndex].Density;
+	return PCG_POINTDATA_MAKERANGE(FTransform, Transform);
 }
 
-const FVector& UPCGPointData::GetBoundsMin(int32 InPointIndex) const
+TPCGValueRange<float> UPCGPointData::GetDensityValueRange(bool bAllocate)
 {
-	return Points[InPointIndex].BoundsMin;
+	return PCG_POINTDATA_MAKERANGE(float, Density);
 }
 
-const FVector& UPCGPointData::GetBoundsMax(int32 InPointIndex) const
+TPCGValueRange<FVector> UPCGPointData::GetBoundsMinValueRange(bool bAllocate)
 {
-	return Points[InPointIndex].BoundsMax;
+	return PCG_POINTDATA_MAKERANGE(FVector, BoundsMin);
 }
 
-const FVector4& UPCGPointData::GetColor(int32 InPointIndex) const
+TPCGValueRange<FVector> UPCGPointData::GetBoundsMaxValueRange(bool bAllocate)
 {
-	return Points[InPointIndex].Color;
+	return PCG_POINTDATA_MAKERANGE(FVector, BoundsMax);
 }
 
-float UPCGPointData::GetSteepness(int32 InPointIndex) const
+TPCGValueRange<FVector4> UPCGPointData::GetColorValueRange(bool bAllocate)
 {
-	return Points[InPointIndex].Steepness;
+	return PCG_POINTDATA_MAKERANGE(FVector4, Color);
 }
 
-int32 UPCGPointData::GetSeed(int32 InPointIndex) const
+TPCGValueRange<float> UPCGPointData::GetSteepnessValueRange(bool bAllocate)
 {
-	return Points[InPointIndex].Seed;
+	return PCG_POINTDATA_MAKERANGE(float, Steepness);
 }
 
-int64 UPCGPointData::GetMetadataEntry(int32 InPointIndex) const
+TPCGValueRange<int32> UPCGPointData::GetSeedValueRange(bool bAllocate)
 {
-	return Points[InPointIndex].MetadataEntry;
+	return PCG_POINTDATA_MAKERANGE(int32, Seed);
 }
 
-FBoxSphereBounds UPCGPointData::GetDensityBounds(int32 InPointIndex) const
+TPCGValueRange<int64> UPCGPointData::GetMetadataEntryValueRange(bool bAllocate)
 {
-	return Points[InPointIndex].GetDensityBounds();
+	return PCG_POINTDATA_MAKERANGE(int64, MetadataEntry);
 }
 
-FBox UPCGPointData::GetLocalDensityBounds(int32 InPointIndex) const
+TConstPCGValueRange<FTransform> UPCGPointData::GetConstTransformValueRange() const
 {
-	return Points[InPointIndex].GetLocalDensityBounds();
+	return PCG_POINTDATA_MAKECONSTRANGE(FTransform, Transform);
 }
 
-FBox UPCGPointData::GetLocalBounds(int32 InPointIndex) const
+TConstPCGValueRange<float> UPCGPointData::GetConstDensityValueRange() const
 {
-	return Points[InPointIndex].GetLocalBounds();
+	return PCG_POINTDATA_MAKECONSTRANGE(float, Density);
 }
 
-FVector UPCGPointData::GetLocalCenter(int32 InPointIndex) const
+TConstPCGValueRange<FVector> UPCGPointData::GetConstBoundsMinValueRange() const
 {
-	return Points[InPointIndex].GetLocalCenter();
+	return PCG_POINTDATA_MAKECONSTRANGE(FVector, BoundsMin);
 }
 
-FVector UPCGPointData::GetExtents(int32 InPointIndex) const
+TConstPCGValueRange<FVector> UPCGPointData::GetConstBoundsMaxValueRange() const
 {
-	return Points[InPointIndex].GetExtents();
+	return PCG_POINTDATA_MAKECONSTRANGE(FVector, BoundsMax);
 }
 
-FVector UPCGPointData::GetScaledExtents(int32 InPointIndex) const
+TConstPCGValueRange<FVector4> UPCGPointData::GetConstColorValueRange() const
 {
-	return Points[InPointIndex].GetScaledExtents();
+	return PCG_POINTDATA_MAKECONSTRANGE(FVector4, Color);
 }
 
-FVector UPCGPointData::GetLocalSize(int32 InPointIndex) const
+TConstPCGValueRange<float> UPCGPointData::GetConstSteepnessValueRange() const
 {
-	return Points[InPointIndex].GetLocalSize();
+	return PCG_POINTDATA_MAKECONSTRANGE(float, Steepness);
 }
 
-FVector UPCGPointData::GetScaledLocalSize(int32 InPointIndex) const
+TConstPCGValueRange<int32> UPCGPointData::GetConstSeedValueRange() const
 {
-	return Points[InPointIndex].GetScaledLocalSize();
+	return PCG_POINTDATA_MAKECONSTRANGE(int32, Seed);
 }
 
-FBox UPCGPointData::GetBounds() const
+TConstPCGValueRange<int64> UPCGPointData::GetConstMetadataEntryValueRange() const
 {
-	if (bBoundsAreDirty)
-	{
-		RecomputeBounds();
-	}
-
-	return Bounds;
-}
-
-void UPCGPointData::RecomputeBounds() const
-{
-	FScopeLock Lock(&CachedDataLock);
-
-	if (!bBoundsAreDirty)
-	{
-		return;
-	}
-
-	FBox NewBounds(EForceInit::ForceInit);
-	for (const FPCGPoint& Point : Points)
-	{
-		FBoxSphereBounds PointBounds = Point.GetDensityBounds();
-		NewBounds += FBox::BuildAABB(PointBounds.Origin, PointBounds.BoxExtent);
-	}
-
-	Bounds = NewBounds;
-	bBoundsAreDirty = false;
-}
-
-void UPCGPointData::CopyPointsFrom(const UPCGPointData* InData, const TArray<int>& InDataIndices)
-{
-	TRACE_CPUPROFILER_EVENT_SCOPE(UPCGPointData::CopyPointsFrom);
-	check(InData);
-	Points.SetNum(InDataIndices.Num());
-
-	// TODO: parallel-for this?
-	for (int PointIndex = 0; PointIndex < InDataIndices.Num(); ++PointIndex)
-	{
-		Points[PointIndex] = InData->Points[InDataIndices[PointIndex]];
-	}
-
-	bBoundsAreDirty = true;
-	bOctreeIsDirty = true;
-	bOctreeOldIsDirty = true;
+	return PCG_POINTDATA_MAKECONSTRANGE(int64, MetadataEntry);
 }
 
 void UPCGPointData::SetPoints(const TArray<FPCGPoint>& InPoints)
@@ -523,57 +314,9 @@ void UPCGPointData::SetPoints(const TArray<FPCGPoint>& InPoints)
 	GetMutablePoints() = InPoints;
 }
 
-void UPCGPointData::InitializeFromActor(AActor* InActor, bool* bOutOptionalSanitizedTagAttributeName)
-{
-	check(InActor);
-	check(Metadata && Metadata->GetAttributeCount() == 0);
-
-	AddSinglePointFromActor(InActor, bOutOptionalSanitizedTagAttributeName);
-}
-
-void UPCGPointData::AddSinglePointFromActor(AActor* InActor, bool* bOutOptionalSanitizedTagAttributeName)
-{
-	check(InActor);
-
-	FPCGPoint& Point = GetMutablePoints().Emplace_GetRef();
-	Point.Steepness = 1.0f;
-	Point.Transform = InActor->GetActorTransform();
-
-	const FVector& Position = Point.Transform.GetLocation();
-	Point.Seed = PCGHelpers::ComputeSeed((int)Position.X, (int)Position.Y, (int)Position.Z);
-
-	const FBox LocalBounds = PCGHelpers::GetActorLocalBounds(InActor);
-	Point.BoundsMin = LocalBounds.Min;
-	Point.BoundsMax = LocalBounds.Max;
-
-	Point.MetadataEntry = Metadata->AddEntry();
-
-	FPCGMetadataAttribute<FSoftObjectPath>* ActorReferenceAttribute = Metadata->FindOrCreateAttribute(PCGPointDataConstants::ActorReferenceAttribute, FSoftObjectPath(), /*bAllowsInterpolation=*/false, /*bOverrideParent=*/false, /*bOverwriteIfTypeMismatch=*/false);
-	if (ActorReferenceAttribute)
-	{
-		ActorReferenceAttribute->SetValue(Point.MetadataEntry, FSoftObjectPath(InActor));
-	}
-
-	bool bSanitizedAttributeNames = false;
-
-	// Parse tags as well
-	for (FName Tag : InActor->Tags)
-	{
-		PCG::Private::FParseTagResult TagData(Tag);
-		if (PCG::Private::SetAttributeFromTag(TagData, Metadata, Point.MetadataEntry, /*bCanCreateAttribute=*/true))
-		{
-			bSanitizedAttributeNames |= TagData.HasBeenSanitized();
-		}
-	}
-
-	if (bOutOptionalSanitizedTagAttributeName)
-	{
-		*bOutOptionalSanitizedTagAttributeName = bSanitizedAttributeNames;
-	}
-}
-
 FPCGPoint UPCGPointData::GetPoint(int32 Index) const
 {
+	// This check has an overhead but this method is blueprint callable and we want to avoid crashing
 	if (Points.IsValidIndex(Index))
 	{
 		return Points[Index];
@@ -585,173 +328,6 @@ FPCGPoint UPCGPointData::GetPoint(int32 Index) const
 	}
 }
 
-bool UPCGPointData::SamplePoint(const FTransform& InTransform, const FBox& InBounds, FPCGPoint& OutPoint, UPCGMetadata* OutMetadata) const
-{
-	// Run a projection but don't change the point transform. There is a large overlap in code/functionality so this shares one code path.
-	FPCGProjectionParams Params{};
-	Params.bProjectPositions = Params.bProjectRotations = Params.bProjectScales = false;
-	Params.ColorBlendMode = EPCGProjectionColorBlendMode::SourceValue;
-
-	// The ProjectPoint implementation in this class returns true if the query point is overlapping the point data, which is what SamplePoint should return, so forward the return value.
-	return ProjectPoint(InTransform, InBounds, Params, OutPoint, OutMetadata);
-}
-
-bool UPCGPointData::ProjectPoint(const FTransform& InTransform, const FBox& InBounds, const FPCGProjectionParams& InParams, FPCGPoint& OutPoint, UPCGMetadata* OutMetadata) const
-{
-	return ProjectPoint(InTransform, InBounds, InParams, OutPoint, OutMetadata, true);
-}
-
-bool UPCGPointData::ProjectPoint(const FTransform& InTransform, const FBox& InBounds, const FPCGProjectionParams& InParams, FPCGPoint& OutPoint, UPCGMetadata* OutMetadata, bool bUseBounds) const
-{
-	//TRACE_CPUPROFILER_EVENT_SCOPE(UPCGPointData::SamplePoint);
-	if (bOctreeIsDirty)
-	{
-		RebuildOctree();
-	}
-
-	TArray<TPair<int32, FVector::FReal>, TInlineAllocator<4>> Contributions;
-	const bool bSampleInVolume = (InBounds.GetExtent() != FVector::ZeroVector);
-
-	if (!bSampleInVolume)
-	{
-		const FVector InPosition = InTransform.GetLocation();
-		PCGPointOctree.FindElementsWithBoundsTest(FBoxCenterAndExtent(InPosition, FVector::Zero()), [this, &InPosition, &Contributions](const PCGPointOctree::FPointRef& InPointRef) 
-		{
-			Contributions.Emplace(InPointRef.Index, PCGPointHelpers::InverseEuclidianDistance(Points[InPointRef.Index], InPosition));
-		});
-	}
-	else
-	{
-		FBox TransformedBounds = InBounds.TransformBy(InTransform);
-		FMatrix InTransformInverseMatrix = InTransform.ToMatrixWithScale().Inverse();
-
-		PCGPointOctree.FindElementsWithBoundsTest(FBoxCenterAndExtent(TransformedBounds.GetCenter(), TransformedBounds.GetExtent()), [this, bUseBounds, &InBounds, &InTransformInverseMatrix, &Contributions](const PCGPointOctree::FPointRef& InPointRef)
-		{
-			const FVector::FReal Contribution = bUseBounds ? PCGPointHelpers::VolumeOverlap(Points[InPointRef.Index], InBounds, InTransformInverseMatrix) : 1.0;
-			if (Contribution > 0)
-			{
-				Contributions.Emplace(InPointRef.Index, Contribution);
-			}
-		});
-	}
-
-	FVector::FReal SumContributions = 0;
-	FVector::FReal MaxContribution = 0;
-	int32 MaxContributor = INDEX_NONE;
-
-	for (const TPair<int32, FVector::FReal>& Contribution : Contributions)
-	{
-		SumContributions += Contribution.Value;
-
-		if (Contribution.Value > MaxContribution)
-		{
-			MaxContribution = Contribution.Value;
-			MaxContributor = Contribution.Key;
-		}
-	}
-
-	if (SumContributions <= 0)
-	{
-		return false;
-	}
-
-	// Rationale: 
-	// When doing volume-to-volume intersection, we want the final density to reflect the amount of overlap
-	// if any - hence the volume overlap computation before.
-	// But, considering that some points may/will overlap (incl. due to steepness), we want to make sure we do not
-	// sum up to more than the total volume. 
-	// Note that this might create some artifacts on the edges in some instances, but we will revisit this once we have a
-	// better and sufficiently efficient solution.
-	const FVector::FReal DensityNormalizationFactor = ((SumContributions > 1.0) ? (1.0 / SumContributions) : 1.0);
-
-	TArray<TPair<const FPCGPoint*, float>, TInlineAllocator<4>> ContributionsForMetadata;
-
-	// Computed weighted average of spatial properties
-	FVector WeightedPosition = FVector::ZeroVector;
-	FQuat WeightedQuat = FQuat(0.0, 0.0, 0.0, 0.0);
-	FVector WeightedScale = FVector::ZeroVector;
-	FVector::FReal WeightedDensity = 0;
-	FVector WeightedBoundsMin = FVector::ZeroVector;
-	FVector WeightedBoundsMax = FVector::ZeroVector;
-	FVector4 WeightedColor = FVector4::Zero();
-	float WeightedSteepness = 0;
-
-	for (const TPair<int32, FVector::FReal>& Contribution : Contributions)
-	{
-		const FPCGPoint& SourcePoint = Points[Contribution.Key];
-		const FVector::FReal Weight = Contribution.Value / SumContributions;
-
-		WeightedPosition += SourcePoint.Transform.GetLocation() * Weight;
-		WeightedQuat = PCGPointHelpers::AddQuatWithWeight(WeightedQuat, SourcePoint.Transform.GetRotation(), Weight);
-		WeightedScale += SourcePoint.Transform.GetScale3D() * Weight;
-
-		if (!bSampleInVolume)
-		{
-			WeightedDensity += PCGPointHelpers::ManhattanDensity(SourcePoint, InTransform.GetLocation());
-		}
-		else
-		{
-			WeightedDensity += SourcePoint.Density * (bUseBounds ? (Contribution.Value * DensityNormalizationFactor) : Weight);
-		}
-
-		WeightedBoundsMin += SourcePoint.BoundsMin * Weight;
-		WeightedBoundsMax += SourcePoint.BoundsMax * Weight;
-		WeightedColor += SourcePoint.Color * Weight;
-		WeightedSteepness += SourcePoint.Steepness * Weight;
-
-		ContributionsForMetadata.Emplace(&Points[Contribution.Key], static_cast<float>(Weight));
-	}
-
-	// Finally, apply changes to point, based on the projection settings
-	if (InParams.bProjectPositions)
-	{
-		OutPoint.Transform.SetLocation(bSampleInVolume ? WeightedPosition : InTransform.GetLocation());
-	}
-	else
-	{
-		OutPoint.Transform.SetLocation(InTransform.GetLocation());
-	}
-
-	if (InParams.bProjectRotations)
-	{
-		WeightedQuat.Normalize();
-		OutPoint.Transform.SetRotation(WeightedQuat);
-	}
-	else
-	{
-		OutPoint.Transform.SetRotation(InTransform.GetRotation());
-	}
-
-	if (InParams.bProjectScales)
-	{
-		OutPoint.Transform.SetScale3D(WeightedScale);
-	}
-	else
-	{
-		OutPoint.Transform.SetScale3D(InTransform.GetScale3D());
-	}
-
-	OutPoint.Density = static_cast<float>(WeightedDensity);
-	OutPoint.BoundsMin = WeightedBoundsMin;
-	OutPoint.BoundsMax = WeightedBoundsMax;
-	OutPoint.Color = WeightedColor;
-	OutPoint.Steepness = WeightedSteepness;
-
-	if (OutMetadata)
-	{
-		//TRACE_CPUPROFILER_EVENT_SCOPE(UPCGPointData::SamplePoint::SetupMetadata);
-		// Initialise metadata entry for this temporary point
-		UPCGMetadataAccessorHelpers::InitializeMetadataWithParent(OutPoint, OutMetadata, Points[MaxContributor], Metadata);
-
-		if (ContributionsForMetadata.Num() > 1)
-		{
-			OutMetadata->ComputePointWeightedAttribute(OutPoint, MakeArrayView(ContributionsForMetadata), Metadata);
-		}
-	}
-
-	return true;
-}
-
 void UPCGPointData::RebuildOctreeOld() const
 {
 PRAGMA_DISABLE_DEPRECATION_WARNINGS
@@ -762,7 +338,7 @@ PRAGMA_DISABLE_DEPRECATION_WARNINGS
 		return;
 	}
 
-	TRACE_CPUPROFILER_EVENT_SCOPE(UPCGPointData::RebuildOctree)
+	TRACE_CPUPROFILER_EVENT_SCOPE(UPCGPointData::RebuildOctreeOld)
 	check(bOctreeOldIsDirty);
 
 	FBox PointBounds = GetBounds();
@@ -778,30 +354,6 @@ PRAGMA_DISABLE_DEPRECATION_WARNINGS
 PRAGMA_ENABLE_DEPRECATION_WARNINGS
 }
 
-void UPCGPointData::RebuildOctree() const
-{
-	FScopeLock Lock(&CachedDataLock);
-
-	if (!bOctreeIsDirty)
-	{
-		return;
-	}
-
-	TRACE_CPUPROFILER_EVENT_SCOPE(UPCGPointData::RebuildOctree)
-	check(bOctreeIsDirty);
-
-	FBox PointBounds = GetBounds();
-	PCGPointOctree::FPointOctree NewOctree(PointBounds.GetCenter(), PointBounds.GetExtent().Length());
-
-	for (int32 Index = 0; Index < Points.Num(); ++Index)
-	{
-		NewOctree.AddElement(PCGPointOctree::FPointRef(Index, Points[Index].GetDensityBounds()));
-	}
-
-	PCGPointOctree = MoveTemp(NewOctree);
-	bOctreeIsDirty = false;
-}
-
 UPCGSpatialData* UPCGPointData::CopyInternal(FPCGContext* Context) const
 {
 	UPCGPointData* NewPointData = FPCGContext::NewObject_AnyThread<UPCGPointData>(Context);
@@ -810,58 +362,15 @@ UPCGSpatialData* UPCGPointData::CopyInternal(FPCGContext* Context) const
 	return NewPointData;
 }
 
-void UPCGPointData::Flatten()
+void UPCGPointData::MoveRange(int32 RangeStartIndex, int32 MoveToIndex, int32 NumElements)
 {
-	if (!Metadata)
+	if (RangeStartIndex != MoveToIndex && NumElements > 0)
 	{
-		return;
-	}
-
-	// If there is no more attributes, reset all keys from points to invalid
-	if (Metadata->GetAttributeCount() == 0)
-	{
-		bool bWasModified = false;
-		for (FPCGPoint& Point : Points)
-		{
-			if (Point.MetadataEntry != PCGInvalidEntryKey)
-			{
-				if (!bWasModified)
-				{
-					Modify();
-					bWasModified = true;
-				}
-
-				Point.MetadataEntry = PCGInvalidEntryKey;
-			}
-		}
-
-		return;
-	}
-
-	// Gather all the keys that are not invalid
-	TArray<PCGMetadataEntryKey> EntryKeys;
-	EntryKeys.Reserve(Points.Num());
-	for (const FPCGPoint& Point : Points)
-	{
-		if (Point.MetadataEntry != PCGInvalidEntryKey)
-		{
-			EntryKeys.Add(Point.MetadataEntry);
-		}
-	}
-
-	// Then flatten and compress the Metadata for all invalid entry keys. Return true if something changed.
-	if (Metadata->FlattenAndCompress(EntryKeys))
-	{
-		Modify();
-
-		// Go over all the points and assign all a new entry key for all points that has a valid entry key in the first place.
-		PCGMetadataEntryKey CurrentEntryKey = 0;
-		for (FPCGPoint& Point : Points)
-		{
-			if (Point.MetadataEntry != PCGInvalidEntryKey)
-			{
-				Point.MetadataEntry = CurrentEntryKey++;
-			}
-		}
+		check(Points.IsValidIndex(MoveToIndex) && Points.IsValidIndex(MoveToIndex + NumElements - 1) && Points.IsValidIndex(RangeStartIndex) && Points.IsValidIndex(RangeStartIndex + NumElements - 1));
+		std::memmove(&Points[MoveToIndex], &Points[RangeStartIndex], sizeof(FPCGPoint) * NumElements);
 	}
 }
+
+#undef PCG_POINTDATA_MAKERANGE_BASE
+#undef PCG_POINTDATA_MAKERANGE
+#undef PCG_POINTDATA_MAKECONSTRANGE
\ No newline at end of file
diff --git a/Engine/Plugins/PCG/Source/PCG/Private/Elements/PCGBlurElement.cpp b/Engine/Plugins/PCG/Source/PCG/Private/Elements/PCGBlurElement.cpp
--- a/Engine/Plugins/PCG/Source/PCG/Private/Elements/PCGBlurElement.cpp
+++ b/Engine/Plugins/PCG/Source/PCG/Private/Elements/PCGBlurElement.cpp
@@ -227,7 +227,7 @@ bool FPCGBlurElement::ExecuteInternal(FPCGContext* InContext) const
 							switch (Settings->BlurMode)
 							{
 							case EPCGBlurElementMode::Linear:
-								Weight = 1.0 - FVector::Distance(PointLocation, IterState.InputPointData->GetTransform(Neighbor).GetLocation()) / Settings->SearchDistance;
+								Weight = 1.0 - FVector::Distance(PointLocation, IterState.InputPointData->GetPoints()[Neighbor].Transform.GetLocation()) / Settings->SearchDistance;
 								break;
 							case EPCGBlurElementMode::Gaussian:
 								Weight = GaussianRatio * FMath::Exp(-SquaredDistance * GaussianExpRatio);
diff --git a/Engine/Plugins/PCG/Source/PCG/Private/Elements/PCGCreatePointsGrid.cpp b/Engine/Plugins/PCG/Source/PCG/Private/Elements/PCGCreatePointsGrid.cpp
--- a/Engine/Plugins/PCG/Source/PCG/Private/Elements/PCGCreatePointsGrid.cpp
+++ b/Engine/Plugins/PCG/Source/PCG/Private/Elements/PCGCreatePointsGrid.cpp
@@ -4,6 +4,7 @@
 
 #include "PCGComponent.h"
 #include "PCGContext.h"
+#include "Data/PCGPointArrayData.h"
 #include "Helpers/PCGAsync.h"
 #include "Helpers/PCGHelpers.h"
 #include "Helpers/PCGSettingsHelpers.h"
@@ -207,67 +208,99 @@ bool FPCGCreatePointsGridElement::ExecuteInternal(FPCGContext* Context) const
 		return true;
 	}
 
-	int32 NumIterations = static_cast<int32>(NumIterations64);
-	
-	TArray<FPCGTaggedData>& Outputs = Context->OutputData.TaggedData;
-	FPCGTaggedData& Output = Outputs.Emplace_GetRef();
+	const int32 NumIterations = static_cast<int32>(NumIterations64);
+		
+	UPCGBasePointData* PtData = nullptr;
+	TPCGValueRange<FTransform> TransformRange;
+	TPCGValueRange<int32> SeedRange;
 
-	UPCGPointData* PtData = FPCGContext::NewObject_AnyThread<UPCGPointData>(Context);
-	check(PtData);
+	auto InitFunc = [Context, Settings, CellSize, NumIterations, &PtData, &TransformRange, &SeedRange]()
+	{ 
+		PtData = FPCGContext::NewObject_AnyThread<UPCGPointData>(Context);
+		
+		check(PtData);
+				
+		PtData->SetNumPoints(NumIterations);
 
-	TArray<FPCGPoint>& OutputPoints = PtData->GetMutablePoints();
-	Output.Data = PtData;
+		// Set Steepness
+		PtData->SetSteepness(Settings->PointSteepness);
+		if (Settings->bSetPointsBounds)
+		{
+			PtData->SetExtents(CellSize * 0.5);
+		}
+						
+		TransformRange = PtData->GetTransformValueRange();
+		SeedRange = PtData->GetSeedValueRange();
+	};
 
-	FPCGAsync::AsyncPointProcessing(Context, NumIterations, OutputPoints, [Settings, CullingShape, &PointCountX, &PointCountY, &PointOffsetX, &PointOffsetY, &PointOffsetZ, &CellSize, &GridExtents, &LocalTransform](int32 Index, FPCGPoint& OutPoint)
+	auto ProcessFunc = [Settings, CullingShape, &PointCountX, &PointCountY, &PointOffsetX, &PointOffsetY, &PointOffsetZ, &CellSize, &GridExtents, &LocalTransform, &TransformRange, &SeedRange](int32 StartIndex, int32 StartWriteIndex, int32 Count)
 	{
-		OutPoint = FPCGPoint();
+		ensure(TransformRange.Num() == SeedRange.Num() && (StartWriteIndex + Count) <= TransformRange.Num());
 
-		double XCoordinate = (Index % PointCountX) + PointOffsetX;
-		double YCoordinate = ((Index / PointCountX) % PointCountY) + PointOffsetY;
-		double ZCoordinate = (Index / (PointCountX * PointCountY)) + PointOffsetZ;
-
-		if (Settings->PointPosition == EPCGPointPosition::CellCenter)
+		int32 NumWritten = 0;
+		for(int32 Index = StartIndex; Index < StartIndex+Count; ++Index)
 		{
-			// If extents are smaller than a point in that dimension, place in center, otherwise offset 
-			XCoordinate += 0.5;
-			YCoordinate += 0.5;
-			ZCoordinate += 0.5;
-		}
+			double XCoordinate = (Index % PointCountX) + PointOffsetX;
+			double YCoordinate = ((Index / PointCountX) % PointCountY) + PointOffsetY;
+			double ZCoordinate = (Index / (PointCountX * PointCountY)) + PointOffsetZ;
 
-		// If the extends are smaller than the point, set point to origin
-		if (GridExtents.X < CellSize.X / 2.0)
-		{
-			XCoordinate = 0.0;
-		}
-		if (GridExtents.Y < CellSize.Y / 2.0)
-		{
-			YCoordinate = 0.0;
-		}
-		if (GridExtents.Z < CellSize.Z / 2.0)
-		{
-			ZCoordinate = 0.0;
-		}
+			if (Settings->PointPosition == EPCGPointPosition::CellCenter)
+			{
+				// If extents are smaller than a point in that dimension, place in center, otherwise offset 
+				XCoordinate += 0.5;
+				YCoordinate += 0.5;
+				ZCoordinate += 0.5;
+			}
 
-		FTransform PointTransform = FTransform(FVector((CellSize.X * XCoordinate) - GridExtents.X, (CellSize.Y * YCoordinate) - GridExtents.Y, (CellSize.Z * ZCoordinate) - GridExtents.Z));
+			// If the extends are smaller than the point, set point to origin
+			if (GridExtents.X < CellSize.X / 2.0)
+			{
+				XCoordinate = 0.0;
+			}
+			if (GridExtents.Y < CellSize.Y / 2.0)
+			{
+				YCoordinate = 0.0;
+			}
+			if (GridExtents.Z < CellSize.Z / 2.0)
+			{
+				ZCoordinate = 0.0;
+			}
 
-		if (Settings->CoordinateSpace == EPCGCoordinateSpace::LocalComponent || Settings->CoordinateSpace == EPCGCoordinateSpace::OriginalComponent)
-		{
-			PointTransform *= LocalTransform;
-		}
+			FTransform PointTransform = FTransform(FVector((CellSize.X * XCoordinate) - GridExtents.X, (CellSize.Y * YCoordinate) - GridExtents.Y, (CellSize.Z * ZCoordinate) - GridExtents.Z));
 
-		OutPoint.Transform = PointTransform;
+			if (Settings->CoordinateSpace == EPCGCoordinateSpace::LocalComponent || Settings->CoordinateSpace == EPCGCoordinateSpace::OriginalComponent)
+			{
+				PointTransform *= LocalTransform;
+			}
 
-		if (Settings->bSetPointsBounds)
-		{
-			OutPoint.SetExtents(CellSize * 0.5);
+			// Discards points outside of the volume
+			if (!CullingShape || (CullingShape->GetDensityAtPosition(PointTransform.GetLocation()) > 0.0f))
+			{
+				const int32 WriteIndex = StartWriteIndex + NumWritten;
+				TransformRange[WriteIndex] = PointTransform;
+				SeedRange[WriteIndex] = PCGHelpers::ComputeSeedFromPosition(PointTransform.GetLocation());
+				NumWritten++;
+			}
 		}
 
-		OutPoint.Steepness = Settings->PointSteepness;
-		OutPoint.Seed = PCGHelpers::ComputeSeedFromPosition(PointTransform.GetLocation());
+		return NumWritten;
+	};
+
+	auto MoveRangeFunc = [&PtData](int32 RangeStartIndex, int32 MoveToIndex, int32 NumElements)
+	{
+		PtData->MoveRange(RangeStartIndex, MoveToIndex, NumElements);
+	};
+
+	auto FinishFunc = [&PtData](int32 NumWritten)
+	{
+		PtData->SetNumPoints(NumWritten);
+	};
+
+	FPCGAsync::AsyncProcessingRangeEx(&Context->AsyncState, NumIterations, InitFunc, ProcessFunc, MoveRangeFunc, FinishFunc, false);
 
-		// Discards points outside of the volume
-		return !CullingShape || (CullingShape->GetDensityAtPosition(PointTransform.GetLocation()) > 0.0f);
-	});
+	TArray<FPCGTaggedData>& Outputs = Context->OutputData.TaggedData;
+	FPCGTaggedData& Output = Outputs.Emplace_GetRef();
+	Output.Data = PtData;
 
 	return true;
 }
diff --git a/Engine/Plugins/PCG/Source/PCG/Private/Elements/PCGSelfPruning.cpp b/Engine/Plugins/PCG/Source/PCG/Private/Elements/PCGSelfPruning.cpp
--- a/Engine/Plugins/PCG/Source/PCG/Private/Elements/PCGSelfPruning.cpp
+++ b/Engine/Plugins/PCG/Source/PCG/Private/Elements/PCGSelfPruning.cpp
@@ -166,12 +166,13 @@ namespace PCGSelfPruningElement
 
 		check(IterationState.InputPointData);
 		const PCGPointOctree::FPointOctree& Octree = IterationState.InputPointData->GetPointOctree();
+		const TConstPCGValueRange<FTransform> TransformRange = IterationState.InputPointData->GetConstTransformValueRange();
 
 		int32 CheckTimeSlicingCount = 0;
 
 		TArray<const PCGPointOctree::FPointRef*, TInlineAllocator<256>> ElementsToTest;
 
-		auto SetupQueryInfo = [&IterationState](const PCGPointOctree::FPointRef& PointRef, FBodyInstance* OtherBodyInstance, EPCGCollisionQueryFlag CollisionQueryFlag, FBodyInstance*& OutInstance, PhysicsInterfaceTypes::FInlineShapeArray& OutShapes, FCollisionShape& OutSimpleShape)
+		auto SetupQueryInfo = [&IterationState, &TransformRange](const PCGPointOctree::FPointRef& PointRef, FBodyInstance* OtherBodyInstance, EPCGCollisionQueryFlag CollisionQueryFlag, FBodyInstance*& OutInstance, PhysicsInterfaceTypes::FInlineShapeArray& OutShapes, FCollisionShape& OutSimpleShape)
 		{
 			const int32 EntryIndex = PointRef.Index;
 			OutInstance = IterationState.CollisionWrapper.GetBodyInstance(EntryIndex);
@@ -197,7 +198,7 @@ namespace PCGSelfPruningElement
 					}
 				}
 
-				OutInstance->UpdateBodyScale(IterationState.InputPointData->GetTransform(PointRef.Index).GetScale3D());
+				OutInstance->UpdateBodyScale(TransformRange[PointRef.Index].GetScale3D());
 				const bool bFirstChoice = FPCGCollisionWrapper::GetShapeArray(OutInstance, CollisionQueryFlag, OutShapes);
 
 				if (OutShapes.IsEmpty())
@@ -268,7 +269,7 @@ namespace PCGSelfPruningElement
 			PhysicsInterfaceTypes::FInlineShapeArray OtherShapes;
 			FCollisionShape OtherSimpleShape;
 
-			const FTransform& PointTransform = IterationState.InputPointData->GetTransform(PointRef.Index);
+			const FTransform& PointTransform = TransformRange[PointRef.Index];
 			const bool bThisHasComplexShapes = SetupQueryInfo(PointRef, nullptr, InCollisionQueryFlag, ThisInstance, ThisShapes, ThisSimpleShape);
 			FTransform TransformNoScale = FTransform(PointTransform.GetRotation(), PointTransform.GetLocation());
 
@@ -281,7 +282,7 @@ namespace PCGSelfPruningElement
 
 				bool bOtherHasComplexShapes = SetupQueryInfo(*OtherPointRef, ThisInstance, OtherCollisionQueryFlag, OtherInstance, OtherShapes, OtherSimpleShape);
 
-				const FTransform& OtherPointTransform = IterationState.InputPointData->GetTransform(OtherPointRef->Index);
+				const FTransform& OtherPointTransform = TransformRange[OtherPointRef->Index];
 				FTransform OtherTransformNoScale = FTransform(OtherPointTransform.GetRotation(), OtherPointTransform.GetLocation());
 
 				// Four cases here:
@@ -389,6 +390,10 @@ namespace PCGSelfPruningElement
 
 		int32 CheckTimeSlicingCount = 0;
 
+		const TConstPCGValueRange<FTransform> TransformRange = IterationState.InputPointData->GetConstTransformValueRange();
+		const TConstPCGValueRange<FVector> BoundsMinRange = IterationState.InputPointData->GetConstBoundsMinValueRange();
+		const TConstPCGValueRange<FVector> BoundsMaxRange = IterationState.InputPointData->GetConstBoundsMaxValueRange();
+
 		while (IterationState.CurrentPointIndex < IterationState.SortedPointRefs.Num())
 		{
 			// Don't check too many times. Pre-increment will make sure we always at least process "TimeSliceFrequencyCheck" elements at each iteration.
@@ -409,12 +414,12 @@ namespace PCGSelfPruningElement
 			}
 
 			IterationState.ExclusionPoints.Add(PointRef.Index);
-			const FTransform& PointTransform = IterationState.InputPointData->GetTransform(PointRef.Index);
-			const FVector PointLocalCenter = IterationState.InputPointData->GetLocalCenter(PointRef.Index);
+			const FTransform& PointTransform = TransformRange[PointRef.Index];
+			const FVector PointLocalCenter = PCGPointHelpers::GetLocalCenter(BoundsMinRange[PointRef.Index], BoundsMaxRange[PointRef.Index]);
 
-			Octree.FindElementsWithBoundsTest(FBoxCenterAndExtent(PointTransform.TransformPosition(PointLocalCenter), FVector::Zero()), [&IterationState, &PointTransform](const PCGPointOctree::FPointRef InPointRef)
+			Octree.FindElementsWithBoundsTest(FBoxCenterAndExtent(PointTransform.TransformPosition(PointLocalCenter), FVector::Zero()), [&IterationState, &PointTransform, &TransformRange](const PCGPointOctree::FPointRef InPointRef)
 			{
-				const FTransform& OtherPointTransform = IterationState.InputPointData->GetTransform(InPointRef.Index);
+				const FTransform& OtherPointTransform = TransformRange[InPointRef.Index];
 				if ((PointTransform.GetLocation() - OtherPointTransform.GetLocation()).SquaredLength() <= SMALL_NUMBER && !IterationState.ExclusionPoints.Contains(InPointRef.Index))
 				{
 					IterationState.ExcludedPoints.Add(InPointRef.Index);
diff --git a/Engine/Plugins/PCG/Source/PCG/Private/PCGCommon.cpp b/Engine/Plugins/PCG/Source/PCG/Private/PCGCommon.cpp
--- a/Engine/Plugins/PCG/Source/PCG/Private/PCGCommon.cpp
+++ b/Engine/Plugins/PCG/Source/PCG/Private/PCGCommon.cpp
@@ -216,6 +216,20 @@ namespace PCGPinIdHelpers
 	}
 }
 
+namespace PCGPointCustomPropertyNames
+{
+	bool IsCustomPropertyName(FName Name)
+	{
+		return Name == PCGPointCustomPropertyNames::ExtentsName ||
+			Name == PCGPointCustomPropertyNames::LocalCenterName ||
+			Name == PCGPointCustomPropertyNames::PositionName ||
+			Name == PCGPointCustomPropertyNames::RotationName ||
+			Name == PCGPointCustomPropertyNames::ScaleName ||
+			Name == PCGPointCustomPropertyNames::LocalSizeName ||
+			Name == PCGPointCustomPropertyNames::ScaledLocalSizeName;
+	}
+}
+
 namespace PCGQualityHelpers
 {
 	FName GetQualityPinLabel()
diff --git a/Engine/Plugins/PCG/Source/PCG/Private/PCGElement.cpp b/Engine/Plugins/PCG/Source/PCG/Private/PCGElement.cpp
--- a/Engine/Plugins/PCG/Source/PCG/Private/PCGElement.cpp
+++ b/Engine/Plugins/PCG/Source/PCG/Private/PCGElement.cpp
@@ -413,6 +413,21 @@ bool IPCGElement::PrepareData(FPCGContext* Context) const
 		}
 	}
 
+	if (!SupportsBasePointDataInputs(Context))
+	{
+		for (FPCGTaggedData& Data : Context->InputData.TaggedData)
+		{
+			if (const UPCGBasePointData* BasePointData = Cast<UPCGBasePointData>(Data.Data))
+			{
+				// No doesn't support anything else than UPCGPointData
+				if (!BasePointData->IsA<UPCGPointData>())
+				{
+					Data.Data = BasePointData->ToPointData(Context);
+				}
+			}
+		}
+	}
+
 	return PrepareDataInternal(Context);
 }
 
diff --git a/Engine/Plugins/PCG/Source/PCG/Private/PCGModule.cpp b/Engine/Plugins/PCG/Source/PCG/Private/PCGModule.cpp
--- a/Engine/Plugins/PCG/Source/PCG/Private/PCGModule.cpp
+++ b/Engine/Plugins/PCG/Source/PCG/Private/PCGModule.cpp
@@ -4,6 +4,7 @@
 
 #include "PCGContext.h"
 #include "PCGElement.h"
+#include "Data/PCGBasePointData.h"
 #include "Data/PCGPointData.h"
 
 #include "Interfaces/IPluginManager.h"
@@ -48,6 +49,8 @@ void FPCGModule::StartupModule()
 
 	// Registering accessor methods
 	AttributeAccessorFactory.RegisterDefaultMethods();
+	AttributeAccessorFactory.RegisterMethods<UPCGBasePointData>(UPCGBasePointData::GetPointAccessorMethods());
+	// @todo_pcg: Eventually remove the UPCGPointData method registration because the UPCGBasePointData accessors are compatible
 	AttributeAccessorFactory.RegisterMethods<UPCGPointData>(UPCGPointData::GetPointAccessorMethods());
 
 	// Register onto the PreExit, because we need the class to be still valid to remove them from the mapping
@@ -64,6 +67,7 @@ void FPCGModule::PreExit()
 	// Unregistering accessor methods
 	AttributeAccessorFactory.UnregisterDefaultMethods();
 	AttributeAccessorFactory.UnregisterMethods<UPCGPointData>();
+	AttributeAccessorFactory.UnregisterMethods<UPCGBasePointData>();
 
 #if WITH_EDITOR
 	DeregisterNativeElementDeterminismTests();
diff --git a/Engine/Plugins/PCG/Source/PCG/Private/PCGPoint.cpp b/Engine/Plugins/PCG/Source/PCG/Private/PCGPoint.cpp
--- a/Engine/Plugins/PCG/Source/PCG/Private/PCGPoint.cpp
+++ b/Engine/Plugins/PCG/Source/PCG/Private/PCGPoint.cpp
@@ -2,24 +2,16 @@
 
 #include "PCGPoint.h"
 
-#include "Math/BoxSphereBounds.h"
+#include "PCGCommon.h"
+
 #include "Metadata/Accessors/PCGCustomAccessor.h"
 #include "Metadata/Accessors/IPCGAttributeAccessorTpl.h"
+
+#include "Math/BoxSphereBounds.h"
 #include "UObject/FortniteReleaseBranchCustomObjectVersion.h"
 
 #include UE_INLINE_GENERATED_CPP_BY_NAME(PCGPoint)
 
-namespace PCGPointCustomPropertyNames
-{
-	const FName ExtentsName = TEXT("Extents");
-	const FName LocalCenterName = TEXT("LocalCenter");
-	const FName PositionName = TEXT("Position");
-	const FName RotationName = TEXT("Rotation");
-	const FName ScaleName = TEXT("Scale");
-	const FName LocalSizeName = TEXT("LocalSize");
-	const FName ScaledLocalSizeName = TEXT("ScaledLocalSize");
-}
-
 /** Serialized fields of a FPCGPoint, the values here can't change as they are being used to mask out serialization */
 enum class EPCGPointSerializeFields : uint8
 {
@@ -139,54 +131,9 @@ bool FPCGPoint::Serialize(FStructuredArchive::FSlot Slot)
 	return true;
 }
 
-FBox FPCGPoint::GetLocalBounds() const
-{
-	return FBox(BoundsMin, BoundsMax);
-}
-
-FBox FPCGPoint::GetLocalDensityBounds() const
-{
-	return FBox((2 - Steepness) * BoundsMin, (2 - Steepness) * BoundsMax);
-}
-
-void FPCGPoint::SetLocalBounds(const FBox& InBounds)
-{
-	BoundsMin = InBounds.Min;
-	BoundsMax = InBounds.Max;
-}
-
-FBoxSphereBounds FPCGPoint::GetDensityBounds() const
-{
-	return FBoxSphereBounds(GetLocalDensityBounds().TransformBy(Transform));
-}
-
-void FPCGPoint::ApplyScaleToBounds()
-{
-	const FVector PointScale = Transform.GetScale3D();
-	Transform.SetScale3D(PointScale.GetSignVector());
-	BoundsMin *= PointScale.GetAbs();
-	BoundsMax *= PointScale.GetAbs();
-}
-
-void FPCGPoint::ResetPointCenter(const FVector& BoundsRatio)
-{
-	const FVector NewCenterLocal = FMath::Lerp(BoundsMin, BoundsMax, BoundsRatio);
-
-	BoundsMin -= NewCenterLocal;
-	BoundsMax -= NewCenterLocal;
-
-	Transform.SetLocation(Transform.GetLocation() + Transform.TransformVector(NewCenterLocal));
-}
-
 bool FPCGPoint::HasCustomPropertyGetterSetter(FName Name)
 {
-	return Name == PCGPointCustomPropertyNames::ExtentsName ||
-		Name == PCGPointCustomPropertyNames::LocalCenterName ||
-		Name == PCGPointCustomPropertyNames::PositionName ||
-		Name == PCGPointCustomPropertyNames::RotationName ||
-		Name == PCGPointCustomPropertyNames::ScaleName ||
-		Name == PCGPointCustomPropertyNames::LocalSizeName ||
-		Name == PCGPointCustomPropertyNames::ScaledLocalSizeName;
+	return PCGPointCustomPropertyNames::IsCustomPropertyName(Name);
 }
 
 TUniquePtr<IPCGAttributeAccessor> FPCGPoint::CreateCustomPropertyAccessor(FName Name)
diff --git a/Engine/Plugins/PCG/Source/PCG/Private/PCGPointArray.cpp b/Engine/Plugins/PCG/Source/PCG/Private/PCGPointArray.cpp
new file mode 100644
--- /dev/null
+++ b/Engine/Plugins/PCG/Source/PCG/Private/PCGPointArray.cpp
@@ -0,0 +1,152 @@
+// Copyright Epic Games, Inc. All Rights Reserved.
+
+#include "PCGPointArray.h"
+
+FPCGPointArray::FPCGPointArray()
+{
+	FPCGPoint DefaultPoint;
+	Transform.SetValue(DefaultPoint.Transform);
+	Density.SetValue(DefaultPoint.Density);
+	BoundsMin.SetValue(DefaultPoint.BoundsMin);
+	BoundsMax.SetValue(DefaultPoint.BoundsMax);
+	Color.SetValue(DefaultPoint.Color);
+	Steepness.SetValue(DefaultPoint.Steepness);
+	Seed.SetValue(DefaultPoint.Seed);
+	MetadataEntry.SetValue(DefaultPoint.MetadataEntry);
+}
+
+void FPCGPointArray::SetNumPoints(int32 InNumPoints)
+{
+	NumPoints = InNumPoints;
+	
+	// Reallocate arrays that are currently allocated
+	Transform.SetNum(NumPoints, /*bAllocate=*/false);
+	Density.SetNum(NumPoints, /*bAllocate=*/false);
+	BoundsMin.SetNum(NumPoints, /*bAllocate=*/false);
+	BoundsMax.SetNum(NumPoints, /*bAllocate=*/false);
+	Color.SetNum(NumPoints, /*bAllocate=*/false);
+	Steepness.SetNum(NumPoints, /*bAllocate=*/false);
+	Seed.SetNum(NumPoints, /*bAllocate=*/false);
+	MetadataEntry.SetNum(NumPoints, /*bAllocate=*/false);
+}
+
+void FPCGPointArray::Allocate(EPCGPointNativeProperties InProperties)
+{
+	if (EnumHasAnyFlags(InProperties, EPCGPointNativeProperties::Transform))
+	{
+		Transform.Allocate();
+	}
+
+	if (EnumHasAnyFlags(InProperties, EPCGPointNativeProperties::Density))
+	{
+		Density.Allocate();
+	}
+
+	if (EnumHasAnyFlags(InProperties, EPCGPointNativeProperties::BoundsMin))
+	{
+		BoundsMin.Allocate();
+	}
+
+	if (EnumHasAnyFlags(InProperties, EPCGPointNativeProperties::BoundsMax))
+	{
+		BoundsMax.Allocate();
+	}
+
+	if (EnumHasAnyFlags(InProperties, EPCGPointNativeProperties::Color))
+	{
+		Color.Allocate();
+	}
+
+	if (EnumHasAnyFlags(InProperties, EPCGPointNativeProperties::Steepness))
+	{
+		Steepness.Allocate();
+	}
+
+	if (EnumHasAnyFlags(InProperties, EPCGPointNativeProperties::Seed))
+	{
+		Seed.Allocate();
+	}
+
+	if (EnumHasAnyFlags(InProperties, EPCGPointNativeProperties::MetadataEntry))
+	{
+		MetadataEntry.Allocate();
+	}
+}
+
+void FPCGPointArray::Free(EPCGPointNativeProperties InProperties)
+{
+	if (EnumHasAnyFlags(InProperties, EPCGPointNativeProperties::Transform))
+	{
+		Transform.Free();
+	}
+
+	if (EnumHasAnyFlags(InProperties, EPCGPointNativeProperties::Density))
+	{
+		Density.Free();
+	}
+
+	if (EnumHasAnyFlags(InProperties, EPCGPointNativeProperties::BoundsMin))
+	{
+		BoundsMin.Free();
+	}
+
+	if (EnumHasAnyFlags(InProperties, EPCGPointNativeProperties::BoundsMax))
+	{
+		BoundsMax.Free();
+	}
+
+	if (EnumHasAnyFlags(InProperties, EPCGPointNativeProperties::Color))
+	{
+		Color.Free();
+	}
+
+	if (EnumHasAnyFlags(InProperties, EPCGPointNativeProperties::Steepness))
+	{
+		Steepness.Free();
+	}
+
+	if (EnumHasAnyFlags(InProperties, EPCGPointNativeProperties::Seed))
+	{
+		Seed.Free();
+	}
+
+	if (EnumHasAnyFlags(InProperties, EPCGPointNativeProperties::MetadataEntry))
+	{
+		MetadataEntry.Free();
+	}
+}
+
+void FPCGPointArray::MoveRange(int32 RangeStartIndex, int32 MoveToIndex, int32 NumElements)
+{
+	if (RangeStartIndex != MoveToIndex && NumElements > 0)
+	{
+		Transform.MoveRange(RangeStartIndex, MoveToIndex, NumElements);
+		Density.MoveRange(RangeStartIndex, MoveToIndex, NumElements);
+		BoundsMin.MoveRange(RangeStartIndex, MoveToIndex, NumElements);
+		BoundsMax.MoveRange(RangeStartIndex, MoveToIndex, NumElements);
+		Color.MoveRange(RangeStartIndex, MoveToIndex, NumElements);
+		Steepness.MoveRange(RangeStartIndex, MoveToIndex, NumElements);
+		Seed.MoveRange(RangeStartIndex, MoveToIndex, NumElements);
+		MetadataEntry.MoveRange(RangeStartIndex, MoveToIndex, NumElements);
+	}
+}
+
+TArray<FTransform> FPCGPointArray::GetTransformCopy() const
+{
+	return Transform.GetCopy();
+}
+
+bool FPCGPointArray::Serialize(FArchive& Ar)
+{
+	Ar << NumPoints;
+	Transform.Serialize(Ar);
+	Density.Serialize(Ar);
+	BoundsMin.Serialize(Ar);
+	BoundsMax.Serialize(Ar);
+	Color.Serialize(Ar);
+	Steepness.Serialize(Ar);
+	Seed.Serialize(Ar);
+	MetadataEntry.Serialize(Ar);
+
+	return true;
+}
diff --git a/Engine/Plugins/PCG/Source/PCG/Private/Tests/PCGTestsCommon.cpp b/Engine/Plugins/PCG/Source/PCG/Private/Tests/PCGTestsCommon.cpp
--- a/Engine/Plugins/PCG/Source/PCG/Private/Tests/PCGTestsCommon.cpp
+++ b/Engine/Plugins/PCG/Source/PCG/Private/Tests/PCGTestsCommon.cpp
@@ -126,57 +126,6 @@ namespace PCGTestsCommon
 		return NewObject<AActor>();
 	}
 
-	UPCGPointData* CreateEmptyPointData()
-	{
-		return NewObject<UPCGPointData>();
-	}
-
-	UPCGPointData* CreatePointData()
-	{
-		UPCGPointData* SinglePointData = CreateEmptyPointData();
-
-		check(SinglePointData);
-
-		SinglePointData->GetMutablePoints().Emplace();
-
-		return SinglePointData;
-	}
-
-	UPCGPointData* CreatePointData(const FVector& InLocation)
-	{
-		UPCGPointData* SinglePointData = CreatePointData();
-
-		check(SinglePointData);
-		check(SinglePointData->GetMutablePoints().Num() == 1);
-
-		SinglePointData->GetMutablePoints()[0].Transform.SetLocation(InLocation);
-
-		return SinglePointData;
-	}
-
-	/** Creates a point data with PointCount many points, and randomizes the Transform, Color, and Density */
-	UPCGPointData* CreateRandomPointData(int32 PointCount, int32 Seed, bool bRandomDensity)
-	{
-		TObjectPtr<UPCGPointData> PointData = PCGTestsCommon::CreateEmptyPointData();
-		TArray<FPCGPoint>& Points = PointData->GetMutablePoints();
-
-		Points.Reserve(PointCount);
-	
-		FRandomStream RandomSource(Seed);
-		for (int I = 0; I < PointCount; ++I)
-		{
-			FQuat Rotation(FRotator(RandomSource.FRandRange(0.f, 360.f)).Quaternion()); 
-			FVector Scale(RandomSource.VRand());
-			FVector Location(RandomSource.VRand());
-
-			FPCGPoint& Point = Points.Emplace_GetRef(FTransform(Rotation, Location, Scale), 1.f, I);
-			Point.Color = RandomSource.VRand();
-			Point.Density = bRandomDensity ? RandomSource.GetFraction() : 1.0f;
-		}
-
-		return PointData;
-	}
-
 	UPCGPolyLineData* CreatePolyLineData()
 	{
 		// TODO: spline, landscape spline
diff --git a/Engine/Plugins/PCG/Source/PCG/Private/Tests/PCGValueRangeTest.cpp b/Engine/Plugins/PCG/Source/PCG/Private/Tests/PCGValueRangeTest.cpp
new file mode 100644
--- /dev/null
+++ b/Engine/Plugins/PCG/Source/PCG/Private/Tests/PCGValueRangeTest.cpp
@@ -0,0 +1,155 @@
+// Copyright Epic Games, Inc. All Rights Reserved.
+
+#if WITH_EDITOR
+
+#include "Tests/PCGTestsCommon.h"
+#include "Utils/PCGValueRange.h"
+
+#include "Math/IntVector.h"
+
+#define PCG_FMT(Text, ...) FString::Printf(TEXT(Text), __VA_ARGS__)
+
+class FPCGValueRangeTestBaseClass : public FPCGTestBaseClass
+{
+public:
+	using FPCGTestBaseClass::FPCGTestBaseClass;
+
+	template <typename T>
+	bool TestRange(const EPCGPointNativeProperties& NativeProperty)
+	{
+		const int32 NumPoints = 100;
+		const int32 Seed = 42;
+
+		UPCGPointData* PointDataA = PCGTestsCommon::CreateRandomPointData<UPCGPointData>(NumPoints, Seed, /*bRandomDensity=*/true);
+		UPCGPointArrayData* PointDataB = PCGTestsCommon::CreateRandomPointData<UPCGPointArrayData>(NumPoints, Seed, /*bRandomDensity=*/true);
+
+		const TConstPCGValueRange<T> RangeA = PointDataA->GetConstValueRange<T>(NativeProperty);
+		const TConstPCGValueRange<T> RangeB = PointDataB->GetConstValueRange<T>(NativeProperty);
+
+		UTEST_EQUAL(PCG_FMT("RangeA.Num() == RangeB.Num()", *StaticEnum<EPCGPointNativeProperties>()->GetNameStringByValue((int64)NativeProperty)), RangeA.Num(), RangeB.Num());
+
+		for (int32 i = 0; i < RangeA.Num(); ++i)
+		{
+			UTEST_EQUAL(PCG_FMT("RangeA[%d] == RangeB[%d]", i, i), RangeA[i], RangeB[i]);
+		}
+
+		return true;
+	};
+};
+
+IMPLEMENT_CUSTOM_SIMPLE_AUTOMATION_TEST(FPCGValueRangeTest_Base, FPCGValueRangeTestBaseClass, "Plugins.PCG.ValueRange.Base", PCGTestsCommon::TestFlags)
+IMPLEMENT_CUSTOM_SIMPLE_AUTOMATION_TEST(FPCGValueRangeTest_PointData_Transform, FPCGValueRangeTestBaseClass, "Plugins.PCG.ValueRange.PointData.Transform", PCGTestsCommon::TestFlags)
+IMPLEMENT_CUSTOM_SIMPLE_AUTOMATION_TEST(FPCGValueRangeTest_PointData_Steepness, FPCGValueRangeTestBaseClass, "Plugins.PCG.ValueRange.PointData.Steepness", PCGTestsCommon::TestFlags)
+IMPLEMENT_CUSTOM_SIMPLE_AUTOMATION_TEST(FPCGValueRangeTest_PointData_BoundsMin, FPCGValueRangeTestBaseClass, "Plugins.PCG.ValueRange.PointData.BoundsMin", PCGTestsCommon::TestFlags)
+IMPLEMENT_CUSTOM_SIMPLE_AUTOMATION_TEST(FPCGValueRangeTest_PointData_BoundsMax, FPCGValueRangeTestBaseClass, "Plugins.PCG.ValueRange.PointData.BoundsMax", PCGTestsCommon::TestFlags)
+IMPLEMENT_CUSTOM_SIMPLE_AUTOMATION_TEST(FPCGValueRangeTest_PointData_Color, FPCGValueRangeTestBaseClass, "Plugins.PCG.ValueRange.PointData.Color", PCGTestsCommon::TestFlags)
+IMPLEMENT_CUSTOM_SIMPLE_AUTOMATION_TEST(FPCGValueRangeTest_PointData_Seed, FPCGValueRangeTestBaseClass, "Plugins.PCG.ValueRange.PointData.Seed", PCGTestsCommon::TestFlags)
+IMPLEMENT_CUSTOM_SIMPLE_AUTOMATION_TEST(FPCGValueRangeTest_PointData_Density, FPCGValueRangeTestBaseClass, "Plugins.PCG.ValueRange.PointData.Density", PCGTestsCommon::TestFlags)
+IMPLEMENT_CUSTOM_SIMPLE_AUTOMATION_TEST(FPCGValueRangeTest_PointData_MetadataEntry, FPCGValueRangeTestBaseClass, "Plugins.PCG.ValueRange.PointData.MetadataEntry", PCGTestsCommon::TestFlags)
+
+bool FPCGValueRangeTest_Base::RunTest(const FString& Parameters)
+{
+	TConstPCGValueRange<float> FloatRange;
+
+	// Default range is empty
+	UTEST_EQUAL("FloatRange is empty", FloatRange.Num(), 0);
+
+	// Default range Num elements is equal to the underlying stride view Num elements
+	UTEST_EQUAL("FloatRange Num() == ViewNum()", FloatRange.Num(), FloatRange.ViewNum());
+
+	// Default range doesn't have a single value because it has no values
+	UTEST_FALSE("FloatRange is not single value", FloatRange.GetSingleValue().IsSet());
+
+	TArray<int> IntArray{ 1, 2, 3 };
+	TConstPCGValueRange<int> IntRange(MakeConstStridedView(IntArray));
+
+	// Range constructed from strided view has same elements as strided view
+	UTEST_EQUAL("IntRange num elements", IntRange.Num(), 3);
+	UTEST_EQUAL("IntRange Num() == ViewNum()", IntRange.Num(), IntRange.ViewNum());
+	
+	// Test values
+	UTEST_EQUAL("Test IntRange[0] value", IntRange[0], 1);
+	UTEST_EQUAL("Test IntRange[1] value", IntRange[1], 2);
+	UTEST_EQUAL("Test IntRange[2] value", IntRange[2], 3);
+
+	// Mutliple value range doesn't have a single value because it has multiple values
+	UTEST_FALSE("IntRange is not single value", FloatRange.GetSingleValue().IsSet());
+
+	const UE::Math::TIntVector3<int> IntVectorA = { 1, 2, 3 };
+	const UE::Math::TIntVector3<int> IntVectorB = { 4, 5, 6 };
+	const UE::Math::TIntVector3<int> IntVectorC = { 7, 8, 9 };
+	const TArray<UE::Math::TIntVector3<int>> IntVectorArray{ IntVectorA, IntVectorB, IntVectorC };
+	const int NumElements = 20;
+	TConstPCGValueRange<UE::Math::TIntVector3<int>> IntVectorRange(MakeConstStridedView(IntVectorArray), NumElements);
+
+	// Range constructed with specific element count
+	UTEST_EQUAL("IntVectorRange num elements", IntVectorRange.Num(), NumElements);
+	UTEST_EQUAL("IntVectorRange num view elements", IntVectorRange.ViewNum(), IntVectorArray.Num())
+
+	// Validate range values against view 
+	for (int32 i = 0; i < IntVectorRange.Num(); ++i)
+	{
+		UTEST_EQUAL("Validate IntVectorRange Value against IntVectorArray", IntVectorRange[i], IntVectorArray[i % IntVectorArray.Num()]);
+	}
+
+	// Test values
+	UTEST_EQUAL("Test IntVectorRange[0].X value", IntVectorRange[0].X, IntVectorA.X);
+	UTEST_EQUAL("Test IntVectorRange[1].Y value", IntVectorRange[1].Y, IntVectorB.Y);
+	UTEST_EQUAL("Test IntVectorRange[2].Z value", IntVectorRange[2].Z, IntVectorC.Z);
+
+	const TArray<UE::Math::TIntVector3<int>> IntVectorArraySingleValue{ IntVectorA };
+	TConstPCGValueRange<UE::Math::TIntVector3<int>> SingleIntVectorRange(MakeConstStridedView(IntVectorArraySingleValue), NumElements);
+
+	// Single value range
+	UTEST_EQUAL("SingleIntVectorRange num elements", SingleIntVectorRange.Num(), NumElements);
+	UTEST_GREATER("SingleIntVectorRange Num() > ViewNum()", SingleIntVectorRange.Num(), SingleIntVectorRange.ViewNum());
+
+	UTEST_TRUE("SingleIntVectorRange is single value", SingleIntVectorRange.GetSingleValue().IsSet());
+
+	UTEST_EQUAL("Test SingleIntVectorRange SingleValue()", SingleIntVectorRange.GetSingleValue().GetValue(), IntVectorA);
+	UTEST_EQUAL("Test equality index 0 and index NumElements-1", SingleIntVectorRange[0], SingleIntVectorRange[NumElements - 1]);
+
+	return true;
+}
+
+bool FPCGValueRangeTest_PointData_Transform::RunTest(const FString& Parameters)
+{
+	return TestRange<FTransform>(EPCGPointNativeProperties::Transform);
+}
+
+bool FPCGValueRangeTest_PointData_Steepness::RunTest(const FString& Parameters)
+{
+	return TestRange<float>(EPCGPointNativeProperties::Steepness);
+}
+
+bool FPCGValueRangeTest_PointData_BoundsMin::RunTest(const FString& Parameters)
+{
+	return TestRange<FVector>(EPCGPointNativeProperties::BoundsMin);
+}
+
+bool FPCGValueRangeTest_PointData_BoundsMax::RunTest(const FString& Parameters)
+{
+	return TestRange<FVector>(EPCGPointNativeProperties::BoundsMax);
+}
+
+bool FPCGValueRangeTest_PointData_Color::RunTest(const FString& Parameters)
+{
+	return TestRange<FVector4>(EPCGPointNativeProperties::Color);
+}
+
+bool FPCGValueRangeTest_PointData_Seed::RunTest(const FString& Parameters)
+{
+	return TestRange<int32>(EPCGPointNativeProperties::Seed);
+}
+
+bool FPCGValueRangeTest_PointData_Density::RunTest(const FString& Parameters)
+{
+	return TestRange<float>(EPCGPointNativeProperties::Density);
+}
+
+bool FPCGValueRangeTest_PointData_MetadataEntry::RunTest(const FString& Parameters)
+{
+	return TestRange<int64>(EPCGPointNativeProperties::MetadataEntry);
+}
+
+#endif // WITH_EDITOR
\ No newline at end of file
diff --git a/Engine/Plugins/PCG/Source/PCG/Public/Data/PCGBasePointData.h b/Engine/Plugins/PCG/Source/PCG/Public/Data/PCGBasePointData.h
--- a/Engine/Plugins/PCG/Source/PCG/Public/Data/PCGBasePointData.h
+++ b/Engine/Plugins/PCG/Source/PCG/Public/Data/PCGBasePointData.h
@@ -2,49 +2,291 @@
 
 #pragma once
 
+#include "PCGCommon.h"
 #include "PCGSpatialData.h"
 
 #include "Utils/PCGPointOctree.h"
+#include "Utils/PCGValueRange.h"
 
 #include "PCGBasePointData.generated.h"
 
+namespace PCGPointDataConstants
+{
+	const FName ActorReferenceAttribute = TEXT("ActorReference");
+}
 
-UCLASS(Abstract)
+UCLASS(Abstract, BlueprintType, ClassGroup = (Procedural))
 class PCG_API UPCGBasePointData : public UPCGSpatialData
 {
 	GENERATED_BODY()
 
 public:
 	UPCGBasePointData(const FObjectInitializer& ObjectInitializer);
+
+	// Get the functions to the accessor factory
+	static FPCGAttributeAccessorMethods GetPointAccessorMethods();
+
+	//~Begin UObject Interface
+	virtual void GetResourceSizeEx(FResourceSizeEx& CumulativeResourceSize) override;
+	//~End UObject Interface
+
+	// ~Begin UPCGSpatialData interface
+	virtual int GetDimension() const override { return 0; }
+	virtual bool SamplePoint(const FTransform& Transform, const FBox& Bounds, FPCGPoint& OutPoint, UPCGMetadata* OutMetadata) const override;
+	virtual bool ProjectPoint(const FTransform& InTransform, const FBox& InBounds, const FPCGProjectionParams& InParams, FPCGPoint& OutPoint, UPCGMetadata* OutMetadata) const override;
+	// ~End UPCGSpatialData interface
+		
+	// ~Begin UPCGData Interface
 	
-	virtual bool IsValidRef(const PCGPointOctree::FPointRef& InPointRef) const PURE_VIRTUAL(UPCGBasePointData::IsValidRef, return false;);
+	/** Make a pass on Metadata to flatten parenting and only keep entries used by points. */
+	virtual void Flatten() override;
+	virtual bool SupportsFullDataCrc() const override { return true; }
+	virtual void AddToCrc(FArchiveCrc32& Ar, bool bFullDataCrc) const override;
+	// ~End UPCGData Interface
+
+	virtual bool ProjectPoint(const FTransform& InTransform, const FBox& InBounds, const FPCGProjectionParams& InParams, FPCGPoint& OutPoint, UPCGMetadata* OutMetadata, bool bUseBounds) const;
+
+	/** Initializes a single point based on the given actor */
+	void InitializeFromActor(AActor* InActor, bool* bOutOptionalSanitizedTagAttributeName = nullptr);
+
+	/** Adds a single point based on the given actor */
+	void AddSinglePointFromActor(AActor* InActor, bool* bOutOptionalSanitizedTagAttributeName = nullptr);
 
-	virtual bool IsPointOctreeDirty() const PURE_VIRTUAL(UPCGBasePointData::IsPointOctreeDirty, return false;);
-	virtual const PCGPointOctree::FPointOctree& GetPointOctree() const PURE_VIRTUAL(UPCGBasePointData::GetPointOctree, static PCGPointOctree::FPointOctree Empty; return Empty;);
+	virtual bool IsValidRef(const PCGPointOctree::FPointRef& InPointRef) const PURE_VIRTUAL(UPCGBasePointData::IsValidRef, return false;);
+	
+	virtual EPCGDataType GetDataType() const override { return EPCGDataType::Point; }
 
 	virtual int32 GetNumPoints() const PURE_VIRTUAL(UPCGBasePointData::GetNumPoints, return 0;);
+	virtual void SetNumPoints(int32 InNumPoints) PURE_VIRTUAL(UPCGBasePointData::SetNumPoints, return;);
 
-	virtual TArray<FTransform> GetTransformsCopy() const PURE_VIRTUAL(UPCGBasePointData::GetTransforms, return TArray<FTransform>(););
+	virtual void AllocateProperties(EPCGPointNativeProperties Properties) {}
+	virtual void FreeProperties(EPCGPointNativeProperties Properties) {}
+	virtual void MoveRange(int32 RangeStartIndex, int32 MoveToIndex, int32 NumElements) PURE_VIRTUAL(UPCGBasePointData::MoveRange, return;);
+
+	virtual TArray<FTransform> GetTransformsCopy() const PURE_VIRTUAL(UPCGBasePointData::GetTransformsCopy, return TArray<FTransform>(););
 	
-	virtual const FTransform& GetTransform(int32 InPointIndex) const PURE_VIRTUAL(UPCGBasePointData::GetTransform, return FTransform::Identity;);
-	virtual float GetDensity(int32 InPointIndex) const PURE_VIRTUAL(UPCGBasePointData::GetDensity, return 0.f;);
-	virtual const FVector& GetBoundsMin(int32 InPointIndex) const PURE_VIRTUAL(UPCGBasePointData::GetBoundsMin, return FVector::ZeroVector;);
-	virtual const FVector& GetBoundsMax(int32 InPointIndex) const PURE_VIRTUAL(UPCGBasePointData::GetBoundsMin, return FVector::ZeroVector;);
-	virtual const FVector4& GetColor(int32 InPointIndex) const PURE_VIRTUAL(UPCGBasePointData::GetColor, static FVector4 Zero(0.0f); return Zero;);
-	virtual float GetSteepness(int32 InPointIndex) const PURE_VIRTUAL(UPCGBasePointData::GetSteepness, return 0.f;);
-	virtual int32 GetSeed(int32 InPointIndex) const PURE_VIRTUAL(UPCGBasePointData::GetSeed, return 0;);
-	virtual int64 GetMetadataEntry(int32 InPointIndex) const PURE_VIRTUAL(UPCGBasePointData::GetMetadataEntry, return 0;);
-
-	virtual FBoxSphereBounds GetDensityBounds(int32 InPointIndex) const PURE_VIRTUAL(UPCGBasePointData::GetDensityBounds, return FBoxSphereBounds(););
-	virtual FBox GetLocalDensityBounds(int32 InPointIndex) const PURE_VIRTUAL(UPCGBasePointData::GetLocalDensityBounds, return FBox(););
-	virtual FBox GetLocalBounds(int32 InPointIndex) const PURE_VIRTUAL(UPCGBasePointData::GetLocalBounds, return FBox(););
-	virtual FVector GetLocalCenter(int32 InPointIndex) const PURE_VIRTUAL(UPCGBasePointData::GetLocalCenter, return FVector::ZeroVector;);
-	virtual FVector GetExtents(int32 InPointIndex) const PURE_VIRTUAL(UPCGBasePointData::GetExtents, return FVector::ZeroVector;);
-	virtual FVector GetScaledExtents(int32 InPointIndex) const PURE_VIRTUAL(UPCGBasePointData::GetScaledExtents, return FVector::ZeroVector;);
-	virtual FVector GetLocalSize(int32 InPointIndex) const PURE_VIRTUAL(UPCGBasePointData::GetLocalSize, return FVector::ZeroVector;);
-	virtual FVector GetScaledLocalSize(int32 InPointIndex) const PURE_VIRTUAL(UPCGBasePointData::GetScaledLocalSize, return FVector::ZeroVector;);
+	/**
+	 * GetValueRange / GetConstValueRange / Get*ValueRange / GetConst*ValueRange
+	 * 
+	 * Returns a value range that can be iterated and read from / written to abstracting the underlying point data structure.
+	 * 
+	 * Calling SetNumPoints/AllocateProperties/FreeProperties can and will probably invalidate ranges so make sure that you do those operations first or that you get a new range after you do.
+	 * 
+	 * @param bAllocate In the case where we return a non-const range by default the memory will be allocated but in some specific case we might not want to allocate (for SingleValue() ranges)
+	 */
+	template <typename T>
+	TPCGValueRange<T> GetValueRange(EPCGPointNativeProperties NativeProperty, bool bAllocate = true);
+
+	virtual TPCGValueRange<FTransform> GetTransformValueRange(bool bAllocate = true) PURE_VIRTUAL(UPCGBasePointData::GetTransformValueRange, return TPCGValueRange<FTransform>(););
+	virtual TPCGValueRange<float> GetDensityValueRange(bool bAllocate = true) PURE_VIRTUAL(UPCGBasePointData::GetDensityValueRange, return TPCGValueRange<float>(););
+	virtual TPCGValueRange<FVector> GetBoundsMinValueRange(bool bAllocate = true) PURE_VIRTUAL(UPCGBasePointData::GetBoundsMinValueRange, return TPCGValueRange<FVector>(););
+	virtual TPCGValueRange<FVector> GetBoundsMaxValueRange(bool bAllocate = true) PURE_VIRTUAL(UPCGBasePointData::GetBoundsMaxValueRange, return TPCGValueRange<FVector>(););
+	virtual TPCGValueRange<FVector4> GetColorValueRange(bool bAllocate = true) PURE_VIRTUAL(UPCGBasePointData::GetColorValueRange, return TPCGValueRange<FVector4>(););
+	virtual TPCGValueRange<float> GetSteepnessValueRange(bool bAllocate = true) PURE_VIRTUAL(UPCGBasePointData::GetSteepnessValueRange, return TPCGValueRange<float>(););
+	virtual TPCGValueRange<int32> GetSeedValueRange(bool bAllocate = true) PURE_VIRTUAL(UPCGBasePointData::GetSeedValueRange, return TPCGValueRange<int32>(););
+	virtual TPCGValueRange<int64> GetMetadataEntryValueRange(bool bAllocate = true) PURE_VIRTUAL(UPCGBasePointData::GetMetadataEntryValueRange, return TPCGValueRange<int64>(););
+
+	template <typename T>
+	TConstPCGValueRange<T> GetConstValueRange(EPCGPointNativeProperties NativeProperty) const;
+
+	virtual TConstPCGValueRange<FTransform> GetConstTransformValueRange() const PURE_VIRTUAL(UPCGBasePointData::GetConstTransformValueRange, return TConstPCGValueRange<FTransform>(););
+	virtual TConstPCGValueRange<float> GetConstDensityValueRange() const PURE_VIRTUAL(UPCGBasePointData::GetConstDensityValueRange, return TConstPCGValueRange<float>(););
+	virtual TConstPCGValueRange<FVector> GetConstBoundsMinValueRange() const PURE_VIRTUAL(UPCGBasePointData::GetConstBoundsMinValueRange, return TConstPCGValueRange<FVector>(););
+	virtual TConstPCGValueRange<FVector> GetConstBoundsMaxValueRange() const PURE_VIRTUAL(UPCGBasePointData::GetConstBoundsMaxValueRange, return TConstPCGValueRange<FVector>(););
+	virtual TConstPCGValueRange<FVector4> GetConstColorValueRange() const PURE_VIRTUAL(UPCGBasePointData::GetConstColorValueRange, return TConstPCGValueRange<FVector4>(););
+	virtual TConstPCGValueRange<float> GetConstSteepnessValueRange() const PURE_VIRTUAL(UPCGBasePointData::GetConstSteepnessValueRange, return TConstPCGValueRange<float>(););
+	virtual TConstPCGValueRange<int32> GetConstSeedValueRange() const PURE_VIRTUAL(UPCGBasePointData::GetConstSeedValueRange, return TConstPCGValueRange<int32>(););
+	virtual TConstPCGValueRange<int64> GetConstMetadataEntryValueRange() const PURE_VIRTUAL(UPCGBasePointData::GetConstMetadataEntryValueRange, return TConstPCGValueRange<int64>(););
+		
+	virtual void SetTransform(const FTransform& InTransform);
+	virtual void SetDensity(float InDensity);
+	virtual void SetBoundsMin(const FVector& InBoundsMin);
+	virtual void SetBoundsMax(const FVector& InBoundsMax);
+	virtual void SetColor(const FVector4& InColor);
+	virtual void SetSteepness(float InSteepness);
+	virtual void SetSeed(int32 InSeed);
+	virtual void SetMetadataEntry(int64 InMetadataEntry);
+
+	virtual void SetExtents(const FVector& InExtents);
+	virtual void SetLocalCenter(const FVector& InLocalCenter);
+
+	virtual const FTransform& GetTransform(int32 InPointIndex) const;
+	virtual float GetDensity(int32 InPointIndex) const;
+	virtual const FVector& GetBoundsMin(int32 InPointIndex) const;
+	virtual const FVector& GetBoundsMax(int32 InPointIndex) const;
+	virtual const FVector4& GetColor(int32 InPointIndex) const;
+	virtual float GetSteepness(int32 InPointIndex) const;
+	virtual int32 GetSeed(int32 InPointIndex) const;
+	virtual int64 GetMetadataEntry(int32 InPointIndex) const;
+
+	virtual FBoxSphereBounds GetDensityBounds(int32 InPointIndex) const;
+	virtual FBox GetLocalDensityBounds(int32 InPointIndex) const;
+	virtual FBox GetLocalBounds(int32 InPointIndex) const;
+	virtual FVector GetLocalCenter(int32 InPointIndex) const;
+	virtual FVector GetExtents(int32 InPointIndex) const;
+	virtual FVector GetScaledExtents(int32 InPointIndex) const;
+	virtual FVector GetLocalSize(int32 InPointIndex) const;
+	virtual FVector GetScaledLocalSize(int32 InPointIndex) const;
+		
+	/** Get the dirty status of the Octree. Note that the Point Octree can be rebuilt from another thread, so this info can be invalidated at anytime. */
+	virtual bool IsPointOctreeDirty() const { return bOctreeIsDirty; }
+	virtual const PCGPointOctree::FPointOctree& GetPointOctree() const;
+	virtual FBox GetBounds() const override;
+		
+	UFUNCTION(BlueprintCallable, Category = SpatialData)
+	void CopyPointsFrom(const UPCGBasePointData* InData, const TArray<int32>& InDataIndices);
 
 	/** Metadata specific */
 	virtual FPCGMetadataLayerID GetDefaultMetadataLayerID() const override { return PCGMetadataLayerID::Elements; }
 	virtual TArray<FPCGMetadataLayerID> GetAllSupportedMetadataLayerIDs() const override { return {PCGMetadataLayerID::Data, PCGMetadataLayerID::Elements}; }
-};
\ No newline at end of file
+
+protected:
+	void CopyPoints(const UPCGBasePointData* From, UPCGBasePointData* To, const TArray<int32>& InDataIndices, bool bCopyAll) const;
+
+	void RebuildOctreeIfNeeded() const
+	{
+		if (bOctreeIsDirty)
+		{
+			RebuildOctree();
+		}
+	}
+
+	void RebuildOctree() const;
+	
+	void RecomputeBoundsIfNeeded() const
+	{
+		if (bBoundsAreDirty)
+		{
+			RecomputeBounds();
+		}
+	}
+	
+	void RecomputeBounds() const;
+
+	virtual void DirtyCache()
+	{
+		bOctreeIsDirty = true;
+		bBoundsAreDirty = true;
+	}
+
+	mutable FCriticalSection CachedDataLock;
+	mutable PCGPointOctree::FPointOctree PCGPointOctree;
+	mutable FBox Bounds;
+
+	mutable bool bOctreeIsDirty = true;
+	mutable bool bBoundsAreDirty = true;
+};
+
+template <typename T>
+inline TPCGValueRange<T> UPCGBasePointData::GetValueRange(EPCGPointNativeProperties NativeProperty, bool bAllocate)
+{
+	if constexpr (std::is_same_v<T, FTransform>)
+	{
+		if (NativeProperty == EPCGPointNativeProperties::Transform)
+		{
+			return GetTransformValueRange(bAllocate);
+		}
+	}
+	else if constexpr (std::is_same_v<T, float>)
+	{
+		if (NativeProperty == EPCGPointNativeProperties::Density)
+		{
+			return GetDensityValueRange(bAllocate);
+		}
+		else if (NativeProperty == EPCGPointNativeProperties::Steepness)
+		{
+			return GetSteepnessValueRange(bAllocate);
+		}
+	}
+	else if constexpr (std::is_same_v<T, FVector>)
+	{
+		if (NativeProperty == EPCGPointNativeProperties::BoundsMin)
+		{
+			return GetBoundsMinValueRange(bAllocate);
+		}
+		else if (NativeProperty == EPCGPointNativeProperties::BoundsMax)
+		{
+			return GetBoundsMaxValueRange(bAllocate);
+		}
+	}
+	else if constexpr (std::is_same_v<T, FVector4>)
+	{
+		if (NativeProperty == EPCGPointNativeProperties::Color)
+		{
+			return GetColorValueRange(bAllocate);
+		}
+	}
+	else if constexpr (std::is_same_v<T, int32>)
+	{
+		if (NativeProperty == EPCGPointNativeProperties::Seed)
+		{
+			return GetSeedValueRange(bAllocate);
+		}
+	}
+	else if constexpr (std::is_same_v<T, int64>)
+	{
+		if (NativeProperty == EPCGPointNativeProperties::MetadataEntry)
+		{
+			return GetMetadataEntryValueRange(bAllocate);
+		}
+	}
+	
+	return TPCGValueRange<T>();
+}
+
+template <typename T>
+inline TConstPCGValueRange<T> UPCGBasePointData::GetConstValueRange(EPCGPointNativeProperties NativeProperty) const
+{
+	if constexpr (std::is_same_v<T, FTransform>)
+	{
+		if (NativeProperty == EPCGPointNativeProperties::Transform)
+		{
+			return GetConstTransformValueRange();
+		}
+	}
+	else if constexpr (std::is_same_v<T, float>)
+	{
+		if (NativeProperty == EPCGPointNativeProperties::Density)
+		{
+			return GetConstDensityValueRange();
+		}
+		else if (NativeProperty == EPCGPointNativeProperties::Steepness)
+		{
+			return GetConstSteepnessValueRange();
+		}
+	}
+	else if constexpr (std::is_same_v<T, FVector>)
+	{
+		if (NativeProperty == EPCGPointNativeProperties::BoundsMin)
+		{
+			return GetConstBoundsMinValueRange();
+		}
+		else if (NativeProperty == EPCGPointNativeProperties::BoundsMax)
+		{
+			return GetConstBoundsMaxValueRange();
+		}
+	}
+	else if constexpr (std::is_same_v<T, FVector4>)
+	{
+		if (NativeProperty == EPCGPointNativeProperties::Color)
+		{
+			return GetConstColorValueRange();
+		}
+	}
+	else if constexpr (std::is_same_v<T, int32>)
+	{
+		if (NativeProperty == EPCGPointNativeProperties::Seed)
+		{
+			return GetConstSeedValueRange();
+		}
+	}
+	else if constexpr (std::is_same_v <T, int64>)
+	{
+		if (NativeProperty == EPCGPointNativeProperties::MetadataEntry)
+		{
+			return GetConstMetadataEntryValueRange();
+		}
+	}
+	
+	return TConstPCGValueRange<T>();
+}
\ No newline at end of file
diff --git a/Engine/Plugins/PCG/Source/PCG/Public/Data/PCGPointArrayData.h b/Engine/Plugins/PCG/Source/PCG/Public/Data/PCGPointArrayData.h
new file mode 100644
--- /dev/null
+++ b/Engine/Plugins/PCG/Source/PCG/Public/Data/PCGPointArrayData.h
@@ -0,0 +1,57 @@
+// Copyright Epic Games, Inc. All Rights Reserved.
+
+#pragma once
+
+#include "PCGBasePointData.h"
+#include "PCGPointArray.h"
+
+#include "PCGPointArrayData.generated.h"
+
+UCLASS(BlueprintType, ClassGroup = (Procedural))
+class PCG_API UPCGPointArrayData : public UPCGBasePointData
+{
+	GENERATED_BODY()
+public:	
+	//~Begin UObject Interface
+	virtual void GetResourceSizeEx(FResourceSizeEx& CumulativeResourceSize) override;
+	//~End UObject Interface
+
+	virtual const UPCGPointData* ToPointData(FPCGContext* Context, const FBox& InBounds = FBox(EForceInit::ForceInit)) const override;
+	virtual UPCGSpatialData* CopyInternal(FPCGContext* Context) const override;
+
+	//~Begin UPCGBasePointData Interface
+	virtual bool IsValidRef(const PCGPointOctree::FPointRef& InPointRef) const override { return InPointRef.Index >= 0 && InPointRef.Index < GetNumPoints(); }
+	
+	virtual int32 GetNumPoints() const override { return PointArray.GetNumPoints(); }
+	virtual void SetNumPoints(int32 InNumPoints) override { PointArray.SetNumPoints(InNumPoints); }
+		
+	virtual void AllocateProperties(EPCGPointNativeProperties Properties) override { PointArray.Allocate(Properties); }
+	virtual void FreeProperties(EPCGPointNativeProperties Properties) override { PointArray.Free(Properties); }
+	virtual void MoveRange(int32 RangeStartIndex, int32 MoveToIndex, int32 NumElements) override { PointArray.MoveRange(RangeStartIndex, MoveToIndex, NumElements); }
+
+	virtual TArray<FTransform> GetTransformsCopy() const { return PointArray.GetTransformCopy(); }
+
+	virtual TPCGValueRange<FTransform> GetTransformValueRange(bool bAllocate = true) override { return PointArray.GetTransformValueRange(bAllocate); }
+	virtual TPCGValueRange<float> GetDensityValueRange(bool bAllocate = true) override { return PointArray.GetDensityValueRange(bAllocate); }
+	virtual TPCGValueRange<FVector> GetBoundsMinValueRange(bool bAllocate = true) override { return PointArray.GetBoundsMinValueRange(bAllocate); }
+	virtual TPCGValueRange<FVector> GetBoundsMaxValueRange(bool bAllocate = true) override { return PointArray.GetBoundsMaxValueRange(bAllocate); }
+	virtual TPCGValueRange<FVector4> GetColorValueRange(bool bAllocate = true) override { return PointArray.GetColorValueRange(bAllocate); }
+	virtual TPCGValueRange<float> GetSteepnessValueRange(bool bAllocate = true) override { return PointArray.GetSteepnessValueRange(bAllocate); }
+	virtual TPCGValueRange<int32> GetSeedValueRange(bool bAllocate = true) override { return PointArray.GetSeedValueRange(bAllocate); }
+	virtual TPCGValueRange<int64> GetMetadataEntryValueRange(bool bAllocate = true) override { return PointArray.GetMetadataEntryValueRange(bAllocate); }
+
+	virtual TConstPCGValueRange<FTransform> GetConstTransformValueRange() const override { return PointArray.GetConstTransformValueRange(); }
+	virtual TConstPCGValueRange<float> GetConstDensityValueRange() const override { return PointArray.GetConstDensityValueRange(); }
+	virtual TConstPCGValueRange<FVector> GetConstBoundsMinValueRange() const override { return PointArray.GetConstBoundsMinValueRange(); }
+	virtual TConstPCGValueRange<FVector> GetConstBoundsMaxValueRange() const override { return PointArray.GetConstBoundsMaxValueRange(); }
+	virtual TConstPCGValueRange<FVector4> GetConstColorValueRange() const override { return PointArray.GetConstColorValueRange(); }
+	virtual TConstPCGValueRange<float> GetConstSteepnessValueRange() const override { return PointArray.GetConstSteepnessValueRange(); }
+	virtual TConstPCGValueRange<int32> GetConstSeedValueRange() const override { return PointArray.GetConstSeedValueRange(); }
+	virtual TConstPCGValueRange<int64> GetConstMetadataEntryValueRange() const override { return PointArray.GetConstMetadataEntryValueRange(); }
+	//~End UPCGBasePointData Interface
+		
+private:
+
+	UPROPERTY()
+	FPCGPointArray PointArray;
+};
\ No newline at end of file
diff --git a/Engine/Plugins/PCG/Source/PCG/Public/Data/PCGPointData.h b/Engine/Plugins/PCG/Source/PCG/Public/Data/PCGPointData.h
--- a/Engine/Plugins/PCG/Source/PCG/Public/Data/PCGPointData.h
+++ b/Engine/Plugins/PCG/Source/PCG/Public/Data/PCGPointData.h
@@ -3,7 +3,9 @@
 #pragma once
 
 #include "Data/PCGBasePointData.h"
+#include "Utils/PCGValueRange.h"
 
+#include "Containers/StridedView.h"
 #include "Math/GenericOctreePublic.h"
 #include "Math/GenericOctree.h"
 
@@ -20,11 +22,6 @@ namespace PCGPointHelpers
 	void BilerpWithSnapping(const FPCGPoint& X0Y0, const FPCGPoint& X1Y0, const FPCGPoint& X0Y1, const FPCGPoint& X1Y1, const UPCGMetadata* SourceMetadata, FPCGPoint& OutPoint, UPCGMetadata* OutMetadata, float XFactor, float YFactor);
 }
 
-namespace PCGPointDataConstants
-{
-	const FName ActorReferenceAttribute = TEXT("ActorReference");
-}
-
 PRAGMA_DISABLE_DEPRECATION_WARNINGS
 
 struct PCG_API UE_DEPRECATED(5.6, "Use PCGPointOctree::FPointRef instead") FPCGPointRef
@@ -83,11 +80,6 @@ PRAGMA_ENABLE_DEPRECATION_WARNINGS
 	virtual void GetResourceSizeEx(FResourceSizeEx& CumulativeResourceSize) override;
 	//~End UObject Interface
 
-	// ~Begin UPCGData interface
-	virtual EPCGDataType GetDataType() const override { return EPCGDataType::Point; }
-	virtual void AddToCrc(FArchiveCrc32& Ar, bool bFullDataCrc) const override;
-	// ~End UPCGData interface
-
 	// Static helper to create an accessor on a data that doesn't yet exist, as accessors for point data don't rely on existing data.
 	static TUniquePtr<IPCGAttributeAccessor> CreateStaticAccessor(const FPCGAttributePropertySelector& InSelector, bool bQuiet = false);
 
@@ -99,47 +91,34 @@ PRAGMA_ENABLE_DEPRECATION_WARNINGS
 
 	virtual TArray<FTransform> GetTransformsCopy() const override;
 
-	virtual const FTransform& GetTransform(int32 InPointIndex) const override;
-	virtual float GetDensity(int32 InPointIndex) const override;
-	virtual const FVector& GetBoundsMin(int32 InPointIndex) const override;
-	virtual const FVector& GetBoundsMax(int32 InPointIndex) const override;
-	virtual const FVector4& GetColor(int32 InPointIndex) const override;
-	virtual float GetSteepness(int32 InPointIndex) const override;
-	virtual int32 GetSeed(int32 InPointIndex) const override;
-	virtual int64 GetMetadataEntry(int32 InPointIndex) const override;
-
-	virtual FBoxSphereBounds GetDensityBounds(int32 InPointIndex) const override;
-	virtual FBox GetLocalDensityBounds(int32 InPointIndex) const override;
-	virtual FBox GetLocalBounds(int32 InPointIndex) const override;
-	virtual FVector GetLocalCenter(int32 InPointIndex) const override;
-	virtual FVector GetExtents(int32 InPointIndex) const override;
-	virtual FVector GetScaledExtents(int32 InPointIndex) const override;
-	virtual FVector GetLocalSize(int32 InPointIndex) const override;
-	virtual FVector GetScaledLocalSize(int32 InPointIndex) const override;
+	virtual TPCGValueRange<FTransform> GetTransformValueRange(bool bAllocate = true) override;
+	virtual TPCGValueRange<float> GetDensityValueRange(bool bAllocate = true) override;
+	virtual TPCGValueRange<FVector> GetBoundsMinValueRange(bool bAllocate = true) override;
+	virtual TPCGValueRange<FVector> GetBoundsMaxValueRange(bool bAllocate = true) override;
+	virtual TPCGValueRange<FVector4> GetColorValueRange(bool bAllocate = true) override;
+	virtual TPCGValueRange<float> GetSteepnessValueRange(bool bAllocate = true) override;
+	virtual TPCGValueRange<int32> GetSeedValueRange(bool bAllocate = true) override;
+	virtual TPCGValueRange<int64> GetMetadataEntryValueRange(bool bAllocate = true) override;
+
+	virtual TConstPCGValueRange<FTransform> GetConstTransformValueRange() const override;
+	virtual TConstPCGValueRange<float> GetConstDensityValueRange() const override;
+	virtual TConstPCGValueRange<FVector> GetConstBoundsMinValueRange() const override;
+	virtual TConstPCGValueRange<FVector> GetConstBoundsMaxValueRange() const override;
+	virtual TConstPCGValueRange<FVector4> GetConstColorValueRange() const override;
+	virtual TConstPCGValueRange<float> GetConstSteepnessValueRange() const override;
+	virtual TConstPCGValueRange<int32> GetConstSeedValueRange() const override;
+	virtual TConstPCGValueRange<int64> GetConstMetadataEntryValueRange() const override;
+	
+	virtual void MoveRange(int32 RangeStartIndex, int32 MoveToIndex, int32 NumElements) override;
 	// ~End UPCGBasePointData interface
 
 	// ~Begin UPCGSpatialData interface
-	virtual int GetDimension() const override { return 0; }
-	virtual FBox GetBounds() const override;
 	virtual const UPCGPointData* ToPointData(FPCGContext* Context, const FBox& InBounds = FBox(EForceInit::ForceInit)) const { return this; }
-	virtual bool SamplePoint(const FTransform& Transform, const FBox& Bounds, FPCGPoint& OutPoint, UPCGMetadata* OutMetadata) const override;
-	virtual bool ProjectPoint(const FTransform& InTransform, const FBox& InBounds, const FPCGProjectionParams& InParams, FPCGPoint& OutPoint, UPCGMetadata* OutMetadata) const override;
-
 protected:
 	virtual UPCGSpatialData* CopyInternal(FPCGContext* Context) const override;
 	// ~End UPCGSpatialData interface
 
-	// ~Begin UPCGData interface
-	virtual bool SupportsFullDataCrc() const override { return true; }
-	// ~End UPCGData interface
-
 public:
-	/** Initializes a single point based on the given actor */
-	void InitializeFromActor(AActor* InActor, bool* bOutOptionalSanitizedTagAttributeName = nullptr);
-
-	/** Adds a single point based on the given actor */
-	void AddSinglePointFromActor(AActor* InActor, bool* bOutOptionalSanitizedTagAttributeName = nullptr);
-
 	UFUNCTION(BlueprintCallable, Category = SpatialData)
 	const TArray<FPCGPoint>& GetPoints() const { return Points; }
 
@@ -149,6 +128,8 @@ public:
 	UFUNCTION(BlueprintCallable, Category = SpatialData)
 	virtual int32 GetNumPoints() const override { return Points.Num(); }
 
+	virtual void SetNumPoints(int32 InNumPoints) override { Points.SetNum(InNumPoints); }
+
 	UFUNCTION(BlueprintCallable, Category = SpatialData)
 	bool IsEmpty() const { return Points.IsEmpty(); }
 
@@ -158,19 +139,7 @@ public:
 	UFUNCTION(BlueprintCallable, Category = SpatialData)
 	void SetPoints(const TArray<FPCGPoint>& InPoints);
 
-	UFUNCTION(BlueprintCallable, Category = SpatialData)
-	void CopyPointsFrom(const UPCGPointData* InData, const TArray<int>& InDataIndices);
-
-	/** Make a pass on Metadata to flatten parenting and only keep entries used by points. */
-	virtual void Flatten() override;
-
 	TArray<FPCGPoint>& GetMutablePoints();
-	
-	/** Get the dirty status of the Octree. Note that the Point Octree can be rebuilt from another thread, so this info can be invalidated at anytime. */
-	virtual bool IsPointOctreeDirty() const { return bOctreeIsDirty; }
-	virtual const PCGPointOctree::FPointOctree& GetPointOctree() const override;
-
-	bool ProjectPoint(const FTransform& InTransform, const FBox& InBounds, const FPCGProjectionParams& InParams, FPCGPoint& OutPoint, UPCGMetadata* OutMetadata, bool bUseBounds) const;
 
 PRAGMA_DISABLE_DEPRECATION_WARNINGS
 	UE_DEPRECATED(5.6, "Use GetPointOctree instead")
@@ -181,18 +150,17 @@ PRAGMA_ENABLE_DEPRECATION_WARNINGS
 
 protected:
 	void RebuildOctreeOld() const;
-	void RebuildOctree() const;
-	void RecomputeBounds() const;
 
+	virtual void DirtyCache() override
+	{
+		Super::DirtyCache();
+		bOctreeOldIsDirty = true;
+	}
+	
 	UPROPERTY()
 	TArray<FPCGPoint> Points;
 
-	mutable FCriticalSection CachedDataLock;
-	mutable PCGPointOctree::FPointOctree PCGPointOctree;
-	mutable FBox Bounds; // TODO: review if this needs to be threadsafe
-	mutable bool bBoundsAreDirty = true;
 	mutable bool bOctreeOldIsDirty = true;
-	mutable bool bOctreeIsDirty = true;
 
 PRAGMA_DISABLE_DEPRECATION_WARNINGS
 	// deprecated
diff --git a/Engine/Plugins/PCG/Source/PCG/Public/Helpers/PCGPointHelpers.h b/Engine/Plugins/PCG/Source/PCG/Public/Helpers/PCGPointHelpers.h
new file mode 100644
--- /dev/null
+++ b/Engine/Plugins/PCG/Source/PCG/Public/Helpers/PCGPointHelpers.h
@@ -0,0 +1,89 @@
+// Copyright Epic Games, Inc. All Rights Reserved.
+
+#pragma once
+
+#include "CoreTypes.h"
+#include "Math/Box.h"
+#include "Math/BoxSphereBounds.h"
+#include "Math/UnrealMathUtility.h"
+
+namespace PCGPointHelpers
+{
+	FORCEINLINE FBox GetLocalBounds(const FVector& InBoundsMin, const FVector& InBoundsMax)
+	{
+		return FBox(InBoundsMin, InBoundsMax);
+	}
+
+	FORCEINLINE void SetLocalBounds(const FBox& InBounds, FVector& OutBoundsMin, FVector& OutBoundsMax)
+	{
+		OutBoundsMin = InBounds.Min;
+		OutBoundsMax = InBounds.Max;
+	}
+
+	FORCEINLINE FBox GetLocalDensityBounds(float InSteepness, const FVector& InBoundsMin, const FVector& InBoundsMax)
+	{
+		return FBox((2 - InSteepness) * InBoundsMin, (2 - InSteepness) * InBoundsMax);
+	}
+
+	FORCEINLINE FVector GetLocalCenter(const FVector& InBoundsMin, const FVector& InBoundsMax)
+	{
+		return (InBoundsMax + InBoundsMin) / 2.0;
+	}
+
+	FORCEINLINE FVector GetExtents(const FVector& InBoundsMin, const FVector& InBoundsMax)
+	{
+		return (InBoundsMax - InBoundsMin) / 2.0;
+	}
+
+	FORCEINLINE FBoxSphereBounds GetDensityBounds(const FTransform& InTransform, float InSteepness, const FVector& InBoundsMin, const FVector& InBoundsMax)
+	{
+		return FBoxSphereBounds(PCGPointHelpers::GetLocalDensityBounds(InSteepness, InBoundsMin, InBoundsMax).TransformBy(InTransform));
+	}
+	
+	FORCEINLINE void SetExtents(const FVector& InExtents, FVector& InOutBoundsMin, FVector& InOutBoundsMax)
+	{
+		const FVector Center = PCGPointHelpers::GetLocalCenter(InOutBoundsMin, InOutBoundsMax);
+		InOutBoundsMin = Center - InExtents;
+		InOutBoundsMax = Center + InExtents;
+	}
+		
+	FORCEINLINE FVector GetScaledExtents(const FTransform& InTransform, const FVector& InBoundsMin, const FVector& InBoundsMax)
+	{ 
+		return PCGPointHelpers::GetExtents(InBoundsMin, InBoundsMax) * InTransform.GetScale3D();
+	}
+
+	FORCEINLINE void SetLocalCenter(const FVector& InCenter, FVector& InOutBoundsMin, FVector& InOutBoundsMax)
+	{
+		const FVector Delta = InCenter - PCGPointHelpers::GetLocalCenter(InOutBoundsMin, InOutBoundsMax);
+		InOutBoundsMin += Delta;
+		InOutBoundsMax += Delta;
+	}
+
+	FORCEINLINE FVector GetLocalSize(const FVector& InBoundsMin, const FVector& InBoundsMax)
+	{ 
+		return InBoundsMax - InBoundsMin; 
+	}
+	
+	FORCEINLINE FVector GetScaledLocalSize(const FTransform& InTransform, const FVector& InBoundsMin, const FVector& InBoundsMax)
+	{ 
+		return PCGPointHelpers::GetLocalSize(InBoundsMin, InBoundsMax) * InTransform.GetScale3D();
+	}
+
+	FORCEINLINE void ApplyScaleToBounds(FTransform& InOutTransform, FVector& InOutBoundsMin, FVector& InOutBoundsMax)
+	{
+		const FVector PointScale = InOutTransform.GetScale3D();
+		InOutTransform.SetScale3D(PointScale.GetSignVector());
+		InOutBoundsMin *= PointScale.GetAbs();
+		InOutBoundsMax *= PointScale.GetAbs();
+	}
+
+	FORCEINLINE void ResetPointCenter(const FVector& BoundsRatio, FTransform& InOutTransform, FVector& InOutBoundsMin, FVector& InOutBoundsMax)
+	{
+		const FVector NewCenterLocal = FMath::Lerp(InOutBoundsMin, InOutBoundsMax, BoundsRatio);
+
+		InOutBoundsMin -= NewCenterLocal;
+		InOutBoundsMax -= NewCenterLocal;
+
+		InOutTransform.SetLocation(InOutTransform.GetLocation() + InOutTransform.TransformVector(NewCenterLocal));
+	}
+}
\ No newline at end of file
diff --git a/Engine/Plugins/PCG/Source/PCG/Public/Metadata/Accessors/PCGAttributeAccessorKeys.h b/Engine/Plugins/PCG/Source/PCG/Public/Metadata/Accessors/PCGAttributeAccessorKeys.h
--- a/Engine/Plugins/PCG/Source/PCG/Public/Metadata/Accessors/PCGAttributeAccessorKeys.h
+++ b/Engine/Plugins/PCG/Source/PCG/Public/Metadata/Accessors/PCGAttributeAccessorKeys.h
@@ -71,6 +71,17 @@ public:
 	template <typename ObjectType>
 	bool GetKeys(int32 InStart, TArrayView<const ObjectType*> OutKeys) const;
 
+	/**
+	 * Retrieve in the given array indices that can be accessed for Accessor/AccessorKeys that support it
+	 * It will wrap around if the index/range goes outside the number of keys.
+	 * @param InStart - Index to start looking in the keys.
+	 * @param InCount - Number of indices to get.
+	 * @param OutKeyIndices - Out array of key indices. 
+	 * @param OutContiguous - Out value notifying caller that OutKeyIndices is empty but key indices are contiguous so [InStart, InStart+Count-1] is a valid range for the accessor
+	 * @return true if it succeeded, false otherwise.
+	 */
+	virtual bool GetKeyIndices(int32 InStart, int32 InCount, TArray<int32>& OutKeyIndices, bool& bOutContiguous) const { return false; }
+
 	/**
 	* Retrieve in the given argument pointer of the wanted type.
 	* Need to be a supported type, such as FPCGPoint, PCGMetadataEntryKey or void.
diff --git a/Engine/Plugins/PCG/Source/PCG/Public/PCGCommon.h b/Engine/Plugins/PCG/Source/PCG/Public/PCGCommon.h
--- a/Engine/Plugins/PCG/Source/PCG/Public/PCGCommon.h
+++ b/Engine/Plugins/PCG/Source/PCG/Public/PCGCommon.h
@@ -43,6 +43,35 @@ namespace PCGPinIdHelpers
 	uint64 GetPinIndexFromPinId(FPCGPinId PinId);
 }
 
+namespace PCGPointCustomPropertyNames
+{
+	bool IsCustomPropertyName(FName Name);
+
+	const FName ExtentsName = TEXT("Extents");
+	const FName LocalCenterName = TEXT("LocalCenter");
+	const FName PositionName = TEXT("Position");
+	const FName RotationName = TEXT("Rotation");
+	const FName ScaleName = TEXT("Scale");
+	const FName LocalSizeName = TEXT("LocalSize");
+	const FName ScaledLocalSizeName = TEXT("ScaledLocalSize");
+}
+
+// Value names need to match EPCGPointProperties
+UENUM(meta = (Bitflags))
+enum class EPCGPointNativeProperties : uint32
+{
+	None = 0 UMETA(Hidden),
+	Transform = 1 << 0,
+	Density = 1 << 1,
+	BoundsMin = 1 << 2,
+	BoundsMax = 1 << 3,
+	Color = 1 << 4,
+	Steepness = 1 << 5,
+	Seed = 1 << 6,
+	MetadataEntry = 1 << 7,
+};
+ENUM_CLASS_FLAGS(EPCGPointNativeProperties);
+
 UENUM(meta = (Bitflags))
 enum class EPCGChangeType : uint32
 {
diff --git a/Engine/Plugins/PCG/Source/PCG/Public/PCGElement.h b/Engine/Plugins/PCG/Source/PCG/Public/PCGElement.h
--- a/Engine/Plugins/PCG/Source/PCG/Public/PCGElement.h
+++ b/Engine/Plugins/PCG/Source/PCG/Public/PCGElement.h
@@ -183,6 +183,9 @@ protected:
 	/** Returns true if node supports data objects that are only resident on GPU (not CPU). */
 	virtual bool SupportsGPUResidentData(FPCGContext* InContext) const { return false; }
 
+	/** Returns true if node supports any point data format. If not a conversion (costly) will happen on the node inputs. */
+	virtual bool SupportsBasePointDataInputs(FPCGContext* InContext) const { return false; }
+
 #if WITH_EDITOR
 	virtual bool ShouldLog() const { return true; }
 #endif
diff --git a/Engine/Plugins/PCG/Source/PCG/Public/PCGPoint.h b/Engine/Plugins/PCG/Source/PCG/Public/PCGPoint.h
--- a/Engine/Plugins/PCG/Source/PCG/Public/PCGPoint.h
+++ b/Engine/Plugins/PCG/Source/PCG/Public/PCGPoint.h
@@ -2,6 +2,7 @@
 
 #pragma once
 
+#include "Helpers/PCGPointHelpers.h"
 
 #include "Math/Box.h"
 #include "PCGPoint.generated.h"
@@ -35,11 +36,6 @@ public:
 	FPCGPoint() = default;
 	FPCGPoint(const FTransform& InTransform, float InDensity, int32 InSeed);
 
-	FBox GetLocalBounds() const;
-	FBox GetLocalDensityBounds() const;
-	void SetLocalBounds(const FBox& InBounds);
-	FBoxSphereBounds GetDensityBounds() const;
-
 	UPROPERTY(BlueprintReadWrite, EditAnywhere, Category = Properties)
 	FTransform Transform;
 
@@ -64,30 +60,70 @@ public:
 	UPROPERTY(BlueprintReadOnly, VisibleAnywhere, Category = "Properties|Metadata")
 	int64 MetadataEntry = -1;
 
-	FVector GetExtents() const { return (BoundsMax - BoundsMin) / 2.0; }
+	FBox GetLocalBounds() const
+	{
+		return PCGPointHelpers::GetLocalBounds(BoundsMin, BoundsMax);
+	}
+
+	FBox GetLocalDensityBounds() const
+	{
+		return PCGPointHelpers::GetLocalDensityBounds(Steepness, BoundsMin, BoundsMax);
+	}
+
+	void SetLocalBounds(const FBox& InBounds)
+	{
+		PCGPointHelpers::SetLocalBounds(InBounds, BoundsMin, BoundsMax);
+	}
+
+	FBoxSphereBounds GetDensityBounds() const
+	{
+		return PCGPointHelpers::GetDensityBounds(Transform, Steepness, BoundsMin, BoundsMax);
+	}
+
+	FVector GetExtents() const 
+	{ 
+		return PCGPointHelpers::GetExtents(BoundsMin, BoundsMax);
+	}
+
 	void SetExtents(const FVector& InExtents)
 	{
-		const FVector Center = GetLocalCenter();
-		BoundsMin = Center - InExtents;
-		BoundsMax = Center + InExtents;
+		PCGPointHelpers::SetExtents(InExtents, BoundsMin, BoundsMax);
 	}
 
-	FVector GetScaledExtents() const { return GetExtents() * Transform.GetScale3D(); }
+	FVector GetScaledExtents() const 
+	{ 
+		return PCGPointHelpers::GetScaledExtents(Transform, BoundsMin, BoundsMax);
+	}
 
-	FVector GetLocalCenter() const { return (BoundsMax + BoundsMin) / 2.0; }
+	FVector GetLocalCenter() const 
+	{ 
+		return PCGPointHelpers::GetLocalCenter(BoundsMin, BoundsMax);
+	}
+	
 	void SetLocalCenter(const FVector& InCenter)
 	{
-		const FVector Delta = InCenter - GetLocalCenter();
-		BoundsMin += Delta;
-		BoundsMax += Delta;
+		PCGPointHelpers::SetLocalCenter(InCenter, BoundsMin, BoundsMax);
 	}
 
-	FVector GetLocalSize() const { return BoundsMax - BoundsMin; }
-	FVector GetScaledLocalSize() const { return GetLocalSize() * Transform.GetScale3D(); }
+	FVector GetLocalSize() const 
+	{ 
+		return PCGPointHelpers::GetLocalSize(BoundsMin, BoundsMax);
+	}
+	
+	FVector GetScaledLocalSize() const 
+	{ 
+		return PCGPointHelpers::GetScaledLocalSize(Transform, BoundsMin, BoundsMax);
+	}
 
-	void ApplyScaleToBounds();
+	void ApplyScaleToBounds()
+	{
+		PCGPointHelpers::ApplyScaleToBounds(Transform, BoundsMin, BoundsMax);
+	}
 
-	void ResetPointCenter(const FVector& BoundsRatio);
+	void ResetPointCenter(const FVector& BoundsRatio)
+	{
+		PCGPointHelpers::ResetPointCenter(BoundsRatio, Transform, BoundsMin, BoundsMax);
+	}
 
 	using PointCustomPropertyGetter = TFunction<bool(const FPCGPoint&, void*)>;
 	using PointCustomPropertySetter = TFunction<bool(FPCGPoint&, const void*)>;
diff --git a/Engine/Plugins/PCG/Source/PCG/Public/PCGPointArray.h b/Engine/Plugins/PCG/Source/PCG/Public/PCGPointArray.h
new file mode 100644
--- /dev/null
+++ b/Engine/Plugins/PCG/Source/PCG/Public/PCGPointArray.h
@@ -0,0 +1,324 @@
+// Copyright Epic Games, Inc. All Rights Reserved.
+
+#pragma once
+
+#include "PCGCommon.h"
+#include "PCGPoint.h"
+
+#include "Helpers/PCGPointHelpers.h"
+#include "Utils/PCGValueRange.h"
+
+#include "CoreTypes.h"
+#include "Containers/StridedView.h"
+#include "Math/MathFwd.h"
+#include "Serialization/Archive.h"
+
+#include "PCGPointArray.generated.h"
+
+/**
+ * Templated struct used to store NumValues of the same type.
+ * If all values are equal then they are stored into the Value field and Values is not allocated
+ * If some values differ then Values is allocated and stores the different values
+ * 
+ * It provides a GetValueRange()/GetConstValueRange() implementation so that those NumValues can be indexed without user knowing of the internal allocations.
+ */
+template<typename T>
+struct FPCGPointArrayProperty
+{
+public:
+	/** Defined for serialization purposes(FPCGPointArray trait WithIdenticalViaEquality) */
+	bool operator==(const FPCGPointArrayProperty<T>& Other) const
+	{
+		if (NumValues != Other.NumValues || Values.Num() != Other.Values.Num())
+		{
+			return false;
+		}
+
+		auto Equals = [](const T& ValueA, const T& ValueB)
+		{
+			if constexpr (std::is_same_v<FVector,T> || std::is_same_v<FTransform,T> || std::is_same_v<FVector4,T>)
+			{
+				return ValueA.Equals(ValueB);
+			}
+			else
+			{
+				return ValueA == ValueB;
+			}
+		};
+
+		if(!Equals(Value, Other.Value))
+		{
+			return false;
+		}
+
+		for (int32 Index = 0; Index < Values.Num(); ++Index)
+		{
+			if (!Equals(Values[Index], Other.Values[Index]))
+			{
+				return false;
+			}
+		}
+
+		return true;
+	}
+	
+	/** If 'Values' isn't currently allocated, allocate it and copy 'Value' at every index of 'Values[0, NumValues - 1]' */
+	void Allocate()
+	{
+		SetNum(NumValues, /*bAllocate=*/true);
+	}
+
+	/** Set the 'NumValues', optionally allocating 'Values' */
+	void SetNum(int32 NewNum, bool bAllocate)
+	{
+		NumValues = NewNum;
+
+		const int32 OldNum = Values.Num();
+		if (OldNum == NewNum)
+		{
+			return;
+		}
+
+		if (OldNum > 0 || bAllocate)
+		{
+			Values.SetNum(NewNum);
+			for (int32 i = OldNum; i < NewNum; ++i)
+			{
+				Values[i] = Value;
+			}
+		}
+	};
+
+	/** Free 'Values', effectively resets all values to 'Value' */
+	void Free()
+	{
+		Values.Reset();
+	}
+
+	/** Set all values('Value/Values') to InValue */
+	void SetValue(const T& InValue)
+	{
+		for (int32 i = 0; i < Values.Num(); ++i)
+		{
+			Values[i] = Value;
+		}
+		Value = InValue;
+	}
+
+	/** Return value at Index */
+	const T& GetValue(int32 Index) const
+	{
+		ensure(Index >= 0 && Index < NumValues);
+		if (Values.Num() == 0)
+		{
+			return Value;
+		}
+
+		return Values[Index];
+	}
+
+	/** Move ranges of values(moves them back into the array).Used by the PCGAsync api */
+	void MoveRange(int32 RangeStartIndex, int32 MoveToIndex, int32 NumElements)
+	{
+		if (IsAllocated())
+		{
+			check(Values.IsValidIndex(MoveToIndex) && Values.IsValidIndex(MoveToIndex + NumElements - 1) && Values.IsValidIndex(RangeStartIndex) && Values.IsValidIndex(RangeStartIndex + NumElements - 1));
+			std::memmove(&Values[MoveToIndex], &Values[RangeStartIndex], sizeof(T) * NumElements);
+		}
+	}
+
+	/* Returns true if 'Values' is allocated */
+	bool IsAllocated() const
+	{
+		ensure(Values.Num() == 0 || Values.Num() == NumValues);
+		return Values.Num() != 0;
+	}
+
+	/** Returns a TArray copy of size NumValues */
+	TArray<T> GetCopy() const
+	{
+		if (IsAllocated())
+		{
+			return Values;
+		}
+
+		TArray<T> Copy;
+		Copy.SetNumUninitialized(NumValues);
+		for (int32 i = 0; i < NumValues; ++i)
+		{
+			Copy[i] = Value;
+		}
+
+		return Copy;
+	}
+
+	/** Returns a TPCGValueRange<T> of NumValues and optionally allocate 'Values' */
+	TPCGValueRange<T> GetValueRange(bool bAllocate)
+	{
+		if (bAllocate)
+		{
+			// Allocate when getting non-const version
+			SetNum(NumValues, /*bForce=*/true);
+		}
+
+		return TPCGValueRange<T>(IsAllocated() ? MakeStridedView(Values) : MakeStridedView(sizeof(Value), &Value, 1), NumValues);
+	}
+
+	/** Returns a TConstPCGValueRange<T> of NumValues */
+	TConstPCGValueRange<T> GetConstValueRange() const
+	{
+		return TConstPCGValueRange<T>(IsAllocated() ? MakeConstStridedView(Values) : MakeConstStridedView(sizeof(Value), &Value, 1), NumValues);
+	}
+		
+	SIZE_T GetSizeBytes() const
+	{
+		return Values.GetAllocatedSize() + sizeof(Value) + sizeof(NumValues) + sizeof(Values);
+	}
+
+	void Serialize(FArchive& Ar)
+	{
+		Ar << NumValues;
+		Ar << Value;
+		Ar << Values;
+	}
+
+private:
+	// Array containing values if allocated
+	TArray<T> Values;
+	// Value representing all values if array is unallocated
+	T Value;
+	// Number of values represented by this FPCGPointArrayProperty
+	int32 NumValues;
+};
+
+USTRUCT()
+struct PCG_API FPCGPointArray
+{
+	GENERATED_BODY()
+public:
+	FPCGPointArray();
+
+	friend bool operator==(const FPCGPointArray&, const FPCGPointArray&) = default;
+
+	void SetNumPoints(int32 InNumPoints);
+	int32 GetNumPoints() const { return NumPoints; }
+		
+	void Allocate(EPCGPointNativeProperties InProperties);
+	void Free(EPCGPointNativeProperties InProperties);
+	void MoveRange(int32 RangeStartIndex, int32 MoveToIndex, int32 NumElements);
+
+	SIZE_T GetSizeBytes() const
+	{
+		return Transform.GetSizeBytes() 
+			+ Density.GetSizeBytes() 
+			+ BoundsMin.GetSizeBytes() 
+			+ BoundsMax.GetSizeBytes() 
+			+ Color.GetSizeBytes() 
+			+ Steepness.GetSizeBytes() 
+			+ Seed.GetSizeBytes() 
+			+ MetadataEntry.GetSizeBytes();
+	}
+
+	TArray<FTransform> GetTransformCopy() const;
+
+	TPCGValueRange<FTransform> GetTransformValueRange(bool bAllocate) { return Transform.GetValueRange(bAllocate); }
+	TPCGValueRange<float> GetDensityValueRange(bool bAllocate) { return Density.GetValueRange(bAllocate); }
+	TPCGValueRange<FVector> GetBoundsMinValueRange(bool bAllocate) { return BoundsMin.GetValueRange(bAllocate); }
+	TPCGValueRange<FVector> GetBoundsMaxValueRange(bool bAllocate) { return BoundsMax.GetValueRange(bAllocate); }
+	TPCGValueRange<FVector4> GetColorValueRange(bool bAllocate) { return Color.GetValueRange(bAllocate); }
+	TPCGValueRange<float> GetSteepnessValueRange(bool bAllocate) { return Steepness.GetValueRange(bAllocate); }
+	TPCGValueRange<int32> GetSeedValueRange(bool bAllocate) { return Seed.GetValueRange(bAllocate); }
+	TPCGValueRange<int64> GetMetadataEntryValueRange(bool bAllocate) { return MetadataEntry.GetValueRange(bAllocate); }
+
+	TConstPCGValueRange<FTransform> GetConstTransformValueRange() const { return Transform.GetConstValueRange(); }
+	TConstPCGValueRange<float> GetConstDensityValueRange() const { return Density.GetConstValueRange(); }
+	TConstPCGValueRange<FVector> GetConstBoundsMinValueRange() const { return BoundsMin.GetConstValueRange(); }
+	TConstPCGValueRange<FVector> GetConstBoundsMaxValueRange() const { return BoundsMax.GetConstValueRange(); }
+	TConstPCGValueRange<FVector4> GetConstColorValueRange() const { return Color.GetConstValueRange(); }
+	TConstPCGValueRange<float> GetConstSteepnessValueRange() const { return Steepness.GetConstValueRange(); }
+	TConstPCGValueRange<int32> GetConstSeedValueRange() const { return Seed.GetConstValueRange(); }
+	TConstPCGValueRange<int64> GetConstMetadataEntryValueRange() const { return MetadataEntry.GetConstValueRange(); }
+
+	void SetTransform(const FTransform& InTransform) { Transform.SetValue(InTransform); }
+	void SetDensity(float InDensity) { Density.SetValue(InDensity); }
+	void SetBoundsMin(const FVector& InBoundsMin) { BoundsMin.SetValue(InBoundsMin); }
+	void SetBoundsMax(const FVector& InBoundsMax) { BoundsMax.SetValue(InBoundsMax); }
+	void SetColor(const FVector4& InColor) { Color.SetValue(InColor); }
+	void SetSteepness(float InSteepness) { Steepness.SetValue(InSteepness); }
+	void SetSeed(int32 InSeed) { Seed.SetValue(InSeed); }
+	void SetMetadataEntry(int64 InMetadataEntry) { MetadataEntry.SetValue(InMetadataEntry); }
+		
+	const FTransform& GetTransform(int32 InPointIndex) const { return Transform.GetValue(InPointIndex); }
+	float GetDensity(int32 InPointIndex) const { return Density.GetValue(InPointIndex); }
+	const FVector& GetBoundsMin(int32 InPointIndex) const { return BoundsMin.GetValue(InPointIndex); }
+	const FVector& GetBoundsMax(int32 InPointIndex) const { return BoundsMax.GetValue(InPointIndex); }
+	const FVector4& GetColor(int32 InPointIndex) const { return Color.GetValue(InPointIndex); }
+	float GetSteepness(int32 InPointIndex) const { return Steepness.GetValue(InPointIndex); }
+	int32 GetSeed(int32 InPointIndex) const { return Seed.GetValue(InPointIndex); }
+	int64 GetMetadataEntry(int32 InPointIndex) const { return MetadataEntry.GetValue(InPointIndex); }
+
+	FBoxSphereBounds GetDensityBounds(int32 InPointIndex) const
+	{
+		return PCGPointHelpers::GetDensityBounds(GetTransform(InPointIndex), GetSteepness(InPointIndex), GetBoundsMin(InPointIndex), GetBoundsMax(InPointIndex));
+	}
+
+	FBox GetLocalDensityBounds(int32 InPointIndex) const
+	{
+		return PCGPointHelpers::GetLocalDensityBounds(GetSteepness(InPointIndex), GetBoundsMin(InPointIndex), GetBoundsMax(InPointIndex));
+	}
+
+	FBox GetLocalBounds(int32 InPointIndex) const
+	{
+		return PCGPointHelpers::GetLocalBounds(GetBoundsMin(InPointIndex), GetBoundsMax(InPointIndex));
+	}
+
+	FVector GetLocalCenter(int32 InPointIndex) const
+	{
+		return PCGPointHelpers::GetLocalCenter(GetBoundsMin(InPointIndex), GetBoundsMax(InPointIndex));
+	}
+
+	FVector GetExtents(int32 InPointIndex) const
+	{
+		return PCGPointHelpers::GetExtents(GetBoundsMin(InPointIndex), GetBoundsMax(InPointIndex));
+	}
+
+	FVector GetScaledExtents(int32 InPointIndex) const
+	{
+		return PCGPointHelpers::GetScaledExtents(GetTransform(InPointIndex), GetBoundsMin(InPointIndex), GetBoundsMax(InPointIndex));
+	}
+	
+	FVector GetLocalSize(int32 InPointIndex) const
+	{
+		return PCGPointHelpers::GetLocalSize(GetBoundsMin(InPointIndex), GetBoundsMax(InPointIndex));
+	}
+
+	FVector GetScaledLocalSize(int32 InPointIndex) const
+	{
+		return PCGPointHelpers::GetScaledLocalSize(GetTransform(InPointIndex), GetBoundsMin(InPointIndex), GetBoundsMax(InPointIndex));
+	}
+
+	// WithSerializer & WithIdenticalViaEquality needed for this struct to properly serialize
+	bool Serialize(FArchive& Ar);
+		
+private:
+	FPCGPointArrayProperty<FTransform>	Transform;
+	FPCGPointArrayProperty<float>		Density;
+	FPCGPointArrayProperty<FVector>		BoundsMin;
+	FPCGPointArrayProperty<FVector>		BoundsMax;
+	FPCGPointArrayProperty<FVector4>	Color;
+	FPCGPointArrayProperty<float>		Steepness;
+	FPCGPointArrayProperty<int32>		Seed;
+	FPCGPointArrayProperty<int64>		MetadataEntry;
+
+	int32 NumPoints = 0;
+};
+
+template<>
+struct TStructOpsTypeTraits<FPCGPointArray> : public TStructOpsTypeTraitsBase2<FPCGPointArray>
+{
+	enum
+	{
+		WithSerializer = true,
+		WithIdenticalViaEquality = true,
+	};
+};
\ No newline at end of file
diff --git a/Engine/Plugins/PCG/Source/PCG/Public/Tests/PCGTestsCommon.h b/Engine/Plugins/PCG/Source/PCG/Public/Tests/PCGTestsCommon.h
--- a/Engine/Plugins/PCG/Source/PCG/Public/Tests/PCGTestsCommon.h
+++ b/Engine/Plugins/PCG/Source/PCG/Public/Tests/PCGTestsCommon.h
@@ -5,7 +5,9 @@
 #include "PCGData.h"
 #include "PCGPoint.h"
 #include "PCGSettings.h"
+#include "Data/PCGBasePointData.h"
 #include "Data/PCGPointData.h"
+#include "Data/PCGPointArrayData.h"
 #include "Metadata/PCGMetadata.h"
 #include "Metadata/PCGMetadataAttributeTpl.h"
 #include "Metadata/PCGMetadataAttributeTraits.h"
@@ -53,17 +55,71 @@ namespace PCGTestsCommon
 	};
 
 	AActor* CreateTemporaryActor();
-	PCG_API UPCGPointData* CreateEmptyPointData();
 	UPCGParamData* CreateEmptyParamData();
 
 	/** Creates a PointData with a single point at the origin */
-	PCG_API UPCGPointData* CreatePointData();
-	/** Creates a PointData with a single point at the provided location */
-	PCG_API UPCGPointData* CreatePointData(const FVector& InLocation);
+	template <typename PointDataType = UPCGPointData, std::enable_if_t<std::is_base_of_v<UPCGBasePointData, PointDataType>, int> = 0>
+	PointDataType* CreateEmptyPointData()
+	{
+		return NewObject<PointDataType>();
+	}
+
+	template <typename PointDataType = UPCGPointData, std::enable_if_t<std::is_base_of_v<UPCGBasePointData, PointDataType>, int> = 0>
+	PointDataType* CreatePointData()
+	{
+		PointDataType* SinglePointData = CreateEmptyPointData<PointDataType>();
+
+		check(SinglePointData);
+
+		SinglePointData->SetNumPoints(1);
+
+		return SinglePointData;
+	}
+
+	template <typename PointDataType = UPCGPointData, std::enable_if_t<std::is_base_of_v<UPCGBasePointData, PointDataType>, int> = 0>
+	PointDataType* CreatePointData(const FVector& InLocation)
+	{
+		PointDataType* SinglePointData = CreatePointData<PointDataType>();
+
+		check(SinglePointData);
+		check(SinglePointData->GetNumPoints() == 1);
 
-	/** Creates a PointData with PointCount many points, and randomizes the Transform, Color, and Density */
-	PCG_API UPCGPointData* CreateRandomPointData(int32 PointCount, int32 Seed, bool RandomDensity = false);
+		TPCGValueRange<FTransform> TransformRange = SinglePointData->GetTransformValueRange();
+		
+		TransformRange[0].SetLocation(InLocation);
 
+		return SinglePointData;
+	}
+
+	/** Creates a point data with PointCount many points, and randomizes the Transform, Color, and Density */
+	template <typename PointDataType = UPCGPointData, std::enable_if_t<std::is_base_of_v<UPCGBasePointData, PointDataType>, int> = 0>
+	PointDataType* CreateRandomPointData(int32 PointCount, int32 Seed, bool bRandomDensity = false)
+	{
+		TObjectPtr<PointDataType> PointData = PCGTestsCommon::CreateEmptyPointData<PointDataType>();
+		
+		PointData->SetNumPoints(PointCount);
+		
+		TPCGValueRange<FTransform> TransformRange = PointData->GetTransformValueRange();
+		TPCGValueRange<int32> SeedRange = PointData->GetSeedValueRange();
+		TPCGValueRange<float> DensityRange = PointData->GetDensityValueRange();
+		TPCGValueRange<FVector4> ColorRange = PointData->GetColorValueRange();
+
+		FRandomStream RandomSource(Seed);
+		for (int I = 0; I < PointCount; ++I)
+		{
+			FQuat Rotation(FRotator(RandomSource.FRandRange(0.f, 360.f)).Quaternion());
+			FVector Scale(RandomSource.VRand());
+			FVector Location(RandomSource.VRand());
+
+			TransformRange[I] = FTransform(Rotation, Location, Scale);
+			SeedRange[I] = I;
+			DensityRange[I] = bRandomDensity ? RandomSource.GetFraction() : 1.0f;
+			ColorRange[I] = RandomSource.VRand();
+		}
+
+		return PointData;
+	}
+		
 	UPCGPolyLineData* CreatePolyLineData();
 	UPCGSurfaceData* CreateSurfaceData();
 	UPCGVolumeData* CreateVolumeData(const FBox& InBounds = FBox::BuildAABB(FVector::ZeroVector, FVector::OneVector * 100));
diff --git a/Engine/Plugins/PCG/Source/PCG/Public/Utils/PCGValueRange.h b/Engine/Plugins/PCG/Source/PCG/Public/Utils/PCGValueRange.h
new file mode 100644
--- /dev/null
+++ b/Engine/Plugins/PCG/Source/PCG/Public/Utils/PCGValueRange.h
@@ -0,0 +1,150 @@
+// Copyright Epic Games, Inc. All Rights Reserved.
+
+#pragma once
+
+#include "CoreTypes.h"
+#include "Containers/ContainersFwd.h"
+#include "Containers/StridedView.h"
+
+#include "HAL/PlatformString.h" // for INT64_FMT
+
+/*
+* TPCGValueRange is used to represent a list of InNumElements values but those
+* values are provided by an underlying TStridedView which can have a different number of values.
+*
+* Most common use case is where we have N Points, a TPCGValueRange might be used to iterate through all the Point Colors.
+* 
+* For UPCGPointData, all Points have a memory allocation representing those Colors so the PCGValueRange number of elements will match the underlying TStridedView number of elements
+* 
+* but if we have another type of Point data which doesn't allocate Colors by default and just uses a single value, We can still use a TPCGValueRange with the same number of elements with
+* an underlying TStridedView with a single element.
+* 
+* This allows us to write processing functions that operate on ranges without knowledge of the underlying memory layout
+* 
+* The range index will be modulated against the underlying TStridedView num elements.
+* 
+* ex: 
+* 
+* Range with 5 Elements & StridedView of 1 Element
+* Range Indices [0, 1, 2, 3 ,4] StridedView Indices [0, 0, 0, 0, 0]
+* 
+* Range with 5 Elements & StridedView of 3 Element(s)
+* Range Indices [0, 1, 2, 3, 4] StridedView Indices [0, 1, 2, 0, 1]
+* 
+* Range with 5 Elements & StridedView of 5 Elements(s) or more
+* Range Indices [0, 1, 2, 3, 4] StridedView Indices [0, 1, 2, 3, 4]
+*/
+template<typename ElementType, typename ViewType = TStridedView<ElementType>>
+class TPCGValueRange
+{
+public:
+	using SizeType = ViewType::SizeType;
+
+	TPCGValueRange()
+		: ElementView(ViewType()), NumElements(0)
+	{
+	}
+
+	TPCGValueRange(ViewType InElementView)
+		: ElementView(InElementView), NumElements(InElementView.Num())
+	{
+	}
+
+	TPCGValueRange(ViewType InElementView, SizeType InNumElements)
+		: ElementView(InElementView), NumElements(InNumElements)
+	{
+	}
+		
+	[[nodiscard]] FORCEINLINE TOptional<ElementType> GetSingleValue() const
+	{
+		return ViewNum() == 1 ? TOptional<ElementType>(GetElement(0)) : TOptional<ElementType>();
+	}
+
+	[[nodiscard]] FORCEINLINE bool IsValidIndex(SizeType Index) const
+	{
+		return (Index >= 0) && (Index < NumElements);
+	}
+
+	[[nodiscard]] FORCEINLINE bool IsEmpty() const
+	{
+		return NumElements == 0;
+	}
+
+	[[nodiscard]] FORCEINLINE SizeType Num() const
+	{
+		return NumElements;
+	}
+
+	[[nodiscard]] FORCEINLINE SizeType ViewNum() const
+	{
+		return ElementView.Num();
+	}
+
+	[[nodiscard]] FORCEINLINE ElementType& operator[](SizeType Index) const
+	{
+		return GetElement(Index);
+	}
+
+	FORCEINLINE void Set(const ElementType& Value)
+	{
+		// Set all values in underlying view to same value
+		for (SizeType i = 0; i < ElementView.Num(); ++i)
+		{
+			ElementView[i] = Value;
+		}
+	}
+
+	struct FIterator
+	{
+		const TPCGValueRange* Owner;
+		SizeType Index;
+
+		FORCEINLINE FIterator& operator++()
+		{
+			++Index;
+			return *this;
+		}
+
+		FORCEINLINE ElementType& operator*()
+		{
+			return Owner->GetElementUnsafe(Index);
+		}
+
+		[[nodiscard]] FORCEINLINE bool operator==(const FIterator& Other) const
+		{
+			return Owner == Other.Owner
+				&& Index == Other.Index;
+		}
+
+		[[nodiscard]] FORCEINLINE bool operator!=(const FIterator& Other) const
+		{
+			return !(*this == Other);
+		}
+	};
+
+	[[nodiscard]] FORCEINLINE FIterator begin() const { return FIterator{ this, 0 }; }
+	[[nodiscard]] FORCEINLINE FIterator end() const { return FIterator{ this, Num() }; }
+
+private:
+	FORCEINLINE void RangeCheck(ViewType::SizeType Index) const
+	{
+		checkf((Index >= 0) & (Index < NumElements), TEXT("Array index out of bounds: %" INT64_FMT " from an array of size %" INT64_FMT), int64(Index), int64(NumElements))
+	}
+	
+	[[nodiscard]] FORCEINLINE ElementType& GetElementUnsafe(SizeType Index) const
+	{
+		return ElementView.GetUnsafe(Index % ElementView.Num());
+	}
+
+	[[nodiscard]] FORCEINLINE ElementType& GetElement(SizeType Index) const
+	{
+		RangeCheck(Index);
+		return ElementView.GetUnsafe(Index % ElementView.Num());
+	}
+
+private:
+	ViewType ElementView;
+	SizeType NumElements;
+};
+
+template<typename ElementType> using TConstPCGValueRange = TPCGValueRange<const ElementType, TConstStridedView<ElementType>>;
\ No newline at end of file
diff --git a/Engine/Plugins/PCG/Source/PCGEditor/Private/PCGLevelToAsset.cpp b/Engine/Plugins/PCG/Source/PCGEditor/Private/PCGLevelToAsset.cpp
--- a/Engine/Plugins/PCG/Source/PCGEditor/Private/PCGLevelToAsset.cpp
+++ b/Engine/Plugins/PCG/Source/PCGEditor/Private/PCGLevelToAsset.cpp
@@ -6,7 +6,9 @@
 #include "PCGEditorUtils.h"
 #include "PCGAssetExporterUtils.h"
 
+#include "Data/PCGBasePointData.h"
 #include "Data/PCGPointData.h"
+#include "Data/PCGPointArrayData.h"
 #include "Helpers/PCGActorHelpers.h"
 #include "Helpers/PCGHelpers.h"
 #include "Helpers/PCGTagHelpers.h"
@@ -145,7 +147,7 @@ bool UPCGLevelToAsset::BP_ExportWorld_Implementation(UWorld* World, const FStrin
 	// Create Root Data
 	UPCGPointData* RootPointData = NewObject<UPCGPointData>(Asset);
 	UPCGMetadata* RootMetadata = RootPointData->MutableMetadata();
-	TArray<FPCGPoint>& Roots = RootPointData->GetMutablePoints();
+
 	RootMetadata->CreateAttribute<FString>(TEXT("Name"), World->GetName(), /*bAllowsInterpolation=*/false, /*bOverrideParent=*/true);
 	RootMetadata->CreateAttribute<FSoftObjectPath>(TEXT("Source"), FSoftObjectPath(World), /*bAllowsInterpolation=*/false, /*bOverrideParent=*/true);
 
@@ -159,7 +161,6 @@ bool UPCGLevelToAsset::BP_ExportWorld_Implementation(UWorld* World, const FStrin
 	// Create PCGData
 	UPCGPointData* PointData = NewObject<UPCGPointData>(Asset);
 	UPCGMetadata* PointMetadata = PointData->MutableMetadata();
-	TArray<FPCGPoint>& Points = PointData->GetMutablePoints();
 
 	// Add to data collection
 	{
@@ -203,15 +204,19 @@ bool UPCGLevelToAsset::BP_ExportWorld_Implementation(UWorld* World, const FStrin
 
 	// Hierarchy root point
 	{
-		FPCGPoint& RootPoint = Points.Emplace_GetRef(FTransform::Identity, /*Density=*/1.0f, /*Seed=*/0);
-		RootPoint.Steepness = 1.0f;
-		RootPoint.BoundsMin = RootPoint.BoundsMax = FVector::Zero();
-
-		PointMetadata->InitializeOnSet(RootPoint.MetadataEntry);
-		ActorIndexAttribute->SetValue(RootPoint.MetadataEntry, 0);
+		PointData->SetNumPoints(1);
+		PointData->SetTransform(FTransform::Identity);
+		PointData->SetDensity(1.0f);
+		PointData->SetBoundsMin(FVector::Zero());
+		PointData->SetBoundsMax(FVector::Zero());
+		PointData->SetSteepness(1.0f);
+		PointData->AllocateProperties(EPCGPointNativeProperties::Transform | EPCGPointNativeProperties::MetadataEntry | EPCGPointNativeProperties::BoundsMin | EPCGPointNativeProperties::BoundsMax | EPCGPointNativeProperties::Seed);
+
+		TPCGValueRange<int64> MetadataEntryRange = PointData->GetMetadataEntryValueRange();
+		PointMetadata->InitializeOnSet(MetadataEntryRange[0]);
+		ActorIndexAttribute->SetValue(MetadataEntryRange[0], 0);
 	}
 
-
 	// Make sure all actors are loaded
 	TMap<UWorldPartition*, TArray<FWorldPartitionReference>> ActorReferencesPerWorldPartition;
 
@@ -336,25 +341,32 @@ bool UPCGLevelToAsset::BP_ExportWorld_Implementation(UWorld* World, const FStrin
 			ParentActor = ParentActor->GetAttachParentActor();
 		}
 
-		auto MakePoint = [&](const FTransform& Transform, const FSoftObjectPath& MeshPath, const FBox& MeshBounds, const TArray<UMaterialInterface*>& MeshMaterials)
+		auto MakePoint = [&](int32 PointIndex, const FTransform& Transform, const FSoftObjectPath& MeshPath, const FBox& MeshBounds, const TArray<UMaterialInterface*>& MeshMaterials)
 		{
-			FPCGPoint& Point = Points.Emplace_GetRef(Transform, /*Density=*/1.0f, PCGHelpers::ComputeSeedFromPosition(Transform.GetLocation()));
-			Point.BoundsMin = MeshBounds.Min;
-			Point.BoundsMax = MeshBounds.Max;
-			Point.Steepness = 1.0f;
-
-			PointMetadata->InitializeOnSet(Point.MetadataEntry);
-			MeshAttribute->SetValue(Point.MetadataEntry, MeshPath);
+			TPCGValueRange<FTransform> TransformRange = PointData->GetTransformValueRange();
+			TPCGValueRange<int32> SeedRange = PointData->GetSeedValueRange();
+			TPCGValueRange<FVector> BoundsMinRange = PointData->GetBoundsMinValueRange();
+			TPCGValueRange<FVector> BoundsMaxRange = PointData->GetBoundsMaxValueRange();
+			TPCGValueRange<int64> MetadataEntryRange = PointData->GetMetadataEntryValueRange();
+
+			TransformRange[PointIndex] = Transform;
+			SeedRange[PointIndex] = PCGHelpers::ComputeSeedFromPosition(Transform.GetLocation());
+			BoundsMinRange[PointIndex] = MeshBounds.Min;
+			BoundsMaxRange[PointIndex] = MeshBounds.Max;
+			
+			int64& MetadataEntry = MetadataEntryRange[PointIndex];
+			PointMetadata->InitializeOnSet(MetadataEntry);
+			MeshAttribute->SetValue(MetadataEntry, MeshPath);
 
 			if (!MeshMaterials.IsEmpty())
 			{
-				MaterialAttribute->SetValue(Point.MetadataEntry, FSoftObjectPath(MeshMaterials[0]));
+				MaterialAttribute->SetValue(MetadataEntry, FSoftObjectPath(MeshMaterials[0]));
 			}
 
-			ActorIndexAttribute->SetValue(Point.MetadataEntry, ActorIndex);
-			ParentIndexAttribute->SetValue(Point.MetadataEntry, ParentActorIndex);
-			RelativeTransformAttribute->SetValue(Point.MetadataEntry, RelativeTransform);
-			HierarchyDepthAttribute->SetValue(Point.MetadataEntry, HierarchyDepth);
+			ActorIndexAttribute->SetValue(MetadataEntry, ActorIndex);
+			ParentIndexAttribute->SetValue(MetadataEntry, ParentActorIndex);
+			RelativeTransformAttribute->SetValue(MetadataEntry, RelativeTransform);
+			HierarchyDepthAttribute->SetValue(MetadataEntry, HierarchyDepth);
 
 			// For all tags, if the tag is of format 'Name:Value' then create attribute Name and assign Value, otherwise
 			// create a boolean attribute with the name given by the sanitized tag string.
@@ -365,7 +377,7 @@ bool UPCGLevelToAsset::BP_ExportWorld_Implementation(UWorld* World, const FStrin
 
 				if (TagData.IsValid() && TagToAttributeMap.Contains(FName(TagData.GetOriginalAttribute())))
 				{
-					PCG::Private::SetAttributeFromTag(TagData, PointMetadata, Point.MetadataEntry, /*bCanCreateAttribute=*/false);
+					PCG::Private::SetAttributeFromTag(TagData, PointMetadata, MetadataEntry, /*bCanCreateAttribute=*/false);
 				}
 			}
 		};
@@ -383,23 +395,28 @@ bool UPCGLevelToAsset::BP_ExportWorld_Implementation(UWorld* World, const FStrin
 			const FBox MeshBounds = StaticMesh->GetBoundingBox();
 			TArray<UMaterialInterface*> Materials = SMC->GetMaterials();
 
+			const int FirstIndex = PointData->GetNumPoints();
+
 			// For all instances (or a single instance if this is a static mesh and not an ISM)
 			// if a static mesh -> use actor transform (which might be wrong?)
 			// if ISM -> get instance transform in world space
 			if (UInstancedStaticMeshComponent* ISMC = Cast<UInstancedStaticMeshComponent>(SMC))
 			{
 				const int InstanceCount = ISMC->GetNumInstances();
+				
+				PointData->SetNumPoints(PointData->GetNumPoints() + InstanceCount);
 				for (int I = 0; I < InstanceCount; ++I)
 				{
 					FTransform InstanceTransform;
 					ISMC->GetInstanceTransform(I, InstanceTransform, /*bWorldSpace=*/true);
 
-					MakePoint(InstanceTransform, MeshPath, MeshBounds, Materials);
+					MakePoint(FirstIndex+I, InstanceTransform, MeshPath, MeshBounds, Materials);
 				}
 			}
 			else
 			{
-				MakePoint(ActorTransform, MeshPath, MeshBounds, Materials);
+				PointData->SetNumPoints(PointData->GetNumPoints() + 1);
+				MakePoint(FirstIndex, ActorTransform, MeshPath, MeshBounds, Materials);
 			}
 		}
 
@@ -408,10 +425,13 @@ bool UPCGLevelToAsset::BP_ExportWorld_Implementation(UWorld* World, const FStrin
 
 	// Finally, create root point in the root data
 	{
-		FPCGPoint& RootPoint = Roots.Emplace_GetRef(FTransform::Identity, 1.0f, 0);
-		RootPoint.BoundsMin = AllActorBounds.Min;
-		RootPoint.BoundsMax = AllActorBounds.Max;
-		RootPoint.Steepness = 1.0f;
+		RootPointData->SetNumPoints(1);
+		RootPointData->SetTransform(FTransform::Identity);
+		RootPointData->SetDensity(1.0f);
+		RootPointData->SetSeed(0);
+		RootPointData->SetBoundsMin(AllActorBounds.Min);
+		RootPointData->SetBoundsMax(AllActorBounds.Max);
+		RootPointData->SetSteepness(1.0f);
 	}
 
 	return true;
