Fix Android build

#rnx

diff --git a/Engine/Plugins/PCG/Source/PCG/Public/Utils/PCGValueRange.h b/Engine/Plugins/PCG/Source/PCG/Public/Utils/PCGValueRange.h
--- a/Engine/Plugins/PCG/Source/PCG/Public/Utils/PCGValueRange.h
+++ b/Engine/Plugins/PCG/Source/PCG/Public/Utils/PCGValueRange.h
@@ -38,7 +38,7 @@ template<typename ElementType, typename ViewType = TStridedView<ElementType>>
 class TPCGValueRange
 {
 public:
-	using SizeType = ViewType::SizeType;
+	using SizeType = typename ViewType::SizeType;
 
 	TPCGValueRange()
 		: ElementView(ViewType()), NumElements(0)
@@ -126,7 +126,7 @@ public:
 	[[nodiscard]] FORCEINLINE FIterator end() const { return FIterator{ this, Num() }; }
 
 private:
-	FORCEINLINE void RangeCheck(ViewType::SizeType Index) const
+	FORCEINLINE void RangeCheck(SizeType Index) const
 	{
 		checkf((Index >= 0) & (Index < NumElements), TEXT("Array index out of bounds: %" INT64_FMT " from an array of size %" INT64_FMT), int64(Index), int64(NumElements))
 	}
