[PCG] Fix ComputeGraphElement hanging IComputeTaskWorker::SubmitWork fails
* For example, when a data provider proxy IsValid returns false

#rnx

diff --git a/Engine/Plugins/PCG/Source/PCG/Private/Compute/Elements/PCGComputeGraphElement.cpp b/Engine/Plugins/PCG/Source/PCG/Private/Compute/Elements/PCGComputeGraphElement.cpp
--- a/Engine/Plugins/PCG/Source/PCG/Private/Compute/Elements/PCGComputeGraphElement.cpp
+++ b/Engine/Plugins/PCG/Source/PCG/Private/Compute/Elements/PCGComputeGraphElement.cpp
@@ -153,6 +153,12 @@ bool FPCGComputeGraphElement::ExecuteInternal(FPCGContext* InContext) const
 	// 6. Keep waiting for execution to complete.
 	if (Context->bGraphEnqueued)
 	{
+		if (Context->bGraphSubmitFailed)
+		{
+			// Submit failed, no choice but to exit so that execution doesn't hang.
+			return true;
+		}
+
 		// Likely we need a frame to pass in order to make progress.
 		SleepUntilNextFrame(Context);
 		return false;
@@ -365,7 +371,19 @@ bool FPCGComputeGraphElement::ExecuteInternal(FPCGContext* InContext) const
 			InContext->SourceComponent->GetScene(),
 			ComputeTaskExecutionGroup::EndOfFrameUpdate,
 			InContext->SourceComponent->GetOwner()->GetFName(),
-			FSimpleDelegate());
+			/*InFallbackDelegate=*/FSimpleDelegate::CreateLambda([ContextHandle = Context->GetOrCreateHandle()]()
+			{
+				// This render thread delegate will be executed if SubmitWork fails at any stage.
+				FPCGComputeGraphContext* Context = FPCGContext::GetContextFromHandle<FPCGComputeGraphContext>(ContextHandle);
+				if (!Context)
+				{
+					return;
+				}
+
+				// Wake up so that we can terminate execution of the ComputeGraphElement.
+				Context->bIsPaused = false;
+				Context->bGraphSubmitFailed = true;
+			}));
 	}
 
 	if (Context->bGraphEnqueued)
diff --git a/Engine/Plugins/PCG/Source/PCG/Private/Compute/Elements/PCGComputeGraphElement.h b/Engine/Plugins/PCG/Source/PCG/Private/Compute/Elements/PCGComputeGraphElement.h
--- a/Engine/Plugins/PCG/Source/PCG/Private/Compute/Elements/PCGComputeGraphElement.h
+++ b/Engine/Plugins/PCG/Source/PCG/Private/Compute/Elements/PCGComputeGraphElement.h
@@ -31,6 +31,9 @@ public:
 	/** Graph enqueued (scheduled for execution by GPU). */
 	bool bGraphEnqueued = false;
 
+	/** Graph was enqueued but was invalid for some reason and the work could not be submitted. */
+	bool bGraphSubmitFailed = false;
+
 	/** Graph executed successfully. */
 	bool bExecutionSuccess = false;
 
