[PCG] Clean up use of CreateUAV in PCGTextureReadback
* Was setting array range as (0, 0) instead of (0, 1), which is potentially problematic on some platforms. Instead now we use a pattern that is more common in the rest of the engine.

#rnx

diff --git a/Engine/Plugins/PCG/Source/PCGCompute/Private/PCGTextureReadback.cpp b/Engine/Plugins/PCG/Source/PCGCompute/Private/PCGTextureReadback.cpp
--- a/Engine/Plugins/PCG/Source/PCGCompute/Private/PCGTextureReadback.cpp
+++ b/Engine/Plugins/PCG/Source/PCGCompute/Private/PCGTextureReadback.cpp
@@ -61,13 +61,7 @@ void FPCGTextureReadbackInterface::Dispatch_RenderThread(FRHICommandListImmediat
 	
 	// Create temporary output texture
 	FTextureRHIRef OutputTexture = RHICreateTexture(TargetTextureDesc);
-	PassParameters.OutputTexture = RHICmdList.CreateUnorderedAccessView(OutputTexture,
-		FRHIViewDesc::CreateTextureUAV()
-			.SetDimensionFromTexture(OutputTexture)
-			.SetMipLevel(0)
-			.SetFormat(TargetTextureDesc.Format)
-			.SetArrayRange(0, 0)
-		);
+	PassParameters.OutputTexture = RHICmdList.CreateUnorderedAccessView(OutputTexture, FRHIViewDesc::CreateTextureUAV().SetDimensionFromTexture(OutputTexture));
 
 	TShaderMapRef<FPCGTextureReadbackCS> ComputeShader(GetGlobalShaderMap(GMaxRHIFeatureLevel));
 	FComputeShaderUtils::Dispatch(RHICmdList, ComputeShader, PassParameters,
