[PCG] Refactor of Attribute Property Selector to make it more generic

* Deprecating `PointProperty` and replacing with `Property`.
* All properties are now stored in `PropertyName` (can't use `AttributeName` as it would break the CDO).
* Still work as is for point properties.
* Add deprecation for `PointProperty` with a custom version.
* Add concept of Layer with the "@" character, for multi level metadata (or even multi layer properties in the future).
* Replace the deprecated `PointProperty` usage with `Property`.
* Start replacing `PointOrParam` to `Any` for a few nodes that will support all type of data (in Attribute noise for now).

diff --git a/Engine/Plugins/PCG/Source/PCG/Private/Data/PCGBasePointData.cpp b/Engine/Plugins/PCG/Source/PCG/Private/Data/PCGBasePointData.cpp
--- a/Engine/Plugins/PCG/Source/PCG/Private/Data/PCGBasePointData.cpp
+++ b/Engine/Plugins/PCG/Source/PCG/Private/Data/PCGBasePointData.cpp
@@ -630,9 +630,9 @@ namespace PCGBasePointData
 	{
 		TUniquePtr<IPCGAttributeAccessor> Accessor;
 
-		if (InSelector.GetSelection() == EPCGAttributePropertySelection::PointProperty)
+		if (InSelector.GetSelection() == EPCGAttributePropertySelection::Property)
 		{
-			const FName PropertyName = InSelector.GetName();
+			const FName PropertyName = InSelector.GetPropertyName();
 
 			if (const int64 EnumValue = StaticEnum<EPCGPointNativeProperties>()->GetValueByName(PropertyName); EnumValue != INDEX_NONE)
 			{
diff --git a/Engine/Plugins/PCG/Source/PCG/Private/Data/PCGPointData.cpp b/Engine/Plugins/PCG/Source/PCG/Private/Data/PCGPointData.cpp
--- a/Engine/Plugins/PCG/Source/PCG/Private/Data/PCGPointData.cpp
+++ b/Engine/Plugins/PCG/Source/PCG/Private/Data/PCGPointData.cpp
@@ -169,7 +169,7 @@ TUniquePtr<IPCGAttributeAccessor> UPCGPointData::CreateStaticAccessor(const FPCG
 {
 	TUniquePtr<IPCGAttributeAccessor> Accessor;
 	
-	if (InSelector.GetSelection() == EPCGAttributePropertySelection::PointProperty)
+	if (InSelector.GetSelection() == EPCGAttributePropertySelection::Property)
 	{
 		const FName PropertyName = InSelector.GetName();
 		if (const FProperty* Property = FPCGPoint::StaticStruct()->FindPropertyByName(PropertyName))
diff --git a/Engine/Plugins/PCG/Source/PCG/Private/Elements/PCGAttributeNoise.cpp b/Engine/Plugins/PCG/Source/PCG/Private/Elements/PCGAttributeNoise.cpp
--- a/Engine/Plugins/PCG/Source/PCG/Private/Elements/PCGAttributeNoise.cpp
+++ b/Engine/Plugins/PCG/Source/PCG/Private/Elements/PCGAttributeNoise.cpp
@@ -194,7 +194,7 @@ FPCGElementPtr UPCGAttributeNoiseSettings::CreateElement() const
 TArray<FPCGPinProperties> UPCGAttributeNoiseSettings::InputPinProperties() const
 {
 	TArray<FPCGPinProperties> PinProperties;
-	FPCGPinProperties& InputPinProperty = PinProperties.Emplace_GetRef(PCGPinConstants::DefaultInputLabel, EPCGDataType::PointOrParam);
+	FPCGPinProperties& InputPinProperty = PinProperties.Emplace_GetRef(PCGPinConstants::DefaultInputLabel, EPCGDataType::Any);
 	InputPinProperty.SetRequiredPin();
 	return PinProperties;
 }
@@ -202,7 +202,7 @@ TArray<FPCGPinProperties> UPCGAttributeNoiseSettings::InputPinProperties() const
 TArray<FPCGPinProperties> UPCGAttributeNoiseSettings::OutputPinProperties() const
 {
 	TArray<FPCGPinProperties> PinProperties;
-	PinProperties.Emplace(PCGPinConstants::DefaultOutputLabel, EPCGDataType::PointOrParam);
+	PinProperties.Emplace(PCGPinConstants::DefaultOutputLabel, EPCGDataType::Any);
 
 	return PinProperties;
 }
@@ -217,16 +217,7 @@ EPCGDataType UPCGAttributeNoiseSettings::GetCurrentPinTypes(const UPCGPin* InPin
 
 	// Output pin narrows to union of inputs on first pin
 	const EPCGDataType InputTypeUnion = GetTypeUnionOfIncidentEdges(PCGPinConstants::DefaultInputLabel);
-
-	// Spatial is collapsed into points
-	if (InputTypeUnion != EPCGDataType::None && (InputTypeUnion & EPCGDataType::Spatial) == InputTypeUnion)
-	{
-		return EPCGDataType::Point;
-	}
-	else
-	{
-		return (InputTypeUnion != EPCGDataType::None) ? InputTypeUnion : (EPCGDataType::PointOrParam);
-	}
+	return (InputTypeUnion != EPCGDataType::None) ? InputTypeUnion : (EPCGDataType::Any);
 }
 
 bool FPCGAttributeNoiseElement::ExecuteInternal(FPCGContext* InContext) const
@@ -348,7 +339,7 @@ bool FPCGAttributeNoiseElement::ExecuteInternal(FPCGContext* InContext) const
 		const TArray<FPCGPoint>* InputPoints = InputPointData ? &InputPointData->GetPoints() : nullptr;
 
 		// Force clamp on Density
-		const bool bClampResult = Settings->bClampResult || (Context->OutputTarget.GetSelection() == EPCGAttributePropertySelection::PointProperty && Context->OutputTarget.GetPointProperty() == EPCGPointProperties::Density);
+		const bool bClampResult = Settings->bClampResult || (Context->OutputTarget.GetSelection() == EPCGAttributePropertySelection::Property && Context->OutputTarget.GetPointProperty() == EPCGPointProperties::Density);
 
 		const bool bDone = PCGMetadataAttribute::CallbackWithRightType(Context->InputAccessor->GetUnderlyingType(), [Settings, Context, Seed, bClampResult, InputPoints](auto&& Dummy) -> bool
 		{
diff --git a/Engine/Plugins/PCG/Source/PCG/Private/Elements/PCGPathfindingElement.cpp b/Engine/Plugins/PCG/Source/PCG/Private/Elements/PCGPathfindingElement.cpp
--- a/Engine/Plugins/PCG/Source/PCG/Private/Elements/PCGPathfindingElement.cpp
+++ b/Engine/Plugins/PCG/Source/PCG/Private/Elements/PCGPathfindingElement.cpp
@@ -223,7 +223,7 @@ bool FPCGPathfindingElement::PrepareDataInternal(FPCGContext* InContext) const
 			}
 
 			// For Points with $Position as the selection, nothing is needed. For all other conditions, set the position based on the selected attribute/property.
-			if (OutPointArray.IsEmpty() || (Selector.GetSelection() == EPCGAttributePropertySelection::PointProperty && Selector.GetPointProperty() != EPCGPointProperties::Position))
+			if (OutPointArray.IsEmpty() || (Selector.GetSelection() == EPCGAttributePropertySelection::Property && Selector.GetPointProperty() != EPCGPointProperties::Position))
 			{
 				return;
 			}
diff --git a/Engine/Plugins/PCG/Source/PCG/Private/Elements/PCGSelfPruning.cpp b/Engine/Plugins/PCG/Source/PCG/Private/Elements/PCGSelfPruning.cpp
--- a/Engine/Plugins/PCG/Source/PCG/Private/Elements/PCGSelfPruning.cpp
+++ b/Engine/Plugins/PCG/Source/PCG/Private/Elements/PCGSelfPruning.cpp
@@ -530,7 +530,7 @@ namespace PCGSelfPruningElement
 				TRACE_CPUPROFILER_EVENT_SCOPE(FPCGSelfPruningElement::Execute::SortingSeeds);
 				Algo::Sort(InState.SortedPointRefs, [&InState](const PCGPointOctree::FPointRef& A, const PCGPointOctree::FPointRef& B) { return PCGSelfPruningAlgorithms::RandomSort(InState.InputPointData, A, B); });
 			}
-			else if (ComparisonSource.GetSelection() == EPCGAttributePropertySelection::PointProperty && ComparisonSource.GetPointProperty() == EPCGPointProperties::Extents && ComparisonSource.GetExtraNames().IsEmpty())
+			else if (ComparisonSource.GetSelection() == EPCGAttributePropertySelection::Property && ComparisonSource.GetPointProperty() == EPCGPointProperties::Extents && ComparisonSource.GetExtraNames().IsEmpty())
 			{
 				TRACE_CPUPROFILER_EVENT_SCOPE(FPCGSelfPruningElement::Execute::SortingExtents);
 				PCGSelfPruningAlgorithms::SortExtents(InState.InputPointData, InState.SortedPointRefs, InParameters, SquaredRadiusEquality);
diff --git a/Engine/Plugins/PCG/Source/PCG/Private/Metadata/PCGAttributePropertySelector.cpp b/Engine/Plugins/PCG/Source/PCG/Private/Metadata/PCGAttributePropertySelector.cpp
--- a/Engine/Plugins/PCG/Source/PCG/Private/Metadata/PCGAttributePropertySelector.cpp
+++ b/Engine/Plugins/PCG/Source/PCG/Private/Metadata/PCGAttributePropertySelector.cpp
@@ -14,13 +14,50 @@ namespace PCGAttributePropertySelectorConstants
 {
 	static const TCHAR* PropertyPrefix = TEXT("$");
 	static const TCHAR* ExtraSeparator = TEXT(".");
+	static const TCHAR* LayerPrefix = TEXT("@");
 	static const TCHAR PropertyPrefixChar = PropertyPrefix[0];
+	static const TCHAR LayerPrefixChar = LayerPrefix[0];
 	static const TCHAR ExtraSeparatorChar = ExtraSeparator[0];
 
 	static const FString ExportTextLeftSentinel = TEXT("PCGBegin(");
 	static const FString ExportTextRightSentinel = TEXT(")PCGEnd");
 }
 
+namespace PCGAttributePropertySelector
+{
+	bool IsReservedAttributeName(const FName& InName)
+	{
+		return InName == PCGMetadataAttributeConstants::LastAttributeName ||
+			InName == PCGMetadataAttributeConstants::LastCreatedAttributeName ||
+			InName == PCGMetadataAttributeConstants::SourceAttributeName ||
+			InName == PCGMetadataAttributeConstants::SourceNameAttributeName;
+	}
+}
+
+bool FPCGAttributePropertySelector::Serialize(FArchive& Ar)
+{
+	Ar.UsingCustomVersion(FPCGCustomVersion::GUID);
+
+	// Don't actually serialize, just write the custom version for PostSerialize
+	return false;
+}
+
+void FPCGAttributePropertySelector::PostSerialize(const FArchive& Ar)
+{
+#if WITH_EDITOR
+	if (Ar.IsLoading() && Ar.IsPersistent() && !Ar.HasAnyPortFlags(PPF_Duplicate | PPF_DuplicateForPIE))
+	{
+PRAGMA_DISABLE_DEPRECATION_WARNINGS
+		if (Ar.CustomVer(FPCGCustomVersion::GUID) < FPCGCustomVersion::AttributePropertySelectorDeprecatePointProperties
+			&& (Selection == EPCGAttributePropertySelection::PointProperty || Selection == EPCGAttributePropertySelection::Property))
+		{
+			SetPointProperty(PointProperty_DEPRECATED, /*bResetExtraNames=*/false);
+		}
+PRAGMA_ENABLE_DEPRECATION_WARNINGS
+	}
+#endif // WITH_EDITOR
+}
+
 bool FPCGAttributePropertySelector::ExportTextItem(FString& ValueStr, FPCGAttributePropertySelector const& DefaultValue, UObject* Parent, int32 PortFlags, UObject* ExportRootScope) const
 {
 	// String guarded by sentinels, don't use `"` because it can be used in the selector.
@@ -62,20 +99,8 @@ bool FPCGAttributePropertySelector::ImportTextItem(const TCHAR*& Buffer, int32 P
 
 FName FPCGAttributePropertySelector::GetName() const
 {
-	switch (Selection)
+	switch (GetSelection())
 	{
-	case EPCGAttributePropertySelection::PointProperty:
-	{
-		if (const UEnum* EnumPtr = StaticEnum<EPCGPointProperties>())
-		{
-			// Need to use the string version and not the name version, because the name verison has "EPCGPointProperties::" as a prefix.
-			return FName(EnumPtr->GetNameStringByValue((int64)PointProperty));
-		}
-		else
-		{
-			return NAME_None;
-		}
-	}
 	case EPCGAttributePropertySelection::ExtraProperty:
 	{
 		if (const UEnum* EnumPtr = StaticEnum<EPCGExtraProperties>())
@@ -91,6 +116,10 @@ FName FPCGAttributePropertySelector::GetName() const
 	{
 		return GetAttributeName();
 	}
+	case EPCGAttributePropertySelection::Property:
+	{
+		return GetPropertyName();
+	}
 	default:
 		return NAME_None;
 	}
@@ -98,64 +127,100 @@ FName FPCGAttributePropertySelector::GetName() const
 
 bool FPCGAttributePropertySelector::SetAttributeName(FName InAttributeName, bool bResetExtraNames)
 {
-	const bool bHasExtraNames = !ExtraNames.IsEmpty();
+	bool bHasChanged = false;
 	if (bResetExtraNames)
 	{
-		ExtraNames.Empty();
+		bHasChanged |= ResetExtraNames();
 	}
 
-	if (Selection == EPCGAttributePropertySelection::Attribute && GetAttributeName() == InAttributeName)
-	{
-		// Nothing changed, except perhaps the extra names
-		return (bHasExtraNames && bResetExtraNames);
-	}
-	else
+	if (Selection != EPCGAttributePropertySelection::Attribute || GetAttributeName() != InAttributeName)
 	{
 		Selection = EPCGAttributePropertySelection::Attribute;
 		AttributeName = InAttributeName;
-		return true;
+		bHasChanged = true;
 	}
+
+	return bHasChanged;
 }
 
-bool FPCGAttributePropertySelector::SetPointProperty(EPCGPointProperties InPointProperty, bool bResetExtraNames)
+bool FPCGAttributePropertySelector::SetLayerName(FName InLayerName, bool bResetExtraNames)
 {
-	const bool bHasExtraNames = !ExtraNames.IsEmpty();
+	bool bHasChanged = false;
 	if (bResetExtraNames)
 	{
-		ExtraNames.Empty();
+		bHasChanged |= ResetExtraNames();
 	}
 
-	if (Selection == EPCGAttributePropertySelection::PointProperty && InPointProperty == PointProperty)
+	if (GetLayerName() != InLayerName)
 	{
-		// Nothing changed, except perhaps the extra names
-		return (bHasExtraNames && bResetExtraNames);
+		LayerName = InLayerName;
+		bHasChanged = true;
 	}
-	else
+
+	return bHasChanged;
+}
+
+bool FPCGAttributePropertySelector::SetPropertyName(FName InPropertyName, bool bResetExtraNames)
+{
+	bool bHasChanged = false;
+	if (bResetExtraNames)
 	{
-		Selection = EPCGAttributePropertySelection::PointProperty;
-		PointProperty = InPointProperty;
-		return true;
+		bHasChanged |= ResetExtraNames();
+	}
+
+	if (Selection != EPCGAttributePropertySelection::Property || GetPropertyName() != InPropertyName)
+	{
+		Selection = EPCGAttributePropertySelection::Property;
+		PropertyName = InPropertyName;
+		bHasChanged = true;
 	}
+
+	return bHasChanged;
+}
+
+bool FPCGAttributePropertySelector::SetPointProperty(EPCGPointProperties InPointProperty, bool bResetExtraNames)
+{
+	UEnum* EnumPtr = StaticEnum<EPCGPointProperties>();
+	check(EnumPtr);
+	
+	const bool bHasChanged = SetPropertyName(*EnumPtr->GetNameStringByValue(static_cast<int64>(InPointProperty)), bResetExtraNames);
+	
+	return bHasChanged;
 }
 
 bool FPCGAttributePropertySelector::SetExtraProperty(EPCGExtraProperties InExtraProperty, bool bResetExtraNames)
 {
-	const bool bHasExtraNames = !ExtraNames.IsEmpty();
+	bool bHasChanged = false;
 	if (bResetExtraNames)
 	{
-		ExtraNames.Empty();
+		bHasChanged |= ResetExtraNames();
 	}
 
-	if (Selection == EPCGAttributePropertySelection::ExtraProperty && InExtraProperty == ExtraProperty)
+	if (Selection != EPCGAttributePropertySelection::ExtraProperty || InExtraProperty != ExtraProperty)
 	{
-		// Nothing changed, except perhaps the extra names
-		return (bHasExtraNames && bResetExtraNames);
+		Selection = EPCGAttributePropertySelection::ExtraProperty;
+		ExtraProperty = InExtraProperty;
+		bHasChanged = true;
+	}
+
+	return bHasChanged;
+}
+
+FString FPCGAttributePropertySelector::GetLayerString(bool bAddLeadingQualifier) const
+{
+	const FName Layer = GetLayerName();
+	if (Layer == NAME_None)
+	{
+		return FString{};
+	}
+	
+	if (bAddLeadingQualifier)
+	{
+		return FString(PCGAttributePropertySelectorConstants::LayerPrefix) + Layer.ToString();
 	}
 	else
 	{
-		Selection = EPCGAttributePropertySelection::ExtraProperty;
-		ExtraProperty = InExtraProperty;
-		return true;
+		return Layer.ToString();
 	}
 }
 
@@ -193,10 +258,18 @@ FString FPCGAttributePropertySelector::GetAttributePropertyAccessorsString(bool
 
 FString FPCGAttributePropertySelector::ToString() const
 {
+	const FString Layer = GetLayerString(/*bAddLeadingQualifier=*/true);
 	const FString Attribute = GetAttributePropertyString(/*bAddPropertyQualifier=*/true);
 	const FString Accessors = GetAttributePropertyAccessorsString(/*bAddLeadingSeparator*/true);
 
-	return Attribute + Accessors;
+	if (Layer.IsEmpty())
+	{
+		return Attribute + Accessors;
+	}
+	else
+	{
+		return Layer + PCGAttributePropertySelectorConstants::ExtraSeparator + Attribute + Accessors;
+	}
 }
 
 bool FPCGAttributePropertySelector::operator==(const FPCGAttributePropertySelector& Other) const
@@ -206,17 +279,17 @@ bool FPCGAttributePropertySelector::operator==(const FPCGAttributePropertySelect
 
 bool FPCGAttributePropertySelector::IsSame(const FPCGAttributePropertySelector& Other, bool bIncludeExtraNames) const
 {
-	if (Selection != Other.Selection || (bIncludeExtraNames && ExtraNames != Other.ExtraNames))
+	if (GetSelection() != Other.GetSelection() || LayerName != Other.LayerName || (bIncludeExtraNames && ExtraNames != Other.ExtraNames))
 	{
 		return false;
 	}
 
-	switch (Selection)
+	switch (GetSelection())
 	{
 	case EPCGAttributePropertySelection::Attribute:
 		return GetAttributeName() == Other.GetAttributeName();
-	case EPCGAttributePropertySelection::PointProperty:
-		return PointProperty == Other.PointProperty;
+	case EPCGAttributePropertySelection::Property:
+		return GetPropertyName() == Other.GetPropertyName();
 	case EPCGAttributePropertySelection::ExtraProperty:
 		return ExtraProperty == Other.ExtraProperty;
 	default:
@@ -227,18 +300,21 @@ bool FPCGAttributePropertySelector::IsSame(const FPCGAttributePropertySelector&
 void FPCGAttributePropertySelector::ImportFromOtherSelector(const FPCGAttributePropertySelector& InOther)
 {
 	Selection = InOther.GetSelection();
+	LayerName = InOther.GetLayerName();
 
-	switch (Selection)
+	switch (GetSelection())
 	{
 	case EPCGAttributePropertySelection::Attribute:
 		SetAttributeName(InOther.GetAttributeName());
 		break;
-	case EPCGAttributePropertySelection::PointProperty:
-		SetPointProperty(InOther.GetPointProperty());
+	case EPCGAttributePropertySelection::Property:
+		SetPropertyName(InOther.GetPropertyName());
 		break;
 	case EPCGAttributePropertySelection::ExtraProperty:
 		SetExtraProperty(InOther.GetExtraProperty());
 		break;
+	default:
+		break;
 	}
 
 	ExtraNames = InOther.GetExtraNames();
@@ -254,96 +330,148 @@ bool FPCGAttributePropertySelector::IsValid() const
 		return false;
 	}
 
-	return (Selection != EPCGAttributePropertySelection::Attribute) || 
-		ThisAttributeName == PCGMetadataAttributeConstants::LastAttributeName ||
-		ThisAttributeName == PCGMetadataAttributeConstants::LastCreatedAttributeName ||
-		ThisAttributeName == PCGMetadataAttributeConstants::SourceAttributeName ||
-		ThisAttributeName == PCGMetadataAttributeConstants::SourceNameAttributeName ||
+	return (Selection != EPCGAttributePropertySelection::Attribute) ||
+		PCGAttributePropertySelector::IsReservedAttributeName(ThisAttributeName) ||
 		FPCGMetadataAttributeBase::IsValidName(ThisAttributeName);
 }
 
-bool FPCGAttributePropertySelector::Update(const FString& NewValue)
+bool FPCGAttributePropertySelector::Reset()
 {
-	TArray<FString> NewValues;
-	if (NewValue.IsEmpty())
+	bool bHasChanged = false;
+	if (LayerName != NAME_None)
 	{
-		NewValues.Add(NewValue);
+		bHasChanged = true;
+		LayerName = NAME_None;
+	}
+
+	if (AttributeName != NAME_None)
+	{
+		bHasChanged = true;
+		AttributeName = NAME_None;
+	}
+	
+	if (PropertyName != NAME_None)
+	{
+		bHasChanged = true;
+		PropertyName = NAME_None;
+	}
+
+
+	bHasChanged |= ResetExtraNames();
+
+	return bHasChanged;
+}
+
+bool FPCGAttributePropertySelector::ResetExtraNames()
+{
+	if (!ExtraNames.IsEmpty())
+	{
+		ExtraNames.Empty();
+		return true;
 	}
 	else
 	{
-		NewValue.ParseIntoArray(NewValues, PCGAttributePropertySelectorConstants::ExtraSeparator, /*InCullEmpty=*/ false);
+		return false;
 	}
+}
 
-	const FString& NewName = NewValues[0];
+bool FPCGAttributePropertySelector::Update(const FString& NewValue)
+{
+	if (NewValue.IsEmpty())
+	{
+		const FPCGAttributePropertySelector Previous = *this;
+		Reset();
+		return *this == Previous;
+	}
+	
+	TArray<FString> NewValues;
+	NewValue.ParseIntoArray(NewValues, PCGAttributePropertySelectorConstants::ExtraSeparator, /*InCullEmpty=*/ false);
+
+	check(!NewValues.IsEmpty())
+
+	// TODO: If we ever have to support multiple layers, this has to change
+	const bool bHasLayerName = !NewValues[0].IsEmpty() && NewValues[0][0] == PCGAttributePropertySelectorConstants::LayerPrefixChar && !PCGAttributePropertySelector::IsReservedAttributeName(*NewValues[0]);
+	const int32 PropertyIndex = bHasLayerName ? 1 : 0;
+	const bool bHasProperty = NewValues.Num() > PropertyIndex && !NewValues[PropertyIndex].IsEmpty() && NewValues[PropertyIndex][0] == PCGAttributePropertySelectorConstants::PropertyPrefixChar;
+
+	const int32 ExtraNameIndex = PropertyIndex + 1;
+	
 	TArray<FString> ExtraNamesTemp;
-	if (NewValues.Num() > 1)
+	if (NewValues.Num() > ExtraNameIndex)
 	{
-		ExtraNamesTemp.Append(&NewValues[1], NewValues.Num() - 1);
+		ExtraNamesTemp.Append(&NewValues[ExtraNameIndex], NewValues.Num() - ExtraNameIndex);
 	}
 
-	const bool bExtraChanged = ExtraNamesTemp != ExtraNames;
+	bool bHasChanged = ExtraNamesTemp != ExtraNames;
 	ExtraNames = MoveTemp(ExtraNamesTemp);
 
-	if (!NewName.IsEmpty() && NewName[0] == PCGAttributePropertySelectorConstants::PropertyPrefixChar)
+	if (bHasLayerName)
 	{
-		// Remove the first character of the name, where the first character is $.
-		const FString NewNameWithoutPrefix = NewName.RightChop(1);
-
-		if (const UEnum* EnumPtr = StaticEnum<EPCGPointProperties>())
+		bHasChanged |= SetLayerName(*NewValues[0].RightChop(1), /*bResetExtraNames=*/ false);
+	}
+	
+	if (bHasProperty)
+	{
+		const FString NewNameWithoutPrefix = NewValues[PropertyIndex].RightChop(1);
+		const UEnum* EnumPtr = StaticEnum<EPCGExtraProperties>();
+		check(EnumPtr);
+		
+		const int32 Index = EnumPtr->GetIndexByNameString(NewNameWithoutPrefix);
+		if (Index != INDEX_NONE)
 		{
-			int32 Index = EnumPtr->GetIndexByNameString(NewNameWithoutPrefix);
-			if (Index != INDEX_NONE)
-			{
-				return SetPointProperty(static_cast<EPCGPointProperties>(EnumPtr->GetValueByIndex(Index)), /*bResetExtraNames=*/ false) || bExtraChanged;
-			}
+			return SetExtraProperty(static_cast<EPCGExtraProperties>(EnumPtr->GetValueByIndex(Index)), /*bResetExtraNames=*/ false) || bHasChanged;
 		}
-
-		if (const UEnum* EnumPtr = StaticEnum<EPCGExtraProperties>())
+		else
 		{
-			int32 Index = EnumPtr->GetIndexByNameString(NewNameWithoutPrefix);
-			if (Index != INDEX_NONE)
-			{
-				return SetExtraProperty(static_cast<EPCGExtraProperties>(EnumPtr->GetValueByIndex(Index)), /*bResetExtraNames=*/ false) || bExtraChanged;
-			}
+			return SetPropertyName(*NewNameWithoutPrefix, /*bResetExtraNames=*/ false) || bHasChanged;
 		}
 	}
-
-	return SetAttributeName(NewName.IsEmpty() ? NAME_None : FName(NewName), /*bResetExtraNames=*/ false) || bExtraChanged;
+	else
+	{
+		return SetAttributeName(*NewValues[PropertyIndex], /*bResetExtraNames=*/ false) || bHasChanged;
+	}
 }
 
 void FPCGAttributePropertySelector::AddToCrc(FArchiveCrc32& Ar) const
 {
-	Ar << Selection;
+	EPCGAttributePropertySelection RealSelection = GetSelection();
+	Ar << RealSelection;
+	Ar << const_cast<FName&>(LayerName);
 	Ar << const_cast<TArray<FString>&>(ExtraNames);
 
-	switch (Selection)
+	switch (RealSelection)
 	{
 	case EPCGAttributePropertySelection::Attribute:
 		Ar << const_cast<FName&>(AttributeName);
 		break;
-	case EPCGAttributePropertySelection::PointProperty:
-		Ar << PointProperty;
+	case EPCGAttributePropertySelection::Property:
+		Ar << const_cast<FName&>(PropertyName);
 		break;
 	case EPCGAttributePropertySelection::ExtraProperty:
 		Ar << ExtraProperty;
 		break;
+	default:
+		break;
 	}
 }
 
 uint32 GetTypeHash(const FPCGAttributePropertySelector& Selector)
 {
-	uint32 Hash = GetTypeHash(Selector.Selection);
-	if (Selector.Selection == EPCGAttributePropertySelection::Attribute)
+	uint32 Hash = HashCombine(GetTypeHash(Selector.Selection), GetTypeHash(Selector.LayerName));
+	
+	switch (Selector.GetSelection())
 	{
+	case EPCGAttributePropertySelection::Attribute:
 		Hash = HashCombine(Hash, GetTypeHash(Selector.AttributeName));
-	}
-	else if (Selector.Selection == EPCGAttributePropertySelection::PointProperty)
-	{
-		Hash = HashCombine(Hash, GetTypeHash(Selector.PointProperty));
-	}
-	else
-	{
+		break;
+	case EPCGAttributePropertySelection::Property:
+		Hash = HashCombine(Hash, GetTypeHash(Selector.PropertyName));
+		break;
+	case EPCGAttributePropertySelection::ExtraProperty:
 		Hash = HashCombine(Hash, GetTypeHash(Selector.ExtraProperty));
+		break;
+	default:
+		break;
 	}
 
 	for (const FString& ExtraName : Selector.ExtraNames)
@@ -354,6 +482,22 @@ uint32 GetTypeHash(const FPCGAttributePropertySelector& Selector)
 	return Hash;
 }
 
+EPCGPointProperties FPCGAttributePropertySelector::GetPointProperty() const
+{
+	UEnum* EnumPtr = StaticEnum<EPCGPointProperties>();
+	check(EnumPtr);
+		
+	const int32 Index = EnumPtr->GetIndexByName(PropertyName);
+	if (Index != INDEX_NONE)
+	{
+		return static_cast<EPCGPointProperties>(EnumPtr->GetValueByIndex(Index));
+	}
+	else
+	{
+		return EPCGPointProperties::Invalid;
+	}
+}
+
 bool FPCGAttributePropertySelector::IsBasicAttribute() const
 {
 	return Selection == EPCGAttributePropertySelection::Attribute && ExtraNames.IsEmpty();
@@ -479,19 +623,29 @@ bool FPCGAttributePropertyOutputSelector::SerializeFromMismatchedTag(const struc
 
 ///////////////////////////////////////////////////////////////////////
 
-bool UPCGAttributePropertySelectorBlueprintHelpers::SetPointProperty(FPCGAttributePropertySelector& Selector, EPCGPointProperties InPointProperty)
+bool UPCGAttributePropertySelectorBlueprintHelpers::SetPointProperty(FPCGAttributePropertySelector& Selector, EPCGPointProperties InPointProperty, bool bResetExtraNames)
+{
+	return Selector.SetPointProperty(InPointProperty, bResetExtraNames);
+}
+
+bool UPCGAttributePropertySelectorBlueprintHelpers::SetAttributeName(FPCGAttributePropertySelector& Selector, FName InAttributeName, bool bResetExtraNames)
 {
-	return Selector.SetPointProperty(InPointProperty);
+	return Selector.SetAttributeName(InAttributeName, bResetExtraNames);
 }
 
-bool UPCGAttributePropertySelectorBlueprintHelpers::SetAttributeName(FPCGAttributePropertySelector& Selector, FName InAttributeName)
+bool UPCGAttributePropertySelectorBlueprintHelpers::SetPropertyName(FPCGAttributePropertySelector& Selector, FName InPropertyName, bool bResetExtraNames)
 {
-	return Selector.SetAttributeName(InAttributeName);
+	return Selector.SetPropertyName(InPropertyName, bResetExtraNames);
 }
 
-bool UPCGAttributePropertySelectorBlueprintHelpers::SetExtraProperty(FPCGAttributePropertySelector& Selector, EPCGExtraProperties InExtraProperty)
+bool UPCGAttributePropertySelectorBlueprintHelpers::SetLayerName(FPCGAttributePropertySelector& Selector, FName InLayerName, bool bResetExtraNames)
 {
-	return Selector.SetExtraProperty(InExtraProperty);
+	return Selector.SetLayerName(InLayerName, bResetExtraNames);
+}
+
+bool UPCGAttributePropertySelectorBlueprintHelpers::SetExtraProperty(FPCGAttributePropertySelector& Selector, EPCGExtraProperties InExtraProperty, bool bResetExtraNames)
+{
+	return Selector.SetExtraProperty(InExtraProperty, bResetExtraNames);
 }
 
 EPCGAttributePropertySelection UPCGAttributePropertySelectorBlueprintHelpers::GetSelection(const FPCGAttributePropertySelector& Selector)
@@ -519,6 +673,16 @@ const TArray<FString>& UPCGAttributePropertySelectorBlueprintHelpers::GetExtraNa
 	return Selector.GetExtraNames();
 }
 
+FName UPCGAttributePropertySelectorBlueprintHelpers::GetPropertyName(const FPCGAttributePropertySelector& Selector)
+{
+	return Selector.GetPropertyName();
+}
+
+FName UPCGAttributePropertySelectorBlueprintHelpers::GetLayerName(const FPCGAttributePropertySelector& Selector)
+{
+	return Selector.GetLayerName();
+}
+
 FName UPCGAttributePropertySelectorBlueprintHelpers::GetName(const FPCGAttributePropertySelector& Selector)
 {
 	return Selector.GetName();
@@ -529,7 +693,7 @@ FPCGAttributePropertyInputSelector UPCGAttributePropertySelectorBlueprintHelpers
 	return Selector.CopyAndFixLast(InData);
 }
 
-FPCGAttributePropertyOutputSelector UPCGAttributePropertySelectorBlueprintHelpers::CopyAndFixSource(const FPCGAttributePropertyOutputSelector& OutputSelector, const FPCGAttributePropertyInputSelector& InputSelector)
+FPCGAttributePropertyOutputSelector UPCGAttributePropertySelectorBlueprintHelpers::CopyAndFixSource(const FPCGAttributePropertyOutputSelector& OutputSelector, const FPCGAttributePropertyInputSelector& InputSelector, const UPCGData* InOptionalData)
 {
-	return OutputSelector.CopyAndFixSource(&InputSelector);
+	return OutputSelector.CopyAndFixSource(&InputSelector, InOptionalData);
 }
\ No newline at end of file
diff --git a/Engine/Plugins/PCG/Source/PCG/Private/Tests/Accessors/PCGAttributePropertySelectorTests.cpp b/Engine/Plugins/PCG/Source/PCG/Private/Tests/Accessors/PCGAttributePropertySelectorTests.cpp
new file mode 100644
--- /dev/null
+++ b/Engine/Plugins/PCG/Source/PCG/Private/Tests/Accessors/PCGAttributePropertySelectorTests.cpp
@@ -0,0 +1,411 @@
+// Copyright Epic Games, Inc. All Rights Reserved.
+
+
+#include "PCGCustomVersion.h"
+#include "Tests/PCGTestsCommon.h"
+
+#include "Metadata/PCGAttributePropertySelector.h"
+
+#if WITH_EDITOR
+class FPCGAttributePropertySelectorTests : public FPCGTestBaseClass
+{
+	using FPCGTestBaseClass::FPCGTestBaseClass;
+
+protected:
+	bool FromStringComparison(const FString& InString, const FPCGAttributePropertySelector& InSelector)
+	{
+		FPCGAttributePropertySelector Selector = FPCGAttributePropertySelector::CreateSelectorFromString(InString);
+		UTEST_EQUAL("Selectors are equal", Selector, InSelector);
+		return true;
+	}
+
+	bool ImportExport(const FString& ExpectedExport, const FPCGAttributePropertySelector& InSelector)
+	{
+		FString Export;
+		// Only first parameter is useful.
+		InSelector.ExportTextItem(Export, FPCGAttributePropertySelector{}, nullptr, 0, nullptr);
+		UTEST_EQUAL("Export texts are equal", Export, ExpectedExport);
+		
+		FPCGAttributePropertySelector ImportedSelector;
+		const TCHAR* Buffer = *Export;
+		ImportedSelector.ImportTextItem(Buffer, 0, nullptr, nullptr);
+		UTEST_EQUAL("Selectors are equal", ImportedSelector, InSelector);
+		const bool bFullBufferConsumed = reinterpret_cast<const void*>(Buffer) == reinterpret_cast<const void*>(*Export + Export.Len());
+		UTEST_TRUE("Full buffer consumed", bFullBufferConsumed);
+		
+		return true;
+	}
+
+	/**
+	* Test that the selector serialized with an old version, with its selection set to PointProperty
+	* are deserializing correctly (after the PostSerialize)
+	*/
+	bool ValidateDeprecatedSelector(const FPCGAttributePropertySelector& ExpectedSelector, EPCGPointProperties PointProperty, TArray<FString> ExtraNames)
+	{
+		// To be removed when deprecated point property no longer exists.
+		
+		FPCGAttributePropertySelector DeprecatedSelector{};
+		DeprecatedSelector.GetExtraNamesMutable() = std::move(ExtraNames);
+		
+PRAGMA_DISABLE_DEPRECATION_WARNINGS
+		DeprecatedSelector.Selection = EPCGAttributePropertySelection::PointProperty;
+		DeprecatedSelector.PointProperty_DEPRECATED = PointProperty;
+PRAGMA_ENABLE_DEPRECATION_WARNINGS
+
+		FArchive DummyAr;
+		DummyAr.UsingCustomVersion(FPCGCustomVersion::GUID);
+		DummyAr.SetCustomVersion(FPCGCustomVersion::GUID, static_cast<int32>(FPCGCustomVersion::AttributePropertySelectorDeprecatePointProperties) - 1, FName("Dummy"));
+		DummyAr.SetIsLoading(true);
+		DummyAr.SetIsPersistent(true);
+		DeprecatedSelector.PostSerialize(DummyAr);
+		
+		UTEST_EQUAL("Selectors are equal", DeprecatedSelector, ExpectedSelector);
+		return true;
+	}
+
+	/**
+	* Test that the selector serialized with an old version, and set in the ctor as a point property (like in some PCGSettings constructors)
+	* are deserializing correctly (after the PostSerialize)
+	*/
+	bool ValidateDeprecatedSelectorCDO(const FPCGAttributePropertySelector& ExpectedSelector, EPCGPointProperties PointProperty, TArray<FString> ExtraNames)
+	{
+		// To be removed when deprecated point property no longer exists.
+		// Version with CDO where the default object call SetPointProperty.
+		FPCGAttributePropertySelector DeprecatedSelector = FPCGAttributePropertySelector::CreatePointPropertySelector(EPCGPointProperties::Density, NAME_None, ExtraNames);
+		
+PRAGMA_DISABLE_DEPRECATION_WARNINGS
+		DeprecatedSelector.PointProperty_DEPRECATED = PointProperty;
+PRAGMA_ENABLE_DEPRECATION_WARNINGS
+
+		FArchive DummyAr;
+		DummyAr.UsingCustomVersion(FPCGCustomVersion::GUID);
+		DummyAr.SetCustomVersion(FPCGCustomVersion::GUID, static_cast<int32>(FPCGCustomVersion::AttributePropertySelectorDeprecatePointProperties) - 1, FName("Dummy"));
+		DummyAr.SetIsLoading(true);
+		DummyAr.SetIsPersistent(true);
+		DeprecatedSelector.PostSerialize(DummyAr);
+		
+		UTEST_EQUAL("Selectors are equal", DeprecatedSelector, ExpectedSelector);
+		return true;
+	}
+
+	/**
+	* Test that the new selector serialized with the new version, and set in the ctor as a point property (like in some PCGSettings constructors)
+	* are deserializing correctly (after the PostSerialize)
+	*/
+	bool ValidateSelectorCDOWithNewerVersion(const FPCGAttributePropertySelector& ExpectedSelector, EPCGPointProperties PointProperty, TArray<FString> ExtraNames)
+	{
+		// To be removed when deprecated point property no longer exists.
+		// Version with CDO where the default object call SetPointProperty.
+		FPCGAttributePropertySelector DeprecatedSelector = FPCGAttributePropertySelector::CreatePointPropertySelector(PointProperty, NAME_None, ExtraNames);
+
+		FArchive DummyAr;
+		DummyAr.UsingCustomVersion(FPCGCustomVersion::GUID);
+		DummyAr.SetCustomVersion(FPCGCustomVersion::GUID, FPCGCustomVersion::AttributePropertySelectorDeprecatePointProperties, FName("Dummy"));
+		DummyAr.SetIsLoading(true);
+		DummyAr.SetIsPersistent(true);
+		DeprecatedSelector.PostSerialize(DummyAr);
+		
+		UTEST_EQUAL("Selectors are equal", DeprecatedSelector, ExpectedSelector);
+		return true;
+	}
+};
+
+IMPLEMENT_CUSTOM_SIMPLE_AUTOMATION_TEST(FPCGAttributePropertySelector_FromString_Attribute, FPCGAttributePropertySelectorTests, "Plugins.PCG.AttributePropertySelector.FromString.Attribute", PCGTestsCommon::TestFlags)
+IMPLEMENT_CUSTOM_SIMPLE_AUTOMATION_TEST(FPCGAttributePropertySelector_FromString_Property, FPCGAttributePropertySelectorTests, "Plugins.PCG.AttributePropertySelector.FromString.Property", PCGTestsCommon::TestFlags)
+IMPLEMENT_CUSTOM_SIMPLE_AUTOMATION_TEST(FPCGAttributePropertySelector_FromString_PointProperty, FPCGAttributePropertySelectorTests, "Plugins.PCG.AttributePropertySelector.FromString.PointProperty", PCGTestsCommon::TestFlags)
+IMPLEMENT_CUSTOM_SIMPLE_AUTOMATION_TEST(FPCGAttributePropertySelector_FromString_ExtraProperty, FPCGAttributePropertySelectorTests, "Plugins.PCG.AttributePropertySelector.FromString.ExtraProperty", PCGTestsCommon::TestFlags)
+IMPLEMENT_CUSTOM_SIMPLE_AUTOMATION_TEST(FPCGAttributePropertySelector_FromString_AtLast, FPCGAttributePropertySelectorTests, "Plugins.PCG.AttributePropertySelector.FromString.@Last", PCGTestsCommon::TestFlags)
+IMPLEMENT_CUSTOM_SIMPLE_AUTOMATION_TEST(FPCGAttributePropertySelector_FromString_AtSource, FPCGAttributePropertySelectorTests, "Plugins.PCG.AttributePropertySelector.FromString.@Source", PCGTestsCommon::TestFlags)
+IMPLEMENT_CUSTOM_SIMPLE_AUTOMATION_TEST(FPCGAttributePropertySelector_FromString_WithLayer_Attribute, FPCGAttributePropertySelectorTests, "Plugins.PCG.AttributePropertySelector.FromString.WithLayer.Attribute", PCGTestsCommon::TestFlags)
+IMPLEMENT_CUSTOM_SIMPLE_AUTOMATION_TEST(FPCGAttributePropertySelector_FromString_WithLayer_Property, FPCGAttributePropertySelectorTests, "Plugins.PCG.AttributePropertySelector.FromString.WithLayer.Property", PCGTestsCommon::TestFlags)
+IMPLEMENT_CUSTOM_SIMPLE_AUTOMATION_TEST(FPCGAttributePropertySelector_FromString_WithLayer_PointProperty, FPCGAttributePropertySelectorTests, "Plugins.PCG.AttributePropertySelector.FromString.WithLayer.PointProperty", PCGTestsCommon::TestFlags)
+IMPLEMENT_CUSTOM_SIMPLE_AUTOMATION_TEST(FPCGAttributePropertySelector_FromString_WithLayer_ExtraProperty, FPCGAttributePropertySelectorTests, "Plugins.PCG.AttributePropertySelector.FromString.WithLayer.ExtraProperty", PCGTestsCommon::TestFlags)
+IMPLEMENT_CUSTOM_SIMPLE_AUTOMATION_TEST(FPCGAttributePropertySelector_FromString_WithLayer_AtLast, FPCGAttributePropertySelectorTests, "Plugins.PCG.AttributePropertySelector.FromString.WithLayer.@Last", PCGTestsCommon::TestFlags)
+IMPLEMENT_CUSTOM_SIMPLE_AUTOMATION_TEST(FPCGAttributePropertySelector_FromString_WithLayer_AtSource, FPCGAttributePropertySelectorTests, "Plugins.PCG.AttributePropertySelector.FromString.WithLayer.@Source", PCGTestsCommon::TestFlags)
+IMPLEMENT_CUSTOM_SIMPLE_AUTOMATION_TEST(FPCGAttributePropertySelector_FromString_WithExtraNames_Attribute, FPCGAttributePropertySelectorTests, "Plugins.PCG.AttributePropertySelector.FromString.WithExtraNames.Attribute", PCGTestsCommon::TestFlags)
+IMPLEMENT_CUSTOM_SIMPLE_AUTOMATION_TEST(FPCGAttributePropertySelector_FromString_WithExtraNames_Property, FPCGAttributePropertySelectorTests, "Plugins.PCG.AttributePropertySelector.FromString.WithExtraNames.Property", PCGTestsCommon::TestFlags)
+IMPLEMENT_CUSTOM_SIMPLE_AUTOMATION_TEST(FPCGAttributePropertySelector_FromString_WithExtraNames_PointProperty, FPCGAttributePropertySelectorTests, "Plugins.PCG.AttributePropertySelector.FromString.WithExtraNames.PointProperty", PCGTestsCommon::TestFlags)
+IMPLEMENT_CUSTOM_SIMPLE_AUTOMATION_TEST(FPCGAttributePropertySelector_FromString_WithExtraNames_ExtraProperty, FPCGAttributePropertySelectorTests, "Plugins.PCG.AttributePropertySelector.FromString.WithExtraNames.ExtraProperty", PCGTestsCommon::TestFlags)
+IMPLEMENT_CUSTOM_SIMPLE_AUTOMATION_TEST(FPCGAttributePropertySelector_FromString_WithExtraNames_AtLast, FPCGAttributePropertySelectorTests, "Plugins.PCG.AttributePropertySelector.FromString.WithExtraNames.@Last", PCGTestsCommon::TestFlags)
+IMPLEMENT_CUSTOM_SIMPLE_AUTOMATION_TEST(FPCGAttributePropertySelector_FromString_WithExtraNames_AtSource, FPCGAttributePropertySelectorTests, "Plugins.PCG.AttributePropertySelector.FromString.WithExtraNames.@Source", PCGTestsCommon::TestFlags)
+IMPLEMENT_CUSTOM_SIMPLE_AUTOMATION_TEST(FPCGAttributePropertySelector_FromString_WithExtraNames_WithLayer_Attribute, FPCGAttributePropertySelectorTests, "Plugins.PCG.AttributePropertySelector.FromString.WithExtraNames.WithLayer.Attribute", PCGTestsCommon::TestFlags)
+IMPLEMENT_CUSTOM_SIMPLE_AUTOMATION_TEST(FPCGAttributePropertySelector_FromString_WithExtraNames_WithLayer_Property, FPCGAttributePropertySelectorTests, "Plugins.PCG.AttributePropertySelector.FromString.WithExtraNames.WithLayer.Property", PCGTestsCommon::TestFlags)
+IMPLEMENT_CUSTOM_SIMPLE_AUTOMATION_TEST(FPCGAttributePropertySelector_FromString_WithExtraNames_WithLayer_PointProperty, FPCGAttributePropertySelectorTests, "Plugins.PCG.AttributePropertySelector.FromString.WithExtraNames.WithLayer.PointProperty", PCGTestsCommon::TestFlags)
+IMPLEMENT_CUSTOM_SIMPLE_AUTOMATION_TEST(FPCGAttributePropertySelector_FromString_WithExtraNames_WithLayer_ExtraProperty, FPCGAttributePropertySelectorTests, "Plugins.PCG.AttributePropertySelector.FromString.WithExtraNames.WithLayer.ExtraProperty", PCGTestsCommon::TestFlags)
+IMPLEMENT_CUSTOM_SIMPLE_AUTOMATION_TEST(FPCGAttributePropertySelector_FromString_WithExtraNames_WithLayer_AtLast, FPCGAttributePropertySelectorTests, "Plugins.PCG.AttributePropertySelector.FromString.WithExtraNames.WithLayer.@Last", PCGTestsCommon::TestFlags)
+IMPLEMENT_CUSTOM_SIMPLE_AUTOMATION_TEST(FPCGAttributePropertySelector_FromString_WithExtraNames_WithLayer_AtSource, FPCGAttributePropertySelectorTests, "Plugins.PCG.AttributePropertySelector.FromString.WithExtraNames.WithLayer.@Source", PCGTestsCommon::TestFlags)
+
+IMPLEMENT_CUSTOM_SIMPLE_AUTOMATION_TEST(FPCGAttributePropertySelector_ExportImport_Attribute, FPCGAttributePropertySelectorTests, "Plugins.PCG.AttributePropertySelector.ExportImport.Attribute", PCGTestsCommon::TestFlags)
+IMPLEMENT_CUSTOM_SIMPLE_AUTOMATION_TEST(FPCGAttributePropertySelector_ExportImport_Property, FPCGAttributePropertySelectorTests, "Plugins.PCG.AttributePropertySelector.ExportImport.Property", PCGTestsCommon::TestFlags)
+IMPLEMENT_CUSTOM_SIMPLE_AUTOMATION_TEST(FPCGAttributePropertySelector_ExportImport_PointProperty, FPCGAttributePropertySelectorTests, "Plugins.PCG.AttributePropertySelector.ExportImport.PointProperty", PCGTestsCommon::TestFlags)
+IMPLEMENT_CUSTOM_SIMPLE_AUTOMATION_TEST(FPCGAttributePropertySelector_ExportImport_ExtraProperty, FPCGAttributePropertySelectorTests, "Plugins.PCG.AttributePropertySelector.ExportImport.ExtraProperty", PCGTestsCommon::TestFlags)
+IMPLEMENT_CUSTOM_SIMPLE_AUTOMATION_TEST(FPCGAttributePropertySelector_ExportImport_AtLast, FPCGAttributePropertySelectorTests, "Plugins.PCG.AttributePropertySelector.ExportImport.@Last", PCGTestsCommon::TestFlags)
+IMPLEMENT_CUSTOM_SIMPLE_AUTOMATION_TEST(FPCGAttributePropertySelector_ExportImport_AtSource, FPCGAttributePropertySelectorTests, "Plugins.PCG.AttributePropertySelector.ExportImport.@Source", PCGTestsCommon::TestFlags)
+IMPLEMENT_CUSTOM_SIMPLE_AUTOMATION_TEST(FPCGAttributePropertySelector_ExportImport_WithLayer_Attribute, FPCGAttributePropertySelectorTests, "Plugins.PCG.AttributePropertySelector.ExportImport.WithLayer.Attribute", PCGTestsCommon::TestFlags)
+IMPLEMENT_CUSTOM_SIMPLE_AUTOMATION_TEST(FPCGAttributePropertySelector_ExportImport_WithLayer_Property, FPCGAttributePropertySelectorTests, "Plugins.PCG.AttributePropertySelector.ExportImport.WithLayer.Property", PCGTestsCommon::TestFlags)
+IMPLEMENT_CUSTOM_SIMPLE_AUTOMATION_TEST(FPCGAttributePropertySelector_ExportImport_WithLayer_PointProperty, FPCGAttributePropertySelectorTests, "Plugins.PCG.AttributePropertySelector.ExportImport.WithLayer.PointProperty", PCGTestsCommon::TestFlags)
+IMPLEMENT_CUSTOM_SIMPLE_AUTOMATION_TEST(FPCGAttributePropertySelector_ExportImport_WithLayer_ExtraProperty, FPCGAttributePropertySelectorTests, "Plugins.PCG.AttributePropertySelector.ExportImport.WithLayer.ExtraProperty", PCGTestsCommon::TestFlags)
+IMPLEMENT_CUSTOM_SIMPLE_AUTOMATION_TEST(FPCGAttributePropertySelector_ExportImport_WithExtraNames_Attribute, FPCGAttributePropertySelectorTests, "Plugins.PCG.AttributePropertySelector.ExportImport.WithExtraNames.Attribute", PCGTestsCommon::TestFlags)
+IMPLEMENT_CUSTOM_SIMPLE_AUTOMATION_TEST(FPCGAttributePropertySelector_ExportImport_WithExtraNames_Property, FPCGAttributePropertySelectorTests, "Plugins.PCG.AttributePropertySelector.ExportImport.WithExtraNames.Property", PCGTestsCommon::TestFlags)
+IMPLEMENT_CUSTOM_SIMPLE_AUTOMATION_TEST(FPCGAttributePropertySelector_ExportImport_WithExtraNames_PointProperty, FPCGAttributePropertySelectorTests, "Plugins.PCG.AttributePropertySelector.ExportImport.WithExtraNames.PointProperty", PCGTestsCommon::TestFlags)
+IMPLEMENT_CUSTOM_SIMPLE_AUTOMATION_TEST(FPCGAttributePropertySelector_ExportImport_WithExtraNames_ExtraProperty, FPCGAttributePropertySelectorTests, "Plugins.PCG.AttributePropertySelector.ExportImport.WithExtraNames.ExtraProperty", PCGTestsCommon::TestFlags)
+IMPLEMENT_CUSTOM_SIMPLE_AUTOMATION_TEST(FPCGAttributePropertySelector_ExportImport_WithExtraNames_WithLayer_Attribute, FPCGAttributePropertySelectorTests, "Plugins.PCG.AttributePropertySelector.ExportImport.WithExtraNames.WithLayer.Attribute", PCGTestsCommon::TestFlags)
+IMPLEMENT_CUSTOM_SIMPLE_AUTOMATION_TEST(FPCGAttributePropertySelector_ExportImport_WithExtraNames_WithLayer_Property, FPCGAttributePropertySelectorTests, "Plugins.PCG.AttributePropertySelector.ExportImport.WithExtraNames.WithLayer.Property", PCGTestsCommon::TestFlags)
+IMPLEMENT_CUSTOM_SIMPLE_AUTOMATION_TEST(FPCGAttributePropertySelector_ExportImport_WithExtraNames_WithLayer_PointProperty, FPCGAttributePropertySelectorTests, "Plugins.PCG.AttributePropertySelector.ExportImport.WithExtraNames.WithLayer.PointProperty", PCGTestsCommon::TestFlags)
+IMPLEMENT_CUSTOM_SIMPLE_AUTOMATION_TEST(FPCGAttributePropertySelector_ExportImport_WithExtraNames_WithLayer_ExtraProperty, FPCGAttributePropertySelectorTests, "Plugins.PCG.AttributePropertySelector.ExportImport.WithExtraNames.WithLayer.ExtraProperty", PCGTestsCommon::TestFlags)
+
+// To be removed when deprecated point property no longer exists.
+IMPLEMENT_CUSTOM_SIMPLE_AUTOMATION_TEST(FPCGAttributePropertySelector_Deprecation_PointProperty, FPCGAttributePropertySelectorTests, "Plugins.PCG.AttributePropertySelector.Deprecation.PointProperty", PCGTestsCommon::TestFlags)
+IMPLEMENT_CUSTOM_SIMPLE_AUTOMATION_TEST(FPCGAttributePropertySelector_Deprecation_WithExtraNames_PointProperty, FPCGAttributePropertySelectorTests, "Plugins.PCG.AttributePropertySelector.Deprecation.WithExtraNames.PointProperty", PCGTestsCommon::TestFlags)
+IMPLEMENT_CUSTOM_SIMPLE_AUTOMATION_TEST(FPCGAttributePropertySelector_Deprecation_CDO_PointProperty, FPCGAttributePropertySelectorTests, "Plugins.PCG.AttributePropertySelector.Deprecation.CDO.PointProperty", PCGTestsCommon::TestFlags)
+IMPLEMENT_CUSTOM_SIMPLE_AUTOMATION_TEST(FPCGAttributePropertySelector_Deprecation_CDO_WithExtraNames_PointProperty, FPCGAttributePropertySelectorTests, "Plugins.PCG.AttributePropertySelector.Deprecation.CDO.WithExtraNames.PointProperty", PCGTestsCommon::TestFlags)
+IMPLEMENT_CUSTOM_SIMPLE_AUTOMATION_TEST(FPCGAttributePropertySelector_Deprecation_CDO_NewVersion_PointProperty, FPCGAttributePropertySelectorTests, "Plugins.PCG.AttributePropertySelector.Deprecation.CDO.NewVersion.PointProperty", PCGTestsCommon::TestFlags)
+IMPLEMENT_CUSTOM_SIMPLE_AUTOMATION_TEST(FPCGAttributePropertySelector_Deprecation_CDO_NewVersion_WithExtraNames_PointProperty, FPCGAttributePropertySelectorTests, "Plugins.PCG.AttributePropertySelector.Deprecation.CDO.NewVersion.WithExtraNames.PointProperty", PCGTestsCommon::TestFlags)
+
+bool FPCGAttributePropertySelector_FromString_Attribute::RunTest(const FString& Parameters)
+{
+	return FromStringComparison(TEXT("MyAttr"), FPCGAttributePropertySelector::CreateAttributeSelector(TEXT("MyAttr")));
+}
+
+bool FPCGAttributePropertySelector_FromString_Property::RunTest(const FString& Parameters)
+{
+	return FromStringComparison(TEXT("$MyProperty"), FPCGAttributePropertySelector::CreatePropertySelector(TEXT("MyProperty")));
+}
+
+bool FPCGAttributePropertySelector_FromString_PointProperty::RunTest(const FString& Parameters)
+{
+	return FromStringComparison(TEXT("$Rotation"), FPCGAttributePropertySelector::CreatePointPropertySelector(EPCGPointProperties::Rotation));
+}
+
+bool FPCGAttributePropertySelector_FromString_ExtraProperty::RunTest(const FString& Parameters)
+{
+	return FromStringComparison(TEXT("$Index"), FPCGAttributePropertySelector::CreateExtraPropertySelector(EPCGExtraProperties::Index));
+}
+
+bool FPCGAttributePropertySelector_FromString_AtLast::RunTest(const FString& Parameters)
+{
+	return FromStringComparison(TEXT("@Last"), FPCGAttributePropertySelector::CreateAttributeSelector(TEXT("@Last")));
+}
+
+bool FPCGAttributePropertySelector_FromString_AtSource::RunTest(const FString& Parameters)
+{
+	return FromStringComparison(TEXT("@Source"), FPCGAttributePropertySelector::CreateAttributeSelector(TEXT("@Source")));
+}
+
+bool FPCGAttributePropertySelector_FromString_WithLayer_Attribute::RunTest(const FString& Parameters)
+{
+	return FromStringComparison(TEXT("@NewLayer.MyAttr"), FPCGAttributePropertySelector::CreateAttributeSelector(TEXT("MyAttr"), TEXT("NewLayer")));
+}
+
+bool FPCGAttributePropertySelector_FromString_WithLayer_Property::RunTest(const FString& Parameters)
+{
+	return FromStringComparison(TEXT("@NewLayer.$MyProperty"), FPCGAttributePropertySelector::CreatePropertySelector(TEXT("MyProperty"), TEXT("NewLayer")));
+}
+
+bool FPCGAttributePropertySelector_FromString_WithLayer_PointProperty::RunTest(const FString& Parameters)
+{
+	return FromStringComparison(TEXT("@NewLayer.$Scale"), FPCGAttributePropertySelector::CreatePointPropertySelector(EPCGPointProperties::Scale, TEXT("NewLayer")));
+}
+
+bool FPCGAttributePropertySelector_FromString_WithLayer_ExtraProperty::RunTest(const FString& Parameters)
+{
+	return FromStringComparison(TEXT("@NewLayer.$Index"), FPCGAttributePropertySelector::CreateExtraPropertySelector(EPCGExtraProperties::Index, TEXT("NewLayer")));
+}
+
+bool FPCGAttributePropertySelector_FromString_WithLayer_AtLast::RunTest(const FString& Parameters)
+{
+	return FromStringComparison(TEXT("@NewLayer.@Last"), FPCGAttributePropertySelector::CreateAttributeSelector(TEXT("@Last"), TEXT("NewLayer")));
+}
+
+bool FPCGAttributePropertySelector_FromString_WithLayer_AtSource::RunTest(const FString& Parameters)
+{
+	return FromStringComparison(TEXT("@NewLayer.@Source"), FPCGAttributePropertySelector::CreateAttributeSelector(TEXT("@Source"), TEXT("NewLayer")));
+}
+
+bool FPCGAttributePropertySelector_FromString_WithExtraNames_Attribute::RunTest(const FString& Parameters)
+{
+	return FromStringComparison(TEXT("MyAttr.X"), FPCGAttributePropertySelector::CreateAttributeSelector(TEXT("MyAttr"), NAME_None, {TEXT("X")}));
+}
+
+bool FPCGAttributePropertySelector_FromString_WithExtraNames_Property::RunTest(const FString& Parameters)
+{
+	return FromStringComparison(TEXT("$MyProperty.XY"), FPCGAttributePropertySelector::CreatePropertySelector(TEXT("MyProperty"), NAME_None, {TEXT("XY")}));
+}
+
+bool FPCGAttributePropertySelector_FromString_WithExtraNames_PointProperty::RunTest(const FString& Parameters)
+{
+	return FromStringComparison(TEXT("$Density.Length"), FPCGAttributePropertySelector::CreatePointPropertySelector(EPCGPointProperties::Density, NAME_None, {TEXT("Length")}));
+}
+
+bool FPCGAttributePropertySelector_FromString_WithExtraNames_ExtraProperty::RunTest(const FString& Parameters)
+{
+	return FromStringComparison(TEXT("$Index.Abs"), FPCGAttributePropertySelector::CreateExtraPropertySelector(EPCGExtraProperties::Index, NAME_None, {TEXT("Abs")}));
+}
+
+bool FPCGAttributePropertySelector_FromString_WithExtraNames_AtLast::RunTest(const FString& Parameters)
+{
+	return FromStringComparison(TEXT("@Last.X"), FPCGAttributePropertySelector::CreateAttributeSelector(TEXT("@Last"), NAME_None, {TEXT("X")}));
+}
+
+bool FPCGAttributePropertySelector_FromString_WithExtraNames_AtSource::RunTest(const FString& Parameters)
+{
+	return FromStringComparison(TEXT("@Source.X"), FPCGAttributePropertySelector::CreateAttributeSelector(TEXT("@Source"), NAME_None, {TEXT("X")}));
+}
+
+bool FPCGAttributePropertySelector_FromString_WithExtraNames_WithLayer_Attribute::RunTest(const FString& Parameters)
+{
+	return FromStringComparison(TEXT("@NewLayer.MyAttr.Position.XY"), FPCGAttributePropertySelector::CreateAttributeSelector(TEXT("MyAttr"), TEXT("NewLayer"), {TEXT("Position"), TEXT("XY")}));
+}
+
+bool FPCGAttributePropertySelector_FromString_WithExtraNames_WithLayer_Property::RunTest(const FString& Parameters)
+{
+	return FromStringComparison(TEXT("@NewLayer.$MyProperty.RGBA"), FPCGAttributePropertySelector::CreatePropertySelector(TEXT("MyProperty"), TEXT("NewLayer"), {TEXT("RGBA")}));
+}
+
+bool FPCGAttributePropertySelector_FromString_WithExtraNames_WithLayer_PointProperty::RunTest(const FString& Parameters)
+{
+	return FromStringComparison(TEXT("@NewLayer.$Steepness.Sign"), FPCGAttributePropertySelector::CreatePointPropertySelector(EPCGPointProperties::Steepness, TEXT("NewLayer"), {TEXT("Sign")}));
+}
+
+bool FPCGAttributePropertySelector_FromString_WithExtraNames_WithLayer_ExtraProperty::RunTest(const FString& Parameters)
+{
+	return FromStringComparison(TEXT("@NewLayer.$Index.Sign"), FPCGAttributePropertySelector::CreateExtraPropertySelector(EPCGExtraProperties::Index, TEXT("NewLayer"), {TEXT("Sign")}));
+}
+
+bool FPCGAttributePropertySelector_FromString_WithExtraNames_WithLayer_AtLast::RunTest(const FString& Parameters)
+{
+	return FromStringComparison(TEXT("@[REDACTED_EMAIL]"), FPCGAttributePropertySelector::CreateAttributeSelector(TEXT("@Last"), TEXT("NewLayer"), {TEXT("Position"), TEXT("XY")}));
+}
+
+bool FPCGAttributePropertySelector_FromString_WithExtraNames_WithLayer_AtSource::RunTest(const FString& Parameters)
+{
+	return FromStringComparison(TEXT("@[REDACTED_EMAIL]"), FPCGAttributePropertySelector::CreateAttributeSelector(TEXT("@Source"), TEXT("NewLayer"), {TEXT("Position"), TEXT("XY")}));
+}
+
+bool FPCGAttributePropertySelector_ExportImport_Attribute::RunTest(const FString& Parameters)
+{
+	return ImportExport(TEXT("PCGBegin(MyAttr)PCGEnd"), FPCGAttributePropertySelector::CreateAttributeSelector(TEXT("MyAttr")));
+}
+
+bool FPCGAttributePropertySelector_ExportImport_Property::RunTest(const FString& Parameters)
+{
+	return ImportExport(TEXT("PCGBegin($MyProperty)PCGEnd"), FPCGAttributePropertySelector::CreatePropertySelector(TEXT("MyProperty")));
+}
+
+bool FPCGAttributePropertySelector_ExportImport_PointProperty::RunTest(const FString& Parameters)
+{
+	return ImportExport(TEXT("PCGBegin($Rotation)PCGEnd"), FPCGAttributePropertySelector::CreatePointPropertySelector(EPCGPointProperties::Rotation));
+}
+
+bool FPCGAttributePropertySelector_ExportImport_ExtraProperty::RunTest(const FString& Parameters)
+{
+	return ImportExport(TEXT("PCGBegin($Index)PCGEnd"), FPCGAttributePropertySelector::CreateExtraPropertySelector(EPCGExtraProperties::Index));
+}
+
+bool FPCGAttributePropertySelector_ExportImport_AtLast::RunTest(const FString& Parameters)
+{
+	return ImportExport(TEXT("PCGBegin(@Last)PCGEnd"), FPCGAttributePropertySelector::CreateAttributeSelector(TEXT("@Last")));
+}
+
+bool FPCGAttributePropertySelector_ExportImport_AtSource::RunTest(const FString& Parameters)
+{
+	return ImportExport(TEXT("PCGBegin(@Source)PCGEnd"), FPCGAttributePropertySelector::CreateAttributeSelector(TEXT("@Source")));
+}
+
+bool FPCGAttributePropertySelector_ExportImport_WithLayer_Attribute::RunTest(const FString& Parameters)
+{
+	return ImportExport(TEXT("PCGBegin(@NewLayer.MyAttr)PCGEnd"), FPCGAttributePropertySelector::CreateAttributeSelector(TEXT("MyAttr"), TEXT("NewLayer")));
+}
+
+bool FPCGAttributePropertySelector_ExportImport_WithLayer_Property::RunTest(const FString& Parameters)
+{
+	return ImportExport(TEXT("PCGBegin(@NewLayer.$MyProperty)PCGEnd"), FPCGAttributePropertySelector::CreatePropertySelector(TEXT("MyProperty"), TEXT("NewLayer")));
+}
+
+bool FPCGAttributePropertySelector_ExportImport_WithLayer_PointProperty::RunTest(const FString& Parameters)
+{
+	return ImportExport(TEXT("PCGBegin(@NewLayer.$Scale)PCGEnd"), FPCGAttributePropertySelector::CreatePointPropertySelector(EPCGPointProperties::Scale, TEXT("NewLayer")));
+}
+
+bool FPCGAttributePropertySelector_ExportImport_WithLayer_ExtraProperty::RunTest(const FString& Parameters)
+{
+	return ImportExport(TEXT("PCGBegin(@NewLayer.$Index)PCGEnd"), FPCGAttributePropertySelector::CreateExtraPropertySelector(EPCGExtraProperties::Index, TEXT("NewLayer")));
+}
+
+bool FPCGAttributePropertySelector_ExportImport_WithExtraNames_Attribute::RunTest(const FString& Parameters)
+{
+	return ImportExport(TEXT("PCGBegin(MyAttr.X)PCGEnd"), FPCGAttributePropertySelector::CreateAttributeSelector(TEXT("MyAttr"), NAME_None, {TEXT("X")}));
+}
+
+bool FPCGAttributePropertySelector_ExportImport_WithExtraNames_Property::RunTest(const FString& Parameters)
+{
+	return ImportExport(TEXT("PCGBegin($MyProperty.XY)PCGEnd"), FPCGAttributePropertySelector::CreatePropertySelector(TEXT("MyProperty"), NAME_None, {TEXT("XY")}));
+}
+
+bool FPCGAttributePropertySelector_ExportImport_WithExtraNames_PointProperty::RunTest(const FString& Parameters)
+{
+	return ImportExport(TEXT("PCGBegin($Density.Length)PCGEnd"), FPCGAttributePropertySelector::CreatePointPropertySelector(EPCGPointProperties::Density, NAME_None, {TEXT("Length")}));
+}
+
+bool FPCGAttributePropertySelector_ExportImport_WithExtraNames_ExtraProperty::RunTest(const FString& Parameters)
+{
+	return ImportExport(TEXT("PCGBegin($Index.Abs)PCGEnd"), FPCGAttributePropertySelector::CreateExtraPropertySelector(EPCGExtraProperties::Index, NAME_None, {TEXT("Abs")}));
+}
+
+bool FPCGAttributePropertySelector_ExportImport_WithExtraNames_WithLayer_Attribute::RunTest(const FString& Parameters)
+{
+	return ImportExport(TEXT("PCGBegin(@NewLayer.MyAttr.Position.XY)PCGEnd"), FPCGAttributePropertySelector::CreateAttributeSelector(TEXT("MyAttr"), TEXT("NewLayer"), {TEXT("Position"), TEXT("XY")}));
+}
+
+bool FPCGAttributePropertySelector_ExportImport_WithExtraNames_WithLayer_Property::RunTest(const FString& Parameters)
+{
+	return ImportExport(TEXT("PCGBegin(@NewLayer.$MyProperty.RGBA)PCGEnd"), FPCGAttributePropertySelector::CreatePropertySelector(TEXT("MyProperty"), TEXT("NewLayer"), {TEXT("RGBA")}));
+}
+
+bool FPCGAttributePropertySelector_ExportImport_WithExtraNames_WithLayer_PointProperty::RunTest(const FString& Parameters)
+{
+	return ImportExport(TEXT("PCGBegin(@NewLayer.$Steepness.Sign)PCGEnd"), FPCGAttributePropertySelector::CreatePointPropertySelector(EPCGPointProperties::Steepness, TEXT("NewLayer"), {TEXT("Sign")}));
+}
+
+bool FPCGAttributePropertySelector_ExportImport_WithExtraNames_WithLayer_ExtraProperty::RunTest(const FString& Parameters)
+{
+	return ImportExport(TEXT("PCGBegin(@NewLayer.$Index.Sign)PCGEnd"), FPCGAttributePropertySelector::CreateExtraPropertySelector(EPCGExtraProperties::Index, TEXT("NewLayer"), {TEXT("Sign")}));
+}
+
+bool FPCGAttributePropertySelector_Deprecation_PointProperty::RunTest(const FString& Parameters)
+{
+	// To be removed when deprecated point property no longer exists.
+	return ValidateDeprecatedSelector(FPCGAttributePropertySelector::CreatePointPropertySelector(EPCGPointProperties::Steepness), EPCGPointProperties::Steepness, {});
+}
+
+bool FPCGAttributePropertySelector_Deprecation_WithExtraNames_PointProperty::RunTest(const FString& Parameters)
+{
+	// To be removed when deprecated point property no longer exists.
+	return ValidateDeprecatedSelector(FPCGAttributePropertySelector::CreatePointPropertySelector(EPCGPointProperties::Rotation, NAME_None, {TEXT("Forward")}), EPCGPointProperties::Rotation, {TEXT("Forward")});
+}
+
+bool FPCGAttributePropertySelector_Deprecation_CDO_PointProperty::RunTest(const FString& Parameters)
+{
+	// To be removed when deprecated point property no longer exists.
+	return ValidateDeprecatedSelectorCDO(FPCGAttributePropertySelector::CreatePointPropertySelector(EPCGPointProperties::Steepness), EPCGPointProperties::Steepness, {});
+}
+
+bool FPCGAttributePropertySelector_Deprecation_CDO_WithExtraNames_PointProperty::RunTest(const FString& Parameters)
+{
+	// To be removed when deprecated point property no longer exists.
+	return ValidateDeprecatedSelectorCDO(FPCGAttributePropertySelector::CreatePointPropertySelector(EPCGPointProperties::Rotation, NAME_None, {TEXT("Forward")}), EPCGPointProperties::Rotation, {TEXT("Forward")});
+}
+
+bool FPCGAttributePropertySelector_Deprecation_CDO_NewVersion_PointProperty::RunTest(const FString& Parameters)
+{
+	// To be removed when deprecated point property no longer exists.
+	return ValidateSelectorCDOWithNewerVersion(FPCGAttributePropertySelector::CreatePointPropertySelector(EPCGPointProperties::Steepness), EPCGPointProperties::Steepness, {});
+}
+
+bool FPCGAttributePropertySelector_Deprecation_CDO_NewVersion_WithExtraNames_PointProperty::RunTest(const FString& Parameters)
+{
+	// To be removed when deprecated point property no longer exists.
+	return ValidateSelectorCDOWithNewerVersion(FPCGAttributePropertySelector::CreatePointPropertySelector(EPCGPointProperties::Rotation, NAME_None, {TEXT("Forward")}), EPCGPointProperties::Rotation, {TEXT("Forward")});
+}
+
+#endif // WITH_EDITOR
diff --git a/Engine/Plugins/PCG/Source/PCG/Public/Metadata/PCGAttributePropertySelector.h b/Engine/Plugins/PCG/Source/PCG/Public/Metadata/PCGAttributePropertySelector.h
--- a/Engine/Plugins/PCG/Source/PCG/Public/Metadata/PCGAttributePropertySelector.h
+++ b/Engine/Plugins/PCG/Source/PCG/Public/Metadata/PCGAttributePropertySelector.h
@@ -12,8 +12,9 @@ UENUM()
 enum class EPCGAttributePropertySelection
 {
 	Attribute,
-	PointProperty,
-	ExtraProperty
+	PointProperty UE_DEPRECATED(5.6, "Please use Property flag now") UMETA(Hidden, DeprecatedProperty, DeprecationMessage="Please use Property flag now."),
+	ExtraProperty,
+	Property,
 };
 
 UENUM()
@@ -26,6 +27,11 @@ class FArchiveCrc32;
 struct FPCGCustomVersion;
 class UPCGData;
 
+#if WITH_EDITOR
+// Friend for unit tests
+class FPCGAttributePropertySelectorTests;
+#endif // WITH_EDITOR
+
 /**
 * Blueprint class to allow to select an attribute or a property.
 * It will handle the logic and can only be modified using the blueprint library defined below.
@@ -47,18 +53,25 @@ public:
 	bool operator==(const FPCGAttributePropertySelector& Other) const;
 	bool IsSame(const FPCGAttributePropertySelector& Other, bool bIncludeExtraNames = true) const;
 
-	// Setters, retrurn true if something changed.
+	bool Reset();
+	bool ResetExtraNames();
+	
+	// Setters, return true if something changed.
 	bool SetPointProperty(EPCGPointProperties InPointProperty, bool bResetExtraNames = true);
 	bool SetExtraProperty(EPCGExtraProperties InExtraProperty, bool bResetExtraNames = true);
 	bool SetAttributeName(FName InAttributeName, bool bResetExtraNames = true);
+	bool SetPropertyName(FName InPropertyName, bool bResetExtraNames = true);
+	bool SetLayerName(FName InLayerName, bool bResetExtraNames = true);
 
 	// Getters
 	EPCGAttributePropertySelection GetSelection() const { return Selection; }
 	const TArray<FString>& GetExtraNames() const { return ExtraNames; }
 	TArray<FString>& GetExtraNamesMutable() { return ExtraNames; }
 	FName GetAttributeName() const { return AttributeName; }
-	EPCGPointProperties GetPointProperty() const { return PointProperty; }
+	EPCGPointProperties GetPointProperty() const;
 	EPCGExtraProperties GetExtraProperty() const { return ExtraProperty; }
+	FName GetPropertyName() const { return PropertyName; }
+	FName GetLayerName() const { return LayerName; }
 
 	// Convenient function to know if it is a basic attribute (attribute and no extra names)
 	bool IsBasicAttribute() const;
@@ -68,8 +81,11 @@ public:
 
 	// Returns qualified attribute/property name with the accessors.
 	FString ToString() const;
+
+	// Returns attribute/property name only, with optional '@' qualifier for layers.
+	FString GetLayerString(bool bAddLeadingQualifier = true) const;
 	
-	// Returns attribute/property name only, with optional '$' qualifier for properties.
+	// Returns attribute/property name only, with optional '$' qualifier for properties and '.'.
 	FString GetAttributePropertyString(bool bAddPropertyQualifier = true) const;
 
 	// Returns the accessor part of the selector, with optional leading '.' separator.
@@ -95,29 +111,46 @@ public:
 
 	// Convenience templated static constructors
 	template <typename T = FPCGAttributePropertySelector>
-	static T CreateAttributeSelector(const FName AttributeName)
+	static T CreateAttributeSelector(const FName AttributeName, const FName LayerName = NAME_None, const TArrayView<const FString>& InExtraNames = {})
 	{
 		static_assert(std::is_base_of_v<FPCGAttributePropertySelector, T>, "The type must be of base class 'FPCGAttributePropertySelector'");
 		T Selector;
+		Selector.SetLayerName(LayerName);
 		Selector.SetAttributeName(AttributeName);
+		Selector.GetExtraNamesMutable() = InExtraNames;
 		return Selector;
 	}
-
+	
+	template <typename T = FPCGAttributePropertySelector>
+	static T CreatePropertySelector(FName PropertyName, const FName LayerName = NAME_None, const TArrayView<const FString>& InExtraNames = {})
+	{
+		static_assert(std::is_base_of_v<FPCGAttributePropertySelector, T>, "The type must be of base class 'FPCGAttributePropertySelector'");
+		T Selector;
+		Selector.SetLayerName(LayerName);
+		Selector.SetPropertyName(PropertyName);
+		Selector.GetExtraNamesMutable() = InExtraNames;
+		return Selector;
+	}
+	
 	template <typename T = FPCGAttributePropertySelector>
-	static T CreatePointPropertySelector(EPCGPointProperties PointProperty)
+	static T CreatePointPropertySelector(EPCGPointProperties PointProperty, const FName LayerName = NAME_None, const TArrayView<const FString>& InExtraNames = {})
 	{
 		static_assert(std::is_base_of_v<FPCGAttributePropertySelector, T>, "The type must be of base class 'FPCGAttributePropertySelector'");
 		T Selector;
+		Selector.SetLayerName(LayerName);
 		Selector.SetPointProperty(PointProperty);
+		Selector.GetExtraNamesMutable() = InExtraNames;
 		return Selector;
 	}
 
 	template <typename T = FPCGAttributePropertySelector>
-	static T CreateExtraPropertySelector(EPCGExtraProperties ExtraProperty)
+	static T CreateExtraPropertySelector(EPCGExtraProperties ExtraProperty, const FName LayerName = NAME_None, const TArrayView<const FString>& InExtraNames = {})
 	{
 		static_assert(std::is_base_of_v<FPCGAttributePropertySelector, T>, "The type must be of base class 'FPCGAttributePropertySelector'");
 		T Selector;
+		Selector.SetLayerName(LayerName);
 		Selector.SetExtraProperty(ExtraProperty);
+		Selector.GetExtraNamesMutable() = InExtraNames;
 		return Selector;
 	}
 
@@ -156,21 +189,37 @@ public:
 	bool ExportTextItem(FString& ValueStr, FPCGAttributePropertySelector const& DefaultValue, UObject* Parent, int32 PortFlags, UObject* ExportRootScope) const;
 	bool ImportTextItem(const TCHAR*& Buffer, int32 PortFlags, UObject* Parent, FOutputDevice* ErrorText);
 
+	bool Serialize(FArchive& Ar);
+	void PostSerialize(const FArchive& Ar);
+
 protected:
 	UPROPERTY()
 	EPCGAttributePropertySelection Selection = EPCGAttributePropertySelection::Attribute;
 
 	UPROPERTY()
-	FName AttributeName = NAME_None;
+	FName LayerName = NAME_None;
 
 	UPROPERTY()
-	EPCGPointProperties PointProperty = EPCGPointProperties::Position;
+	FName AttributeName = NAME_None;
 
 	UPROPERTY()
-	EPCGExtraProperties ExtraProperty = EPCGExtraProperties::Index;
+	FName PropertyName = NAME_None;
 
 	UPROPERTY()
 	TArray<FString> ExtraNames;
+
+	UPROPERTY()
+	EPCGExtraProperties ExtraProperty = EPCGExtraProperties::Index;
+	
+#if WITH_EDITORONLY_DATA
+	UPROPERTY()
+	EPCGPointProperties PointProperty_DEPRECATED = EPCGPointProperties::Position;
+#endif // WITH_EDITORONLY_DATA
+
+#if WITH_EDITOR
+	// To unit test deprecation
+	friend FPCGAttributePropertySelectorTests;
+#endif // WITH_EDITOR
 };
 
 /** Struct that will default on @Last (or @LastCreated for previously created selectors). */
@@ -223,6 +272,8 @@ struct TStructOpsTypeTraits<FPCGAttributePropertySelector> : public TStructOpsTy
 	{
 		WithExportTextItem = true,
 		WithImportTextItem = true,
+		WithPostSerialize = true,
+		WithSerializer = true
 	};
 };
 
@@ -239,13 +290,7 @@ struct TStructOpsTypeTraits<FPCGAttributePropertyInputSelector> : public TStruct
 };
 
 template<>
-struct TStructOpsTypeTraits<FPCGAttributePropertyOutputSelector> : public TStructOpsTypeTraits<FPCGAttributePropertySelector>
-{
-	enum
-	{
-		WithStructuredSerializeFromMismatchedTag = true
-	};
-};
+struct TStructOpsTypeTraits<FPCGAttributePropertyOutputSelector> : public TStructOpsTypeTraits<FPCGAttributePropertyInputSelector> {};
 
 /**
 * Helper class to allow the BP to call the custom setters and getters on FPCGAttributePropertySelector.
@@ -257,13 +302,19 @@ class UPCGAttributePropertySelectorBlueprintHelpers : public UBlueprintFunctionL
 
 public:
 	UFUNCTION(BlueprintCallable, Category = "PCG|PCGAttributePropertySelector", meta = (ScriptMethod))
-	static bool SetPointProperty(UPARAM(ref) FPCGAttributePropertySelector& Selector, EPCGPointProperties InPointProperty);
+	static bool SetPointProperty(UPARAM(ref) FPCGAttributePropertySelector& Selector, EPCGPointProperties InPointProperty, bool bResetExtraNames = true);
+
+	UFUNCTION(BlueprintCallable, Category = "PCG|PCGAttributePropertySelector", meta = (ScriptMethod))
+	static bool SetAttributeName(UPARAM(ref) FPCGAttributePropertySelector& Selector, FName InAttributeName, bool bResetExtraNames = true);
+
+	UFUNCTION(BlueprintCallable, Category = "PCG|PCGAttributePropertySelector", meta = (ScriptMethod))
+	static bool SetPropertyName(UPARAM(ref) FPCGAttributePropertySelector& Selector, FName InPropertyName, bool bResetExtraNames = true);
 
 	UFUNCTION(BlueprintCallable, Category = "PCG|PCGAttributePropertySelector", meta = (ScriptMethod))
-	static bool SetAttributeName(UPARAM(ref) FPCGAttributePropertySelector& Selector, FName InAttributeName);
+	static bool SetExtraProperty(UPARAM(ref) FPCGAttributePropertySelector& Selector, EPCGExtraProperties InExtraProperty, bool bResetExtraNames = true);
 
 	UFUNCTION(BlueprintCallable, Category = "PCG|PCGAttributePropertySelector", meta = (ScriptMethod))
-	static bool SetExtraProperty(UPARAM(ref) FPCGAttributePropertySelector& Selector, EPCGExtraProperties InExtraProperty);
+	static bool SetLayerName(UPARAM(ref) FPCGAttributePropertySelector& Selector, FName InLayerName, bool bResetExtraNames = true);
 
 	UFUNCTION(BlueprintCallable, Category = "PCG|PCGAttributePropertySelector", meta = (ScriptMethod))
 	static EPCGAttributePropertySelection GetSelection(UPARAM(const, ref) const FPCGAttributePropertySelector& Selector);
@@ -274,6 +325,12 @@ public:
 	UFUNCTION(BlueprintCallable, Category = "PCG|PCGAttributePropertySelector", meta = (ScriptMethod))
 	static FName GetAttributeName(UPARAM(const, ref) const FPCGAttributePropertySelector& Selector);
 
+	UFUNCTION(BlueprintCallable, Category = "PCG|PCGAttributePropertySelector", meta = (ScriptMethod))
+	static FName GetPropertyName(UPARAM(const, ref) const FPCGAttributePropertySelector& Selector);
+
+	UFUNCTION(BlueprintCallable, Category = "PCG|PCGAttributePropertySelector", meta = (ScriptMethod))
+	static FName GetLayerName(UPARAM(const, ref) const FPCGAttributePropertySelector& Selector);
+
 	UFUNCTION(BlueprintCallable, Category = "PCG|PCGAttributePropertySelector", meta = (ScriptMethod))
 	static EPCGExtraProperties GetExtraProperty(UPARAM(const, ref) const FPCGAttributePropertySelector& Selector);
 
@@ -287,5 +344,5 @@ public:
 	static FPCGAttributePropertyInputSelector CopyAndFixLast(UPARAM(const, ref) const FPCGAttributePropertyInputSelector& Selector, const UPCGData* InData);
 
 	UFUNCTION(BlueprintCallable, Category = "PCG|PCGAttributePropertySelector", meta = (ScriptMethod))
-	static FPCGAttributePropertyOutputSelector CopyAndFixSource(UPARAM(const, ref) const FPCGAttributePropertyOutputSelector& OutputSelector, const FPCGAttributePropertyInputSelector& InputSelector);
+	static FPCGAttributePropertyOutputSelector CopyAndFixSource(UPARAM(const, ref) const FPCGAttributePropertyOutputSelector& OutputSelector, const FPCGAttributePropertyInputSelector& InputSelector, const UPCGData* InOptionalData = nullptr);
 };
\ No newline at end of file
diff --git a/Engine/Plugins/PCG/Source/PCG/Public/PCGCustomVersion.h b/Engine/Plugins/PCG/Source/PCG/Public/PCGCustomVersion.h
--- a/Engine/Plugins/PCG/Source/PCG/Public/PCGCustomVersion.h
+++ b/Engine/Plugins/PCG/Source/PCG/Public/PCGCustomVersion.h
@@ -80,9 +80,12 @@ struct FPCGCustomVersion
 		// Attributes and tags can now contain spaces and will no longer be parsed by spaces.
 		AttributesAndTagsCanContainSpaces = 23,
 
-		// MultiLevelMetadata
+		// Multi Layer Metadata
 		MultiLevelMetadata = 24,
 
+		// Refactor of the Attribute Property selector to deprecate point properties and support any property.
+		AttributePropertySelectorDeprecatePointProperties = 25,
+
 		// -----<new versions can be added above this line>-------------------------------------------------
 		VersionPlusOne,
 		LatestVersion = VersionPlusOne - 1
diff --git a/Engine/Plugins/PCG/Source/PCG/Public/PCGPoint.h b/Engine/Plugins/PCG/Source/PCG/Public/PCGPoint.h
--- a/Engine/Plugins/PCG/Source/PCG/Public/PCGPoint.h
+++ b/Engine/Plugins/PCG/Source/PCG/Public/PCGPoint.h
@@ -25,7 +25,9 @@ enum class EPCGPointProperties : uint8
 	LocalCenter UMETA(Tooltip = "The local center location of the point's volume, halfway between the minimum and maximum bounds."),
 	Seed UMETA(Tooltip = "Used to seed random processes during various operations."),
 	LocalSize UMETA(Tooltip = "The difference between the maximum and minimum bounds of the point."),
-	ScaledLocalSize UMETA(Tooltip = "The difference between the maximum and minimum bounds of the point, after only the scale has been applied.")
+	ScaledLocalSize UMETA(Tooltip = "The difference between the maximum and minimum bounds of the point, after only the scale has been applied."),
+
+	Invalid = 255 UMETA(Hidden)
 };
 
 USTRUCT(BlueprintType)
diff --git a/Engine/Plugins/PCG/Source/PCGEditor/Private/Details/PCGAttributePropertySelectorDetails.cpp b/Engine/Plugins/PCG/Source/PCGEditor/Private/Details/PCGAttributePropertySelectorDetails.cpp
--- a/Engine/Plugins/PCG/Source/PCGEditor/Private/Details/PCGAttributePropertySelectorDetails.cpp
+++ b/Engine/Plugins/PCG/Source/PCGEditor/Private/Details/PCGAttributePropertySelectorDetails.cpp
@@ -30,6 +30,11 @@ namespace PCGAttributePropertySelectorDetails
 		{
 			for (int32 i = 0; i < EnumPtr->NumEnums() - 1; ++i)
 			{
+				if (EnumPtr->HasMetaData(TEXT("Hidden"), i))
+				{
+					continue;
+				}
+				
 				FString EnumName = EnumPtr->GetDisplayNameTextByIndex(i).ToString();
 				EnumType EnumValue = (EnumType)EnumPtr->GetValueByIndex(i);
 				FText Tooltip = EnumPtr->GetToolTipTextByIndex(i);
