[PCG] Add layer for spline control points properties accessor

* Add a new layer `[REDACTED_HANDLE]` that will also be used by the metadata when we'll support it
* Add a new accessor `FPCGInterpCurveAccessor` to access InterpCurve structs
* Add new keys for spline data
* Add new accessor methods for spline data
* Support for $Position, $Rotation, $Scale, $ArriveTangent, $LeaveTangent and $InterpType
* Small changes in the PCG Property accessor, as the extra properties being a r-value reference, need to be moved not forwarded.

diff --git a/Engine/Plugins/PCG/Source/PCG/Private/Data/PCGSplineData.cpp b/Engine/Plugins/PCG/Source/PCG/Private/Data/PCGSplineData.cpp
--- a/Engine/Plugins/PCG/Source/PCG/Private/Data/PCGSplineData.cpp
+++ b/Engine/Plugins/PCG/Source/PCG/Private/Data/PCGSplineData.cpp
@@ -9,12 +9,15 @@
 #include "Data/PCGSpatialData.h"
 #include "Elements/PCGSplineSampler.h"
 #include "Helpers/PCGHelpers.h"
+#include "Metadata/Accessors/PCGSplineAccessor.h"
 
 #include "Components/SplineComponent.h"
 #include "Serialization/ArchiveCrc32.h"
 
 #include UE_INLINE_GENERATED_CPP_BY_NAME(PCGSplineData)
 
+#define LOCTEXT_NAMESPACE "PCGSplineData"
+
 void UPCGSplineData::Initialize(const USplineComponent* InSpline)
 {
 	check(InSpline);
@@ -176,6 +179,118 @@ FTransform UPCGSplineData::GetTransformAtAlpha(float Alpha) const
 	return SplineStruct.GetTransformAtSplineInputKey(GetInputKeyAtAlpha(Alpha), ESplineCoordinateSpace::World);
 }
 
+TUniquePtr<IPCGAttributeAccessor> UPCGSplineData::CreateStaticAccessor(const FPCGAttributePropertySelector& InSelector, bool bQuiet)
+{
+	static const FProperty* PCGSplineStructProperty = UPCGSplineData::StaticClass()->FindPropertyByName(GET_MEMBER_NAME_CHECKED(UPCGSplineData, SplineStruct));
+	static const FProperty* SplineCurveProperty = FPCGSplineStruct::StaticStruct()->FindPropertyByName(GET_MEMBER_NAME_CHECKED(FPCGSplineStruct, SplineCurves));
+
+	const FName LayerName = InSelector.GetLayerName();
+
+	if (InSelector.GetSelection() == EPCGAttributePropertySelection::Property && (LayerName.IsNone() || LayerName == PCGSplineData::ControlPointLayerName))
+	{
+		const FName PropertyName = InSelector.GetName();
+		TArray<const FProperty*> ExtraProperties = {PCGSplineStructProperty, SplineCurveProperty};
+		
+		if (PropertyName == TEXT("Position") || PropertyName == TEXT("Scale"))
+		{
+			const FStructProperty* VectorProperty = CastField<FStructProperty>(FSplineCurves::StaticStruct()->FindPropertyByName(PropertyName));
+			check(VectorProperty);
+			return MakeUnique<FPCGInterpCurveAccessor<FInterpCurveVector>>(VectorProperty, std::move(ExtraProperties));
+		}
+		else if (PropertyName == TEXT("Rotation"))
+		{
+			const FStructProperty* QuatProperty = CastField<FStructProperty>(FSplineCurves::StaticStruct()->FindPropertyByName(PropertyName));
+			check(QuatProperty);
+			return MakeUnique<FPCGInterpCurveAccessor<FInterpCurveQuat>>(QuatProperty, std::move(ExtraProperties));
+		}
+		else if (PropertyName == TEXT("LeaveTangent") || PropertyName == TEXT("ArriveTangent") || PropertyName == TEXT("InterpType"))
+		{
+			// Implementation note:
+			// Tangents and Interpolation type are all linked to the Position property on FSplineCurves, even if technically the rotation and the scale also have tangents and interpolation type.
+			// But in the spline details view, and everywhere we have tangents and interpolation, they are always related to the Position interp curve.
+			// If we need access to the rotation/scale tangents/interpolation modes, we can add them later on.
+			const FStructProperty* PositionProperty = CastField<FStructProperty>(FSplineCurves::StaticStruct()->FindPropertyByName(TEXT("Position")));
+			check(PositionProperty);
+			if (PropertyName == TEXT("LeaveTangent"))
+			{
+				return MakeUnique<FPCGInterpCurveAccessor<FInterpCurveVector, EPCGInterpCurveAccessorTarget::LeaveTangent>>(PositionProperty, std::move(ExtraProperties));
+			}
+			else if (PropertyName == TEXT("ArriveTangent"))
+			{
+				return MakeUnique<FPCGInterpCurveAccessor<FInterpCurveVector, EPCGInterpCurveAccessorTarget::ArriveTangent>>(PositionProperty, std::move(ExtraProperties));
+			}
+			else
+			{
+				return MakeUnique<FPCGInterpCurveAccessor<FInterpCurveVector, EPCGInterpCurveAccessorTarget::InterpMode>>(PositionProperty, std::move(ExtraProperties));
+			}
+		}
+
+		if (!bQuiet)
+		{
+			PCGLog::LogErrorOnGraph(FText::Format(LOCTEXT("FailCreateAccessor", "Property {0} is not valid for a control point."), FText::FromName(PropertyName)));
+		}
+	}
+	
+	return {};
+}
+
+FPCGAttributeAccessorMethods UPCGSplineData::GetSplineAccessorMethods()
+{
+	auto CreateAccessorFunc = [](UPCGData* InData, const FPCGAttributePropertySelector& InSelector, bool bQuiet) -> TUniquePtr<IPCGAttributeAccessor>
+	{
+		return CreateStaticAccessor(InSelector, bQuiet);
+	};
+
+	auto CreateConstAccessorFunc = [](const UPCGData* InData, const FPCGAttributePropertySelector& InSelector, bool bQuiet) -> TUniquePtr<const IPCGAttributeAccessor>
+	{
+		return CreateStaticAccessor(InSelector, bQuiet);
+	};
+
+	auto CreateAccessorKeysFunc = [](UPCGData* InData, const FPCGAttributePropertySelector& InSelector, bool bQuiet) -> TUniquePtr<IPCGAttributeAccessorKeys>
+	{
+		// If the selector is not a Property or ExtraProperty, return an invalid key to create the default keys for spatial data.
+		const EPCGAttributePropertySelection Selection = InSelector.GetSelection();
+		if (Selection == EPCGAttributePropertySelection::Property || Selection == EPCGAttributePropertySelection::ExtraProperty)
+		{
+			UPCGSplineData* SplineData = CastChecked<UPCGSplineData>(InData);
+			return MakeUnique<FPCGAttributeAccessorKeysSplineData>(SplineData);
+		}
+		else
+		{
+			return {};
+		}
+	};
+	
+	auto CreateConstAccessorKeysFunc = [](const UPCGData* InData, const FPCGAttributePropertySelector& InSelector, bool bQuiet) -> TUniquePtr<const IPCGAttributeAccessorKeys>
+	{
+		// If the selector is not a Property or ExtraProperty, return an invalid key to create the default keys for spatial data.
+		const EPCGAttributePropertySelection Selection = InSelector.GetSelection();
+		if (Selection == EPCGAttributePropertySelection::Property || Selection == EPCGAttributePropertySelection::ExtraProperty)
+		{
+			const UPCGSplineData* SplineData = CastChecked<UPCGSplineData>(InData);
+			return MakeUnique<const FPCGAttributeAccessorKeysSplineData>(SplineData);
+		}
+		else
+		{
+			return {};
+		}
+	};
+	
+	FPCGAttributeAccessorMethods Methods
+	{
+		.CreateAccessorFunc = CreateAccessorFunc,
+		.CreateConstAccessorFunc = CreateConstAccessorFunc,
+		.CreateAccessorKeysFunc = CreateAccessorKeysFunc,
+		.CreateConstAccessorKeysFunc = CreateConstAccessorKeysFunc
+	};
+
+#if WITH_EDITOR
+	Methods.FillSelectorMenuEntryFromEnum<EPCGSplineStructProperties>({LOCTEXT("ControlPointSelectorMenuEntry", "Spline Control Point")});
+#endif // WITH_EDITOR
+
+	return Methods;
+}
+
 const UPCGPointData* UPCGSplineData::CreatePointData(FPCGContext* Context) const
 {
 	TRACE_CPUPROFILER_EVENT_SCOPE(UPCGSplineData::CreatePointData);
@@ -414,3 +529,5 @@ UPCGSpatialData* UPCGSplineProjectionData::CopyInternal(FPCGContext* Context) co
 
 	return NewProjectionData;
 }
+
+#undef LOCTEXT_NAMESPACE
\ No newline at end of file
diff --git a/Engine/Plugins/PCG/Source/PCG/Private/Metadata/Accessors/PCGAttributeAccessorFactory.cpp b/Engine/Plugins/PCG/Source/PCG/Private/Metadata/Accessors/PCGAttributeAccessorFactory.cpp
--- a/Engine/Plugins/PCG/Source/PCG/Private/Metadata/Accessors/PCGAttributeAccessorFactory.cpp
+++ b/Engine/Plugins/PCG/Source/PCG/Private/Metadata/Accessors/PCGAttributeAccessorFactory.cpp
@@ -4,6 +4,7 @@
 
 #include "PCGModule.h"
 #include "Data/PCGSpatialData.h"
+#include "Metadata/PCGAttributePropertySelector.h"
 #include "Metadata/PCGMetadata.h"
 #include "Metadata/PCGMetadataAttribute.h"
 #include "Metadata/Accessors/PCGAttributeAccessorHelpers.h"
@@ -123,6 +124,33 @@ decltype(auto) FPCGAttributeAccessorFactory::CallOnMethod(PCGData* InData, Func
 	return AccessorType();
 }
 
+#if WITH_EDITOR
+void FPCGAttributeAccessorMethods::FillSelectorMenuEntryFromEnum(const UEnum* EnumType, const TArray<FText>& InMenuHierarchy)
+{
+	if (!EnumType)
+	{
+		return;
+	}
+
+	const int32 NumEnums = EnumType->ContainsExistingMax() ? EnumType->NumEnums() - 1 : EnumType->NumEnums();
+	for (int32 i = 0; i < NumEnums; ++i)
+	{
+		if (EnumType->HasMetaData(TEXT("Hidden"), i))
+		{
+			continue;
+		}
+
+		// Not const to be able to move them.
+		FText EnumName = EnumType->GetDisplayNameTextByIndex(i);
+		FText Tooltip = EnumType->GetToolTipTextByIndex(i);
+		const FName InvariantName = EnumType->GetNameByIndex(i);
+
+		// Can't move the menu hierarchy as it will be copied for each entry.
+		AttributeSelectorOptions.Emplace(std::move(EnumName), std::move(Tooltip), FPCGAttributePropertySelector::CreatePropertySelector(InvariantName), InMenuHierarchy);
+	}
+}
+#endif // WITH_EDITOR
+
 FPCGAttributeAccessorFactory& FPCGAttributeAccessorFactory::GetMutableInstance()
 {
 	return FPCGModule::GetMutableAttributeAccessorFactory();
@@ -189,6 +217,7 @@ TUniquePtr<const IPCGAttributeAccessorKeys> FPCGAttributeAccessorFactory::Create
 	});
 }
 
+#if WITH_EDITOR
 void FPCGAttributeAccessorFactory::ForEachSelectorMenuEntry(TFunctionRef<void(const FPCGAttributeSelectorMenuEntry&)> Callback) const
 {
 	for (const auto& [Class, Methods] : AccessorMethods)
@@ -199,6 +228,7 @@ void FPCGAttributeAccessorFactory::ForEachSelectorMenuEntry(TFunctionRef<void(co
 		}
 	}
 }
+#endif // WITH_EDITOR
 
 void FPCGAttributeAccessorFactory::RegisterMethods(TSubclassOf<UPCGData> PCGDataClass, FPCGAttributeAccessorMethods&& InAccessorMethods)
 {
diff --git a/Engine/Plugins/PCG/Source/PCG/Private/Metadata/Accessors/PCGAttributeAccessorKeys.cpp b/Engine/Plugins/PCG/Source/PCG/Private/Metadata/Accessors/PCGAttributeAccessorKeys.cpp
--- a/Engine/Plugins/PCG/Source/PCG/Private/Metadata/Accessors/PCGAttributeAccessorKeys.cpp
+++ b/Engine/Plugins/PCG/Source/PCG/Private/Metadata/Accessors/PCGAttributeAccessorKeys.cpp
@@ -1,10 +1,11 @@
 // Copyright Epic Games, Inc. All Rights Reserved.
 
 #include "Metadata/Accessors/PCGAttributeAccessorKeys.h"
-#include "Metadata/PCGMetadata.h"
-#include "Metadata/PCGMetadataAttribute.h"
 
 #include "PCGPoint.h"
+#include "Data/PCGSplineData.h"
+#include "Metadata/PCGMetadata.h"
+#include "Metadata/PCGMetadataAttribute.h"
 
 //////////////////////////////////////////////////////////////////// 
 FPCGAttributeAccessorKeysEntries::FPCGAttributeAccessorKeysEntries(const FPCGMetadataAttributeBase* Attribute)
@@ -115,6 +116,11 @@ FPCGAttributeAccessorKeysPoints::FPCGAttributeAccessorKeysPoints(const FPCGPoint
 	: FPCGAttributeAccessorKeysPoints(TArrayView<const FPCGPoint>(&InPoint, 1))
 {}
 
+bool FPCGAttributeAccessorKeysPoints::IsClassSupported(const UStruct* InClass) const
+{
+	return PCGAttributeAccessorKeys::IsClassSupported<FPCGPoint>(InClass);
+}
+
 bool FPCGAttributeAccessorKeysPoints::GetPointKeys(int32 InStart, TArrayView<FPCGPoint*> OutPoints)
 {
 	return PCGAttributeAccessorKeys::GetKeys(Points, InStart, OutPoints, [](FPCGPoint& Point) -> FPCGPoint* { return &Point; });
@@ -172,6 +178,11 @@ FPCGAttributeAccessorKeysPointsSubset::FPCGAttributeAccessorKeysPointsSubset(TAr
 	, Points(std::move(*reinterpret_cast<TArray<FPCGPoint*>*>(&InPointPtrs)))
 {}
 
+bool FPCGAttributeAccessorKeysPointsSubset::IsClassSupported(const UStruct* InClass) const
+{
+	return PCGAttributeAccessorKeys::IsClassSupported<FPCGPoint>(InClass);
+}
+
 bool FPCGAttributeAccessorKeysPointsSubset::GetPointKeys(int32 InStart, TArrayView<FPCGPoint*> OutPoints)
 {
 	return PCGAttributeAccessorKeys::GetKeys(Points, InStart, OutPoints, [](FPCGPoint* Point) -> FPCGPoint* { return Point; });
@@ -200,4 +211,11 @@ bool FPCGAttributeAccessorKeysPointsSubset::GetMetadataEntryKeys(int32 InStart,
 bool FPCGAttributeAccessorKeysPointsSubset::GetMetadataEntryKeys(int32 InStart, TArrayView<const PCGMetadataEntryKey*> OutEntryKeys) const
 {
 	return PCGAttributeAccessorKeys::GetKeys(Points, InStart, OutEntryKeys, [](const FPCGPoint* Point) -> const PCGMetadataEntryKey* { return &(Point->MetadataEntry); });
+}
+
+/////////////////////////////////////////////////////////
+
+int32 FPCGAttributeAccessorKeysSplineData::GetNum() const
+{
+	return Ptr ? Ptr->SplineStruct.SplineCurves.Position.Points.Num() : 0;
 }
\ No newline at end of file
diff --git a/Engine/Plugins/PCG/Source/PCG/Private/Metadata/Accessors/PCGPropertyAccessor.cpp b/Engine/Plugins/PCG/Source/PCG/Private/Metadata/Accessors/PCGPropertyAccessor.cpp
--- a/Engine/Plugins/PCG/Source/PCG/Private/Metadata/Accessors/PCGPropertyAccessor.cpp
+++ b/Engine/Plugins/PCG/Source/PCG/Private/Metadata/Accessors/PCGPropertyAccessor.cpp
@@ -7,7 +7,7 @@
 TArray<const void*> PCGPropertyAccessor::GetContainerKeys(int32 Index, int32 Range, const IPCGAttributeAccessorKeys& Keys)
 {
 	TArray<const void*> ContainerKeys;
-	ContainerKeys.SetNum(Range);
+	ContainerKeys.SetNumUninitialized(Range);
 	TArrayView<const void*> ContainerKeysView(ContainerKeys);
 	if (!Keys.GetKeys(Index, ContainerKeysView))
 	{
@@ -20,7 +20,7 @@ TArray<const void*> PCGPropertyAccessor::GetContainerKeys(int32 Index, int32 Ran
 TArray<void*> PCGPropertyAccessor::GetContainerKeys(int32 Index, int32 Range, IPCGAttributeAccessorKeys& Keys)
 {
 	TArray<void*> ContainerKeys;
-	ContainerKeys.SetNum(Range);
+	ContainerKeys.SetNumUninitialized(Range);
 	TArrayView<void*> ContainerKeysView(ContainerKeys);
 	if (!Keys.GetKeys(Index, ContainerKeysView))
 	{
@@ -31,7 +31,7 @@ TArray<void*> PCGPropertyAccessor::GetContainerKeys(int32 Index, int32 Range, IP
 }
 
 IPCGPropertyChain::IPCGPropertyChain(const FProperty* Property, TArray<const FProperty*>&& ExtraProperties)
-	: PropertyChain(std::forward<TArray<const FProperty*>>(ExtraProperties))
+	: PropertyChain(std::move(ExtraProperties))
 {
 	// Fix property chain
 	if (PropertyChain.IsEmpty() || PropertyChain.Last() != Property)
@@ -40,6 +40,18 @@ IPCGPropertyChain::IPCGPropertyChain(const FProperty* Property, TArray<const FPr
 	}
 }
 
+const UStruct* IPCGPropertyChain::GetTopPropertyStruct() const
+{
+	if (!PropertyChain.IsEmpty() && PropertyChain[0])
+	{
+		return PropertyChain[0]->GetOwnerStruct();
+	}
+	else
+	{
+		return nullptr;
+	}
+}
+
 bool FPCGEnumPropertyAccessor::GetRangeImpl(TArrayView<int64> OutValues, int32 Index, const IPCGAttributeAccessorKeys& Keys) const
 {
 	return PCGPropertyAccessor::IterateGet(GetPropertyChain(), OutValues, Index, Keys, [this](const void* PropertyAddressData) -> Type
diff --git a/Engine/Plugins/PCG/Source/PCG/Private/PCGModule.cpp b/Engine/Plugins/PCG/Source/PCG/Private/PCGModule.cpp
--- a/Engine/Plugins/PCG/Source/PCG/Private/PCGModule.cpp
+++ b/Engine/Plugins/PCG/Source/PCG/Private/PCGModule.cpp
@@ -6,6 +6,7 @@
 #include "PCGElement.h"
 #include "Data/PCGBasePointData.h"
 #include "Data/PCGPointData.h"
+#include "Data/PCGSplineData.h"
 
 #include "Interfaces/IPluginManager.h"
 #include "Misc/CoreDelegates.h"
@@ -52,6 +53,7 @@ void FPCGModule::StartupModule()
 	AttributeAccessorFactory.RegisterMethods<UPCGBasePointData>(UPCGBasePointData::GetPointAccessorMethods());
 	// @todo_pcg: Eventually remove the UPCGPointData method registration because the UPCGBasePointData accessors are compatible
 	AttributeAccessorFactory.RegisterMethods<UPCGPointData>(UPCGPointData::GetPointAccessorMethods());
+	AttributeAccessorFactory.RegisterMethods<UPCGSplineData>(UPCGSplineData::GetSplineAccessorMethods());
 
 	// Register onto the PreExit, because we need the class to be still valid to remove them from the mapping
 	FCoreDelegates::OnPreExit.AddRaw(this, &FPCGModule::PreExit);
@@ -65,8 +67,9 @@ void FPCGModule::ShutdownModule()
 void FPCGModule::PreExit()
 {
 	// Unregistering accessor methods
-	AttributeAccessorFactory.UnregisterDefaultMethods();
+	AttributeAccessorFactory.UnregisterMethods<UPCGSplineData>();
 	AttributeAccessorFactory.UnregisterMethods<UPCGPointData>();
+	AttributeAccessorFactory.UnregisterDefaultMethods();
 	AttributeAccessorFactory.UnregisterMethods<UPCGBasePointData>();
 
 #if WITH_EDITOR
diff --git a/Engine/Plugins/PCG/Source/PCG/Public/Data/PCGSplineData.h b/Engine/Plugins/PCG/Source/PCG/Public/Data/PCGSplineData.h
--- a/Engine/Plugins/PCG/Source/PCG/Public/Data/PCGSplineData.h
+++ b/Engine/Plugins/PCG/Source/PCG/Public/Data/PCGSplineData.h
@@ -15,6 +15,11 @@ class UPCGSpatialData;
 class USplineComponent;
 class UPCGSurfaceData;
 
+namespace PCGSplineData
+{
+	const FName ControlPointLayerName = "ControlPoints";
+}
+
 UCLASS(BlueprintType, ClassGroup = (Procedural))
 class PCG_API UPCGSplineData : public UPCGPolyLineData
 {
@@ -45,6 +50,12 @@ public:
 	virtual FTransform GetTransformAtAlpha(float Alpha) const override;
 	virtual bool IsClosed() const override { return SplineStruct.bClosedLoop; }
 	//~End UPCGPolyLineData interface
+	
+	/** Static helper to create an accessor on a data that doesn't yet exist, as accessors for spline data don't rely on existing data. */
+	static TUniquePtr<IPCGAttributeAccessor> CreateStaticAccessor(const FPCGAttributePropertySelector& InSelector, bool bQuiet = false);
+
+	// Get the functions to the accessor factory
+	static FPCGAttributeAccessorMethods GetSplineAccessorMethods();
 
 	//~Begin UPCGSpatialDataWithPointCache interface
 	virtual const UPCGPointData* CreatePointData(FPCGContext* Context) const override;
diff --git a/Engine/Plugins/PCG/Source/PCG/Public/Data/PCGSplineStruct.h b/Engine/Plugins/PCG/Source/PCG/Public/Data/PCGSplineStruct.h
--- a/Engine/Plugins/PCG/Source/PCG/Public/Data/PCGSplineStruct.h
+++ b/Engine/Plugins/PCG/Source/PCG/Public/Data/PCGSplineStruct.h
@@ -7,6 +7,17 @@
 
 #include "PCGSplineStruct.generated.h"
 
+UENUM(Blueprintable)
+enum class EPCGSplineStructProperties : uint8
+{
+	Position       UMETA(Tooltip = "Location component of the control point's transform."),
+	Rotation       UMETA(Tooltip = "Rotation component of the control point's transform."),
+	Scale          UMETA(Tooltip = "Scale component of the control point's transform."),
+	ArriveTangent  UMETA(Tooltip = "Arrive Tangent at the control point."),
+	LeaveTangent   UMETA(Tooltip = "Leave Tangent at the control point."),
+	InterpType     UMETA(Tooltip = "Type of interpolation at the control point for position (same as set on the spline control point). Underlying enum: ESplinePointType.")
+};
+
 /** Subset of the Spline Component API in a standalone struct */
 USTRUCT()
 struct PCG_API FPCGSplineStruct
diff --git a/Engine/Plugins/PCG/Source/PCG/Public/Metadata/Accessors/PCGAttributeAccessorFactory.h b/Engine/Plugins/PCG/Source/PCG/Public/Metadata/Accessors/PCGAttributeAccessorFactory.h
--- a/Engine/Plugins/PCG/Source/PCG/Public/Metadata/Accessors/PCGAttributeAccessorFactory.h
+++ b/Engine/Plugins/PCG/Source/PCG/Public/Metadata/Accessors/PCGAttributeAccessorFactory.h
@@ -16,9 +16,17 @@ class IPCGAttributeAccessorKeys;
 class UPCGData;
 class UPCGMetadata;
 
+#if WITH_EDITOR
 struct FPCGAttributeSelectorMenuEntry
 {
-	/** Display name in the menu. Can use "|" to indicate a hierarchy of menus */
+	FPCGAttributeSelectorMenuEntry(FText InDisplayName, FText InTooltip, FPCGAttributePropertySelector InSelector, TArray<FText> InMenuHierarchy = {})
+		: DisplayName(std::move(InDisplayName))
+		, Tooltip(std::move(InTooltip))
+		, Selector(std::move(InSelector))
+		, MenuHierarchy(std::move(InMenuHierarchy))
+	{}
+
+	/** Display name for the entry. */
 	FText DisplayName;
 
 	/** Tooltip to display for the entry. */
@@ -26,7 +34,11 @@ struct FPCGAttributeSelectorMenuEntry
 
 	/** Selector to copy if this entry is selected. */
 	FPCGAttributePropertySelector Selector;
+	
+	/** Hierarchy of the menus, for all the sub menus to reach this entry. */
+	TArray<FText> MenuHierarchy;
 };
+#endif // WITH_EDITOR
 
 struct FPCGAttributeAccessorMethods
 {
@@ -35,8 +47,15 @@ struct FPCGAttributeAccessorMethods
 	TFunction<TUniquePtr<IPCGAttributeAccessorKeys>(UPCGData*, const FPCGAttributePropertySelector&, const bool)> CreateAccessorKeysFunc;
 	TFunction<TUniquePtr<const IPCGAttributeAccessorKeys>(const UPCGData*, const FPCGAttributePropertySelector&, const bool)> CreateConstAccessorKeysFunc;
 
+#if WITH_EDITOR
+	template <typename EnumType>
+	void FillSelectorMenuEntryFromEnum(const TArray<FText>& InMenuHierarchy = {}) { FillSelectorMenuEntryFromEnum(StaticEnum<EnumType>(), InMenuHierarchy); }
+
+	PCG_API void FillSelectorMenuEntryFromEnum(const UEnum* EnumType, const TArray<FText>& InMenuHierarchy = {});
+	
 	// Special option to gather all the possible options for a given data type in the Selector context menu.
 	TArray<FPCGAttributeSelectorMenuEntry> AttributeSelectorOptions;
+#endif // WITH_EDITOR
 };
 
 class FPCGAttributeAccessorFactory
@@ -59,8 +78,10 @@ public:
 	/** Create a simple key based on the data passed as input. Internal function deliberately not exposed. Use PCGAttributeAccessorHelpers::CreateConstKeys. */
 	TUniquePtr<const IPCGAttributeAccessorKeys> CreateSimpleConstKeys(const UPCGData* InData, const FPCGAttributePropertySelector& InSelector, const bool bQuiet) const;
 
+#if WITH_EDITOR
 	/** Call a Callback on all the FPCGAttributeSelectorMenuEntry currently registered. */
 	PCG_API void ForEachSelectorMenuEntry(TFunctionRef<void(const FPCGAttributeSelectorMenuEntry&)> Callback) const;
+#endif // WITH_EDITOR
 	
 	template <typename T, typename std::enable_if_t<std::is_base_of_v<UPCGData, T>, bool> = false>
 	void RegisterMethods(FPCGAttributeAccessorMethods&& InAccessorMethods) { RegisterMethods(T::StaticClass(), std::move(InAccessorMethods)); }
diff --git a/Engine/Plugins/PCG/Source/PCG/Public/Metadata/Accessors/PCGAttributeAccessorKeys.h b/Engine/Plugins/PCG/Source/PCG/Public/Metadata/Accessors/PCGAttributeAccessorKeys.h
--- a/Engine/Plugins/PCG/Source/PCG/Public/Metadata/Accessors/PCGAttributeAccessorKeys.h
+++ b/Engine/Plugins/PCG/Source/PCG/Public/Metadata/Accessors/PCGAttributeAccessorKeys.h
@@ -4,12 +4,11 @@
 
 #include "Metadata/PCGMetadataCommon.h"
 
-
+class FPCGMetadataAttributeBase;
 class FPCGMetadataLayer;
 struct FPCGPoint;
-
-class FPCGMetadataAttributeBase;
 class UPCGMetadata;
+class UPCGSplineData;
 
 ///////////////////////////////////////////////////////////////////////
 
@@ -40,6 +39,26 @@ namespace PCGAttributeAccessorKeys
 
 		return true;
 	}
+
+	/**
+	 * Utility wrapper around IsChildOf to check if a class passed as argument is compatible with the templated class, which is what the keys are storing.
+	 * Can only be checked if ObjectType is a UClass or UStruct.
+	 * @tparam ObjectType Class of the underlying object. For example for the points keys, ObjectType will be FPCGPoint.
+	 * @param InClass Class of the requested object. A point property could check that the keys is supporting FPCGPoint.
+	 * @return True if it is supported. False otherwise.
+	 */
+	template <typename ObjectType>
+	bool IsClassSupported(const UStruct* InClass)
+	{		
+		if constexpr (TModels_V<CStaticClassProvider, ObjectType> || TModels_V<CStaticStructProvider, ObjectType>)
+		{
+			return InClass && InClass->IsChildOf<ObjectType>();
+		}
+		else
+		{
+			return false;
+		}
+	}
 }
 
 ///////////////////////////////////////////////////////////////////////
@@ -129,6 +148,9 @@ public:
 
 	bool IsReadOnly() const { return bIsReadOnly; }
 
+	/** Returns true if GetGenericObjectKeys would return this class/struct */
+	virtual bool IsClassSupported(const UStruct* InClass) const { return false; }
+
 protected:
 	virtual bool GetPointKeys(int32 InStart, TArrayView<FPCGPoint*> OutPoints) { return false; }
 	virtual bool GetPointKeys(int32 InStart, TArrayView<const FPCGPoint*> OutPoints) const { return false; }
@@ -192,6 +214,8 @@ public:
 
 	virtual int32 GetNum() const override { return Points.Num(); }
 
+	virtual bool IsClassSupported(const UStruct* InClass) const override;
+
 protected:
 	virtual bool GetPointKeys(int32 InStart, TArrayView<FPCGPoint*> OutPoints) override;
 	virtual bool GetPointKeys(int32 InStart, TArrayView<const FPCGPoint*> OutPoints) const override;
@@ -221,6 +245,8 @@ public:
 
 	virtual int32 GetNum() const override { return Points.Num(); }
 
+	virtual bool IsClassSupported(const UStruct* InClass) const override;
+
 protected:
 	virtual bool GetPointKeys(int32 InStart, TArrayView<FPCGPoint*> OutPoints) override;
 	virtual bool GetPointKeys(int32 InStart, TArrayView<const FPCGPoint*> OutPoints) const override;
@@ -265,6 +291,11 @@ public:
 
 	virtual int32 GetNum() const override { return Objects.Num(); }
 
+	virtual bool IsClassSupported(const UStruct* InClass) const override
+	{
+		return PCGAttributeAccessorKeys::IsClassSupported<ObjectType>(InClass);
+	}
+
 protected:
 	virtual bool GetGenericObjectKeys(int32 InStart, TArrayView<void*> OutObjects) override
 	{
@@ -310,6 +341,11 @@ public:
 
 	virtual int32 GetNum() const override { return 1; }
 
+	virtual bool IsClassSupported(const UStruct* InClass) const override
+	{
+		return PCGAttributeAccessorKeys::IsClassSupported<ObjectType>(InClass);
+	}
+
 protected:
 	virtual bool GetGenericObjectKeys(int32 InStart, TArrayView<void*> OutObjects) override
 	{
@@ -381,6 +417,19 @@ protected:
 	TArrayView<void*> Ptrs;
 };
 
+/////////////////////////////////////////////////////////////////
+
+/**
+* Unique Key around a single spline data.
+*/
+class FPCGAttributeAccessorKeysSplineData : public FPCGAttributeAccessorKeysSingleObjectPtr<UPCGSplineData>
+{
+public:
+	using FPCGAttributeAccessorKeysSingleObjectPtr::FPCGAttributeAccessorKeysSingleObjectPtr;
+	
+	virtual int32 GetNum() const override;
+};
+
 template <typename ObjectType>
 inline bool IPCGAttributeAccessorKeys::GetKeys(int32 InStart, TArrayView<ObjectType*> OutKeys)
 {
diff --git a/Engine/Plugins/PCG/Source/PCG/Public/Metadata/Accessors/PCGPropertyAccessor.h b/Engine/Plugins/PCG/Source/PCG/Public/Metadata/Accessors/PCGPropertyAccessor.h
--- a/Engine/Plugins/PCG/Source/PCG/Public/Metadata/Accessors/PCGPropertyAccessor.h
+++ b/Engine/Plugins/PCG/Source/PCG/Public/Metadata/Accessors/PCGPropertyAccessor.h
@@ -19,7 +19,7 @@ namespace PCGPropertyAccessor
 	TArray<void*> GetContainerKeys(int32 Index, int32 Range, IPCGAttributeAccessorKeys& Keys);
 
 	template <typename T>
-	void AddressOffset(const TArray<const FProperty*>& InProperties, TArray<T>& InContainerKeys)
+	void AddressOffset(const TArray<const FProperty*>& InProperties, TArrayView<T> InContainerKeys)
 	{
 		for (int32 j = 0; j < InProperties.Num(); ++j)
 		{
@@ -98,6 +98,9 @@ protected:
 
 	const TArray<const FProperty*>& GetPropertyChain() const { return PropertyChain; }
 
+	/** Return the owner class/struct of the top property (if it is a class/struct). Useful to make sure the top property matches the keys type. */
+	const UStruct* GetTopPropertyStruct() const;
+
 private:
 	TArray<const FProperty*> PropertyChain;
 };
@@ -117,7 +120,7 @@ public:
 
 	FPCGNumericPropertyAccessor(const FNumericProperty* InProperty, TArray<const FProperty*>&& ExtraProperties = {})
 		: Super(/*bInReadOnly=*/ false)
-		, IPCGPropertyChain(InProperty, std::forward<TArray<const FProperty*>>(ExtraProperties))
+		, IPCGPropertyChain(InProperty, std::move(ExtraProperties))
 		, Property(InProperty)
 	{
 		static_assert(std::is_integral_v<T> || std::is_floating_point_v<T>);
@@ -172,7 +175,7 @@ public:
 
 	FPCGEnumPropertyAccessor(const FEnumProperty* InProperty, TArray<const FProperty*>&& ExtraProperties = {})
 		: Super(/*bInReadOnly=*/ false)
-		, IPCGPropertyChain(InProperty, std::forward<TArray<const FProperty*>>(ExtraProperties))
+		, IPCGPropertyChain(InProperty, std::move(ExtraProperties))
 		, Property(InProperty)
 	{
 		check(Property);
@@ -201,7 +204,7 @@ public:
 
 	FPCGPropertyStructAccessor(const FStructProperty* InProperty, TArray<const FProperty*>&& ExtraProperties = {})
 		: Super(/*bInReadOnly=*/ false)
-		, IPCGPropertyChain(InProperty, std::forward<TArray<const FProperty*>>(ExtraProperties))
+		, IPCGPropertyChain(InProperty, std::move(ExtraProperties))
 		, Property(InProperty)
 	{
 		static_assert(PCG::Private::IsPCGType<T>());
@@ -243,7 +246,7 @@ public:
 
 	FPCGPropertyAccessor(const PropertyType* InProperty, TArray<const FProperty*>&& ExtraProperties = {})
 		: Super(/*bInReadOnly=*/ false)
-		, IPCGPropertyChain(InProperty, std::forward<TArray<const FProperty*>>(ExtraProperties))
+		, IPCGPropertyChain(InProperty, std::move(ExtraProperties))
 		, Property(InProperty)
 	{
 		static_assert(PCG::Private::IsPCGType<T>());
@@ -285,7 +288,7 @@ public:
 
 	FPCGPropertyPathAccessor(const FProperty* InProperty, TArray<const FProperty*>&& ExtraProperties = {})
 		: Super(/*bInReadOnly=*/ false)
-		, IPCGPropertyChain(InProperty, std::forward<TArray<const FProperty*>>(ExtraProperties))
+		, IPCGPropertyChain(InProperty, std::move(ExtraProperties))
 		, Property(InProperty)
 	{
 		static_assert(std::is_same_v<FSoftObjectPath, T> || std::is_same_v<FSoftClassPath, T>);
@@ -325,7 +328,7 @@ public:
 
 	FPCGPropertySoftObjectPathAccessor(const FSoftObjectProperty* InProperty, TArray<const FProperty*>&& ExtraProperties = {})
 		: Super(/*bInReadOnly=*/ false)
-		, IPCGPropertyChain(InProperty, std::forward<TArray<const FProperty*>>(ExtraProperties))
+		, IPCGPropertyChain(InProperty, std::move(ExtraProperties))
 		, Property(InProperty)
 	{
 		check(Property);
@@ -364,7 +367,7 @@ public:
 
 	FPCGPropertySoftClassPathAccessor(const FSoftClassProperty* InProperty, TArray<const FProperty*>&& ExtraProperties = {})
 		: Super(/*bInReadOnly=*/ false)
-		, IPCGPropertyChain(InProperty, std::forward<TArray<const FProperty*>>(ExtraProperties))
+		, IPCGPropertyChain(InProperty, std::move(ExtraProperties))
 		, Property(InProperty)
 	{
 		check(Property);
@@ -405,7 +408,7 @@ public:
 
 	FPCGPropertyObjectPtrAccessor(const PropertyType* InProperty, TArray<const FProperty*>&& ExtraProperties = {})
 		: Super(/*bInReadOnly=*/ false)
-		, IPCGPropertyChain(InProperty, std::forward<TArray<const FProperty*>>(ExtraProperties))
+		, IPCGPropertyChain(InProperty, std::move(ExtraProperties))
 		, Property(InProperty)
 	{
 		// Making sure it is the right properties.
@@ -445,7 +448,7 @@ public:
 
 	FPCGAttributePropertySelectorAccessor(const FStructProperty* InProperty, TArray<const FProperty*>&& ExtraProperties = {})
 		: Super(/*bInReadOnly=*/ false)
-		, IPCGPropertyChain(InProperty, std::forward<TArray<const FProperty*>>(ExtraProperties))
+		, IPCGPropertyChain(InProperty, std::move(ExtraProperties))
 		, Property(InProperty)
 	{
 		ensure(InProperty && InProperty->Struct && InProperty->Struct->IsA(FPCGAttributePropertySelector::StaticStruct()->GetClass()));
@@ -483,7 +486,7 @@ public:
 
 	FPCGLinearColorAccessor(const FStructProperty* InProperty, TArray<const FProperty*>&& ExtraProperties = {})
 		: Super(/*bInReadOnly=*/ false)
-		, IPCGPropertyChain(InProperty, std::forward<TArray<const FProperty*>>(ExtraProperties))
+		, IPCGPropertyChain(InProperty, std::move(ExtraProperties))
 		, Property(InProperty)
 	{
 		ensure(InProperty && InProperty->Struct && InProperty->Struct == TBaseStructure<FLinearColor>::Get());
@@ -522,7 +525,7 @@ public:
 
 	FPCGColorAccessor(const FStructProperty* InProperty, TArray<const FProperty*>&& ExtraProperties = {})
 		: Super(/*bInReadOnly=*/ false)
-		, IPCGPropertyChain(InProperty, std::forward<TArray<const FProperty*>>(ExtraProperties))
+		, IPCGPropertyChain(InProperty, std::move(ExtraProperties))
 		, Property(InProperty)
 	{
 		ensure(InProperty && InProperty->Struct && InProperty->Struct == TBaseStructure<FColor>::Get());
diff --git a/Engine/Plugins/PCG/Source/PCG/Public/Metadata/Accessors/PCGSplineAccessor.h b/Engine/Plugins/PCG/Source/PCG/Public/Metadata/Accessors/PCGSplineAccessor.h
new file mode 100644
--- /dev/null
+++ b/Engine/Plugins/PCG/Source/PCG/Public/Metadata/Accessors/PCGSplineAccessor.h
@@ -0,0 +1,147 @@
+// Copyright Epic Games, Inc. All Rights Reserved.
+
+#pragma once
+
+#include "PCGPropertyAccessor.h"
+
+#include "Components/SplineComponent.h"
+
+enum class EPCGInterpCurveAccessorTarget
+{
+	Value,
+	ArriveTangent,
+	LeaveTangent,
+	InterpMode
+};
+
+/**
+* Templated accessor for any interp curve. It's important that the keys only have a single value, the struct that holds the spline curve, since
+* the interp curve is basically an array.
+* Keys supported: PCGSplineData, FPCGSplineStruct, FSplineCurves
+*/
+template <typename CurveType,
+	EPCGInterpCurveAccessorTarget Target = EPCGInterpCurveAccessorTarget::Value,
+	std::enable_if_t<std::is_base_of_v<FInterpCurve<typename CurveType::ElementType>, CurveType>, int> = 0>
+class FPCGInterpCurveAccessor : public IPCGAttributeAccessorT<FPCGInterpCurveAccessor<CurveType, Target>>, IPCGPropertyChain
+{
+public:
+	// The underlying type is either the Curve element type or the same as the enum property accessor if the target is the InterpMode, as it is an enum below.
+	using Type = std::conditional_t<Target == EPCGInterpCurveAccessorTarget::InterpMode, FPCGEnumPropertyAccessor::Type, typename CurveType::ElementType>;
+	using Super = IPCGAttributeAccessorT<FPCGInterpCurveAccessor<CurveType, Target>>;
+
+	FPCGInterpCurveAccessor(const FStructProperty* InProperty, TArray<const FProperty*>&& ExtraProperties = {})
+		: Super(/*bInReadOnly=*/ false)
+		, IPCGPropertyChain(InProperty, std::move(ExtraProperties))
+		, Property(InProperty)
+	{
+		static_assert(PCG::Private::IsPCGType<Type>());		
+		check(Property && Property->Struct->IsChildOf(TBaseStructure<CurveType>::Get()));
+		TopPropertyStruct = GetTopPropertyStruct();
+	}
+
+	bool GetRangeImpl(TArrayView<Type> OutValues, int32 Index, const IPCGAttributeAccessorKeys& Keys) const
+	{
+		// Spline struct contains an array, so there is only a single address, to the array.
+		const void* ContainerKeys = nullptr;
+		TArrayView<const void*> ContainerKeysView(&ContainerKeys, 1);
+		if (!Keys.GetKeys(Index, ContainerKeysView))
+		{
+			return false;
+		}
+
+		// Validation to not access keys that are not the expected type. Done after the GetKeys, as we also want to discard other type of incompatible
+		// keys (like a Default Metadata entry key)
+		if (!ensure(Keys.IsClassSupported(TopPropertyStruct)))
+		{
+			return false;
+		}
+
+		PCGPropertyAccessor::AddressOffset(GetPropertyChain(), ContainerKeysView);
+
+		const CurveType& InterpCurve = *static_cast<const CurveType*>(ContainerKeys);
+		const int32 NumPoints = InterpCurve.Points.Num();
+		
+		if (NumPoints == 0)
+		{
+			return false;
+		}
+		
+		for (int32 i = 0; i < OutValues.Num(); ++i)
+		{
+			const int32 CurrIndex = (Index + i) % NumPoints;
+			if constexpr (Target == EPCGInterpCurveAccessorTarget::InterpMode)
+			{
+				OutValues[i] = static_cast<Type>(ConvertInterpCurveModeToSplinePointType(InterpCurve.Points[CurrIndex].InterpMode));
+			}
+			else if constexpr (Target == EPCGInterpCurveAccessorTarget::ArriveTangent)
+			{
+				OutValues[i] = InterpCurve.Points[CurrIndex].ArriveTangent;
+			}
+			else if constexpr (Target == EPCGInterpCurveAccessorTarget::LeaveTangent)
+			{
+				OutValues[i] = InterpCurve.Points[CurrIndex].LeaveTangent;
+			}
+			else // if constexpr (Target == EPCGInterpCurveAccessorTarget::Value)
+			{
+				OutValues[i] = InterpCurve.Points[CurrIndex].OutVal;
+			}
+		}
+
+		return true;
+	}
+
+	bool SetRangeImpl(TArrayView<const Type> InValues, int32 Index, IPCGAttributeAccessorKeys& Keys, EPCGAttributeAccessorFlags)
+	{		
+		// Spline struct contains an array, so there is only a single address, to the array.
+		void* ContainerKeys = nullptr;
+		TArrayView<void*> ContainerKeysView(&ContainerKeys, 1);
+		if (!Keys.GetKeys(Index, ContainerKeysView))
+		{
+			return false;
+		}
+
+		// Validation to not access keys that are not the expected type. Done after the GetKeys, as we also want to discard other type of incompatible
+		// keys (like a Default Metadata entry key)
+		if (!ensure(Keys.IsClassSupported(TopPropertyStruct)))
+		{
+			return false;
+		}
+
+		PCGPropertyAccessor::AddressOffset(GetPropertyChain(), ContainerKeysView);
+
+		CurveType& InterpCurve = *static_cast<CurveType*>(ContainerKeys);
+		const int32 NumPoints = InterpCurve.Points.Num();
+		if (NumPoints == 0)
+		{
+			return false;
+		}
+		
+		for (int32 i = 0; i < InValues.Num(); ++i)
+		{
+			const int32 CurrIndex = (Index + i) % NumPoints;
+
+			if constexpr (Target == EPCGInterpCurveAccessorTarget::InterpMode)
+			{
+				InterpCurve.Points[CurrIndex].InterpMode = ConvertSplinePointTypeToInterpCurveMode(static_cast<ESplinePointType::Type>(InValues[i]));
+			}
+			else if constexpr (Target == EPCGInterpCurveAccessorTarget::ArriveTangent)
+			{
+				InterpCurve.Points[CurrIndex].ArriveTangent = InValues[i];
+			}
+			else if constexpr (Target == EPCGInterpCurveAccessorTarget::LeaveTangent)
+			{
+				InterpCurve.Points[CurrIndex].LeaveTangent = InValues[i];
+			}
+			else // if constexpr (Target == EPCGInterpCurveAccessorTarget::Value)
+			{
+				InterpCurve.Points[CurrIndex].OutVal = InValues[i];
+			}
+		}
+
+		return true;
+	}
+
+private:
+	const FStructProperty* Property = nullptr;
+	const UStruct* TopPropertyStruct = nullptr;
+};
\ No newline at end of file
