[PCG] Enable PCGPointArrayData path via CVar (pcg.Experimental.EnablePointArrayData)
- Add support for PointArrayData to CreatePointGrid, BoundsModifier (FPCGPointOperationElementBase), SurfaceSampler nodes
- FPCGContext::NewPointData_AnyThread will allocate the proper point data sub type based on CVar
- UPCGSpatialData::ToBasePointData() will properly forward to ToPointData or ToPointArrayData

#rnx

diff --git a/Engine/Plugins/PCG/Source/PCG/Private/Data/PCGBasePointData.cpp b/Engine/Plugins/PCG/Source/PCG/Private/Data/PCGBasePointData.cpp
--- a/Engine/Plugins/PCG/Source/PCG/Private/Data/PCGBasePointData.cpp
+++ b/Engine/Plugins/PCG/Source/PCG/Private/Data/PCGBasePointData.cpp
@@ -2,6 +2,9 @@
 
 #include "Data/PCGBasePointData.h"
 
+#include "PCGContext.h"
+#include "Data/PCGPointArrayData.h"
+#include "Data/PCGPointData.h"
 #include "Helpers/PCGHelpers.h"
 #include "Helpers/PCGPointHelpers.h"
 #include "Helpers/PCGTagHelpers.h"
@@ -717,18 +720,61 @@ FBox UPCGBasePointData::GetBounds() const
 	return Bounds;
 }
 
+void UPCGBasePointData::CopyPointsTo(UPCGBasePointData* To, int32 ReadStartIndex, int32 WriteStartIndex, int32 Count) const
+{
+	const TConstPCGValueRange<FTransform> FromTransformRange = GetConstTransformValueRange();
+	const TConstPCGValueRange<float> FromSteepnessRange = GetConstSteepnessValueRange();
+	const TConstPCGValueRange<float> FromDensityRange = GetConstDensityValueRange();
+	const TConstPCGValueRange<FVector> FromBoundsMinRange = GetConstBoundsMinValueRange();
+	const TConstPCGValueRange<FVector> FromBoundsMaxRange = GetConstBoundsMaxValueRange();
+	const TConstPCGValueRange<FVector4> FromColorRange = GetConstColorValueRange();
+	const TConstPCGValueRange<int64> FromMetadataEntryRange = GetConstMetadataEntryValueRange();
+	const TConstPCGValueRange<int32> FromSeedRange = GetConstSeedValueRange();
+	
+	To->AllocateProperties(GetAllocatedProperties());
+
+	TPCGValueRange<FTransform> ToTransformRange = To->GetTransformValueRange(/*bAllocate=*/false);
+	TPCGValueRange<float> ToSteepnessRange = To->GetSteepnessValueRange(/*bAllocate=*/false);
+	TPCGValueRange<float> ToDensityRange = To->GetDensityValueRange(/*bAllocate=*/false);
+	TPCGValueRange<FVector> ToBoundsMinRange = To->GetBoundsMinValueRange(/*bAllocate=*/false);
+	TPCGValueRange<FVector> ToBoundsMaxRange = To->GetBoundsMaxValueRange(/*bAllocate=*/false);
+	TPCGValueRange<FVector4> ToColorRange = To->GetColorValueRange(/*bAllocate=*/false);
+	TPCGValueRange<int64> ToMetadataEntryRange = To->GetMetadataEntryValueRange(/*bAllocate=*/false);
+	TPCGValueRange<int32> ToSeedRange = To->GetSeedValueRange(/*bAllocate=*/false);
+
+	check(ReadStartIndex + Count <= GetNumPoints() && WriteStartIndex + Count <= To->GetNumPoints());
+
+	for (int32 Index = 0; Index < Count; ++Index)
+	{
+		const int32 ReadIndex = ReadStartIndex + Index;
+		const int32 WriteIndex = WriteStartIndex + Index;
+
+		ToTransformRange[WriteIndex] = FromTransformRange[ReadIndex];
+		ToDensityRange[WriteIndex] = FromDensityRange[ReadIndex];
+		ToBoundsMinRange[WriteIndex] = FromBoundsMinRange[ReadIndex];
+		ToBoundsMaxRange[WriteIndex] = FromBoundsMaxRange[ReadIndex];
+		ToColorRange[WriteIndex] = FromColorRange[ReadIndex];
+		ToSteepnessRange[WriteIndex] = FromSteepnessRange[ReadIndex];
+		ToSeedRange[WriteIndex] = FromSeedRange[ReadIndex];
+		ToMetadataEntryRange[WriteIndex] = FromMetadataEntryRange[ReadIndex];
+	}
+
+	To->DirtyCache();
+}
+
 void UPCGBasePointData::CopyPointsFrom(const UPCGBasePointData* InData, const TArray<int32>& InDataIndices)
 {
 	return CopyPoints(InData, this, InDataIndices, false);
 }
 
-void UPCGBasePointData::CopyPoints(const UPCGBasePointData* From, UPCGBasePointData* To, const TArray<int32>& InDataIndices, bool bCopyAll) const
+void UPCGBasePointData::CopyPoints(const UPCGBasePointData* From, UPCGBasePointData* To, const TArray<int32>& InDataIndices, bool bCopyAll)
 {
 	TRACE_CPUPROFILER_EVENT_SCOPE(UPCGBasePointData::CopyPoints);
 	check(From);
 
 	const int32 NumPoints = bCopyAll ? From->GetNumPoints() : InDataIndices.Num();
 	To->SetNumPoints(NumPoints);
+	To->AllocateProperties(From->GetAllocatedProperties());
 
 	const TConstPCGValueRange<FTransform> FromTransformRange = From->GetConstTransformValueRange();
 	const TConstPCGValueRange<float> FromSteepnessRange = From->GetConstSteepnessValueRange();
@@ -739,14 +785,14 @@ void UPCGBasePointData::CopyPoints(const UPCGBasePointData* From, UPCGBasePointD
 	const TConstPCGValueRange<int64> FromMetadataEntryRange = From->GetConstMetadataEntryValueRange();
 	const TConstPCGValueRange<int32> FromSeedRange = From->GetConstSeedValueRange();
 
-	TPCGValueRange<FTransform> ToTransformRange = To->GetTransformValueRange();
-	TPCGValueRange<float> ToSteepnessRange = To->GetSteepnessValueRange();
-	TPCGValueRange<float> ToDensityRange = To->GetDensityValueRange();
-	TPCGValueRange<FVector> ToBoundsMinRange = To->GetBoundsMinValueRange();
-	TPCGValueRange<FVector> ToBoundsMaxRange = To->GetBoundsMaxValueRange();
-	TPCGValueRange<FVector4> ToColorRange = To->GetColorValueRange();
-	TPCGValueRange<int64> ToMetadataEntryRange = To->GetMetadataEntryValueRange();
-	TPCGValueRange<int32> ToSeedRange = To->GetSeedValueRange();
+	TPCGValueRange<FTransform> ToTransformRange = To->GetTransformValueRange(/*bAllocate=*/false);
+	TPCGValueRange<float> ToSteepnessRange = To->GetSteepnessValueRange(/*bAllocate=*/false);
+	TPCGValueRange<float> ToDensityRange = To->GetDensityValueRange(/*bAllocate=*/false);
+	TPCGValueRange<FVector> ToBoundsMinRange = To->GetBoundsMinValueRange(/*bAllocate=*/false);
+	TPCGValueRange<FVector> ToBoundsMaxRange = To->GetBoundsMaxValueRange(/*bAllocate=*/false);
+	TPCGValueRange<FVector4> ToColorRange = To->GetColorValueRange(/*bAllocate=*/false);
+	TPCGValueRange<int64> ToMetadataEntryRange = To->GetMetadataEntryValueRange(/*bAllocate=*/false);
+	TPCGValueRange<int32> ToSeedRange = To->GetSeedValueRange(/*bAllocate=*/false);
 
 	for (int32 Index = 0; Index < NumPoints; ++Index)
 	{
diff --git a/Engine/Plugins/PCG/Source/PCG/Private/Data/PCGPointArrayData.cpp b/Engine/Plugins/PCG/Source/PCG/Private/Data/PCGPointArrayData.cpp
--- a/Engine/Plugins/PCG/Source/PCG/Private/Data/PCGPointArrayData.cpp
+++ b/Engine/Plugins/PCG/Source/PCG/Private/Data/PCGPointArrayData.cpp
@@ -31,3 +31,15 @@ UPCGSpatialData* UPCGPointArrayData::CopyInternal(FPCGContext* Context) const
 
 	return NewPointData;
 }
+
+void UPCGPointArrayData::CopyPointsTo(UPCGBasePointData* OutData, int32 ReadStartIndex, int32 WriteStartIndex, int32 Count) const
+{
+	if (UPCGPointArrayData* PointArrayData = Cast<UPCGPointArrayData>(OutData))
+	{
+		PointArray.CopyPointsTo(PointArrayData->PointArray, ReadStartIndex, WriteStartIndex, Count);
+	}
+	else
+	{
+		Super::CopyPointsTo(OutData, ReadStartIndex, WriteStartIndex, Count);
+	}
+}
\ No newline at end of file
diff --git a/Engine/Plugins/PCG/Source/PCG/Private/Data/PCGPointData.cpp b/Engine/Plugins/PCG/Source/PCG/Private/Data/PCGPointData.cpp
--- a/Engine/Plugins/PCG/Source/PCG/Private/Data/PCGPointData.cpp
+++ b/Engine/Plugins/PCG/Source/PCG/Private/Data/PCGPointData.cpp
@@ -362,6 +362,18 @@ UPCGSpatialData* UPCGPointData::CopyInternal(FPCGContext* Context) const
 	return NewPointData;
 }
 
+void UPCGPointData::CopyPointsTo(UPCGBasePointData* OutData, int32 ReadStartIndex, int32 WriteStartIndex, int32 Count) const
+{
+	if (UPCGPointData* OutPointData = Cast<UPCGPointData>(OutData))
+	{
+		std::memcpy(&OutPointData->GetMutablePoints()[WriteStartIndex], &Points[ReadStartIndex], sizeof(FPCGPoint) * Count);
+	}
+	else
+	{
+		Super::CopyPointsTo(OutData, ReadStartIndex, WriteStartIndex, Count);
+	}
+}
+
 void UPCGPointData::MoveRange(int32 RangeStartIndex, int32 MoveToIndex, int32 NumElements)
 {
 	if (RangeStartIndex != MoveToIndex && NumElements > 0)
diff --git a/Engine/Plugins/PCG/Source/PCG/Private/Data/PCGSpatialData.cpp b/Engine/Plugins/PCG/Source/PCG/Private/Data/PCGSpatialData.cpp
--- a/Engine/Plugins/PCG/Source/PCG/Private/Data/PCGSpatialData.cpp
+++ b/Engine/Plugins/PCG/Source/PCG/Private/Data/PCGSpatialData.cpp
@@ -6,6 +6,7 @@
 #include "Data/PCGDifferenceData.h"
 #include "Data/PCGIntersectionData.h"
 #include "Data/PCGPointData.h"
+#include "Data/PCGPointArrayData.h"
 #include "Data/PCGProjectionData.h"
 #include "Data/PCGUnionData.h"
 #include "Elements/PCGExecuteBlueprint.h"
@@ -21,6 +22,36 @@ UPCGSpatialData::UPCGSpatialData(const FObjectInitializer& ObjectInitializer)
 	Metadata->SetupLayer(PCGMetadataLayerID::Data, /*bIsDefault=*/true);
 }
 
+const UPCGBasePointData* UPCGSpatialData::ToBasePointData(FPCGContext* Context) const
+{
+	if (CVarPCGEnablePointArrayData.GetValueOnAnyThread())
+	{
+		return ToPointArrayData(Context);
+	}
+	else
+	{
+		return ToPointData(Context);
+	}
+}
+
+const UPCGPointArrayData* UPCGSpatialData::ToPointArrayData(FPCGContext* Context) const
+{
+	// @todo_pcg: this is sub-optimal but until we implement this in other spatial types (including the cached types) it works.
+	if (const UPCGPointData* PointData = ToPointData(Context))
+	{
+		UPCGPointArrayData* PointArrayData = FPCGContext::NewObject_AnyThread<UPCGPointArrayData>(Context);
+		PointArrayData->InitializeFromData(PointData);
+
+		UPCGBasePointData::CopyPoints(PointData, PointArrayData, {}, /*bCopyAll=*/true);
+
+		return PointArrayData;
+	}
+	else
+	{
+		return nullptr;
+	}
+}
+
 void UPCGSpatialDataWithPointCache::GetResourceSizeEx(FResourceSizeEx& CumulativeResourceSize)
 {
 	Super::GetResourceSizeEx(CumulativeResourceSize);
diff --git a/Engine/Plugins/PCG/Source/PCG/Private/Elements/PCGBoundsModifier.cpp b/Engine/Plugins/PCG/Source/PCG/Private/Elements/PCGBoundsModifier.cpp
--- a/Engine/Plugins/PCG/Source/PCG/Private/Elements/PCGBoundsModifier.cpp
+++ b/Engine/Plugins/PCG/Source/PCG/Private/Elements/PCGBoundsModifier.cpp
@@ -3,7 +3,7 @@
 #include "Elements/PCGBoundsModifier.h"
 
 #include "PCGContext.h"
-#include "PCGPoint.h"
+#include "Data/PCGBasePointData.h"
 
 #include UE_INLINE_GENERATED_CPP_BY_NAME(PCGBoundsModifier)
 
@@ -28,6 +28,17 @@ FText UPCGBoundsModifierSettings::GetNodeTooltipText() const
 }
 #endif // WITH_EDITOR
 
+EPCGPointNativeProperties FPCGBoundsModifier::GetPropertiesToAllocate(const UPCGBasePointData* Input) const
+{
+	EPCGPointNativeProperties PropertiesToAllocate = Input->GetAllocatedProperties();
+	
+	PropertiesToAllocate |= EPCGPointNativeProperties::BoundsMin;
+	PropertiesToAllocate |= EPCGPointNativeProperties::BoundsMax;
+	PropertiesToAllocate |= EPCGPointNativeProperties::Steepness;
+
+	return PropertiesToAllocate;
+}
+
 bool FPCGBoundsModifier::ExecuteInternal(FPCGContext* Context) const
 {
 	TRACE_CPUPROFILER_EVENT_SCOPE(FPCGBoundsModifier::Execute);
@@ -41,76 +52,128 @@ bool FPCGBoundsModifier::ExecuteInternal(FPCGContext* Context) const
 	switch (Settings->Mode)
 	{
 		case EPCGBoundsModifierMode::Intersect:
-		return ExecutePointOperation(BoundsModifierContext, [Settings](const FPCGPoint& InPoint, FPCGPoint& OutPoint)
-		{
-			OutPoint = InPoint;
-			OutPoint.SetLocalBounds(InPoint.GetLocalBounds().Overlap(FBox(Settings->BoundsMin, Settings->BoundsMax)));
-
-			if (Settings->bAffectSteepness)
+			return ExecutePointOperation(BoundsModifierContext, [Settings](const UPCGBasePointData* InputData, UPCGBasePointData* OutputData, int32 StartIndex, int32 Count)
 			{
-				OutPoint.Steepness = FMath::Min(InPoint.Steepness, Settings->Steepness);
-			}
-
-			return true;
-		}, PCGBoundsModifier::PointsPerChunk);
-
-	case EPCGBoundsModifierMode::Include:
-		return ExecutePointOperation(BoundsModifierContext, [Settings](const FPCGPoint& InPoint, FPCGPoint& OutPoint)
-		{
-			OutPoint = InPoint;
-			OutPoint.SetLocalBounds(InPoint.GetLocalBounds() + FBox(Settings->BoundsMin, Settings->BoundsMax));
-
-			if (Settings->bAffectSteepness)
+				InputData->CopyPointsTo(OutputData, StartIndex, StartIndex, Count);
+
+				TPCGValueRange<FVector> BoundsMinRange = OutputData->GetBoundsMinValueRange();
+				TPCGValueRange<FVector> BoundsMaxRange = OutputData->GetBoundsMaxValueRange();
+				TPCGValueRange<float> SteepnessRange = OutputData->GetSteepnessValueRange();
+
+				const FBox SettingsBounds(Settings->BoundsMin, Settings->BoundsMax);
+
+				for (int32 Index = StartIndex; Index < (StartIndex + Count); ++Index)
+				{
+					FVector& BoundsMin = BoundsMinRange[Index];
+					FVector& BoundsMax = BoundsMaxRange[Index];
+					const FBox LocalBounds = PCGPointHelpers::GetLocalBounds(BoundsMin, BoundsMax);
+					PCGPointHelpers::SetLocalBounds(LocalBounds.Overlap(SettingsBounds), BoundsMin, BoundsMax);
+				
+					if (Settings->bAffectSteepness)
+					{
+						SteepnessRange[Index] = FMath::Min(SteepnessRange[Index], Settings->Steepness);
+					}
+				}
+
+				return true;
+			}, PCGBoundsModifier::PointsPerChunk);
+
+		case EPCGBoundsModifierMode::Include:
+			return ExecutePointOperation(BoundsModifierContext, [Settings](const UPCGBasePointData* InputData, UPCGBasePointData* OutputData, int32 StartIndex, int32 Count)
 			{
-				OutPoint.Steepness = FMath::Max(InPoint.Steepness, Settings->Steepness);
-			}
-
-			return true;
-		}, PCGBoundsModifier::PointsPerChunk);
-
-	case EPCGBoundsModifierMode::Translate:
-		return ExecutePointOperation(BoundsModifierContext, [Settings](const FPCGPoint& InPoint, FPCGPoint& OutPoint)
-		{
-			OutPoint = InPoint;
-			OutPoint.BoundsMin += Settings->BoundsMin;
-			OutPoint.BoundsMax += Settings->BoundsMax;
-
-			if (Settings->bAffectSteepness)
+				InputData->CopyPointsTo(OutputData, StartIndex, StartIndex, Count);
+
+				TPCGValueRange<FVector> BoundsMinRange = OutputData->GetBoundsMinValueRange();
+				TPCGValueRange<FVector> BoundsMaxRange = OutputData->GetBoundsMaxValueRange();
+				TPCGValueRange<float> SteepnessRange = OutputData->GetSteepnessValueRange();
+
+				const FBox SettingsBounds(Settings->BoundsMin, Settings->BoundsMax);
+
+				for (int32 Index = StartIndex; Index < (StartIndex + Count); ++Index)
+				{
+					FVector& BoundsMin = BoundsMinRange[Index];
+					FVector& BoundsMax = BoundsMaxRange[Index];
+					const FBox LocalBounds = PCGPointHelpers::GetLocalBounds(BoundsMin, BoundsMax);
+					PCGPointHelpers::SetLocalBounds(LocalBounds + SettingsBounds, BoundsMin, BoundsMax);
+				
+					if (Settings->bAffectSteepness)
+					{
+						SteepnessRange[Index] = FMath::Max(SteepnessRange[Index], Settings->Steepness);
+					}
+				}
+
+				return true;
+			}, PCGBoundsModifier::PointsPerChunk);
+
+		case EPCGBoundsModifierMode::Translate:
+			return ExecutePointOperation(BoundsModifierContext, [Settings](const UPCGBasePointData* InputData, UPCGBasePointData* OutputData, int32 StartIndex, int32 Count)
 			{
-				OutPoint.Steepness = FMath::Clamp(InPoint.Steepness + Settings->Steepness, 0.0f, 1.0f);
-			}
-
-			return true;
-		}, PCGBoundsModifier::PointsPerChunk);
-
-	case EPCGBoundsModifierMode::Scale:
-		return ExecutePointOperation(BoundsModifierContext, [Settings](const FPCGPoint& InPoint, FPCGPoint& OutPoint)
-		{
-			OutPoint = InPoint;
-			OutPoint.BoundsMin *= Settings->BoundsMin;
-			OutPoint.BoundsMax *= Settings->BoundsMax;
-
-			if (Settings->bAffectSteepness)
+				InputData->CopyPointsTo(OutputData, StartIndex, StartIndex, Count);
+
+				TPCGValueRange<FVector> BoundsMinRange = OutputData->GetBoundsMinValueRange();
+				TPCGValueRange<FVector> BoundsMaxRange = OutputData->GetBoundsMaxValueRange();
+				TPCGValueRange<float> SteepnessRange = OutputData->GetSteepnessValueRange();
+
+				for (int32 Index = StartIndex; Index < (StartIndex + Count); ++Index)
+				{
+					BoundsMinRange[Index] += Settings->BoundsMin;
+					BoundsMaxRange[Index] += Settings->BoundsMax;
+
+					if (Settings->bAffectSteepness)
+					{
+						SteepnessRange[Index] = FMath::Clamp(SteepnessRange[Index] + Settings->Steepness, 0.0f, 1.0f);
+					}
+				}
+						
+				return true;
+			}, PCGBoundsModifier::PointsPerChunk);
+
+		case EPCGBoundsModifierMode::Scale:
+			return ExecutePointOperation(BoundsModifierContext, [Settings](const UPCGBasePointData* InputData, UPCGBasePointData* OutputData, int32 StartIndex, int32 Count)
 			{
-				OutPoint.Steepness = FMath::Clamp(InPoint.Steepness * Settings->Steepness, 0.0f, 1.0f);
-			}
-
-			return true;
-		}, PCGBoundsModifier::PointsPerChunk);
-
-	case EPCGBoundsModifierMode::Set:
-		return ExecutePointOperation(BoundsModifierContext, [Settings](const FPCGPoint& InPoint, FPCGPoint& OutPoint)
-		{
-			OutPoint = InPoint;
-			OutPoint.SetLocalBounds(FBox(Settings->BoundsMin, Settings->BoundsMax));
-
-			if (Settings->bAffectSteepness)
+				InputData->CopyPointsTo(OutputData, StartIndex, StartIndex, Count);
+							
+				TPCGValueRange<FVector> BoundsMinRange = OutputData->GetBoundsMinValueRange();
+				TPCGValueRange<FVector> BoundsMaxRange = OutputData->GetBoundsMaxValueRange();
+				TPCGValueRange<float> SteepnessRange = OutputData->GetSteepnessValueRange();
+
+				for (int32 Index = StartIndex; Index < (StartIndex + Count); ++Index)
+				{
+					BoundsMinRange[Index] *= Settings->BoundsMin;
+					BoundsMaxRange[Index] *= Settings->BoundsMax;
+
+					if (Settings->bAffectSteepness)
+					{
+						SteepnessRange[Index] = FMath::Clamp(SteepnessRange[Index] * Settings->Steepness, 0.0f, 1.0f);
+					}
+				}
+			
+				return true;
+			}, PCGBoundsModifier::PointsPerChunk);
+
+		case EPCGBoundsModifierMode::Set:
+			return ExecutePointOperation(BoundsModifierContext, [Settings](const UPCGBasePointData* InputData, UPCGBasePointData* OutputData, int32 StartIndex, int32 Count)
 			{
-				OutPoint.Steepness = Settings->Steepness;
-			}
-
-			return true;
-		}, PCGBoundsModifier::PointsPerChunk);
+				InputData->CopyPointsTo(OutputData, StartIndex, StartIndex, Count);
+
+				TPCGValueRange<FVector> BoundsMinRange = OutputData->GetBoundsMinValueRange();
+				TPCGValueRange<FVector> BoundsMaxRange = OutputData->GetBoundsMaxValueRange();
+				TPCGValueRange<float> SteepnessRange = OutputData->GetSteepnessValueRange();
+
+				const FBox SettingsBounds(Settings->BoundsMin, Settings->BoundsMax);
+
+				for (int32 Index = StartIndex; Index < (StartIndex + Count); ++Index)
+				{
+					PCGPointHelpers::SetLocalBounds(SettingsBounds, BoundsMinRange[Index], BoundsMaxRange[Index]);
+				
+					if (Settings->bAffectSteepness)
+					{
+						SteepnessRange[Index] = Settings->Steepness;
+					}
+				}
+
+				return true;
+			}, PCGBoundsModifier::PointsPerChunk);
 
 		default:
 			checkNoEntry();
diff --git a/Engine/Plugins/PCG/Source/PCG/Private/Elements/PCGCreatePointsGrid.cpp b/Engine/Plugins/PCG/Source/PCG/Private/Elements/PCGCreatePointsGrid.cpp
--- a/Engine/Plugins/PCG/Source/PCG/Private/Elements/PCGCreatePointsGrid.cpp
+++ b/Engine/Plugins/PCG/Source/PCG/Private/Elements/PCGCreatePointsGrid.cpp
@@ -216,7 +216,7 @@ bool FPCGCreatePointsGridElement::ExecuteInternal(FPCGContext* Context) const
 
 	auto InitFunc = [Context, Settings, CellSize, NumIterations, &PtData, &TransformRange, &SeedRange]()
 	{ 
-		PtData = FPCGContext::NewObject_AnyThread<UPCGPointData>(Context);
+		PtData = FPCGContext::NewPointData_AnyThread(Context);
 		
 		check(PtData);
 				
diff --git a/Engine/Plugins/PCG/Source/PCG/Private/Elements/PCGPointOperationElementBase.cpp b/Engine/Plugins/PCG/Source/PCG/Private/Elements/PCGPointOperationElementBase.cpp
--- a/Engine/Plugins/PCG/Source/PCG/Private/Elements/PCGPointOperationElementBase.cpp
+++ b/Engine/Plugins/PCG/Source/PCG/Private/Elements/PCGPointOperationElementBase.cpp
@@ -2,6 +2,7 @@
 
 #include "Elements/PCGPointOperationElementBase.h"
 
+#include "Data/PCGBasePointData.h"
 #include "Data/PCGPointData.h"
 #include "Helpers/PCGAsync.h"
 
@@ -44,22 +45,27 @@ bool FPCGPointOperationElementBase::PreparePointOperationData(ContextType* InCon
 				return EPCGTimeSliceInitResult::NoOperation;
 			}
 
-			OutState.InputPointData = SpatialData->ToPointData(InContext);
-
-			if (!OutState.InputPointData)
+			OutState.InputData = SupportsBasePointDataInputs(InContext) ? SpatialData->ToBasePointData(InContext) : SpatialData->ToPointData(InContext);
+			if (!OutState.InputData)
 			{
 				PCGLog::LogErrorOnGraph(LOCTEXT("InputMissingPointData", "Unable to get Point data from input"), InContext);
 				return EPCGTimeSliceInitResult::NoOperation;
 			}
 
-			OutState.NumPoints = OutState.InputPointData->GetPoints().Num();
+			OutState.NumPoints = OutState.InputData->GetNumPoints();
 
 			// Create and initialize the output points
-			OutState.OutputPointData = FPCGContext::NewObject_AnyThread<UPCGPointData>(InContext);
-			OutState.OutputPointData->InitializeFromData(OutState.InputPointData);
-			OutState.OutputPointData->GetMutablePoints().SetNumUninitialized(OutState.NumPoints);
+			OutState.OutputData = SupportsBasePointDataInputs(InContext) ? FPCGContext::NewPointData_AnyThread(InContext) : FPCGContext::NewObject_AnyThread<UPCGPointData>(InContext);
+			OutState.OutputData->InitializeFromData(OutState.InputData);
+			OutState.OutputData->SetNumPoints(OutState.NumPoints);
+			OutState.OutputData->AllocateProperties(GetPropertiesToAllocate(OutState.InputData));
+
+			Output.Data = OutState.OutputData;
 
-			Output.Data = OutState.OutputPointData;
+PRAGMA_DISABLE_DEPRECATION_WARNINGS
+			OutState.InputPointData = Cast<UPCGPointData>(OutState.InputData);
+			OutState.OutputPointData = Cast<UPCGPointData>(OutState.OutputData);
+PRAGMA_ENABLE_DEPRECATION_WARNINGS
 
 			return EPCGTimeSliceInitResult::Success;
 		});
@@ -67,51 +73,104 @@ bool FPCGPointOperationElementBase::PreparePointOperationData(ContextType* InCon
 	return true;
 }
 
-bool FPCGPointOperationElementBase::ExecutePointOperation(ContextType* Context, TFunctionRef<PointExecSignature> PointFunction, int32 PointsPerChunk) const
+namespace PCGPointOperationElementBase
 {
-	TRACE_CPUPROFILER_EVENT_SCOPE(FPCGPointOperationElementBase::ExecutePointProcessing);
-	check(Context);
-
-	// Standard check that the time slice state has been prepared. If the result is NoOp or Failure, result in no output
-	if (!Context->DataIsPreparedForExecution() || Context->GetExecutionStateResult() != EPCGTimeSliceInitResult::Success)
-	{
-		Context->OutputData.TaggedData.Empty();
-		return true;
-	}
-
-	return ExecuteSlice(Context, [this, &PointFunction, &PointsPerChunk](ContextType* Context, const ExecStateType& ExecState, const IterStateType& IterState, const uint32 IterIndex)
+	bool ExecuteSliceOneToOne(FPCGPointOperationElementBase::ContextType* Context, const FPCGPointOperationElementBase::IterStateType& IterState, TFunctionRef<bool(const FPCGPoint& InPoint, FPCGPoint& OutPoint)> Callback, int32 PointsPerChunk)
 	{
+PRAGMA_DISABLE_DEPRECATION_WARNINGS
+		check(IterState.InputPointData && IterState.OutputPointData);
 		const TArray<FPCGPoint>& InPoints = IterState.InputPointData->GetPoints();
 		TArray<FPCGPoint>& OutPoints = IterState.OutputPointData->GetMutablePoints();
-
-		// If this input created an error, result in no output
-		if (Context->GetIterationStateResult(IterIndex) != EPCGTimeSliceInitResult::Success)
-		{
-			OutPoints.Empty();
-			return true;
-		}
+PRAGMA_ENABLE_DEPRECATION_WARNINGS
 
 		// Conversion lambda from index to point ref for ease of use
-		auto InternalPointFunction = [&PointFunction, &InPoints, &OutPoints](int32 ReadIndex, int32 WriteIndex)->bool
+		auto InternalPointFunction = [&Callback, &InPoints, &OutPoints](int32 ReadIndex, int32 WriteIndex) -> bool
 		{
-			return PointFunction(InPoints[ReadIndex], OutPoints[WriteIndex]);
+			return Callback(InPoints[ReadIndex], OutPoints[WriteIndex]);
 		};
 
-		const bool bAsyncDone = FPCGAsync::AsyncProcessingOneToOneEx(
+		return FPCGAsync::AsyncProcessingOneToOneEx(
 			&Context->AsyncState,
 			IterState.NumPoints,
 			/*InitializeFunc=*/[] {}, // Not useful for this context, since its preferred to initialize in PrepareDataInternal, so empty lambda
 			std::move(InternalPointFunction),
 			Context->TimeSliceIsEnabled(),
 			PointsPerChunk);
+	}
 
-		if (bAsyncDone)
+	bool ExecuteSliceRange(FPCGPointOperationElementBase::ContextType* Context, const FPCGPointOperationElementBase::IterStateType& IterState, TFunctionRef<bool(const UPCGBasePointData* InPointData, UPCGBasePointData* OutPointData, int32 InStartIndex, int32 InCount)> Callback, int32 PointsPerChunk)
+	{
+		const UPCGBasePointData* InputData = IterState.InputData;
+		UPCGBasePointData* OutputData = IterState.OutputData;
+
+		auto ProcessRangeFunc = [&Callback, &InputData, &OutputData](int32 StartReadIndex, int32 StartWriteIndex, int32 Count)
 		{
-			PCGE_LOG_C(Verbose, LogOnly, Context, FText::Format(LOCTEXT("PointProcessInfo", "Processed {0} points"), IterState.NumPoints));
+			check(StartReadIndex == StartWriteIndex);
+			Callback(InputData, OutputData, StartReadIndex, Count);
+			return Count;
+		};
+
+		return FPCGAsync::AsyncProcessingRangeEx(
+			&Context->AsyncState,
+			IterState.NumPoints,
+			/*InitializeFunc=*/[] {}, // Not useful for this context, since its preferred to initialize in PrepareDataInternal, so empty lambda
+			std::move(ProcessRangeFunc),
+			/*MoveRangeFunc=*/[](int32 ReadIndex, int32 WriteIndex, int32 Count) {}, // This is a one to one element so no point get discarded, so empty lambda
+			/*FinalizeFunc=*/[](int32) {}, // Not useful here, so empty lambda
+			Context->TimeSliceIsEnabled(),
+			PointsPerChunk);
+	}
+
+	template <typename Func>
+	bool ExecutePointOperation(const FPCGPointOperationElementBase* Element, FPCGPointOperationElementBase::ContextType* Context, Func Callback, int32 PointsPerChunk)
+	{
+		TRACE_CPUPROFILER_EVENT_SCOPE(FPCGPointOperationElementBase::ExecutePointProcessing);
+		check(Context);
+
+		// Standard check that the time slice state has been prepared. If the result is NoOp or Failure, result in no output
+		if (!Context->DataIsPreparedForExecution() || Context->GetExecutionStateResult() != EPCGTimeSliceInitResult::Success)
+		{
+			Context->OutputData.TaggedData.Empty();
+			return true;
 		}
 
-		return bAsyncDone;
-	});
+		return Element->ExecuteSlice(Context, [&Callback, &PointsPerChunk](FPCGPointOperationElementBase::ContextType* Context, const FPCGPointOperationElementBase::ExecStateType& ExecState, const FPCGPointOperationElementBase::IterStateType& IterState, const uint32 IterIndex)
+		{
+			// If this input created an error, result in no output
+			if (Context->GetIterationStateResult(IterIndex) != EPCGTimeSliceInitResult::Success)
+			{
+				IterState.OutputData->SetNumPoints(0);
+				return true;
+			}
+
+			bool bAsyncDone = false;
+			if constexpr (std::is_invocable_v<Func, const FPCGPoint&, FPCGPoint&>)
+			{
+				bAsyncDone = ExecuteSliceOneToOne(Context, IterState, Callback, PointsPerChunk);
+			}
+			else
+			{
+				bAsyncDone = ExecuteSliceRange(Context, IterState, Callback, PointsPerChunk);
+			}
+
+			if (bAsyncDone)
+			{
+				PCGE_LOG_C(Verbose, LogOnly, Context, FText::Format(LOCTEXT("PointProcessInfo", "Processed {0} points"), IterState.NumPoints));
+			}
+
+			return bAsyncDone;
+		});
+	}
+}
+
+bool FPCGPointOperationElementBase::ExecutePointOperationWithPoints(ContextType* Context, TFunctionRef<bool(const FPCGPoint& InPoint, FPCGPoint& OutPoint)> Callback, int32 PointsPerChunk) const
+{
+	return PCGPointOperationElementBase::ExecutePointOperation(this, Context, Callback, PointsPerChunk);
+}
+
+bool FPCGPointOperationElementBase::ExecutePointOperationWithIndices(ContextType* Context, TFunctionRef<bool(const UPCGBasePointData* InPointData, UPCGBasePointData* OutPointData, int32 InStartIndex, int32 InCount)> Callback, int32 PointsPerChunk) const
+{
+	return PCGPointOperationElementBase::ExecutePointOperation(this, Context, Callback, PointsPerChunk);
 }
 
 #undef LOCTEXT_NAMESPACE
diff --git a/Engine/Plugins/PCG/Source/PCG/Private/Elements/PCGSurfaceSampler.cpp b/Engine/Plugins/PCG/Source/PCG/Private/Elements/PCGSurfaceSampler.cpp
--- a/Engine/Plugins/PCG/Source/PCG/Private/Elements/PCGSurfaceSampler.cpp
+++ b/Engine/Plugins/PCG/Source/PCG/Private/Elements/PCGSurfaceSampler.cpp
@@ -149,9 +149,9 @@ namespace PCGSurfaceSampler
 		return FIntVector2(CellMinX + (Index % CellCountX), CellMinY + (Index / CellCountX));
 	}
 
-	UPCGPointData* SampleSurface(FPCGContext* Context, const UPCGSurfaceData* InSurface, const UPCGSpatialData* InBoundingShape, const FBox& EffectiveBounds, const FSurfaceSamplerParams& ExecutionParams)
+	UPCGBasePointData* SampleSurface(FPCGContext* Context, const FSurfaceSamplerParams& ExecutionParams, const UPCGSurfaceData* InSurface, const UPCGSpatialData* InBoundingShape, const FBox& EffectiveBounds, bool bForcePointData)
 	{
-		UPCGPointData* SampledData = FPCGContext::NewObject_AnyThread<UPCGPointData>(Context);
+		UPCGBasePointData* SampledData = bForcePointData ? FPCGContext::NewObject_AnyThread<UPCGPointData>(Context) : FPCGContext::NewPointData_AnyThread(Context);
 		SampledData->InitializeFromData(InSurface);
 
 		FSurfaceSamplerData SamplerData;
@@ -163,90 +163,156 @@ namespace PCGSurfaceSampler
 		return SampledData;
 	}
 
-	bool SampleSurface(FPCGContext* Context, const FSurfaceSamplerData& SamplerData, const UPCGSurfaceData* InSurface, const UPCGSpatialData* InBoundingShape, UPCGPointData* SampledData, const bool bTimeSlicingIsEnabled)
+	UPCGPointData* SampleSurface(FPCGContext* Context, const UPCGSurfaceData* InSurface, const UPCGSpatialData* InBoundingShape, const FBox& EffectiveBounds, const FSurfaceSamplerParams& ExecutionParams)
 	{
-		check(InSurface && SampledData);
-
-		TArray<FPCGPoint>& SampledPoints = SampledData->GetMutablePoints();
+		return CastChecked<UPCGPointData>(SampleSurface(Context, ExecutionParams, InSurface, InBoundingShape, EffectiveBounds, /*bForcePointData=*/true));
+	}
 
+	UPCGBasePointData* SampleSurface(FPCGContext* Context, const FSurfaceSamplerParams& ExecutionParams, const UPCGSurfaceData* InSurface, const UPCGSpatialData* InBoundingShape, const FBox& EffectiveBounds)
+	{
+		return SampleSurface(Context, ExecutionParams, InSurface, InBoundingShape, EffectiveBounds, /*bForcePointData=*/false);
+	}
+	
+	bool SampleSurface(FPCGContext* Context, const FSurfaceSamplerData& SamplerData, const UPCGSurfaceData* InSurface, const UPCGSpatialData* InBoundingShape, UPCGBasePointData* SampledData, const bool bTimeSlicingIsEnabled)
+	{
+		check(InSurface && SampledData);
 		const FPCGProjectionParams ProjectionParams{};
 
 		// Cache pointer ahead of time to avoid dereferencing object pointer which does access tracking and supports lazy loading, and can come with substantial
 		// overhead (add trace marker to FObjectPtr::Get to see).
 		UPCGMetadata* OutMetadata = SampledData->Metadata.Get();
+				
+		auto InitializeFunc = [&SamplerData, SampledData]()
+		{
+			SampledData->SetNumPoints(SamplerData.CellCount);
+			SampledData->SetSteepness(SamplerData.Params.PointSteepness);
+			SampledData->SetDensity(1.0f);
+
+			FPCGPoint DefaultPoint;
+			PCGPointHelpers::SetExtents(SamplerData.Params.PointExtents, DefaultPoint.BoundsMin, DefaultPoint.BoundsMax);
 
-		auto AsyncProcessFunc = [&SamplerData, InBoundingShape, InSurface, &ProjectionParams, OutMetadata](int32 Index, FPCGPoint& OutPoint)
+			SampledData->SetBoundsMin(DefaultPoint.BoundsMin);
+			SampledData->SetBoundsMax(DefaultPoint.BoundsMax);
+
+			EPCGPointNativeProperties PropertiesToAllocate = EPCGPointNativeProperties::Transform
+				| EPCGPointNativeProperties::Density
+				| EPCGPointNativeProperties::Seed
+				| EPCGPointNativeProperties::MetadataEntry
+				| EPCGPointNativeProperties::Color;
+
+			SampledData->AllocateProperties(PropertiesToAllocate);
+		};
+
+		auto AsyncProcessRangeFunc = [&SamplerData, InBoundingShape, InSurface, &ProjectionParams, OutMetadata, SampledData](int32 StartReadIndex, int32 StartWriteIndex, int32 Count)
 		{
-			const FSurfaceSamplerParams& SamplerParams = SamplerData.Params;
-			const FIntVector2 Indices = SamplerData.ComputeCellIndices(Index);
+			TPCGValueRange<FTransform> TransformRange = SampledData->GetTransformValueRange(/*bAllocate=*/false);
+			TPCGValueRange<float> DensityRange = SampledData->GetDensityValueRange(/*bAllocate=*/false);
+			TPCGValueRange<FVector> BoundsMinRange = SampledData->GetBoundsMinValueRange(/*bAllocate=*/false);
+			TPCGValueRange<FVector> BoundsMaxRange = SampledData->GetBoundsMaxValueRange(/*bAllocate=*/false);
+			TPCGValueRange<int32> SeedRange = SampledData->GetSeedValueRange(/*bAllocate=*/false);
+			TPCGValueRange<int64> MetadataEntryRange = SampledData->GetMetadataEntryValueRange(/*bAllocate=*/false);
+			TPCGValueRange<FVector4> ColorRange = SampledData->GetColorValueRange(/*bAllocate=*/false);
 
-			const FVector::FReal CurrentX = Indices.X * SamplerData.CellSize.X;
-			const FVector::FReal CurrentY = Indices.Y * SamplerData.CellSize.Y;
-			const FVector InnerCellSize = SamplerData.InnerCellSize;
+			int32 NumWritten = 0;
 
-			FRandomStream RandomSource(PCGHelpers::ComputeSeed(SamplerData.Seed, Indices.X, Indices.Y));
-			const float Chance = RandomSource.FRand();
+			for(int32 ReadIndex = StartReadIndex; ReadIndex < StartReadIndex+Count; ++ReadIndex)
+			{				
+				const int32 WriteIndex = StartWriteIndex + NumWritten;
 
-			const float Ratio = SamplerData.Ratio;
+				const FSurfaceSamplerParams& SamplerParams = SamplerData.Params;
+				const FIntVector2 Indices = SamplerData.ComputeCellIndices(ReadIndex);
 
-			if (Chance >= Ratio)
-			{
-				return false;
-			}
+				const FVector::FReal CurrentX = Indices.X * SamplerData.CellSize.X;
+				const FVector::FReal CurrentY = Indices.Y * SamplerData.CellSize.Y;
+				const FVector InnerCellSize = SamplerData.InnerCellSize;
 
-			const float RandX = RandomSource.FRand();
-			const float RandY = RandomSource.FRand();
+				FRandomStream RandomSource(PCGHelpers::ComputeSeed(SamplerData.Seed, Indices.X, Indices.Y));
+				const float Chance = RandomSource.FRand();
 
-			FVector TentativeLocation = FVector(CurrentX + RandX * InnerCellSize.X, CurrentY + RandY * InnerCellSize.Y, SamplerData.PreProjectionDisplacement);
+				const float Ratio = SamplerData.Ratio;
 
-			// If pre-projected points need a local transformation (ex. World Ray Hit Query) and not default to -Z
-			if (SamplerData.bNeedsLocalTransformation)
-			{
-				// Transform the pre-projected sample point around the surface's origin from local to world space
-				TentativeLocation = SamplerData.PreProjectionTransform.TransformPosition(TentativeLocation);
-			}
+				if (Chance >= Ratio)
+				{
+					continue;
+				}
 
-			const FBox LocalBound(-SamplerParams.PointExtents, SamplerParams.PointExtents);
+				const float RandX = RandomSource.FRand();
+				const float RandY = RandomSource.FRand();
 
-			// The output at this point is not initialized
-			OutPoint = FPCGPoint();
+				FVector TentativeLocation = FVector(CurrentX + RandX * InnerCellSize.X, CurrentY + RandY * InnerCellSize.Y, SamplerData.PreProjectionDisplacement);
 
-			// Firstly project onto elected generating shape to move to final position.
-			if (!InSurface->ProjectPoint(FTransform(TentativeLocation), LocalBound, ProjectionParams, OutPoint, OutMetadata))
-			{
-				return false;
-			}
+				// If pre-projected points need a local transformation (ex. World Ray Hit Query) and not default to -Z
+				if (SamplerData.bNeedsLocalTransformation)
+				{
+					// Transform the pre-projected sample point around the surface's origin from local to world space
+					TentativeLocation = SamplerData.PreProjectionTransform.TransformPosition(TentativeLocation);
+				}
 
-			// Set physical properties that are needed for the bounding shape checks, etc.
-			OutPoint.SetExtents(SamplerParams.PointExtents);
-			OutPoint.Steepness = SamplerParams.PointSteepness;
+				const FBox LocalBound(-SamplerParams.PointExtents, SamplerParams.PointExtents);
 
-			// Now run gauntlet of shape network (if there is one) to accept or reject the point.
-			if (InBoundingShape)
-			{
-				FPCGPoint BoundingShapeSample;
+				// The output at this point is not initialized
+				FPCGPoint OutPoint = FPCGPoint();
+
+				// Firstly project onto elected generating shape to move to final position.
+				if (!InSurface->ProjectPoint(FTransform(TentativeLocation), LocalBound, ProjectionParams, OutPoint, OutMetadata))
+				{
+					continue;
+				}
+											
+				// Now run gauntlet of shape network (if there is one) to accept or reject the point.
+				if (InBoundingShape)
+				{
+					FPCGPoint BoundingShapeSample;
 #if WITH_EDITOR
-				if (!InBoundingShape->SamplePoint(OutPoint.Transform, OutPoint.GetLocalBounds(), BoundingShapeSample, nullptr) && !SamplerParams.bKeepZeroDensityPoints)
+					if (!InBoundingShape->SamplePoint(OutPoint.Transform, OutPoint.GetLocalBounds(), BoundingShapeSample, nullptr) && !SamplerParams.bKeepZeroDensityPoints)
 #else
-				if (!InBoundingShape->SamplePoint(OutPoint.Transform, OutPoint.GetLocalBounds(), BoundingShapeSample, nullptr))
+					if (!InBoundingShape->SamplePoint(OutPoint.Transform, OutPoint.GetLocalBounds(), BoundingShapeSample, nullptr))
 #endif
+					{
+						continue;
+					}
+
+					// Produce smooth density field
+					DensityRange[WriteIndex] *= BoundingShapeSample.Density;
+				}
+
+				// Apply final parameters on the point
+				if (SamplerParams.bApplyDensityToPoints)
 				{
-					return false;
+					DensityRange[WriteIndex] *= (SamplerParams.bApplyDensityToPoints ? ((Ratio - Chance) / Ratio) : 1.0f);
 				}
 
-				// Produce smooth density field
-				OutPoint.Density *= BoundingShapeSample.Density;
+				SeedRange[WriteIndex] = RandomSource.GetCurrentSeed();
+				TransformRange[WriteIndex] = OutPoint.Transform;
+				MetadataEntryRange[WriteIndex] = OutPoint.MetadataEntry;
+				ColorRange[WriteIndex] = OutPoint.Color;
+				NumWritten++;
 			}
 
-			// Apply final parameters on the point
-			OutPoint.Density *= (SamplerParams.bApplyDensityToPoints ? ((Ratio - Chance) / Ratio) : 1.0f);
-			OutPoint.Seed = RandomSource.GetCurrentSeed();
+			return NumWritten;
+		};
+
+		auto MoveDataRangeFunc = [SampledData](int32 RangeStartIndex, int32 MoveToIndex, int32 NumElements)
+		{
+			SampledData->MoveRange(RangeStartIndex, MoveToIndex, NumElements);
+		};
 
-			return true;
+		auto FinishedFunc = [SampledData](int32 NumWritten)
+		{
+			SampledData->SetNumPoints(NumWritten);
 		};
 
 		FPCGAsyncState* AsyncState = Context ? &Context->AsyncState : nullptr;
-		return FPCGAsync::AsyncProcessing<FPCGPoint>(AsyncState, SamplerData.CellCount, SampledPoints, AsyncProcessFunc, /*bEnableTimeSlicing=*/Context && bTimeSlicingIsEnabled);
+		const bool bDone = FPCGAsync::AsyncProcessingRangeEx(
+			AsyncState, 
+			SamplerData.CellCount,
+			InitializeFunc,
+			AsyncProcessRangeFunc,
+			MoveDataRangeFunc,
+			FinishedFunc,
+			/*bEnableTimeSlicing=*/Context && bTimeSlicingIsEnabled);
+
+		return bDone;
 	}
 
 #if WITH_EDITOR
@@ -441,8 +507,11 @@ bool FPCGSurfaceSamplerElement::PrepareDataInternal(FPCGContext* InContext) cons
 			const UPCGSurfaceData* GeneratingShape = GeneratingShapes[IterationIndex];
 			check(GeneratingShape);
 
-			OutState.OutputPoints = FPCGContext::NewObject_AnyThread<UPCGPointData>(Context);
-			OutState.OutputPoints->InitializeFromData(GeneratingShape);
+			OutState.OutputData = FPCGContext::NewPointData_AnyThread(Context);
+PRAGMA_DISABLE_DEPRECATION_WARNINGS
+			OutState.OutputPoints = Cast<UPCGPointData>(OutState.OutputData);
+PRAGMA_ENABLE_DEPRECATION_WARNINGS
+			OutState.OutputData->InitializeFromData(GeneratingShape);
 
 			// This bounds will be used to generate the pre-projected grid
 			FBox EffectiveGridBounds = FBox(EForceInit::ForceInit);
@@ -491,7 +560,7 @@ bool FPCGSurfaceSamplerElement::PrepareDataInternal(FPCGContext* InContext) cons
 			}
 
 			// Assigning this here prevents the need to root
-			Outputs[IterationIndex].Data = OutState.OutputPoints;
+			Outputs[IterationIndex].Data = OutState.OutputData;
 			return EPCGTimeSliceInitResult::Success;
 		});
 
@@ -523,7 +592,7 @@ bool FPCGSurfaceSamplerElement::ExecuteInternal(FPCGContext* InContext) const
 		for (FPCGTaggedData& Input : TimeSlicedContext->InputData.GetInputs())
 		{
 			FPCGTaggedData& Output = TimeSlicedContext->OutputData.TaggedData.Add_GetRef(Input);
-			UPCGPointData* PointData = FPCGContext::NewObject_AnyThread<UPCGPointData>(InContext);
+			UPCGBasePointData* PointData = FPCGContext::NewPointData_AnyThread(InContext);
 			PointData->InitializeFromData(Cast<UPCGSpatialData>(Input.Data));
 			Output.Data = PointData;
 		}
@@ -539,7 +608,7 @@ bool FPCGSurfaceSamplerElement::ExecuteInternal(FPCGContext* InContext) const
 		// This iteration resulted in an early out for no sampling operation. Early out with empty point data.
 		if (InitResult == EPCGTimeSliceInitResult::NoOperation)
 		{
-			Context->OutputData.TaggedData[IterationIndex].Data = FPCGContext::NewObject_AnyThread<UPCGPointData>(Context);
+			Context->OutputData.TaggedData[IterationIndex].Data = FPCGContext::NewPointData_AnyThread(Context);
 
 			return true;
 		}
@@ -548,11 +617,11 @@ bool FPCGSurfaceSamplerElement::ExecuteInternal(FPCGContext* InContext) const
 		check(InitResult == EPCGTimeSliceInitResult::Success);
 
 		// Run the execution until the time slice is finished
-		const bool bAsyncDone = PCGSurfaceSampler::SampleSurface(Context, IterState.SamplerData, ExecState.GeneratingShapes[IterationIndex], ExecState.BoundingShape, IterState.OutputPoints, Context->TimeSliceIsEnabled());
+		const bool bAsyncDone = PCGSurfaceSampler::SampleSurface(Context, IterState.SamplerData, ExecState.GeneratingShapes[IterationIndex], ExecState.BoundingShape, IterState.OutputData, Context->TimeSliceIsEnabled());
 
 		if (bAsyncDone)
 		{
-			PCGE_LOG_C(Verbose, LogOnly, Context, FText::Format(LOCTEXT("GenerationInfo", "Generated {0} points in {1} cells"), IterState.OutputPoints->GetPoints().Num(), IterState.SamplerData.CellCount));
+			PCGE_LOG_C(Verbose, LogOnly, Context, FText::Format(LOCTEXT("GenerationInfo", "Generated {0} points in {1} cells"), IterState.OutputData->GetNumPoints(), IterState.SamplerData.CellCount));
 		}
 
 		return bAsyncDone;
diff --git a/Engine/Plugins/PCG/Source/PCG/Private/PCGContext.cpp b/Engine/Plugins/PCG/Source/PCG/Private/PCGContext.cpp
--- a/Engine/Plugins/PCG/Source/PCG/Private/PCGContext.cpp
+++ b/Engine/Plugins/PCG/Source/PCG/Private/PCGContext.cpp
@@ -8,6 +8,8 @@
 #include "PCGPin.h"
 #include "PCGSubsystem.h"
 #include "Data/PCGSpatialData.h"
+#include "Data/PCGPointData.h"
+#include "Data/PCGPointArrayData.h"
 #include "Helpers/PCGHelpers.h"
 #include "Metadata/PCGAttributePropertySelector.h"
 #include "Metadata/PCGMetadata.h"
@@ -24,6 +26,11 @@
 
 #define LOCTEXT_NAMESPACE "PCGContext"
 
+TAutoConsoleVariable<bool> CVarPCGEnablePointArrayData(
+	TEXT("pcg.Experimental.EnablePointArrayData"),
+	false,
+	TEXT("Whether to use the new UPCGPointArrayData when calling FPCGContext::CreatePointData_AnyThread"));
+
 namespace PCGContextHelpers
 {
 	template <typename T>
@@ -35,6 +42,18 @@ namespace PCGContextHelpers
 	}
 }
 
+UPCGBasePointData* FPCGContext::NewPointData_AnyThread(FPCGContext* Context)
+{
+	if (CVarPCGEnablePointArrayData.GetValueOnAnyThread())
+	{
+		return NewObject_AnyThread<UPCGPointArrayData>(Context);
+	}
+	else
+	{
+		return NewObject_AnyThread<UPCGPointData>(Context);
+	}
+}
+
 FPCGTaskId FPCGContext::GetGraphExecutionTaskId() const
 {
 	return ensure(Stack) ? Stack->GetGraphExecutionTaskId() : InvalidPCGTaskId;
diff --git a/Engine/Plugins/PCG/Source/PCG/Private/PCGPointArray.cpp b/Engine/Plugins/PCG/Source/PCG/Private/PCGPointArray.cpp
--- a/Engine/Plugins/PCG/Source/PCG/Private/PCGPointArray.cpp
+++ b/Engine/Plugins/PCG/Source/PCG/Private/PCGPointArray.cpp
@@ -150,3 +150,62 @@ bool FPCGPointArray::Serialize(FArchive& Ar)
 
 	return true;
 }
+
+void FPCGPointArray::CopyPointsTo(FPCGPointArray& OutPointArray, int32 StartReadIndex, int32 StartWriteIndex, int32 Count) const
+{
+	Transform.CopyTo(OutPointArray.Transform, StartReadIndex, StartWriteIndex, Count);
+	Density.CopyTo(OutPointArray.Density, StartReadIndex, StartWriteIndex, Count);
+	BoundsMin.CopyTo(OutPointArray.BoundsMin, StartReadIndex, StartWriteIndex, Count);
+	BoundsMax.CopyTo(OutPointArray.BoundsMax, StartReadIndex, StartWriteIndex, Count);
+	Color.CopyTo(OutPointArray.Color, StartReadIndex, StartWriteIndex, Count);
+	Steepness.CopyTo(OutPointArray.Steepness, StartReadIndex, StartWriteIndex, Count);
+	Seed.CopyTo(OutPointArray.Seed, StartReadIndex, StartWriteIndex, Count);
+	MetadataEntry.CopyTo(OutPointArray.MetadataEntry, StartReadIndex, StartWriteIndex, Count);
+}
+
+EPCGPointNativeProperties FPCGPointArray::GetAllocatedProperties() const
+{
+	EPCGPointNativeProperties AllocatedProperties = EPCGPointNativeProperties::None;
+
+	if (Transform.IsAllocated())
+	{
+		AllocatedProperties |= EPCGPointNativeProperties::Transform;
+	}
+
+	if (Density.IsAllocated())
+	{
+		AllocatedProperties |= EPCGPointNativeProperties::Density;
+	}
+
+	if (BoundsMin.IsAllocated())
+	{
+		AllocatedProperties |= EPCGPointNativeProperties::BoundsMin;
+	}
+
+	if (BoundsMax.IsAllocated())
+	{
+		AllocatedProperties |= EPCGPointNativeProperties::BoundsMax;
+	}
+
+	if (Color.IsAllocated())
+	{
+		AllocatedProperties |= EPCGPointNativeProperties::Color;
+	}
+
+	if (Steepness.IsAllocated())
+	{
+		AllocatedProperties |= EPCGPointNativeProperties::Steepness;
+	}
+
+	if (Seed.IsAllocated())
+	{
+		AllocatedProperties |= EPCGPointNativeProperties::Seed;
+	}
+
+	if (MetadataEntry.IsAllocated())
+	{
+		AllocatedProperties |= EPCGPointNativeProperties::MetadataEntry;
+	}
+
+	return AllocatedProperties;
+}
\ No newline at end of file
diff --git a/Engine/Plugins/PCG/Source/PCG/Private/Tests/Elements/PCGCreatePointsGridTest.cpp b/Engine/Plugins/PCG/Source/PCG/Private/Tests/Elements/PCGCreatePointsGridTest.cpp
--- a/Engine/Plugins/PCG/Source/PCG/Private/Tests/Elements/PCGCreatePointsGridTest.cpp
+++ b/Engine/Plugins/PCG/Source/PCG/Private/Tests/Elements/PCGCreatePointsGridTest.cpp
@@ -5,7 +5,7 @@
 #include "PCGComponent.h"
 #include "PCGContext.h"
 #include "PCGVolume.h"
-#include "Data/PCGPointData.h"
+#include "Data/PCGBasePointData.h"
 
 #include "Elements/PCGCreatePointsGrid.h"
 
@@ -55,18 +55,18 @@ bool FPCGCreatePointsGridTest_CellCenter::RunTest(const FString& Parameters)
 
 	UTEST_EQUAL("Output count", Outputs.Num(), 1);
 
-	const UPCGPointData* OutPointData = Cast<UPCGPointData>(Outputs[0].Data);
+	const UPCGBasePointData* OutPointData = Cast<UPCGBasePointData>(Outputs[0].Data);
 
 	UTEST_NOT_NULL("Output point data", OutPointData);
 
-	const TArray<FPCGPoint>& OutPoints = OutPointData->GetPoints();
+	const TConstPCGValueRange<FTransform> Transforms = OutPointData->GetConstTransformValueRange();
 
-	UTEST_EQUAL("Output point count", OutPoints.Num(), 100);
+	UTEST_EQUAL("Output point count", Transforms.Num(), 100);
 
-	UTEST_EQUAL("Point 0 has the correct transform,", OutPoints[0].Transform, FTransform(FRotator::ZeroRotator, FVector(-450.0, -450.0, 0), FVector::One()));
-	UTEST_EQUAL("Point 23 has the correct transform,", OutPoints[23].Transform, FTransform(FRotator::ZeroRotator, FVector(-150.0, -250.0, 0), FVector::One()));
-	UTEST_EQUAL("Point 65 has the correct transform,", OutPoints[65].Transform, FTransform(FRotator::ZeroRotator, FVector(50.0, 150.0, 0), FVector::One()));
-	UTEST_EQUAL("Point 92 has the correct transform,", OutPoints[92].Transform, FTransform(FRotator::ZeroRotator, FVector(-250.0, 450.0, 0), FVector::One()));
+	UTEST_EQUAL("Point 0 has the correct transform,", Transforms[0], FTransform(FRotator::ZeroRotator, FVector(-450.0, -450.0, 0), FVector::One()));
+	UTEST_EQUAL("Point 23 has the correct transform,", Transforms[23], FTransform(FRotator::ZeroRotator, FVector(-150.0, -250.0, 0), FVector::One()));
+	UTEST_EQUAL("Point 65 has the correct transform,", Transforms[65], FTransform(FRotator::ZeroRotator, FVector(50.0, 150.0, 0), FVector::One()));
+	UTEST_EQUAL("Point 92 has the correct transform,", Transforms[92], FTransform(FRotator::ZeroRotator, FVector(-250.0, 450.0, 0), FVector::One()));
 
 	return true;
 }
@@ -82,20 +82,20 @@ bool FPCGCreatePointsGridTest_CellCorner::RunTest(const FString& Parameters)
 
 	UTEST_EQUAL("Output count", Outputs.Num(), 1);
 
-	const UPCGPointData* OutPointData = Cast<UPCGPointData>(Outputs[0].Data);
+	const UPCGBasePointData* OutPointData = Cast<UPCGBasePointData>(Outputs[0].Data);
 
 	UTEST_NOT_NULL("Output point data", OutPointData);
 
-	const TArray<FPCGPoint>& OutPoints = OutPointData->GetPoints();
+	const TConstPCGValueRange<FTransform> Transforms = OutPointData->GetConstTransformValueRange();
 
 	// Cell Corners generate (10 + 1) * (10 + 1) * 2
-	UTEST_EQUAL("Output point count", OutPoints.Num(), 242);
+	UTEST_EQUAL("Output point count", Transforms.Num(), 242);
 
-	UTEST_EQUAL("Point 0 has the correct transform,", OutPoints[0].Transform, FTransform(FRotator::ZeroRotator, FVector(-500.0, -500.0, -50.0), FVector::One()));
-	UTEST_EQUAL("Point 23 has the correct transform,", OutPoints[23].Transform, FTransform(FRotator::ZeroRotator, FVector(-400.0, -300.0, -50.0), FVector::One()));
-	UTEST_EQUAL("Point 65 has the correct transform,", OutPoints[65].Transform, FTransform(FRotator::ZeroRotator, FVector(500.0, 0.0, -50.0), FVector::One()));
-	UTEST_EQUAL("Point 92 has the correct transform,", OutPoints[92].Transform, FTransform(FRotator::ZeroRotator, FVector(-100.0, 300.0, -50.0), FVector::One()));
-	UTEST_EQUAL("Point 145 has the correct transform,", OutPoints[145].Transform, FTransform(FRotator::ZeroRotator, FVector(-300.0, -300.0, 50.0), FVector::One()));
+	UTEST_EQUAL("Point 0 has the correct transform,", Transforms[0], FTransform(FRotator::ZeroRotator, FVector(-500.0, -500.0, -50.0), FVector::One()));
+	UTEST_EQUAL("Point 23 has the correct transform,", Transforms[23], FTransform(FRotator::ZeroRotator, FVector(-400.0, -300.0, -50.0), FVector::One()));
+	UTEST_EQUAL("Point 65 has the correct transform,", Transforms[65], FTransform(FRotator::ZeroRotator, FVector(500.0, 0.0, -50.0), FVector::One()));
+	UTEST_EQUAL("Point 92 has the correct transform,", Transforms[92], FTransform(FRotator::ZeroRotator, FVector(-100.0, 300.0, -50.0), FVector::One()));
+	UTEST_EQUAL("Point 145 has the correct transform,", Transforms[145], FTransform(FRotator::ZeroRotator, FVector(-300.0, -300.0, 50.0), FVector::One()));
 
 	return true;
 }
@@ -111,18 +111,18 @@ bool FPCGCreatePointsGridTest_XYCenter::RunTest(const FString& Parameters)
 
 	UTEST_EQUAL("Output count", Outputs.Num(), 1);
 
-	const UPCGPointData* OutPointData = Cast<UPCGPointData>(Outputs[0].Data);
+	const UPCGBasePointData* OutPointData = Cast<UPCGBasePointData>(Outputs[0].Data);
 
 	UTEST_NOT_NULL("Output point data", OutPointData);
 
-	const TArray<FPCGPoint>& OutPoints = OutPointData->GetPoints();
+	const TConstPCGValueRange<FTransform> Transforms = OutPointData->GetConstTransformValueRange();
 
-	UTEST_EQUAL("Output point count", OutPoints.Num(), 4);
+	UTEST_EQUAL("Output point count", Transforms.Num(), 4);
 
-	UTEST_EQUAL("Point 0 has the correct transform,", OutPoints[0].Transform, FTransform(FRotator::ZeroRotator, FVector(-50.0, -50.0, 0.0), FVector::One()));
-	UTEST_EQUAL("Point 1 has the correct transform,", OutPoints[1].Transform, FTransform(FRotator::ZeroRotator, FVector(50.0, -50.0, 0.0), FVector::One()));
-	UTEST_EQUAL("Point 2 has the correct transform,", OutPoints[2].Transform, FTransform(FRotator::ZeroRotator, FVector(-50.0, 50.0, 0.0), FVector::One()));
-	UTEST_EQUAL("Point 3 has the correct transform,", OutPoints[3].Transform, FTransform(FRotator::ZeroRotator, FVector(50.0, 50.0, 0.0), FVector::One()));
+	UTEST_EQUAL("Point 0 has the correct transform,", Transforms[0], FTransform(FRotator::ZeroRotator, FVector(-50.0, -50.0, 0.0), FVector::One()));
+	UTEST_EQUAL("Point 1 has the correct transform,", Transforms[1], FTransform(FRotator::ZeroRotator, FVector(50.0, -50.0, 0.0), FVector::One()));
+	UTEST_EQUAL("Point 2 has the correct transform,", Transforms[2], FTransform(FRotator::ZeroRotator, FVector(-50.0, 50.0, 0.0), FVector::One()));
+	UTEST_EQUAL("Point 3 has the correct transform,", Transforms[3], FTransform(FRotator::ZeroRotator, FVector(50.0, 50.0, 0.0), FVector::One()));
 
 	return true;
 }
@@ -138,18 +138,18 @@ bool FPCGCreatePointsGridTest_XYCorner::RunTest(const FString& Parameters)
 
 	UTEST_EQUAL("Output count", Outputs.Num(), 1);
 
-	const UPCGPointData* OutPointData = Cast<UPCGPointData>(Outputs[0].Data);
+	const UPCGBasePointData* OutPointData = Cast<UPCGBasePointData>(Outputs[0].Data);
 
 	UTEST_NOT_NULL("Output point data", OutPointData);
 
-	const TArray<FPCGPoint>& OutPoints = OutPointData->GetPoints();
+	const TConstPCGValueRange<FTransform> Transforms = OutPointData->GetConstTransformValueRange();
 
-	UTEST_EQUAL("Output point count", OutPoints.Num(), 9);
+	UTEST_EQUAL("Output point count", Transforms.Num(), 9);
 
-	UTEST_EQUAL("Point 1 has the correct transform,", OutPoints[1].Transform, FTransform(FRotator::ZeroRotator, FVector(0.0, -100.0, 0.0), FVector::One()));
-	UTEST_EQUAL("Point 3 has the correct transform,", OutPoints[3].Transform, FTransform(FRotator::ZeroRotator, FVector(-100.0, 0.0, 0.0), FVector::One()));
-	UTEST_EQUAL("Point 5 has the correct transform,", OutPoints[5].Transform, FTransform(FRotator::ZeroRotator, FVector(100.0, 0.0, 0.0), FVector::One()));
-	UTEST_EQUAL("Point 7 has the correct transform,", OutPoints[7].Transform, FTransform(FRotator::ZeroRotator, FVector(0.0, 100.0, 0.0), FVector::One()));
+	UTEST_EQUAL("Point 1 has the correct transform,", Transforms[1], FTransform(FRotator::ZeroRotator, FVector(0.0, -100.0, 0.0), FVector::One()));
+	UTEST_EQUAL("Point 3 has the correct transform,", Transforms[3], FTransform(FRotator::ZeroRotator, FVector(-100.0, 0.0, 0.0), FVector::One()));
+	UTEST_EQUAL("Point 5 has the correct transform,", Transforms[5], FTransform(FRotator::ZeroRotator, FVector(100.0, 0.0, 0.0), FVector::One()));
+	UTEST_EQUAL("Point 7 has the correct transform,", Transforms[7], FTransform(FRotator::ZeroRotator, FVector(0.0, 100.0, 0.0), FVector::One()));
 
 	return true;
 }
@@ -165,18 +165,18 @@ bool FPCGCreatePointsGridTest_XZCenter::RunTest(const FString& Parameters)
 
 	UTEST_EQUAL("Output count", Outputs.Num(), 1);
 
-	const UPCGPointData* OutPointData = Cast<UPCGPointData>(Outputs[0].Data);
+	const UPCGBasePointData* OutPointData = Cast<UPCGBasePointData>(Outputs[0].Data);
 
 	UTEST_NOT_NULL("Output point data", OutPointData);
 
-	const TArray<FPCGPoint>& OutPoints = OutPointData->GetPoints();
+	const TConstPCGValueRange<FTransform> Transforms = OutPointData->GetConstTransformValueRange();
 
-	UTEST_EQUAL("Output point count", OutPoints.Num(), 4);
+	UTEST_EQUAL("Output point count", Transforms.Num(), 4);
 
-	UTEST_EQUAL("Point 0 has the correct transform,", OutPoints[0].Transform, FTransform(FRotator::ZeroRotator, FVector(-50.0, 0.0, -50.0), FVector::One()));
-	UTEST_EQUAL("Point 1 has the correct transform,", OutPoints[1].Transform, FTransform(FRotator::ZeroRotator, FVector(50.0, 0.0, -50.0), FVector::One()));
-	UTEST_EQUAL("Point 2 has the correct transform,", OutPoints[2].Transform, FTransform(FRotator::ZeroRotator, FVector(-50.0, 0.0, 50.0), FVector::One()));
-	UTEST_EQUAL("Point 3 has the correct transform,", OutPoints[3].Transform, FTransform(FRotator::ZeroRotator, FVector(50.0, 0.0, 50.0), FVector::One()));
+	UTEST_EQUAL("Point 0 has the correct transform,", Transforms[0], FTransform(FRotator::ZeroRotator, FVector(-50.0, 0.0, -50.0), FVector::One()));
+	UTEST_EQUAL("Point 1 has the correct transform,", Transforms[1], FTransform(FRotator::ZeroRotator, FVector(50.0, 0.0, -50.0), FVector::One()));
+	UTEST_EQUAL("Point 2 has the correct transform,", Transforms[2], FTransform(FRotator::ZeroRotator, FVector(-50.0, 0.0, 50.0), FVector::One()));
+	UTEST_EQUAL("Point 3 has the correct transform,", Transforms[3], FTransform(FRotator::ZeroRotator, FVector(50.0, 0.0, 50.0), FVector::One()));
 	
 	return true;
 }
@@ -192,18 +192,18 @@ bool FPCGCreatePointsGridTest_XZCorner::RunTest(const FString& Parameters)
 
 	UTEST_EQUAL("Output count", Outputs.Num(), 1);
 
-	const UPCGPointData* OutPointData = Cast<UPCGPointData>(Outputs[0].Data);
+	const UPCGBasePointData* OutPointData = Cast<UPCGBasePointData>(Outputs[0].Data);
 
 	UTEST_NOT_NULL("Output point data", OutPointData);
 
-	const TArray<FPCGPoint>& OutPoints = OutPointData->GetPoints();
+	const TConstPCGValueRange<FTransform> Transforms = OutPointData->GetConstTransformValueRange();
 
-	UTEST_EQUAL("Output point count", OutPoints.Num(), 9);
+	UTEST_EQUAL("Output point count", Transforms.Num(), 9);
 
-	UTEST_EQUAL("Point 1 has the correct transform,", OutPoints[1].Transform, FTransform(FRotator::ZeroRotator, FVector(0.0, 0.0, -100.0), FVector::One()));
-	UTEST_EQUAL("Point 3 has the correct transform,", OutPoints[3].Transform, FTransform(FRotator::ZeroRotator, FVector(-100.0, 0.0, 0.0), FVector::One()));
-	UTEST_EQUAL("Point 5 has the correct transform,", OutPoints[5].Transform, FTransform(FRotator::ZeroRotator, FVector(100.0, 0.0, 0.0), FVector::One()));
-	UTEST_EQUAL("Point 7 has the correct transform,", OutPoints[7].Transform, FTransform(FRotator::ZeroRotator, FVector(0.0, 0.0, 100.0), FVector::One()));
+	UTEST_EQUAL("Point 1 has the correct transform,", Transforms[1], FTransform(FRotator::ZeroRotator, FVector(0.0, 0.0, -100.0), FVector::One()));
+	UTEST_EQUAL("Point 3 has the correct transform,", Transforms[3], FTransform(FRotator::ZeroRotator, FVector(-100.0, 0.0, 0.0), FVector::One()));
+	UTEST_EQUAL("Point 5 has the correct transform,", Transforms[5], FTransform(FRotator::ZeroRotator, FVector(100.0, 0.0, 0.0), FVector::One()));
+	UTEST_EQUAL("Point 7 has the correct transform,", Transforms[7], FTransform(FRotator::ZeroRotator, FVector(0.0, 0.0, 100.0), FVector::One()));
 
 	return true;
 }
@@ -219,18 +219,18 @@ bool FPCGCreatePointsGridTest_YZCenter::RunTest(const FString& Parameters)
 
 	UTEST_EQUAL("Output count", Outputs.Num(), 1);
 
-	const UPCGPointData* OutPointData = Cast<UPCGPointData>(Outputs[0].Data);
+	const UPCGBasePointData* OutPointData = Cast<UPCGBasePointData>(Outputs[0].Data);
 
 	UTEST_NOT_NULL("Output point data", OutPointData);
 
-	const TArray<FPCGPoint>& OutPoints = OutPointData->GetPoints();
+	const TConstPCGValueRange<FTransform> Transforms = OutPointData->GetConstTransformValueRange();
 
-	UTEST_EQUAL("Output point count", OutPoints.Num(), 4);
+	UTEST_EQUAL("Output point count", Transforms.Num(), 4);
 
-	UTEST_EQUAL("Point 0 has the correct transform,", OutPoints[0].Transform, FTransform(FRotator::ZeroRotator, FVector(0.0, -50.0, -50.0), FVector::One()));
-	UTEST_EQUAL("Point 1 has the correct transform,", OutPoints[1].Transform, FTransform(FRotator::ZeroRotator, FVector(0.0, 50.0, -50.0), FVector::One()));
-	UTEST_EQUAL("Point 2 has the correct transform,", OutPoints[2].Transform, FTransform(FRotator::ZeroRotator, FVector(0.0, -50.0, 50.0), FVector::One()));
-	UTEST_EQUAL("Point 3 has the correct transform,", OutPoints[3].Transform, FTransform(FRotator::ZeroRotator, FVector(0.0, 50.0, 50.0), FVector::One()));
+	UTEST_EQUAL("Point 0 has the correct transform,", Transforms[0], FTransform(FRotator::ZeroRotator, FVector(0.0, -50.0, -50.0), FVector::One()));
+	UTEST_EQUAL("Point 1 has the correct transform,", Transforms[1], FTransform(FRotator::ZeroRotator, FVector(0.0, 50.0, -50.0), FVector::One()));
+	UTEST_EQUAL("Point 2 has the correct transform,", Transforms[2], FTransform(FRotator::ZeroRotator, FVector(0.0, -50.0, 50.0), FVector::One()));
+	UTEST_EQUAL("Point 3 has the correct transform,", Transforms[3], FTransform(FRotator::ZeroRotator, FVector(0.0, 50.0, 50.0), FVector::One()));
 
 	return true;
 }
@@ -246,18 +246,18 @@ bool FPCGCreatePointsGridTest_YZCorner::RunTest(const FString& Parameters)
 
 	UTEST_EQUAL("Output count", Outputs.Num(), 1);
 
-	const UPCGPointData* OutPointData = Cast<UPCGPointData>(Outputs[0].Data);
+	const UPCGBasePointData* OutPointData = Cast<UPCGBasePointData>(Outputs[0].Data);
 
 	UTEST_NOT_NULL("Output point data", OutPointData);
 
-	const TArray<FPCGPoint>& OutPoints = OutPointData->GetPoints();
+	const TConstPCGValueRange<FTransform> Transforms = OutPointData->GetConstTransformValueRange();
 
-	UTEST_EQUAL("Output point count", OutPoints.Num(), 9);
+	UTEST_EQUAL("Output point count", Transforms.Num(), 9);
 
-	UTEST_EQUAL("Point 1 has the correct transform,", OutPoints[1].Transform, FTransform(FRotator::ZeroRotator, FVector(0.0, 0.0, -100.0), FVector::One()));
-	UTEST_EQUAL("Point 3 has the correct transform,", OutPoints[3].Transform, FTransform(FRotator::ZeroRotator, FVector(0.0, -100.0, 0.0), FVector::One()));
-	UTEST_EQUAL("Point 5 has the correct transform,", OutPoints[5].Transform, FTransform(FRotator::ZeroRotator, FVector(0.0, 100.0, 0.0), FVector::One()));
-	UTEST_EQUAL("Point 7 has the correct transform,", OutPoints[7].Transform, FTransform(FRotator::ZeroRotator, FVector(0.0, 0.0, 100.0), FVector::One()));
+	UTEST_EQUAL("Point 1 has the correct transform,", Transforms[1], FTransform(FRotator::ZeroRotator, FVector(0.0, 0.0, -100.0), FVector::One()));
+	UTEST_EQUAL("Point 3 has the correct transform,", Transforms[3], FTransform(FRotator::ZeroRotator, FVector(0.0, -100.0, 0.0), FVector::One()));
+	UTEST_EQUAL("Point 5 has the correct transform,", Transforms[5], FTransform(FRotator::ZeroRotator, FVector(0.0, 100.0, 0.0), FVector::One()));
+	UTEST_EQUAL("Point 7 has the correct transform,", Transforms[7], FTransform(FRotator::ZeroRotator, FVector(0.0, 0.0, 100.0), FVector::One()));
 
 	return true;
 }
@@ -305,15 +305,15 @@ bool FPCGCreatePointsGridTest_OversizedCell::RunTest(const FString& Parameters)
 
 	UTEST_EQUAL("Output count", Outputs.Num(), 1);
 
-	const UPCGPointData* OutPointData = Cast<UPCGPointData>(Outputs[0].Data);
+	const UPCGBasePointData* OutPointData = Cast<UPCGBasePointData>(Outputs[0].Data);
 
 	UTEST_NOT_NULL("Output point data", OutPointData);
 
-	const TArray<FPCGPoint>& OutPoints = OutPointData->GetPoints();
+	const TConstPCGValueRange<FTransform> Transforms = OutPointData->GetConstTransformValueRange();
 
-	UTEST_EQUAL("Output point count", OutPoints.Num(), 1);
+	UTEST_EQUAL("Output point count", Transforms.Num(), 1);
 
-	UTEST_EQUAL("Point 0 has the correct transform,", OutPoints[0].Transform, FTransform(FRotator::ZeroRotator, FVector(0.0, 0.0, 0.0), FVector::One()));
+	UTEST_EQUAL("Point 0 has the correct transform,", Transforms[0], FTransform(FRotator::ZeroRotator, FVector(0.0, 0.0, 0.0), FVector::One()));
 
 	return true;
 }
@@ -341,24 +341,24 @@ bool FPCGCreatePointsGridTest_LocalCenter::RunTest(const FString& Parameters)
 
 	UTEST_EQUAL("Output count", Outputs.Num(), 1);
 
-	const UPCGPointData* OutPointData = Cast<UPCGPointData>(Outputs[0].Data);
+	const UPCGBasePointData* OutPointData = Cast<UPCGBasePointData>(Outputs[0].Data);
 
 	UTEST_NOT_NULL("Output point data", OutPointData);
 
-	const TArray<FPCGPoint>& OutPoints = OutPointData->GetPoints();
+	const TConstPCGValueRange<FTransform> Transforms = OutPointData->GetConstTransformValueRange();
 
-	UTEST_EQUAL("Output point count", OutPoints.Num(), 4);
+	UTEST_EQUAL("Output point count", Transforms.Num(), 4);
 
 	// Check to see if the position of the points are within the bounds of the volume
-	for (int i = 0; i < OutPoints.Num(); ++i)
+	for (int i = 0; i < Transforms.Num(); ++i)
 	{
-		UTEST_TRUE(FString::Format(TEXT("Point {0}, {1} is within the volume"), { i }), TestBounds.IsInsideOrOn(OutPoints[i].Transform.GetLocation()));
+		UTEST_TRUE(FString::Format(TEXT("Point {0}, {1} is within the volume"), { i }), TestBounds.IsInsideOrOn(Transforms[i].GetLocation()));
 	}
 
-	UTEST_EQUAL("Point 0 has the correct transform,", OutPoints[0].Transform, FTransform(FRotator::ZeroRotator, FVector(50.0, 50.0, 0.0), FVector::One()));
-	UTEST_EQUAL("Point 1 has the correct transform,", OutPoints[1].Transform, FTransform(FRotator::ZeroRotator, FVector(150.0, 50.0, 0.0), FVector::One()));
-	UTEST_EQUAL("Point 2 has the correct transform,", OutPoints[2].Transform, FTransform(FRotator::ZeroRotator, FVector(50.0, 150.0, 0.0), FVector::One()));
-	UTEST_EQUAL("Point 3 has the correct transform,", OutPoints[3].Transform, FTransform(FRotator::ZeroRotator, FVector(150.0, 150.0, 0.0), FVector::One()));
+	UTEST_EQUAL("Point 0 has the correct transform,", Transforms[0], FTransform(FRotator::ZeroRotator, FVector(50.0, 50.0, 0.0), FVector::One()));
+	UTEST_EQUAL("Point 1 has the correct transform,", Transforms[1], FTransform(FRotator::ZeroRotator, FVector(150.0, 50.0, 0.0), FVector::One()));
+	UTEST_EQUAL("Point 2 has the correct transform,", Transforms[2], FTransform(FRotator::ZeroRotator, FVector(50.0, 150.0, 0.0), FVector::One()));
+	UTEST_EQUAL("Point 3 has the correct transform,", Transforms[3], FTransform(FRotator::ZeroRotator, FVector(150.0, 150.0, 0.0), FVector::One()));
 	
 	return true;
 }
@@ -386,25 +386,25 @@ bool FPCGCreatePointsGridTest_LocalCorner::RunTest(const FString& Parameters)
 
 	UTEST_EQUAL("Output count", Outputs.Num(), 1);
 
-	const UPCGPointData* OutPointData = Cast<UPCGPointData>(Outputs[0].Data);
+	const UPCGBasePointData* OutPointData = Cast<UPCGBasePointData>(Outputs[0].Data);
 
 	UTEST_NOT_NULL("Output point data", OutPointData);
 
-	const TArray<FPCGPoint>& OutPoints = OutPointData->GetPoints();
+	const TConstPCGValueRange<FTransform> Transforms = OutPointData->GetConstTransformValueRange();
 
 	//Cell Corners generate (2 + 1) * (2 + 1) * 2
-	UTEST_EQUAL("Output point count", OutPoints.Num(), 18);
+	UTEST_EQUAL("Output point count", Transforms.Num(), 18);
 
 	// Check to see if the position of the points are within the bounds of the volume
-	for (int i = 0; i < OutPoints.Num(); ++i)
+	for (int i = 0; i < Transforms.Num(); ++i)
 	{
-		UTEST_TRUE(FString::Format(TEXT("Point {0} is within the volume"), { i }), TestBounds.IsInsideOrOn(OutPoints[i].Transform.GetLocation()));
+		UTEST_TRUE(FString::Format(TEXT("Point {0} is within the volume"), { i }), TestBounds.IsInsideOrOn(Transforms[i].GetLocation()));
 	}
 
-	UTEST_EQUAL("Point 1 has the correct transform,", OutPoints[1].Transform, FTransform(FRotator::ZeroRotator, FVector(100.0, 0.0, -50.0), FVector::One()));
-	UTEST_EQUAL("Point 3 has the correct transform,", OutPoints[3].Transform, FTransform(FRotator::ZeroRotator, FVector(0.0, 100.0, -50.0), FVector::One()));
-	UTEST_EQUAL("Point 5 has the correct transform,", OutPoints[5].Transform, FTransform(FRotator::ZeroRotator, FVector(200.0, 100.0, -50.0), FVector::One()));
-	UTEST_EQUAL("Point 11 has the correct transform,", OutPoints[11].Transform, FTransform(FRotator::ZeroRotator, FVector(200.0, 0.0, 50.0), FVector::One()));
+	UTEST_EQUAL("Point 1 has the correct transform,", Transforms[1], FTransform(FRotator::ZeroRotator, FVector(100.0, 0.0, -50.0), FVector::One()));
+	UTEST_EQUAL("Point 3 has the correct transform,", Transforms[3], FTransform(FRotator::ZeroRotator, FVector(0.0, 100.0, -50.0), FVector::One()));
+	UTEST_EQUAL("Point 5 has the correct transform,", Transforms[5], FTransform(FRotator::ZeroRotator, FVector(200.0, 100.0, -50.0), FVector::One()));
+	UTEST_EQUAL("Point 11 has the correct transform,", Transforms[11], FTransform(FRotator::ZeroRotator, FVector(200.0, 0.0, 50.0), FVector::One()));
 
 	return true;
 }
\ No newline at end of file
diff --git a/Engine/Plugins/PCG/Source/PCG/Public/Data/PCGBasePointData.h b/Engine/Plugins/PCG/Source/PCG/Public/Data/PCGBasePointData.h
--- a/Engine/Plugins/PCG/Source/PCG/Public/Data/PCGBasePointData.h
+++ b/Engine/Plugins/PCG/Source/PCG/Public/Data/PCGBasePointData.h
@@ -136,12 +136,16 @@ public:
 	UFUNCTION(BlueprintCallable, Category = SpatialData)
 	void CopyPointsFrom(const UPCGBasePointData* InData, const TArray<int32>& InDataIndices);
 
+	virtual void CopyPointsTo(UPCGBasePointData* OutData, int32 ReadStartIndex, int32 WriteStartIndex, int32 Count) const;
+	virtual EPCGPointNativeProperties GetAllocatedProperties() const { return EPCGPointNativeProperties::All; }
+
 	/** Metadata specific */
 	virtual FPCGMetadataLayerID GetDefaultMetadataLayerID() const override { return PCGMetadataLayerID::Elements; }
 	virtual TArray<FPCGMetadataLayerID> GetAllSupportedMetadataLayerIDs() const override { return {PCGMetadataLayerID::Data, PCGMetadataLayerID::Elements}; }
 
+	static void CopyPoints(const UPCGBasePointData* From, UPCGBasePointData* To, const TArray<int32>& InDataIndices, bool bCopyAll);
+
 protected:
-	void CopyPoints(const UPCGBasePointData* From, UPCGBasePointData* To, const TArray<int32>& InDataIndices, bool bCopyAll) const;
 
 	void RebuildOctreeIfNeeded() const
 	{
diff --git a/Engine/Plugins/PCG/Source/PCG/Public/Data/PCGPointArrayData.h b/Engine/Plugins/PCG/Source/PCG/Public/Data/PCGPointArrayData.h
--- a/Engine/Plugins/PCG/Source/PCG/Public/Data/PCGPointArrayData.h
+++ b/Engine/Plugins/PCG/Source/PCG/Public/Data/PCGPointArrayData.h
@@ -17,7 +17,10 @@ public:
 	//~End UObject Interface
 
 	virtual const UPCGPointData* ToPointData(FPCGContext* Context, const FBox& InBounds = FBox(EForceInit::ForceInit)) const override;
+	virtual const UPCGPointArrayData* ToPointArrayData(FPCGContext* Context) const override { return this; }
 	virtual UPCGSpatialData* CopyInternal(FPCGContext* Context) const override;
+	virtual void CopyPointsTo(UPCGBasePointData* OutData, int32 ReadStartIndex, int32 WriteStartIndex, int32 Count) const;
+	virtual EPCGPointNativeProperties GetAllocatedProperties() const override { return PointArray.GetAllocatedProperties(); }
 
 	//~Begin UPCGBasePointData Interface
 	virtual bool IsValidRef(const PCGPointOctree::FPointRef& InPointRef) const override { return InPointRef.Index >= 0 && InPointRef.Index < GetNumPoints(); }
diff --git a/Engine/Plugins/PCG/Source/PCG/Public/Data/PCGPointData.h b/Engine/Plugins/PCG/Source/PCG/Public/Data/PCGPointData.h
--- a/Engine/Plugins/PCG/Source/PCG/Public/Data/PCGPointData.h
+++ b/Engine/Plugins/PCG/Source/PCG/Public/Data/PCGPointData.h
@@ -110,6 +110,8 @@ PRAGMA_ENABLE_DEPRECATION_WARNINGS
 	virtual TConstPCGValueRange<int64> GetConstMetadataEntryValueRange() const override;
 	
 	virtual void MoveRange(int32 RangeStartIndex, int32 MoveToIndex, int32 NumElements) override;
+
+	virtual void CopyPointsTo(UPCGBasePointData* OutData, int32 ReadStartIndex, int32 WriteStartIndex, int32 Count) const override;
 	// ~End UPCGBasePointData interface
 
 	// ~Begin UPCGSpatialData interface
diff --git a/Engine/Plugins/PCG/Source/PCG/Public/Data/PCGSpatialData.h b/Engine/Plugins/PCG/Source/PCG/Public/Data/PCGSpatialData.h
--- a/Engine/Plugins/PCG/Source/PCG/Public/Data/PCGSpatialData.h
+++ b/Engine/Plugins/PCG/Source/PCG/Public/Data/PCGSpatialData.h
@@ -12,7 +12,9 @@
 
 class AActor;
 struct FPCGContext;
+class UPCGBasePointData;
 class UPCGPointData;
+class UPCGPointArrayData;
 class UPCGIntersectionData;
 class UPCGUnionData;
 class UPCGDifferenceData;
@@ -90,6 +92,10 @@ public:
 
 	virtual const UPCGPointData* ToPointData(FPCGContext* Context, const FBox& InBounds = FBox(EForceInit::ForceInit)) const PURE_VIRTUAL(UPCGSpatialData::ToPointData, return nullptr;);
 
+	virtual const UPCGPointArrayData* ToPointArrayData(FPCGContext* Context) const;
+
+	const UPCGBasePointData* ToBasePointData(FPCGContext* Context) const;
+
 	/** Sample rotation, scale and other attributes from this data at the query position. Returns true if Transform location and Bounds overlaps this data. */
 	virtual bool SamplePoint(const FTransform& Transform, const FBox& Bounds, FPCGPoint& OutPoint, UPCGMetadata* OutMetadata) const PURE_VIRTUAL(UPCGSpatialData::SamplePoint, return false;);
 
diff --git a/Engine/Plugins/PCG/Source/PCG/Public/Elements/PCGBoundsModifier.h b/Engine/Plugins/PCG/Source/PCG/Public/Elements/PCGBoundsModifier.h
--- a/Engine/Plugins/PCG/Source/PCG/Public/Elements/PCGBoundsModifier.h
+++ b/Engine/Plugins/PCG/Source/PCG/Public/Elements/PCGBoundsModifier.h
@@ -61,6 +61,9 @@ public:
 
 class FPCGBoundsModifier : public FPCGPointOperationElementBase
 {
+public:
+	virtual bool SupportsBasePointDataInputs(FPCGContext* InContext) const { return true; }
 protected:
 	virtual bool ExecuteInternal(FPCGContext* Context) const override;
+	virtual EPCGPointNativeProperties GetPropertiesToAllocate(const UPCGBasePointData* Input) const override;
 };
\ No newline at end of file
diff --git a/Engine/Plugins/PCG/Source/PCG/Public/Elements/PCGCreatePointsGrid.h b/Engine/Plugins/PCG/Source/PCG/Public/Elements/PCGCreatePointsGrid.h
--- a/Engine/Plugins/PCG/Source/PCG/Public/Elements/PCGCreatePointsGrid.h
+++ b/Engine/Plugins/PCG/Source/PCG/Public/Elements/PCGCreatePointsGrid.h
@@ -93,6 +93,7 @@ class FPCGCreatePointsGridElement : public IPCGElement
 {
 public:
 	virtual void GetDependenciesCrc(const FPCGDataCollection& InInput, const UPCGSettings* InSettings, UPCGComponent* InComponent, FPCGCrc& OutCrc) const override;
+	virtual bool SupportsBasePointDataInputs(FPCGContext* InContext) const override { return true; }
 
 protected:
 	virtual bool ExecuteInternal(FPCGContext* Context) const override;
diff --git a/Engine/Plugins/PCG/Source/PCG/Public/Elements/PCGPointOperationElementBase.h b/Engine/Plugins/PCG/Source/PCG/Public/Elements/PCGPointOperationElementBase.h
--- a/Engine/Plugins/PCG/Source/PCG/Public/Elements/PCGPointOperationElementBase.h
+++ b/Engine/Plugins/PCG/Source/PCG/Public/Elements/PCGPointOperationElementBase.h
@@ -18,8 +18,18 @@ namespace PCGPointOperation
 	/** Stores the input and output data as the state of the time sliced execution */
 	struct IterationState
 	{
+		// This will be set to the same value as InputData when using UPCGPointData type to support deprecated code
+		UE_DEPRECATED(5.6, "Use InputData instead")
 		const UPCGPointData* InputPointData = nullptr;
+
+		// This will be set to the same value as OutputData when using UPCGPointData type to support deprecated code
+		UE_DEPRECATED(5.6, "Use OutputData instead")
 		UPCGPointData* OutputPointData = nullptr;
+
+		const UPCGBasePointData* InputData = nullptr;
+
+		UPCGBasePointData* OutputData = nullptr;
+
 		int32 NumPoints = 0;
 	};
 }
@@ -29,8 +39,6 @@ namespace PCGPointOperation
  */
 class PCG_API FPCGPointOperationElementBase : public TPCGTimeSlicedElementBase<PCGTimeSlice::FEmptyStruct, PCGPointOperation::IterationState>
 {
-	using PointExecSignature = bool(const FPCGPoint& InPoint, FPCGPoint& OutPoint);
-
 protected:
 	//~Begin IPCGElement interface
 	/** Conveniently calls PreparePointOperationData to prepare the time sliced element for execution. May be overridden, but PreparePointOperationData must be called. */
@@ -39,8 +47,25 @@ protected:
 	//~End IPCGElement interface
 
 	/** Executes the PointFunction function/lambda for every point copied from PreparePointOperationData */
-	bool ExecutePointOperation(ContextType* Context, TFunctionRef<PointExecSignature> PointFunction, int32 PointsPerChunk = PCGPointOperation::Constants::PointsPerChunk) const;
+	template <typename Func = TFunctionRef<bool(const FPCGPoint& InPoint, FPCGPoint& OutPoint)>>
+	bool ExecutePointOperation(ContextType* Context, Func Callback, int32 PointsPerChunk = PCGPointOperation::Constants::PointsPerChunk) const
+	{
+		if constexpr (std::is_invocable_v<Func, const FPCGPoint&, FPCGPoint&>)
+		{
+			return ExecutePointOperationWithPoints(Context, Callback, PointsPerChunk);
+		}
+		else
+		{
+			return ExecutePointOperationWithIndices(Context, Callback, PointsPerChunk);
+		}
+	}
 
 	/** Mandatory call. Using the context, prepares the state data for time slice execution */
 	bool PreparePointOperationData(ContextType* Context, FName InputPinLabel = PCGPinConstants::DefaultInputLabel) const;
+
+	virtual EPCGPointNativeProperties GetPropertiesToAllocate(const UPCGBasePointData* Input) const { return EPCGPointNativeProperties::All; }
+
+private:
+	bool ExecutePointOperationWithPoints(ContextType* Context, TFunctionRef<bool(const FPCGPoint& InPoint, FPCGPoint& OutPoint)> Callback, int32 PointsPerChunk) const;
+	bool ExecutePointOperationWithIndices(ContextType* Context, TFunctionRef<bool(const UPCGBasePointData* InPointData, UPCGBasePointData* OutPointData, int32 StartIndex, int32 Count)> Callback, int32 PointsPerChunk) const;
 };
\ No newline at end of file
diff --git a/Engine/Plugins/PCG/Source/PCG/Public/Elements/PCGSurfaceSampler.h b/Engine/Plugins/PCG/Source/PCG/Public/Elements/PCGSurfaceSampler.h
--- a/Engine/Plugins/PCG/Source/PCG/Public/Elements/PCGSurfaceSampler.h
+++ b/Engine/Plugins/PCG/Source/PCG/Public/Elements/PCGSurfaceSampler.h
@@ -11,6 +11,7 @@
 
 class UPCGNode;
 class UPCGPointData;
+class UPCGBasePointData;
 class UPCGSurfaceData;
 class UPCGSurfaceSamplerSettings;
 
@@ -78,14 +79,21 @@ namespace PCGSurfaceSampler
 	struct FSurfaceSamplerIterationState
 	{
 		FSurfaceSamplerData SamplerData;
+
+		UPCGBasePointData* OutputData = nullptr;
+
+		// This will be set to the same value as OutputData when using UPCGPointData type to support deprecated code
+		UE_DEPRECATED(5.6, "Use OutputData instead")
 		UPCGPointData* OutputPoints = nullptr;
 	};
 
 	/** Sample a surface and returns the resulting point data. Can't be timesliced. */
+	PCG_API UPCGBasePointData* SampleSurface(FPCGContext* Context, const FSurfaceSamplerParams& ExecutionParams, const UPCGSurfaceData* InSurface, const UPCGSpatialData* InBoundingShape, const FBox& EffectiveBounds);
+		
 	PCG_API UPCGPointData* SampleSurface(FPCGContext* Context, const UPCGSurfaceData* InSurface, const UPCGSpatialData* InBoundingShape, const FBox& EffectiveBounds, const FSurfaceSamplerParams& ExecutionParams);
 
 	/** Sample a surface and write the results in the given point data. Can be timesliced, and will return false if the processing is not done, true otherwise. */
-	bool SampleSurface(FPCGContext* Context, const FSurfaceSamplerData& SamplerData, const UPCGSurfaceData* InSurface, const UPCGSpatialData* InBoundingShape, UPCGPointData* SampledData, const bool bTimeSlicingIsEnabled = false);
+	bool SampleSurface(FPCGContext* Context, const FSurfaceSamplerData& SamplerData, const UPCGSurfaceData* InSurface, const UPCGSpatialData* InBoundingShape, UPCGBasePointData* SampledData, const bool bTimeSlicingIsEnabled = false);
 }
 
 UCLASS(MinimalAPI, BlueprintType, ClassGroup = (Procedural))
@@ -161,6 +169,7 @@ public:
 	virtual void GetDependenciesCrc(const FPCGDataCollection& InInput, const UPCGSettings* InSettings, UPCGComponent* InComponent, FPCGCrc& OutCrc) const override;
 	// Might be sampling landscape or other external data, worth computing a full CRC in case we can halt change propagation/re-executions
 	virtual bool ShouldComputeFullOutputDataCrc(FPCGContext* Context) const override { return true; }
+	virtual bool SupportsBasePointDataInputs(FPCGContext* InContext) const { return true; }
 
 protected:
 	virtual bool PrepareDataInternal(FPCGContext* InContext) const override;
diff --git a/Engine/Plugins/PCG/Source/PCG/Public/PCGCommon.h b/Engine/Plugins/PCG/Source/PCG/Public/PCGCommon.h
--- a/Engine/Plugins/PCG/Source/PCG/Public/PCGCommon.h
+++ b/Engine/Plugins/PCG/Source/PCG/Public/PCGCommon.h
@@ -69,6 +69,7 @@ enum class EPCGPointNativeProperties : uint32
 	Steepness = 1 << 5,
 	Seed = 1 << 6,
 	MetadataEntry = 1 << 7,
+	All = Transform | Density | BoundsMin | BoundsMax | Color | Steepness | Seed | MetadataEntry,
 };
 ENUM_CLASS_FLAGS(EPCGPointNativeProperties);
 
diff --git a/Engine/Plugins/PCG/Source/PCG/Public/PCGContext.h b/Engine/Plugins/PCG/Source/PCG/Public/PCGContext.h
--- a/Engine/Plugins/PCG/Source/PCG/Public/PCGContext.h
+++ b/Engine/Plugins/PCG/Source/PCG/Public/PCGContext.h
@@ -12,6 +12,7 @@
 
 #include "PCGContext.generated.h"
 
+class UPCGBasePointData;
 class UPCGComponent;
 class UPCGGraphInterface;
 class UPCGSettingsInterface;
@@ -19,6 +20,8 @@ class UPCGSpatialData;
 struct FPCGSettingsOverridableParam;
 struct FPCGStack;
 
+extern PCG_API TAutoConsoleVariable<bool> CVarPCGEnablePointArrayData;
+
 namespace PCGContextHelpers
 {
 	template<typename SettingsType>
@@ -119,6 +122,9 @@ struct PCG_API FPCGContext
 		return Context ? Context->NewObject_AnyThread_Impl<T>(std::forward<Args>(InArgs)...) : ::NewObject<T>(std::forward<Args>(InArgs)...);
 	}
 
+	// Temporary so that we can toggle between UPCGPointData and UPCGPointArrayData with a CVar
+	static UPCGBasePointData* NewPointData_AnyThread(FPCGContext* Context);
+
 	bool ContainsAsyncObject(const UObject* InAsyncObject) { return AsyncObjects.Contains(InAsyncObject); }
 
 	// Return the settings casted in the wanted type.
diff --git a/Engine/Plugins/PCG/Source/PCG/Public/PCGPointArray.h b/Engine/Plugins/PCG/Source/PCG/Public/PCGPointArray.h
--- a/Engine/Plugins/PCG/Source/PCG/Public/PCGPointArray.h
+++ b/Engine/Plugins/PCG/Source/PCG/Public/PCGPointArray.h
@@ -182,6 +182,30 @@ public:
 		Ar << Values;
 	}
 
+	FORCEINLINE void CopyTo(FPCGPointArrayProperty<T>& OutProperty, int32 StartReadIndex, int32 StartWriteIndex, int32 Count) const
+	{
+		check(StartReadIndex + Count <= NumValues && StartWriteIndex + Count <= OutProperty.NumValues);
+		if (!IsAllocated())
+		{
+			if (!OutProperty.IsAllocated())
+			{
+				OutProperty.Value = Value;
+			}
+			else
+			{
+				for (int32 Index = StartWriteIndex; Index < StartWriteIndex + Count; ++Index)
+				{
+					OutProperty.Values[Index] = Value;
+				}
+			}
+		}
+		else
+		{
+			check(OutProperty.IsAllocated());
+			std::memcpy(&OutProperty.Values[StartWriteIndex], &Values[StartReadIndex], sizeof(T) * Count);
+		}
+	}
+
 private:
 	// Array containing values if allocated
 	TArray<T> Values;
@@ -203,6 +227,7 @@ public:
 	void SetNumPoints(int32 InNumPoints);
 	int32 GetNumPoints() const { return NumPoints; }
 		
+	EPCGPointNativeProperties GetAllocatedProperties() const;
 	void Allocate(EPCGPointNativeProperties InProperties);
 	void Free(EPCGPointNativeProperties InProperties);
 	void MoveRange(int32 RangeStartIndex, int32 MoveToIndex, int32 NumElements);
@@ -299,7 +324,8 @@ public:
 
 	// WithSerializer & WithIdenticalViaEquality needed for this struct to properly serialize
 	bool Serialize(FArchive& Ar);
-		
+
+	void CopyPointsTo(FPCGPointArray& OutPointArray, int32 InStartReadIndex, int32 InStartWriteIndex, int32 Count) const;
 private:
 	FPCGPointArrayProperty<FTransform>	Transform;
 	FPCGPointArrayProperty<float>		Density;
