[PCG] Fix multiple issues with refactor of attribute property selector

* Parsing was not correct in some invalid cases and was crashing
* Add tests to cover found crashes

diff --git a/Engine/Plugins/PCG/Source/PCG/Private/Metadata/PCGAttributePropertySelector.cpp b/Engine/Plugins/PCG/Source/PCG/Private/Metadata/PCGAttributePropertySelector.cpp
--- a/Engine/Plugins/PCG/Source/PCG/Private/Metadata/PCGAttributePropertySelector.cpp
+++ b/Engine/Plugins/PCG/Source/PCG/Private/Metadata/PCGAttributePropertySelector.cpp
@@ -337,28 +337,9 @@ bool FPCGAttributePropertySelector::IsValid() const
 
 bool FPCGAttributePropertySelector::Reset()
 {
-	bool bHasChanged = false;
-	if (LayerName != NAME_None)
-	{
-		bHasChanged = true;
-		LayerName = NAME_None;
-	}
-
-	if (AttributeName != NAME_None)
-	{
-		bHasChanged = true;
-		AttributeName = NAME_None;
-	}
-	
-	if (PropertyName != NAME_None)
-	{
-		bHasChanged = true;
-		PropertyName = NAME_None;
-	}
-
-
-	bHasChanged |= ResetExtraNames();
-
+	static const FPCGAttributePropertySelector EmptySelector{};
+	const bool bHasChanged = EmptySelector != *this;
+	*this = EmptySelector;
 	return bHasChanged;
 }
 
@@ -379,9 +360,7 @@ bool FPCGAttributePropertySelector::Update(const FString& NewValue)
 {
 	if (NewValue.IsEmpty())
 	{
-		const FPCGAttributePropertySelector Previous = *this;
-		Reset();
-		return *this == Previous;
+		return Reset();
 	}
 	
 	TArray<FString> NewValues;
@@ -390,7 +369,7 @@ bool FPCGAttributePropertySelector::Update(const FString& NewValue)
 	check(!NewValues.IsEmpty())
 
 	// TODO: If we ever have to support multiple layers, this has to change
-	const bool bHasLayerName = !NewValues[0].IsEmpty() && NewValues[0][0] == PCGAttributePropertySelectorConstants::LayerPrefixChar && !PCGAttributePropertySelector::IsReservedAttributeName(*NewValues[0]);
+	const bool bHasLayerName = NewValues[0].Len() > 1 && NewValues[0][0] == PCGAttributePropertySelectorConstants::LayerPrefixChar && !PCGAttributePropertySelector::IsReservedAttributeName(*NewValues[0]);
 	const int32 PropertyIndex = bHasLayerName ? 1 : 0;
 	const bool bHasProperty = NewValues.Num() > PropertyIndex && !NewValues[PropertyIndex].IsEmpty() && NewValues[PropertyIndex][0] == PCGAttributePropertySelectorConstants::PropertyPrefixChar;
 
@@ -405,9 +384,12 @@ bool FPCGAttributePropertySelector::Update(const FString& NewValue)
 	bool bHasChanged = ExtraNamesTemp != ExtraNames;
 	ExtraNames = MoveTemp(ExtraNamesTemp);
 
-	if (bHasLayerName)
+	const FName NewLayerName = bHasLayerName ? FName{*NewValues[0].RightChop(1)} : FName{NAME_None};
+	bHasChanged |= SetLayerName(NewLayerName, /*bResetExtraNames=*/ false);
+
+	if (NewValues.Num() <= PropertyIndex)
 	{
-		bHasChanged |= SetLayerName(*NewValues[0].RightChop(1), /*bResetExtraNames=*/ false);
+		return SetAttributeName(NAME_None, /*bResetExtraNames=*/ false) || bHasChanged;
 	}
 	
 	if (bHasProperty)
diff --git a/Engine/Plugins/PCG/Source/PCG/Private/Tests/Accessors/PCGAttributePropertySelectorTests.cpp b/Engine/Plugins/PCG/Source/PCG/Private/Tests/Accessors/PCGAttributePropertySelectorTests.cpp
--- a/Engine/Plugins/PCG/Source/PCG/Private/Tests/Accessors/PCGAttributePropertySelectorTests.cpp
+++ b/Engine/Plugins/PCG/Source/PCG/Private/Tests/Accessors/PCGAttributePropertySelectorTests.cpp
@@ -162,6 +162,12 @@ IMPLEMENT_CUSTOM_SIMPLE_AUTOMATION_TEST(FPCGAttributePropertySelector_Deprecatio
 IMPLEMENT_CUSTOM_SIMPLE_AUTOMATION_TEST(FPCGAttributePropertySelector_Deprecation_CDO_NewVersion_PointProperty, FPCGAttributePropertySelectorTests, "Plugins.PCG.AttributePropertySelector.Deprecation.CDO.NewVersion.PointProperty", PCGTestsCommon::TestFlags)
 IMPLEMENT_CUSTOM_SIMPLE_AUTOMATION_TEST(FPCGAttributePropertySelector_Deprecation_CDO_NewVersion_WithExtraNames_PointProperty, FPCGAttributePropertySelectorTests, "Plugins.PCG.AttributePropertySelector.Deprecation.CDO.NewVersion.WithExtraNames.PointProperty", PCGTestsCommon::TestFlags)
 
+IMPLEMENT_CUSTOM_SIMPLE_AUTOMATION_TEST(FPCGAttributePropertySelector_Invalid_JustAt, FPCGAttributePropertySelectorTests, "Plugins.PCG.AttributePropertySelector.Invalid.JustAt", PCGTestsCommon::TestFlags)
+IMPLEMENT_CUSTOM_SIMPLE_AUTOMATION_TEST(FPCGAttributePropertySelector_Invalid_InvalidSymbols, FPCGAttributePropertySelectorTests, "Plugins.PCG.AttributePropertySelector.Invalid.InvalidSymbols", PCGTestsCommon::TestFlags)
+
+IMPLEMENT_CUSTOM_SIMPLE_AUTOMATION_TEST(FPCGAttributePropertySelector_JustLayerName, FPCGAttributePropertySelectorTests, "Plugins.PCG.AttributePropertySelector.JustLayerName", PCGTestsCommon::TestFlags)
+IMPLEMENT_CUSTOM_SIMPLE_AUTOMATION_TEST(FPCGAttributePropertySelector_UpdateLayerNameToAtLast, FPCGAttributePropertySelectorTests, "Plugins.PCG.AttributePropertySelector.UpdateLayerNameToAtLast", PCGTestsCommon::TestFlags)
+
 bool FPCGAttributePropertySelector_FromString_Attribute::RunTest(const FString& Parameters)
 {
 	return FromStringComparison(TEXT("MyAttr"), FPCGAttributePropertySelector::CreateAttributeSelector(TEXT("MyAttr")));
@@ -408,4 +414,41 @@ bool FPCGAttributePropertySelector_Deprecation_CDO_NewVersion_WithExtraNames_Poi
 	return ValidateSelectorCDOWithNewerVersion(FPCGAttributePropertySelector::CreatePointPropertySelector(EPCGPointProperties::Rotation, NAME_None, {TEXT("Forward")}), EPCGPointProperties::Rotation, {TEXT("Forward")});
 }
 
+bool FPCGAttributePropertySelector_Invalid_JustAt::RunTest(const FString& Parameters)
+{
+	FPCGAttributePropertySelector Selector = FPCGAttributePropertySelector::CreateSelectorFromString(TEXT("@"));
+	UTEST_EQUAL("Selector has no layer name", Selector.GetLayerName(), FName{NAME_None});
+	UTEST_EQUAL("Selector is an attribute", Selector.GetSelection(), EPCGAttributePropertySelection::Attribute);
+	UTEST_FALSE("Selector is invalid", Selector.IsValid());
+	return true;
+}
+
+bool FPCGAttributePropertySelector_Invalid_InvalidSymbols::RunTest(const FString& Parameters)
+{
+	FPCGAttributePropertySelector Selector = FPCGAttributePropertySelector::CreateSelectorFromString(TEXT(")(*&^"));
+	UTEST_EQUAL("Selector has no layer name", Selector.GetLayerName(), FName{NAME_None});
+	UTEST_EQUAL("Selector is an attribute", Selector.GetSelection(), EPCGAttributePropertySelection::Attribute);
+	UTEST_FALSE("Selector is invalid", Selector.IsValid());
+	return true;
+}
+
+bool FPCGAttributePropertySelector_JustLayerName::RunTest(const FString& Parameters)
+{
+	FPCGAttributePropertySelector Selector = FPCGAttributePropertySelector::CreateSelectorFromString(TEXT("@MyLayer"));
+	UTEST_EQUAL("Selector has the right layer name", Selector.GetLayerName(), FName{TEXT("MyLayer")});
+	UTEST_EQUAL("Selector is an attribute", Selector.GetSelection(), EPCGAttributePropertySelection::Attribute);
+	UTEST_EQUAL("Selector has an attribute name of None", Selector.GetAttributeName(), FName{NAME_None});
+	return true;
+}
+
+bool FPCGAttributePropertySelector_UpdateLayerNameToAtLast::RunTest(const FString& Parameters)
+{
+	FPCGAttributePropertySelector Selector = FPCGAttributePropertySelector::CreateSelectorFromString(TEXT("@Las"));
+	Selector.Update(TEXT("@Last"));
+	UTEST_EQUAL("Selector has no layer name", Selector.GetLayerName(), FName{NAME_None});
+	UTEST_EQUAL("Selector is an attribute", Selector.GetSelection(), EPCGAttributePropertySelection::Attribute);
+	UTEST_EQUAL("Selector has an attribute name of @Last", Selector.GetAttributeName(), FName{TEXT("@Last")});
+	return true;
+}
+
 #endif // WITH_EDITOR
