[PCG] Rename PCG Metadata Layers to PCG Metadata Domains

#rnx

diff --git a/Engine/Plugins/PCG/Config/BasePCG.ini b/Engine/Plugins/PCG/Config/BasePCG.ini
--- a/Engine/Plugins/PCG/Config/BasePCG.ini
+++ b/Engine/Plugins/PCG/Config/BasePCG.ini
@@ -78,6 +78,7 @@
 +PropertyRedirects=(OldName="PCGAttributeTransferSettings.SourceAttributeProperty",NewName="PCGAttributeTransferSettings.InputSource")
 +PropertyRedirects=(OldName="PCGAttributeTransferSettings.TargetAttributeProperty",NewName="PCGAttributeTransferSettings.OutputTarget")
 +PropertyRedirects=(OldName="PCGRandomChoiceSettings.bOutputDiscardedPoints",NewName="PCGRandomChoiceSettings.bOutputDiscardedEntries")
++PropertyRedirects=(OldName="/Script/PCG.PCGAttributePropertySelector.LayerName",NewName="/Script/PCG.PCGAttributePropertySelector.DomainName")
 
 +ClassRedirects=(OldName="PCGPointSamplerSettings",NewName="/Script/PCG.PCGSelectPointsSettings")
 +ClassRedirects=(OldName="PCGDensityNoiseSettings",NewName="/Script/PCG.PCGAttributeNoiseSettings")
diff --git a/Engine/Plugins/PCG/Source/PCG/Private/Data/PCGBasePointData.cpp b/Engine/Plugins/PCG/Source/PCG/Private/Data/PCGBasePointData.cpp
--- a/Engine/Plugins/PCG/Source/PCG/Private/Data/PCGBasePointData.cpp
+++ b/Engine/Plugins/PCG/Source/PCG/Private/Data/PCGBasePointData.cpp
@@ -16,6 +16,8 @@
 #include "GameFramework/Actor.h"
 #include "Serialization/ArchiveCrc32.h"
 
+#define LOCTEXT_NAMESPACE "PCGBasePointData"
+
 static TAutoConsoleVariable<bool> CVarCacheFullPointDataCrc(
 	TEXT("pcg.Cache.FullPointDataCrc"),
 	true,
@@ -675,7 +677,7 @@ UPCGBasePointData::UPCGBasePointData(const FObjectInitializer& ObjectInitializer
 	: Super(ObjectInitializer)
 {
 	check(Metadata);
-	Metadata->SetupLayer(PCGMetadataLayerID::Elements, /*bIsDefault=*/true);
+	Metadata->SetupDomain(PCGMetadataDomainID::Elements, /*bIsDefault=*/true);
 }
 
 FPCGAttributeAccessorMethods UPCGBasePointData::GetPointAccessorMethods()
@@ -710,6 +712,10 @@ FPCGAttributeAccessorMethods UPCGBasePointData::GetPointAccessorMethods()
 		.CreateConstAccessorKeysFunc = CreateConstAccessorKeysFunc
 	};
 
+#if WITH_EDITOR
+	Methods.FillSelectorMenuEntryFromEnum<EPCGPointProperties>({LOCTEXT("PointSelectorMenuEntry", "Point")});
+#endif // WITH_EDITOR
+
 	return Methods;
 }
 
@@ -1490,4 +1496,6 @@ void UPCGBasePointData::AddToCrc(FArchiveCrc32& Ar, bool bFullDataCrc) const
 	{
 		PCGMetadata->AddToCrc(Ar, bFullDataCrc);
 	}
-}
\ No newline at end of file
+}
+
+#undef LOCTEXT_NAMESPACE
\ No newline at end of file
diff --git a/Engine/Plugins/PCG/Source/PCG/Private/Data/PCGSpatialData.cpp b/Engine/Plugins/PCG/Source/PCG/Private/Data/PCGSpatialData.cpp
--- a/Engine/Plugins/PCG/Source/PCG/Private/Data/PCGSpatialData.cpp
+++ b/Engine/Plugins/PCG/Source/PCG/Private/Data/PCGSpatialData.cpp
@@ -19,7 +19,7 @@ UPCGSpatialData::UPCGSpatialData(const FObjectInitializer& ObjectInitializer)
 	: Super(ObjectInitializer)
 {
 	Metadata = ObjectInitializer.CreateDefaultSubobject<UPCGMetadata>(this, TEXT("Metadata"));
-	Metadata->SetupLayer(PCGMetadataLayerID::Data, /*bIsDefault=*/true);
+	Metadata->SetupDomain(PCGMetadataDomainID::Data, /*bIsDefault=*/true);
 }
 
 const UPCGBasePointData* UPCGSpatialData::ToBasePointData(FPCGContext* Context) const
diff --git a/Engine/Plugins/PCG/Source/PCG/Private/Data/PCGSplineData.cpp b/Engine/Plugins/PCG/Source/PCG/Private/Data/PCGSplineData.cpp
--- a/Engine/Plugins/PCG/Source/PCG/Private/Data/PCGSplineData.cpp
+++ b/Engine/Plugins/PCG/Source/PCG/Private/Data/PCGSplineData.cpp
@@ -184,9 +184,9 @@ TUniquePtr<IPCGAttributeAccessor> UPCGSplineData::CreateStaticAccessor(const FPC
 	static const FProperty* PCGSplineStructProperty = UPCGSplineData::StaticClass()->FindPropertyByName(GET_MEMBER_NAME_CHECKED(UPCGSplineData, SplineStruct));
 	static const FProperty* SplineCurveProperty = FPCGSplineStruct::StaticStruct()->FindPropertyByName(GET_MEMBER_NAME_CHECKED(FPCGSplineStruct, SplineCurves));
 
-	const FName LayerName = InSelector.GetLayerName();
+	const FName DomainName = InSelector.GetDomainName();
 
-	if (InSelector.GetSelection() == EPCGAttributePropertySelection::Property && (LayerName.IsNone() || LayerName == PCGSplineData::ControlPointLayerName))
+	if (InSelector.GetSelection() == EPCGAttributePropertySelection::Property && (DomainName.IsNone() || DomainName == PCGSplineData::ControlPointDomainName))
 	{
 		const FName PropertyName = InSelector.GetName();
 		TArray<const FProperty*> ExtraProperties = {PCGSplineStructProperty, SplineCurveProperty};
diff --git a/Engine/Plugins/PCG/Source/PCG/Private/Helpers/PCGMetadataHelpers.cpp b/Engine/Plugins/PCG/Source/PCG/Private/Helpers/PCGMetadataHelpers.cpp
--- a/Engine/Plugins/PCG/Source/PCG/Private/Helpers/PCGMetadataHelpers.cpp
+++ b/Engine/Plugins/PCG/Source/PCG/Private/Helpers/PCGMetadataHelpers.cpp
@@ -19,14 +19,14 @@ namespace PCGMetadataHelpers
 		return Metadata1 && Metadata2 && Metadata1->GetRoot() == Metadata2->GetRoot();
 	}
 	
-	bool HasSameRoot(const FPCGMetadataLayer* Metadata1, const FPCGMetadataLayer* Metadata2)
+	bool HasSameRoot(const FPCGMetadataDomain* Metadata1, const FPCGMetadataDomain* Metadata2)
 	{
 		return Metadata1 && Metadata2 && Metadata1->GetRoot() == Metadata2->GetRoot();
 	}
 	
-	FPCGMetadataLayer* GetDefaultMetadataLayer(UPCGMetadata* InMetadata)
+	FPCGMetadataDomain* GetDefaultMetadataDomain(UPCGMetadata* InMetadata)
 	{
-		return InMetadata ? InMetadata->GetDefaultMetadataLayer() : nullptr;
+		return InMetadata ? InMetadata->GetDefaultMetadataDomain() : nullptr;
 	}
 
 	const UPCGMetadata* GetParentMetadata(const UPCGMetadata* Metadata)
@@ -40,7 +40,7 @@ namespace PCGMetadataHelpers
 		return Parent.Get();
 	}
 
-	const FPCGMetadataLayer* GetParentMetadata(const FPCGMetadataLayer* Metadata)
+	const FPCGMetadataDomain* GetParentMetadata(const FPCGMetadataDomain* Metadata)
 	{
 		check(Metadata);
 		return Metadata->GetParent();
@@ -396,7 +396,7 @@ namespace PCGMetadataHelpers
 		return CopyAttributes(TargetData, SourceData, AttributeSelectors, /*bSameOrigin=*/false, OptionalContext);
 	}
 
-	void ComputePointWeightedAttribute(FPCGMetadataLayer* InOutMetadata, FPCGPoint& OutPoint, const TArrayView<TPair<const FPCGPoint*, float>>& InWeightedPoints, const FPCGMetadataLayer* InMetadata)
+	void ComputePointWeightedAttribute(FPCGMetadataDomain* InOutMetadata, FPCGPoint& OutPoint, const TArrayView<TPair<const FPCGPoint*, float>>& InWeightedPoints, const FPCGMetadataDomain* InMetadata)
 	{
 		check(InOutMetadata);
 		TArray<TPair<PCGMetadataEntryKey, float>, TInlineAllocator<4>> InWeightedKeys;
@@ -410,7 +410,7 @@ namespace PCGMetadataHelpers
 		InOutMetadata->ComputeWeightedAttribute(OutPoint.MetadataEntry, MakeArrayView(InWeightedKeys), InMetadata);
 	}
 
-	void SetPointAttributes(FPCGMetadataLayer* InOutMetadata, const TArrayView<const FPCGPoint>& InPoints, const FPCGMetadataLayer* InMetadata, const TArrayView<FPCGPoint>& OutPoints, FPCGContext* OptionalContext)
+	void SetPointAttributes(FPCGMetadataDomain* InOutMetadata, const TArrayView<const FPCGPoint>& InPoints, const FPCGMetadataDomain* InMetadata, const TArrayView<FPCGPoint>& OutPoints, FPCGContext* OptionalContext)
 	{
 		check(InOutMetadata);
 		if (!InMetadata || InMetadata->GetAttributeCount() == 0 || InOutMetadata->GetAttributeCount() == 0)
diff --git a/Engine/Plugins/PCG/Source/PCG/Private/Metadata/Accessors/PCGAttributeAccessorHelpers.cpp b/Engine/Plugins/PCG/Source/PCG/Private/Metadata/Accessors/PCGAttributeAccessorHelpers.cpp
--- a/Engine/Plugins/PCG/Source/PCG/Private/Metadata/Accessors/PCGAttributeAccessorHelpers.cpp
+++ b/Engine/Plugins/PCG/Source/PCG/Private/Metadata/Accessors/PCGAttributeAccessorHelpers.cpp
@@ -473,11 +473,11 @@ TUniquePtr<const IPCGAttributeAccessor> PCGAttributeAccessorHelpers::CreateConst
 	}
 	else
 	{
-		return CreateConstAccessor(InAttribute, InMetadata->GetConstDefaultMetadataLayer(), bQuiet);
+		return CreateConstAccessor(InAttribute, InMetadata->GetConstDefaultMetadataDomain(), bQuiet);
 	}
 }
 
-TUniquePtr<const IPCGAttributeAccessor> PCGAttributeAccessorHelpers::CreateConstAccessor(const FPCGMetadataAttributeBase* InAttribute, const FPCGMetadataLayer* InMetadata, bool bQuiet)
+TUniquePtr<const IPCGAttributeAccessor> PCGAttributeAccessorHelpers::CreateConstAccessor(const FPCGMetadataAttributeBase* InAttribute, const FPCGMetadataDomain* InMetadata, bool bQuiet)
 {
 	if (!InMetadata || !InAttribute)
 	{
@@ -550,11 +550,11 @@ TUniquePtr<IPCGAttributeAccessor> PCGAttributeAccessorHelpers::CreateAccessor(FP
 	}
 	else
 	{
-		return CreateAccessor(InAttribute, InMetadata->GetDefaultMetadataLayer());
+		return CreateAccessor(InAttribute, InMetadata->GetDefaultMetadataDomain());
 	}
 }
 
-TUniquePtr<IPCGAttributeAccessor> PCGAttributeAccessorHelpers::CreateAccessor(FPCGMetadataAttributeBase* InAttribute, FPCGMetadataLayer* InMetadata, bool bQuiet)
+TUniquePtr<IPCGAttributeAccessor> PCGAttributeAccessorHelpers::CreateAccessor(FPCGMetadataAttributeBase* InAttribute, FPCGMetadataDomain* InMetadata, bool bQuiet)
 {
 	if (!InMetadata || !InAttribute)
 	{
@@ -621,7 +621,7 @@ TUniquePtr<IPCGAttributeAccessor> PCGAttributeAccessorHelpers::CreateAccessorWit
 			InMatchingAccessor->Get<AttributeType>(DefaultValue, FPCGAttributeAccessorKeysEntries(PCGInvalidEntryKey));
 			FPCGMetadataAttribute<AttributeType>* Attribute = Metadata->CreateAttribute<AttributeType>(AttributeName, DefaultValue, /*bAllowInterpolation=*/true, /*bOverrideParent=*/false);
 
-			return Attribute ? MakeUnique<FPCGAttributeAccessor<AttributeType>>(Attribute, Attribute->GetMetadataLayer()) : nullptr;
+			return Attribute ? MakeUnique<FPCGAttributeAccessor<AttributeType>>(Attribute, Attribute->GetMetadataDomain()) : nullptr;
 		};
 
 		Result = PCGMetadataAttribute::CallbackWithRightType(InMatchingAccessor->GetUnderlyingType(), std::move(CreateAttributeAndAccessor));
diff --git a/Engine/Plugins/PCG/Source/PCG/Private/Metadata/Accessors/PCGAttributeAccessorKeys.cpp b/Engine/Plugins/PCG/Source/PCG/Private/Metadata/Accessors/PCGAttributeAccessorKeys.cpp
--- a/Engine/Plugins/PCG/Source/PCG/Private/Metadata/Accessors/PCGAttributeAccessorKeys.cpp
+++ b/Engine/Plugins/PCG/Source/PCG/Private/Metadata/Accessors/PCGAttributeAccessorKeys.cpp
@@ -22,21 +22,21 @@ FPCGAttributeAccessorKeysEntries::FPCGAttributeAccessorKeysEntries(PCGMetadataEn
 }
 
 FPCGAttributeAccessorKeysEntries::FPCGAttributeAccessorKeysEntries(const UPCGMetadata* Metadata, bool bAddDefaultValueIfEmpty)
-	: FPCGAttributeAccessorKeysEntries(Metadata ? Metadata->GetConstDefaultMetadataLayer() : nullptr, bAddDefaultValueIfEmpty)
+	: FPCGAttributeAccessorKeysEntries(Metadata ? Metadata->GetConstDefaultMetadataDomain() : nullptr, bAddDefaultValueIfEmpty)
 {}
 
 FPCGAttributeAccessorKeysEntries::FPCGAttributeAccessorKeysEntries(UPCGMetadata* Metadata, bool bAddDefaultValueIfEmpty)
-	: FPCGAttributeAccessorKeysEntries(Metadata ? Metadata->GetDefaultMetadataLayer() : nullptr, bAddDefaultValueIfEmpty)
+	: FPCGAttributeAccessorKeysEntries(Metadata ? Metadata->GetDefaultMetadataDomain() : nullptr, bAddDefaultValueIfEmpty)
 {}
 
-FPCGAttributeAccessorKeysEntries::FPCGAttributeAccessorKeysEntries(const FPCGMetadataLayer* Metadata, bool bAddDefaultValueIfEmpty)
+FPCGAttributeAccessorKeysEntries::FPCGAttributeAccessorKeysEntries(const FPCGMetadataDomain* Metadata, bool bAddDefaultValueIfEmpty)
 	: IPCGAttributeAccessorKeys(/*bInReadOnly=*/ true)
 {
 	InitializeFromMetadata(Metadata, bAddDefaultValueIfEmpty);
 	Entries = TArrayView<PCGMetadataEntryKey>(ExtractedEntries);
 }
 
-FPCGAttributeAccessorKeysEntries::FPCGAttributeAccessorKeysEntries(FPCGMetadataLayer* Metadata, bool bAddDefaultValueIfEmpty)
+FPCGAttributeAccessorKeysEntries::FPCGAttributeAccessorKeysEntries(FPCGMetadataDomain* Metadata, bool bAddDefaultValueIfEmpty)
 	: IPCGAttributeAccessorKeys(/*bInReadOnly=*/ false)
 {
 	InitializeFromMetadata(Metadata, bAddDefaultValueIfEmpty);
@@ -57,7 +57,7 @@ FPCGAttributeAccessorKeysEntries::FPCGAttributeAccessorKeysEntries(const TArrayV
 
 }
 
-void FPCGAttributeAccessorKeysEntries::InitializeFromMetadata(const FPCGMetadataLayer* Metadata, bool bAddDefaultValueIfEmpty)
+void FPCGAttributeAccessorKeysEntries::InitializeFromMetadata(const FPCGMetadataDomain* Metadata, bool bAddDefaultValueIfEmpty)
 {
 	if (!Metadata)
 	{
diff --git a/Engine/Plugins/PCG/Source/PCG/Private/Metadata/PCGAttributePropertySelector.cpp b/Engine/Plugins/PCG/Source/PCG/Private/Metadata/PCGAttributePropertySelector.cpp
--- a/Engine/Plugins/PCG/Source/PCG/Private/Metadata/PCGAttributePropertySelector.cpp
+++ b/Engine/Plugins/PCG/Source/PCG/Private/Metadata/PCGAttributePropertySelector.cpp
@@ -14,9 +14,9 @@ namespace PCGAttributePropertySelectorConstants
 {
 	static const TCHAR* PropertyPrefix = TEXT("$");
 	static const TCHAR* ExtraSeparator = TEXT(".");
-	static const TCHAR* LayerPrefix = TEXT("@");
+	static const TCHAR* DomainPrefix = TEXT("@");
 	static const TCHAR PropertyPrefixChar = PropertyPrefix[0];
-	static const TCHAR LayerPrefixChar = LayerPrefix[0];
+	static const TCHAR DomainPrefixChar = DomainPrefix[0];
 	static const TCHAR ExtraSeparatorChar = ExtraSeparator[0];
 
 	static const FString ExportTextLeftSentinel = TEXT("PCGBegin(");
@@ -143,7 +143,7 @@ bool FPCGAttributePropertySelector::SetAttributeName(FName InAttributeName, bool
 	return bHasChanged;
 }
 
-bool FPCGAttributePropertySelector::SetLayerName(FName InLayerName, bool bResetExtraNames)
+bool FPCGAttributePropertySelector::SetDomainName(FName InDomainName, bool bResetExtraNames)
 {
 	bool bHasChanged = false;
 	if (bResetExtraNames)
@@ -151,9 +151,9 @@ bool FPCGAttributePropertySelector::SetLayerName(FName InLayerName, bool bResetE
 		bHasChanged |= ResetExtraNames();
 	}
 
-	if (GetLayerName() != InLayerName)
+	if (GetDomainName() != InDomainName)
 	{
-		LayerName = InLayerName;
+		DomainName = InDomainName;
 		bHasChanged = true;
 	}
 
@@ -206,21 +206,21 @@ bool FPCGAttributePropertySelector::SetExtraProperty(EPCGExtraProperties InExtra
 	return bHasChanged;
 }
 
-FString FPCGAttributePropertySelector::GetLayerString(bool bAddLeadingQualifier) const
+FString FPCGAttributePropertySelector::GetDomainString(bool bAddLeadingQualifier) const
 {
-	const FName Layer = GetLayerName();
-	if (Layer == NAME_None)
+	const FName Domain = GetDomainName();
+	if (Domain == NAME_None)
 	{
 		return FString{};
 	}
 	
 	if (bAddLeadingQualifier)
 	{
-		return FString(PCGAttributePropertySelectorConstants::LayerPrefix) + Layer.ToString();
+		return FString(PCGAttributePropertySelectorConstants::DomainPrefix) + Domain.ToString();
 	}
 	else
 	{
-		return Layer.ToString();
+		return Domain.ToString();
 	}
 }
 
@@ -258,17 +258,17 @@ FString FPCGAttributePropertySelector::GetAttributePropertyAccessorsString(bool
 
 FString FPCGAttributePropertySelector::ToString() const
 {
-	const FString Layer = GetLayerString(/*bAddLeadingQualifier=*/true);
+	const FString Domain = GetDomainString(/*bAddLeadingQualifier=*/true);
 	const FString Attribute = GetAttributePropertyString(/*bAddPropertyQualifier=*/true);
 	const FString Accessors = GetAttributePropertyAccessorsString(/*bAddLeadingSeparator*/true);
 
-	if (Layer.IsEmpty())
+	if (Domain.IsEmpty())
 	{
 		return Attribute + Accessors;
 	}
 	else
 	{
-		return Layer + PCGAttributePropertySelectorConstants::ExtraSeparator + Attribute + Accessors;
+		return Domain + PCGAttributePropertySelectorConstants::ExtraSeparator + Attribute + Accessors;
 	}
 }
 
@@ -279,7 +279,7 @@ bool FPCGAttributePropertySelector::operator==(const FPCGAttributePropertySelect
 
 bool FPCGAttributePropertySelector::IsSame(const FPCGAttributePropertySelector& Other, bool bIncludeExtraNames) const
 {
-	if (GetSelection() != Other.GetSelection() || LayerName != Other.LayerName || (bIncludeExtraNames && ExtraNames != Other.ExtraNames))
+	if (GetSelection() != Other.GetSelection() || DomainName != Other.DomainName || (bIncludeExtraNames && ExtraNames != Other.ExtraNames))
 	{
 		return false;
 	}
@@ -300,7 +300,7 @@ bool FPCGAttributePropertySelector::IsSame(const FPCGAttributePropertySelector&
 void FPCGAttributePropertySelector::ImportFromOtherSelector(const FPCGAttributePropertySelector& InOther)
 {
 	Selection = InOther.GetSelection();
-	LayerName = InOther.GetLayerName();
+	DomainName = InOther.GetDomainName();
 
 	switch (GetSelection())
 	{
@@ -368,9 +368,9 @@ bool FPCGAttributePropertySelector::Update(const FString& NewValue)
 
 	check(!NewValues.IsEmpty())
 
-	// TODO: If we ever have to support multiple layers, this has to change
-	const bool bHasLayerName = NewValues[0].Len() > 1 && NewValues[0][0] == PCGAttributePropertySelectorConstants::LayerPrefixChar && !PCGAttributePropertySelector::IsReservedAttributeName(*NewValues[0]);
-	const int32 PropertyIndex = bHasLayerName ? 1 : 0;
+	// TODO: If we ever have to support multiple domains, this has to change
+	const bool bHasDomainName = NewValues[0].Len() > 1 && NewValues[0][0] == PCGAttributePropertySelectorConstants::DomainPrefixChar && !PCGAttributePropertySelector::IsReservedAttributeName(*NewValues[0]);
+	const int32 PropertyIndex = bHasDomainName ? 1 : 0;
 	const bool bHasProperty = NewValues.Num() > PropertyIndex && !NewValues[PropertyIndex].IsEmpty() && NewValues[PropertyIndex][0] == PCGAttributePropertySelectorConstants::PropertyPrefixChar;
 
 	const int32 ExtraNameIndex = PropertyIndex + 1;
@@ -384,8 +384,8 @@ bool FPCGAttributePropertySelector::Update(const FString& NewValue)
 	bool bHasChanged = ExtraNamesTemp != ExtraNames;
 	ExtraNames = MoveTemp(ExtraNamesTemp);
 
-	const FName NewLayerName = bHasLayerName ? FName{*NewValues[0].RightChop(1)} : FName{NAME_None};
-	bHasChanged |= SetLayerName(NewLayerName, /*bResetExtraNames=*/ false);
+	const FName NewDomainName = bHasDomainName ? FName{*NewValues[0].RightChop(1)} : FName{NAME_None};
+	bHasChanged |= SetDomainName(NewDomainName, /*bResetExtraNames=*/ false);
 
 	if (NewValues.Num() <= PropertyIndex)
 	{
@@ -418,7 +418,7 @@ void FPCGAttributePropertySelector::AddToCrc(FArchiveCrc32& Ar) const
 {
 	EPCGAttributePropertySelection RealSelection = GetSelection();
 	Ar << RealSelection;
-	Ar << const_cast<FName&>(LayerName);
+	Ar << const_cast<FName&>(DomainName);
 	Ar << const_cast<TArray<FString>&>(ExtraNames);
 
 	switch (RealSelection)
@@ -439,7 +439,7 @@ void FPCGAttributePropertySelector::AddToCrc(FArchiveCrc32& Ar) const
 
 uint32 GetTypeHash(const FPCGAttributePropertySelector& Selector)
 {
-	uint32 Hash = HashCombine(GetTypeHash(Selector.Selection), GetTypeHash(Selector.LayerName));
+	uint32 Hash = HashCombine(GetTypeHash(Selector.Selection), GetTypeHash(Selector.DomainName));
 	
 	switch (Selector.GetSelection())
 	{
@@ -620,9 +620,9 @@ bool UPCGAttributePropertySelectorBlueprintHelpers::SetPropertyName(FPCGAttribut
 	return Selector.SetPropertyName(InPropertyName, bResetExtraNames);
 }
 
-bool UPCGAttributePropertySelectorBlueprintHelpers::SetLayerName(FPCGAttributePropertySelector& Selector, FName InLayerName, bool bResetExtraNames)
+bool UPCGAttributePropertySelectorBlueprintHelpers::SetDomainName(FPCGAttributePropertySelector& Selector, FName InDomainName, bool bResetExtraNames)
 {
-	return Selector.SetLayerName(InLayerName, bResetExtraNames);
+	return Selector.SetDomainName(InDomainName, bResetExtraNames);
 }
 
 bool UPCGAttributePropertySelectorBlueprintHelpers::SetExtraProperty(FPCGAttributePropertySelector& Selector, EPCGExtraProperties InExtraProperty, bool bResetExtraNames)
@@ -660,9 +660,9 @@ FName UPCGAttributePropertySelectorBlueprintHelpers::GetPropertyName(const FPCGA
 	return Selector.GetPropertyName();
 }
 
-FName UPCGAttributePropertySelectorBlueprintHelpers::GetLayerName(const FPCGAttributePropertySelector& Selector)
+FName UPCGAttributePropertySelectorBlueprintHelpers::GetDomainName(const FPCGAttributePropertySelector& Selector)
 {
-	return Selector.GetLayerName();
+	return Selector.GetDomainName();
 }
 
 FName UPCGAttributePropertySelectorBlueprintHelpers::GetName(const FPCGAttributePropertySelector& Selector)
diff --git a/Engine/Plugins/PCG/Source/PCG/Private/Metadata/PCGMetadata.cpp b/Engine/Plugins/PCG/Source/PCG/Private/Metadata/PCGMetadata.cpp
--- a/Engine/Plugins/PCG/Source/PCG/Private/Metadata/PCGMetadata.cpp
+++ b/Engine/Plugins/PCG/Source/PCG/Private/Metadata/PCGMetadata.cpp
@@ -17,17 +17,17 @@
 ////////////////////////////
 
 template <typename Func, typename ...Args>
-decltype(auto) UPCGMetadata::WithMetadataLayer(const FPCGMetadataLayerID& InMetadataLayerID, Func InFunc, Args&& ...InArgs)
+decltype(auto) UPCGMetadata::WithMetadataDomain(const FPCGMetadataDomainID& InMetadataDomainID, Func InFunc, Args&& ...InArgs)
 {
-	if (FPCGMetadataLayer* FoundMetadataLayer = GetMetadataLayer(InMetadataLayerID))
+	if (FPCGMetadataDomain* FoundMetadataDomain = GetMetadataDomain(InMetadataDomainID))
 	{
-		return Invoke(InFunc, *FoundMetadataLayer, std::forward<Args>(InArgs)...);
+		return Invoke(InFunc, *FoundMetadataDomain, std::forward<Args>(InArgs)...);
 	}
 	else
 	{
-		UE_LOG(LogPCG, Error, TEXT("Failed to find MetadataLayer with id %s"), *InMetadataLayerID.DebugName.ToString());
+		UE_LOG(LogPCG, Error, TEXT("Failed to find MetadataDomain with id %s"), *InMetadataDomainID.DebugName.ToString());
 		
-		using ReturnType = std::invoke_result_t<Func, FPCGMetadataLayer&, Args...>;
+		using ReturnType = std::invoke_result_t<Func, FPCGMetadataDomain&, Args...>;
 		if constexpr (!std::is_same_v<ReturnType, void>)
 		{
 			return ReturnType{};
@@ -36,17 +36,17 @@ decltype(auto) UPCGMetadata::WithMetadataLayer(const FPCGMetadataLayerID& InMeta
 }
 
 template <typename Func>
-decltype(auto) UPCGMetadata::WithMetadataLayer_Lambda(const FPCGMetadataLayerID& InMetadataLayerID, Func InFunc)
+decltype(auto) UPCGMetadata::WithMetadataDomain_Lambda(const FPCGMetadataDomainID& InMetadataDomainID, Func InFunc)
 {
-	if (FPCGMetadataLayer* FoundMetadataLayer = GetMetadataLayer(InMetadataLayerID))
+	if (FPCGMetadataDomain* FoundMetadataDomain = GetMetadataDomain(InMetadataDomainID))
 	{
-		return InFunc(FoundMetadataLayer);
+		return InFunc(FoundMetadataDomain);
 	}
 	else
 	{
-		UE_LOG(LogPCG, Error, TEXT("Failed to find layer with id %s"), *InMetadataLayerID.DebugName.ToString());
+		UE_LOG(LogPCG, Error, TEXT("Failed to find domain with id %s"), *InMetadataDomainID.DebugName.ToString());
 		
-		using ReturnType = std::invoke_result_t<Func, FPCGMetadataLayer*>;
+		using ReturnType = std::invoke_result_t<Func, FPCGMetadataDomain*>;
 		if constexpr (!std::is_same_v<ReturnType, void>)
 		{
 			return ReturnType{};
@@ -55,17 +55,17 @@ decltype(auto) UPCGMetadata::WithMetadataLayer_Lambda(const FPCGMetadataLayerID&
 }
 
 template <typename Func, typename ...Args>
-decltype(auto) UPCGMetadata::WithConstMetadataLayer(const FPCGMetadataLayerID& InMetadataLayerID, Func InFunc, Args&& ...InArgs) const
+decltype(auto) UPCGMetadata::WithConstMetadataDomain(const FPCGMetadataDomainID& InMetadataDomainID, Func InFunc, Args&& ...InArgs) const
 {
-	if (const FPCGMetadataLayer* FoundMetadataLayer = GetConstMetadataLayer(InMetadataLayerID))
+	if (const FPCGMetadataDomain* FoundMetadataDomain = GetConstMetadataDomain(InMetadataDomainID))
 	{
-		return Invoke(InFunc, *FoundMetadataLayer, std::forward<Args>(InArgs)...);
+		return Invoke(InFunc, *FoundMetadataDomain, std::forward<Args>(InArgs)...);
 	}
 	else
 	{
-		UE_LOG(LogPCG, Error, TEXT("Failed to find layer with id %s"), *InMetadataLayerID.DebugName.ToString());
+		UE_LOG(LogPCG, Error, TEXT("Failed to find domain with id %s"), *InMetadataDomainID.DebugName.ToString());
 		
-		using ReturnType = typename std::invoke_result_t<Func, FPCGMetadataLayer&, Args...>;
+		using ReturnType = typename std::invoke_result_t<Func, FPCGMetadataDomain&, Args...>;
 		if constexpr (!std::is_same_v<ReturnType, void>)
 		{
 			return ReturnType{};
@@ -74,28 +74,28 @@ decltype(auto) UPCGMetadata::WithConstMetadataLayer(const FPCGMetadataLayerID& I
 }
 
 template <typename Func>
-void UPCGMetadata::ForEachValidUniqueConstMetadataLayer(Func InFunc) const
+void UPCGMetadata::ForEachValidUniqueConstMetadataDomain(Func InFunc) const
 {
-	auto Apply = [this, &InFunc](const FPCGMetadataLayerID& LayerID)
+	auto Apply = [this, &InFunc](const FPCGMetadataDomainID& DomainID)
 	{
-		if (LayerID == DefaultLayer && !LayerID.IsDefault())
+		if (DomainID == DefaultDomain && !DomainID.IsDefault())
 		{
 			return;
 		}
 		
-		const FPCGMetadataLayer* CurrentMetadataLayer = MetadataLayers[LayerID].Get();
-		if (!CurrentMetadataLayer)
+		const FPCGMetadataDomain* CurrentMetadataDomain = MetadataDomains[DomainID].Get();
+		if (!CurrentMetadataDomain)
 		{
 			return;
 		}
 		
-		InFunc(LayerID, *CurrentMetadataLayer);
+		InFunc(DomainID, *CurrentMetadataDomain);
 	};
 
 	// Always check the default first
-	Apply(PCGMetadataLayerID::Default);
+	Apply(PCGMetadataDomainID::Default);
 	
-	for (const auto& It : MetadataLayers)
+	for (const auto& It : MetadataDomains)
 	{
 		if (It.Key.IsDefault())
 		{
@@ -107,28 +107,28 @@ void UPCGMetadata::ForEachValidUniqueConstMetadataLayer(Func InFunc) const
 }
 
 template <typename Func>
-void UPCGMetadata::ForEachValidUniqueMetadataLayer(Func InFunc)
+void UPCGMetadata::ForEachValidUniqueMetadataDomain(Func InFunc)
 {	
-	auto Apply = [this, &InFunc](const FPCGMetadataLayerID& LayerID)
+	auto Apply = [this, &InFunc](const FPCGMetadataDomainID& DomainID)
 	{
-		if (LayerID == DefaultLayer && !LayerID.IsDefault())
+		if (DomainID == DefaultDomain && !DomainID.IsDefault())
 		{
 			return;
 		}
 		
-		FPCGMetadataLayer* CurrentMetadataLayer = MetadataLayers[LayerID].Get();
-		if (!CurrentMetadataLayer)
+		FPCGMetadataDomain* CurrentMetadataDomain = MetadataDomains[DomainID].Get();
+		if (!CurrentMetadataDomain)
 		{
 			return;
 		}
 		
-		InFunc(LayerID, *CurrentMetadataLayer);
+		InFunc(DomainID, *CurrentMetadataDomain);
 	};
 	
 	// Always check the default first
-	Apply(PCGMetadataLayerID::Default);
+	Apply(PCGMetadataDomainID::Default);
 	
-	for (const auto& It : MetadataLayers)
+	for (const auto& It : MetadataDomains)
 	{
 		if (It.Key.IsDefault())
 		{
@@ -139,92 +139,92 @@ void UPCGMetadata::ForEachValidUniqueMetadataLayer(Func InFunc)
 	}
 }
 
-FPCGMetadataLayer* UPCGMetadata::GetMetadataLayer(const FPCGMetadataLayerID& InMetadataLayerID)
+FPCGMetadataDomain* UPCGMetadata::GetMetadataDomain(const FPCGMetadataDomainID& InMetadataDomainID)
 {
-	return FindOrCreateMetadataLayer(InMetadataLayerID);
+	return FindOrCreateMetadataDomain(InMetadataDomainID);
 }
 
-const FPCGMetadataLayer* UPCGMetadata::GetConstMetadataLayer(const FPCGMetadataLayerID& InMetadataLayerID) const
+const FPCGMetadataDomain* UPCGMetadata::GetConstMetadataDomain(const FPCGMetadataDomainID& InMetadataDomainID) const
 {
-	const TSharedPtr<FPCGMetadataLayer>* FoundMetadataLayer = MetadataLayers.Find(InMetadataLayerID);
-	return FoundMetadataLayer ? FoundMetadataLayer->Get() : nullptr;
+	const TSharedPtr<FPCGMetadataDomain>* FoundMetadataDomain = MetadataDomains.Find(InMetadataDomainID);
+	return FoundMetadataDomain ? FoundMetadataDomain->Get() : nullptr;
 }
 
-FPCGMetadataLayer* UPCGMetadata::FindOrCreateMetadataLayer(const FPCGMetadataLayerID& InMetadataLayerID)
+FPCGMetadataDomain* UPCGMetadata::FindOrCreateMetadataDomain(const FPCGMetadataDomainID& InMetadataDomainID)
 {
-	TSharedPtr<FPCGMetadataLayer>* FoundMetadataLayer = MetadataLayers.Find(InMetadataLayerID);
-	if(!ensure(FoundMetadataLayer))
+	TSharedPtr<FPCGMetadataDomain>* FoundMetadataDomain = MetadataDomains.Find(InMetadataDomainID);
+	if(!ensure(FoundMetadataDomain))
 	{
 		return nullptr;
 	}
 
-	if (!FoundMetadataLayer->IsValid())
+	if (!FoundMetadataDomain->IsValid())
 	{
-		UE::TScopeLock<UE::FSpinLock> Lock(MetadataLayersSpinLock);
+		UE::TScopeLock<UE::FSpinLock> Lock(MetadataDomainsSpinLock);
 		
-		FoundMetadataLayer = MetadataLayers.Find(InMetadataLayerID);
-		if (!FoundMetadataLayer->IsValid())
+		FoundMetadataDomain = MetadataDomains.Find(InMetadataDomainID);
+		if (!FoundMetadataDomain->IsValid())
 		{
-			CreateMetadataLayer_Unsafe(InMetadataLayerID);
+			CreateMetadataDomain_Unsafe(InMetadataDomainID);
 		}
 		
 	}
 
-	return FoundMetadataLayer->Get();
+	return FoundMetadataDomain->Get();
 }
 
-FPCGMetadataLayer* UPCGMetadata::CreateMetadataLayer_Unsafe(const FPCGMetadataLayerID& InMetadataLayerID)
+FPCGMetadataDomain* UPCGMetadata::CreateMetadataDomain_Unsafe(const FPCGMetadataDomainID& InMetadataDomainID)
 {
-	check(MetadataLayers.Contains(InMetadataLayerID) && !MetadataLayers[InMetadataLayerID].IsValid());
-	if (!InMetadataLayerID.IsDefault() && InMetadataLayerID == DefaultLayer)
+	check(MetadataDomains.Contains(InMetadataDomainID) && !MetadataDomains[InMetadataDomainID].IsValid());
+	if (!InMetadataDomainID.IsDefault() && InMetadataDomainID == DefaultDomain)
 	{
 		// Nothing to do
 	}
-	else if (!MetadataLayers[InMetadataLayerID].IsValid())
+	else if (!MetadataDomains[InMetadataDomainID].IsValid())
 	{
-		MetadataLayers[InMetadataLayerID] = MakeShared<FPCGMetadataLayer>(this, InMetadataLayerID);
+		MetadataDomains[InMetadataDomainID] = MakeShared<FPCGMetadataDomain>(this, InMetadataDomainID);
 	}
 	
-	return MetadataLayers[InMetadataLayerID].Get();
+	return MetadataDomains[InMetadataDomainID].Get();
 }
 
 UPCGMetadata::UPCGMetadata(const FObjectInitializer& ObjectInitializer)
 	: Super(ObjectInitializer)
 {
-	// Always initialize the default level
-	SetupLayer(PCGMetadataLayerID::Default, /*bIsDefault=*/true);
-	CreateMetadataLayer_Unsafe(PCGMetadataLayerID::Default);
+	// Always initialize the default domain
+	SetupDomain(PCGMetadataDomainID::Default, /*bIsDefault=*/true);
+	CreateMetadataDomain_Unsafe(PCGMetadataDomainID::Default);
 }
 
-void UPCGMetadata::SetupLayer(FPCGMetadataLayerID LayerID, bool bIsDefault)
+void UPCGMetadata::SetupDomain(FPCGMetadataDomainID DomainID, bool bIsDefault)
 {
-	check(!MetadataLayers.Contains(LayerID));
-	if (bIsDefault && !LayerID.IsDefault())
+	check(!MetadataDomains.Contains(DomainID));
+	if (bIsDefault && !DomainID.IsDefault())
 	{
-		check(MetadataLayers.Contains(PCGMetadataLayerID::Default) && MetadataLayers[PCGMetadataLayerID::Default].IsValid());
-		// If the default level is not Default, we need to change it
-		if (!DefaultLayer.IsDefault())
+		check(MetadataDomains.Contains(PCGMetadataDomainID::Default) && MetadataDomains[PCGMetadataDomainID::Default].IsValid());
+		// If the default domain is not Default, we need to change it
+		if (!DefaultDomain.IsDefault())
 		{
-			check(MetadataLayers.Contains(DefaultLayer) && MetadataLayers[DefaultLayer] == MetadataLayers[PCGMetadataLayerID::Default]);
-			MetadataLayers[DefaultLayer].Reset();
+			check(MetadataDomains.Contains(DefaultDomain) && MetadataDomains[DefaultDomain] == MetadataDomains[PCGMetadataDomainID::Default]);
+			MetadataDomains[DefaultDomain].Reset();
 		}
 		
-		DefaultLayer = LayerID;
-		// Map the default layer to the DefaultID layer and update layer to the not-default layerID.
-		MetadataLayers.Emplace(LayerID, MetadataLayers[PCGMetadataLayerID::Default])->LayerID = LayerID;
+		DefaultDomain = DomainID;
+		// Map the default domain to the DefaultID domain and update domain to the not-default domain ID.
+		MetadataDomains.Emplace(DomainID, MetadataDomains[PCGMetadataDomainID::Default])->DomainID = DomainID;
 	}
 	else
 	{
-		MetadataLayers.Emplace(LayerID);
+		MetadataDomains.Emplace(DomainID);
 	}
 }
 
-void UPCGMetadata::SetupLayersFromPCGDataType(const TSubclassOf<UPCGData>& PCGDataType)
+void UPCGMetadata::SetupDomainsFromPCGDataType(const TSubclassOf<UPCGData>& PCGDataType)
 {
 	const UPCGData* CDO = *PCGDataType ? CastChecked<const UPCGData>(PCGDataType->GetDefaultObject()) : nullptr;
 	if (CDO)
 	{
-		SetupLayersFromOtherMetadataIfNeeded(CDO->ConstMetadata());
+		SetupDomainsFromOtherMetadataIfNeeded(CDO->ConstMetadata());
 	}
 }
 
@@ -234,45 +234,45 @@ void UPCGMetadata::Serialize(FArchive& InArchive)
 
 	InArchive.UsingCustomVersion(FPCGCustomVersion::GUID);
 
-	TArray<FPCGMetadataLayerID> LayerIDs;
+	TArray<FPCGMetadataDomainID> DomainIDs;
 	
 	if (InArchive.IsLoading())
 	{
 		if (InArchive.CustomVer(FPCGCustomVersion::GUID) < FPCGCustomVersion::MultiLevelMetadata)
 		{
 			const UPCGData* Outer = Cast<const UPCGData>(GetOuter());
-			FPCGMetadataLayer* DefaultMetadataLayer = GetDefaultMetadataLayer();
-			check(DefaultMetadataLayer);
-			DefaultMetadataLayer->Serialize(InArchive);
+			FPCGMetadataDomain* DefaultMetadataDomain = GetDefaultMetadataDomain();
+			check(DefaultMetadataDomain);
+			DefaultMetadataDomain->Serialize(InArchive);
 		}
 		else
 		{
-			FPCGMetadataLayerID ArchiveDefaultLevel;
-			InArchive << ArchiveDefaultLevel;
+			FPCGMetadataDomainID ArchiveDefaultDomain;
+			InArchive << ArchiveDefaultDomain;
 
-			if (ArchiveDefaultLevel != DefaultLayer)
+			if (ArchiveDefaultDomain != DefaultDomain)
 			{
-				UE_LOG(LogPCG, Warning, TEXT("Mismatch between default metadata levels while loading. "
+				UE_LOG(LogPCG, Warning, TEXT("Mismatch between default metadata domains while loading. "
 								 "You should make sure to update the serialized metadata with the new default. "
-								 "Serialized level: %s ; Current level: %s"),
-								 *ArchiveDefaultLevel.DebugName.ToString(), *DefaultLayer.DebugName.ToString());
+								 "Serialized domain: %s ; Current domain: %s"),
+								 *ArchiveDefaultDomain.DebugName.ToString(), *DefaultDomain.DebugName.ToString());
 			}
 			
-			InArchive << LayerIDs;
+			InArchive << DomainIDs;
 			
-			for (const FPCGMetadataLayerID& LayerID : LayerIDs)
+			for (const FPCGMetadataDomainID& DomainID : DomainIDs)
 			{
 				bool bIsValid = false;
 				InArchive << bIsValid;
 				
-				// If the layer is the default layer, nothing to do.
-				if (!LayerID.IsDefault())
+				// If the domain is the default domain, nothing to do.
+				if (!DomainID.IsDefault())
 				{
-					FPCGMetadataLayer* NewMetadataLayer = FindOrCreateMetadataLayer(LayerID);
-					check(NewMetadataLayer);
+					FPCGMetadataDomain* NewMetadataDomain = FindOrCreateMetadataDomain(DomainID);
+					check(NewMetadataDomain);
 					if (bIsValid)
 					{
-						NewMetadataLayer->Serialize(InArchive);
+						NewMetadataDomain->Serialize(InArchive);
 					}
 				}
 			}
@@ -280,21 +280,21 @@ void UPCGMetadata::Serialize(FArchive& InArchive)
 	}
 	else
 	{
-		MetadataLayers.GetKeys(LayerIDs);
-		// Sort the layers to always serialize in the same order.
-		LayerIDs.Sort();
+		MetadataDomains.GetKeys(DomainIDs);
+		// Sort the domains to always serialize in the same order.
+		DomainIDs.Sort();
 
-		InArchive << DefaultLayer;
-		InArchive << LayerIDs;
-		for (const FPCGMetadataLayerID& LayerID : LayerIDs)
+		InArchive << DefaultDomain;
+		InArchive << DomainIDs;
+		for (const FPCGMetadataDomainID& DomainID : DomainIDs)
 		{
-			bool bIsValid = MetadataLayers[LayerID].IsValid();
+			bool bIsValid = MetadataDomains[DomainID].IsValid();
 			InArchive << bIsValid;
 
-			// Don't serialize the default level
-			if (bIsValid &&  !LayerID.IsDefault())
+			// Don't serialize the default domain
+			if (bIsValid &&  !DomainID.IsDefault())
 			{
-				MetadataLayers[LayerID]->Serialize(InArchive);
+				MetadataDomains[DomainID]->Serialize(InArchive);
 			}
 		}
 	}
@@ -305,18 +305,18 @@ void UPCGMetadata::AddToCrc(FArchiveCrc32& Ar, bool bFullDataCrc) const
 	const UPCGData* Data = Cast<UPCGData>(GetOuter());
 	check(Data);
 	
-	TArray<FPCGMetadataLayerID> LayerIDs;
-	MetadataLayers.GetKeys(LayerIDs);
-	LayerIDs.Sort();
+	TArray<FPCGMetadataDomainID> DomainIDs;
+	MetadataDomains.GetKeys(DomainIDs);
+	DomainIDs.Sort();
 
-	for (const FPCGMetadataLayerID& LayerID : LayerIDs)
+	for (const FPCGMetadataDomainID& DomainID : DomainIDs)
 	{
-		bool bIsValid = MetadataLayers[LayerID].IsValid();
+		bool bIsValid = MetadataDomains[DomainID].IsValid();
 		Ar << bIsValid;
 
 		if (bIsValid)
 		{
-			MetadataLayers[LayerID]->AddToCrc(Ar, Data, bFullDataCrc);
+			MetadataDomains[DomainID]->AddToCrc(Ar, Data, bFullDataCrc);
 		}
 	}
 }
@@ -344,16 +344,16 @@ void UPCGMetadata::InitializeWithAttributeFilter(const UPCGMetadata* InParent, c
 
 	if (InParent != this && InParent)
 	{
-		SetupLayersFromOtherMetadataIfNeeded(InParent);
+		SetupDomainsFromOtherMetadataIfNeeded(InParent);
 		Parent = InParent;
 		
-		TMap<FPCGMetadataLayerID, TSet<FName>> FilteredAttributeSplit = FPCGAttributeIdentifier::SplitByLayer(InFilteredAttributes);
+		TMap<FPCGMetadataDomainID, TSet<FName>> FilteredAttributeSplit = FPCGAttributeIdentifier::SplitByDomain(InFilteredAttributes);
 		
-		InParent->ForEachValidUniqueConstMetadataLayer([this, &FilteredAttributeSplit, InFilterMode, InMatchOperator](const FPCGMetadataLayerID& LayerID, const FPCGMetadataLayer& OtherMetadataLayer)
+		InParent->ForEachValidUniqueConstMetadataDomain([this, &FilteredAttributeSplit, InFilterMode, InMatchOperator](const FPCGMetadataDomainID& DomainID, const FPCGMetadataDomain& OtherMetadataDomain)
 		{
-			TSet<FName>* FilteredAttributesPtr = FilteredAttributeSplit.Find(LayerID);
-			FPCGMetadataLayer* CurrentMetadataLayer = FindOrCreateMetadataLayer(LayerID);
-			CurrentMetadataLayer->InitializeWithAttributeFilter(&OtherMetadataLayer, (FilteredAttributesPtr ? *FilteredAttributesPtr : TSet<FName>{}), InFilterMode, InMatchOperator);
+			TSet<FName>* FilteredAttributesPtr = FilteredAttributeSplit.Find(DomainID);
+			FPCGMetadataDomain* CurrentMetadataDomain = FindOrCreateMetadataDomain(DomainID);
+			CurrentMetadataDomain->InitializeWithAttributeFilter(&OtherMetadataDomain, (FilteredAttributesPtr ? *FilteredAttributesPtr : TSet<FName>{}), InFilterMode, InMatchOperator);
 		});
 	}
 }
@@ -383,14 +383,14 @@ void UPCGMetadata::InitializeAsCopyWithAttributeFilter(const UPCGMetadata* InMet
 		return;
 	}
 
-	SetupLayersFromOtherMetadataIfNeeded(InMetadataToCopy);
+	SetupDomainsFromOtherMetadataIfNeeded(InMetadataToCopy);
 		
-	TMap<FPCGMetadataLayerID, TSet<FName>> FilteredAttributeSplit = FPCGAttributeIdentifier::SplitByLayer(InFilteredAttributes);
-	InMetadataToCopy->ForEachValidUniqueConstMetadataLayer([this, &FilteredAttributeSplit, InFilterMode, &EntriesToCopy](const FPCGMetadataLayerID& LayerID, const FPCGMetadataLayer& OtherMetadataLayer)
+	TMap<FPCGMetadataDomainID, TSet<FName>> FilteredAttributeSplit = FPCGAttributeIdentifier::SplitByDomain(InFilteredAttributes);
+	InMetadataToCopy->ForEachValidUniqueConstMetadataDomain([this, &FilteredAttributeSplit, InFilterMode, &EntriesToCopy](const FPCGMetadataDomainID& DomainID, const FPCGMetadataDomain& OtherMetadataDomain)
 	{
-		TSet<FName>* FilteredAttributesPtr = FilteredAttributeSplit.Find(LayerID);
-		FPCGMetadataLayer* CurrentMetadataLayer = FindOrCreateMetadataLayer(LayerID);
-		CurrentMetadataLayer->InitializeAsCopyWithAttributeFilter(&OtherMetadataLayer, (FilteredAttributesPtr ? *FilteredAttributesPtr : TSet<FName>{}), InFilterMode, EntriesToCopy);
+		TSet<FName>* FilteredAttributesPtr = FilteredAttributeSplit.Find(DomainID);
+		FPCGMetadataDomain* CurrentMetadataDomain = FindOrCreateMetadataDomain(DomainID);
+		CurrentMetadataDomain->InitializeAsCopyWithAttributeFilter(&OtherMetadataDomain, (FilteredAttributesPtr ? *FilteredAttributesPtr : TSet<FName>{}), InFilterMode, EntriesToCopy);
 	});
 }
 
@@ -406,14 +406,14 @@ void UPCGMetadata::AddAttributesFiltered(const UPCGMetadata* InOther, const TSet
 		return;
 	}
 	
-	TMap<FPCGMetadataLayerID, TSet<FName>> FilteredAttributeSplit = FPCGAttributeIdentifier::SplitByLayer(InFilteredAttributes);
+	TMap<FPCGMetadataDomainID, TSet<FName>> FilteredAttributeSplit = FPCGAttributeIdentifier::SplitByDomain(InFilteredAttributes);
 
-	InOther->ForEachValidUniqueConstMetadataLayer([this, &FilteredAttributeSplit, InFilterMode, InMatchOperator, InOther](const FPCGMetadataLayerID& LayerID, const FPCGMetadataLayer& OtherMetadataLayer)
+	InOther->ForEachValidUniqueConstMetadataDomain([this, &FilteredAttributeSplit, InFilterMode, InMatchOperator, InOther](const FPCGMetadataDomainID& DomainID, const FPCGMetadataDomain& OtherMetadataDomain)
 	{
-		TSet<FName>* FilteredAttributesPtr = FilteredAttributeSplit.Find(LayerID);
-		FPCGMetadataLayer* CurrentMetadataLayer = FindOrCreateMetadataLayer(LayerID);
+		TSet<FName>* FilteredAttributesPtr = FilteredAttributeSplit.Find(DomainID);
+		FPCGMetadataDomain* CurrentMetadataDomain = FindOrCreateMetadataDomain(DomainID);
 				
-		if (CurrentMetadataLayer->AddAttributesFiltered(&OtherMetadataLayer, (FilteredAttributesPtr ? *FilteredAttributesPtr : TSet<FName>{}), InFilterMode, InMatchOperator))
+		if (CurrentMetadataDomain->AddAttributesFiltered(&OtherMetadataDomain, (FilteredAttributesPtr ? *FilteredAttributesPtr : TSet<FName>{}), InFilterMode, InMatchOperator))
 		{
 			if (InOther != Parent)
 			{
@@ -456,10 +456,10 @@ void UPCGMetadata::CopyAttributes(const UPCGMetadata* InOther)
 		return;
 	}
 
-	InOther->ForEachValidUniqueConstMetadataLayer([this](const FPCGMetadataLayerID& LayerID, const FPCGMetadataLayer& OtherMetadataLayer)
+	InOther->ForEachValidUniqueConstMetadataDomain([this](const FPCGMetadataDomainID& DomainID, const FPCGMetadataDomain& OtherMetadataDomain)
 	{
-		FPCGMetadataLayer* CurrentMetadataLayer = FindOrCreateMetadataLayer(LayerID);
-		CurrentMetadataLayer->CopyAttributes(&OtherMetadataLayer);
+		FPCGMetadataDomain* CurrentMetadataDomain = FindOrCreateMetadataDomain(DomainID);
+		CurrentMetadataDomain->CopyAttributes(&OtherMetadataDomain);
 	});
 }
 
@@ -539,7 +539,7 @@ void UPCGMetadata::FlattenImpl()
 	TRACE_CPUPROFILER_EVENT_SCOPE(UPCGMetadata::FlattenImpl);
 	Modify();
 
-	ForEachValidUniqueMetadataLayer([](const FPCGMetadataLayerID& LayerID, FPCGMetadataLayer& MetadataLayer) { MetadataLayer.FlattenImpl();});
+	ForEachValidUniqueMetadataDomain([](const FPCGMetadataDomainID& DomainID, FPCGMetadataDomain& MetadataDomain) { MetadataDomain.FlattenImpl();});
 
 	Parent = nullptr;
 }
@@ -548,7 +548,7 @@ bool UPCGMetadata::FlattenAndCompress(const TArray<PCGMetadataEntryKey>& InEntry
 {
 	TRACE_CPUPROFILER_EVENT_SCOPE(UPCGMetadata::FlattenAndCompress);
 	
-	ForEachValidUniqueMetadataLayer([&InEntryKeysToKeep](const FPCGMetadataLayerID& LayerID, FPCGMetadataLayer& MetadataLayer) { MetadataLayer.FlattenAndCompress(InEntryKeysToKeep);});
+	ForEachValidUniqueMetadataDomain([&InEntryKeysToKeep](const FPCGMetadataDomainID& DomainID, FPCGMetadataDomain& MetadataDomain) { MetadataDomain.FlattenAndCompress(InEntryKeysToKeep);});
 
 	Parent = nullptr;
 	return true;
@@ -556,12 +556,12 @@ bool UPCGMetadata::FlattenAndCompress(const TArray<PCGMetadataEntryKey>& InEntry
 
 FPCGMetadataAttributeBase* UPCGMetadata::GetMutableAttribute(FPCGAttributeIdentifier AttributeName)
 {
-	return WithMetadataLayer(AttributeName.MetadataLayer, &FPCGMetadataLayer::GetMutableAttribute, AttributeName.Name);
+	return WithMetadataDomain(AttributeName.MetadataDomain, &FPCGMetadataDomain::GetMutableAttribute, AttributeName.Name);
 }
 
 const FPCGMetadataAttributeBase* UPCGMetadata::GetConstAttribute(FPCGAttributeIdentifier AttributeName) const
 {
-	return WithConstMetadataLayer(AttributeName.MetadataLayer, &FPCGMetadataLayer::GetConstAttribute, AttributeName.Name);
+	return WithConstMetadataDomain(AttributeName.MetadataDomain, &FPCGMetadataDomain::GetConstAttribute, AttributeName.Name);
 }
 
 bool UPCGMetadata::BP_HasAttribute(FName AttributeName) const
@@ -571,7 +571,7 @@ bool UPCGMetadata::BP_HasAttribute(FName AttributeName) const
 
 bool UPCGMetadata::HasAttribute(FPCGAttributeIdentifier AttributeName) const
 {
-	return WithConstMetadataLayer(AttributeName.MetadataLayer, &FPCGMetadataLayer::HasAttribute, AttributeName.Name);
+	return WithConstMetadataDomain(AttributeName.MetadataDomain, &FPCGMetadataDomain::HasAttribute, AttributeName.Name);
 }
 
 bool UPCGMetadata::HasCommonAttributes(const UPCGMetadata* InMetadata) const
@@ -582,15 +582,15 @@ bool UPCGMetadata::HasCommonAttributes(const UPCGMetadata* InMetadata) const
 	}
 	
 	bool bHasCommonAttribute = false;
-	ForEachValidUniqueConstMetadataLayer([InMetadata, &bHasCommonAttribute](const FPCGMetadataLayerID& LayerID, const FPCGMetadataLayer& MetadataLayer)
+	ForEachValidUniqueConstMetadataDomain([InMetadata, &bHasCommonAttribute](const FPCGMetadataDomainID& DomainID, const FPCGMetadataDomain& MetadataDomain)
 	{
 		if (bHasCommonAttribute)
 		{
 			return;
 		}
 		
-		const FPCGMetadataLayer* OtherMetadataLayer = InMetadata->GetConstMetadataLayer(LayerID);
-		bHasCommonAttribute = OtherMetadataLayer && MetadataLayer.HasCommonAttributes(OtherMetadataLayer);
+		const FPCGMetadataDomain* OtherMetadataDomain = InMetadata->GetConstMetadataDomain(DomainID);
+		bHasCommonAttribute = OtherMetadataDomain && MetadataDomain.HasCommonAttributes(OtherMetadataDomain);
 	});
 
 	return bHasCommonAttribute;
@@ -599,7 +599,7 @@ bool UPCGMetadata::HasCommonAttributes(const UPCGMetadata* InMetadata) const
 int32 UPCGMetadata::GetAttributeCount() const
 {
 	int32 Count = 0;
-	ForEachValidUniqueConstMetadataLayer([&Count](const FPCGMetadataLayerID& LayerID, const FPCGMetadataLayer& MetadataLayer) { Count += MetadataLayer.GetAttributeCount(); });
+	ForEachValidUniqueConstMetadataDomain([&Count](const FPCGMetadataDomainID& DomainID, const FPCGMetadataDomain& MetadataDomain) { Count += MetadataDomain.GetAttributeCount(); });
 	
 	return Count;
 }
@@ -609,7 +609,7 @@ void UPCGMetadata::GetAttributes(TArray<FName>& AttributeNames, TArray<EPCGMetad
 	AttributeNames.Reset();
 	AttributeTypes.Reset();
 
-	ForEachValidUniqueConstMetadataLayer([&AttributeNames, &AttributeTypes](const FPCGMetadataLayerID& LayerID, const FPCGMetadataLayer& MetadataLayer) { MetadataLayer.GetAttributes(AttributeNames, AttributeTypes); });
+	ForEachValidUniqueConstMetadataDomain([&AttributeNames, &AttributeTypes](const FPCGMetadataDomainID& DomainID, const FPCGMetadataDomain& MetadataDomain) { MetadataDomain.GetAttributes(AttributeNames, AttributeTypes); });
 }
 
 void UPCGMetadata::GetAllAttributes(TArray<FPCGAttributeIdentifier>& AttributeNames, TArray<EPCGMetadataTypes>& AttributeTypes) const
@@ -619,19 +619,19 @@ void UPCGMetadata::GetAllAttributes(TArray<FPCGAttributeIdentifier>& AttributeNa
 
 	TArray<FName> SubAttributeNames;
 
-	ForEachValidUniqueConstMetadataLayer([&AttributeNames, &AttributeTypes, &SubAttributeNames](const FPCGMetadataLayerID& LayerID, const FPCGMetadataLayer& MetadataLayer)
+	ForEachValidUniqueConstMetadataDomain([&AttributeNames, &AttributeTypes, &SubAttributeNames](const FPCGMetadataDomainID& DomainID, const FPCGMetadataDomain& MetadataDomain)
 	{
 		SubAttributeNames.Empty();
-		MetadataLayer.GetAttributes(SubAttributeNames, AttributeTypes);
-		Algo::Transform(SubAttributeNames, AttributeNames, [&LayerID](const FName Name) { return FPCGAttributeIdentifier(Name, LayerID);});
+		MetadataDomain.GetAttributes(SubAttributeNames, AttributeTypes);
+		Algo::Transform(SubAttributeNames, AttributeNames, [&DomainID](const FName Name) { return FPCGAttributeIdentifier(Name, DomainID);});
 	});
 }
 
 FName UPCGMetadata::GetLatestAttributeNameOrNone() const
 {
-	const FPCGMetadataLayer* DefaultMetadataLayer = GetConstDefaultMetadataLayer();
-	check(DefaultMetadataLayer);
-	return DefaultMetadataLayer->GetLatestAttributeNameOrNone();
+	const FPCGMetadataDomain* DefaultMetadataDomain = GetConstDefaultMetadataDomain();
+	check(DefaultMetadataDomain);
+	return DefaultMetadataDomain->GetLatestAttributeNameOrNone();
 }
 
 bool UPCGMetadata::ParentHasAttribute(FPCGAttributeIdentifier AttributeName) const
@@ -666,26 +666,26 @@ PCG_IMPL_CREATE_TYPED_ATTRIBUTE(CreateSoftClassPathAttribute, const FSoftClassPa
 
 bool UPCGMetadata::CreateAttributeFromProperty(FPCGAttributeIdentifier AttributeName, const UObject* Object, const FProperty* InProperty)
 {
-	FPCGMetadataLayer* FoundMetadataLayer = FindOrCreateMetadataLayer(AttributeName.MetadataLayer);
-	return FoundMetadataLayer ? FoundMetadataLayer->CreateAttributeFromProperty(AttributeName.Name, Object, InProperty) : false;
+	FPCGMetadataDomain* FoundMetadataDomain = FindOrCreateMetadataDomain(AttributeName.MetadataDomain);
+	return FoundMetadataDomain ? FoundMetadataDomain->CreateAttributeFromProperty(AttributeName.Name, Object, InProperty) : false;
 }
 
 bool UPCGMetadata::CreateAttributeFromDataProperty(FPCGAttributeIdentifier AttributeName, const void* Data, const FProperty* InProperty)
 {
-	FPCGMetadataLayer* FoundMetadataLayer = FindOrCreateMetadataLayer(AttributeName.MetadataLayer);
-	return FoundMetadataLayer ? FoundMetadataLayer->CreateAttributeFromDataProperty(AttributeName.Name, Data, InProperty) : false;
+	FPCGMetadataDomain* FoundMetadataDomain = FindOrCreateMetadataDomain(AttributeName.MetadataDomain);
+	return FoundMetadataDomain ? FoundMetadataDomain->CreateAttributeFromDataProperty(AttributeName.Name, Data, InProperty) : false;
 }
 
 bool UPCGMetadata::SetAttributeFromProperty(FPCGAttributeIdentifier AttributeName, PCGMetadataEntryKey& EntryKey, const UObject* Object, const FProperty* InProperty, bool bCreate)
 {
-	FPCGMetadataLayer* FoundMetadataLayer = FindOrCreateMetadataLayer(AttributeName.MetadataLayer);
-	return FoundMetadataLayer ? FoundMetadataLayer->SetAttributeFromProperty(AttributeName.Name, EntryKey, Object, InProperty, bCreate) : false;
+	FPCGMetadataDomain* FoundMetadataDomain = FindOrCreateMetadataDomain(AttributeName.MetadataDomain);
+	return FoundMetadataDomain ? FoundMetadataDomain->SetAttributeFromProperty(AttributeName.Name, EntryKey, Object, InProperty, bCreate) : false;
 }
 
 bool UPCGMetadata::SetAttributeFromDataProperty(FPCGAttributeIdentifier AttributeName, PCGMetadataEntryKey& EntryKey, const void* Data, const FProperty* InProperty, bool bCreate)
 {
-	FPCGMetadataLayer* FoundMetadataLayer = FindOrCreateMetadataLayer(AttributeName.MetadataLayer);
-	return FoundMetadataLayer ? FoundMetadataLayer->SetAttributeFromDataProperty(AttributeName.Name, EntryKey, Data, InProperty, bCreate) : false;
+	FPCGMetadataDomain* FoundMetadataDomain = FindOrCreateMetadataDomain(AttributeName.MetadataDomain);
+	return FoundMetadataDomain ? FoundMetadataDomain->SetAttributeFromDataProperty(AttributeName.Name, EntryKey, Data, InProperty, bCreate) : false;
 }
 
 bool UPCGMetadata::BP_CopyExistingAttribute(FName AttributeToCopy, FName NewAttributeName, bool bKeepParent)
@@ -700,32 +700,32 @@ bool UPCGMetadata::CopyExistingAttribute(FPCGAttributeIdentifier AttributeToCopy
 
 FPCGMetadataAttributeBase* UPCGMetadata::CopyAttribute(FPCGAttributeIdentifier AttributeToCopy, FName NewAttributeName, bool bKeepParent, bool bCopyEntries, bool bCopyValues)
 {
-	return WithMetadataLayer_Lambda(AttributeToCopy.MetadataLayer, [&](FPCGMetadataLayer* MetadataLayer) { return MetadataLayer->CopyAttribute(AttributeToCopy.Name, NewAttributeName, bKeepParent, bCopyEntries, bCopyValues); });
+	return WithMetadataDomain_Lambda(AttributeToCopy.MetadataDomain, [&](FPCGMetadataDomain* MetadataDomain) { return MetadataDomain->CopyAttribute(AttributeToCopy.Name, NewAttributeName, bKeepParent, bCopyEntries, bCopyValues); });
 }
 
 FPCGMetadataAttributeBase* UPCGMetadata::CopyAttribute(const FPCGMetadataAttributeBase* OriginalAttribute, FPCGAttributeIdentifier NewAttributeName, bool bKeepParent, bool bCopyEntries, bool bCopyValues)
 {
 	check(OriginalAttribute);
-	const FPCGMetadataLayer* OriginalMetadataLayer = OriginalAttribute->GetMetadataLayer();
-	check(OriginalMetadataLayer);
+	const FPCGMetadataDomain* OriginalMetadataDomain = OriginalAttribute->GetMetadataDomain();
+	check(OriginalMetadataDomain);
 	const UPCGMetadata* RootMetadata = GetRoot();
-	check((RootMetadata && RootMetadata->MetadataLayers.Contains(OriginalMetadataLayer->LayerID) && RootMetadata->MetadataLayers[OriginalMetadataLayer->LayerID].Get() == OriginalMetadataLayer->GetRoot()) || !bKeepParent);
+	check((RootMetadata && RootMetadata->MetadataDomains.Contains(OriginalMetadataDomain->DomainID) && RootMetadata->MetadataDomains[OriginalMetadataDomain->DomainID].Get() == OriginalMetadataDomain->GetRoot()) || !bKeepParent);
 
-	FPCGMetadataAttributeBase* Attribute = WithMetadataLayer_Lambda(NewAttributeName.MetadataLayer, [&](FPCGMetadataLayer* MetadataLayer) -> FPCGMetadataAttributeBase*
+	FPCGMetadataAttributeBase* Attribute = WithMetadataDomain_Lambda(NewAttributeName.MetadataDomain, [&](FPCGMetadataDomain* MetadataDomain) -> FPCGMetadataAttributeBase*
 	{
 		// TODO: Maybe the operation needs to be handled differently? 
-		if (OriginalMetadataLayer->LayerID != MetadataLayer->LayerID)
+		if (OriginalMetadataDomain->DomainID != MetadataDomain->DomainID)
 		{
-			// UE_LOG(LogPCG, Error, TEXT("[Metadata - CopyAttribute] Can't copy across levels."));
+			// UE_LOG(LogPCG, Error, TEXT("[Metadata - CopyAttribute] Can't copy across domains."));
 			// return nullptr;
 		}
 
-		return MetadataLayer->CopyAttribute(OriginalAttribute, NewAttributeName.Name, bKeepParent, bCopyEntries, bCopyValues);
+		return MetadataDomain->CopyAttribute(OriginalAttribute, NewAttributeName.Name, bKeepParent, bCopyEntries, bCopyValues);
 	});
 	
 	if (!Attribute)
 	{
-		UE_LOG(LogPCG, Error, TEXT("[Metadata - CopyAttribute] Metadata level not existing in current metadata or copy failed."));
+		UE_LOG(LogPCG, Error, TEXT("[Metadata - CopyAttribute] Metadata domain does not exist in current metadata or copy failed."));
 	}
 
 	return Attribute;
@@ -738,9 +738,9 @@ bool UPCGMetadata::BP_RenameAttribute(FName AttributeToRename, FName NewAttribut
 
 bool UPCGMetadata::RenameAttribute(FPCGAttributeIdentifier AttributeToRename, FName NewAttributeName)
 {
-	if (FPCGMetadataLayer* FoundMetadataLayer = GetMetadataLayer(AttributeToRename.MetadataLayer))
+	if (FPCGMetadataDomain* FoundMetadataDomain = GetMetadataDomain(AttributeToRename.MetadataDomain))
 	{
-		return FoundMetadataLayer->RenameAttribute(AttributeToRename.Name, NewAttributeName);
+		return FoundMetadataDomain->RenameAttribute(AttributeToRename.Name, NewAttributeName);
 	}
 	else
 	{
@@ -755,7 +755,7 @@ void UPCGMetadata::BP_ClearAttribute(FName AttributeToClear)
 
 void UPCGMetadata::ClearAttribute(FPCGAttributeIdentifier AttributeToClear)
 {
-	return WithMetadataLayer(AttributeToClear.MetadataLayer, &FPCGMetadataLayer::ClearAttribute, AttributeToClear.Name);
+	return WithMetadataDomain(AttributeToClear.MetadataDomain, &FPCGMetadataDomain::ClearAttribute, AttributeToClear.Name);
 }
 
 void UPCGMetadata::BP_DeleteAttribute(FName AttributeToDelete)
@@ -765,62 +765,62 @@ void UPCGMetadata::BP_DeleteAttribute(FName AttributeToDelete)
 
 void UPCGMetadata::DeleteAttribute(FPCGAttributeIdentifier AttributeToDelete)
 {
-	return WithMetadataLayer(AttributeToDelete.MetadataLayer, &FPCGMetadataLayer::DeleteAttribute, AttributeToDelete.Name);
+	return WithMetadataDomain(AttributeToDelete.MetadataDomain, &FPCGMetadataDomain::DeleteAttribute, AttributeToDelete.Name);
 }
 
 bool UPCGMetadata::ChangeAttributeType(FPCGAttributeIdentifier AttributeName, int16 AttributeNewType)
 {
-	return WithMetadataLayer(AttributeName.MetadataLayer, &FPCGMetadataLayer::ChangeAttributeType, AttributeName.Name, AttributeNewType);
+	return WithMetadataDomain(AttributeName.MetadataDomain, &FPCGMetadataDomain::ChangeAttributeType, AttributeName.Name, AttributeNewType);
 }
 
 int64 UPCGMetadata::GetItemCountForChild() const
 {
-	return WithConstMetadataLayer(PCGMetadataLayerID::Default, &FPCGMetadataLayer::GetItemCountForChild);
+	return WithConstMetadataDomain(PCGMetadataDomainID::Default, &FPCGMetadataDomain::GetItemCountForChild);
 }
 
 int64 UPCGMetadata::GetLocalItemCount() const
 {
-	return WithConstMetadataLayer(PCGMetadataLayerID::Default, &FPCGMetadataLayer::GetLocalItemCount);
+	return WithConstMetadataDomain(PCGMetadataDomainID::Default, &FPCGMetadataDomain::GetLocalItemCount);
 }
 
 int64 UPCGMetadata::AddEntry(int64 ParentEntry)
 {
-	return WithMetadataLayer(PCGMetadataLayerID::Default, &FPCGMetadataLayer::AddEntry, ParentEntry);
+	return WithMetadataDomain(PCGMetadataDomainID::Default, &FPCGMetadataDomain::AddEntry, ParentEntry);
 }
 
 TArray<int64> UPCGMetadata::AddEntries(TArrayView<const int64> ParentEntryKeys)
 {
-	return WithMetadataLayer(PCGMetadataLayerID::Default, &FPCGMetadataLayer::AddEntries, ParentEntryKeys);
+	return WithMetadataDomain(PCGMetadataDomainID::Default, &FPCGMetadataDomain::AddEntries, ParentEntryKeys);
 }
 
 void UPCGMetadata::AddEntriesInPlace(TArrayView<int64*> ParentEntryKeys)
 {
-	return WithMetadataLayer(PCGMetadataLayerID::Default, &FPCGMetadataLayer::AddEntriesInPlace, ParentEntryKeys);
+	return WithMetadataDomain(PCGMetadataDomainID::Default, &FPCGMetadataDomain::AddEntriesInPlace, ParentEntryKeys);
 }
 
 int64 UPCGMetadata::AddEntryPlaceholder()
 {
-	return WithMetadataLayer(PCGMetadataLayerID::Default, &FPCGMetadataLayer::AddEntryPlaceholder);
+	return WithMetadataDomain(PCGMetadataDomainID::Default, &FPCGMetadataDomain::AddEntryPlaceholder);
 }
 
 void UPCGMetadata::AddDelayedEntries(const TArray<TTuple<int64, int64>>& AllEntries)
 {
-	return WithMetadataLayer(PCGMetadataLayerID::Default, &FPCGMetadataLayer::AddDelayedEntries, AllEntries);
+	return WithMetadataDomain(PCGMetadataDomainID::Default, &FPCGMetadataDomain::AddDelayedEntries, AllEntries);
 }
 
 bool UPCGMetadata::InitializeOnSet(PCGMetadataEntryKey& InOutKey, PCGMetadataEntryKey InParentKeyA, const UPCGMetadata* InParentMetadataA, PCGMetadataEntryKey InParentKeyB, const UPCGMetadata* InParentMetadataB)
 {
-	return WithMetadataLayer(PCGMetadataLayerID::Default, &FPCGMetadataLayer::InitializeOnSet, InOutKey, InParentKeyA, InParentMetadataA ? InParentMetadataA->GetConstMetadataLayer(PCGMetadataLayerID::Default) : nullptr, InParentKeyB, InParentMetadataB ? InParentMetadataB->GetConstMetadataLayer(PCGMetadataLayerID::Default) : nullptr);
+	return WithMetadataDomain(PCGMetadataDomainID::Default, &FPCGMetadataDomain::InitializeOnSet, InOutKey, InParentKeyA, InParentMetadataA ? InParentMetadataA->GetConstMetadataDomain(PCGMetadataDomainID::Default) : nullptr, InParentKeyB, InParentMetadataB ? InParentMetadataB->GetConstMetadataDomain(PCGMetadataDomainID::Default) : nullptr);
 }
 
 PCGMetadataEntryKey UPCGMetadata::GetParentKey(PCGMetadataEntryKey LocalItemKey) const
 {
-	return WithConstMetadataLayer(PCGMetadataLayerID::Default, &FPCGMetadataLayer::GetParentKey, LocalItemKey);
+	return WithConstMetadataDomain(PCGMetadataDomainID::Default, &FPCGMetadataDomain::GetParentKey, LocalItemKey);
 }
 
 void UPCGMetadata::GetParentKeys(TArrayView<PCGMetadataEntryKey> LocalItemKeys, const TBitArray<>* Mask) const
 {
-	return WithConstMetadataLayer(PCGMetadataLayerID::Default, &FPCGMetadataLayer::GetParentKeys, LocalItemKeys, Mask);
+	return WithConstMetadataDomain(PCGMetadataDomainID::Default, &FPCGMetadataDomain::GetParentKeys, LocalItemKeys, Mask);
 }
 
 void UPCGMetadata::MergePointAttributes(const FPCGPoint& InPointA, const FPCGPoint& InPointB, FPCGPoint& OutPoint, EPCGMetadataOp Op)
@@ -840,32 +840,32 @@ void UPCGMetadata::MergeAttributes(PCGMetadataEntryKey InKeyA, const UPCGMetadat
 
 void UPCGMetadata::MergeAttributesSubset(PCGMetadataEntryKey InKeyA, const UPCGMetadata* InMetadataA, const UPCGMetadata* InMetadataSubsetA, PCGMetadataEntryKey InKeyB, const UPCGMetadata* InMetadataB, const UPCGMetadata* InMetadataSubsetB, PCGMetadataEntryKey& OutKey, EPCGMetadataOp Op)
 {
-	// TODO Support more levels?
-	const FPCGMetadataLayer* InMetadataLayerA = InMetadataA->GetConstMetadataLayer(PCGMetadataLayerID::Default);
-	const FPCGMetadataLayer* InMetadataLayerSubsetA = InMetadataSubsetA->GetConstMetadataLayer(PCGMetadataLayerID::Default);
-	const FPCGMetadataLayer* InMetadataLayerB = InMetadataB->GetConstMetadataLayer(PCGMetadataLayerID::Default);
-	const FPCGMetadataLayer* InMetadataLayerSubsetB = InMetadataSubsetB->GetConstMetadataLayer(PCGMetadataLayerID::Default);
+	// TODO Support more domains?
+	const FPCGMetadataDomain* InMetadataDomainA = InMetadataA->GetConstMetadataDomain(PCGMetadataDomainID::Default);
+	const FPCGMetadataDomain* InMetadataDomainSubsetA = InMetadataSubsetA->GetConstMetadataDomain(PCGMetadataDomainID::Default);
+	const FPCGMetadataDomain* InMetadataDomainB = InMetadataB->GetConstMetadataDomain(PCGMetadataDomainID::Default);
+	const FPCGMetadataDomain* InMetadataDomainSubsetB = InMetadataSubsetB->GetConstMetadataDomain(PCGMetadataDomainID::Default);
 
-	return WithMetadataLayer(PCGMetadataLayerID::Default, &FPCGMetadataLayer::MergeAttributesSubset, InKeyA, InMetadataLayerA, InMetadataLayerSubsetA, InKeyB, InMetadataLayerB, InMetadataLayerSubsetB, OutKey, Op); 
+	return WithMetadataDomain(PCGMetadataDomainID::Default, &FPCGMetadataDomain::MergeAttributesSubset, InKeyA, InMetadataDomainA, InMetadataDomainSubsetA, InKeyB, InMetadataDomainB, InMetadataDomainSubsetB, OutKey, Op); 
 }
 
 void UPCGMetadata::ResetWeightedAttributes(PCGMetadataEntryKey& OutKey)
 {
-	return WithMetadataLayer(PCGMetadataLayerID::Default, &FPCGMetadataLayer::ResetWeightedAttributes, OutKey);
+	return WithMetadataDomain(PCGMetadataDomainID::Default, &FPCGMetadataDomain::ResetWeightedAttributes, OutKey);
 }
 
 void UPCGMetadata::AccumulateWeightedAttributes(PCGMetadataEntryKey InKey, const UPCGMetadata* InMetadata, float Weight, bool bSetNonInterpolableAttributes, PCGMetadataEntryKey& OutKey)
 {
-	// TODO: Support other levels?
-	return WithMetadataLayer(PCGMetadataLayerID::Default, &FPCGMetadataLayer::AccumulateWeightedAttributes, InKey, InMetadata ? InMetadata->GetConstMetadataLayer(PCGMetadataLayerID::Default) : nullptr, Weight, bSetNonInterpolableAttributes, OutKey);
+	// TODO: Support other domains?
+	return WithMetadataDomain(PCGMetadataDomainID::Default, &FPCGMetadataDomain::AccumulateWeightedAttributes, InKey, InMetadata ? InMetadata->GetConstMetadataDomain(PCGMetadataDomainID::Default) : nullptr, Weight, bSetNonInterpolableAttributes, OutKey);
 }
 
 void UPCGMetadata::ComputePointWeightedAttribute(FPCGPoint& OutPoint, const TArrayView<TPair<const FPCGPoint*, float>>& InWeightedPoints, const UPCGMetadata* InMetadata)
 {
-	// TODO: Support other levels?
-	return WithMetadataLayer_Lambda(PCGMetadataLayerID::Default, [&](FPCGMetadataLayer* MetadataLayer)
+	// TODO: Support other domains?
+	return WithMetadataDomain_Lambda(PCGMetadataDomainID::Default, [&](FPCGMetadataDomain* MetadataDomain)
 	{
-		PCGMetadataHelpers::ComputePointWeightedAttribute(MetadataLayer, OutPoint, InWeightedPoints, InMetadata->GetConstMetadataLayer(PCGMetadataLayerID::Default));
+		PCGMetadataHelpers::ComputePointWeightedAttribute(MetadataDomain, OutPoint, InWeightedPoints, InMetadata->GetConstMetadataDomain(PCGMetadataDomainID::Default));
 	});
 }
 
@@ -876,13 +876,13 @@ void UPCGMetadata::ComputeWeightedAttribute(PCGMetadataEntryKey& OutKey, const T
 		return;
 	}
 
-	// TODO: Support other levels?
-	return WithMetadataLayer(PCGMetadataLayerID::Default, &FPCGMetadataLayer::ComputeWeightedAttribute, OutKey, InWeightedKeys, InMetadata->GetConstMetadataLayer(PCGMetadataLayerID::Default));
+	// TODO: Support other domains?
+	return WithMetadataDomain(PCGMetadataDomainID::Default, &FPCGMetadataDomain::ComputeWeightedAttribute, OutKey, InWeightedKeys, InMetadata->GetConstMetadataDomain(PCGMetadataDomainID::Default));
 }
 
 int64 UPCGMetadata::GetItemKeyCountForParent() const
 {
-	return WithConstMetadataLayer(PCGMetadataLayerID::Default, &FPCGMetadataLayer::GetItemKeyCountForParent);
+	return WithConstMetadataDomain(PCGMetadataDomainID::Default, &FPCGMetadataDomain::GetItemKeyCountForParent);
 }
 
 void UPCGMetadata::SetAttributes(PCGMetadataEntryKey InKey, const UPCGMetadata* InMetadata, PCGMetadataEntryKey& OutKey)
@@ -892,10 +892,10 @@ void UPCGMetadata::SetAttributes(PCGMetadataEntryKey InKey, const UPCGMetadata*
 		return;
 	}
 
-	// TODO: Support other levels?
-	return WithMetadataLayer_Lambda(PCGMetadataLayerID::Default, [&InKey, InMetadataLayer = InMetadata->GetConstMetadataLayer(PCGMetadataLayerID::Default), &OutKey](FPCGMetadataLayer* MetadataLayer)
+	// TODO: Support other domains?
+	return WithMetadataDomain_Lambda(PCGMetadataDomainID::Default, [&InKey, InMetadataDomain = InMetadata->GetConstMetadataDomain(PCGMetadataDomainID::Default), &OutKey](FPCGMetadataDomain* MetadataDomain)
 	{
-		return MetadataLayer->SetAttributes(InKey, InMetadataLayer, OutKey);
+		return MetadataDomain->SetAttributes(InKey, InMetadataDomain, OutKey);
 	});
 }
 
@@ -906,10 +906,10 @@ void UPCGMetadata::SetPointAttributes(const TArrayView<const FPCGPoint>& InPoint
 		return;
 	}
 
-	// TODO: Support other levels?
-	return WithMetadataLayer_Lambda(PCGMetadataLayerID::Default, [&](FPCGMetadataLayer* MetadataLayer)
+	// TODO: Support other domains?
+	return WithMetadataDomain_Lambda(PCGMetadataDomainID::Default, [&](FPCGMetadataDomain* MetadataDomain)
 	{
-		PCGMetadataHelpers::SetPointAttributes(MetadataLayer, InPoints, InMetadata->GetConstMetadataLayer(PCGMetadataLayerID::Default), OutPoints, OptionalContext);
+		PCGMetadataHelpers::SetPointAttributes(MetadataDomain, InPoints, InMetadata->GetConstMetadataDomain(PCGMetadataDomainID::Default), OutPoints, OptionalContext);
 	});
 }
 
@@ -920,10 +920,10 @@ void UPCGMetadata::SetAttributes(const TArrayView<const PCGMetadataEntryKey>& In
 		return;
 	}
 
-	// TODO: Support other levels?
-	return WithMetadataLayer_Lambda(PCGMetadataLayerID::Default, [&](FPCGMetadataLayer* MetadataLayer)
+	// TODO: Support other domains?
+	return WithMetadataDomain_Lambda(PCGMetadataDomainID::Default, [&](FPCGMetadataDomain* MetadataDomain)
 	{
-		return MetadataLayer->SetAttributes(InOriginalKeys, InMetadata->GetConstMetadataLayer(PCGMetadataLayerID::Default), InOutOptionalKeys, OptionalContext);
+		return MetadataDomain->SetAttributes(InOriginalKeys, InMetadata->GetConstMetadataDomain(PCGMetadataDomainID::Default), InOutOptionalKeys, OptionalContext);
 	});
 }
 
@@ -986,22 +986,22 @@ void UPCGMetadata::SetLastCachedSelectorOnOwner(FName AttributeName)
 	}
 }
 
-void UPCGMetadata::SetupLayersFromOtherMetadataIfNeeded(const UPCGMetadata* OtherMetadata)
+void UPCGMetadata::SetupDomainsFromOtherMetadataIfNeeded(const UPCGMetadata* OtherMetadata)
 {
 	// Only do this if our outer is not a PCG data and we are not setuped.
 	// To be used with caution, only with floating metadata.
-	if (!OtherMetadata || OtherMetadata->DefaultLayer.IsDefault() || !DefaultLayer.IsDefault() || (GetOuter() && GetOuter()->IsA<UPCGData>()))
+	if (!OtherMetadata || OtherMetadata->DefaultDomain.IsDefault() || !DefaultDomain.IsDefault() || (GetOuter() && GetOuter()->IsA<UPCGData>()))
 	{
 		return;
 	}
 
-	for (const auto& It : OtherMetadata->MetadataLayers)
+	for (const auto& It : OtherMetadata->MetadataDomains)
 	{
-		if (MetadataLayers.Contains(It.Key))
+		if (MetadataDomains.Contains(It.Key))
 		{
 			continue;
 		}
 		
-		SetupLayer(It.Key, It.Key == OtherMetadata->DefaultLayer);
+		SetupDomain(It.Key, It.Key == OtherMetadata->DefaultDomain);
 	}
 }
diff --git a/Engine/Plugins/PCG/Source/PCG/Private/Metadata/PCGMetadataAttribute.cpp b/Engine/Plugins/PCG/Source/PCG/Private/Metadata/PCGMetadataAttribute.cpp
--- a/Engine/Plugins/PCG/Source/PCG/Private/Metadata/PCGMetadataAttribute.cpp
+++ b/Engine/Plugins/PCG/Source/PCG/Private/Metadata/PCGMetadataAttribute.cpp
@@ -26,7 +26,7 @@ namespace PCGMetadataAttributeBase
 	}
 }
 
-FPCGMetadataAttributeBase::FPCGMetadataAttributeBase(FPCGMetadataLayer* InMetadata, FName InName, const FPCGMetadataAttributeBase* InParent, bool bInAllowsInterpolation)
+FPCGMetadataAttributeBase::FPCGMetadataAttributeBase(FPCGMetadataDomain* InMetadata, FName InName, const FPCGMetadataAttributeBase* InParent, bool bInAllowsInterpolation)
 	: Metadata(InMetadata)
 	, Parent(InParent)
 	, bAllowsInterpolation(bInAllowsInterpolation)
@@ -36,12 +36,12 @@ FPCGMetadataAttributeBase::FPCGMetadataAttributeBase(FPCGMetadataLayer* InMetada
 
 PRAGMA_DISABLE_DEPRECATION_WARNINGS
 FPCGMetadataAttributeBase::FPCGMetadataAttributeBase(UPCGMetadata* InMetadata, FName InName, const FPCGMetadataAttributeBase* InParent, bool bInAllowsInterpolation)
-	: FPCGMetadataAttributeBase(InMetadata ? InMetadata->GetDefaultMetadataLayer() : nullptr, InName, InParent, bInAllowsInterpolation)
+	: FPCGMetadataAttributeBase(InMetadata ? InMetadata->GetDefaultMetadataDomain() : nullptr, InName, InParent, bInAllowsInterpolation)
 {}
 
 void FPCGMetadataAttributeBase::Serialize(UPCGMetadata* InMetadata, FArchive& InArchive)
 {
-	Serialize(InMetadata ? InMetadata->GetDefaultMetadataLayer() : nullptr, InArchive);
+	Serialize(InMetadata ? InMetadata->GetDefaultMetadataDomain() : nullptr, InArchive);
 }
 
 const UPCGMetadata* FPCGMetadataAttributeBase::GetMetadata() const
@@ -51,7 +51,7 @@ const UPCGMetadata* FPCGMetadataAttributeBase::GetMetadata() const
 
 PRAGMA_ENABLE_DEPRECATION_WARNINGS
 
-void FPCGMetadataAttributeBase::Serialize(FPCGMetadataLayer* InMetadata, FArchive& InArchive)
+void FPCGMetadataAttributeBase::Serialize(FPCGMetadataDomain* InMetadata, FArchive& InArchive)
 {
 	InArchive << EntryToValueKeyMap;
 	Metadata = InMetadata;
diff --git a/Engine/Plugins/PCG/Source/PCG/Private/Metadata/PCGMetadataDomain.cpp b/Engine/Plugins/PCG/Source/PCG/Private/Metadata/PCGMetadataDomain.cpp
new file mode 100644
--- /dev/null
+++ b/Engine/Plugins/PCG/Source/PCG/Private/Metadata/PCGMetadataDomain.cpp
@@ -0,0 +1,1349 @@
+// Copyright Epic Games, Inc. All Rights Reserved.
+
+#include "Metadata/PCGMetadataDomain.h"
+
+#include "PCGContext.h"
+#include "PCGElement.h"
+#include "Elements/Metadata/PCGMetadataElementCommon.h"
+#include "Metadata/PCGMetadata.h"
+#include "Helpers/PCGPropertyHelpers.h"
+
+#include "Algo/AnyOf.h"
+#include "Async/ParallelFor.h"
+#include "Serialization/ArchiveCrc32.h"
+
+namespace PCGMetadata
+{
+	template<typename DataType>
+	bool CreateAttributeFromPropertyHelper(FPCGMetadataDomain* Metadata, FName AttributeName, const DataType* DataPtr, const FProperty* InProperty)
+	{
+		if (!InProperty || !DataPtr || !Metadata)
+		{
+			return false;
+		}
+
+		auto CreateAttribute = [AttributeName, Metadata]<typename T>(T PropertyValue) -> bool
+		{
+			FPCGMetadataAttributeBase* BaseAttribute = Metadata->FindOrCreateAttribute<T>(AttributeName, PropertyValue, /*bAllowsInterpolation=*/false, /*bOverrideParent=*/false, /*bOverwriteIfTypeMismatch=*/true);
+			
+			return (BaseAttribute != nullptr);
+		};
+
+		return PCGPropertyHelpers::GetPropertyValueWithCallback(DataPtr, InProperty, CreateAttribute);
+	}
+
+	template<typename DataType>
+	bool SetAttributeFromPropertyHelper(FPCGMetadataDomain* Metadata, FName AttributeName, PCGMetadataEntryKey& EntryKey, const DataType* DataPtr, const FProperty* InProperty, bool bCreate)
+	{
+		if (!InProperty || !DataPtr || !Metadata)
+		{
+			return false;
+		}
+
+		// Check if an attribute already exists or not if we ask to create a new one
+		if (!bCreate && !Metadata->HasAttribute(AttributeName))
+		{
+			return false;
+		}
+
+		auto CreateAttributeAndSet = [AttributeName, Metadata, bCreate, &EntryKey]<typename PropertyType>(PropertyType PropertyValue) -> bool
+		{
+			FPCGMetadataAttributeBase* BaseAttribute = Metadata->GetMutableAttribute(AttributeName);
+
+			if (!BaseAttribute && bCreate)
+			{
+				// Interpolation is disabled and no parent override.
+				BaseAttribute = Metadata->CreateAttribute<PropertyType>(AttributeName, PropertyValue, false, false);
+			}
+
+			if (!BaseAttribute)
+			{
+				return false;
+			}
+
+			// Allow to set the value if both type matches or if we can construct AttributeType from PropertyType.
+			return PCGMetadataAttribute::CallbackWithRightType(BaseAttribute->GetTypeId(), [&EntryKey, &PropertyValue, BaseAttribute, Metadata]<typename AttributeType>(AttributeType AttributeValue) -> bool
+				{
+					FPCGMetadataAttribute<AttributeType>* Attribute = static_cast<FPCGMetadataAttribute<AttributeType>*>(BaseAttribute);
+
+					// Special cased because FSoftObjectPath currently has a deprecated constructor from FName which generates compile warnings.
+					constexpr bool bAssigningNameToSoftObjectPath = std::is_same_v<AttributeType, FSoftObjectPath> && std::is_same_v<PropertyType, FName>;
+
+					if constexpr (std::is_same_v<AttributeType, PropertyType>)
+					{
+						Metadata->InitializeOnSet(EntryKey);
+						Attribute->SetValue(EntryKey, PropertyValue);
+						return true;
+					}
+					else if constexpr (std::is_constructible_v<AttributeType, PropertyType> && !bAssigningNameToSoftObjectPath)
+					{
+						Metadata->InitializeOnSet(EntryKey);
+						Attribute->SetValue(EntryKey, AttributeType(PropertyValue));
+						return true;
+					}
+					else
+					{
+						return false;
+					}
+				});
+		};
+
+		return PCGPropertyHelpers::GetPropertyValueWithCallback(DataPtr, InProperty, CreateAttributeAndSet);
+	}
+}
+
+
+//////////////////////////
+/// FPCGMetadataDomain
+//////////////////////////
+
+
+void FPCGMetadataDomain::Serialize(FArchive& InArchive)
+{
+	int32 NumAttributes = (InArchive.IsLoading() ? 0 : Attributes.Num());
+	// We need to keep track of the max attribute Id, since it won't necessary be equal to the number of attributes + 1.
+	int64 MaxAttributeId = -1;
+
+	InArchive << NumAttributes;
+
+	if (InArchive.IsLoading())
+	{
+		for (int32 AttributeIndex = 0; AttributeIndex < NumAttributes; ++AttributeIndex)
+		{
+			FName AttributeName = NAME_None;
+			InArchive << AttributeName;
+
+			int32 AttributeTypeId = 0;
+			InArchive << AttributeTypeId;
+
+			FPCGMetadataAttributeBase* SerializedAttribute = PCGMetadataAttribute::AllocateEmptyAttributeFromType(static_cast<int16>(AttributeTypeId));
+			if (ensure(SerializedAttribute))
+			{
+				SerializedAttribute->Name = AttributeName;
+				SerializedAttribute->Serialize(this, InArchive);
+				Attributes.Add(AttributeName, SerializedAttribute);
+				MaxAttributeId = FMath::Max(SerializedAttribute->AttributeId, MaxAttributeId);
+			}
+		}
+	}
+	else
+	{
+		for (auto& AttributePair : Attributes)
+		{
+			InArchive << AttributePair.Key;
+			
+			int32 AttributeTypeId = AttributePair.Value->GetTypeId();
+			InArchive << AttributeTypeId;
+
+			AttributePair.Value->Serialize(this, InArchive);
+		}
+	}
+
+	InArchive << ParentKeys;
+
+	// Finally, initialize non-serialized members
+	if (InArchive.IsLoading())
+	{
+		// The next attribute id need to be bigger than the max attribute id of all attributes (or we could have collisions).
+		// Therefore by construction, it should never be less than the number of attributes (but can be greater).
+		NextAttributeId = MaxAttributeId + 1;
+		check(NextAttributeId >= Attributes.Num());
+		ItemKeyOffset = (Parent ? Parent->GetItemCountForChild() : 0);
+	}
+}
+
+void FPCGMetadataDomain::AddToCrc(FArchiveCrc32& Ar, const UPCGData* Data, bool bFullDataCrc) const
+{
+	check(Data);
+
+	// Sort attributes so we have a consistent processing path
+	TArray<FName> AttributeNames;
+
+	{
+		AttributeLock.ReadLock();
+		AttributeNames.Reserve(Attributes.Num());
+
+		for (const TPair<FName, FPCGMetadataAttributeBase*>& Attribute : Attributes)
+		{
+			AttributeNames.Add(Attribute.Key);
+		}
+		AttributeLock.ReadUnlock();
+
+		Algo::Sort(AttributeNames, [this](const FName& A, const FName& B) { return A.LexicalLess(B); });
+	}
+
+	// Add attributes to CRC
+	for (FName AttributeName : AttributeNames)
+	{
+		Ar << AttributeName;
+
+		FPCGAttributePropertyInputSelector InputSource;
+		InputSource.SetAttributeName(AttributeName);
+
+		TUniquePtr<const IPCGAttributeAccessor> InputAccessor = PCGAttributeAccessorHelpers::CreateConstAccessor(Data, InputSource);
+		TUniquePtr<const IPCGAttributeAccessorKeys> InputKeys = PCGAttributeAccessorHelpers::CreateConstKeys(Data, InputSource);
+
+		auto Callback = [&InputAccessor, &InputKeys, &Ar](auto&& Dummy)
+		{
+			using AttributeType = std::decay_t<decltype(Dummy)>;
+			TArray<AttributeType> Values;
+			if constexpr (std::is_trivially_copyable_v<AttributeType>)
+			{
+				Values.SetNumUninitialized(InputKeys->GetNum());
+			}
+			else
+			{
+				Values.SetNum(InputKeys->GetNum());
+			}
+			
+			InputAccessor->GetRange<AttributeType>(Values, 0, *InputKeys);
+
+			for (AttributeType& Value : Values)
+			{
+				// Add value to Crc
+				PCG::Private::Serialize(Ar, Value);
+			}
+		};
+
+		PCGMetadataAttribute::CallbackWithRightType(InputAccessor->GetUnderlyingType(), Callback);
+	}
+}
+
+FPCGMetadataDomain::FPCGMetadataDomain(UPCGMetadata* InTopMetadata, FPCGMetadataDomainID InMetadataDomainID)
+	: TopMetadata(InTopMetadata)
+	, DomainID(InMetadataDomainID)
+{
+	check(InTopMetadata);
+}
+
+FPCGMetadataDomain::~FPCGMetadataDomain()
+{
+	FWriteScopeLock ScopeLock(AttributeLock);
+	for (TPair<FName, FPCGMetadataAttributeBase*>& AttributeEntry : Attributes)
+	{
+		delete AttributeEntry.Value;
+		AttributeEntry.Value = nullptr;
+	}
+	Attributes.Reset();
+}
+
+void FPCGMetadataDomain::Initialize(const FPCGMetadataDomain* InParent, bool bAddAttributesFromParent)
+{
+	// If we are adding attributes from parent, then we use exclude filter with empty list so
+	// that all parameters added. Otherwise use include filter with empty list so none are added.
+	const EPCGMetadataFilterMode bFilter = bAddAttributesFromParent ? EPCGMetadataFilterMode::ExcludeAttributes : EPCGMetadataFilterMode::IncludeAttributes;
+	InitializeWithAttributeFilter(InParent, TSet<FName>(), bFilter);
+}
+
+void FPCGMetadataDomain::InitializeWithAttributeFilter(const FPCGMetadataDomain* InParent, const TSet<FName>& InFilteredAttributes, EPCGMetadataFilterMode InFilterMode, EPCGStringMatchingOperator InMatchOperator)
+{
+	if (Parent || Attributes.Num() != 0)
+	{
+		// Already initialized; note that while that might be construed as a warning, there are legit cases where this is correct
+		return;
+	}
+
+	Parent = ((InParent != this) ? InParent : nullptr);
+	ItemKeyOffset = Parent ? Parent->GetItemCountForChild() : 0;
+
+	// If we have been given an include list which is empty, then don't bother adding any attributes
+	const bool bSkipAddingAttributesFromParent = (InFilterMode == EPCGMetadataFilterMode::IncludeAttributes) && (InFilteredAttributes.Num() == 0);
+	if (!bSkipAddingAttributesFromParent)
+	{
+		AddAttributesFiltered(InParent, InFilteredAttributes, InFilterMode, InMatchOperator);
+	}
+}
+
+void FPCGMetadataDomain::InitializeAsCopy(const FPCGMetadataDomain* InMetadataToCopy, const TArray<PCGMetadataEntryKey>* EntriesToCopy)
+{
+	InitializeAsCopyWithAttributeFilter(InMetadataToCopy, TSet<FName>(), EPCGMetadataFilterMode::ExcludeAttributes, EntriesToCopy);
+}
+
+void FPCGMetadataDomain::InitializeAsCopyWithAttributeFilter(const FPCGMetadataDomain* InMetadataToCopy, const TSet<FName>& InFilteredAttributes, EPCGMetadataFilterMode InFilterMode, const TArray<PCGMetadataEntryKey>* EntriesToCopy)
+{
+	if (!InMetadataToCopy)
+	{
+		return;
+	}
+
+	check(InMetadataToCopy);
+	if (Parent || Attributes.Num() != 0)
+	{
+		UE_LOG(LogPCG, Error, TEXT("Metadata has already been initialized or already contains attributes"));
+		return;
+	}
+
+	const bool bSkipAttributesInFilterList = (InFilterMode == EPCGMetadataFilterMode::ExcludeAttributes);
+
+	// If we have a partial copy, it will flatten the metadata, so we don't need a parent.
+	// Otherwise, we keep the parent hierarchy.
+	const bool bPartialCopy = EntriesToCopy && EntriesToCopy->Num() <= InMetadataToCopy->GetItemCountForChild();
+	TArray<PCGMetadataEntryKey> NewEntryKeys;
+	TArray<PCGMetadataValueKey> NewValueKeys;
+	if (bPartialCopy)
+	{
+		const int32 Count = EntriesToCopy->Num();
+		NewEntryKeys.SetNumUninitialized(Count);
+		NewValueKeys.SetNumUninitialized(Count);
+		ParentKeys.SetNumUninitialized(Count);
+		for (int32 j = 0; j < Count; ++j)
+		{
+			NewEntryKeys[j] = PCGMetadataEntryKey(j);
+			ParentKeys[j] = -1;
+		}
+
+		ItemKeyOffset = 0;
+	}
+	else
+	{
+		ParentKeys = InMetadataToCopy->ParentKeys;
+		ItemKeyOffset = InMetadataToCopy->ItemKeyOffset;
+		Parent = InMetadataToCopy->Parent;
+	}
+
+	// Copy attributes
+	for (const TPair<FName, FPCGMetadataAttributeBase*>& OtherAttribute : InMetadataToCopy->Attributes)
+	{
+		const bool bAttributeInFilterList = InFilteredAttributes.Contains(OtherAttribute.Key);
+		const bool bSkipThisAttribute = (bSkipAttributesInFilterList == bAttributeInFilterList);
+
+		if (!bSkipThisAttribute)
+		{
+			// Don't copy entries if we have a partial copy, we will set them all after.
+			FPCGMetadataAttributeBase* Attribute = CopyAttribute(OtherAttribute.Value, OtherAttribute.Key, /*bKeepParent=*/false, /*bCopyEntries=*/!bPartialCopy, /*bCopyValues=*/true);
+
+			if (bPartialCopy && OtherAttribute.Value && Attribute)
+			{
+				OtherAttribute.Value->GetValueKeys(*EntriesToCopy, NewValueKeys);
+				Attribute->SetValuesFromValueKeys(NewEntryKeys, NewValueKeys);
+			}
+		}
+	}
+}
+
+bool FPCGMetadataDomain::AddAttributesFiltered(const FPCGMetadataDomain* InOther, const TSet<FName>& InFilteredAttributes, EPCGMetadataFilterMode InFilterMode, EPCGStringMatchingOperator InMatchOperator)
+{
+	if (!InOther)
+	{
+		return false;
+	}
+
+	TArray<FString> InFilteredAttributesStrings;
+	if (InMatchOperator != EPCGStringMatchingOperator::Equal)
+	{
+		Algo::Transform(InFilteredAttributes, InFilteredAttributesStrings, [](const FName& InFilteredAttribute) {return InFilteredAttribute.ToString(); });
+	}
+
+	auto IsAttributeInFilterList = [&InFilteredAttributes, &InFilteredAttributesStrings, InMatchOperator](FName OtherAttributeName) -> bool
+	{
+		if (InMatchOperator == EPCGStringMatchingOperator::Equal)
+		{
+			return InFilteredAttributes.Contains(OtherAttributeName);
+		}
+		else
+		{
+			const FString OtherAttributeString = OtherAttributeName.ToString();
+			if (InMatchOperator == EPCGStringMatchingOperator::Substring)
+			{
+				return Algo::AnyOf(InFilteredAttributesStrings, [&OtherAttributeString](const FString& InAttribute) { return OtherAttributeString.Contains(InAttribute); });
+			}
+			else if (InMatchOperator == EPCGStringMatchingOperator::Matches)
+			{
+				return Algo::AnyOf(InFilteredAttributesStrings, [&OtherAttributeString](const FString& InAttribute) { return OtherAttributeString.MatchesWildcard(InAttribute); });
+			}
+			else
+			{
+				checkNoEntry();
+				return false;
+			}
+		}
+	};
+
+	bool bAttributeAdded = false;
+
+	for (const TPair<FName, FPCGMetadataAttributeBase*>& OtherAttribute : InOther->Attributes)
+	{
+		// Skip this attribute if it is in an exclude list, or if it is not in an include list
+		const bool bAttributeInFilterList = IsAttributeInFilterList(OtherAttribute.Key);
+		const bool bSkipAttributesInFilterList = InFilterMode == EPCGMetadataFilterMode::ExcludeAttributes;
+		const bool bSkipThisAttribute = bSkipAttributesInFilterList == bAttributeInFilterList;
+
+		if(bSkipThisAttribute || !OtherAttribute.Value)
+		{
+			continue;
+		}
+		else if (HasAttribute(OtherAttribute.Key))
+		{
+			// If both the current attribute and the other attribute have the same type - nothing to do
+			// If the current attribute can be broadcasted to the other but not the other way around - change the type
+			// If none of this is true - do nothing
+			const FPCGMetadataAttributeBase* Attribute = GetConstAttribute(OtherAttribute.Key);
+			check(Attribute);
+
+			if(Attribute->GetTypeId() != OtherAttribute.Value->GetTypeId() && 
+				!PCG::Private::IsBroadcastable(OtherAttribute.Value->GetTypeId(), Attribute->GetTypeId()) &&
+				PCG::Private::IsBroadcastable(Attribute->GetTypeId(), OtherAttribute.Value->GetTypeId()))
+			{
+				ChangeAttributeType(OtherAttribute.Key, OtherAttribute.Value->GetTypeId());
+			}
+		}
+		else if (CopyAttribute(OtherAttribute.Value, OtherAttribute.Key, /*bKeepParent=*/InOther == Parent, /*bCopyEntries=*/false, /*bCopyValues=*/false))
+		{
+			bAttributeAdded = true;
+		}
+	}
+
+	return bAttributeAdded;
+}
+
+bool FPCGMetadataDomain::AddAttribute(const FPCGMetadataDomain* InOther, FName AttributeName)
+{
+	if (!InOther || !InOther->HasAttribute(AttributeName) || HasAttribute(AttributeName))
+	{
+		return false;
+	}
+
+	return CopyAttribute(InOther->GetConstAttribute(AttributeName), AttributeName, /*bKeepParent=*/InOther == Parent, /*bCopyEntries=*/false, /*bCopyValues=*/false) != nullptr;
+}
+
+void FPCGMetadataDomain::CopyAttributes(const FPCGMetadataDomain* InOther)
+{
+	if (!InOther || InOther == Parent)
+	{
+		return;
+	}
+
+	if (GetItemCountForChild() != InOther->GetItemCountForChild())
+	{
+		UE_LOG(LogPCG, Error, TEXT("Mismatch in copy attributes since the entries do not match"));
+		return;
+	}
+
+	for (const TPair<FName, FPCGMetadataAttributeBase*> OtherAttribute : InOther->Attributes)
+	{
+		if (HasAttribute(OtherAttribute.Key))
+		{
+			continue;
+		}
+		else
+		{
+			CopyAttribute(OtherAttribute.Value, OtherAttribute.Key, /*bKeepParent=*/false, /*bCopyEntries=*/true, /*bCopyValues=*/true);
+		}
+	}
+}
+
+void FPCGMetadataDomain::CopyAttribute(const FPCGMetadataDomain* InOther, FName AttributeToCopy, FName NewAttributeName)
+{
+	if (!InOther)
+	{
+		return;
+	}
+	else if (HasAttribute(NewAttributeName) || !InOther->HasAttribute(AttributeToCopy))
+	{
+		return;
+	}
+	else if (InOther == Parent)
+	{
+		CopyExistingAttribute(AttributeToCopy, NewAttributeName);
+		return;
+	}
+
+	if (GetItemCountForChild() != InOther->GetItemCountForChild())
+	{
+		UE_LOG(LogPCG, Error, TEXT("Mismatch in copy attributes since the entries do not match"));
+		return;
+	}
+
+	CopyAttribute(InOther->GetConstAttribute(AttributeToCopy), NewAttributeName, /*bKeepParent=*/false, /*bCopyEntries=*/true, /*bCopyValues=*/true);
+}
+
+const FPCGMetadataDomain* FPCGMetadataDomain::GetRoot() const
+{
+	if (Parent)
+	{
+		return Parent->GetRoot();
+	}
+	else
+	{
+		return this;
+	}
+}
+
+bool FPCGMetadataDomain::HasParent(const FPCGMetadataDomain* InTentativeParent) const
+{
+	if (!InTentativeParent)
+	{
+		return false;
+	}
+
+	const FPCGMetadataDomain* HierarchicalParent = Parent;
+	while (HierarchicalParent && HierarchicalParent != InTentativeParent)
+	{
+		HierarchicalParent = HierarchicalParent->Parent;
+	}
+
+	return HierarchicalParent == InTentativeParent;
+}
+
+void FPCGMetadataDomain::FlattenImpl()
+{
+	TRACE_CPUPROFILER_EVENT_SCOPE(FPCGSubMetadata::FlattenImpl);
+
+	const int32 NumEntries = GetItemCountForChild();
+
+	AttributeLock.WriteLock();
+	for (auto& AttributePair : Attributes)
+	{
+		FPCGMetadataAttributeBase* Attribute = AttributePair.Value;
+		check(Attribute);
+
+		// For all stored entries (from the root), we need to make sure that entries that should have a concrete value have it
+		// Optimization notes:
+		// - we could skip entries that existed prior to attribute existence, etc.
+		// - we could skip entries that have no parent, but that would require checking against the parent entries in the parent hierarchy
+		for (int64 EntryKey = 0; EntryKey < NumEntries; ++EntryKey)
+		{
+			// Get value using value inheritance as expected
+			PCGMetadataValueKey ValueKey = Attribute->GetValueKey(EntryKey);
+			if (ValueKey != PCGDefaultValueKey)
+			{
+				// Set concrete non-default value
+				Attribute->SetValueFromValueKey(EntryKey, ValueKey);
+			}
+		}
+
+		// Finally, flatten values
+		Attribute->Flatten();
+	}
+	AttributeLock.WriteUnlock();
+
+	Parent = nullptr;
+	ParentKeys.Reset();
+	ParentKeys.Init(PCGInvalidEntryKey, NumEntries);
+	ItemKeyOffset = 0;
+}
+
+bool FPCGMetadataDomain::FlattenAndCompress(const TArray<PCGMetadataEntryKey>& InEntryKeysToKeep)
+{
+	TRACE_CPUPROFILER_EVENT_SCOPE(FPCGSubMetadata::FlattenAndCompress);
+
+	// No keys or no parents, nothing to do
+	if (Attributes.IsEmpty())
+	{
+		return false;
+	}
+
+	AttributeLock.WriteLock();
+	for (auto& AttributePair : Attributes)
+	{
+		FPCGMetadataAttributeBase* Attribute = AttributePair.Value;
+		check(Attribute);
+
+		Attribute->FlattenAndCompress(InEntryKeysToKeep);
+	}
+	AttributeLock.WriteUnlock();
+
+	Parent = nullptr;
+	ParentKeys.Reset();
+	ParentKeys.Init(PCGInvalidEntryKey, InEntryKeysToKeep.Num());
+	ItemKeyOffset = 0;
+
+	return true;
+}
+
+void FPCGMetadataDomain::AddAttributeInternal(FName AttributeName, FPCGMetadataAttributeBase* Attribute)
+{
+	// This call assumes we have a write lock on the attribute map.
+	Attributes.Add(AttributeName, Attribute);
+}
+
+void FPCGMetadataDomain::RemoveAttributeInternal(FName AttributeName)
+{
+	Attributes.Remove(AttributeName);
+}
+
+void FPCGMetadataDomain::SetLastCachedSelectorOnOwner(FName AttributeName)
+{
+	TopMetadata->SetLastCachedSelectorOnOwner(AttributeName);
+}
+
+FPCGMetadataAttributeBase* FPCGMetadataDomain::GetMutableAttribute(FName AttributeName)
+{
+	FPCGMetadataAttributeBase* Attribute = nullptr;
+
+	AttributeLock.ReadLock();
+	if (FPCGMetadataAttributeBase** FoundAttribute = Attributes.Find(AttributeName))
+	{
+		Attribute = *FoundAttribute;
+
+		// Also when accessing an attribute, notify the PCG Data owner that the latest attribute manipulated is this one.
+		SetLastCachedSelectorOnOwner(AttributeName);
+	}
+	AttributeLock.ReadUnlock();
+
+	return Attribute;
+}
+
+const FPCGMetadataAttributeBase* FPCGMetadataDomain::GetConstAttribute(FName AttributeName) const
+{
+	const FPCGMetadataAttributeBase* Attribute = nullptr;
+
+	AttributeLock.ReadLock();
+	if (const FPCGMetadataAttributeBase* const* FoundAttribute = Attributes.Find(AttributeName))
+	{
+		Attribute = *FoundAttribute;
+	}
+	AttributeLock.ReadUnlock();
+
+	return Attribute;
+}
+
+const FPCGMetadataAttributeBase* FPCGMetadataDomain::GetConstAttributeById(int32 InAttributeId) const
+{
+	const FPCGMetadataAttributeBase* Attribute = nullptr;
+
+	AttributeLock.ReadLock();
+	for (const auto& AttributePair : Attributes)
+	{
+		if (AttributePair.Value && AttributePair.Value->AttributeId == InAttributeId)
+		{
+			Attribute = AttributePair.Value;
+			break;
+		}
+	}
+	AttributeLock.ReadUnlock();
+
+	return Attribute;
+}
+
+bool FPCGMetadataDomain::HasAttribute(FName AttributeName) const
+{
+	FReadScopeLock ScopeLock(AttributeLock);
+	return Attributes.Contains(AttributeName);
+}
+
+bool FPCGMetadataDomain::HasCommonAttributes(const FPCGMetadataDomain* InMetadata) const
+{
+	if (!InMetadata)
+	{
+		return false;
+	}
+
+	bool bHasCommonAttribute = false;
+
+	AttributeLock.ReadLock();
+	for (const TPair<FName, FPCGMetadataAttributeBase*>& AttributePair : Attributes)
+	{
+		if (InMetadata->HasAttribute(AttributePair.Key))
+		{
+			bHasCommonAttribute = true;
+			break;
+		}
+	}
+	AttributeLock.ReadUnlock();
+
+	return bHasCommonAttribute;
+}
+
+int32 FPCGMetadataDomain::GetAttributeCount() const
+{
+	FReadScopeLock ScopeLock(AttributeLock);
+	return Attributes.Num();
+}
+
+void FPCGMetadataDomain::GetAttributes(TArray<FName>& AttributeNames, TArray<EPCGMetadataTypes>& AttributeTypes) const
+{
+	AttributeLock.ReadLock();
+	for (const TPair<FName, FPCGMetadataAttributeBase*>& Attribute : Attributes)
+	{
+		check(Attribute.Value && Attribute.Value->Name == Attribute.Key);
+		AttributeNames.Add(Attribute.Key);
+
+		if (Attribute.Value->GetTypeId() < static_cast<uint16>(EPCGMetadataTypes::Unknown))
+		{
+			AttributeTypes.Add(static_cast<EPCGMetadataTypes>(Attribute.Value->GetTypeId()));
+		}
+		else
+		{
+			AttributeTypes.Add(EPCGMetadataTypes::Unknown);
+		}
+	}
+	AttributeLock.ReadUnlock();
+}
+
+FName FPCGMetadataDomain::GetLatestAttributeNameOrNone() const
+{
+	FName LatestAttributeName = NAME_None;
+	int64 MaxAttributeId = -1;
+	
+	AttributeLock.ReadLock();
+	for (const TPair<FName, FPCGMetadataAttributeBase*>& It : Attributes)
+	{
+		if (It.Value && (It.Value->AttributeId > MaxAttributeId))
+		{
+			MaxAttributeId = It.Value->AttributeId;
+			LatestAttributeName = It.Key;
+		}
+	}
+	AttributeLock.ReadUnlock();
+
+	return LatestAttributeName;
+}
+
+bool FPCGMetadataDomain::ParentHasAttribute(FName AttributeName) const
+{
+	return Parent && Parent->HasAttribute(AttributeName);
+}
+
+bool FPCGMetadataDomain::CreateAttributeFromProperty(FName AttributeName, const UObject* Object, const FProperty* InProperty)
+{
+	return PCGMetadata::CreateAttributeFromPropertyHelper<UObject>(this, AttributeName, Object, InProperty);
+}
+
+bool FPCGMetadataDomain::CreateAttributeFromDataProperty(FName AttributeName, const void* Data, const FProperty* InProperty)
+{
+	return PCGMetadata::CreateAttributeFromPropertyHelper<void>(this, AttributeName, Data, InProperty);
+}
+
+bool FPCGMetadataDomain::SetAttributeFromProperty(FName AttributeName, PCGMetadataEntryKey& EntryKey, const UObject* Object, const FProperty* InProperty, bool bCreate)
+{
+	return PCGMetadata::SetAttributeFromPropertyHelper<UObject>(this, AttributeName, EntryKey, Object, InProperty, bCreate);
+}
+
+bool FPCGMetadataDomain::SetAttributeFromDataProperty(FName AttributeName, PCGMetadataEntryKey& EntryKey, const void* Data, const FProperty* InProperty, bool bCreate)
+{
+	return PCGMetadata::SetAttributeFromPropertyHelper<void>(this, AttributeName, EntryKey, Data, InProperty, bCreate);
+}
+
+bool FPCGMetadataDomain::CopyExistingAttribute(FName AttributeToCopy, FName NewAttributeName, bool bKeepParent)
+{
+	return CopyAttribute(AttributeToCopy, NewAttributeName, bKeepParent, /*bCopyEntries=*/true, /*bCopyValues=*/true) != nullptr;
+}
+
+FPCGMetadataAttributeBase* FPCGMetadataDomain::CopyAttribute(FName AttributeToCopy, FName NewAttributeName, bool bKeepParent, bool bCopyEntries, bool bCopyValues)
+{
+	const FPCGMetadataAttributeBase* OriginalAttribute = nullptr;
+
+	AttributeLock.ReadLock();
+	if (FPCGMetadataAttributeBase** AttributeFound = Attributes.Find(AttributeToCopy))
+	{
+		OriginalAttribute = *AttributeFound;
+	}
+	AttributeLock.ReadUnlock();
+
+	if (!OriginalAttribute && Parent)
+	{
+		OriginalAttribute = Parent->GetConstAttribute(AttributeToCopy);
+	}
+
+	if (!OriginalAttribute)
+	{
+		UE_LOG(LogPCG, Warning, TEXT("Attribute %s does not exist, therefore cannot be copied"), *AttributeToCopy.ToString());
+		return nullptr;
+	}
+
+	return CopyAttribute(OriginalAttribute, NewAttributeName, bKeepParent, bCopyEntries, bCopyValues);
+}
+
+FPCGMetadataAttributeBase* FPCGMetadataDomain::CopyAttribute(const FPCGMetadataAttributeBase* OriginalAttribute, FName NewAttributeName, bool bKeepParent, bool bCopyEntries, bool bCopyValues)
+{
+	check(OriginalAttribute);
+	check(OriginalAttribute->GetMetadataDomain()->GetRoot() == GetRoot() || !bKeepParent);
+	FPCGMetadataAttributeBase* NewAttribute = OriginalAttribute->Copy(NewAttributeName, this, bKeepParent, bCopyEntries, bCopyValues);
+
+	if (NewAttribute)
+	{
+		AttributeLock.WriteLock();
+		NewAttribute->AttributeId = NextAttributeId++;
+		AddAttributeInternal(NewAttributeName, NewAttribute);
+
+		// Also when creating an attribute, notify the PCG Data owner that the latest attribute manipulated is this one.
+		SetLastCachedSelectorOnOwner(NewAttributeName);
+
+		AttributeLock.WriteUnlock();
+	}
+
+	return NewAttribute;
+}
+
+bool FPCGMetadataDomain::RenameAttribute(FName AttributeToRename, FName NewAttributeName)
+{
+	if (!FPCGMetadataAttributeBase::IsValidName(NewAttributeName))
+	{
+		UE_LOG(LogPCG, Error, TEXT("New attribute name %s is not valid"), *NewAttributeName.ToString());
+		return false;
+	}
+
+	bool bRenamed = false;
+	AttributeLock.WriteLock();
+	if (FPCGMetadataAttributeBase** AttributeFound = Attributes.Find(AttributeToRename))
+	{
+		FPCGMetadataAttributeBase* Attribute = *AttributeFound;
+		RemoveAttributeInternal(AttributeToRename);
+		Attribute->Name = NewAttributeName;
+		AddAttributeInternal(NewAttributeName, Attribute);
+
+		// Also when renaming an attribute, notify the PCG Data owner that the latest attribute manipulated is this one.
+		SetLastCachedSelectorOnOwner(NewAttributeName);
+		
+		bRenamed = true;
+	}
+	AttributeLock.WriteUnlock();
+
+	if (!bRenamed)
+	{
+		UE_LOG(LogPCG, Warning, TEXT("Attribute %s does not exist and therefore cannot be renamed"), *AttributeToRename.ToString());
+	}
+
+	return bRenamed;
+}
+
+void FPCGMetadataDomain::ClearAttribute(FName AttributeToClear)
+{
+	FPCGMetadataAttributeBase* Attribute = nullptr;
+
+	AttributeLock.ReadLock();
+	if (FPCGMetadataAttributeBase** AttributeFound = Attributes.Find(AttributeToClear))
+	{
+		Attribute = *AttributeFound;
+	}
+	AttributeLock.ReadUnlock();
+
+	// If the attribute exists, then we can lose all the entries
+	// If it doesn't but it exists in the parent hierarchy, then we must create a new attribute.
+	if (Attribute)
+	{
+		Attribute->ClearEntries();
+	}
+}
+
+void FPCGMetadataDomain::DeleteAttribute(FName AttributeToDelete)
+{
+	FPCGMetadataAttributeBase* Attribute = nullptr;
+
+	// If it's a local attribute, then just delete it
+	AttributeLock.WriteLock();
+	if (FPCGMetadataAttributeBase** AttributeFound = Attributes.Find(AttributeToDelete))
+	{
+		Attribute = *AttributeFound;
+		RemoveAttributeInternal(AttributeToDelete);
+	}
+	AttributeLock.WriteUnlock();
+
+	if (Attribute)
+	{
+		delete Attribute;
+	}
+	else
+	{
+		UE_LOG(LogPCG, Verbose, TEXT("Attribute %s does not exist and therefore cannot be deleted"), *AttributeToDelete.ToString());
+	}
+}
+
+bool FPCGMetadataDomain::ChangeAttributeType(FName AttributeName, int16 AttributeNewType)
+{
+	FPCGMetadataAttributeBase* Attribute = GetMutableAttribute(AttributeName);
+
+	if (!Attribute)
+	{
+		UE_LOG(LogPCG, Error, TEXT("Attribute '%s' does not exist and therefore cannot change its type"), *AttributeName.ToString());
+		return false;
+	}
+
+	if (Attribute->GetTypeId() == AttributeNewType)
+	{
+		// Nothing to do, attribute is already the type we want
+		return true;
+	}
+
+	if (FPCGMetadataAttributeBase* NewAttribute = Attribute->CopyToAnotherType(AttributeNewType))
+	{
+		NewAttribute->AttributeId = Attribute->AttributeId;
+
+		AttributeLock.WriteLock();
+		RemoveAttributeInternal(AttributeName);
+		AddAttributeInternal(AttributeName, NewAttribute);
+		AttributeLock.WriteUnlock();
+
+		delete Attribute;
+		Attribute = nullptr;
+	}
+
+	return true;
+}
+
+int64 FPCGMetadataDomain::GetItemCountForChild() const
+{
+	FReadScopeLock ScopeLock(ItemLock);
+	return ParentKeys.Num() + ItemKeyOffset;
+}
+
+int64 FPCGMetadataDomain::GetLocalItemCount() const
+{
+	FReadScopeLock ScopeLock(ItemLock);
+	return ParentKeys.Num();
+}
+
+int64 FPCGMetadataDomain::AddEntry(int64 ParentEntry)
+{
+	FWriteScopeLock ScopeLock(ItemLock);
+	return ParentKeys.Add(ParentEntry) + ItemKeyOffset;
+}
+
+TArray<int64> FPCGMetadataDomain::AddEntries(TArrayView<const int64> ParentEntryKeys)
+{
+	TArray<int64> Result;
+	Result.Reserve(ParentEntryKeys.Num());
+
+	FWriteScopeLock ScopeLock(ItemLock);
+	ParentKeys.Reserve(ParentKeys.Num() + ParentEntryKeys.Num());
+	for (const int64 ParentEntry : ParentEntryKeys)
+	{
+		Result.Add(ParentKeys.Add(ParentEntry) + ItemKeyOffset);
+	}
+
+	return Result;
+}
+
+void FPCGMetadataDomain::AddEntriesInPlace(TArrayView<int64*> ParentEntryKeys)
+{
+	FWriteScopeLock ScopeLock(ItemLock);
+	ParentKeys.Reserve(ParentKeys.Num() + ParentEntryKeys.Num());
+	for (int64* ParentEntry : ParentEntryKeys)
+	{
+		*ParentEntry = ParentKeys.Add(*ParentEntry) + ItemKeyOffset;
+	}
+}
+
+int64 FPCGMetadataDomain::AddEntryPlaceholder()
+{
+	FReadScopeLock ScopeLock(ItemLock);
+	return ParentKeys.Num() + DelayedEntriesIndex.IncrementExchange() + ItemKeyOffset;
+}
+
+void FPCGMetadataDomain::AddDelayedEntries(const TArray<TTuple<int64, int64>>& AllEntries)
+{
+	FWriteScopeLock ScopeLock(ItemLock);
+	ParentKeys.AddUninitialized(AllEntries.Num());
+	for (const TTuple<int64, int64>& Entry : AllEntries)
+	{
+		int64 Index = Entry.Get<0>() - ItemKeyOffset;
+		check(Index < ParentKeys.Num());
+		ParentKeys[Index] = Entry.Get<1>();
+	}
+
+	DelayedEntriesIndex.Exchange(0);
+}
+
+bool FPCGMetadataDomain::InitializeOnSet(PCGMetadataEntryKey& InOutKey, PCGMetadataEntryKey InParentKeyA, const FPCGMetadataDomain* InParentMetadataA, PCGMetadataEntryKey InParentKeyB, const FPCGMetadataDomain* InParentMetadataB)
+{
+	if (InOutKey == PCGInvalidEntryKey)
+	{
+		if (InParentKeyA != PCGInvalidEntryKey && Parent == InParentMetadataA)
+		{
+			InOutKey = AddEntry(InParentKeyA);
+			return true;
+		}
+		else if (InParentKeyB != PCGInvalidEntryKey && Parent == InParentMetadataB)
+		{
+			InOutKey = AddEntry(InParentKeyB);
+			return true;
+		}
+		else
+		{
+			InOutKey = AddEntry();
+			return false;
+		}
+	}
+	else if(InOutKey < ItemKeyOffset)
+	{
+		InOutKey = AddEntry(InOutKey);
+		return false;
+	}
+	else
+	{
+		return false;
+	}
+}
+
+PCGMetadataEntryKey FPCGMetadataDomain::GetParentKey(PCGMetadataEntryKey LocalItemKey) const
+{
+	if (LocalItemKey < ItemKeyOffset)
+	{
+		// Key is already in parent referential
+		return LocalItemKey;
+	}
+	else
+	{
+		FReadScopeLock ScopeLock(ItemLock);
+		if (LocalItemKey - ItemKeyOffset < ParentKeys.Num())
+		{
+			return ParentKeys[LocalItemKey - ItemKeyOffset];
+		}
+		else
+		{
+			UE_LOG(LogPCG, Warning, TEXT("Invalid metadata key - check for entry key not properly initialized"));
+			return PCGInvalidEntryKey;
+		}
+	}
+}
+
+void FPCGMetadataDomain::GetParentKeys(TArrayView<PCGMetadataEntryKey> LocalItemKeys, const TBitArray<>* Mask) const
+{
+	auto GetParentKey_Unsafe = [this](PCGMetadataEntryKey& LocalItemKey) -> void
+	{
+		if (LocalItemKey < ItemKeyOffset)
+		{
+			// Key is already in parent referential
+			return;
+		}
+		else if (LocalItemKey - ItemKeyOffset < ParentKeys.Num())
+		{
+			LocalItemKey = ParentKeys[LocalItemKey - ItemKeyOffset];
+		}
+		else
+		{
+			UE_LOG(LogPCG, Warning, TEXT("Invalid metadata key - check for entry key not properly initialized"));
+			LocalItemKey = PCGInvalidEntryKey;
+		}
+	};
+
+	FReadScopeLock ScopeLock(ItemLock);
+	if (Mask && ensure(LocalItemKeys.Num() == Mask->Num()))
+	{
+		for (TConstSetBitIterator<> It(*Mask); It; ++It)
+		{
+			GetParentKey_Unsafe(LocalItemKeys[It.GetIndex()]);
+		}
+	}
+	else
+	{
+		for (PCGMetadataEntryKey& LocalItemKey : LocalItemKeys)
+		{
+			GetParentKey_Unsafe(LocalItemKey);
+		}
+	}
+}
+
+void FPCGMetadataDomain::MergeAttributes(PCGMetadataEntryKey InKeyA, const FPCGMetadataDomain* InMetadataA, PCGMetadataEntryKey InKeyB, const FPCGMetadataDomain* InMetadataB, PCGMetadataEntryKey& OutKey, EPCGMetadataOp Op)
+{
+	MergeAttributesSubset(InKeyA, InMetadataA, InMetadataA, InKeyB, InMetadataB, InMetadataB, OutKey, Op);
+}
+
+void FPCGMetadataDomain::MergeAttributesSubset(PCGMetadataEntryKey InKeyA, const FPCGMetadataDomain* InMetadataA, const FPCGMetadataDomain* InMetadataSubsetA, PCGMetadataEntryKey InKeyB, const FPCGMetadataDomain* InMetadataB, const FPCGMetadataDomain* InMetadataSubsetB, PCGMetadataEntryKey& OutKey, EPCGMetadataOp Op)
+{
+	//TRACE_CPUPROFILER_EVENT_SCOPE(FPCGMetadataDomain::MergeAttributesSubset);
+	// Early out: nothing to do if both input metadata are null / points have no assigned metadata
+	if (!InMetadataA && !InMetadataB)
+	{
+		return;
+	}
+
+	// For each attribute in the current metadata, query the values from point A & B, apply operation on the result and finally store in the out point.
+	InitializeOnSet(OutKey, InKeyA, InMetadataA, InKeyB, InMetadataB);
+
+	AttributeLock.ReadLock();
+	for(const TPair<FName, FPCGMetadataAttributeBase*>& AttributePair : Attributes)
+	{
+		const FName& AttributeName = AttributePair.Key;
+		FPCGMetadataAttributeBase* Attribute = AttributePair.Value;
+
+		// Get attribute from A
+		const FPCGMetadataAttributeBase* AttributeA = nullptr;
+		if (InMetadataA && ((InMetadataA == InMetadataSubsetA) || (InMetadataSubsetA && InMetadataSubsetA->HasAttribute(AttributeName))))
+		{
+			AttributeA = InMetadataA->GetConstAttribute(AttributeName);
+		}
+
+		if (AttributeA && AttributeA->GetTypeId() != Attribute->GetTypeId())
+		{
+			UE_LOG(LogPCG, Error, TEXT("Metadata type mismatch with attribute %s"), *AttributeName.ToString());
+			AttributeA = nullptr;
+		}
+
+		// Get attribute from B
+		const FPCGMetadataAttributeBase* AttributeB = nullptr;
+		if (InMetadataB && ((InMetadataB == InMetadataSubsetB) || (InMetadataSubsetB && InMetadataSubsetB->HasAttribute(AttributeName))))
+		{
+			AttributeB = InMetadataB->GetConstAttribute(AttributeName);
+		}
+
+		if (AttributeB && AttributeB->GetTypeId() != Attribute->GetTypeId())
+		{
+			UE_LOG(LogPCG, Error, TEXT("Metadata type mismatch with attribute %s"), *AttributeName.ToString());
+			AttributeB = nullptr;
+		}
+
+		if (AttributeA || AttributeB)
+		{
+			Attribute->SetValue(OutKey, AttributeA, InKeyA, AttributeB, InKeyB, Op);
+		}
+	}
+	AttributeLock.ReadUnlock();
+}
+
+void FPCGMetadataDomain::ResetWeightedAttributes(PCGMetadataEntryKey& OutKey)
+{
+	InitializeOnSet(OutKey);
+
+	AttributeLock.ReadLock();
+	for(const TPair<FName, FPCGMetadataAttributeBase*>& AttributePair : Attributes)
+	{
+		FPCGMetadataAttributeBase* Attribute = AttributePair.Value;
+		
+		if (Attribute && Attribute->AllowsInterpolation())
+		{
+			Attribute->SetZeroValue(OutKey);
+		}
+	}
+	AttributeLock.ReadUnlock();
+}
+
+void FPCGMetadataDomain::AccumulateWeightedAttributes(PCGMetadataEntryKey InKey, const FPCGMetadataDomain* InMetadata, float Weight, bool bSetNonInterpolableAttributes, PCGMetadataEntryKey& OutKey)
+{
+	if (!InMetadata)
+	{
+		return;
+	}
+
+	bool bHasSetParent = InitializeOnSet(OutKey, InKey, InMetadata);
+
+	const bool bShouldSetNonInterpolableAttributes = bSetNonInterpolableAttributes && !bHasSetParent;
+
+	AttributeLock.ReadLock();
+	for(const TPair<FName, FPCGMetadataAttributeBase*>& AttributePair : Attributes)
+	{
+		const FName& AttributeName = AttributePair.Key;
+		FPCGMetadataAttributeBase* Attribute = AttributePair.Value;
+
+		if (const FPCGMetadataAttributeBase* OtherAttribute = InMetadata->GetConstAttribute(AttributeName))
+		{
+			if (OtherAttribute->GetTypeId() != Attribute->GetTypeId())
+			{
+				UE_LOG(LogPCG, Error, TEXT("Metadata type mismatch with attribute %s"), *AttributeName.ToString());
+				continue;
+			}
+
+			if (Attribute->AllowsInterpolation())
+			{
+				Attribute->AccumulateValue(OutKey, OtherAttribute, InKey, Weight);
+			}
+			else if (bShouldSetNonInterpolableAttributes)
+			{
+				Attribute->SetValue(OutKey, OtherAttribute, InKey);
+			}
+		}
+	}
+	AttributeLock.ReadUnlock();
+}
+
+void FPCGMetadataDomain::ComputeWeightedAttribute(PCGMetadataEntryKey& OutKey, const TArrayView<TPair<PCGMetadataEntryKey, float>>& InWeightedKeys, const FPCGMetadataDomain* InMetadata)
+{
+	if (!InMetadata || InWeightedKeys.IsEmpty())
+	{
+		return;
+	}
+
+	// Could ensure that InitializeOnSet returns false...
+	AttributeLock.ReadLock();
+	for (const TPair<FName, FPCGMetadataAttributeBase*>& AttributePair : Attributes)
+	{
+		const FName& AttributeName = AttributePair.Key;
+		FPCGMetadataAttributeBase* Attribute = AttributePair.Value;
+
+		if (!Attribute->AllowsInterpolation())
+		{
+			continue;
+		}
+
+		if (const FPCGMetadataAttributeBase* OtherAttribute = InMetadata->GetConstAttribute(AttributeName))
+		{
+			if (OtherAttribute->GetTypeId() != Attribute->GetTypeId())
+			{
+				UE_LOG(LogPCG, Error, TEXT("Metadata type mismatch with attribute %s"), *AttributeName.ToString());
+				continue;
+			}
+
+			Attribute->SetWeightedValue(OutKey, OtherAttribute, InWeightedKeys);
+		}
+	}
+	AttributeLock.ReadUnlock();
+}
+
+int64 FPCGMetadataDomain::GetItemKeyCountForParent() const
+{
+	return ItemKeyOffset;
+}
+
+void FPCGMetadataDomain::SetAttributes(PCGMetadataEntryKey InKey, const FPCGMetadataDomain* InMetadata, PCGMetadataEntryKey& OutKey)
+{
+	if (!InMetadata)
+	{
+		return;
+	}
+
+	if (InitializeOnSet(OutKey, InKey, InMetadata))
+	{
+		// Early out; we don't need to do anything else at this point
+		return;
+	}
+
+	AttributeLock.ReadLock();
+	for(const TPair<FName, FPCGMetadataAttributeBase*>& AttributePair : Attributes)
+	{
+		const FName& AttributeName = AttributePair.Key;
+		FPCGMetadataAttributeBase* Attribute = AttributePair.Value;
+
+		if (const FPCGMetadataAttributeBase* OtherAttribute = InMetadata->GetConstAttribute(AttributeName))
+		{
+			if (OtherAttribute->GetTypeId() != Attribute->GetTypeId())
+			{
+				UE_LOG(LogPCG, Error, TEXT("Metadata type mismatch with attribute %s"), *AttributeName.ToString());
+				continue;
+			}
+
+			Attribute->SetValue(OutKey, OtherAttribute, InKey);
+		}
+	}
+	AttributeLock.ReadUnlock();
+}
+
+void FPCGMetadataDomain::SetAttributes(const TArrayView<const PCGMetadataEntryKey>& InOriginalKeys, const FPCGMetadataDomain* InMetadata, const TArrayView<PCGMetadataEntryKey>* InOutOptionalKeys, FPCGContext* OptionalContext)
+{
+	if (!InMetadata || InMetadata->GetAttributeCount() == 0 || GetAttributeCount() == 0 || InOriginalKeys.IsEmpty())
+	{
+		return;
+	}
+
+	TRACE_CPUPROFILER_EVENT_SCOPE(FPCGSubMetadata::SetAttributes);
+
+	check(!InOutOptionalKeys || (InOriginalKeys.Num() == InOutOptionalKeys->Num()));
+
+	// There are a few things we can do to optimize here -
+	// basically, we don't need to set attributes more than once for a given <in, out> pair
+	TArray<PCGMetadataEntryKey, TInlineAllocator<256>> InKeys;
+	TArray<PCGMetadataEntryKey, TInlineAllocator<256>> OutKeys;
+
+	if (InOutOptionalKeys)
+	{
+		TRACE_CPUPROFILER_EVENT_SCOPE(FPCGSubMetadata::SetAttributes::CreateDeduplicatedKeys);
+		TMap<TPair<PCGMetadataEntryKey, PCGMetadataEntryKey>, int> PairMapping;
+
+		for (int KeyIndex = 0; KeyIndex < InOriginalKeys.Num(); ++KeyIndex)
+		{
+			PCGMetadataEntryKey InKey = InOriginalKeys[KeyIndex];
+			PCGMetadataEntryKey& OutKey = (*InOutOptionalKeys)[KeyIndex];
+
+			if (int* MatchingPairIndex = PairMapping.Find(TPair<PCGMetadataEntryKey, PCGMetadataEntryKey>(InKey, OutKey)))
+			{
+				OutKey = *MatchingPairIndex;
+			}
+			else
+			{
+				int NewIndex = InKeys.Add(InKey);
+
+				PairMapping.Emplace(TPair<PCGMetadataEntryKey, PCGMetadataEntryKey>(InKey, OutKey), NewIndex);
+				OutKeys.Add(OutKey);
+				OutKey = NewIndex;
+			}
+		}
+	}
+	else
+	{
+		InKeys.Append(InOriginalKeys);
+		OutKeys.Init(PCGInvalidEntryKey, InOriginalKeys.Num());
+	}
+
+	{
+		TRACE_CPUPROFILER_EVENT_SCOPE(FPCGSubMetadata::SetAttributes::InitializeOnSet);
+
+		for (int32 KeyIndex = 0; KeyIndex < InKeys.Num(); ++KeyIndex)
+		{
+			InitializeOnSet(OutKeys[KeyIndex], InKeys[KeyIndex], InMetadata);
+		}
+	}
+
+	AttributeLock.ReadLock();
+	int32 AttributeOffset = 0;
+	const int32 DefaultAttributesPerDispatch = 64;
+	int32 AttributesPerDispatch = OptionalContext ? DefaultAttributesPerDispatch : 1;
+	if (OptionalContext && OptionalContext->AsyncState.NumAvailableTasks > 0)
+	{
+		AttributesPerDispatch = FMath::Min(OptionalContext->AsyncState.NumAvailableTasks, AttributesPerDispatch);
+	}
+
+	while (AttributeOffset < Attributes.Num())
+	{
+		TArray<FName> AttributeNames;
+		TArray<EPCGMetadataTypes> AttributeTypes;
+		GetAttributes(AttributeNames, AttributeTypes);
+
+		const int32 AttributeCountInCurrentDispatch = FMath::Min(AttributesPerDispatch, Attributes.Num() - AttributeOffset);
+		ParallelFor(AttributeCountInCurrentDispatch, [this, OptionalContext, AttributeOffset, InMetadata, &AttributeNames, &InKeys, &OutKeys](int32 WorkerIndex)
+		{
+			const FName AttributeName = AttributeNames[AttributeOffset + WorkerIndex];
+			FPCGMetadataAttributeBase* Attribute = Attributes[AttributeName];
+
+			if (const FPCGMetadataAttributeBase* OtherAttribute = InMetadata->GetConstAttribute(AttributeName))
+			{
+				if (!PCG::Private::IsBroadcastableOrConstructible(OtherAttribute->GetTypeId(), Attribute->GetTypeId()))
+				{
+					PCGE_LOG_C(Error, GraphAndLog, OptionalContext, FText::Format(NSLOCTEXT("PCGMetadata", "TypeMismatch", "Metadata type mismatch with attribute '{0}'"), FText::FromName(AttributeName)));
+					return;
+				}
+
+				if (Attribute == OtherAttribute)
+				{
+					TArray<PCGMetadataValueKey> ValueKeys;
+					Attribute->GetValueKeys(TArrayView<const PCGMetadataEntryKey>(InKeys), ValueKeys);
+					Attribute->SetValuesFromValueKeys(OutKeys, ValueKeys);
+				}
+				else
+				{
+					// Create accessor for the other attribute
+					TUniquePtr<const IPCGAttributeAccessor> OtherAttributeAccessor = PCGAttributeAccessorHelpers::CreateConstAccessor(OtherAttribute, InMetadata);
+
+					TArrayView<PCGMetadataEntryKey> InKeysView(InKeys);
+					FPCGAttributeAccessorKeysEntries OtherAttributeKeys(InKeysView);
+
+					if (!OtherAttributeAccessor)
+					{
+						return;
+					}
+
+					auto GetAndSetValues = [Attribute, &OutKeys, &OtherAttributeAccessor, &OtherAttributeKeys](auto Dummy) -> bool
+					{
+						using Type = decltype(Dummy);
+
+						auto SetValues = [Attribute, &OutKeys](const TArrayView<Type>& View, const int32 Start, const int32 Range)
+						{
+							TArrayView<PCGMetadataEntryKey> Keys(OutKeys.GetData() + Start, Range);
+							static_cast<FPCGMetadataAttribute<Type>*>(Attribute)->SetValues(Keys, View);
+						};
+
+						return PCGMetadataElementCommon::ApplyOnAccessorRange<Type>(OtherAttributeKeys, *OtherAttributeAccessor, SetValues, EPCGAttributeAccessorFlags::AllowBroadcast | EPCGAttributeAccessorFlags::AllowConstructible);
+					};
+
+					PCGMetadataAttribute::CallbackWithRightType(Attribute->GetTypeId(), GetAndSetValues);
+				}
+			}
+		});
+
+		AttributeOffset += AttributeCountInCurrentDispatch;
+	}
+	AttributeLock.ReadUnlock();
+
+	if (InOutOptionalKeys)
+	{
+		// Finally, copy back the actual out keys to the original out keys
+		for (PCGMetadataEntryKey& OutKey : *InOutOptionalKeys)
+		{
+			OutKey = OutKeys[OutKey];
+		}
+	}
+}
+
+void FPCGMetadataDomain::SetAttributes(const TArrayView<const PCGMetadataEntryKey>& InKeys, const FPCGMetadataDomain* InMetadata, const TArrayView<PCGMetadataEntryKey>& OutKeys, FPCGContext* OptionalContext)
+{
+	SetAttributes(InKeys, InMetadata, &OutKeys, OptionalContext);
+}
diff --git a/Engine/Plugins/PCG/Source/PCG/Private/Metadata/PCGMetadataLayer.cpp b/Engine/Plugins/PCG/Source/PCG/Private/Metadata/PCGMetadataLayer.cpp
deleted file mode 100644
diff --git a/Engine/Plugins/PCG/Source/PCG/Private/PCGParamData.cpp b/Engine/Plugins/PCG/Source/PCG/Private/PCGParamData.cpp
--- a/Engine/Plugins/PCG/Source/PCG/Private/PCGParamData.cpp
+++ b/Engine/Plugins/PCG/Source/PCG/Private/PCGParamData.cpp
@@ -17,9 +17,9 @@ UPCGParamData::UPCGParamData(const FObjectInitializer& ObjectInitializer)
 	: Super(ObjectInitializer)
 {
 	Metadata = ObjectInitializer.CreateDefaultSubobject<UPCGMetadata>(this, TEXT("Metadata"));
-	for (const FPCGMetadataLayerID& MetadataLevel : UPCGParamData::GetAllSupportedMetadataLayerIDs())
+	for (const FPCGMetadataDomainID& MetadataLevel : UPCGParamData::GetAllSupportedMetadataDomainIDs())
 	{
-		Metadata->SetupLayer(MetadataLevel, MetadataLevel == UPCGParamData::GetDefaultMetadataLayerID());
+		Metadata->SetupDomain(MetadataLevel, MetadataLevel == UPCGParamData::GetDefaultMetadataDomainID());
 	}
 }
 
diff --git a/Engine/Plugins/PCG/Source/PCG/Private/Tests/Accessors/PCGAttributeAccessorTest.cpp b/Engine/Plugins/PCG/Source/PCG/Private/Tests/Accessors/PCGAttributeAccessorTest.cpp
--- a/Engine/Plugins/PCG/Source/PCG/Private/Tests/Accessors/PCGAttributeAccessorTest.cpp
+++ b/Engine/Plugins/PCG/Source/PCG/Private/Tests/Accessors/PCGAttributeAccessorTest.cpp
@@ -393,7 +393,7 @@ bool FPCGAttributeAccessorSingleSetDefaultTest::RunTest(const FString& Parameter
 bool FPCGAttributeAccessorSingleSetTest::RunTest(const FString& Parameters)
 {
 	FPCGAttributeAccessorTestHelpers::AttributeData AttributeData;
-	FPCGAttributeAccessorKeysEntries Keys(const_cast<FPCGMetadataLayer*>(AttributeData.FloatAttribute->GetMetadataLayer()));
+	FPCGAttributeAccessorKeysEntries Keys(const_cast<FPCGMetadataDomain*>(AttributeData.FloatAttribute->GetMetadataDomain()));
 
 	if (!TestTrue("Set float attribute", AttributeData.FloatAccessor->Set(*AttributeData.TempFloats, Keys))
 		|| !TestEqual("float value 0", AttributeData.FloatAttribute->GetValueFromItemKey(0), *AttributeData.TempFloats))
@@ -444,7 +444,7 @@ bool FPCGAttributeAccessorSingleInvalidKeySetTest::RunTest(const FString& Parame
 bool FPCGAttributeAccessorSetRangeTest::RunTest(const FString& Parameters)
 {
 	FPCGAttributeAccessorTestHelpers::AttributeData AttributeData;
-	FPCGAttributeAccessorKeysEntries Keys(const_cast<FPCGMetadataLayer*>(AttributeData.FloatAttribute->GetMetadataLayer()));
+	FPCGAttributeAccessorKeysEntries Keys(const_cast<FPCGMetadataDomain*>(AttributeData.FloatAttribute->GetMetadataDomain()));
 
 	const int32 Index = 0;
 
diff --git a/Engine/Plugins/PCG/Source/PCG/Private/Tests/Accessors/PCGAttributePropertySelectorTests.cpp b/Engine/Plugins/PCG/Source/PCG/Private/Tests/Accessors/PCGAttributePropertySelectorTests.cpp
--- a/Engine/Plugins/PCG/Source/PCG/Private/Tests/Accessors/PCGAttributePropertySelectorTests.cpp
+++ b/Engine/Plugins/PCG/Source/PCG/Private/Tests/Accessors/PCGAttributePropertySelectorTests.cpp
@@ -116,24 +116,24 @@ IMPLEMENT_CUSTOM_SIMPLE_AUTOMATION_TEST(FPCGAttributePropertySelector_FromString
 IMPLEMENT_CUSTOM_SIMPLE_AUTOMATION_TEST(FPCGAttributePropertySelector_FromString_ExtraProperty, FPCGAttributePropertySelectorTests, "Plugins.PCG.AttributePropertySelector.FromString.ExtraProperty", PCGTestsCommon::TestFlags)
 IMPLEMENT_CUSTOM_SIMPLE_AUTOMATION_TEST(FPCGAttributePropertySelector_FromString_AtLast, FPCGAttributePropertySelectorTests, "Plugins.PCG.AttributePropertySelector.FromString.@Last", PCGTestsCommon::TestFlags)
 IMPLEMENT_CUSTOM_SIMPLE_AUTOMATION_TEST(FPCGAttributePropertySelector_FromString_AtSource, FPCGAttributePropertySelectorTests, "Plugins.PCG.AttributePropertySelector.FromString.@Source", PCGTestsCommon::TestFlags)
-IMPLEMENT_CUSTOM_SIMPLE_AUTOMATION_TEST(FPCGAttributePropertySelector_FromString_WithLayer_Attribute, FPCGAttributePropertySelectorTests, "Plugins.PCG.AttributePropertySelector.FromString.WithLayer.Attribute", PCGTestsCommon::TestFlags)
-IMPLEMENT_CUSTOM_SIMPLE_AUTOMATION_TEST(FPCGAttributePropertySelector_FromString_WithLayer_Property, FPCGAttributePropertySelectorTests, "Plugins.PCG.AttributePropertySelector.FromString.WithLayer.Property", PCGTestsCommon::TestFlags)
-IMPLEMENT_CUSTOM_SIMPLE_AUTOMATION_TEST(FPCGAttributePropertySelector_FromString_WithLayer_PointProperty, FPCGAttributePropertySelectorTests, "Plugins.PCG.AttributePropertySelector.FromString.WithLayer.PointProperty", PCGTestsCommon::TestFlags)
-IMPLEMENT_CUSTOM_SIMPLE_AUTOMATION_TEST(FPCGAttributePropertySelector_FromString_WithLayer_ExtraProperty, FPCGAttributePropertySelectorTests, "Plugins.PCG.AttributePropertySelector.FromString.WithLayer.ExtraProperty", PCGTestsCommon::TestFlags)
-IMPLEMENT_CUSTOM_SIMPLE_AUTOMATION_TEST(FPCGAttributePropertySelector_FromString_WithLayer_AtLast, FPCGAttributePropertySelectorTests, "Plugins.PCG.AttributePropertySelector.FromString.WithLayer.@Last", PCGTestsCommon::TestFlags)
-IMPLEMENT_CUSTOM_SIMPLE_AUTOMATION_TEST(FPCGAttributePropertySelector_FromString_WithLayer_AtSource, FPCGAttributePropertySelectorTests, "Plugins.PCG.AttributePropertySelector.FromString.WithLayer.@Source", PCGTestsCommon::TestFlags)
+IMPLEMENT_CUSTOM_SIMPLE_AUTOMATION_TEST(FPCGAttributePropertySelector_FromString_WithDomain_Attribute, FPCGAttributePropertySelectorTests, "Plugins.PCG.AttributePropertySelector.FromString.WithDomain.Attribute", PCGTestsCommon::TestFlags)
+IMPLEMENT_CUSTOM_SIMPLE_AUTOMATION_TEST(FPCGAttributePropertySelector_FromString_WithDomain_Property, FPCGAttributePropertySelectorTests, "Plugins.PCG.AttributePropertySelector.FromString.WithDomain.Property", PCGTestsCommon::TestFlags)
+IMPLEMENT_CUSTOM_SIMPLE_AUTOMATION_TEST(FPCGAttributePropertySelector_FromString_WithDomain_PointProperty, FPCGAttributePropertySelectorTests, "Plugins.PCG.AttributePropertySelector.FromString.WithDomain.PointProperty", PCGTestsCommon::TestFlags)
+IMPLEMENT_CUSTOM_SIMPLE_AUTOMATION_TEST(FPCGAttributePropertySelector_FromString_WithDomain_ExtraProperty, FPCGAttributePropertySelectorTests, "Plugins.PCG.AttributePropertySelector.FromString.WithDomain.ExtraProperty", PCGTestsCommon::TestFlags)
+IMPLEMENT_CUSTOM_SIMPLE_AUTOMATION_TEST(FPCGAttributePropertySelector_FromString_WithDomain_AtLast, FPCGAttributePropertySelectorTests, "Plugins.PCG.AttributePropertySelector.FromString.WithDomain.@Last", PCGTestsCommon::TestFlags)
+IMPLEMENT_CUSTOM_SIMPLE_AUTOMATION_TEST(FPCGAttributePropertySelector_FromString_WithDomain_AtSource, FPCGAttributePropertySelectorTests, "Plugins.PCG.AttributePropertySelector.FromString.WithDomain.@Source", PCGTestsCommon::TestFlags)
 IMPLEMENT_CUSTOM_SIMPLE_AUTOMATION_TEST(FPCGAttributePropertySelector_FromString_WithExtraNames_Attribute, FPCGAttributePropertySelectorTests, "Plugins.PCG.AttributePropertySelector.FromString.WithExtraNames.Attribute", PCGTestsCommon::TestFlags)
 IMPLEMENT_CUSTOM_SIMPLE_AUTOMATION_TEST(FPCGAttributePropertySelector_FromString_WithExtraNames_Property, FPCGAttributePropertySelectorTests, "Plugins.PCG.AttributePropertySelector.FromString.WithExtraNames.Property", PCGTestsCommon::TestFlags)
 IMPLEMENT_CUSTOM_SIMPLE_AUTOMATION_TEST(FPCGAttributePropertySelector_FromString_WithExtraNames_PointProperty, FPCGAttributePropertySelectorTests, "Plugins.PCG.AttributePropertySelector.FromString.WithExtraNames.PointProperty", PCGTestsCommon::TestFlags)
 IMPLEMENT_CUSTOM_SIMPLE_AUTOMATION_TEST(FPCGAttributePropertySelector_FromString_WithExtraNames_ExtraProperty, FPCGAttributePropertySelectorTests, "Plugins.PCG.AttributePropertySelector.FromString.WithExtraNames.ExtraProperty", PCGTestsCommon::TestFlags)
 IMPLEMENT_CUSTOM_SIMPLE_AUTOMATION_TEST(FPCGAttributePropertySelector_FromString_WithExtraNames_AtLast, FPCGAttributePropertySelectorTests, "Plugins.PCG.AttributePropertySelector.FromString.WithExtraNames.@Last", PCGTestsCommon::TestFlags)
 IMPLEMENT_CUSTOM_SIMPLE_AUTOMATION_TEST(FPCGAttributePropertySelector_FromString_WithExtraNames_AtSource, FPCGAttributePropertySelectorTests, "Plugins.PCG.AttributePropertySelector.FromString.WithExtraNames.@Source", PCGTestsCommon::TestFlags)
-IMPLEMENT_CUSTOM_SIMPLE_AUTOMATION_TEST(FPCGAttributePropertySelector_FromString_WithExtraNames_WithLayer_Attribute, FPCGAttributePropertySelectorTests, "Plugins.PCG.AttributePropertySelector.FromString.WithExtraNames.WithLayer.Attribute", PCGTestsCommon::TestFlags)
-IMPLEMENT_CUSTOM_SIMPLE_AUTOMATION_TEST(FPCGAttributePropertySelector_FromString_WithExtraNames_WithLayer_Property, FPCGAttributePropertySelectorTests, "Plugins.PCG.AttributePropertySelector.FromString.WithExtraNames.WithLayer.Property", PCGTestsCommon::TestFlags)
-IMPLEMENT_CUSTOM_SIMPLE_AUTOMATION_TEST(FPCGAttributePropertySelector_FromString_WithExtraNames_WithLayer_PointProperty, FPCGAttributePropertySelectorTests, "Plugins.PCG.AttributePropertySelector.FromString.WithExtraNames.WithLayer.PointProperty", PCGTestsCommon::TestFlags)
-IMPLEMENT_CUSTOM_SIMPLE_AUTOMATION_TEST(FPCGAttributePropertySelector_FromString_WithExtraNames_WithLayer_ExtraProperty, FPCGAttributePropertySelectorTests, "Plugins.PCG.AttributePropertySelector.FromString.WithExtraNames.WithLayer.ExtraProperty", PCGTestsCommon::TestFlags)
-IMPLEMENT_CUSTOM_SIMPLE_AUTOMATION_TEST(FPCGAttributePropertySelector_FromString_WithExtraNames_WithLayer_AtLast, FPCGAttributePropertySelectorTests, "Plugins.PCG.AttributePropertySelector.FromString.WithExtraNames.WithLayer.@Last", PCGTestsCommon::TestFlags)
-IMPLEMENT_CUSTOM_SIMPLE_AUTOMATION_TEST(FPCGAttributePropertySelector_FromString_WithExtraNames_WithLayer_AtSource, FPCGAttributePropertySelectorTests, "Plugins.PCG.AttributePropertySelector.FromString.WithExtraNames.WithLayer.@Source", PCGTestsCommon::TestFlags)
+IMPLEMENT_CUSTOM_SIMPLE_AUTOMATION_TEST(FPCGAttributePropertySelector_FromString_WithExtraNames_WithDomain_Attribute, FPCGAttributePropertySelectorTests, "Plugins.PCG.AttributePropertySelector.FromString.WithExtraNames.WithDomain.Attribute", PCGTestsCommon::TestFlags)
+IMPLEMENT_CUSTOM_SIMPLE_AUTOMATION_TEST(FPCGAttributePropertySelector_FromString_WithExtraNames_WithDomain_Property, FPCGAttributePropertySelectorTests, "Plugins.PCG.AttributePropertySelector.FromString.WithExtraNames.WithDomain.Property", PCGTestsCommon::TestFlags)
+IMPLEMENT_CUSTOM_SIMPLE_AUTOMATION_TEST(FPCGAttributePropertySelector_FromString_WithExtraNames_WithDomain_PointProperty, FPCGAttributePropertySelectorTests, "Plugins.PCG.AttributePropertySelector.FromString.WithExtraNames.WithDomain.PointProperty", PCGTestsCommon::TestFlags)
+IMPLEMENT_CUSTOM_SIMPLE_AUTOMATION_TEST(FPCGAttributePropertySelector_FromString_WithExtraNames_WithDomain_ExtraProperty, FPCGAttributePropertySelectorTests, "Plugins.PCG.AttributePropertySelector.FromString.WithExtraNames.WithDomain.ExtraProperty", PCGTestsCommon::TestFlags)
+IMPLEMENT_CUSTOM_SIMPLE_AUTOMATION_TEST(FPCGAttributePropertySelector_FromString_WithExtraNames_WithDomain_AtLast, FPCGAttributePropertySelectorTests, "Plugins.PCG.AttributePropertySelector.FromString.WithExtraNames.WithDomain.@Last", PCGTestsCommon::TestFlags)
+IMPLEMENT_CUSTOM_SIMPLE_AUTOMATION_TEST(FPCGAttributePropertySelector_FromString_WithExtraNames_WithDomain_AtSource, FPCGAttributePropertySelectorTests, "Plugins.PCG.AttributePropertySelector.FromString.WithExtraNames.WithDomain.@Source", PCGTestsCommon::TestFlags)
 
 IMPLEMENT_CUSTOM_SIMPLE_AUTOMATION_TEST(FPCGAttributePropertySelector_ExportImport_Attribute, FPCGAttributePropertySelectorTests, "Plugins.PCG.AttributePropertySelector.ExportImport.Attribute", PCGTestsCommon::TestFlags)
 IMPLEMENT_CUSTOM_SIMPLE_AUTOMATION_TEST(FPCGAttributePropertySelector_ExportImport_Property, FPCGAttributePropertySelectorTests, "Plugins.PCG.AttributePropertySelector.ExportImport.Property", PCGTestsCommon::TestFlags)
@@ -141,18 +141,18 @@ IMPLEMENT_CUSTOM_SIMPLE_AUTOMATION_TEST(FPCGAttributePropertySelector_ExportImpo
 IMPLEMENT_CUSTOM_SIMPLE_AUTOMATION_TEST(FPCGAttributePropertySelector_ExportImport_ExtraProperty, FPCGAttributePropertySelectorTests, "Plugins.PCG.AttributePropertySelector.ExportImport.ExtraProperty", PCGTestsCommon::TestFlags)
 IMPLEMENT_CUSTOM_SIMPLE_AUTOMATION_TEST(FPCGAttributePropertySelector_ExportImport_AtLast, FPCGAttributePropertySelectorTests, "Plugins.PCG.AttributePropertySelector.ExportImport.@Last", PCGTestsCommon::TestFlags)
 IMPLEMENT_CUSTOM_SIMPLE_AUTOMATION_TEST(FPCGAttributePropertySelector_ExportImport_AtSource, FPCGAttributePropertySelectorTests, "Plugins.PCG.AttributePropertySelector.ExportImport.@Source", PCGTestsCommon::TestFlags)
-IMPLEMENT_CUSTOM_SIMPLE_AUTOMATION_TEST(FPCGAttributePropertySelector_ExportImport_WithLayer_Attribute, FPCGAttributePropertySelectorTests, "Plugins.PCG.AttributePropertySelector.ExportImport.WithLayer.Attribute", PCGTestsCommon::TestFlags)
-IMPLEMENT_CUSTOM_SIMPLE_AUTOMATION_TEST(FPCGAttributePropertySelector_ExportImport_WithLayer_Property, FPCGAttributePropertySelectorTests, "Plugins.PCG.AttributePropertySelector.ExportImport.WithLayer.Property", PCGTestsCommon::TestFlags)
-IMPLEMENT_CUSTOM_SIMPLE_AUTOMATION_TEST(FPCGAttributePropertySelector_ExportImport_WithLayer_PointProperty, FPCGAttributePropertySelectorTests, "Plugins.PCG.AttributePropertySelector.ExportImport.WithLayer.PointProperty", PCGTestsCommon::TestFlags)
-IMPLEMENT_CUSTOM_SIMPLE_AUTOMATION_TEST(FPCGAttributePropertySelector_ExportImport_WithLayer_ExtraProperty, FPCGAttributePropertySelectorTests, "Plugins.PCG.AttributePropertySelector.ExportImport.WithLayer.ExtraProperty", PCGTestsCommon::TestFlags)
+IMPLEMENT_CUSTOM_SIMPLE_AUTOMATION_TEST(FPCGAttributePropertySelector_ExportImport_WithDomain_Attribute, FPCGAttributePropertySelectorTests, "Plugins.PCG.AttributePropertySelector.ExportImport.WithDomain.Attribute", PCGTestsCommon::TestFlags)
+IMPLEMENT_CUSTOM_SIMPLE_AUTOMATION_TEST(FPCGAttributePropertySelector_ExportImport_WithDomain_Property, FPCGAttributePropertySelectorTests, "Plugins.PCG.AttributePropertySelector.ExportImport.WithDomain.Property", PCGTestsCommon::TestFlags)
+IMPLEMENT_CUSTOM_SIMPLE_AUTOMATION_TEST(FPCGAttributePropertySelector_ExportImport_WithDomain_PointProperty, FPCGAttributePropertySelectorTests, "Plugins.PCG.AttributePropertySelector.ExportImport.WithDomain.PointProperty", PCGTestsCommon::TestFlags)
+IMPLEMENT_CUSTOM_SIMPLE_AUTOMATION_TEST(FPCGAttributePropertySelector_ExportImport_WithDomain_ExtraProperty, FPCGAttributePropertySelectorTests, "Plugins.PCG.AttributePropertySelector.ExportImport.WithDomain.ExtraProperty", PCGTestsCommon::TestFlags)
 IMPLEMENT_CUSTOM_SIMPLE_AUTOMATION_TEST(FPCGAttributePropertySelector_ExportImport_WithExtraNames_Attribute, FPCGAttributePropertySelectorTests, "Plugins.PCG.AttributePropertySelector.ExportImport.WithExtraNames.Attribute", PCGTestsCommon::TestFlags)
 IMPLEMENT_CUSTOM_SIMPLE_AUTOMATION_TEST(FPCGAttributePropertySelector_ExportImport_WithExtraNames_Property, FPCGAttributePropertySelectorTests, "Plugins.PCG.AttributePropertySelector.ExportImport.WithExtraNames.Property", PCGTestsCommon::TestFlags)
 IMPLEMENT_CUSTOM_SIMPLE_AUTOMATION_TEST(FPCGAttributePropertySelector_ExportImport_WithExtraNames_PointProperty, FPCGAttributePropertySelectorTests, "Plugins.PCG.AttributePropertySelector.ExportImport.WithExtraNames.PointProperty", PCGTestsCommon::TestFlags)
 IMPLEMENT_CUSTOM_SIMPLE_AUTOMATION_TEST(FPCGAttributePropertySelector_ExportImport_WithExtraNames_ExtraProperty, FPCGAttributePropertySelectorTests, "Plugins.PCG.AttributePropertySelector.ExportImport.WithExtraNames.ExtraProperty", PCGTestsCommon::TestFlags)
-IMPLEMENT_CUSTOM_SIMPLE_AUTOMATION_TEST(FPCGAttributePropertySelector_ExportImport_WithExtraNames_WithLayer_Attribute, FPCGAttributePropertySelectorTests, "Plugins.PCG.AttributePropertySelector.ExportImport.WithExtraNames.WithLayer.Attribute", PCGTestsCommon::TestFlags)
-IMPLEMENT_CUSTOM_SIMPLE_AUTOMATION_TEST(FPCGAttributePropertySelector_ExportImport_WithExtraNames_WithLayer_Property, FPCGAttributePropertySelectorTests, "Plugins.PCG.AttributePropertySelector.ExportImport.WithExtraNames.WithLayer.Property", PCGTestsCommon::TestFlags)
-IMPLEMENT_CUSTOM_SIMPLE_AUTOMATION_TEST(FPCGAttributePropertySelector_ExportImport_WithExtraNames_WithLayer_PointProperty, FPCGAttributePropertySelectorTests, "Plugins.PCG.AttributePropertySelector.ExportImport.WithExtraNames.WithLayer.PointProperty", PCGTestsCommon::TestFlags)
-IMPLEMENT_CUSTOM_SIMPLE_AUTOMATION_TEST(FPCGAttributePropertySelector_ExportImport_WithExtraNames_WithLayer_ExtraProperty, FPCGAttributePropertySelectorTests, "Plugins.PCG.AttributePropertySelector.ExportImport.WithExtraNames.WithLayer.ExtraProperty", PCGTestsCommon::TestFlags)
+IMPLEMENT_CUSTOM_SIMPLE_AUTOMATION_TEST(FPCGAttributePropertySelector_ExportImport_WithExtraNames_WithDomain_Attribute, FPCGAttributePropertySelectorTests, "Plugins.PCG.AttributePropertySelector.ExportImport.WithExtraNames.WithDomain.Attribute", PCGTestsCommon::TestFlags)
+IMPLEMENT_CUSTOM_SIMPLE_AUTOMATION_TEST(FPCGAttributePropertySelector_ExportImport_WithExtraNames_WithDomain_Property, FPCGAttributePropertySelectorTests, "Plugins.PCG.AttributePropertySelector.ExportImport.WithExtraNames.WithDomain.Property", PCGTestsCommon::TestFlags)
+IMPLEMENT_CUSTOM_SIMPLE_AUTOMATION_TEST(FPCGAttributePropertySelector_ExportImport_WithExtraNames_WithDomain_PointProperty, FPCGAttributePropertySelectorTests, "Plugins.PCG.AttributePropertySelector.ExportImport.WithExtraNames.WithDomain.PointProperty", PCGTestsCommon::TestFlags)
+IMPLEMENT_CUSTOM_SIMPLE_AUTOMATION_TEST(FPCGAttributePropertySelector_ExportImport_WithExtraNames_WithDomain_ExtraProperty, FPCGAttributePropertySelectorTests, "Plugins.PCG.AttributePropertySelector.ExportImport.WithExtraNames.WithDomain.ExtraProperty", PCGTestsCommon::TestFlags)
 
 // To be removed when deprecated point property no longer exists.
 IMPLEMENT_CUSTOM_SIMPLE_AUTOMATION_TEST(FPCGAttributePropertySelector_Deprecation_PointProperty, FPCGAttributePropertySelectorTests, "Plugins.PCG.AttributePropertySelector.Deprecation.PointProperty", PCGTestsCommon::TestFlags)
@@ -165,8 +165,8 @@ IMPLEMENT_CUSTOM_SIMPLE_AUTOMATION_TEST(FPCGAttributePropertySelector_Deprecatio
 IMPLEMENT_CUSTOM_SIMPLE_AUTOMATION_TEST(FPCGAttributePropertySelector_Invalid_JustAt, FPCGAttributePropertySelectorTests, "Plugins.PCG.AttributePropertySelector.Invalid.JustAt", PCGTestsCommon::TestFlags)
 IMPLEMENT_CUSTOM_SIMPLE_AUTOMATION_TEST(FPCGAttributePropertySelector_Invalid_InvalidSymbols, FPCGAttributePropertySelectorTests, "Plugins.PCG.AttributePropertySelector.Invalid.InvalidSymbols", PCGTestsCommon::TestFlags)
 
-IMPLEMENT_CUSTOM_SIMPLE_AUTOMATION_TEST(FPCGAttributePropertySelector_JustLayerName, FPCGAttributePropertySelectorTests, "Plugins.PCG.AttributePropertySelector.JustLayerName", PCGTestsCommon::TestFlags)
-IMPLEMENT_CUSTOM_SIMPLE_AUTOMATION_TEST(FPCGAttributePropertySelector_UpdateLayerNameToAtLast, FPCGAttributePropertySelectorTests, "Plugins.PCG.AttributePropertySelector.UpdateLayerNameToAtLast", PCGTestsCommon::TestFlags)
+IMPLEMENT_CUSTOM_SIMPLE_AUTOMATION_TEST(FPCGAttributePropertySelector_JustDomainName, FPCGAttributePropertySelectorTests, "Plugins.PCG.AttributePropertySelector.JustDomainName", PCGTestsCommon::TestFlags)
+IMPLEMENT_CUSTOM_SIMPLE_AUTOMATION_TEST(FPCGAttributePropertySelector_UpdateDomainNameToAtLast, FPCGAttributePropertySelectorTests, "Plugins.PCG.AttributePropertySelector.UpdateDomainNameToAtLast", PCGTestsCommon::TestFlags)
 
 bool FPCGAttributePropertySelector_FromString_Attribute::RunTest(const FString& Parameters)
 {
@@ -198,34 +198,34 @@ bool FPCGAttributePropertySelector_FromString_AtSource::RunTest(const FString& P
 	return FromStringComparison(TEXT("@Source"), FPCGAttributePropertySelector::CreateAttributeSelector(TEXT("@Source")));
 }
 
-bool FPCGAttributePropertySelector_FromString_WithLayer_Attribute::RunTest(const FString& Parameters)
+bool FPCGAttributePropertySelector_FromString_WithDomain_Attribute::RunTest(const FString& Parameters)
 {
-	return FromStringComparison(TEXT("@NewLayer.MyAttr"), FPCGAttributePropertySelector::CreateAttributeSelector(TEXT("MyAttr"), TEXT("NewLayer")));
+	return FromStringComparison(TEXT("@NewDomain.MyAttr"), FPCGAttributePropertySelector::CreateAttributeSelector(TEXT("MyAttr"), TEXT("NewDomain")));
 }
 
-bool FPCGAttributePropertySelector_FromString_WithLayer_Property::RunTest(const FString& Parameters)
+bool FPCGAttributePropertySelector_FromString_WithDomain_Property::RunTest(const FString& Parameters)
 {
-	return FromStringComparison(TEXT("@NewLayer.$MyProperty"), FPCGAttributePropertySelector::CreatePropertySelector(TEXT("MyProperty"), TEXT("NewLayer")));
+	return FromStringComparison(TEXT("@NewDomain.$MyProperty"), FPCGAttributePropertySelector::CreatePropertySelector(TEXT("MyProperty"), TEXT("NewDomain")));
 }
 
-bool FPCGAttributePropertySelector_FromString_WithLayer_PointProperty::RunTest(const FString& Parameters)
+bool FPCGAttributePropertySelector_FromString_WithDomain_PointProperty::RunTest(const FString& Parameters)
 {
-	return FromStringComparison(TEXT("@NewLayer.$Scale"), FPCGAttributePropertySelector::CreatePointPropertySelector(EPCGPointProperties::Scale, TEXT("NewLayer")));
+	return FromStringComparison(TEXT("@NewDomain.$Scale"), FPCGAttributePropertySelector::CreatePointPropertySelector(EPCGPointProperties::Scale, TEXT("NewDomain")));
 }
 
-bool FPCGAttributePropertySelector_FromString_WithLayer_ExtraProperty::RunTest(const FString& Parameters)
+bool FPCGAttributePropertySelector_FromString_WithDomain_ExtraProperty::RunTest(const FString& Parameters)
 {
-	return FromStringComparison(TEXT("@NewLayer.$Index"), FPCGAttributePropertySelector::CreateExtraPropertySelector(EPCGExtraProperties::Index, TEXT("NewLayer")));
+	return FromStringComparison(TEXT("@NewDomain.$Index"), FPCGAttributePropertySelector::CreateExtraPropertySelector(EPCGExtraProperties::Index, TEXT("NewDomain")));
 }
 
-bool FPCGAttributePropertySelector_FromString_WithLayer_AtLast::RunTest(const FString& Parameters)
+bool FPCGAttributePropertySelector_FromString_WithDomain_AtLast::RunTest(const FString& Parameters)
 {
-	return FromStringComparison(TEXT("@NewLayer.@Last"), FPCGAttributePropertySelector::CreateAttributeSelector(TEXT("@Last"), TEXT("NewLayer")));
+	return FromStringComparison(TEXT("@NewDomain.@Last"), FPCGAttributePropertySelector::CreateAttributeSelector(TEXT("@Last"), TEXT("NewDomain")));
 }
 
-bool FPCGAttributePropertySelector_FromString_WithLayer_AtSource::RunTest(const FString& Parameters)
+bool FPCGAttributePropertySelector_FromString_WithDomain_AtSource::RunTest(const FString& Parameters)
 {
-	return FromStringComparison(TEXT("@NewLayer.@Source"), FPCGAttributePropertySelector::CreateAttributeSelector(TEXT("@Source"), TEXT("NewLayer")));
+	return FromStringComparison(TEXT("@NewDomain.@Source"), FPCGAttributePropertySelector::CreateAttributeSelector(TEXT("@Source"), TEXT("NewDomain")));
 }
 
 bool FPCGAttributePropertySelector_FromString_WithExtraNames_Attribute::RunTest(const FString& Parameters)
@@ -258,34 +258,34 @@ bool FPCGAttributePropertySelector_FromString_WithExtraNames_AtSource::RunTest(c
 	return FromStringComparison(TEXT("@Source.X"), FPCGAttributePropertySelector::CreateAttributeSelector(TEXT("@Source"), NAME_None, {TEXT("X")}));
 }
 
-bool FPCGAttributePropertySelector_FromString_WithExtraNames_WithLayer_Attribute::RunTest(const FString& Parameters)
+bool FPCGAttributePropertySelector_FromString_WithExtraNames_WithDomain_Attribute::RunTest(const FString& Parameters)
 {
-	return FromStringComparison(TEXT("@NewLayer.MyAttr.Position.XY"), FPCGAttributePropertySelector::CreateAttributeSelector(TEXT("MyAttr"), TEXT("NewLayer"), {TEXT("Position"), TEXT("XY")}));
+	return FromStringComparison(TEXT("@NewDomain.MyAttr.Position.XY"), FPCGAttributePropertySelector::CreateAttributeSelector(TEXT("MyAttr"), TEXT("NewDomain"), {TEXT("Position"), TEXT("XY")}));
 }
 
-bool FPCGAttributePropertySelector_FromString_WithExtraNames_WithLayer_Property::RunTest(const FString& Parameters)
+bool FPCGAttributePropertySelector_FromString_WithExtraNames_WithDomain_Property::RunTest(const FString& Parameters)
 {
-	return FromStringComparison(TEXT("@NewLayer.$MyProperty.RGBA"), FPCGAttributePropertySelector::CreatePropertySelector(TEXT("MyProperty"), TEXT("NewLayer"), {TEXT("RGBA")}));
+	return FromStringComparison(TEXT("@NewDomain.$MyProperty.RGBA"), FPCGAttributePropertySelector::CreatePropertySelector(TEXT("MyProperty"), TEXT("NewDomain"), {TEXT("RGBA")}));
 }
 
-bool FPCGAttributePropertySelector_FromString_WithExtraNames_WithLayer_PointProperty::RunTest(const FString& Parameters)
+bool FPCGAttributePropertySelector_FromString_WithExtraNames_WithDomain_PointProperty::RunTest(const FString& Parameters)
 {
-	return FromStringComparison(TEXT("@NewLayer.$Steepness.Sign"), FPCGAttributePropertySelector::CreatePointPropertySelector(EPCGPointProperties::Steepness, TEXT("NewLayer"), {TEXT("Sign")}));
+	return FromStringComparison(TEXT("@NewDomain.$Steepness.Sign"), FPCGAttributePropertySelector::CreatePointPropertySelector(EPCGPointProperties::Steepness, TEXT("NewDomain"), {TEXT("Sign")}));
 }
 
-bool FPCGAttributePropertySelector_FromString_WithExtraNames_WithLayer_ExtraProperty::RunTest(const FString& Parameters)
+bool FPCGAttributePropertySelector_FromString_WithExtraNames_WithDomain_ExtraProperty::RunTest(const FString& Parameters)
 {
-	return FromStringComparison(TEXT("@NewLayer.$Index.Sign"), FPCGAttributePropertySelector::CreateExtraPropertySelector(EPCGExtraProperties::Index, TEXT("NewLayer"), {TEXT("Sign")}));
+	return FromStringComparison(TEXT("@NewDomain.$Index.Sign"), FPCGAttributePropertySelector::CreateExtraPropertySelector(EPCGExtraProperties::Index, TEXT("NewDomain"), {TEXT("Sign")}));
 }
 
-bool FPCGAttributePropertySelector_FromString_WithExtraNames_WithLayer_AtLast::RunTest(const FString& Parameters)
+bool FPCGAttributePropertySelector_FromString_WithExtraNames_WithDomain_AtLast::RunTest(const FString& Parameters)
 {
-	return FromStringComparison(TEXT("@[REDACTED_EMAIL]"), FPCGAttributePropertySelector::CreateAttributeSelector(TEXT("@Last"), TEXT("NewLayer"), {TEXT("Position"), TEXT("XY")}));
+	return FromStringComparison(TEXT("@[REDACTED_EMAIL]"), FPCGAttributePropertySelector::CreateAttributeSelector(TEXT("@Last"), TEXT("NewDomain"), {TEXT("Position"), TEXT("XY")}));
 }
 
-bool FPCGAttributePropertySelector_FromString_WithExtraNames_WithLayer_AtSource::RunTest(const FString& Parameters)
+bool FPCGAttributePropertySelector_FromString_WithExtraNames_WithDomain_AtSource::RunTest(const FString& Parameters)
 {
-	return FromStringComparison(TEXT("@[REDACTED_EMAIL]"), FPCGAttributePropertySelector::CreateAttributeSelector(TEXT("@Source"), TEXT("NewLayer"), {TEXT("Position"), TEXT("XY")}));
+	return FromStringComparison(TEXT("@[REDACTED_EMAIL]"), FPCGAttributePropertySelector::CreateAttributeSelector(TEXT("@Source"), TEXT("NewDomain"), {TEXT("Position"), TEXT("XY")}));
 }
 
 bool FPCGAttributePropertySelector_ExportImport_Attribute::RunTest(const FString& Parameters)
@@ -318,24 +318,24 @@ bool FPCGAttributePropertySelector_ExportImport_AtSource::RunTest(const FString&
 	return ImportExport(TEXT("PCGBegin(@Source)PCGEnd"), FPCGAttributePropertySelector::CreateAttributeSelector(TEXT("@Source")));
 }
 
-bool FPCGAttributePropertySelector_ExportImport_WithLayer_Attribute::RunTest(const FString& Parameters)
+bool FPCGAttributePropertySelector_ExportImport_WithDomain_Attribute::RunTest(const FString& Parameters)
 {
-	return ImportExport(TEXT("PCGBegin(@NewLayer.MyAttr)PCGEnd"), FPCGAttributePropertySelector::CreateAttributeSelector(TEXT("MyAttr"), TEXT("NewLayer")));
+	return ImportExport(TEXT("PCGBegin(@NewDomain.MyAttr)PCGEnd"), FPCGAttributePropertySelector::CreateAttributeSelector(TEXT("MyAttr"), TEXT("NewDomain")));
 }
 
-bool FPCGAttributePropertySelector_ExportImport_WithLayer_Property::RunTest(const FString& Parameters)
+bool FPCGAttributePropertySelector_ExportImport_WithDomain_Property::RunTest(const FString& Parameters)
 {
-	return ImportExport(TEXT("PCGBegin(@NewLayer.$MyProperty)PCGEnd"), FPCGAttributePropertySelector::CreatePropertySelector(TEXT("MyProperty"), TEXT("NewLayer")));
+	return ImportExport(TEXT("PCGBegin(@NewDomain.$MyProperty)PCGEnd"), FPCGAttributePropertySelector::CreatePropertySelector(TEXT("MyProperty"), TEXT("NewDomain")));
 }
 
-bool FPCGAttributePropertySelector_ExportImport_WithLayer_PointProperty::RunTest(const FString& Parameters)
+bool FPCGAttributePropertySelector_ExportImport_WithDomain_PointProperty::RunTest(const FString& Parameters)
 {
-	return ImportExport(TEXT("PCGBegin(@NewLayer.$Scale)PCGEnd"), FPCGAttributePropertySelector::CreatePointPropertySelector(EPCGPointProperties::Scale, TEXT("NewLayer")));
+	return ImportExport(TEXT("PCGBegin(@NewDomain.$Scale)PCGEnd"), FPCGAttributePropertySelector::CreatePointPropertySelector(EPCGPointProperties::Scale, TEXT("NewDomain")));
 }
 
-bool FPCGAttributePropertySelector_ExportImport_WithLayer_ExtraProperty::RunTest(const FString& Parameters)
+bool FPCGAttributePropertySelector_ExportImport_WithDomain_ExtraProperty::RunTest(const FString& Parameters)
 {
-	return ImportExport(TEXT("PCGBegin(@NewLayer.$Index)PCGEnd"), FPCGAttributePropertySelector::CreateExtraPropertySelector(EPCGExtraProperties::Index, TEXT("NewLayer")));
+	return ImportExport(TEXT("PCGBegin(@NewDomain.$Index)PCGEnd"), FPCGAttributePropertySelector::CreateExtraPropertySelector(EPCGExtraProperties::Index, TEXT("NewDomain")));
 }
 
 bool FPCGAttributePropertySelector_ExportImport_WithExtraNames_Attribute::RunTest(const FString& Parameters)
@@ -358,24 +358,24 @@ bool FPCGAttributePropertySelector_ExportImport_WithExtraNames_ExtraProperty::Ru
 	return ImportExport(TEXT("PCGBegin($Index.Abs)PCGEnd"), FPCGAttributePropertySelector::CreateExtraPropertySelector(EPCGExtraProperties::Index, NAME_None, {TEXT("Abs")}));
 }
 
-bool FPCGAttributePropertySelector_ExportImport_WithExtraNames_WithLayer_Attribute::RunTest(const FString& Parameters)
+bool FPCGAttributePropertySelector_ExportImport_WithExtraNames_WithDomain_Attribute::RunTest(const FString& Parameters)
 {
-	return ImportExport(TEXT("PCGBegin(@NewLayer.MyAttr.Position.XY)PCGEnd"), FPCGAttributePropertySelector::CreateAttributeSelector(TEXT("MyAttr"), TEXT("NewLayer"), {TEXT("Position"), TEXT("XY")}));
+	return ImportExport(TEXT("PCGBegin(@NewDomain.MyAttr.Position.XY)PCGEnd"), FPCGAttributePropertySelector::CreateAttributeSelector(TEXT("MyAttr"), TEXT("NewDomain"), {TEXT("Position"), TEXT("XY")}));
 }
 
-bool FPCGAttributePropertySelector_ExportImport_WithExtraNames_WithLayer_Property::RunTest(const FString& Parameters)
+bool FPCGAttributePropertySelector_ExportImport_WithExtraNames_WithDomain_Property::RunTest(const FString& Parameters)
 {
-	return ImportExport(TEXT("PCGBegin(@NewLayer.$MyProperty.RGBA)PCGEnd"), FPCGAttributePropertySelector::CreatePropertySelector(TEXT("MyProperty"), TEXT("NewLayer"), {TEXT("RGBA")}));
+	return ImportExport(TEXT("PCGBegin(@NewDomain.$MyProperty.RGBA)PCGEnd"), FPCGAttributePropertySelector::CreatePropertySelector(TEXT("MyProperty"), TEXT("NewDomain"), {TEXT("RGBA")}));
 }
 
-bool FPCGAttributePropertySelector_ExportImport_WithExtraNames_WithLayer_PointProperty::RunTest(const FString& Parameters)
+bool FPCGAttributePropertySelector_ExportImport_WithExtraNames_WithDomain_PointProperty::RunTest(const FString& Parameters)
 {
-	return ImportExport(TEXT("PCGBegin(@NewLayer.$Steepness.Sign)PCGEnd"), FPCGAttributePropertySelector::CreatePointPropertySelector(EPCGPointProperties::Steepness, TEXT("NewLayer"), {TEXT("Sign")}));
+	return ImportExport(TEXT("PCGBegin(@NewDomain.$Steepness.Sign)PCGEnd"), FPCGAttributePropertySelector::CreatePointPropertySelector(EPCGPointProperties::Steepness, TEXT("NewDomain"), {TEXT("Sign")}));
 }
 
-bool FPCGAttributePropertySelector_ExportImport_WithExtraNames_WithLayer_ExtraProperty::RunTest(const FString& Parameters)
+bool FPCGAttributePropertySelector_ExportImport_WithExtraNames_WithDomain_ExtraProperty::RunTest(const FString& Parameters)
 {
-	return ImportExport(TEXT("PCGBegin(@NewLayer.$Index.Sign)PCGEnd"), FPCGAttributePropertySelector::CreateExtraPropertySelector(EPCGExtraProperties::Index, TEXT("NewLayer"), {TEXT("Sign")}));
+	return ImportExport(TEXT("PCGBegin(@NewDomain.$Index.Sign)PCGEnd"), FPCGAttributePropertySelector::CreateExtraPropertySelector(EPCGExtraProperties::Index, TEXT("NewDomain"), {TEXT("Sign")}));
 }
 
 bool FPCGAttributePropertySelector_Deprecation_PointProperty::RunTest(const FString& Parameters)
@@ -417,7 +417,7 @@ bool FPCGAttributePropertySelector_Deprecation_CDO_NewVersion_WithExtraNames_Poi
 bool FPCGAttributePropertySelector_Invalid_JustAt::RunTest(const FString& Parameters)
 {
 	FPCGAttributePropertySelector Selector = FPCGAttributePropertySelector::CreateSelectorFromString(TEXT("@"));
-	UTEST_EQUAL("Selector has no layer name", Selector.GetLayerName(), FName{NAME_None});
+	UTEST_EQUAL("Selector has no domain name", Selector.GetDomainName(), FName{NAME_None});
 	UTEST_EQUAL("Selector is an attribute", Selector.GetSelection(), EPCGAttributePropertySelection::Attribute);
 	UTEST_FALSE("Selector is invalid", Selector.IsValid());
 	return true;
@@ -426,26 +426,26 @@ bool FPCGAttributePropertySelector_Invalid_JustAt::RunTest(const FString& Parame
 bool FPCGAttributePropertySelector_Invalid_InvalidSymbols::RunTest(const FString& Parameters)
 {
 	FPCGAttributePropertySelector Selector = FPCGAttributePropertySelector::CreateSelectorFromString(TEXT(")(*&^"));
-	UTEST_EQUAL("Selector has no layer name", Selector.GetLayerName(), FName{NAME_None});
+	UTEST_EQUAL("Selector has no domain name", Selector.GetDomainName(), FName{NAME_None});
 	UTEST_EQUAL("Selector is an attribute", Selector.GetSelection(), EPCGAttributePropertySelection::Attribute);
 	UTEST_FALSE("Selector is invalid", Selector.IsValid());
 	return true;
 }
 
-bool FPCGAttributePropertySelector_JustLayerName::RunTest(const FString& Parameters)
+bool FPCGAttributePropertySelector_JustDomainName::RunTest(const FString& Parameters)
 {
-	FPCGAttributePropertySelector Selector = FPCGAttributePropertySelector::CreateSelectorFromString(TEXT("@MyLayer"));
-	UTEST_EQUAL("Selector has the right layer name", Selector.GetLayerName(), FName{TEXT("MyLayer")});
+	FPCGAttributePropertySelector Selector = FPCGAttributePropertySelector::CreateSelectorFromString(TEXT("@MyDomain"));
+	UTEST_EQUAL("Selector has the right domain name", Selector.GetDomainName(), FName{TEXT("MyDomain")});
 	UTEST_EQUAL("Selector is an attribute", Selector.GetSelection(), EPCGAttributePropertySelection::Attribute);
 	UTEST_EQUAL("Selector has an attribute name of None", Selector.GetAttributeName(), FName{NAME_None});
 	return true;
 }
 
-bool FPCGAttributePropertySelector_UpdateLayerNameToAtLast::RunTest(const FString& Parameters)
+bool FPCGAttributePropertySelector_UpdateDomainNameToAtLast::RunTest(const FString& Parameters)
 {
 	FPCGAttributePropertySelector Selector = FPCGAttributePropertySelector::CreateSelectorFromString(TEXT("@Las"));
 	Selector.Update(TEXT("@Last"));
-	UTEST_EQUAL("Selector has no layer name", Selector.GetLayerName(), FName{NAME_None});
+	UTEST_EQUAL("Selector has no domain name", Selector.GetDomainName(), FName{NAME_None});
 	UTEST_EQUAL("Selector is an attribute", Selector.GetSelection(), EPCGAttributePropertySelection::Attribute);
 	UTEST_EQUAL("Selector has an attribute name of @Last", Selector.GetAttributeName(), FName{TEXT("@Last")});
 	return true;
diff --git a/Engine/Plugins/PCG/Source/PCG/Public/Data/PCGBasePointData.h b/Engine/Plugins/PCG/Source/PCG/Public/Data/PCGBasePointData.h
--- a/Engine/Plugins/PCG/Source/PCG/Public/Data/PCGBasePointData.h
+++ b/Engine/Plugins/PCG/Source/PCG/Public/Data/PCGBasePointData.h
@@ -140,8 +140,8 @@ public:
 	virtual EPCGPointNativeProperties GetAllocatedProperties() const { return EPCGPointNativeProperties::All; }
 
 	/** Metadata specific */
-	virtual FPCGMetadataLayerID GetDefaultMetadataLayerID() const override { return PCGMetadataLayerID::Elements; }
-	virtual TArray<FPCGMetadataLayerID> GetAllSupportedMetadataLayerIDs() const override { return {PCGMetadataLayerID::Data, PCGMetadataLayerID::Elements}; }
+	virtual FPCGMetadataDomainID GetDefaultMetadataDomainID() const override { return PCGMetadataDomainID::Elements; }
+	virtual TArray<FPCGMetadataDomainID> GetAllSupportedMetadataDomainIDs() const override { return {PCGMetadataDomainID::Data, PCGMetadataDomainID::Elements}; }
 
 	static void CopyPoints(const UPCGBasePointData* From, UPCGBasePointData* To, const TArray<int32>& InDataIndices, bool bCopyAll);
 
diff --git a/Engine/Plugins/PCG/Source/PCG/Public/Data/PCGSplineData.h b/Engine/Plugins/PCG/Source/PCG/Public/Data/PCGSplineData.h
--- a/Engine/Plugins/PCG/Source/PCG/Public/Data/PCGSplineData.h
+++ b/Engine/Plugins/PCG/Source/PCG/Public/Data/PCGSplineData.h
@@ -17,7 +17,7 @@ class UPCGSurfaceData;
 
 namespace PCGSplineData
 {
-	const FName ControlPointLayerName = "ControlPoints";
+	const FName ControlPointDomainName = "ControlPoints";
 }
 
 UCLASS(BlueprintType, ClassGroup = (Procedural))
diff --git a/Engine/Plugins/PCG/Source/PCG/Public/Helpers/PCGMetadataHelpers.h b/Engine/Plugins/PCG/Source/PCG/Public/Helpers/PCGMetadataHelpers.h
--- a/Engine/Plugins/PCG/Source/PCG/Public/Helpers/PCGMetadataHelpers.h
+++ b/Engine/Plugins/PCG/Source/PCG/Public/Helpers/PCGMetadataHelpers.h
@@ -9,7 +9,7 @@
 #include "HAL/Platform.h"
 
 class FPCGMetadataAttributeBase;
-class FPCGMetadataLayer;
+class FPCGMetadataDomain;
 class UPCGData;
 class UPCGMetadata;
 struct FPCGContext;
@@ -19,12 +19,12 @@ template <typename FuncType> class TFunction;
 namespace PCGMetadataHelpers
 {
 	PCG_API bool HasSameRoot(const UPCGMetadata* Metadata1, const UPCGMetadata* Metadata2);
-	PCG_API bool HasSameRoot(const FPCGMetadataLayer* Metadata1, const FPCGMetadataLayer* Metadata2);
+	PCG_API bool HasSameRoot(const FPCGMetadataDomain* Metadata1, const FPCGMetadataDomain* Metadata2);
 	PCG_API const UPCGMetadata* GetParentMetadata(const UPCGMetadata* Metadata);
-	PCG_API const FPCGMetadataLayer* GetParentMetadata(const FPCGMetadataLayer* Metadata);
+	PCG_API const FPCGMetadataDomain* GetParentMetadata(const FPCGMetadataDomain* Metadata);
 
 	// Utility function to break circular dependency for UPCGMetadata deprecation in PCGMetadataAttributeTpl.
-	PCG_API FPCGMetadataLayer* GetDefaultMetadataLayer(UPCGMetadata* InMetadata);
+	PCG_API FPCGMetadataDomain* GetDefaultMetadataDomain(UPCGMetadata* InMetadata);
 
 	// Helpers functions to cast in spatial/param and return metadata. Nullptr if data doesn't have metadata
 	PCG_API const UPCGMetadata* GetConstMetadata(const UPCGData* InData);
@@ -66,7 +66,7 @@ namespace PCGMetadataHelpers
 	/** Copy all the attributes coming from Source Data into Target Data. */
 	PCG_API bool CopyAllAttributes(const UPCGData* SourceData, UPCGData* TargetData, FPCGContext* OptionalContext = nullptr);
 
-	PCG_API void SetPointAttributes(FPCGMetadataLayer* InOutMetadata, const TArrayView<const FPCGPoint>& InPoints, const FPCGMetadataLayer* InMetadata, const TArrayView<FPCGPoint>& OutPoints, FPCGContext* OptionalContext = nullptr);
+	PCG_API void SetPointAttributes(FPCGMetadataDomain* InOutMetadata, const TArrayView<const FPCGPoint>& InPoints, const FPCGMetadataDomain* InMetadata, const TArrayView<FPCGPoint>& OutPoints, FPCGContext* OptionalContext = nullptr);
 
-	PCG_API void ComputePointWeightedAttribute(FPCGMetadataLayer* InOutMetadata, FPCGPoint& OutPoint, const TArrayView<TPair<const FPCGPoint*, float>>& InWeightedPoints, const FPCGMetadataLayer* InMetadata);
+	PCG_API void ComputePointWeightedAttribute(FPCGMetadataDomain* InOutMetadata, FPCGPoint& OutPoint, const TArrayView<TPair<const FPCGPoint*, float>>& InWeightedPoints, const FPCGMetadataDomain* InMetadata);
 }
diff --git a/Engine/Plugins/PCG/Source/PCG/Public/Metadata/Accessors/PCGAttributeAccessor.h b/Engine/Plugins/PCG/Source/PCG/Public/Metadata/Accessors/PCGAttributeAccessor.h
--- a/Engine/Plugins/PCG/Source/PCG/Public/Metadata/Accessors/PCGAttributeAccessor.h
+++ b/Engine/Plugins/PCG/Source/PCG/Public/Metadata/Accessors/PCGAttributeAccessor.h
@@ -6,7 +6,7 @@
 
 #include "Metadata/PCGMetadata.h"
 #include "Metadata/PCGMetadataAttributeTpl.h"
-#include "Metadata/PCGMetadataLayer.h"
+#include "Metadata/PCGMetadataDomain.h"
 
 /**
 * Templated accessor class for attributes. Will wrap around a typed attribute.
@@ -20,30 +20,30 @@ public:
 	using Super = IPCGAttributeAccessorT<FPCGAttributeAccessor<T>>;
 
 	// Can't write if metadata is null
-	FPCGAttributeAccessor(FPCGMetadataAttribute<T>* InAttribute, FPCGMetadataLayer* InMetadataLayer, bool bForceReadOnly = false)
-		: Super(/*bInReadOnly=*/ InMetadataLayer == nullptr || bForceReadOnly)
+	FPCGAttributeAccessor(FPCGMetadataAttribute<T>* InAttribute, FPCGMetadataDomain* InMetadataDomain, bool bForceReadOnly = false)
+		: Super(/*bInReadOnly=*/ InMetadataDomain == nullptr || bForceReadOnly)
 		, Attribute(InAttribute)
-		, MetadataLayer(InMetadataLayer)
+		, MetadataDomain(InMetadataDomain)
 	{
-		check(InAttribute && (!InMetadataLayer || InAttribute->GetMetadataLayer() == InMetadataLayer));
+		check(InAttribute && (!InMetadataDomain || InAttribute->GetMetadataDomain() == InMetadataDomain));
 	}
 
-	FPCGAttributeAccessor(const FPCGMetadataAttribute<T>* InAttribute, const FPCGMetadataLayer* InMetadataLayer, bool bForceReadOnly = false)
+	FPCGAttributeAccessor(const FPCGMetadataAttribute<T>* InAttribute, const FPCGMetadataDomain* InMetadataDomain, bool bForceReadOnly = false)
 		: Super(/*bInReadOnly=*/ true)
 		, Attribute(const_cast<FPCGMetadataAttribute<T>*>(InAttribute))
-		, MetadataLayer(const_cast<FPCGMetadataLayer*>(InMetadataLayer))
+		, MetadataDomain(const_cast<FPCGMetadataDomain*>(InMetadataDomain))
 	{
-		check(InAttribute && (!InMetadataLayer || InAttribute->GetMetadataLayer() == InMetadataLayer));
+		check(InAttribute && (!InMetadataDomain || InAttribute->GetMetadataDomain() == InMetadataDomain));
 	}
 
-	UE_DEPRECATED(5.6, "Use the FPCGMetadataLayer version")
+	UE_DEPRECATED(5.6, "Use the FPCGMetadataDomain version")
 	FPCGAttributeAccessor(FPCGMetadataAttribute<T>* InAttribute, UPCGMetadata* InMetadata, bool bForceReadOnly = false)
-		: FPCGAttributeAccessor(InAttribute, InMetadata->GetDefaultMetadataLayer(), bForceReadOnly)
+		: FPCGAttributeAccessor(InAttribute, InMetadata->GetDefaultMetadataDomain(), bForceReadOnly)
 	{}
 
-	UE_DEPRECATED(5.6, "Use the FPCGMetadataLayer version")
+	UE_DEPRECATED(5.6, "Use the FPCGMetadataDomain version")
 	FPCGAttributeAccessor(FPCGMetadataAttribute<T>* InAttribute, const UPCGMetadata* InMetadata, bool bForceReadOnly = false)
-		: FPCGAttributeAccessor(InAttribute, InMetadata->GetConstDefaultMetadataLayer(), bForceReadOnly)
+		: FPCGAttributeAccessor(InAttribute, InMetadata->GetConstDefaultMetadataDomain(), bForceReadOnly)
 	{}
 
 	bool GetRangeImpl(TArrayView<T> OutValues, int32 Index, const IPCGAttributeAccessorKeys& Keys) const
@@ -140,7 +140,7 @@ public:
 		for (int EntryIndex = 0; EntryIndex < EntryKeysView.Num(); ++EntryIndex)
 		{
 			PCGMetadataEntryKey& EntryKey = *EntryKeysView[EntryIndex];
-			if (EntryKey == PCGInvalidEntryKey || (EntryKey < MetadataLayer->GetItemKeyCountForParent() && !bCanReuseEntryKeys))
+			if (EntryKey == PCGInvalidEntryKey || (EntryKey < MetadataDomain->GetItemKeyCountForParent() && !bCanReuseEntryKeys))
 			{
 				EntriesToSet.Add(&EntryKey);
 			}
@@ -149,7 +149,7 @@ public:
 		if (!EntriesToSet.IsEmpty())
 		{
 			TRACE_CPUPROFILER_EVENT_SCOPE(FPCGAttributeAccessor::Prepare::AddEntriesInPlace);
-			MetadataLayer->AddEntriesInPlace(EntriesToSet);
+			MetadataDomain->AddEntriesInPlace(EntriesToSet);
 		}
 
 		if (bPreallocateValues && ensure(!bWasPrepared))
@@ -177,6 +177,6 @@ public:
 
 private:
 	FPCGMetadataAttribute<T>* Attribute = nullptr;
-	FPCGMetadataLayer* MetadataLayer = nullptr;
+	FPCGMetadataDomain* MetadataDomain = nullptr;
 	bool bWasPrepared = false;
 };
diff --git a/Engine/Plugins/PCG/Source/PCG/Public/Metadata/Accessors/PCGAttributeAccessorHelpers.h b/Engine/Plugins/PCG/Source/PCG/Public/Metadata/Accessors/PCGAttributeAccessorHelpers.h
--- a/Engine/Plugins/PCG/Source/PCG/Public/Metadata/Accessors/PCGAttributeAccessorHelpers.h
+++ b/Engine/Plugins/PCG/Source/PCG/Public/Metadata/Accessors/PCGAttributeAccessorHelpers.h
@@ -68,7 +68,7 @@ namespace PCGAttributeAccessorHelpers
 	* Creates a const accessor to an attribute without requiring a selector.
 	*/
 	PCG_API TUniquePtr<const IPCGAttributeAccessor> CreateConstAccessor(const FPCGMetadataAttributeBase* InAttribute, const UPCGMetadata* InMetadata, bool bQuiet = false);
-	PCG_API TUniquePtr<const IPCGAttributeAccessor> CreateConstAccessor(const FPCGMetadataAttributeBase* InAttribute, const FPCGMetadataLayer* InMetadata, bool bQuiet = false);
+	PCG_API TUniquePtr<const IPCGAttributeAccessor> CreateConstAccessor(const FPCGMetadataAttributeBase* InAttribute, const FPCGMetadataDomain* InMetadata, bool bQuiet = false);
 
 	/**
 	* Creates an accessor to the property or attribute pointed at by the InSelector.
@@ -91,7 +91,7 @@ namespace PCGAttributeAccessorHelpers
 	* Creates an accessor to an attribute without requiring a selector.
 	*/
 	PCG_API TUniquePtr<IPCGAttributeAccessor> CreateAccessor(FPCGMetadataAttributeBase* InAttribute, UPCGMetadata* InMetadata, bool bQuiet = false);
-	PCG_API TUniquePtr<IPCGAttributeAccessor> CreateAccessor(FPCGMetadataAttributeBase* InAttribute, FPCGMetadataLayer* InMetadata, bool bQuiet = false);
+	PCG_API TUniquePtr<IPCGAttributeAccessor> CreateAccessor(FPCGMetadataAttributeBase* InAttribute, FPCGMetadataDomain* InMetadata, bool bQuiet = false);
 
 	PCG_API TUniquePtr<const IPCGAttributeAccessorKeys> CreateConstKeys(const UPCGData* InData, const FPCGAttributePropertySelector& InSelector);
 	PCG_API TUniquePtr<IPCGAttributeAccessorKeys> CreateKeys(UPCGData* InData, const FPCGAttributePropertySelector& InSelector);
diff --git a/Engine/Plugins/PCG/Source/PCG/Public/Metadata/Accessors/PCGAttributeAccessorKeys.h b/Engine/Plugins/PCG/Source/PCG/Public/Metadata/Accessors/PCGAttributeAccessorKeys.h
--- a/Engine/Plugins/PCG/Source/PCG/Public/Metadata/Accessors/PCGAttributeAccessorKeys.h
+++ b/Engine/Plugins/PCG/Source/PCG/Public/Metadata/Accessors/PCGAttributeAccessorKeys.h
@@ -5,7 +5,7 @@
 #include "Metadata/PCGMetadataCommon.h"
 
 class FPCGMetadataAttributeBase;
-class FPCGMetadataLayer;
+class FPCGMetadataDomain;
 struct FPCGPoint;
 class UPCGMetadata;
 class UPCGSplineData;
@@ -183,8 +183,8 @@ public:
 	explicit FPCGAttributeAccessorKeysEntries(const UPCGMetadata* Metadata, bool bAddDefaultValueIfEmpty = false);
 	explicit FPCGAttributeAccessorKeysEntries(UPCGMetadata* Metadata, bool bAddDefaultValueIfEmpty = true);
 
-	explicit FPCGAttributeAccessorKeysEntries(const FPCGMetadataLayer* Metadata, bool bAddDefaultValueIfEmpty = false);
-	explicit FPCGAttributeAccessorKeysEntries(FPCGMetadataLayer* Metadata, bool bAddDefaultValueIfEmpty = true);
+	explicit FPCGAttributeAccessorKeysEntries(const FPCGMetadataDomain* Metadata, bool bAddDefaultValueIfEmpty = false);
+	explicit FPCGAttributeAccessorKeysEntries(FPCGMetadataDomain* Metadata, bool bAddDefaultValueIfEmpty = true);
 
 	virtual int32 GetNum() const override { return Entries.Num(); }
 
@@ -192,7 +192,7 @@ protected:
 	virtual bool GetMetadataEntryKeys(int32 InStart, TArrayView<PCGMetadataEntryKey*> OutEntryKeys) override;
 	virtual bool GetMetadataEntryKeys(int32 InStart, TArrayView<const PCGMetadataEntryKey*> OutEntryKeys) const override;
 
-	void InitializeFromMetadata(const FPCGMetadataLayer* Metadata, bool bAddDefaultValueIfEmpty);
+	void InitializeFromMetadata(const FPCGMetadataDomain* Metadata, bool bAddDefaultValueIfEmpty);
 
 	TArrayView<PCGMetadataEntryKey> Entries;
 	TArray<PCGMetadataEntryKey> ExtractedEntries;
diff --git a/Engine/Plugins/PCG/Source/PCG/Public/Metadata/PCGAttributePropertySelector.h b/Engine/Plugins/PCG/Source/PCG/Public/Metadata/PCGAttributePropertySelector.h
--- a/Engine/Plugins/PCG/Source/PCG/Public/Metadata/PCGAttributePropertySelector.h
+++ b/Engine/Plugins/PCG/Source/PCG/Public/Metadata/PCGAttributePropertySelector.h
@@ -61,7 +61,7 @@ public:
 	bool SetExtraProperty(EPCGExtraProperties InExtraProperty, bool bResetExtraNames = true);
 	bool SetAttributeName(FName InAttributeName, bool bResetExtraNames = true);
 	bool SetPropertyName(FName InPropertyName, bool bResetExtraNames = true);
-	bool SetLayerName(FName InLayerName, bool bResetExtraNames = true);
+	bool SetDomainName(FName InDomainName, bool bResetExtraNames = true);
 
 	// Getters
 	EPCGAttributePropertySelection GetSelection() const { return Selection; }
@@ -71,7 +71,7 @@ public:
 	EPCGPointProperties GetPointProperty() const;
 	EPCGExtraProperties GetExtraProperty() const { return ExtraProperty; }
 	FName GetPropertyName() const { return PropertyName; }
-	FName GetLayerName() const { return LayerName; }
+	FName GetDomainName() const { return DomainName; }
 
 	// Convenient function to know if it is a basic attribute (attribute and no extra names)
 	bool IsBasicAttribute() const;
@@ -82,8 +82,8 @@ public:
 	// Returns qualified attribute/property name with the accessors.
 	FString ToString() const;
 
-	// Returns attribute/property name only, with optional '@' qualifier for layers.
-	FString GetLayerString(bool bAddLeadingQualifier = true) const;
+	// Returns attribute/property name only, with optional '@' qualifier for domains.
+	FString GetDomainString(bool bAddLeadingQualifier = true) const;
 	
 	// Returns attribute/property name only, with optional '$' qualifier for properties and '.'.
 	FString GetAttributePropertyString(bool bAddPropertyQualifier = true) const;
@@ -111,44 +111,44 @@ public:
 
 	// Convenience templated static constructors
 	template <typename T = FPCGAttributePropertySelector>
-	static T CreateAttributeSelector(const FName AttributeName, const FName LayerName = NAME_None, const TArrayView<const FString>& InExtraNames = {})
+	static T CreateAttributeSelector(const FName AttributeName, const FName DomainName = NAME_None, const TArrayView<const FString>& InExtraNames = {})
 	{
 		static_assert(std::is_base_of_v<FPCGAttributePropertySelector, T>, "The type must be of base class 'FPCGAttributePropertySelector'");
 		T Selector;
-		Selector.SetLayerName(LayerName);
+		Selector.SetDomainName(DomainName);
 		Selector.SetAttributeName(AttributeName);
 		Selector.GetExtraNamesMutable() = InExtraNames;
 		return Selector;
 	}
 	
 	template <typename T = FPCGAttributePropertySelector>
-	static T CreatePropertySelector(FName PropertyName, const FName LayerName = NAME_None, const TArrayView<const FString>& InExtraNames = {})
+	static T CreatePropertySelector(FName PropertyName, const FName DomainName = NAME_None, const TArrayView<const FString>& InExtraNames = {})
 	{
 		static_assert(std::is_base_of_v<FPCGAttributePropertySelector, T>, "The type must be of base class 'FPCGAttributePropertySelector'");
 		T Selector;
-		Selector.SetLayerName(LayerName);
+		Selector.SetDomainName(DomainName);
 		Selector.SetPropertyName(PropertyName);
 		Selector.GetExtraNamesMutable() = InExtraNames;
 		return Selector;
 	}
 	
 	template <typename T = FPCGAttributePropertySelector>
-	static T CreatePointPropertySelector(EPCGPointProperties PointProperty, const FName LayerName = NAME_None, const TArrayView<const FString>& InExtraNames = {})
+	static T CreatePointPropertySelector(EPCGPointProperties PointProperty, const FName DomainName = NAME_None, const TArrayView<const FString>& InExtraNames = {})
 	{
 		static_assert(std::is_base_of_v<FPCGAttributePropertySelector, T>, "The type must be of base class 'FPCGAttributePropertySelector'");
 		T Selector;
-		Selector.SetLayerName(LayerName);
+		Selector.SetDomainName(DomainName);
 		Selector.SetPointProperty(PointProperty);
 		Selector.GetExtraNamesMutable() = InExtraNames;
 		return Selector;
 	}
 
 	template <typename T = FPCGAttributePropertySelector>
-	static T CreateExtraPropertySelector(EPCGExtraProperties ExtraProperty, const FName LayerName = NAME_None, const TArrayView<const FString>& InExtraNames = {})
+	static T CreateExtraPropertySelector(EPCGExtraProperties ExtraProperty, const FName DomainName = NAME_None, const TArrayView<const FString>& InExtraNames = {})
 	{
 		static_assert(std::is_base_of_v<FPCGAttributePropertySelector, T>, "The type must be of base class 'FPCGAttributePropertySelector'");
 		T Selector;
-		Selector.SetLayerName(LayerName);
+		Selector.SetDomainName(DomainName);
 		Selector.SetExtraProperty(ExtraProperty);
 		Selector.GetExtraNamesMutable() = InExtraNames;
 		return Selector;
@@ -197,7 +197,7 @@ protected:
 	EPCGAttributePropertySelection Selection = EPCGAttributePropertySelection::Attribute;
 
 	UPROPERTY()
-	FName LayerName = NAME_None;
+	FName DomainName = NAME_None;
 
 	UPROPERTY()
 	FName AttributeName = NAME_None;
@@ -314,7 +314,7 @@ public:
 	static bool SetExtraProperty(UPARAM(ref) FPCGAttributePropertySelector& Selector, EPCGExtraProperties InExtraProperty, bool bResetExtraNames = true);
 
 	UFUNCTION(BlueprintCallable, Category = "PCG|PCGAttributePropertySelector", meta = (ScriptMethod))
-	static bool SetLayerName(UPARAM(ref) FPCGAttributePropertySelector& Selector, FName InLayerName, bool bResetExtraNames = true);
+	static bool SetDomainName(UPARAM(ref) FPCGAttributePropertySelector& Selector, FName InDomainName, bool bResetExtraNames = true);
 
 	UFUNCTION(BlueprintCallable, Category = "PCG|PCGAttributePropertySelector", meta = (ScriptMethod))
 	static EPCGAttributePropertySelection GetSelection(UPARAM(const, ref) const FPCGAttributePropertySelector& Selector);
@@ -329,7 +329,7 @@ public:
 	static FName GetPropertyName(UPARAM(const, ref) const FPCGAttributePropertySelector& Selector);
 
 	UFUNCTION(BlueprintCallable, Category = "PCG|PCGAttributePropertySelector", meta = (ScriptMethod))
-	static FName GetLayerName(UPARAM(const, ref) const FPCGAttributePropertySelector& Selector);
+	static FName GetDomainName(UPARAM(const, ref) const FPCGAttributePropertySelector& Selector);
 
 	UFUNCTION(BlueprintCallable, Category = "PCG|PCGAttributePropertySelector", meta = (ScriptMethod))
 	static EPCGExtraProperties GetExtraProperty(UPARAM(const, ref) const FPCGAttributePropertySelector& Selector);
diff --git a/Engine/Plugins/PCG/Source/PCG/Public/Metadata/PCGMetadata.h b/Engine/Plugins/PCG/Source/PCG/Public/Metadata/PCGMetadata.h
--- a/Engine/Plugins/PCG/Source/PCG/Public/Metadata/PCGMetadata.h
+++ b/Engine/Plugins/PCG/Source/PCG/Public/Metadata/PCGMetadata.h
@@ -6,7 +6,7 @@
 #include "PCGMetadataAttributeTraits.h"
 #include "PCGMetadataCommon.h" // IWYU pragma: keep
 #include "Metadata/PCGMetadataAttributeTpl.h"
-#include "Metadata/PCGMetadataLayer.h"
+#include "Metadata/PCGMetadataDomain.h"
 
 #include "Misc/SpinLock.h"
 
@@ -18,7 +18,7 @@ struct FPCGContext;
 UCLASS(BlueprintType)
 class PCG_API UPCGMetadata : public UObject
 {
-	friend FPCGMetadataLayer;
+	friend FPCGMetadataDomain;
 	
 	GENERATED_BODY()
 
@@ -29,14 +29,14 @@ public:
 	virtual void Serialize(FArchive& InArchive) override;
 	//~ End UObject interface
 
-	/** To be called by the outer pcg data to initialize the layers supported for this data. Can mark this one as default, will override previously set default. */
-	void SetupLayer(FPCGMetadataLayerID LayerID, bool bIsDefault);
+	/** To be called by the outer pcg data to initialize the domains supported for this data. Can mark this one as default, will override previously set default. */
+	void SetupDomain(FPCGMetadataDomainID DomainID, bool bIsDefault);
 
 	/** If we have floating metadata, setup it to match a given data type. */
 	template <typename DataType, std::enable_if_t<std::is_base_of_v<UPCGData, DataType>, int> = 0>
-	void SetupLayersFromPCGDataType() { SetupFromPCGDataType(DataType::StaticClass()); }
+	void SetupDomainsFromPCGDataType() { SetupFromPCGDataType(DataType::StaticClass()); }
 
-	void SetupLayersFromPCGDataType(const TSubclassOf<UPCGData>& PCGDataType);
+	void SetupDomainsFromPCGDataType(const TSubclassOf<UPCGData>& PCGDataType);
 
 	/** Initializes the metadata from a parent metadata, if any (can be null). Copies attributes and values. */
 	UFUNCTION(BlueprintCallable, Category = "PCG|Metadata")
@@ -391,36 +391,36 @@ public:
 	/** Computes Crc from all attributes & keys from outer's data. */
 	void AddToCrc(FArchiveCrc32& Ar, bool bFullDataCrc) const;
 
-	/** Expected to be used when metadata layers are already setup. */
-	FPCGMetadataLayer* GetMetadataLayer(const FPCGMetadataLayerID& InMetadataLayerID);
+	/** Expected to be used when metadata domains are already setup. */
+	FPCGMetadataDomain* GetMetadataDomain(const FPCGMetadataDomainID& InMetadataDomainID);
 	
-	/** Expected to be used when metadata layer are already setup. */
-	const FPCGMetadataLayer* GetConstMetadataLayer(const FPCGMetadataLayerID& InMetadataLayerID) const;
+	/** Expected to be used when metadata domain are already setup. */
+	const FPCGMetadataDomain* GetConstMetadataDomain(const FPCGMetadataDomainID& InMetadataDomainID) const;
 
-	FPCGMetadataLayer* GetDefaultMetadataLayer() { return GetMetadataLayer(PCGMetadataLayerID::Default); }
-	const FPCGMetadataLayer* GetConstDefaultMetadataLayer() const { return GetConstMetadataLayer(PCGMetadataLayerID::Default); }
+	FPCGMetadataDomain* GetDefaultMetadataDomain() { return GetMetadataDomain(PCGMetadataDomainID::Default); }
+	const FPCGMetadataDomain* GetConstDefaultMetadataDomain() const { return GetConstMetadataDomain(PCGMetadataDomainID::Default); }
 
 private:
-	// Utility functions to forward calls to FPCGMetadataLayer.
+	// Utility functions to forward calls to FPCGMetadataDomain.
 	template <typename Func, typename ...Args>
-	decltype(auto) WithMetadataLayer(const FPCGMetadataLayerID& InMetadataLayerID, Func InFunc, Args&& ...InArgs);
+	decltype(auto) WithMetadataDomain(const FPCGMetadataDomainID& InMetadataDomainID, Func InFunc, Args&& ...InArgs);
 	
 	template <typename Func>
-	decltype(auto) WithMetadataLayer_Lambda(const FPCGMetadataLayerID& InMetadataLayerID, Func InFunc);
+	decltype(auto) WithMetadataDomain_Lambda(const FPCGMetadataDomainID& InMetadataDomainID, Func InFunc);
 
 	template <typename Func, typename ...Args>
-	decltype(auto) WithConstMetadataLayer(const FPCGMetadataLayerID& InMetadataLayerID, Func InFunc, Args&& ...InArgs) const;
+	decltype(auto) WithConstMetadataDomain(const FPCGMetadataDomainID& InMetadataDomainID, Func InFunc, Args&& ...InArgs) const;
 
 	template <typename Func>
-	void ForEachValidUniqueConstMetadataLayer(Func InFunc) const;
+	void ForEachValidUniqueConstMetadataDomain(Func InFunc) const;
 
 	template <typename Func>
-	void ForEachValidUniqueMetadataLayer(Func InFunc);
+	void ForEachValidUniqueMetadataDomain(Func InFunc);
 
 protected:
-	FPCGMetadataLayer* FindOrCreateMetadataLayer(const FPCGMetadataLayerID& InMetadataLayerID);
+	FPCGMetadataDomain* FindOrCreateMetadataDomain(const FPCGMetadataDomainID& InMetadataDomainID);
 	
-	FPCGMetadataLayer* CreateMetadataLayer_Unsafe(const FPCGMetadataLayerID& InMetadataLayerID);
+	FPCGMetadataDomain* CreateMetadataDomain_Unsafe(const FPCGMetadataDomainID& InMetadataDomainID);
 	
 	FPCGMetadataAttributeBase* CopyAttribute(FPCGAttributeIdentifier AttributeToCopy, FName NewAttributeName, bool bKeepParent, bool bCopyEntries, bool bCopyValues);
 
@@ -428,7 +428,7 @@ protected:
 
 	void SetLastCachedSelectorOnOwner(FName AttributeName);
 
-	void SetupLayersFromOtherMetadataIfNeeded(const UPCGMetadata* OtherMetadata);
+	void SetupDomainsFromOtherMetadataIfNeeded(const UPCGMetadata* OtherMetadata);
 
 	UPROPERTY()
 	TObjectPtr<const UPCGMetadata> Parent;
@@ -438,15 +438,15 @@ protected:
 	UPROPERTY()
 	TSet<TWeakObjectPtr<const UPCGMetadata>> OtherParents;
 
-	FPCGMetadataLayerID DefaultLayer = PCGMetadataLayerID::Default;
-	TMap<FPCGMetadataLayerID, TSharedPtr<FPCGMetadataLayer>> MetadataLayers;
-	UE::FSpinLock MetadataLayersSpinLock;
+	FPCGMetadataDomainID DefaultDomain = PCGMetadataDomainID::Default;
+	TMap<FPCGMetadataDomainID, TSharedPtr<FPCGMetadataDomain>> MetadataDomains;
+	UE::FSpinLock MetadataDomainsSpinLock;
 };
 
 template<typename T>
 FPCGMetadataAttribute<T>* UPCGMetadata::CreateAttribute(FPCGAttributeIdentifier AttributeName, const T& DefaultValue, bool bAllowsInterpolation, bool bOverrideParent)
 {
-	FPCGMetadataLayer* FoundSubMetadata = FindOrCreateMetadataLayer(AttributeName.MetadataLayer);
+	FPCGMetadataDomain* FoundSubMetadata = FindOrCreateMetadataDomain(AttributeName.MetadataDomain);
 	check(FoundSubMetadata);
 	return FoundSubMetadata->CreateAttribute<T>(AttributeName.Name, DefaultValue, bAllowsInterpolation, bOverrideParent);
 }
@@ -454,14 +454,14 @@ FPCGMetadataAttribute<T>* UPCGMetadata::CreateAttribute(FPCGAttributeIdentifier
 template<typename T>
 FPCGMetadataAttribute<T>* UPCGMetadata::GetMutableTypedAttribute_Unsafe(FPCGAttributeIdentifier AttributeName)
 {
-	FPCGMetadataLayer* FoundSubMetadata = GetMetadataLayer(AttributeName.MetadataLayer);
+	FPCGMetadataDomain* FoundSubMetadata = GetMetadataDomain(AttributeName.MetadataDomain);
 	return FoundSubMetadata ? FoundSubMetadata->GetMutableTypedAttribute_Unsafe<T>(AttributeName.Name) : nullptr;
 }
 
 template<typename T>
 FPCGMetadataAttribute<T>* UPCGMetadata::FindOrCreateAttribute(FPCGAttributeIdentifier AttributeName, const T& DefaultValue, bool bAllowsInterpolation, bool bOverrideParent, bool bOverwriteIfTypeMismatch)
 {
-	FPCGMetadataLayer* FoundSubMetadata = FindOrCreateMetadataLayer(AttributeName.MetadataLayer);
+	FPCGMetadataDomain* FoundSubMetadata = FindOrCreateMetadataDomain(AttributeName.MetadataDomain);
 	check(FoundSubMetadata);
 	return FoundSubMetadata->FindOrCreateAttribute<T>(AttributeName.Name, DefaultValue, bAllowsInterpolation, bOverrideParent, bOverwriteIfTypeMismatch);
 }
diff --git a/Engine/Plugins/PCG/Source/PCG/Public/Metadata/PCGMetadataAttribute.h b/Engine/Plugins/PCG/Source/PCG/Public/Metadata/PCGMetadataAttribute.h
--- a/Engine/Plugins/PCG/Source/PCG/Public/Metadata/PCGMetadataAttribute.h
+++ b/Engine/Plugins/PCG/Source/PCG/Public/Metadata/PCGMetadataAttribute.h
@@ -6,7 +6,7 @@
 
 #include "UObject/ObjectPtr.h" // IWYU pragma: keep
 
-class FPCGMetadataLayer;
+class FPCGMetadataDomain;
 class UPCGMetadata;
 
 namespace PCGMetadataAttributeConstants
@@ -21,16 +21,16 @@ class PCG_API FPCGMetadataAttributeBase
 {
 public:
 	FPCGMetadataAttributeBase() = default;
-	FPCGMetadataAttributeBase(FPCGMetadataLayer* InMetadata, FName InName, const FPCGMetadataAttributeBase* InParent, bool bInAllowsInterpolation);
+	FPCGMetadataAttributeBase(FPCGMetadataDomain* InMetadata, FName InName, const FPCGMetadataAttributeBase* InParent, bool bInAllowsInterpolation);
 
-	UE_DEPRECATED(5.6, "Use the version with FPCGMetadataLayer")
+	UE_DEPRECATED(5.6, "Use the version with FPCGMetadataDomain")
 	FPCGMetadataAttributeBase(UPCGMetadata* InMetadata, FName InName, const FPCGMetadataAttributeBase* InParent, bool bInAllowsInterpolation);
 
 	virtual ~FPCGMetadataAttributeBase() = default;
 	
-	virtual void Serialize(FPCGMetadataLayer* InMetadata, FArchive& InArchive);
+	virtual void Serialize(FPCGMetadataDomain* InMetadata, FArchive& InArchive);
 
-	UE_DEPRECATED(5.6, "Use the version with FPCGMetadataLayer")
+	UE_DEPRECATED(5.6, "Use the version with FPCGMetadataDomain")
 	virtual void Serialize(UPCGMetadata* InMetadata, FArchive& InArchive);
 
 	/** Unparents current attribute by flattening the values, entries, etc. */
@@ -42,13 +42,13 @@ public:
 	virtual void Reset() = 0;
 	
 	const UPCGMetadata* GetMetadata() const;
-	const FPCGMetadataLayer* GetMetadataLayer() const { return Metadata; }
+	const FPCGMetadataDomain* GetMetadataDomain() const { return Metadata; }
 	
 	int16 GetTypeId() const { return TypeId; }
 
-	virtual FPCGMetadataAttributeBase* Copy(FName NewName, FPCGMetadataLayer* InMetadata, bool bKeepParent, bool bCopyEntries = true, bool bCopyValues = true) const = 0;
+	virtual FPCGMetadataAttributeBase* Copy(FName NewName, FPCGMetadataDomain* InMetadata, bool bKeepParent, bool bCopyEntries = true, bool bCopyValues = true) const = 0;
 
-	UE_DEPRECATED(5.6, "Use the version with FPCGMetadataLayer")
+	UE_DEPRECATED(5.6, "Use the version with FPCGMetadataDomain")
 	virtual FPCGMetadataAttributeBase* Copy(FName NewName, UPCGMetadata* InMetadata, bool bKeepParent, bool bCopyEntries = true, bool bCopyValues = true) const = 0;
 	
 	virtual FPCGMetadataAttributeBase* CopyToAnotherType(int16 Type) const = 0;
@@ -116,7 +116,7 @@ protected:
 	TMap<PCGMetadataEntryKey, PCGMetadataValueKey> EntryToValueKeyMap;
 	mutable FRWLock EntryMapLock;
 
-	FPCGMetadataLayer* Metadata = nullptr;
+	FPCGMetadataDomain* Metadata = nullptr;
 	const FPCGMetadataAttributeBase* Parent = nullptr;
 	int16 TypeId = 0;
 	bool bAllowsInterpolation = false;
diff --git a/Engine/Plugins/PCG/Source/PCG/Public/Metadata/PCGMetadataAttributeTpl.h b/Engine/Plugins/PCG/Source/PCG/Public/Metadata/PCGMetadataAttributeTpl.h
--- a/Engine/Plugins/PCG/Source/PCG/Public/Metadata/PCGMetadataAttributeTpl.h
+++ b/Engine/Plugins/PCG/Source/PCG/Public/Metadata/PCGMetadataAttributeTpl.h
@@ -11,7 +11,7 @@
 #include "Metadata/PCGMetadataAttributeTraits.h"
 #include "Misc/ScopeRWLock.h"
 
-class FPCGMetadataLayer;
+class FPCGMetadataDomain;
 class UPCGMetadata;
 
 template<typename T>
@@ -22,7 +22,7 @@ class FPCGMetadataAttribute : public FPCGMetadataAttributeBase
 
 public:
 	/** TODO: All deprecated versions taking a UPCGMetadata will be added at the end of this declaration. For the review, it will be in place for "clarity" */
-	FPCGMetadataAttribute(FPCGMetadataLayer* InMetadata, FName InName, const FPCGMetadataAttributeBase* InParent, const T& InDefaultValue, bool bInAllowsInterpolation)
+	FPCGMetadataAttribute(FPCGMetadataDomain* InMetadata, FName InName, const FPCGMetadataAttributeBase* InParent, const T& InDefaultValue, bool bInAllowsInterpolation)
 		: FPCGMetadataAttributeBase(InMetadata, InName, InParent, bInAllowsInterpolation)
 		, DefaultValue(InDefaultValue)
 	{
@@ -37,9 +37,9 @@ public:
 		}
 	}
 
-	UE_DEPRECATED(5.6, "Use the version with FPCGMetadataLayer")
+	UE_DEPRECATED(5.6, "Use the version with FPCGMetadataDomain")
 	FPCGMetadataAttribute(UPCGMetadata* InMetadata, FName InName, const FPCGMetadataAttributeBase* InParent, const T& InDefaultValue, bool bInAllowsInterpolation)
-		: FPCGMetadataAttribute(PCGMetadataHelpers::GetDefaultMetadataLayer(InMetadata), InName, InParent, InDefaultValue, bAllowsInterpolation)
+		: FPCGMetadataAttribute(PCGMetadataHelpers::GetDefaultMetadataDomain(InMetadata), InName, InParent, InDefaultValue, bAllowsInterpolation)
 	{}
 
 	// This constructor is used only during serialization
@@ -48,7 +48,7 @@ public:
 		TypeId = PCG::Private::MetadataTypes<T>::Id;
 	}
 
-	virtual void Serialize(FPCGMetadataLayer* InMetadata, FArchive& InArchive) override
+	virtual void Serialize(FPCGMetadataDomain* InMetadata, FArchive& InArchive) override
 	{
 		FPCGMetadataAttributeBase::Serialize(InMetadata, InArchive);
 
@@ -194,24 +194,24 @@ public:
 
 	const FPCGMetadataAttribute* GetParent() const { return static_cast<const FPCGMetadataAttribute*>(Parent); }
 
-	UE_DEPRECATED(5.6, "Use the version with FPCGMetadataLayer")
+	UE_DEPRECATED(5.6, "Use the version with FPCGMetadataDomain")
 	FPCGMetadataAttribute* TypedCopy(FName NewName, UPCGMetadata* InMetadata, bool bKeepParent, bool bCopyEntries = true, bool bCopyValues = true)
 	{
-		return static_cast<FPCGMetadataAttribute*>(TypedCopy(NewName, PCGMetadataHelpers::GetDefaultMetadataLayer(InMetadata), bKeepParent, bCopyEntries, bCopyValues));
+		return static_cast<FPCGMetadataAttribute*>(TypedCopy(NewName, PCGMetadataHelpers::GetDefaultMetadataDomain(InMetadata), bKeepParent, bCopyEntries, bCopyValues));
 	}
 	
-	FPCGMetadataAttribute* TypedCopy(FName NewName, FPCGMetadataLayer* InMetadata, bool bKeepParent, bool bCopyEntries = true, bool bCopyValues = true)
+	FPCGMetadataAttribute* TypedCopy(FName NewName, FPCGMetadataDomain* InMetadata, bool bKeepParent, bool bCopyEntries = true, bool bCopyValues = true)
 	{
 		return static_cast<FPCGMetadataAttribute*>(Copy(NewName, InMetadata, bKeepParent, bCopyEntries, bCopyValues));
 	}
 
-	UE_DEPRECATED(5.6, "Use the version with FPCGMetadataLayer")
+	UE_DEPRECATED(5.6, "Use the version with FPCGMetadataDomain")
 	virtual FPCGMetadataAttribute* Copy(FName NewName, UPCGMetadata* InMetadata, bool bKeepParent, bool bCopyEntries = true, bool bCopyValues = true) const override
 	{
-		return static_cast<FPCGMetadataAttribute*>(Copy(NewName, PCGMetadataHelpers::GetDefaultMetadataLayer(InMetadata), bKeepParent, bCopyEntries, bCopyValues));
+		return static_cast<FPCGMetadataAttribute*>(Copy(NewName, PCGMetadataHelpers::GetDefaultMetadataDomain(InMetadata), bKeepParent, bCopyEntries, bCopyValues));
 	}
 	
-	virtual FPCGMetadataAttributeBase* Copy(FName NewName, FPCGMetadataLayer* InMetadata, bool bKeepParent, bool bCopyEntries = true, bool bCopyValues = true) const override
+	virtual FPCGMetadataAttributeBase* Copy(FName NewName, FPCGMetadataDomain* InMetadata, bool bKeepParent, bool bCopyEntries = true, bool bCopyValues = true) const override
 	{
 		return CopyInternal<T>(NewName, InMetadata, bKeepParent, bCopyEntries, bCopyValues);
 	}
@@ -222,14 +222,14 @@ public:
 	UE_DEPRECATED(5.6, "Do not use anymore, use CopyToAnotherType")
 	FPCGMetadataAttributeBase* CopyInternal(FName NewName, UPCGMetadata* InMetadata, bool bKeepParent, bool bCopyEntries, bool bCopyValues) const
 	{
-		return CopyInternal<U>(NewName, PCGMetadataHelpers::GetDefaultMetadataLayer(InMetadata), bKeepParent, bCopyEntries, bCopyValues);
+		return CopyInternal<U>(NewName, PCGMetadataHelpers::GetDefaultMetadataDomain(InMetadata), bKeepParent, bCopyEntries, bCopyValues);
 	}
 
 protected:
 	// TODO: add enable if only on compatible types, but this has some repercussion on using metadata on types that aren't normally supported.
-	// TODO: The UPCGMetadata version was exposed, but should have not, mark the FPCGMetadataLayer version protected. Will be moved at the end, before commit.
+	// TODO: The UPCGMetadata version was exposed, but should have not, mark the FPCGMetadataDomain version protected. Will be moved at the end, before commit.
 	template<typename U>
-	FPCGMetadataAttributeBase* CopyInternal(FName NewName, FPCGMetadataLayer* InMetadata, bool bKeepParent, bool bCopyEntries, bool bCopyValues) const
+	FPCGMetadataAttributeBase* CopyInternal(FName NewName, FPCGMetadataDomain* InMetadata, bool bKeepParent, bool bCopyEntries, bool bCopyValues) const
 	{
 		// If we copy an attribute where we don't want to keep the parent, while copying entries and/or values, we'll lose data.
 		// In that case, we will copy all the data from this attribute and all its ancestors.
@@ -255,10 +255,10 @@ protected:
 		TArray<const FPCGMetadataAttribute<T>*, TInlineAllocator<2>> Parents = { this };
 		if (!bKeepParent && (bCopyEntries || bCopyValues))
 		{
-			const FPCGMetadataLayer* CurrentMetadata = Metadata;
+			const FPCGMetadataDomain* CurrentMetadata = Metadata;
 			const FPCGMetadataAttribute<T>* Current = this;
 
-			const FPCGMetadataLayer* ParentMetadata = PCGMetadataHelpers::GetParentMetadata(CurrentMetadata);
+			const FPCGMetadataDomain* ParentMetadata = PCGMetadataHelpers::GetParentMetadata(CurrentMetadata);
 			while(ParentMetadata && Current->Parent)
 			{
 				CurrentMetadata = ParentMetadata;
diff --git a/Engine/Plugins/PCG/Source/PCG/Public/Metadata/PCGMetadataCommon.h b/Engine/Plugins/PCG/Source/PCG/Public/Metadata/PCGMetadataCommon.h
--- a/Engine/Plugins/PCG/Source/PCG/Public/Metadata/PCGMetadataCommon.h
+++ b/Engine/Plugins/PCG/Source/PCG/Public/Metadata/PCGMetadataCommon.h
@@ -44,62 +44,62 @@ enum class EPCGMetadataFilterMode : uint8
 };
 
 UENUM()
-enum class EPCGMetadataLayerFlag : uint8
+enum class EPCGMetadataDomainFlag : uint8
 {
-	/** Depends on the data. Should map to the same concept before multi-level metadata. */
+	/** Depends on the data. Should map to the same concept before multi-domain metadata. */
 	Default = 0,
 	
-	/** Metadata data at the data level. */
+	/** Metadata at the data domain. */
 	Data = 1,
 
 	/** Metadata on elements like points on point data and entries on param data. */
 	Elements = 2,
 
-	/** For data that can have more levels. */
+	/** For data that can have more domains. */
 	Custom = 255
 };
 
 USTRUCT(BlueprintType)
-struct FPCGMetadataLayerID
+struct FPCGMetadataDomainID
 {
 	GENERATED_BODY();
 
-	FPCGMetadataLayerID() = default;
+	FPCGMetadataDomainID() = default;
 	
-	FPCGMetadataLayerID(EPCGMetadataLayerFlag InFlag, int32 InCustomFlag = -1, FName InDebugName = NAME_None)
+	FPCGMetadataDomainID(EPCGMetadataDomainFlag InFlag, int32 InCustomFlag = -1, FName InDebugName = NAME_None)
 		: Flag(InFlag)
 		, CustomFlag(InCustomFlag)
 		, DebugName(InDebugName)
 	{
-		check(CustomFlag == -1 || InFlag == EPCGMetadataLayerFlag::Custom);
+		check(CustomFlag == -1 || InFlag == EPCGMetadataDomainFlag::Custom);
 	}
 
-	bool operator==(const FPCGMetadataLayerID& Other) const
+	bool operator==(const FPCGMetadataDomainID& Other) const
 	{
 		return Flag == Other.Flag && CustomFlag == Other.CustomFlag;
 	}
 
-	bool operator<(const FPCGMetadataLayerID& Other) const
+	bool operator<(const FPCGMetadataDomainID& Other) const
 	{
 		return Flag < Other.Flag || (Flag == Other.Flag && CustomFlag < Other.CustomFlag);
 	}
 
-	friend uint32 GetTypeHash(const FPCGMetadataLayerID& Item)
+	friend uint32 GetTypeHash(const FPCGMetadataDomainID& Item)
 	{
 		return HashCombine(static_cast<uint32>(Item.Flag), static_cast<uint32>(Item.CustomFlag));
 	}
 
-	friend FArchive& operator<<(FArchive& Ar, FPCGMetadataLayerID& Item)
+	friend FArchive& operator<<(FArchive& Ar, FPCGMetadataDomainID& Item)
 	{
 		Ar << Item.Flag;
 		Ar << Item.CustomFlag;
 		return Ar;
 	}
 
-	bool IsDefault() const { return Flag == EPCGMetadataLayerFlag::Default; }
+	bool IsDefault() const { return Flag == EPCGMetadataDomainFlag::Default; }
 
 	UPROPERTY(EditAnywhere, BlueprintReadWrite, Category = "")
-	EPCGMetadataLayerFlag Flag = EPCGMetadataLayerFlag::Default;
+	EPCGMetadataDomainFlag Flag = EPCGMetadataDomainFlag::Default;
 	
 	UPROPERTY(EditAnywhere, BlueprintReadWrite, Category = "")
 	int32 CustomFlag = -1;
@@ -108,11 +108,11 @@ struct FPCGMetadataLayerID
 	FName DebugName = NAME_None;
 };
 
-namespace PCGMetadataLayerID
+namespace PCGMetadataDomainID
 {
-	static const FPCGMetadataLayerID Default{EPCGMetadataLayerFlag::Default, -1, TEXT("Default")};
-	static const FPCGMetadataLayerID Elements{EPCGMetadataLayerFlag::Elements, -1, TEXT("Elements")};
-	static const FPCGMetadataLayerID Data{EPCGMetadataLayerFlag::Data, -1, TEXT("Data")};
+	static const FPCGMetadataDomainID Default{EPCGMetadataDomainFlag::Default, -1, TEXT("Default")};
+	static const FPCGMetadataDomainID Elements{EPCGMetadataDomainFlag::Elements, -1, TEXT("Elements")};
+	static const FPCGMetadataDomainID Data{EPCGMetadataDomainFlag::Data, -1, TEXT("Data")};
 }
 
 USTRUCT(BlueprintType)
@@ -124,9 +124,9 @@ struct FPCGAttributeIdentifier
 
 	// Needs to be backward compatible with FNames (and everything that can be constructed into a FName).
 	template <typename T = FName, std::enable_if_t<std::is_constructible_v<FName, T>, bool> = true>
-	FPCGAttributeIdentifier(const T& InName, FPCGMetadataLayerID InMetadataLayerID = PCGMetadataLayerID::Default)
+	FPCGAttributeIdentifier(const T& InName, FPCGMetadataDomainID InMetadataDomainID = PCGMetadataDomainID::Default)
 		: Name(InName)
-		, MetadataLayer(InMetadataLayerID)
+		, MetadataDomain(InMetadataDomainID)
 	{}
 
 	template <typename T = FName, std::enable_if_t<std::is_constructible_v<FName, T>, bool> = true>
@@ -144,18 +144,18 @@ struct FPCGAttributeIdentifier
 
 	bool operator==(const FPCGAttributeIdentifier& Other) const
 	{
-		return Name == Other.Name && MetadataLayer == Other.MetadataLayer;
+		return Name == Other.Name && MetadataDomain == Other.MetadataDomain;
 	}
 
 	friend uint32 GetTypeHash(const FPCGAttributeIdentifier& Item)
 	{
-		return HashCombine(GetTypeHash(Item.Name), GetTypeHash(Item.MetadataLayer));
+		return HashCombine(GetTypeHash(Item.Name), GetTypeHash(Item.MetadataDomain));
 	}
 	
 	friend FArchive& operator<<(FArchive& Ar, FPCGAttributeIdentifier& Item)
 	{
 		Ar << Item.Name;
-		Ar << Item.MetadataLayer;
+		Ar << Item.MetadataDomain;
 		return Ar;
 	}
 	
@@ -177,12 +177,12 @@ struct FPCGAttributeIdentifier
 	}
 
 	template <typename Container>
-	static TMap<FPCGMetadataLayerID, TSet<FName>> SplitByLayer(const Container& InContainer)
+	static TMap<FPCGMetadataDomainID, TSet<FName>> SplitByDomain(const Container& InContainer)
 	{
-		TMap<FPCGMetadataLayerID, TSet<FName>> OutMap;
+		TMap<FPCGMetadataDomainID, TSet<FName>> OutMap;
 		for (const FPCGAttributeIdentifier& It : InContainer)
 		{
-			OutMap.FindOrAdd(It.MetadataLayer).Add(It.Name);
+			OutMap.FindOrAdd(It.MetadataDomain).Add(It.Name);
 		}
 
 		return OutMap;
@@ -192,6 +192,6 @@ struct FPCGAttributeIdentifier
 	FName Name = NAME_None;
 
 	UPROPERTY(EditAnywhere, BlueprintReadWrite, Category = "", meta = (ShowOnlyInnerProperties))
-	FPCGMetadataLayerID MetadataLayer;
+	FPCGMetadataDomainID MetadataDomain;
 };
 
diff --git a/Engine/Plugins/PCG/Source/PCG/Public/Metadata/PCGMetadataDomain.h b/Engine/Plugins/PCG/Source/PCG/Public/Metadata/PCGMetadataDomain.h
new file mode 100644
--- /dev/null
+++ b/Engine/Plugins/PCG/Source/PCG/Public/Metadata/PCGMetadataDomain.h
@@ -0,0 +1,396 @@
+// Copyright Epic Games, Inc. All Rights Reserved.
+
+#pragma once
+
+#include "PCGCommon.h"
+#include "Metadata/PCGMetadataCommon.h"
+#include "Metadata/PCGMetadataAttributeTpl.h"
+
+class UPCGMetadata;
+
+class PCG_API FPCGMetadataDomain
+{
+	friend UPCGMetadata;
+
+public:
+	FPCGMetadataDomain(UPCGMetadata* InTopMetadata, FPCGMetadataDomainID InMetadataDomainID);
+	virtual ~FPCGMetadataDomain();
+	
+	virtual void Serialize(FArchive& InArchive);
+	
+	/** Initializes the sub metadata from a parent sub metadata, if any (can be null). Copies attributes and values. */
+	void Initialize(const FPCGMetadataDomain* InParent) { InitializeWithAttributeFilter(InParent, TSet<FName>()); }
+
+	/** Initializes the sub metadata from a parent sub metadata, if any (can be null) with the option to not add attributes from the parent. */
+	void Initialize(const FPCGMetadataDomain* InParent, bool bAddAttributesFromParent);
+
+	/**
+	 * Initializes the sub metadata from a parent sub metadata. Copies attributes and values.
+	 * @param InParent The parent metadata to use as a template, if any (can be null).
+	 * @param InFilteredAttributes Optional list of attributes to exclude or include when adding the attributes from the parent.
+	 * @param InFilterMode Defines attribute filter operation.
+	 */
+	void InitializeWithAttributeFilter(const FPCGMetadataDomain* InParent, const TSet<FName>& InFilteredAttributes, EPCGMetadataFilterMode InFilterMode = EPCGMetadataFilterMode::ExcludeAttributes, EPCGStringMatchingOperator InMatchOperator = EPCGStringMatchingOperator::Equal);
+
+	/** Initializes the sub metadata from a parent sub metadata by copying all attributes to it.
+	* @param InMetadataToCopy Metadata to copy from
+	* @param InOptionalEntriesToCopy Optional array that contains the keys to copy over. This array order will be respected, so it can also be used to re-order entries. Can be null.
+	*/
+	void InitializeAsCopy(const FPCGMetadataDomain* InMetadataToCopy, const TArray<PCGMetadataEntryKey>* InOptionalEntriesToCopy = nullptr);
+
+	/** Initializes the metadata from a parent metadata by copy filtered attributes only to it
+	* @param InMetadataToCopy Metadata to copy from
+	* @param InFilteredAttributes Attributes to keep/exclude, can be empty.
+	* @param InFilterMode Filter to know if we should keep or exclude InFilteredAttributes.
+	* @param InOptionalEntriesToCopy Optional array that contains the keys to copy over. This array order will be respected, so it can also be used to re-order entries. Can be null.
+	*/
+	void InitializeAsCopyWithAttributeFilter(const FPCGMetadataDomain* InMetadataToCopy, const TSet<FName>& InFilteredAttributes, EPCGMetadataFilterMode InFilterMode = EPCGMetadataFilterMode::ExcludeAttributes, const TArray<PCGMetadataEntryKey>* InOptionalEntriesToCopy = nullptr);
+
+	/** Creates missing attributes from another metadata if they are not currently present - note that this does not copy values */
+	bool AddAttributes(const FPCGMetadataDomain* InOther) { return AddAttributesFiltered(InOther, TSet<FName>()); }
+
+	/**
+	 * Creates missing attributes from another metadata if they are not currently present - note that this does not copy values.
+	 * @param InOther The other metadata to obtain a list of attributes from.
+	 * @param InFilteredAttributes Optional list of attributes to exclude or include when adding the attributes.
+	 * @param InFilterMode Defines attribute filter operation.
+	 */
+	bool AddAttributesFiltered(const FPCGMetadataDomain* InOther, const TSet<FName>& InFilteredAttributes, EPCGMetadataFilterMode InFilterMode = EPCGMetadataFilterMode::ExcludeAttributes, EPCGStringMatchingOperator InMatchOperator = EPCGStringMatchingOperator::Equal);
+
+	/** Creates missing attribute from another metadata if it is not currently present - note that this does not copy values */
+	bool AddAttribute(const FPCGMetadataDomain* InOther, FName AttributeName);
+
+	/** Copies attributes from another metadata, including entries & values. Warning: this is intended when dealing with the same data set */
+	void CopyAttributes(const FPCGMetadataDomain* InOther);
+
+	/** Copies an attribute from another metadata, including entries & values. Warning: this is intended when dealing with the same data set */
+	void CopyAttribute(const FPCGMetadataDomain* InOther, FName AttributeToCopy, FName NewAttributeName);
+
+	/** Copies another attribute, with options to keep its parent and copy entries/values */
+	FPCGMetadataAttributeBase* CopyAttribute(const FPCGMetadataAttributeBase* OriginalAttribute, FName NewAttributeName, bool bKeepParent, bool bCopyEntries, bool bCopyValues);
+
+	/** Returns this metadata's parent */
+	const FPCGMetadataDomain* GetParent() const { return Parent; }
+	const FPCGMetadataDomain* GetRoot() const;
+	const UPCGMetadata* GetTopMetadata() const { return TopMetadata; }
+	bool HasParent(const FPCGMetadataDomain* InTentativeParent) const;
+
+	/** Unparents current metadata by flattening the attributes (values, entries, etc.) */
+	void FlattenImpl();
+
+	/** Unparents current metadata, flatten attribute and only keep the entries specified. Return true if something has changed and keys needs be updated. */
+	bool FlattenAndCompress(const TArray<PCGMetadataEntryKey>& InEntryKeysToKeep);
+
+	/** Creates an attribute given a property.
+	* @param AttributeName: Target attribute to create
+	* @param Object: Object to get the property value from
+	* @param Property: The property to set from
+	* @returns true if the attribute creation succeeded
+	*/
+	bool CreateAttributeFromProperty(FName AttributeName, const UObject* Object, const FProperty* Property);
+
+	/** Creates an attribute given a property.
+	* @param AttributeName: Target attribute to create
+	* @param Data: Data pointer to get the property value from
+	* @param Property: The property to set from
+	* @returns true if the attribute creation succeeded
+	*/
+	bool CreateAttributeFromDataProperty(FName AttributeName, const void* Data, const FProperty* Property);
+
+	/** Set an attribute given a property and its value.
+	* @param AttributeName: Target attribute to set the property's value to
+	* @param EntryKey: Metadata entry key to set the value to
+	* @param Object: Object to get the property value from
+	* @param Property: The property to set from
+	* @param bCreate: If true and the attribute doesn't exists, it will create an attribute based on the property type
+	* @returns true if the attribute creation (if required) and the value set succeeded
+	*/
+	bool SetAttributeFromProperty(FName AttributeName, PCGMetadataEntryKey& EntryKey, const UObject* Object, const FProperty* Property, bool bCreate);
+
+	/** Set an attribute given a property and its value.
+	* @param AttributeName: Target attribute to set the property's value to
+	* @param EntryKey: Metadata entry key to set the value to
+	* @param Data: Data pointer to get the property value from
+	* @param Property: The property to set from
+	* @param bCreate: If true and the attribute doesn't exists, it will create an attribute based on the property type
+	* @returns true if the attribute creation (if required) and the value set succeeded
+	*/
+	bool SetAttributeFromDataProperty(FName AttributeName, PCGMetadataEntryKey& EntryKey, const void* Data, const FProperty* Property, bool bCreate);
+
+	/** Get attributes */
+	FPCGMetadataAttributeBase* GetMutableAttribute(FName AttributeName);
+	const FPCGMetadataAttributeBase* GetConstAttribute(FName AttributeName) const;
+	const FPCGMetadataAttributeBase* GetConstAttributeById(int32 AttributeId) const;
+	
+	bool HasAttribute(FName AttributeName) const;
+	bool HasCommonAttributes(const FPCGMetadataDomain* InMetadata) const;
+
+	/** Return the number of attributes in this metadata. */
+	int32 GetAttributeCount() const;
+
+	template <typename T>
+	FPCGMetadataAttribute<T>* GetMutableTypedAttribute(FName AttributeName);
+
+	template <typename T>
+	FPCGMetadataAttribute<T>* GetMutableTypedAttribute_Unsafe(FName AttributeName);
+
+	template <typename T>
+	const FPCGMetadataAttribute<T>* GetConstTypedAttribute(FName AttributeName) const;
+	
+	void GetAttributes(TArray<FName>& AttributeNames, TArray<EPCGMetadataTypes>& AttributeTypes) const;
+
+	/** Returns name of the most recently created attribute, or none if no attributes are present. */
+	FName GetLatestAttributeNameOrNone() const;
+
+	/** Delete/Hide attribute */
+	// Due to stream inheriting, we might want to consider "hiding" parent stream and deleting local streams only
+	void DeleteAttribute(FName AttributeName);
+
+	/** Copy attribute */
+	bool CopyExistingAttribute(FName AttributeToCopy, FName NewAttributeName, bool bKeepParent = true);
+
+	/** Rename attribute */
+	bool RenameAttribute(FName AttributeToRename, FName NewAttributeName);
+
+	/** Clear/Reinit attribute */
+	void ClearAttribute(FName AttributeToClear);
+
+	/** Change type of an attribute */
+	bool ChangeAttributeType(FName AttributeName, int16 AttributeNewType);
+
+	/** Adds a unique entry key to the metadata */
+	int64 AddEntry(int64 ParentEntryKey = -1);
+
+	/** Adds a unique entry key to the metadata for all the parent entry keys. */
+	TArray<int64> AddEntries(TArrayView<const int64> ParentEntryKeys);
+
+	/** Adds a unique entry key to the metadata for all the parent entry keys, in place. */
+	void AddEntriesInPlace(TArrayView<int64*> ParentEntryKeys);
+
+	/** Advanced method.
+	*   In a MT context, we might not want to add the entry directly (because of write lock). Call this to generate an unique index in the MT context
+	*   And call AddDelayedEntries at the end when you want to add all the entries.
+	*   Make sure to not call AddEntry in a different thread or even in the same thread if you use this function, or it will change all the indexes.
+	*/
+	int64 AddEntryPlaceholder();
+
+	/** Advanced method.
+	*   If you used AddEntryPlaceholder, call this function at the end of your MT processing to add all the entries in one shot.
+	*   Make sure to call this one at the end!!
+	*   @param AllEntries Array of pairs. First item is the EntryIndex generated by AddEntryPlaceholder, the second the ParentEntryKey (cf AddEntry)
+	*/
+	void AddDelayedEntries(const TArray<TTuple<int64, int64>>& AllEntries);
+
+	/** Initializes the metadata entry key. Returns true if key set from either parent */
+	bool InitializeOnSet(PCGMetadataEntryKey& InOutKey, PCGMetadataEntryKey InParentKeyA = PCGInvalidEntryKey, const FPCGMetadataDomain* InParentMetadataA = nullptr, PCGMetadataEntryKey InParentKeyB = PCGInvalidEntryKey, const FPCGMetadataDomain* InParentMetadataB = nullptr);
+
+	/** Metadata chaining mechanism */
+	PCGMetadataEntryKey GetParentKey(PCGMetadataEntryKey LocalItemKey) const;
+
+	/** Metadata chaining mechanism for bulk version. Can provide a mask to only update only a subset of the passed keys. */
+	void GetParentKeys(TArrayView<PCGMetadataEntryKey> LocalItemKeys, const TBitArray<>* Mask = nullptr) const;
+
+	/** Attributes operations */
+	void MergeAttributes(PCGMetadataEntryKey InKeyA, const FPCGMetadataDomain* InMetadataA, PCGMetadataEntryKey InKeyB, const FPCGMetadataDomain* InMetadataB, PCGMetadataEntryKey& OutKey, EPCGMetadataOp Op);
+	void MergeAttributesSubset(PCGMetadataEntryKey InKeyA, const FPCGMetadataDomain* InMetadataA, const FPCGMetadataDomain* InMetadataSubetA, PCGMetadataEntryKey InKeyB, const FPCGMetadataDomain* InMetadataB, const FPCGMetadataDomain* InMetadataSubsetb, PCGMetadataEntryKey& OutKey, EPCGMetadataOp Op);
+
+	void ResetWeightedAttributes(PCGMetadataEntryKey& OutKey);
+	void AccumulateWeightedAttributes(PCGMetadataEntryKey InKey, const FPCGMetadataDomain* InMetadata, float Weight, bool bSetNonInterpolableAttributes, PCGMetadataEntryKey& OutKey);
+
+	void SetAttributes(PCGMetadataEntryKey InKey, const FPCGMetadataDomain* InMetadata, PCGMetadataEntryKey& OutKey);
+	void SetAttributes(const TArrayView<const PCGMetadataEntryKey>& InOriginalKeys, const FPCGMetadataDomain* InMetadata, const TArrayView<PCGMetadataEntryKey>* InOutOptionalKeys = nullptr, FPCGContext* OptionalContext = nullptr);
+	void SetAttributes(const TArrayView<const PCGMetadataEntryKey>& InKeys, const FPCGMetadataDomain* InMetadata, const TArrayView<PCGMetadataEntryKey>& OutKeys, FPCGContext* OptionalContext = nullptr);
+
+	void ComputeWeightedAttribute(PCGMetadataEntryKey& OutKey, const TArrayView<TPair<PCGMetadataEntryKey, float>>& InWeightedKeys, const FPCGMetadataDomain* InMetadata);
+
+	int64 GetItemKeyCountForParent() const;
+	int64 GetLocalItemCount() const;
+
+	/** Return the number of entries in metadata including the parent entries. */
+	int64 GetItemCountForChild() const;
+
+	/**
+	* Create a new attribute. If the attribute already exists, it will raise a warning (use FindOrCreateAttribute if this usecase can arise)
+	* If the attribute already exists but is of the wrong type, it will fail and return nullptr. Same if the name is invalid.
+	* Return a typed attribute pointer, of the requested type T.
+	*/
+	template<typename T>
+	FPCGMetadataAttribute<T>* CreateAttribute(FName AttributeName, const T& DefaultValue, bool bAllowsInterpolation, bool bOverrideParent);
+
+	/**
+	* Find or create an attribute. Follows CreateAttribute signature.
+	* Extra boolean bOverwriteIfTypeMismatch allows to overwrite an existing attribute if the type mismatch.
+	* Same as CreateAttribute, it will return nullptr if the attribute name is invalid.
+	* Return a typed attribute pointer, of the requested type T.
+	*/
+	template<typename T>
+	FPCGMetadataAttribute<T>* FindOrCreateAttribute(FName AttributeName, const T& DefaultValue = T{}, bool bAllowsInterpolation = true, bool bOverrideParent = true, bool bOverwriteIfTypeMismatch = true);
+
+	/** Computes Crc from all attributes & keys from outer's data. */
+	void AddToCrc(FArchiveCrc32& Ar, const UPCGData* Data, bool bFullDataCrc) const;
+
+protected:
+	FPCGMetadataAttributeBase* CopyAttribute(FName AttributeToCopy, FName NewAttributeName, bool bKeepParent, bool bCopyEntries, bool bCopyValues);
+
+	bool ParentHasAttribute(FName AttributeName) const;
+
+	void AddAttributeInternal(FName AttributeName, FPCGMetadataAttributeBase* Attribute);
+	void RemoveAttributeInternal(FName AttributeName);
+
+	void SetLastCachedSelectorOnOwner(FName AttributeName);
+
+	UPCGMetadata* TopMetadata = nullptr;
+	FPCGMetadataDomainID DomainID;
+	const FPCGMetadataDomain* Parent = nullptr;
+
+	TMap<FName, FPCGMetadataAttributeBase*> Attributes;
+	PCGMetadataAttributeKey NextAttributeId = 0;
+
+	TArray<PCGMetadataEntryKey> ParentKeys;
+	int64 ItemKeyOffset = 0;
+
+	mutable FRWLock AttributeLock;
+	mutable FRWLock ItemLock;
+
+	TAtomic<int64> DelayedEntriesIndex = 0;
+};
+
+
+template<typename T>
+FPCGMetadataAttribute<T>* FPCGMetadataDomain::CreateAttribute(FName AttributeName, const T& DefaultValue, bool bAllowsInterpolation, bool bOverrideParent)
+{
+	if (!FPCGMetadataAttributeBase::IsValidName(AttributeName))
+	{
+		UE_LOG(LogPCG, Error, TEXT("Attribute name '%s' is invalid"), *AttributeName.ToString());
+		return nullptr;
+	}
+
+	const FPCGMetadataAttributeBase* ParentAttribute = nullptr;
+
+	if (bOverrideParent && Parent)
+	{
+		ParentAttribute = Parent->GetConstAttribute(AttributeName);
+	}
+
+	if (ParentAttribute && (ParentAttribute->GetTypeId() != PCG::Private::MetadataTypes<T>::Id))
+	{
+		// Can't parent if the types do not match
+		ParentAttribute = nullptr;
+	}
+
+	FPCGMetadataAttribute<T>* NewAttribute = new FPCGMetadataAttribute<T>(this, AttributeName, ParentAttribute, DefaultValue, bAllowsInterpolation);
+
+	{
+		FWriteScopeLock WriteLock(AttributeLock);
+
+		if (FPCGMetadataAttributeBase** ExistingAttribute = Attributes.Find(AttributeName))
+		{
+			delete NewAttribute;
+			if ((*ExistingAttribute)->GetTypeId() != PCG::Private::MetadataTypes<T>::Id)
+			{
+				UE_LOG(LogPCG, Error, TEXT("Attribute %s already exists but is not the right type. Abort."), *AttributeName.ToString());
+				return nullptr;
+			}
+			else
+			{
+				UE_LOG(LogPCG, Warning, TEXT("Attribute %s already exists"), *AttributeName.ToString());
+				NewAttribute = static_cast<FPCGMetadataAttribute<T>*>(*ExistingAttribute);
+			}
+		}
+		else
+		{
+			NewAttribute->AttributeId = NextAttributeId++;
+			AddAttributeInternal(AttributeName, NewAttribute);
+
+			// Also when creating an attribute, notify the PCG Data owner that the latest attribute manipulated is this one.
+			SetLastCachedSelectorOnOwner(AttributeName);
+		}
+	}
+
+	return NewAttribute;
+}
+
+template<typename T>
+FPCGMetadataAttribute<T>* FPCGMetadataDomain::GetMutableTypedAttribute_Unsafe(FName AttributeName)
+{
+	FPCGMetadataAttribute<T>* Attribute = nullptr;
+
+	if (FPCGMetadataAttributeBase** FoundAttribute = Attributes.Find(AttributeName))
+	{
+		FPCGMetadataAttributeBase* BaseAttribute = *FoundAttribute;
+		if (BaseAttribute && BaseAttribute->GetTypeId() == PCG::Private::MetadataTypes<T>::Id)
+		{
+			Attribute = static_cast<FPCGMetadataAttribute<T>*>(BaseAttribute);
+		}
+	}
+
+	return Attribute;
+}
+
+template<typename T>
+FPCGMetadataAttribute<T>* FPCGMetadataDomain::FindOrCreateAttribute(FName AttributeName, const T& DefaultValue, bool bAllowsInterpolation, bool bOverrideParent, bool bOverwriteIfTypeMismatch)
+{
+	{
+		FReadScopeLock ScopeLock(AttributeLock);
+		if (FPCGMetadataAttribute<T>* Attribute = GetMutableTypedAttribute_Unsafe<T>(AttributeName))
+		{
+			return Attribute;
+		}
+	}
+
+	FWriteScopeLock ScopeLock(AttributeLock);
+	if (FPCGMetadataAttribute<T>* Attribute = GetMutableTypedAttribute_Unsafe<T>(AttributeName))
+	{
+		return Attribute;
+	}
+
+	// If an attribute with this name exists here, there is a type mismatch.
+	if (FPCGMetadataAttributeBase** FoundAttribute = Attributes.Find(AttributeName))
+	{
+		if (bOverwriteIfTypeMismatch)
+		{
+			delete *FoundAttribute;
+			RemoveAttributeInternal(AttributeName);
+		}
+		else
+		{
+			return nullptr;
+		}
+	}
+
+	// A new attribute will be created.
+	if (!FPCGMetadataAttributeBase::IsValidName(AttributeName))
+	{
+		UE_LOG(LogPCG, Error, TEXT("Attribute name '%s' is invalid"), *AttributeName.ToString());
+		return nullptr;
+	}
+
+	// Parent is const and therefore should be safe
+	const FPCGMetadataAttributeBase* ParentAttribute = (bOverrideParent && Parent) ? Parent->GetConstTypedAttribute<T>(AttributeName) : nullptr;
+	FPCGMetadataAttribute<T>* NewAttribute = new FPCGMetadataAttribute<T>(this, AttributeName, ParentAttribute, DefaultValue, bAllowsInterpolation);
+	NewAttribute->AttributeId = NextAttributeId++;
+	AddAttributeInternal(AttributeName, NewAttribute);
+
+	// Also when creating an attribute, notify the PCG Data owner that the latest attribute manipulated is this one.
+	SetLastCachedSelectorOnOwner(AttributeName);
+
+	return NewAttribute;
+}
+
+template <typename T>
+FPCGMetadataAttribute<T>* FPCGMetadataDomain::GetMutableTypedAttribute(FName AttributeName)
+{
+	FPCGMetadataAttributeBase* BaseAttribute = GetMutableAttribute(AttributeName);
+	return (BaseAttribute && (BaseAttribute->GetTypeId() == PCG::Private::MetadataTypes<T>::Id))
+		? static_cast<FPCGMetadataAttribute<T>*>(BaseAttribute)
+		: nullptr;
+}
+
+template <typename T>
+const FPCGMetadataAttribute<T>* FPCGMetadataDomain::GetConstTypedAttribute(FName AttributeName) const
+{
+	const FPCGMetadataAttributeBase* BaseAttribute = GetConstAttribute(AttributeName);
+	return (BaseAttribute && (BaseAttribute->GetTypeId() == PCG::Private::MetadataTypes<T>::Id))
+		? static_cast<const FPCGMetadataAttribute<T>*>(BaseAttribute)
+		: nullptr;
+}
\ No newline at end of file
diff --git a/Engine/Plugins/PCG/Source/PCG/Public/Metadata/PCGMetadataLayer.h b/Engine/Plugins/PCG/Source/PCG/Public/Metadata/PCGMetadataLayer.h
deleted file mode 100644
diff --git a/Engine/Plugins/PCG/Source/PCG/Public/PCGCustomVersion.h b/Engine/Plugins/PCG/Source/PCG/Public/PCGCustomVersion.h
--- a/Engine/Plugins/PCG/Source/PCG/Public/PCGCustomVersion.h
+++ b/Engine/Plugins/PCG/Source/PCG/Public/PCGCustomVersion.h
@@ -80,7 +80,7 @@ struct FPCGCustomVersion
 		// Attributes and tags can now contain spaces and will no longer be parsed by spaces.
 		AttributesAndTagsCanContainSpaces = 23,
 
-		// Multi Layer Metadata
+		// Multi Domain Metadata
 		MultiLevelMetadata = 24,
 
 		// Refactor of the Attribute Property selector to deprecate point properties and support any property.
diff --git a/Engine/Plugins/PCG/Source/PCG/Public/PCGData.h b/Engine/Plugins/PCG/Source/PCG/Public/PCGData.h
--- a/Engine/Plugins/PCG/Source/PCG/Public/PCGData.h
+++ b/Engine/Plugins/PCG/Source/PCG/Public/PCGData.h
@@ -80,8 +80,8 @@ public:
 	virtual UPCGMetadata* MutableMetadata() { return Metadata; }
 
 	virtual void Flatten();
-	virtual FPCGMetadataLayerID GetDefaultMetadataLayerID() const { return PCGMetadataLayerID::Data; }
-	virtual TArray<FPCGMetadataLayerID> GetAllSupportedMetadataLayerIDs() const { return {PCGMetadataLayerID::Data}; }
+	virtual FPCGMetadataDomainID GetDefaultMetadataDomainID() const { return PCGMetadataDomainID::Data; }
+	virtual TArray<FPCGMetadataDomainID> GetAllSupportedMetadataDomainIDs() const { return {PCGMetadataDomainID::Data}; }
 
 	// Not accessible through blueprint to make sure the constness is preserved.
 	UPROPERTY(VisibleAnywhere, Category = Metadata)
diff --git a/Engine/Plugins/PCG/Source/PCG/Public/PCGParamData.h b/Engine/Plugins/PCG/Source/PCG/Public/PCGParamData.h
--- a/Engine/Plugins/PCG/Source/PCG/Public/PCGParamData.h
+++ b/Engine/Plugins/PCG/Source/PCG/Public/PCGParamData.h
@@ -29,8 +29,8 @@ public:
 	virtual UPCGParamData* DuplicateData(FPCGContext* Context, bool bInitializeMetadata = true) const override;
 	// ~End UPCGData interface
 
-	virtual FPCGMetadataLayerID GetDefaultMetadataLayerID() const override { return PCGMetadataLayerID::Elements; }
-	virtual TArray<FPCGMetadataLayerID> GetAllSupportedMetadataLayerIDs() const override { return {PCGMetadataLayerID::Data, PCGMetadataLayerID::Elements}; }
+	virtual FPCGMetadataDomainID GetDefaultMetadataDomainID() const override { return PCGMetadataDomainID::Elements; }
+	virtual TArray<FPCGMetadataDomainID> GetAllSupportedMetadataDomainIDs() const override { return {PCGMetadataDomainID::Data, PCGMetadataDomainID::Elements}; }
 
 	/** Returns the entry for the given name */
 	UFUNCTION(BlueprintCallable, Category = Params)
