[PCG] Fix issue when moving actors that have no bounds in partitioned

* Also add a condition to not discard previous data if there is a no op in Object Changed because the actor is delayed

#rnx

diff --git a/Engine/Plugins/PCG/Source/PCG/Private/PCGActorAndComponentMapping.cpp b/Engine/Plugins/PCG/Source/PCG/Private/PCGActorAndComponentMapping.cpp
--- a/Engine/Plugins/PCG/Source/PCG/Private/PCGActorAndComponentMapping.cpp
+++ b/Engine/Plugins/PCG/Source/PCG/Private/PCGActorAndComponentMapping.cpp
@@ -743,7 +743,9 @@ void FPCGActorAndComponentMapping::ForAllIntersectingPartitionActors(const FBox&
 		return;
 	}
 
-	auto ForAllIntersectingPartitionActorsOfGridSize = [InFunc, &InBounds, this](const FPCGGridDescriptor& GridDescriptor)
+	const bool bInBoundsAreEmpty = FMath::IsNearlyZero(InBounds.GetVolume());
+
+	auto ForAllIntersectingPartitionActorsOfGridSize = [InFunc, &InBounds, this, bInBoundsAreEmpty](const FPCGGridDescriptor& GridDescriptor)
 	{
 		FReadScopeLock ReadLock(PartitionActorsMapLock);
 
@@ -767,9 +769,15 @@ void FPCGActorAndComponentMapping::ForAllIntersectingPartitionActors(const FBox&
 					{
 						if (APCGPartitionActor* Actor = ActorPtr->Get())
 						{
-							// Exclude any surrounding cells which are touching this cell but not meaningfully overlapping.
-							const FBox Overlap = Actor->GetFixedBounds().Overlap(InBounds);
-							if (Overlap.GetVolume() > 0)
+							const FBox ActorBounds = Actor->GetFixedBounds();
+							// Exclude any surrounding cells which are touching this cell but not meaningfully overlapping, only if the InBounds are
+							// not "empty" (volume > 0). If the InBounds are "empty" (just a single point), the overlap will always be zero,
+							// so just do an intersection check.
+							const bool bShouldProcess = bInBoundsAreEmpty
+								? ActorBounds.Intersect(InBounds)
+								: !FMath::IsNearlyZero(Actor->GetFixedBounds().Overlap(InBounds).GetVolume());
+							
+							if (bShouldProcess)
 							{
 								InFunc(Actor);
 							}
@@ -1735,7 +1743,8 @@ void FPCGActorAndComponentMapping::OnObjectPropertyChanged(UObject* InObject, FP
 
 	if (bNoOperation)
 	{
-		if (Actor)
+		// Don't remove it as it will be needed when the delay will be called.
+		if (Actor && !(DelayedChangedActors.Contains(Actor)))
 		{
 			auto RemoveAll = [this](AActor* InActor, auto RecursiveCall) -> void
 			{
