Fixed use after free in UPCGMetadata::MetadataDomains

Old code was doing:
MetadataDomains[PCGMetadataDomainID::Default] returns a reference T& that points to the internal memory of TMap::MetadataDomains
Emplace of MetadataLayers allocates a new entry. No more space -> reallocate internal memory
memory pointed by T& is no more valid -> use after free

#rnx

diff --git a/Engine/Plugins/PCG/Source/PCG/Private/Metadata/PCGMetadata.cpp b/Engine/Plugins/PCG/Source/PCG/Private/Metadata/PCGMetadata.cpp
--- a/Engine/Plugins/PCG/Source/PCG/Private/Metadata/PCGMetadata.cpp
+++ b/Engine/Plugins/PCG/Source/PCG/Private/Metadata/PCGMetadata.cpp
@@ -211,7 +211,8 @@ void UPCGMetadata::SetupDomain(FPCGMetadataDomainID DomainID, bool bIsDefault)
 		
 		DefaultDomain = DomainID;
 		// Map the default domain to the DefaultID domain and update domain to the not-default domain ID.
-		MetadataDomains.Emplace(DomainID, MetadataDomains[PCGMetadataDomainID::Default])->DomainID = DomainID;
+		TSharedPtr<FPCGMetadataDomain> DefaultIDLayer = MetadataDomains[PCGMetadataDomainID::Default];
+		MetadataDomains.Emplace(DomainID, DefaultIDLayer)->DomainID = DomainID;
 	}
 	else
 	{
