Debug object tree view improvements
- Changed icons (used graph icons and zoom to content)
- Added row coloring if there are local warnings/errors, and italics if it's a dynamic graph.
- Made the row bold if it's the currently selected debug object (not super useful for users, but easier to understand for devs).
- Changed the way naming is presented so the graph names are shown first, then their originating node.
- Fixed sorting when there are larger numbers so that 20 appears after 3.
- Shown downstream executions in the tree view.
- Made the jump to properly set the debug object tree view in the target editor.
- Clicking on a downstream entry will select the appropriate node, like in the Find window.
- Double-click on an entry will expand the graph down to the first occurrence (bfs) of the current graph.

Light refactoring around the visual logs so we do filtering at the outset, not early so we get the "full" error state and cull when building the message instead.

diff --git a/Engine/Plugins/PCG/Source/PCG/Private/Graph/PCGStackContext.cpp b/Engine/Plugins/PCG/Source/PCG/Private/Graph/PCGStackContext.cpp
--- a/Engine/Plugins/PCG/Source/PCG/Private/Graph/PCGStackContext.cpp
+++ b/Engine/Plugins/PCG/Source/PCG/Private/Graph/PCGStackContext.cpp
@@ -208,6 +208,19 @@ const UPCGNode* FPCGStack::GetCurrentFrameNode() const
 	return StackFrames.IsEmpty() ? nullptr : Cast<const UPCGNode>(StackFrames.Last().Object.Get());
 }
 
+const UPCGNode* FPCGStack::GetNodeForCurrentFrame() const
+{
+	for (int StackIndex = GetStackFrames().Num() - 1; StackIndex >= 0; --StackIndex)
+	{
+		if (const UPCGNode* Node = Cast<const UPCGNode>(StackFrames[StackIndex].Object.Get()))
+		{
+			return Node;
+		}
+	}
+
+	return nullptr;
+}
+
 bool FPCGStack::HasObject(const UObject* InObject) const
 {
 	return !!Algo::FindByPredicate(StackFrames, [InObject](const FPCGStackFrame& Frame)
diff --git a/Engine/Plugins/PCG/Source/PCG/Private/PCGSubsystem.cpp b/Engine/Plugins/PCG/Source/PCG/Private/PCGSubsystem.cpp
--- a/Engine/Plugins/PCG/Source/PCG/Private/PCGSubsystem.cpp
+++ b/Engine/Plugins/PCG/Source/PCG/Private/PCGSubsystem.cpp
@@ -1791,7 +1791,7 @@ void UPCGSubsystem::ClearExecutionMetadata(UPCGComponent* InComponent)
 	ClearExecutionMetadata(Stack);
 }
 
-TArray<FPCGStack> UPCGSubsystem::GetExecutedStacks(const UPCGComponent* InComponent, const UPCGGraph* InSubgraph)
+TArray<FPCGStack> UPCGSubsystem::GetExecutedStacks(const UPCGComponent* InComponent, const UPCGGraph* InSubgraph, bool bOnlyWithSubgraphAsCurrentFrame)
 {
 	FReadScopeLock Lock(ExecutedStacksLock);
 
@@ -1799,7 +1799,8 @@ TArray<FPCGStack> UPCGSubsystem::GetExecutedStacks(const UPCGComponent* InCompon
 
 	for (const FPCGStack& Stack : ExecutedStacks)
 	{
-		if ((!InComponent || Stack.GetRootComponent() == InComponent) && Stack.GetGraphForCurrentFrame() == InSubgraph)
+		if ((!InComponent || Stack.GetRootComponent() == InComponent) && 
+			((bOnlyWithSubgraphAsCurrentFrame && Stack.GetGraphForCurrentFrame() == InSubgraph) || (!bOnlyWithSubgraphAsCurrentFrame && Stack.HasObject(InSubgraph))))
 		{
 			MatchingStacks.Add(Stack);
 		}
diff --git a/Engine/Plugins/PCG/Source/PCG/Private/Utils/PCGNodeVisualLogs.cpp b/Engine/Plugins/PCG/Source/PCG/Private/Utils/PCGNodeVisualLogs.cpp
--- a/Engine/Plugins/PCG/Source/PCG/Private/Utils/PCGNodeVisualLogs.cpp
+++ b/Engine/Plugins/PCG/Source/PCG/Private/Utils/PCGNodeVisualLogs.cpp
@@ -12,6 +12,13 @@
 
 #if WITH_EDITOR
 
+namespace PCGNodeVisualLogsConstants
+{
+	static constexpr int MaxLogsInSummary = 8;
+	static const FText Warning = LOCTEXT("PCGLogWarning", "Warning");
+	static const FText Error = LOCTEXT("PCGLogError", "Error");
+}
+
 void FPCGNodeVisualLogs::Log(const FPCGStack& InPCGStack, ELogVerbosity::Type InVerbosity, const FText& InMessage)
 {
 	bool bAdded = false;
@@ -96,103 +103,134 @@ bool FPCGNodeVisualLogs::HasLogsOfVerbosity(const FPCGStack& InPCGStack, ELogVer
 	return false;
 }
 
-FText FPCGNodeVisualLogs::GetLogsSummaryText(const UPCGNode* InNode, ELogVerbosity::Type& OutMinimumVerbosity) const
+FPCGPerNodeVisualLogs FPCGNodeVisualLogs::GetLogs(const FPCGStack& InPCGStack) const
 {
-	FReadScopeLock ScopedReadLock(LogsLock);
+	FPCGPerNodeVisualLogs Logs;
 
-	// Compose user friendly summary of first N errors
-	FText Summary = FText::GetEmpty();
-	OutMinimumVerbosity = ELogVerbosity::All;
+	ForAllMatchingLogs(InPCGStack, [&Logs](const FPCGStack&, const FPCGPerNodeVisualLogs& InLogs)
+	{
+		Logs.Append(InLogs);
+		return true;
+	});
+	
+	return Logs;
+}
 
-	int LogCounter = 0;
+void FPCGNodeVisualLogs::ForAllMatchingLogs(const FPCGStack& InPCGStack, TFunctionRef<bool(const FPCGStack&, const FPCGPerNodeVisualLogs&)> InFunc) const
+{
+	FReadScopeLock ScopedReadLock(LogsLock);
+	for (const TPair<FPCGStack, FPCGPerNodeVisualLogs>& Entry : StackToLogs)
+	{
+		if (Entry.Key.BeginsWith(InPCGStack))
+		{
+			if (!InFunc(Entry.Key, Entry.Value))
+			{
+				break;
+			}
+		}
+	}
+}
+
+void FPCGNodeVisualLogs::GetLogs(const UPCGNode* InNode, FPCGPerNodeVisualLogs& OutLogs, TArray<const UPCGComponent*>& OutComponents) const
+{
+	OutLogs.Reset();
+	OutComponents.Reset();
 
+	FReadScopeLock ScopedReadLock(LogsLock);
 	for (const TPair<FPCGStack, FPCGPerNodeVisualLogs>& Entry : StackToLogs)
 	{
 		const FPCGStack& Stack = Entry.Key;
-
 		if (Stack.HasObject(InNode))
 		{
-			const FPCGPerNodeVisualLogs& NodeLogs = Entry.Value;
-
-			for (int LogIndex = 0; LogIndex < NodeLogs.Num(); ++LogIndex)
+			OutLogs.Append(Entry.Value);
+			const UPCGComponent* RootComponent = Stack.GetRootComponent();
+			for (int LogIndex = 0; LogIndex < Entry.Value.Num(); ++LogIndex)
 			{
-				++LogCounter;
-
-				if (LogCounter >= MaxLogsInSummary)
-				{
-					Summary = FText::Format(FText::FromString(TEXT("{0}\n...")), Summary);
-					return Summary;
-				}
-
-				if (LogCounter > 1)
-				{
-					Summary = FText::Format(FText::FromString(TEXT("{0}\n")), Summary);
-				}
-
-				const UPCGComponent* Component = Stack.GetRootComponent();
-				FText ActorName = (Component && Component->GetOwner()) ? FText::FromString(Component->GetOwner()->GetActorLabel()) : FText::FromString(TEXT("MissingComponent"));
-				
-				OutMinimumVerbosity = FMath::Min(OutMinimumVerbosity, NodeLogs[LogIndex].Verbosity);
-				const FText VerbosityText = NodeLogs[LogIndex].Verbosity == ELogVerbosity::Warning ? FText::FromString(TEXT("Warning")) : FText::FromString(TEXT("Error"));
-
-				Summary = FText::Format(LOCTEXT("NodeTooltipLogWithActor", "{0}[{1}] {2}: {3}"), Summary, ActorName, VerbosityText, NodeLogs[LogIndex].Message);
+				OutComponents.Add(RootComponent);
 			}
 		}
 	}
-
-	return Summary;
 }
 
-FText FPCGNodeVisualLogs::GetLogsSummaryText(const FPCGStack& InBaseStack, ELogVerbosity::Type* OutMinimumVerbosity) const
+FText FPCGNodeVisualLogs::GetSummaryText(const FPCGPerNodeVisualLogs& InLogs, const TArray<const UPCGComponent*>* InComponents, ELogVerbosity::Type* OutMinimumVerbosity)
 {
-	FReadScopeLock ScopedReadLock(LogsLock);
+	check(!InComponents || InLogs.Num() == InComponents->Num());
 
 	FText Summary = FText::GetEmpty();
+	ELogVerbosity::Type MinVerbosity = ELogVerbosity::All;
 
-	if (OutMinimumVerbosity)
-	{
-		*OutMinimumVerbosity = ELogVerbosity::All;
-	}
+	const int32 NumLogs = FMath::Min(InLogs.Num(), PCGNodeVisualLogsConstants::MaxLogsInSummary);
 
-	int32 LogCounter = 0;
-
-	for (const TPair<FPCGStack, FPCGPerNodeVisualLogs>& Entry : StackToLogs)
+	for (int32 LogIndex = 0; LogIndex < NumLogs; ++LogIndex)
 	{
-		if (!Entry.Key.BeginsWith(InBaseStack))
-		{
-			continue;
-		}
+		const FPCGNodeLogEntry& LogEntry = InLogs[LogIndex];
 
-		const FPCGPerNodeVisualLogs& NodeLogs = Entry.Value;
-
-		for (int32 LogIndex = 0; LogIndex < NodeLogs.Num(); ++LogIndex)
+		const FText& MessageVerbosity = (LogEntry.Verbosity == ELogVerbosity::Warning ? PCGNodeVisualLogsConstants::Warning : PCGNodeVisualLogsConstants::Error);
+		if (InComponents)
 		{
-			++LogCounter;
+			const UPCGComponent* Component = (*InComponents)[LogIndex];
+			const FText ActorName = (Component && Component->GetOwner()) ? FText::FromString(Component->GetOwner()->GetActorLabel()) : LOCTEXT("PCGLogMissingComponent", "MissingComponent");
 
-			if (LogCounter > MaxLogsInSummary)
+			if (Summary.IsEmpty())
 			{
-				Summary = FText::Format(FText::FromString(TEXT("{0}\n...")), Summary);
-				return Summary;
+				Summary = FText::Format(LOCTEXT("NodeTooltipLogWithActorEmpty", "[{0}] {1}: {2}"), ActorName, MessageVerbosity, LogEntry.Message);
 			}
-
-			if (LogCounter > 1)
+			else
 			{
-				Summary = FText::Format(FText::FromString(TEXT("{0}\n")), Summary);
+				Summary = FText::Format(LOCTEXT("NodeTooltipLogWithActor", "{0}\n[{1}] {2}: {3}"), Summary, ActorName, MessageVerbosity, LogEntry.Message);
 			}
-
-			if (OutMinimumVerbosity)
+		}
+		else
+		{
+			if (Summary.IsEmpty())
 			{
-				*OutMinimumVerbosity = FMath::Min(*OutMinimumVerbosity, NodeLogs[LogIndex].Verbosity);
+				Summary = FText::Format(LOCTEXT("NodeTooltipLogEmpty", "{1}: {2}"), MessageVerbosity, LogEntry.Message);
+			}
+			else
+			{
+				Summary = FText::Format(LOCTEXT("NodeTooltipLog", "{0}\n{1}: {2}"), Summary, MessageVerbosity, LogEntry.Message);
 			}
-
-			const FText MessageVerbosity = NodeLogs[LogIndex].Verbosity == ELogVerbosity::Warning ? FText::FromString(TEXT("Warning")) : FText::FromString(TEXT("Error"));
-			Summary = FText::Format(LOCTEXT("NodeTooltipLog", "{0}{1}: {2}"), Summary, MessageVerbosity, NodeLogs[LogIndex].Message);
 		}
+
+		MinVerbosity = FMath::Min(MinVerbosity, LogEntry.Verbosity);
+	}
+
+	// Check log level for all entries, not only the first entries
+	for (int32 LogIndex = NumLogs; LogIndex < InLogs.Num(); ++LogIndex)
+	{
+		const FPCGNodeLogEntry& LogEntry = InLogs[LogIndex];
+		MinVerbosity = FMath::Min(MinVerbosity, LogEntry.Verbosity);
+	}
+
+	// Finally, if we had at most the limit of entries, we'll add an ellipsis.
+	if (InLogs.Num() > PCGNodeVisualLogsConstants::MaxLogsInSummary)
+	{
+		Summary = FText::Format(LOCTEXT("NodeTooltipEllipsis", "{0}\n..."), Summary);
+	}
+
+	if (OutMinimumVerbosity)
+	{
+		*OutMinimumVerbosity = MinVerbosity;
 	}
 
 	return Summary;
 }
 
+
+FText FPCGNodeVisualLogs::GetLogsSummaryText(const UPCGNode* InNode, ELogVerbosity::Type& OutMinimumVerbosity) const
+{
+	FPCGPerNodeVisualLogs Logs;
+	TArray<const UPCGComponent*> Components;
+	GetLogs(InNode, Logs, Components);
+
+	return GetSummaryText(Logs, &Components, &OutMinimumVerbosity);
+}
+
+FText FPCGNodeVisualLogs::GetLogsSummaryText(const FPCGStack& InBaseStack, ELogVerbosity::Type* OutMinimumVerbosity) const
+{
+	return GetSummaryText(GetLogs(InBaseStack), nullptr, OutMinimumVerbosity);
+}
+
 void FPCGNodeVisualLogs::ClearLogs(const FPCGStack& InPCGStack)
 {
 	TSet<const UPCGNode*> TouchedNodes;
diff --git a/Engine/Plugins/PCG/Source/PCG/Public/Graph/PCGStackContext.h b/Engine/Plugins/PCG/Source/PCG/Public/Graph/PCGStackContext.h
--- a/Engine/Plugins/PCG/Source/PCG/Public/Graph/PCGStackContext.h
+++ b/Engine/Plugins/PCG/Source/PCG/Public/Graph/PCGStackContext.h
@@ -24,7 +24,7 @@ struct PCG_API FPCGStackFrame
 {
 	GENERATED_BODY()
 
-	FPCGStackFrame() {}
+	FPCGStackFrame() = default;
 
 	explicit FPCGStackFrame(const UObject* InObject)
 	{
@@ -147,6 +147,9 @@ public:
 	/** If current frame (top of stack) corresponds to a node returns that node, otherwise returns null. */
 	const UPCGNode* GetCurrentFrameNode() const;
 
+	/** Gets the node from the node frame closest to the top of the stack (most recent), or null if no such node present. */
+	const UPCGNode* GetNodeForCurrentFrame() const;
+
 	/** Stack has a frame corresponding to the given object. */
 	bool HasObject(const UObject* InObject) const;
 
diff --git a/Engine/Plugins/PCG/Source/PCG/Public/PCGSubsystem.h b/Engine/Plugins/PCG/Source/PCG/Public/PCGSubsystem.h
--- a/Engine/Plugins/PCG/Source/PCG/Public/PCGSubsystem.h
+++ b/Engine/Plugins/PCG/Source/PCG/Public/PCGSubsystem.h
@@ -338,7 +338,7 @@ public:
 	void ClearExecutionMetadata(UPCGComponent* InComponent);
 
 	/** Get a list of stacks that were executed during the last execution. */
-	TArray<FPCGStack> GetExecutedStacks(const UPCGComponent* InComponent, const UPCGGraph* InSubgraph);
+	TArray<FPCGStack> GetExecutedStacks(const UPCGComponent* InComponent, const UPCGGraph* InSubgraph, bool bOnlyWithSubgraphAsCurrentFrame = true);
 	TArray<FPCGStack> GetExecutedStacks(const FPCGStack& BeginningWithStack);
 	void ClearExecutedStacks(FPCGStack BeginningWithStack);
 	void ClearExecutedStacks(const UPCGGraph* InContainingGraph);
diff --git a/Engine/Plugins/PCG/Source/PCG/Public/Utils/PCGNodeVisualLogs.h b/Engine/Plugins/PCG/Source/PCG/Public/Utils/PCGNodeVisualLogs.h
--- a/Engine/Plugins/PCG/Source/PCG/Public/Utils/PCGNodeVisualLogs.h
+++ b/Engine/Plugins/PCG/Source/PCG/Public/Utils/PCGNodeVisualLogs.h
@@ -16,6 +16,8 @@ class UPCGNode;
 
 struct FPCGNodeLogEntry
 {
+	FPCGNodeLogEntry() = default;
+
 	explicit FPCGNodeLogEntry(const FText& InMessage, ELogVerbosity::Type InVerbosity)
 		: Message(InMessage)
 		, Verbosity(InVerbosity)
@@ -44,6 +46,15 @@ public:
 	/** Returns true if an issue with given severity was logged during last execution. */
 	bool HasLogsOfVerbosity(const FPCGStack& InPCGStack, ELogVerbosity::Type InVerbosity) const;
 
+	/** Returns all logs from a given stack */
+	FPCGPerNodeVisualLogs GetLogs(const FPCGStack& InPCGStack) const;
+
+	/** Returns all logs for a given node */
+	void GetLogs(const UPCGNode* InNode, FPCGPerNodeVisualLogs& OutLogs, TArray<const UPCGComponent*>& OutComponents) const;
+
+	/** Calls a function for all stacks that match in the logs or until the function returns false. */
+	void ForAllMatchingLogs(const FPCGStack& InPCGStack, TFunctionRef<bool(const FPCGStack&, const FPCGPerNodeVisualLogs&)> InFunc) const;
+
 	/** Summary text of all visual logs produced while executing the provided base stack, appropriate for display in graph editor tooltip. */
 	FText GetLogsSummaryText(const FPCGStack& InBaseStack, ELogVerbosity::Type* OutMinimumVerbosity = nullptr) const;
 
@@ -59,9 +70,10 @@ public:
 	/** Clear all errors and warnings corresponding to the given component. */
 	void ClearLogs(const UPCGComponent* InComponent);
 
-private:
-	const int MaxLogsInSummary = 8;
+	/** Generates summary text from log list, with optional component pointers to show actor label. */
+	static FText GetSummaryText(const FPCGPerNodeVisualLogs& InLogs, const TArray<const UPCGComponent*>* InComponents, ELogVerbosity::Type* OutMinimumVerbosity = nullptr);
 
+private:
 	TMap<FPCGStack, FPCGPerNodeVisualLogs> StackToLogs;
 	mutable FRWLock LogsLock;
 };
diff --git a/Engine/Plugins/PCG/Source/PCGEditor/Private/PCGEditor.cpp b/Engine/Plugins/PCG/Source/PCGEditor/Private/PCGEditor.cpp
--- a/Engine/Plugins/PCG/Source/PCGEditor/Private/PCGEditor.cpp
+++ b/Engine/Plugins/PCG/Source/PCGEditor/Private/PCGEditor.cpp
@@ -266,6 +266,14 @@ UPCGEditorGraph* FPCGEditor::GetPCGEditorGraph()
 	return PCGEditorGraph;
 }
 
+void FPCGEditor::SetStackBeingInspectedFromAnotherEditor(const FPCGStack& FullStack)
+{
+	if (DebugObjectTreeWidget)
+	{
+		DebugObjectTreeWidget->SetDebugObjectFromStackFromAnotherEditor(FullStack);
+	}
+}
+
 void FPCGEditor::SetStackBeingInspected(const FPCGStack& FullStack)
 {
 	if (FullStack == StackBeingInspected)
@@ -2820,15 +2828,38 @@ void FPCGEditor::OnNodeDoubleClicked(UEdGraphNode* Node)
 {
 	if (Node != nullptr)
 	{
-		if (UObject* Object = Node->GetJumpTargetForDoubleClick())
+		UObject* Object = Node->GetJumpTargetForDoubleClick();
+
+		// "Normal" node
+		if (UPCGSettings* PCGSettings = Cast<UPCGSettings>(Object))
+		{
+			JumpToDefinition(PCGSettings->GetClass());
+		}
+		else
 		{
-			if (UPCGSettings* PCGSettings = Cast<UPCGSettings>(Object))
+			const UPCGEditorGraphNodeBase* PCGEditorGraphNode = Cast<const UPCGEditorGraphNodeBase>(Node);
+			const UPCGNode* PCGNode = PCGEditorGraphNode ? PCGEditorGraphNode->GetPCGNode() : nullptr;
+
+			FPCGStack StackToInspect;
+
+			// If we're inspecting, we'll try to find a match in the stacks for subgraphs instead of relying on the static/template subgraph
+			if (GetStackBeingInspected() && DebugObjectTreeWidget)
 			{
-				JumpToDefinition(PCGSettings->GetClass());
+				DebugObjectTreeWidget->GetFirstStackFromSelection(PCGNode, /*Graph=*/nullptr, StackToInspect);
+				Object = const_cast<UPCGGraph*>(StackToInspect.GetGraphForCurrentFrame());
 			}
-			else
+
+			if (Object)
 			{
+				// Open other editor...
 				GEditor->GetEditorSubsystem<UAssetEditorSubsystem>()->OpenEditorForAsset(Object);
+				IAssetEditorInstance* EditorInstance = GEditor->GetEditorSubsystem<UAssetEditorSubsystem>()->FindEditorForAsset(Object, /*bFocusIfOpen*/true);
+
+				FPCGEditor* OtherPCGEditor = static_cast<FPCGEditor*>(EditorInstance);
+				if (OtherPCGEditor && !StackToInspect.GetStackFrames().IsEmpty())
+				{
+					OtherPCGEditor->SetStackBeingInspectedFromAnotherEditor(StackToInspect);
+				}
 			}
 		}
 	}
diff --git a/Engine/Plugins/PCG/Source/PCGEditor/Private/PCGEditor.h b/Engine/Plugins/PCG/Source/PCGEditor/Private/PCGEditor.h
--- a/Engine/Plugins/PCG/Source/PCGEditor/Private/PCGEditor.h
+++ b/Engine/Plugins/PCG/Source/PCGEditor/Private/PCGEditor.h
@@ -51,6 +51,9 @@ public:
 	/** Sets the execution stack that want to inspect. */
 	void SetStackBeingInspected(const FPCGStack& FullStack);
 
+	/** Sets the execution stack from another editor, which will set directly in the debug object tree view. */
+	void SetStackBeingInspectedFromAnotherEditor(const FPCGStack& FullStack);
+
 	/** Clear current inspection. */
 	void ClearStackBeingInspected();
 
diff --git a/Engine/Plugins/PCG/Source/PCGEditor/Private/SPCGEditorGraphDebugObjectTree.cpp b/Engine/Plugins/PCG/Source/PCGEditor/Private/SPCGEditorGraphDebugObjectTree.cpp
--- a/Engine/Plugins/PCG/Source/PCGEditor/Private/SPCGEditorGraphDebugObjectTree.cpp
+++ b/Engine/Plugins/PCG/Source/PCGEditor/Private/SPCGEditorGraphDebugObjectTree.cpp
@@ -10,6 +10,7 @@
 
 #include "PCGEditor.h"
 #include "PCGEditorGraph.h"
+#include "PCGEditorStyle.h"
 
 #include "PropertyCustomizationHelpers.h"
 #include "Selection.h"
@@ -29,56 +30,107 @@ namespace PCGEditorGraphDebugObjectTree
 	{
 		const UObject* ItemObject = Item->GetObject();
 
-		if (Item->IsDebuggable())
+		if (!IsValid(ItemObject))
 		{
 			OutBrush = FAppStyle::Get().GetBrush("LevelEditor.Tabs.Debug");
 		}
-		else if (IsValid(ItemObject) && ItemObject->IsA<AActor>())
+		else if (ItemObject->IsA<AActor>())
 		{
 			OutBrush = FSlateIconFinder::FindIconBrushForClass(ItemObject->GetClass());
 		}
 		else
 		{
-			OutBrush = Item->IsExpanded() ? FAppStyle::Get().GetBrush("Icons.FolderOpen") : FAppStyle::Get().GetBrush("Icons.FolderClosed");
+			OutBrush = FSlateIconFinder::FindIconBrushForClass(UPCGGraph::StaticClass());
 		}
 
 		OutColorAndOpacity = Item->IsDebuggable() ? FLinearColor::White : FLinearColor(1.0f, 1.0f, 1.0f, 0.5f);
 	}
 
-	void GetErrorIconState(FPCGEditorGraphDebugObjectItemPtr Item, EVisibility& OutIconVisibility, FText& OutIconTooltipText, FLinearColor& OutColorAndOpacity)
+	void GetErrorInfo(FPCGEditorGraphDebugObjectItemPtr Item, EVisibility& OutIconVisibility, FText& OutIconTooltipText, FLinearColor& OutIconColorAndOpacity, FLinearColor& OutRowTextColorAndOpacity)
 	{
+		OutIconVisibility = EVisibility::Hidden;
 		OutIconTooltipText = FText();
-		OutColorAndOpacity = FLinearColor::White;
+		OutIconColorAndOpacity = FLinearColor::White;
+		OutRowTextColorAndOpacity = FLinearColor::White;
 
+		if (!ensure(Item->GetPCGStack()))
+		{
+			return;
+		}
+
+		ELogVerbosity::Type MinLocalVerbosity = ELogVerbosity::All;
 		ELogVerbosity::Type MinVerbosity = ELogVerbosity::All;
 
 		if (const UPCGSubsystem* Subsystem = UPCGSubsystem::GetActiveEditorInstance())
 		{
-			if (Item->GetPCGStack())
+			FPCGPerNodeVisualLogs Logs;
+			Subsystem->GetNodeVisualLogs().ForAllMatchingLogs(*Item->GetPCGStack(), [&Item, &Logs, &MinLocalVerbosity](const FPCGStack& InStack, const FPCGPerNodeVisualLogs& InLogs)
 			{
-				OutIconTooltipText = Subsystem->GetNodeVisualLogs().GetLogsSummaryText(*Item->GetPCGStack(), &MinVerbosity);
-			}
-			else
-			{
-				for (const FPCGEditorGraphDebugObjectItemPtr& Child : Item->GetChildren())
+				// Need to verify this, and clarify this when it's a loop item.
+				// i.e. should get the graph + loop index and compare
+				if (Item->GetPCGStack()->GetGraphForCurrentFrame() == InStack.GetGraphForCurrentFrame())
 				{
-					if (Child->GetPCGStack())
+					for (const FPCGNodeLogEntry& LogEntry : InLogs)
 					{
-						OutIconTooltipText = Subsystem->GetNodeVisualLogs().GetLogsSummaryText(*Child->GetPCGStack(), &MinVerbosity);
-						break;
+						MinLocalVerbosity = FMath::Min(MinLocalVerbosity, LogEntry.Verbosity);
 					}
 				}
-			}
+
+				Logs.Append(InLogs);
+				return true;
+			});
+
+			OutIconTooltipText = Subsystem->GetNodeVisualLogs().GetSummaryText(Logs, nullptr, &MinVerbosity);
 		}
 
 		OutIconVisibility = OutIconTooltipText.IsEmpty() ? EVisibility::Hidden : EVisibility::Visible;
 
+		constexpr FLinearColor WarningColor(1.0f, 0.75f, 0.0f, 0.9f);
+		constexpr FLinearColor ErrorColor(1.0f, 0.0f, 0.0f, 0.9f);
+
+		if (MinLocalVerbosity < ELogVerbosity::All)
+		{
+			OutRowTextColorAndOpacity = MinLocalVerbosity <= ELogVerbosity::Error ? ErrorColor : WarningColor;
+		}
+
 		if (MinVerbosity < ELogVerbosity::All)
 		{
-			constexpr FLinearColor WarningColor(1.0f, 0.75f, 0.0f, 0.9f);
-			constexpr FLinearColor ErrorColor(1.0f, 0.0f, 0.0f, 0.9f);
+			OutIconColorAndOpacity = MinVerbosity <= ELogVerbosity::Error ? ErrorColor : WarningColor;
+		}
+	}
 
-			OutColorAndOpacity = MinVerbosity <= ELogVerbosity::Error ? ErrorColor : WarningColor;
+	FString GetStringFromName(FName InName, bool bForSorting)
+	{
+		if (bForSorting)
+		{
+			return InName.GetPlainNameString() + FString::Printf(TEXT("%06d"), InName.GetNumber());
+		}
+		else
+		{
+			return InName.ToString();
+		}
+	}
+
+	FSlateFontInfo GetRowFont(FPCGEditorGraphDebugObjectItemPtr Item)
+	{
+		const bool bUseBold = Item->IsSelected();
+		const bool bUseItalic = Item->IsDynamic();
+
+		if (!bUseBold && !bUseItalic)
+		{
+			return FAppStyle::GetFontStyle("NormalFont");
+		}
+		else if (bUseBold && !bUseItalic)
+		{
+			return FAppStyle::GetFontStyle("NormalFontBold");
+		}
+		else if (!bUseBold && bUseItalic)
+		{
+			return FAppStyle::GetFontStyle("NormalFontItalic");
+		}
+		else
+		{
+			return FAppStyle::GetFontStyle("NormalFontBoldItalic");
 		}
 	}
 }
@@ -122,7 +174,7 @@ void FPCGEditorGraphDebugObjectItem::SortChildren(bool bIsAscending, bool bIsRec
 		}
 
 		// Otherwise fall back to alphanumeric order.
-		return (InLHS->GetLabel() < InRHS->GetLabel()) == bIsAscending;
+		return (InLHS->GetLabel(/*bForSorting=*/true) < InRHS->GetLabel(/*bForSorting=*/true)) == bIsAscending;
 	});
 
 	if (bIsRecursive)
@@ -142,16 +194,16 @@ void FPCGEditorGraphDebugObjectItem::OnObjectsReplaced(const TMap<UObject*, UObj
 	}
 }
 
-FString FPCGEditorGraphDebugObjectItem_Actor::GetLabel() const
+FString FPCGEditorGraphDebugObjectItem_Actor::GetLabel(bool bForSorting) const
 {
 	return Actor.IsValid() ? Actor->GetActorNameOrLabel() : FString();
 }
 
-FString FPCGEditorGraphDebugObjectItem_PCGComponent::GetLabel() const
+FString FPCGEditorGraphDebugObjectItem_PCGComponent::GetLabel(bool bForSorting) const
 {
 	if (PCGComponent.IsValid() && PCGGraph.IsValid())
 	{
-		return PCGComponent->GetName() + FString(TEXT(" - ") + PCGGraph->GetName());
+		return PCGEditorGraphDebugObjectTree::GetStringFromName(PCGComponent->GetFName(), bForSorting) + FString(TEXT(" - ")) + PCGEditorGraphDebugObjectTree::GetStringFromName(PCGGraph->GetFName(), bForSorting);
 	}
 
 	return FString();
@@ -167,17 +219,24 @@ void FPCGEditorGraphDebugObjectItem_PCGComponent::OnObjectsReplaced(const TMap<U
 	}
 }
 
-FString FPCGEditorGraphDebugObjectItem_PCGSubgraph::GetLabel() const
+FString FPCGEditorGraphDebugObjectItem_PCGSubgraph::GetLabel(bool bForSorting) const
 {
 	if (PCGNode.IsValid() && PCGGraph.IsValid())
 	{
-		return PCGNode->GetName() + FString(TEXT(" - ") + PCGGraph->GetName());
+		if (PCGNode->HasAuthoredTitle())
+		{
+			return PCGEditorGraphDebugObjectTree::GetStringFromName(PCGGraph->GetFName(), bForSorting) + FString(TEXT(" - ")) + PCGNode->GetAuthoredTitleLine().ToString();
+		}
+		else
+		{
+			return PCGEditorGraphDebugObjectTree::GetStringFromName(PCGGraph->GetFName(), bForSorting) + FString(TEXT(" - ")) + PCGEditorGraphDebugObjectTree::GetStringFromName(PCGNode->GetFName(), bForSorting);
+		}
 	}
 
 	return FString();
 }
 
-FString FPCGEditorGraphDebugObjectItem_PCGLoopIndex::GetLabel() const
+FString FPCGEditorGraphDebugObjectItem_PCGLoopIndex::GetLabel(bool bForSorting) const
 {
 	return FString::Format(TEXT("{0}"), { LoopIndex });
 }
@@ -198,7 +257,8 @@ void SPCGEditorGraphDebugObjectItemRow::Construct(const FArguments& InArgs, cons
 	FText ErrorIconTooltipText;
 	EVisibility ErrorIconVisibility;
 	FLinearColor ErrorIconColorAndOpacity;
-	PCGEditorGraphDebugObjectTree::GetErrorIconState(InItem, ErrorIconVisibility, ErrorIconTooltipText, ErrorIconColorAndOpacity);
+	FLinearColor RowTextColorAndOpacity;
+	PCGEditorGraphDebugObjectTree::GetErrorInfo(InItem, ErrorIconVisibility, ErrorIconTooltipText, ErrorIconColorAndOpacity, RowTextColorAndOpacity);
 	
 	ChildSlot
 	[
@@ -217,8 +277,9 @@ void SPCGEditorGraphDebugObjectItemRow::Construct(const FArguments& InArgs, cons
 		[
 			SNew(STextBlock)
 			.Text(FText::FromString(Item->GetLabel()))
+			.Font_Lambda([this]() { return PCGEditorGraphDebugObjectTree::GetRowFont(Item); })
 			// Highlight based on data available for currently inspected node. Computed dynamically in lambda to respond to inspection changes.
-			.ColorAndOpacity_Lambda([this] { return Item->IsGrayedOut() ? FColor(75, 75, 75) : FColor::White; })
+			.ColorAndOpacity_Lambda([this, RowTextColorAndOpacity](){ return Item->IsGrayedOut() ? FColor(75, 75, 75) : RowTextColorAndOpacity; })
 		]
 		+SHorizontalBox::Slot()
 		.AutoWidth()
@@ -290,11 +351,6 @@ void SPCGEditorGraphDebugObjectTree::Construct(const FArguments& InArgs, TShared
 		LOCTEXT("SetDebugObject", "Set debug object from Level Editor selection."),
 		TAttribute<bool>::Create(TAttribute<bool>::FGetter::CreateSP(this, &SPCGEditorGraphDebugObjectTree::IsSetDebugObjectFromSelectionButtonEnabled)));
 
-	const TSharedRef<SWidget> BrowseButton = PropertyCustomizationHelpers::MakeBrowseButton(
-		FSimpleDelegate::CreateSP(this, &SPCGEditorGraphDebugObjectTree::SelectedDebugObject_OnClicked),
-		LOCTEXT("DebugSelectActor", "Select and frame the debug actor in the Level Editor."),
-		TAttribute<bool>::Create(TAttribute<bool>::FGetter::CreateSP(this, &SPCGEditorGraphDebugObjectTree::IsSelectDebugObjectButtonEnabled)));
-
 	ChildSlot
 	[
 		SNew(SVerticalBox)
@@ -315,7 +371,17 @@ void SPCGEditorGraphDebugObjectTree::Construct(const FArguments& InArgs, TShared
 			.VAlign(VAlign_Center)
 			.Padding(4.0f)
 			[
-				BrowseButton
+				SNew(SButton)
+				.ButtonStyle(FAppStyle::Get(), "SimpleButton")
+				.OnClicked(this, &SPCGEditorGraphDebugObjectTree::FocusOnDebugObject_OnClicked)
+				.IsEnabled(this, &SPCGEditorGraphDebugObjectTree::IsFocusOnDebugObjectButtonEnabled)
+				.ToolTipText(LOCTEXT("DebugSelectActor", "Select and frame the debug actor in the Level Editor."))
+				.ContentPadding(0.0f)
+				[
+					SNew(SImage)
+					.Image(FPCGEditorStyle::Get().GetBrush("PCG.Editor.ZoomToSelection"))
+					.ColorAndOpacity(FSlateColor::UseForeground())
+				]
 			]
 		]
 		+SVerticalBox::Slot()
@@ -376,31 +442,6 @@ void SPCGEditorGraphDebugObjectTree::Tick(const FGeometry& AllottedGeometry, con
 	}
 }
 
-void SPCGEditorGraphDebugObjectTree::SetDebugObjectSelection(const FPCGStack& FullStack)
-{
-	bDisableDebugObjectChangeNotification = true;
-
-	for (FPCGEditorGraphDebugObjectItemPtr& Item : AllGraphItems)
-	{
-		const FPCGStack* ItemStack = Item->GetPCGStack();
-		const bool bSelected = ItemStack && *ItemStack == FullStack;
-
-		DebugObjectTreeView->SetItemSelection(Item, bSelected);
-
-		if (bSelected)
-		{
-			FPCGEditorGraphDebugObjectItemPtr Parent = Item->GetParent();
-			while (Parent)
-			{
-				DebugObjectTreeView->SetItemExpansion(Parent, true);
-				Parent = Parent->GetParent();
-			}
-		}
-	}
-
-	bDisableDebugObjectChangeNotification = false;
-}
-
 void SPCGEditorGraphDebugObjectTree::OnObjectsReplaced(const TMap<UObject*, UObject*>& ReplacementMap)
 {
 	for (FPCGEditorGraphDebugObjectItemPtr& GraphItem : AllGraphItems)
@@ -439,7 +480,7 @@ void SPCGEditorGraphDebugObjectTree::SetNodeBeingInspected(const UPCGNode* InPCG
 	RequestRefresh();
 }
 
-void SPCGEditorGraphDebugObjectTree::SelectedDebugObject_OnClicked() const
+FReply SPCGEditorGraphDebugObjectTree::FocusOnDebugObject_OnClicked() const
 {
 	if (UPCGComponent* PCGComponent = PCGEditor.Pin()->GetPCGComponentBeingInspected())
 	{
@@ -452,9 +493,11 @@ void SPCGEditorGraphDebugObjectTree::SelectedDebugObject_OnClicked() const
 			GEditor->SelectComponent(PCGComponent, /*bInSelected=*/true, /*bNotify=*/true, /*bSelectEvenIfHidden=*/true);
 		}
 	}
+
+	return FReply::Handled();
 }
 
-bool SPCGEditorGraphDebugObjectTree::IsSelectDebugObjectButtonEnabled() const
+bool SPCGEditorGraphDebugObjectTree::IsFocusOnDebugObjectButtonEnabled() const
 {
 	return PCGEditor.IsValid() && (PCGEditor.Pin()->GetPCGComponentBeingInspected() != nullptr);
 }
@@ -500,29 +543,127 @@ void SPCGEditorGraphDebugObjectTree::SetDebugObjectFromSelection_OnClicked()
 
 			for (const FPCGStack& Stack : StackContext.GetStacks())
 			{
-				FPCGEditorGraphDebugObjectItemPtr* DebugObjectItem = AllGraphItems.FindByPredicate([&PCGComponent, &Stack](const FPCGEditorGraphDebugObjectItemPtr Item)
+				if (SetDebugObjectFromStack(Stack))
 				{
-					return Item->GetPCGComponent() == PCGComponent && *Item->GetPCGStack() == Stack;
-				});
+					break;
+				}
+			}
+		}
+	}
+}
 
-				if (DebugObjectItem)
-				{
-					DebugObjectTreeView->SetSingleExpandedItem(*DebugObjectItem);
+bool SPCGEditorGraphDebugObjectTree::SetDebugObjectFromStackFromAnotherEditor(const FPCGStack& InStack)
+{
+	RefreshTree();
+	return SetDebugObjectFromStack(InStack);
+}
 
-					FPCGEditorGraphDebugObjectItemPtr Parent = DebugObjectItem->Get()->GetParent();
-					while (Parent)
-					{
-						DebugObjectTreeView->SetItemExpansion(Parent, true);
-						Parent = Parent->GetParent();
-					}
+bool SPCGEditorGraphDebugObjectTree::SetDebugObjectFromStack(const FPCGStack& InStack)
+{
+	if (FPCGEditorGraphDebugObjectItemPtr Item = GetItemFromStack(InStack))
+	{
+		ExpandAndSelectDebugObject(Item);
+		return true;
+	}
+	else
+	{
+		return false;
+	}
+}
 
-					DebugObjectTreeView->SetSelection(*DebugObjectItem);
-					DebugObjectTreeView->RequestScrollIntoView(*DebugObjectItem);
-					break;
-				}
+void SPCGEditorGraphDebugObjectTree::ExpandAndSelectDebugObject(FPCGEditorGraphDebugObjectItemPtr InItem)
+{
+	if (!InItem)
+	{
+		return;
+	}
+
+	FPCGEditorGraphDebugObjectItemPtr Parent = InItem->GetParent();
+	while (Parent)
+	{
+		DebugObjectTreeView->SetItemExpansion(Parent, true);
+		Parent = Parent->GetParent();
+	}
+
+	DebugObjectTreeView->SetSelection(InItem);
+	DebugObjectTreeView->RequestScrollIntoView(InItem);
+}
+
+FPCGEditorGraphDebugObjectItemPtr SPCGEditorGraphDebugObjectTree::GetItemFromStack(const FPCGStack& InStack) const
+{
+	for (const FPCGEditorGraphDebugObjectItemPtr& Item : AllGraphItems)
+	{
+		if (Item->GetPCGStack() && *(Item->GetPCGStack()) == InStack)
+		{
+			return Item;
+		}
+	}
+
+	return nullptr;
+}
+
+bool SPCGEditorGraphDebugObjectTree::GetFirstStackFromSelection(const UPCGNode* InNode, const UPCGGraph* InGraph, FPCGStack& OutStack) const
+{
+	if (SelectedStack.GetStackFrames().IsEmpty())
+	{
+		return false;
+	}
+
+	for (const FPCGEditorGraphDebugObjectItemPtr& Item : AllGraphItems)
+	{
+		const FPCGStack* ItemStack = Item->GetPCGStack();
+
+		// Early validation - common stack
+		if (!ItemStack || !ItemStack->BeginsWith(SelectedStack))
+		{
+			continue;
+		}
+
+		// Then let's find the occurrence of the node
+		const TArray<FPCGStackFrame>& ItemStackFrames = ItemStack->GetStackFrames();
+		
+		const UPCGNode* ItemNode = nullptr;
+		int NodeIndex = SelectedStack.GetStackFrames().Num();
+		while (NodeIndex < ItemStackFrames.Num())
+		{
+			ItemNode = Cast<const UPCGNode>(ItemStackFrames[NodeIndex].Object.Get());
+			if (ItemNode)
+			{
+				break;
+			}
+			else
+			{
+				++NodeIndex;
 			}
 		}
+
+		if (ItemNode != InNode)
+		{
+			continue;
+		}
+
+		// And make sure the graph that's somewhere after it in the stack is present
+		// Implementation note: the graph is either the next (static) or the second next in the stack.
+		const UPCGGraph* ItemGraph = nullptr;
+
+		if (ItemStackFrames.IsValidIndex(NodeIndex + 1))
+		{
+			ItemGraph = Cast<const UPCGGraph>(ItemStackFrames[NodeIndex + 1].Object.Get());
+		}
+
+		if (!ItemGraph && ItemStackFrames.IsValidIndex(NodeIndex + 2))
+		{
+			ItemGraph = Cast<const UPCGGraph>(ItemStackFrames[NodeIndex + 2].Object.Get());
+		}
+
+		if (ItemGraph && (!InGraph || ItemGraph == InGraph))
+		{
+			OutStack = *Item->GetPCGStack();
+			return true;
+		}
 	}
+
+	return false;
 }
 
 void SPCGEditorGraphDebugObjectTree::UpdateIsSetDebugObjectFromSelectionEnabled()
@@ -672,7 +813,7 @@ void SPCGEditorGraphDebugObjectTree::AddStacksToTree(const TArray<FPCGStack>& St
 		};
 
 		// Generic function to create the subgraph/loop item and hook it up to the parent properly
-		auto AddSubgraphOrLoopItemToStack = [&Stack, &InOutStackToItem, &AddToGraphItemsAndAttachItemToParent](const UPCGGraph* InGraph, const UPCGNode* InSubgraphNode, int FrameCutoff, bool bInIsDebuggable, bool bInDisplayGrayedOut)
+		auto AddSubgraphOrLoopItemToStack = [&Stack, &InOutStackToItem, &AddToGraphItemsAndAttachItemToParent](const UPCGGraph* InGraph, const UPCGNode* InSubgraphNode, int FrameCutoff, bool bInIsDebuggable, bool bInDisplayGrayedOut, bool bIsDynamic)
 		{
 			FPCGStack GraphStack = Stack;
 			GraphStack.GetStackFramesMutable().SetNum(FrameCutoff);
@@ -681,7 +822,7 @@ void SPCGEditorGraphDebugObjectTree::AddStacksToTree(const TArray<FPCGStack>& St
 			{
 				FPCGEditorGraphDebugObjectItemPtr GraphItem = InOutStackToItem.Emplace(
 					GraphStack,
-					MakeShared<FPCGEditorGraphDebugObjectItem_PCGSubgraph>(InSubgraphNode, InGraph, GraphStack, bInIsDebuggable, bInDisplayGrayedOut));
+					MakeShared<FPCGEditorGraphDebugObjectItem_PCGSubgraph>(InSubgraphNode, InGraph, GraphStack, bInIsDebuggable, bInDisplayGrayedOut, bIsDynamic));
 
 				AddToGraphItemsAndAttachItemToParent(GraphStack, GraphItem);
 			}
@@ -725,7 +866,7 @@ void SPCGEditorGraphDebugObjectTree::AddStacksToTree(const TArray<FPCGStack>& St
 				// Previous stack was node, therefore static subgraph.
 				else if (const UPCGNode* SubgraphNode = Cast<const UPCGNode>(PreviousStackFrame.Object.Get()))
 				{
-					AddSubgraphOrLoopItemToStack(StackGraph, SubgraphNode, FrameIndex + 1, bIsDebuggable, bDisplayGrayedOut);
+					AddSubgraphOrLoopItemToStack(StackGraph, SubgraphNode, FrameIndex + 1, bIsDebuggable, bDisplayGrayedOut, /*bIsDynamic=*/false);
 				}
 				// Previous stack was loop index, therefore loop subgraph.
 				else if (FrameIndex >= 2 && (PreviousStackFrame.LoopIndex != INDEX_NONE))
@@ -734,7 +875,7 @@ void SPCGEditorGraphDebugObjectTree::AddStacksToTree(const TArray<FPCGStack>& St
 					if (ensure(LoopSubgraphNode))
 					{
 						// Take the stack up to the looped subgraph node, add a item for the node + graph
-						AddSubgraphOrLoopItemToStack(StackGraph, LoopSubgraphNode, FrameIndex - 1, /*bIsDebuggable=*/false, /*bIsGrayedOut=*/false);
+						AddSubgraphOrLoopItemToStack(StackGraph, LoopSubgraphNode, FrameIndex - 1, /*bIsDebuggable=*/false, /*bIsGrayedOut=*/false, /*bIsDynamic=*/true);
 
 						// Take full stack up until this point which will be the unique stack for the loop iteration.
 						FPCGStack LoopIterationStack = Stack;
@@ -756,22 +897,10 @@ void SPCGEditorGraphDebugObjectTree::AddStacksToTree(const TArray<FPCGStack>& St
 					const UPCGNode* DynamicSubgraphNode = Cast<const UPCGNode>(StackFrames[FrameIndex - 2].Object.Get());
 					if (DynamicSubgraphNode)
 					{
-						AddSubgraphOrLoopItemToStack(StackGraph, DynamicSubgraphNode, FrameIndex + 1, bIsDebuggable, bDisplayGrayedOut);
+						AddSubgraphOrLoopItemToStack(StackGraph, DynamicSubgraphNode, FrameIndex + 1, bIsDebuggable, bDisplayGrayedOut, /*bIsDynamic=*/true);
 					}
 				}
 			}
-
-			// Check if any of the remaining stack frames contain valid debug targets - invocations of the current edited graph.
-			// If not then we will abandon this stack, pruning the tree.
-			bRemainingStackContainsEditedGraph = false;
-			for (int RemainingFrameIndex = FrameIndex + 1; RemainingFrameIndex < StackFrames.Num(); ++RemainingFrameIndex)
-			{
-				if (StackFrames[RemainingFrameIndex].Object == GraphBeingEdited)
-				{
-					bRemainingStackContainsEditedGraph = true;
-					break;
-				}
-			}
 		}
 	}
 }
@@ -802,7 +931,7 @@ void SPCGEditorGraphDebugObjectTree::RefreshTree()
 		TMap<AActor*, TSharedPtr<FPCGEditorGraphDebugObjectItem_Actor>> ActorItems;
 		TMap<const FPCGStack, FPCGEditorGraphDebugObjectItemPtr> StackToItem;
 
-		TArray<FPCGStack> GraphStacks = Subsystem->GetExecutedStacks(nullptr, PCGGraph);
+		TArray<FPCGStack> GraphStacks = Subsystem->GetExecutedStacks(nullptr, PCGGraph, /*bOnlyWithSubgraphAsCurrentFrame=*/false);
 		TMap<const UPCGComponent*, TArray<FPCGStack>> ComponentsGraphStacks;
 		for (FPCGStack& GraphStack : GraphStacks)
 		{
@@ -852,7 +981,7 @@ void SPCGEditorGraphDebugObjectTree::SortTreeItems(bool bIsAscending, bool bIsRe
 {
 	RootItems.Sort([bIsAscending](const FPCGEditorGraphDebugObjectItemPtr& InLHS, const FPCGEditorGraphDebugObjectItemPtr& InRHS)
 	{
-		return (InLHS->GetLabel() < InRHS->GetLabel()) == bIsAscending;
+		return (InLHS->GetLabel(/*bForSorting=*/true) < InRHS->GetLabel(/*bForSorting=*/true)) == bIsAscending;
 	});
 
 	if (bIsRecursive)
@@ -871,34 +1000,23 @@ void SPCGEditorGraphDebugObjectTree::RestoreTreeState()
 	ExpandedStacks.Empty(ExpandedStacksBefore.Num());
 	for (const FPCGStack& ExpandedStack : ExpandedStacksBefore)
 	{
-		for (FPCGEditorGraphDebugObjectItemPtr& Item : AllGraphItems)
+		if (FPCGEditorGraphDebugObjectItemPtr Item = GetItemFromStack(ExpandedStack))
 		{
-			const FPCGStack* ItemStack = Item->GetPCGStack();
-
-			if (ItemStack && *ItemStack == ExpandedStack)
-			{
-				DebugObjectTreeView->SetItemExpansion(Item, true);
-			}
+			DebugObjectTreeView->SetItemExpansion(Item, true);
 		}
 	}
 
 	bool bFoundMatchingStack = false;
 
 	// Try to restore user item selection by exact matching.
-	for (FPCGEditorGraphDebugObjectItemPtr& Item : AllGraphItems)
+	if (FPCGEditorGraphDebugObjectItemPtr ExactMatchItem = GetItemFromStack(SelectedStack))
 	{
-		const FPCGStack* ItemStack = Item->GetPCGStack();
-
-		if (ItemStack && SelectedStack == *ItemStack)
+		if (!DebugObjectTreeView->IsItemSelected(ExactMatchItem))
 		{
-			if (!DebugObjectTreeView->IsItemSelected(Item))
-			{
-				DebugObjectTreeView->SetItemSelection(Item, true);
-			}
-			
-			bFoundMatchingStack = true;
-			break;
+			DebugObjectTreeView->SetItemSelection(ExactMatchItem, true);
 		}
+
+		bFoundMatchingStack = true;
 	}
 
 	// Try to restore user item selection by fuzzy matching (e.g. share the same owner) if no exactly matching stack was found.
@@ -1032,10 +1150,10 @@ UPCGGraph* SPCGEditorGraphDebugObjectTree::GetPCGGraph() const
 TSharedRef<ITableRow> SPCGEditorGraphDebugObjectTree::MakeTreeRowWidget(FPCGEditorGraphDebugObjectItemPtr InItem, const TSharedRef<STableViewBase>& InOwnerTable)
 {
 	return SNew( STableRow<TSharedPtr<SPCGEditorGraphDebugObjectItemRow>>, InOwnerTable)
-		[
-			SNew(SPCGEditorGraphDebugObjectItemRow, InOwnerTable, InItem)
-				.OnDoubleClickFunc([this](const FPCGEditorGraphDebugObjectItemPtr& Item) { ExpandAndSelectDebugObject(Item); } )
-		];
+	[
+		SNew(SPCGEditorGraphDebugObjectItemRow, InOwnerTable, InItem)
+		.OnDoubleClickFunc([this](const FPCGEditorGraphDebugObjectItemPtr& Item) { ExpandAndSelectFirstLeafDebugObject(Item); } )
+	];
 }
 
 void SPCGEditorGraphDebugObjectTree::OnGetChildren(FPCGEditorGraphDebugObjectItemPtr InItem, TArray<FPCGEditorGraphDebugObjectItemPtr>& OutChildren) const
@@ -1059,6 +1177,12 @@ void SPCGEditorGraphDebugObjectTree::OnSelectionChanged(FPCGEditorGraphDebugObje
 		PreviouslySelectedStack = SelectedStack;
 	}
 
+	// Unmark the previously selected object
+	if (FPCGEditorGraphDebugObjectItemPtr PreviousItem = GetItemFromStack(SelectedStack))
+	{
+		PreviousItem->SetSelected(false);
+	}
+
 	// Reset selected item information.
 	SelectedStack = FPCGStack();
 	SelectedGraph = nullptr;
@@ -1089,20 +1213,72 @@ void SPCGEditorGraphDebugObjectTree::OnSelectionChanged(FPCGEditorGraphDebugObje
 		}
 	}
 
-	if (bDisableDebugObjectChangeNotification)
-	{
-		return;
-	}
+	// Backup selected stack if we're going to jump elsewhere
+	SelectedStackBeforeCulling = SelectedStack;
+
+	const UPCGGraph* CurrentGraph = PCGEditor.Pin()->GetPCGGraph();
 
 	// Only attempt to inspect stacks that correspond to the edited graph. Other graphs need to be inspected in their own editor.
-	if (InItem && SelectedStack.GetGraphForCurrentFrame() == PCGEditor.Pin()->GetPCGGraph())
+	if (InItem)
 	{
-		PCGEditor.Pin()->SetStackBeingInspected(SelectedStack);
+		// Basically - if the selection is something from "upstream", then we can clear the stack being inspected.
+		if (SelectedStack.HasObject(CurrentGraph))
+		{
+			const UPCGNode* LastNode = nullptr;
+
+			// Two things should happen here:
+			// If the last graph in the stack is the current graph, then we can just use that as the selected stack
+			// Otherwise, we'll cut the stack down to the last instance of this graph, and instead move the editor selection
+			// to the first node after the graph
+			TArray<FPCGStackFrame>& StackFrames = SelectedStack.GetStackFramesMutable();
+			while(!StackFrames.IsEmpty())
+			{
+				const UObject* LastObject = StackFrames.Last().Object.Get();
+				if (LastObject == CurrentGraph)
+				{
+					break;
+				}
+				else if (const UPCGNode* Node = Cast<const UPCGNode>(LastObject))
+				{
+					LastNode = Node;
+				}
+
+				StackFrames.Pop();
+			}
+
+			if (ensure(!SelectedStack.GetStackFrames().IsEmpty()))
+			{
+				PCGEditor.Pin()->SetStackBeingInspected(SelectedStack);
+
+				if (LastNode)
+				{
+					PCGEditor.Pin()->JumpToNode(LastNode);
+				}
+			}
+			else
+			{
+				PCGEditor.Pin()->ClearStackBeingInspected();
+			}
+		}
+		else
+		{
+			// Selected an upstream item - let's clear the stack
+			PCGEditor.Pin()->ClearStackBeingInspected();
+		}
 	}
 	else
 	{
 		PCGEditor.Pin()->ClearStackBeingInspected();
 	}
+
+	// Finally mark the selected item as selected
+	if (SelectedStack.GetGraphForCurrentFrame() == CurrentGraph)
+	{
+		if (FPCGEditorGraphDebugObjectItemPtr CurrentSelectedItem = GetItemFromStack(SelectedStack))
+		{
+			CurrentSelectedItem->SetSelected(true);
+		}
+	}
 }
 
 void SPCGEditorGraphDebugObjectTree::OnExpansionChanged(FPCGEditorGraphDebugObjectItemPtr InItem, bool bInIsExpanded)
@@ -1146,40 +1322,50 @@ void SPCGEditorGraphDebugObjectTree::OnSetExpansionRecursive(FPCGEditorGraphDebu
 	}
 }
 
-void SPCGEditorGraphDebugObjectTree::ExpandAndSelectDebugObject(FPCGEditorGraphDebugObjectItemPtr InItem)
+void SPCGEditorGraphDebugObjectTree::ExpandAndSelectFirstLeafDebugObject(FPCGEditorGraphDebugObjectItemPtr InItem)
 {
 	if (!InItem.IsValid())
 	{
 		return;
 	}
 
-	// Expand this item in the tree view.
-	OnSetExpansionRecursive(InItem, /*bInExpand=*/true);
+	// Regardless of what happens, we'll expand the currently selected item.
+	DebugObjectTreeView->SetItemExpansion(InItem, true);
 
-	FPCGEditorGraphDebugObjectItemPtr Item = InItem;
-	int NumChildren = 1; // If we land on the deepest item already, we should just select it regardless of how many siblings it has.
+	// Find first occurrence of tree (breadth-first) of current graph in children.
+	FPCGEditorGraphDebugObjectItemPtr ItemToSelect;
 
-	// Find the deepest entry in this branch of the tree view.
-	while (true)
+	const UPCGGraph* CurrentGraph = PCGEditor.Pin()->GetPCGGraph();
+	TArray<FPCGEditorGraphDebugObjectItemPtr> ToVisit = InItem->GetChildren().Array();
+	int32 VisitIndex = 0;
+	
+	while(VisitIndex < ToVisit.Num())
 	{
-		const TSet<TSharedPtr<FPCGEditorGraphDebugObjectItem>>& Children = Item->GetChildren();
-		TSet<FPCGEditorGraphDebugObjectItemPtr>::TConstIterator It = Children.CreateConstIterator();
-
-		if (!It)
-		{
-			break;
-		}
-		else
+		const FPCGEditorGraphDebugObjectItemPtr& Item = ToVisit[VisitIndex]; 
+		const FPCGStack* ItemStack = Item ? Item->GetPCGStack() : nullptr;
+		if (ItemStack)
 		{
-			Item = *It;
-			NumChildren = Children.Num();
+			if(!ItemStack->GetStackFrames().IsEmpty() && ItemStack->GetStackFrames().Last().Object == CurrentGraph)
+			{
+				ItemToSelect = Item;
+				break;
+			}
+			else
+			{
+				ToVisit.Append(Item->GetChildren().Array());
+				++VisitIndex;
+			}
 		}
 	}
 
-	// Set the discovered item as the debug object if it is the only child.
-	if (NumChildren == 1)
+	if (ItemToSelect.IsValid())
 	{
-		DebugObjectTreeView->SetSelection(Item);
+		ExpandAndSelectDebugObject(ItemToSelect);
+	}
+	else
+	{
+		// Nothing found, just reselect the original item
+		DebugObjectTreeView->SetSelection(InItem);
 	}
 }
 
@@ -1242,13 +1428,22 @@ void SPCGEditorGraphDebugObjectTree::ContextMenu_JumpToGraphInTree()
 			}
 		}
 
-		if (JumpToPCGGraph && JumpToPCGNode)
+		if (JumpToPCGGraph)
 		{
 			GEditor->GetEditorSubsystem<UAssetEditorSubsystem>()->OpenEditorForAsset(JumpToPCGGraph);
+			IAssetEditorInstance* EditorInstance = GEditor->GetEditorSubsystem<UAssetEditorSubsystem>()->FindEditorForAsset(JumpToPCGGraph, /*bFocusIfOpen*/true);
+			FPCGEditor* OtherPCGEditor = static_cast<FPCGEditor*>(EditorInstance);
 
-			if (IAssetEditorInstance* EditorInstance = GEditor->GetEditorSubsystem<UAssetEditorSubsystem>()->FindEditorForAsset(JumpToPCGGraph, /*bFocusIfOpen*/true))
+			if (OtherPCGEditor)
 			{
-				static_cast<FPCGEditor*>(EditorInstance)->JumpToNode(JumpToPCGNode);
+				// Implementation note - for selection purposes we had culled the stack to be in the local referential
+				// however in the case of the jump here we want the full stack
+				OtherPCGEditor->SetStackBeingInspectedFromAnotherEditor(SelectedStackBeforeCulling);
+
+				if (JumpToPCGNode)
+				{
+					OtherPCGEditor->JumpToNode(JumpToPCGNode);
+				}
 			}
 		}
 	}
diff --git a/Engine/Plugins/PCG/Source/PCGEditor/Private/SPCGEditorGraphDebugObjectTree.h b/Engine/Plugins/PCG/Source/PCGEditor/Private/SPCGEditorGraphDebugObjectTree.h
--- a/Engine/Plugins/PCG/Source/PCGEditor/Private/SPCGEditorGraphDebugObjectTree.h
+++ b/Engine/Plugins/PCG/Source/PCGEditor/Private/SPCGEditorGraphDebugObjectTree.h
@@ -40,13 +40,18 @@ public:
 
 	bool IsGrayedOut() const { return bGrayedOut; }
 
+	bool IsSelected() const { return bSelected; }
+	void SetSelected(bool bInSelected) { bSelected = bInSelected; }
+
 	/** Optional sort priority, if returns INDEX_NONE then sort will fall back to alphabetical. */
 	virtual int32 GetSortPriority() const { return INDEX_NONE; }
 
 	/** Whether this item represents a currently debuggable object for the current edited graph. */
 	virtual bool IsDebuggable() const { return false; }
+	/** Whether this item represents a dynamically executed element (as in dynamic subgraph). */
+	virtual bool IsDynamic() const { return false; }
 
-	virtual FString GetLabel() const = 0;
+	virtual FString GetLabel(bool bForSorting = false) const = 0;
 	virtual UPCGComponent* GetPCGComponent() const = 0;
 	virtual const FPCGStack* GetPCGStack() const { return nullptr; }
 	virtual const UObject* GetObject() const  = 0;
@@ -62,6 +67,7 @@ protected:
 	TSet<TSharedPtr<FPCGEditorGraphDebugObjectItem>> Children;
 	bool bIsExpanded = false;
 	bool bGrayedOut = false;
+	bool bSelected = false;
 };
 
 class FPCGEditorGraphDebugObjectItem_Actor : public FPCGEditorGraphDebugObjectItem
@@ -74,7 +80,7 @@ public:
 		PCGStack.PushFrame(InActor.Get());
 	}
 
-	virtual FString GetLabel() const override;
+	virtual FString GetLabel(bool bForSorting = false) const override;
 	virtual UPCGComponent* GetPCGComponent() const override { return nullptr; }
 	virtual const UObject* GetObject() const override { return Actor.Get(); };
 	virtual const FPCGStack* GetPCGStack() const override { return &PCGStack; }
@@ -103,7 +109,7 @@ public:
 		, bIsDebuggable(bInIsDebuggable)
 	{}
 
-	virtual FString GetLabel() const override;
+	virtual FString GetLabel(bool bForSorting = false) const override;
 	virtual UPCGComponent* GetPCGComponent() const override { return PCGComponent.Get(); }
 	virtual const UObject* GetObject() const override { return PCGComponent.Get(); };
 	virtual const FPCGStack* GetPCGStack() const override { return &PCGStack; }
@@ -130,19 +136,22 @@ public:
 		TWeakObjectPtr<const UPCGGraph> InPCGGraph,
 		const FPCGStack& InPCGStack,
 		bool bInIsDebuggable,
-		bool bInHasInspectionData)
+		bool bInHasInspectionData,
+		bool bInIsDynamic)
 		: FPCGEditorGraphDebugObjectItem(bInHasInspectionData)
 		, PCGNode(InPCGNode)
 		, PCGGraph(InPCGGraph)
 		, PCGStack(InPCGStack)
 		, bIsDebuggable(bInIsDebuggable)
+		, bIsDynamic(bInIsDynamic)
 	{}
 
-	virtual FString GetLabel() const override;
+	virtual FString GetLabel(bool bForSorting = false) const override;
 	virtual UPCGComponent* GetPCGComponent() const override { return GetParent() ? GetParent()->GetPCGComponent() : nullptr; }
 	virtual const UObject* GetObject() const override { return PCGNode.Get(); };
 	virtual const FPCGStack* GetPCGStack() const override { return &PCGStack; }
 	virtual bool IsDebuggable() const override { return bIsDebuggable; }
+	virtual bool IsDynamic() const override { return bIsDynamic; }
 	virtual const UPCGGraph* GetPCGGraph() const override { return PCGGraph.Get(); }
 
 protected:
@@ -153,6 +162,7 @@ protected:
 
 	FPCGStack PCGStack;
 	bool bIsDebuggable = false;
+	bool bIsDynamic = false;
 };
 
 class FPCGEditorGraphDebugObjectItem_PCGLoopIndex : public FPCGEditorGraphDebugObjectItem
@@ -173,7 +183,7 @@ public:
 
 	virtual int32 GetLoopIndex() const { return LoopIndex; }
 
-	virtual FString GetLabel() const override;
+	virtual FString GetLabel(bool bForSorting = false) const override;
 	virtual UPCGComponent* GetPCGComponent() const override { return GetParent() ? GetParent()->GetPCGComponent() : nullptr; }
 	virtual const UObject* GetObject() const override { return nullptr; };
 	virtual const FPCGStack* GetPCGStack() const override { return &PCGStack; }
@@ -232,13 +242,17 @@ public:
 
 	void SetNodeBeingInspected(const UPCGNode* InPCGNode);
 
-	void SetDebugObjectSelection(const FPCGStack& FullStack);
-
 	void OnObjectsReplaced(const TMap<UObject*, UObject*>& ReplacementMap);
 
+	/** If the stack matches to an item, Expands the tree view to make that item visible and select it, and return true. Otherwise, return false. */
+	bool SetDebugObjectFromStackFromAnotherEditor(const FPCGStack& InStack);
+
+	/** Return the first stack downstream of the currently selected stack with the provided node & graph. Return false if nothing is found. */
+	bool GetFirstStackFromSelection(const UPCGNode* InNode, const UPCGGraph* InGraph, FPCGStack& OutStack) const;
+
 private:
-	void SelectedDebugObject_OnClicked() const;
-	bool IsSelectDebugObjectButtonEnabled() const;
+	FReply FocusOnDebugObject_OnClicked() const;
+	bool IsFocusOnDebugObjectButtonEnabled() const;
 
 	void SetDebugObjectFromSelection_OnClicked();
 	bool IsSetDebugObjectFromSelectionButtonEnabled() const { return IsSetDebugObjectFromSelectionEnabled.IsSet() && IsSetDebugObjectFromSelectionEnabled.GetValue(); }
@@ -266,9 +280,16 @@ private:
 	void OnExpansionChanged(FPCGEditorGraphDebugObjectItemPtr InItem, bool bInIsExpanded);
 	void OnSetExpansionRecursive(FPCGEditorGraphDebugObjectItemPtr InItem, bool bInExpand) const;
 
-	/** Expand the given row and select the deepest entry as the debug object if it is unambiguous (the only entry at its level in the tree). */
+	/** Returns the matching item from the AllGraphItems array if any. */
+	FPCGEditorGraphDebugObjectItemPtr GetItemFromStack(const FPCGStack& InStack) const;
+	/** If the stack matches to an item, Expands the tree view to make that item visible and select it, and return true. Otherwise, return false. */
+	bool SetDebugObjectFromStack(const FPCGStack& InStack);
+	/** Expands the tree view to make an item visible and select it. */
 	void ExpandAndSelectDebugObject(FPCGEditorGraphDebugObjectItemPtr InItem);
 
+	/** Expand the given row and select the deepest entry as the debug object if it is unambiguous (the only entry at its level in the tree). */
+	void ExpandAndSelectFirstLeafDebugObject(FPCGEditorGraphDebugObjectItemPtr InItem);
+
 	TSharedPtr<SWidget> OpenContextMenu();
 
 	/** Jump-to context menu command. */
@@ -283,9 +304,6 @@ private:
 
 	bool bNeedsRefresh = false;
 
-	/** Set true to avoid broadcasting debug object change notifications when setting the object from code. */
-	bool bDisableDebugObjectChangeNotification = false;
-
 	/** Latest value for IsSetDebugObjectFromSelectionButtonEnabled */
 	TOptional<bool> IsSetDebugObjectFromSelectionEnabled;
 
@@ -295,6 +313,9 @@ private:
 	/** Used to retain item selection state across tree refreshes. */
 	FPCGStack SelectedStack;
 
+	/** Stack before it was culled so it's in the local state. */
+	FPCGStack SelectedStackBeforeCulling;
+
 	/** Used to retain item selection state across tree refreshes if the SelectedStack is invalidated (e.g. through BP reconstruction). */
 	TWeakObjectPtr<const UPCGGraph> SelectedGraph = nullptr;
 	TWeakObjectPtr<const AActor> SelectedOwner = nullptr;
