[PCG] Fix CDO break with Attribute Property Selector

* Nodes that were setting the selector in their ctor to a point property with something else than Position AND were using the default value were not serialized and broken on reload (because the deprecated property was not set to the right value)
* Added the deprecated property set to restore the CDO.

#rnx

diff --git a/Engine/Plugins/PCG/Source/PCG/Private/Metadata/PCGAttributePropertySelector.cpp b/Engine/Plugins/PCG/Source/PCG/Private/Metadata/PCGAttributePropertySelector.cpp
--- a/Engine/Plugins/PCG/Source/PCG/Private/Metadata/PCGAttributePropertySelector.cpp
+++ b/Engine/Plugins/PCG/Source/PCG/Private/Metadata/PCGAttributePropertySelector.cpp
@@ -184,6 +184,13 @@ bool FPCGAttributePropertySelector::SetPointProperty(EPCGPointProperties InPoint
 	check(EnumPtr);
 	
 	const bool bHasChanged = SetPropertyName(*EnumPtr->GetNameStringByValue(static_cast<int64>(InPointProperty)), bResetExtraNames);
+
+#if WITH_EDITOR
+PRAGMA_DISABLE_DEPRECATION_WARNINGS
+	// Not doing this will break the CDO for all nodes that set this explicitly in their constructor.
+	PointProperty_DEPRECATED = InPointProperty;
+PRAGMA_ENABLE_DEPRECATION_WARNINGS
+#endif // WITH_EDITOR
 	
 	return bHasChanged;
 }
