[PCG] Refactor how GPU backends for nodes are created.

* New kernel object which captures all GPU-related information (moved all of this out of UPCGSettings)
* Data descriptions are now computed within the compute graph itself rather than through UPCGSettings calls, using the new kernel objects.
* Procedural static mesh spawning code now driven by kernel object.
* Remove GPU execution code from Skinned Mesh Supported as GPU instancing path is not supported for this primitive type.

This is required work to allow PCG nodes to emit multiple kernels in follow up changes.

#rnx

diff --git a/Engine/Plugins/PCG/Shaders/Private/Elements/PCGSkinnedMeshSpawner.usf b/Engine/Plugins/PCG/Shaders/Private/Elements/PCGSkinnedMeshSpawner.usf
deleted file mode 100644
diff --git a/Engine/Plugins/PCG/Source/PCG/Internal/Compute/PCGComputeCommon.h b/Engine/Plugins/PCG/Source/PCG/Internal/Compute/PCGComputeCommon.h
--- a/Engine/Plugins/PCG/Source/PCG/Internal/Compute/PCGComputeCommon.h
+++ b/Engine/Plugins/PCG/Source/PCG/Internal/Compute/PCGComputeCommon.h
@@ -108,9 +108,6 @@ namespace PCGComputeHelpers
 	void LogKernelError(const FPCGContext* Context, const UPCGSettings* Settings, const FText& InText);
 #endif
 
-	/** Computes a pin data description based on the given GPU pin properties. */
-	void ComputeOutputPinDataDesc(const FPCGPinPropertiesGPU& PinProperties, const UPCGSettings* Settings, const UPCGDataBinding* Binding, FPCGDataCollectionDesc& OutPinDesc);
-
 	/** Returns true if the given buffer size is dangerously large. Optionally emits error log. */
 	bool IsBufferSizeTooLarge(uint64 InBufferSizeBytes, bool bInLogError = true);
 
diff --git a/Engine/Plugins/PCG/Source/PCG/Internal/Compute/PCGDataForGPU.h b/Engine/Plugins/PCG/Source/PCG/Internal/Compute/PCGDataForGPU.h
--- a/Engine/Plugins/PCG/Source/PCG/Internal/Compute/PCGDataForGPU.h
+++ b/Engine/Plugins/PCG/Source/PCG/Internal/Compute/PCGDataForGPU.h
@@ -271,8 +271,4 @@ namespace PCGDataForGPUHelpers
 		const FPCGDataCollectionDesc& InDataCollectionDescription,
 		uint32& OutCustomFloatCount,
 		TArray<FUint32Vector4>& OutAttributeIdOffsetStrides);
-
-	/** Compute a description of all data arriving on InputPin. */
-	FPCGDataCollectionDesc ComputeInputPinDataDesc(const UPCGSettings* Settings, const FName& InputPinLabel, const UPCGDataBinding* Binding);
-	FPCGDataCollectionDesc ComputeInputPinDataDesc(const UPCGPin* InputPin, const UPCGDataBinding* Binding);
 }
diff --git a/Engine/Plugins/PCG/Source/PCG/Private/Compute/DataInterfaces/Elements/PCGCopyPointsDataInterface.cpp b/Engine/Plugins/PCG/Source/PCG/Private/Compute/DataInterfaces/Elements/PCGCopyPointsDataInterface.cpp
--- a/Engine/Plugins/PCG/Source/PCG/Private/Compute/DataInterfaces/Elements/PCGCopyPointsDataInterface.cpp
+++ b/Engine/Plugins/PCG/Source/PCG/Private/Compute/DataInterfaces/Elements/PCGCopyPointsDataInterface.cpp
@@ -3,6 +3,7 @@
 #include "PCGCopyPointsDataInterface.h"
 
 #include "Elements/PCGCopyPoints.h"
+#include "Elements/PCGCopyPointsKernel.h"
 
 #include "GlobalRenderResources.h"
 #include "RHIResources.h"
@@ -95,13 +96,16 @@ void UPCGCopyPointsDataInterface::GetHLSL(FString& OutHLSL, FString const& InDat
 UComputeDataProvider* UPCGCopyPointsDataInterface::CreateDataProvider(TObjectPtr<UObject> InBinding, uint64 InInputMask, uint64 InOutputMask) const
 {
 	TObjectPtr<UPCGCopyPointsDataProvider> DataProvider = NewObject<UPCGCopyPointsDataProvider>();
-	DataProvider->Settings = Cast<UPCGCopyPointsSettings>(Settings);
+	DataProvider->ProducerKernel = CastChecked<UPCGCopyPointsKernel>(ProducerKernel);
 
 	return DataProvider;
 }
 
 FComputeDataProviderRenderProxy* UPCGCopyPointsDataProvider::GetRenderProxy()
 {
+	const UPCGCopyPointsSettings* Settings = ProducerKernel->Settings;
+	check(Settings);
+
 	return new FPCGCopyPointsDataProviderProxy(
 		static_cast<uint32>(Settings->RotationInheritance),
 		static_cast<uint32>((Settings->RotationInheritance != EPCGCopyPointsInheritanceMode::Source || Settings->bApplyTargetRotationToPositions) ? 1 : 0),
diff --git a/Engine/Plugins/PCG/Source/PCG/Private/Compute/DataInterfaces/Elements/PCGCopyPointsDataInterface.h b/Engine/Plugins/PCG/Source/PCG/Private/Compute/DataInterfaces/Elements/PCGCopyPointsDataInterface.h
--- a/Engine/Plugins/PCG/Source/PCG/Private/Compute/DataInterfaces/Elements/PCGCopyPointsDataInterface.h
+++ b/Engine/Plugins/PCG/Source/PCG/Private/Compute/DataInterfaces/Elements/PCGCopyPointsDataInterface.h
@@ -8,8 +8,7 @@
 #include "PCGCopyPointsDataInterface.generated.h"
 
 class FPCGCopyPointsDataInterfaceParameters;
-class UPCGCopyPointsSettings;
-class UPCGSettings;
+class UPCGCopyPointsKernel;
 
 /** Data Interface to marshal Copy Points settings to the GPU. */
 UCLASS(ClassGroup = (Procedural))
@@ -25,9 +24,6 @@ public:
 	void GetHLSL(FString& OutHLSL, FString const& InDataInterfaceName) const override;
 	UComputeDataProvider* CreateDataProvider(TObjectPtr<UObject> InBinding, uint64 InInputMask, uint64 InOutputMask) const override;
 	//~ End UComputeDataInterface Interface
-
-	UPROPERTY()
-	TObjectPtr<const UPCGSettings> Settings = nullptr;
 };
 
 UCLASS()
@@ -41,7 +37,7 @@ public:
 	//~ End UComputeDataProvider Interface
 
 	UPROPERTY()
-	TObjectPtr<const UPCGCopyPointsSettings> Settings;
+	TObjectPtr<const UPCGCopyPointsKernel> ProducerKernel;
 };
 
 class FPCGCopyPointsDataProviderProxy : public FComputeDataProviderRenderProxy
diff --git a/Engine/Plugins/PCG/Source/PCG/Private/Compute/DataInterfaces/Elements/PCGMetadataPartitionDataInterface.cpp b/Engine/Plugins/PCG/Source/PCG/Private/Compute/DataInterfaces/Elements/PCGMetadataPartitionDataInterface.cpp
--- a/Engine/Plugins/PCG/Source/PCG/Private/Compute/DataInterfaces/Elements/PCGMetadataPartitionDataInterface.cpp
+++ b/Engine/Plugins/PCG/Source/PCG/Private/Compute/DataInterfaces/Elements/PCGMetadataPartitionDataInterface.cpp
@@ -5,6 +5,7 @@
 #include "Compute/PCGComputeCommon.h"
 #include "Compute/PCGDataBinding.h"
 #include "Elements/Metadata/PCGMetadataPartition.h"
+#include "Elements/Metadata/PCGMetadataPartitionKernel.h"
 
 #include "GlobalRenderResources.h"
 #include "RenderGraphBuilder.h"
@@ -83,12 +84,16 @@ UComputeDataProvider* UPCGMetadataPartitionDataInterface::CreateDataProvider(TOb
 {
 	TRACE_CPUPROFILER_EVENT_SCOPE(UPCGMetadataPartitionDataInterface::CreateDataProvider);
 	UPCGDataBinding* Binding = CastChecked<UPCGDataBinding>(InBinding);
+	check(ProducerKernel);
+	const UPCGMetadataPartitionSettings* Settings = Cast<UPCGMetadataPartitionKernel>(ProducerKernel)->Settings;
+	check(Settings);
 
 	TObjectPtr<UPCGMetaDataPartitionDataProvider> DataProvider = NewObject<UPCGMetaDataPartitionDataProvider>();
 
 	if (ensure(Settings->PartitionAttributeSelectors.Num() == 1 && Settings->PartitionAttributeSelectors[0].IsBasicAttribute()))
 	{
-		const FPCGDataCollectionDesc InputDataDesc = PCGDataForGPUHelpers::ComputeInputPinDataDesc(Settings, PCGPinConstants::DefaultInputLabel, Binding);
+		FPCGDataCollectionDesc InputDataDesc;
+		ensure(Binding->GetKernelPinDataDesc_GameThread(ProducerKernel, PCGPinConstants::DefaultInputLabel, /*bIsInput*/true, InputDataDesc));
 
 		const int32 PartitionAttributeId = Binding->GetAttributeId(Settings->PartitionAttributeSelectors[0].GetAttributeName(), EPCGKernelAttributeType::StringKey);
 		if (PartitionAttributeId != INDEX_NONE)
diff --git a/Engine/Plugins/PCG/Source/PCG/Private/Compute/DataInterfaces/Elements/PCGMetadataPartitionDataInterface.h b/Engine/Plugins/PCG/Source/PCG/Private/Compute/DataInterfaces/Elements/PCGMetadataPartitionDataInterface.h
--- a/Engine/Plugins/PCG/Source/PCG/Private/Compute/DataInterfaces/Elements/PCGMetadataPartitionDataInterface.h
+++ b/Engine/Plugins/PCG/Source/PCG/Private/Compute/DataInterfaces/Elements/PCGMetadataPartitionDataInterface.h
@@ -8,7 +8,6 @@
 #include "PCGMetadataPartitionDataInterface.generated.h"
 
 class FPCGMetadataPartitionDataInterfaceParameters;
-class UPCGMetadataPartitionSettings;
 
 /** Data Interface to marshal Attribute Partition settings to the GPU. */
 UCLASS(ClassGroup = (Procedural))
@@ -24,9 +23,6 @@ public:
 	void GetHLSL(FString& OutHLSL, FString const& InDataInterfaceName) const override;
 	UComputeDataProvider* CreateDataProvider(TObjectPtr<UObject> InBinding, uint64 InInputMask, uint64 InOutputMask) const override;
 	//~ End UComputeDataInterface Interface
-
-	UPROPERTY()
-	TObjectPtr<const UPCGMetadataPartitionSettings> Settings = nullptr;
 };
 
 UCLASS()
diff --git a/Engine/Plugins/PCG/Source/PCG/Private/Compute/DataInterfaces/Elements/PCGSkinnedMeshSpawnerDataInterface.cpp b/Engine/Plugins/PCG/Source/PCG/Private/Compute/DataInterfaces/Elements/PCGSkinnedMeshSpawnerDataInterface.cpp
deleted file mode 100644
diff --git a/Engine/Plugins/PCG/Source/PCG/Private/Compute/DataInterfaces/Elements/PCGSkinnedMeshSpawnerDataInterface.h b/Engine/Plugins/PCG/Source/PCG/Private/Compute/DataInterfaces/Elements/PCGSkinnedMeshSpawnerDataInterface.h
deleted file mode 100644
diff --git a/Engine/Plugins/PCG/Source/PCG/Private/Compute/DataInterfaces/Elements/PCGStaticMeshSpawnerDataInterface.cpp b/Engine/Plugins/PCG/Source/PCG/Private/Compute/DataInterfaces/Elements/PCGStaticMeshSpawnerDataInterface.cpp
--- a/Engine/Plugins/PCG/Source/PCG/Private/Compute/DataInterfaces/Elements/PCGStaticMeshSpawnerDataInterface.cpp
+++ b/Engine/Plugins/PCG/Source/PCG/Private/Compute/DataInterfaces/Elements/PCGStaticMeshSpawnerDataInterface.cpp
@@ -5,6 +5,7 @@
 #include "Compute/PCGComputeCommon.h"
 #include "Compute/PCGDataBinding.h"
 #include "Elements/PCGStaticMeshSpawner.h"
+#include "Elements/PCGStaticMeshSpawnerKernel.h"
 
 #include "GlobalRenderResources.h"
 #include "RHIResources.h"
@@ -175,11 +176,13 @@ UComputeDataProvider* UPCGStaticMeshSpawnerDataInterface::CreateDataProvider(TOb
 	TRACE_CPUPROFILER_EVENT_SCOPE(UPCGStaticMeshSpawnerDataInterface::CreateDataProvider);
 	UPCGDataBinding* Binding = CastChecked<UPCGDataBinding>(InBinding);
 
+	const UPCGStaticMeshSpawnerKernel* Kernel = CastChecked<UPCGStaticMeshSpawnerKernel>(ProducerKernel);
+	check(Kernel->Settings);
+
 	TObjectPtr<UPCGStaticMeshSpawnerDataProvider> DataProvider = NewObject<UPCGStaticMeshSpawnerDataProvider>();
-	DataProvider->Settings = CastChecked<UPCGStaticMeshSpawnerSettings>(Settings);
 	
 	// If there were 0 input points for this execution, we will not have created any primitives, so check for null.
-	if (const FPCGSpawnerPrimitives* Primitives = Binding->MeshSpawnersToPrimitives.Find(Settings))
+	if (const FPCGSpawnerPrimitives* Primitives = Binding->MeshSpawnersToPrimitives.Find(Kernel))
 	{
 		DataProvider->AttributeIdOffsetStrides = Primitives->AttributeIdOffsetStrides;
 		DataProvider->SelectionCDF = Primitives->SelectionCDF;
@@ -191,7 +194,7 @@ UComputeDataProvider* UPCGStaticMeshSpawnerDataInterface::CreateDataProvider(TOb
 
 		// We know the name and type of the selected mesh attribute statically and declared the attribute in GetKernelAttributeKeys,
 		// so the attribute ID should be present in the attribute table.
-		DataProvider->SelectedMeshAttributeId = Binding->GetAttributeId(DataProvider->Settings->OutAttributeName, EPCGKernelAttributeType::StringKey);
+		DataProvider->SelectedMeshAttributeId = Binding->GetAttributeId(Kernel->Settings->OutAttributeName, EPCGKernelAttributeType::StringKey);
 		ensure(DataProvider->SelectedMeshAttributeId != INDEX_NONE);
 	}
 
diff --git a/Engine/Plugins/PCG/Source/PCG/Private/Compute/DataInterfaces/Elements/PCGStaticMeshSpawnerDataInterface.h b/Engine/Plugins/PCG/Source/PCG/Private/Compute/DataInterfaces/Elements/PCGStaticMeshSpawnerDataInterface.h
--- a/Engine/Plugins/PCG/Source/PCG/Private/Compute/DataInterfaces/Elements/PCGStaticMeshSpawnerDataInterface.h
+++ b/Engine/Plugins/PCG/Source/PCG/Private/Compute/DataInterfaces/Elements/PCGStaticMeshSpawnerDataInterface.h
@@ -26,9 +26,6 @@ public:
 	UComputeDataProvider* CreateDataProvider(TObjectPtr<UObject> InBinding, uint64 InInputMask, uint64 InOutputMask) const override;
 	//~ End UComputeDataInterface Interface
 
-	UPROPERTY()
-	TObjectPtr<const UPCGSettings> Settings = nullptr;
-
 	static constexpr uint32 MAX_ATTRIBUTES = 64;
 };
 
@@ -42,9 +39,6 @@ public:
 	FComputeDataProviderRenderProxy* GetRenderProxy() override;
 	//~ End UComputeDataProvider Interface
 
-	UPROPERTY()
-	TObjectPtr<const UPCGStaticMeshSpawnerSettings> Settings = nullptr;
-
 	UPROPERTY()
 	TArray<FUintVector4> AttributeIdOffsetStrides;
 
diff --git a/Engine/Plugins/PCG/Source/PCG/Private/Compute/DataInterfaces/PCGComputeDataInterface.h b/Engine/Plugins/PCG/Source/PCG/Private/Compute/DataInterfaces/PCGComputeDataInterface.h
--- a/Engine/Plugins/PCG/Source/PCG/Private/Compute/DataInterfaces/PCGComputeDataInterface.h
+++ b/Engine/Plugins/PCG/Source/PCG/Private/Compute/DataInterfaces/PCGComputeDataInterface.h
@@ -7,6 +7,7 @@
 
 #include "PCGComputeDataInterface.generated.h"
 
+class UPCGComputeKernel;
 class UPCGDataBinding;
 
 UCLASS(Abstract, ClassGroup = (Procedural))
@@ -28,6 +29,12 @@ public:
 		OutputPinLabelAlias = InOptionalOutputPinLabelAlias ? *InOptionalOutputPinLabelAlias : InOutputPinLabel;
 	}
 
+	UPROPERTY()
+	int32 GraphBindingIndex = INDEX_NONE;
+
+	UPROPERTY()
+	TObjectPtr<const UPCGComputeKernel> ProducerKernel;
+
 protected:
 	/** Label of output pin that this data interface is associated with. */
 	UPROPERTY()
diff --git a/Engine/Plugins/PCG/Source/PCG/Private/Compute/DataInterfaces/PCGCustomKernelDataInterface.cpp b/Engine/Plugins/PCG/Source/PCG/Private/Compute/DataInterfaces/PCGCustomKernelDataInterface.cpp
--- a/Engine/Plugins/PCG/Source/PCG/Private/Compute/DataInterfaces/PCGCustomKernelDataInterface.cpp
+++ b/Engine/Plugins/PCG/Source/PCG/Private/Compute/DataInterfaces/PCGCustomKernelDataInterface.cpp
@@ -5,7 +5,7 @@
 #include "PCGComponent.h"
 #include "PCGModule.h"
 #include "PCGSettings.h"
-#include "Compute/Elements/PCGCustomHLSL.h"
+#include "Compute/PCGComputeKernel.h"
 
 #include "RenderGraphBuilder.h"
 #include "RenderGraphResources.h"
@@ -93,21 +93,14 @@ UComputeDataProvider* UPCGCustomKernelDataInterface::CreateDataProvider(TObjectP
 {
 	TRACE_CPUPROFILER_EVENT_SCOPE(UPCGCustomKernelDataInterface::CreateDataProvider);
 	check(Settings);
+	check(Kernel);
 
 	UPCGDataBinding* Binding = CastChecked<UPCGDataBinding>(InBinding);
 	check(Binding->SourceComponent.IsValid() && Binding->SourceComponent.Get());
 
 	UPCGCustomComputeKernelDataProvider* Provider = NewObject<UPCGCustomComputeKernelDataProvider>();
-	Provider->ThreadCount = Settings->ComputeKernelThreadCount(Binding);
-
-	if (const UPCGCustomHLSLSettings* CustomHLSLSettings = Cast<UPCGCustomHLSLSettings>(Settings))
-	{
-		Provider->ThreadCountMultiplier = CustomHLSLSettings->GetThreadCountMultiplier();
-	}
-	else
-	{
-		Provider->ThreadCountMultiplier = 1;
-	}
+	Provider->ThreadCount = Kernel->ComputeThreadCount(Binding);
+	Provider->ThreadCountMultiplier = Kernel->GetThreadCountMultiplier();
 
 	Provider->Seed = static_cast<uint32>(Settings->GetSeed(Binding->SourceComponent.Get()));
 	Provider->SeedSettings = static_cast<uint32>(Settings->Seed);
diff --git a/Engine/Plugins/PCG/Source/PCG/Private/Compute/DataInterfaces/PCGCustomKernelDataInterface.h b/Engine/Plugins/PCG/Source/PCG/Private/Compute/DataInterfaces/PCGCustomKernelDataInterface.h
--- a/Engine/Plugins/PCG/Source/PCG/Private/Compute/DataInterfaces/PCGCustomKernelDataInterface.h
+++ b/Engine/Plugins/PCG/Source/PCG/Private/Compute/DataInterfaces/PCGCustomKernelDataInterface.h
@@ -10,6 +10,7 @@
 #include "PCGCustomKernelDataInterface.generated.h"
 
 class FPCGCustomKernelDataInterfaceParameters;
+class UPCGComputeKernel;
 
 /** Interface for any meta data provided to the compute kernel, such as num threads. */
 UCLASS(ClassGroup = (Procedural))
@@ -34,6 +35,9 @@ public:
 
 	UPROPERTY()
 	TObjectPtr<const UPCGSettings> Settings = nullptr;
+
+	UPROPERTY()
+	TObjectPtr<const UPCGComputeKernel> Kernel = nullptr;
 };
 
 /** Compute Framework Data Provider for each custom compute kernel. */
diff --git a/Engine/Plugins/PCG/Source/PCG/Private/Compute/DataInterfaces/PCGDataCollectionDataInterface.cpp b/Engine/Plugins/PCG/Source/PCG/Private/Compute/DataInterfaces/PCGDataCollectionDataInterface.cpp
--- a/Engine/Plugins/PCG/Source/PCG/Private/Compute/DataInterfaces/PCGDataCollectionDataInterface.cpp
+++ b/Engine/Plugins/PCG/Source/PCG/Private/Compute/DataInterfaces/PCGDataCollectionDataInterface.cpp
@@ -436,10 +436,12 @@ UComputeDataProvider* UPCGDataCollectionDataInterface::CreateDataProvider(TObjec
 {
 	TRACE_CPUPROFILER_EVENT_SCOPE(UPCGDataCollectionDataInterface::CreateDataProvider);
 	check(ProducerSettings);
+
 	UPCGDataBinding* Binding = CastChecked<UPCGDataBinding>(InBinding);
+	check(Binding->Graph);
 
 	FPCGDataCollectionDesc PinDataDesc;
-	ensure(Binding->GetOutputPinDataDesc_GameThread(OutputPinLabel, ProducerSettings, PinDataDesc));
+	ensure(Binding->GetKernelPinDataDesc_GameThread(GraphBindingIndex, PinDataDesc));
 
 	if (PCGComputeHelpers::IsBufferSizeTooLarge(PinDataDesc.ComputePackedSizeBytes()))
 	{
diff --git a/Engine/Plugins/PCG/Source/PCG/Private/Compute/DataInterfaces/PCGDataCollectionUploadDataInterface.cpp b/Engine/Plugins/PCG/Source/PCG/Private/Compute/DataInterfaces/PCGDataCollectionUploadDataInterface.cpp
--- a/Engine/Plugins/PCG/Source/PCG/Private/Compute/DataInterfaces/PCGDataCollectionUploadDataInterface.cpp
+++ b/Engine/Plugins/PCG/Source/PCG/Private/Compute/DataInterfaces/PCGDataCollectionUploadDataInterface.cpp
@@ -30,17 +30,7 @@ UComputeDataProvider* UPCGDataCollectionUploadDataInterface::CreateDataProvider(
 	// Pick the data items from input data collection using any of the compute graph element virtual input pin labels.
 	check(!DownstreamInputPinLabelAliases.IsEmpty());
 
-	if (ProducerSettings && Cast<UPCGNode>(ProducerSettings->GetOuter()))
-	{
-		ensure(Binding->GetOutputPinDataDesc_GameThread(OutputPinLabel, ProducerSettings, Provider->PinDesc));
-	}
-	else
-	{
-		Provider->PinDesc = FPCGDataCollectionDesc::BuildFromInputDataCollectionAndInputPinLabel(
-			Binding->InputDataCollection,
-			DownstreamInputPinLabelAliases[0],
-			Binding);
-	}
+	ensure(Binding->GetKernelPinDataDesc_GameThread(GraphBindingIndex, Provider->PinDesc));
 
 	Provider->DownstreamInputPinLabels = DownstreamInputPinLabelAliases;
 
diff --git a/Engine/Plugins/PCG/Source/PCG/Private/Compute/DataInterfaces/PCGInstanceDataInterface.cpp b/Engine/Plugins/PCG/Source/PCG/Private/Compute/DataInterfaces/PCGInstanceDataInterface.cpp
--- a/Engine/Plugins/PCG/Source/PCG/Private/Compute/DataInterfaces/PCGInstanceDataInterface.cpp
+++ b/Engine/Plugins/PCG/Source/PCG/Private/Compute/DataInterfaces/PCGInstanceDataInterface.cpp
@@ -9,6 +9,7 @@
 #include "Compute/PCGComputeCommon.h"
 #include "Compute/PCGDataBinding.h"
 #include "Data/PCGTextureData.h"
+#include "Elements/PCGStaticMeshSpawnerKernel.h"
 
 #include "GlobalRenderResources.h"
 #include "GlobalShader.h"
@@ -129,12 +130,12 @@ UComputeDataProvider* UPCGInstanceDataInterface::CreateDataProvider(TObjectPtr<U
 {
 	TRACE_CPUPROFILER_EVENT_SCOPE(UPCGInstanceDataInterface::CreateDataProvider);
 	UPCGDataBinding* Binding = CastChecked<UPCGDataBinding>(InBinding);
-	check(ProducerSettings);
+	check(ProducerKernel);
 
 	TObjectPtr<UPCGInstanceDataProvider> DataProvider = NewObject<UPCGInstanceDataProvider>();
 
 	// If there were 0 input points (or too many input points) for this execution, we will not have created any primitives, so check for null.
-	if (FPCGSpawnerPrimitives* FoundPrimitives = Binding->MeshSpawnersToPrimitives.Find(ProducerSettings))
+	if (FPCGSpawnerPrimitives* FoundPrimitives = Binding->MeshSpawnersToPrimitives.Find(Cast<UPCGStaticMeshSpawnerKernel>(ProducerKernel)))
 	{
 		if (!ensure(FoundPrimitives->Primitives.Num() <= PCGComputeConstants::MAX_PRIMITIVE_COMPONENTS_PER_SPAWNER))
 		{
diff --git a/Engine/Plugins/PCG/Source/PCG/Private/Compute/DataInterfaces/PCGInstanceDataInterface.h b/Engine/Plugins/PCG/Source/PCG/Private/Compute/DataInterfaces/PCGInstanceDataInterface.h
--- a/Engine/Plugins/PCG/Source/PCG/Private/Compute/DataInterfaces/PCGInstanceDataInterface.h
+++ b/Engine/Plugins/PCG/Source/PCG/Private/Compute/DataInterfaces/PCGInstanceDataInterface.h
@@ -31,17 +31,11 @@ public:
 	UComputeDataProvider* CreateDataProvider(TObjectPtr<UObject> InBinding, uint64 InInputMask, uint64 InOutputMask) const override;
 	//~ End UComputeDataInterface Interface
 
-	void SetProducerSettings(const UPCGSettings* InProducerSettings) { ProducerSettings = InProducerSettings; }
-
 public:
 	UPROPERTY()
 	FName InputPinProvidingData = NAME_None;
 
 private:
-	/** Settings of node that produces this data, normally the upstream node. */
-	UPROPERTY()
-	TObjectPtr<const UPCGSettings> ProducerSettings;
-
 	static TCHAR const* TemplateFilePath;
 };
 
diff --git a/Engine/Plugins/PCG/Source/PCG/Private/Compute/Elements/PCGComputeGraphElement.cpp b/Engine/Plugins/PCG/Source/PCG/Private/Compute/Elements/PCGComputeGraphElement.cpp
--- a/Engine/Plugins/PCG/Source/PCG/Private/Compute/Elements/PCGComputeGraphElement.cpp
+++ b/Engine/Plugins/PCG/Source/PCG/Private/Compute/Elements/PCGComputeGraphElement.cpp
@@ -4,36 +4,24 @@
 
 #include "PCGComponent.h"
 #include "PCGGraph.h"
-#include "PCGManagedResource.h"
 #include "PCGModule.h"
 #include "PCGSubsystem.h"
 #include "Components/PCGProceduralISMComponent.h"
 #include "Compute/PCGComputeCommon.h"
 #include "Compute/Data/PCGProxyForGPUData.h"
 #include "Compute/DataInterfaces/PCGDataCollectionDataInterface.h"
-#include "Elements/PCGStaticMeshSpawner.h"
-#include "Engine/StaticMesh.h"
+#include "Elements/PCGStaticMeshSpawnerKernel.h"
 #include "Graph/PCGGraphCompiler.h"
-#include "InstanceDataPackers/PCGInstanceDataPackerBase.h"
-#include "MeshSelectors/PCGMeshSelectorByAttribute.h"
-#include "MeshSelectors/PCGMeshSelectorWeighted.h"
 
 #include "ComputeWorkerInterface.h"
 #include "ComputeFramework/ComputeFramework.h"
 #include "ComputeFramework/ComputeKernelCompileResult.h"
-#include "GameFramework/Actor.h"
 #include "Logging/LogVerbosity.h"
 
 #include UE_INLINE_GENERATED_CPP_BY_NAME(PCGComputeGraphElement)
 
 #define LOCTEXT_NAMESPACE "PCGComputeGraphElement"
 
-namespace PCGComputeGraphElement
-{
-	static const FText CouldNotLoadStaticMeshFormat = LOCTEXT("CouldNotLoadStaticMesh", "Could not load static mesh from path '{0}'.");
-	static const FText TooManyPrimitivesFormat = LOCTEXT("TooManyPrimitives", "Attempted to emit too many primitive components, terminated after creating '{0}'.");
-}
-
 void FPCGComputeGraphContext::AddExtraStructReferencedObjects(FReferenceCollector& Collector)
 {
 	Collector.AddPropertyReferences(FComputeGraphInstance::StaticStruct(), &ComputeGraphInstance);
@@ -173,14 +161,9 @@ bool FPCGComputeGraphElement::ExecuteInternal(FPCGContext* InContext) const
 		Context->DataBinding->Initialize(Context->ComputeGraph.Get(), Context->SourceComponent, Context->InputData, Context->ComputeGraph->GetAttributeStaticTable());
 
 		// Perform validation after input data is initialized.
-		for (TSoftObjectPtr<const UPCGNode> Node : Context->ComputeGraph->KernelToNode)
+		if (!Context->ComputeGraph->IsGraphValid(InContext, /*bQuiet=*/false))
 		{
-			const UPCGSettings* Settings = Node.Get() ? Node->GetSettings() : nullptr;
-
-			if (!Settings || !Settings->IsKernelValid(Context, /*bQuiet=*/false))
-			{
-				return true;
-			}
+			return true;
 		}
 
 		const bool bAnyComponentsSetup = SetupProceduralISMComponents(Context, DataBinding);
@@ -279,7 +262,7 @@ bool FPCGComputeGraphElement::ExecuteInternal(FPCGContext* InContext) const
 	{
 		bool bAllProxiesCreated = true;
 
-		for (const TPair<TObjectPtr<const UPCGSettings>, FPCGSpawnerPrimitives>& SpawnerToPrimitives : Context->DataBinding->MeshSpawnersToPrimitives)
+		for (const TPair<TObjectPtr<const UPCGStaticMeshSpawnerKernel>, FPCGSpawnerPrimitives>& SpawnerToPrimitives : Context->DataBinding->MeshSpawnersToPrimitives)
 		{
 			for (const TObjectPtr<UPrimitiveComponent>& Primitive : SpawnerToPrimitives.Value.Primitives)
 			{
@@ -419,303 +402,14 @@ bool FPCGComputeGraphElement::SetupProceduralISMComponents(FPCGContext* InContex
 
 	bool bAnyComponentsSetup = false;
 
-	for (TSoftObjectPtr<const UPCGSettings> Settings : Context->ComputeGraph->StaticMeshSpawners)
+	for (const UPCGStaticMeshSpawnerKernel* Kernel : Context->ComputeGraph->StaticMeshSpawners)
 	{
-		const UPCGStaticMeshSpawnerSettings* SpawnerSettings = Cast<UPCGStaticMeshSpawnerSettings>(Settings.Get());
-		if (!ensure(SpawnerSettings) || !SpawnerSettings->bEnabled)
+		if (!ensure(Kernel))
 		{
 			continue;
 		}
 
-		AActor* TargetActor = SpawnerSettings->TargetActor.Get() ? SpawnerSettings->TargetActor.Get() : InContext->GetTargetActor(nullptr);
-		if (!ensure(SpawnerSettings->MeshSelectorParameters) || !ensure(TargetActor))
-		{
-			continue;
-		}
-
-		const FPCGDataCollectionDesc InputDataDesc = PCGDataForGPUHelpers::ComputeInputPinDataDesc(SpawnerSettings, PCGPinConstants::DefaultInputLabel, InBinding);
-
-		const uint32 TotalInputPointCount = InputDataDesc.ComputeDataElementCount(EPCGDataType::Point);
-		if (TotalInputPointCount == 0)
-		{
-			continue;
-		}
-
-		const FBox LocalBounds = InContext->SourceComponent->GetGridBounds().ShiftBy(-InContext->SourceComponent->GetOwner()->GetActorLocation());
-
-		uint32 CustomFloatCount = 0;
-
-		TArray<FUint32Vector4> AttributeIdOffsetStrides;
-		if (SpawnerSettings->InstanceDataPackerParameters && SpawnerSettings->InstanceDataPackerParameters->GetAttributeNames(/*OutNames=*/nullptr))
-		{
-			TArray<FName> AttributeNames;
-			SpawnerSettings->InstanceDataPackerParameters->GetAttributeNames(&AttributeNames);
-
-			for (FName AttributeName : AttributeNames)
-			{
-				if (AttributeName == NAME_None)
-				{
-					PCG_KERNEL_VALIDATION_ERR(Context, SpawnerSettings, /*bQuiet=*/false,
-						FText::Format(LOCTEXT("InstanceDataAttributeInvalid", "Invalid instance data attribute specified '{0}'."), FText::FromName(AttributeName)));
-
-					continue;
-				}
-
-				if (!InputDataDesc.ContainsAttributeOnAnyData(AttributeName))
-				{
-					PCG_KERNEL_VALIDATION_ERR(Context, SpawnerSettings, /*bQuiet=*/false,
-						FText::Format(LOCTEXT("InstanceDataAttributeNotFound", "Instance data attribute '{0}' not found."), FText::FromName(AttributeName)));
-
-					continue;
-				}
-			}
-
-			PCGDataForGPUHelpers::ComputeCustomFloatPacking(InContext, SpawnerSettings, AttributeNames, InBinding, InputDataDesc, CustomFloatCount, AttributeIdOffsetStrides);
-		}
-
-		TArray<float> PrimitiveSelectionCDF;
-		int32 SelectorAttributeId = INDEX_NONE;
-		TArray<int32> PrimitiveStringKeys;
-		TArray<FBox> PrimitiveMeshBounds;
-		TArray<FPCGProceduralISMComponentDescriptor> ComponentsToCreate;
-
-		if (const UPCGMeshSelectorByAttribute* SelectorByAttribute = Cast<UPCGMeshSelectorByAttribute>(SpawnerSettings->MeshSelectorParameters))
-		{
-			const FName SelectorName = SelectorByAttribute->AttributeName;
-
-			// Compute how many instances we expect for each mesh.
-			TMap<int32, uint32> StringKeyToInstanceCount;
-
-			for (const FPCGDataDesc& Desc : InputDataDesc.DataDescs)
-			{
-				if (Desc.ElementCount <= 0)
-				{
-					continue;
-				}
-
-				for (const FPCGKernelAttributeDesc& AttributeDesc : Desc.AttributeDescs)
-				{
-					if (AttributeDesc.Name == SelectorName)
-					{
-						if (AttributeDesc.Type == EPCGKernelAttributeType::StringKey)
-						{
-							ensure(SelectorAttributeId == INDEX_NONE || SelectorAttributeId == AttributeDesc.Index);
-							SelectorAttributeId = AttributeDesc.Index;
-
-							for (int32 AttributeStringKeyValue : AttributeDesc.GetUniqueStringKeys())
-							{
-								if (AttributeStringKeyValue > 0)
-								{
-									StringKeyToInstanceCount.FindOrAdd(AttributeStringKeyValue, 0) += Desc.ElementCount;
-								}
-							}
-
-							// We could early out but we currently continue looping to generate the below warning for mismatched types.
-						}
-						else
-						{
-							// Currently only a single type per attribute name is supported (the name/type in the attribute table). It's possible to wire a graph
-							// for which an attribute is present with multiple types. Warn if this is encountered.
-							PCG_KERNEL_VALIDATION_WARN(Context, SpawnerSettings, /*bQuiet=*/false, FText::Format(
-								LOCTEXT("MeshSelectorAttributeNotUsable", "Attribute '{0}' not usable for mesh selection, only attributes of type String Key are supported."),
-								FText::FromName(SelectorName)));
-						}
-					}
-				}
-			}
-
-			if (SelectorAttributeId == INDEX_NONE)
-			{
-				PCG_KERNEL_VALIDATION_ERR(Context, SpawnerSettings, /*bQuiet=*/false, FText::Format(
-					LOCTEXT("MeshSelectorAttributeNotFound", "Mesh selector attribute '{0}' not found."),
-					FText::FromName(SelectorName)));
-
-				continue;
-			}
-
-			for (const TPair<int32, uint32>& StringKeyAndInstanceCount : StringKeyToInstanceCount)
-			{
-				if (!ensure(InBinding->GetStringTable().IsValidIndex(StringKeyAndInstanceCount.Key)))
-				{
-					continue;
-				}
-
-				const FString& MeshPathString = InBinding->GetStringTable()[StringKeyAndInstanceCount.Key];
-				if (MeshPathString.IsEmpty())
-				{
-					continue;
-				}
-
-				if (ComponentsToCreate.Num() >= PCGComputeConstants::MAX_PRIMITIVE_COMPONENTS_PER_SPAWNER)
-				{
-					PCG_KERNEL_VALIDATION_WARN(Context, SpawnerSettings, /*bQuiet=*/false,
-						FText::Format(PCGComputeGraphElement::TooManyPrimitivesFormat, FText::FromString(FString::FromInt(PCGComputeConstants::MAX_PRIMITIVE_COMPONENTS_PER_SPAWNER))));
-
-					break;
-				}
-
-				UStaticMesh* StaticMesh = Cast<UStaticMesh>(FSoftObjectPath(MeshPathString).TryLoad());
-				if (!StaticMesh)
-				{
-					PCG_KERNEL_VALIDATION_ERR(Context, SpawnerSettings, /*bQuiet=*/false,
-						FText::Format(PCGComputeGraphElement::CouldNotLoadStaticMeshFormat, FText::FromString(MeshPathString)));
-
-					continue;
-				}
-
-				FPCGProceduralISMComponentDescriptor Descriptor;
-				Descriptor = SelectorByAttribute->TemplateDescriptor;
-				Descriptor.NumInstances = StringKeyAndInstanceCount.Value;
-				Descriptor.LocalBounds = LocalBounds;
-				Descriptor.NumCustomFloats = CustomFloatCount;
-				Descriptor.StaticMesh = StaticMesh;
-
-				PrimitiveStringKeys.Add(StringKeyAndInstanceCount.Key);
-				ComponentsToCreate.Add(MoveTemp(Descriptor));
-
-				if (SpawnerSettings->bApplyMeshBoundsToPoints)
-				{
-					PrimitiveMeshBounds.Add(Descriptor.StaticMesh->GetBoundingBox());
-				}
-			}
-
-			PrimitiveSelectionCDF.SetNumZeroed(ComponentsToCreate.Num());
-		}
-		else if (const UPCGMeshSelectorWeighted* SelectorWeighted = Cast<UPCGMeshSelectorWeighted>(SpawnerSettings->MeshSelectorParameters))
-		{
-			if (SelectorWeighted->MeshEntries.IsEmpty())
-			{
-				PCG_KERNEL_VALIDATION_ERR(Context, SpawnerSettings, /*bQuiet=*/false, LOCTEXT("NoMeshEntries", "No mesh entries provided."));
-
-				continue;
-			}
-
-			float CumulativeWeight = 0.0f;
-
-			float TotalWeight = 0.0f;
-			for (const FPCGMeshSelectorWeightedEntry& Entry : SelectorWeighted->MeshEntries)
-			{
-				TotalWeight += Entry.Weight;
-			}
-
-			if (!ensure(TotalWeight > UE_SMALL_NUMBER))
-			{
-				continue;
-			}
-
-			PrimitiveSelectionCDF.Reserve(SelectorWeighted->MeshEntries.Num());
-
-			for (const FPCGMeshSelectorWeightedEntry& Entry : SelectorWeighted->MeshEntries)
-			{
-				if (Entry.Descriptor.StaticMesh.IsNull())
-				{
-					continue;
-				}
-
-				if (ComponentsToCreate.Num() >= PCGComputeConstants::MAX_PRIMITIVE_COMPONENTS_PER_SPAWNER)
-				{
-					PCG_KERNEL_VALIDATION_WARN(Context, SpawnerSettings, /*bQuiet=*/false,
-						FText::Format(PCGComputeGraphElement::TooManyPrimitivesFormat, FText::FromString(FString::FromInt(PCGComputeConstants::MAX_PRIMITIVE_COMPONENTS_PER_SPAWNER))));
-
-					break;
-				}
-
-				if (UStaticMesh* StaticMesh = Entry.Descriptor.StaticMesh.LoadSynchronous())
-				{
-					const float Weight = float(Entry.Weight) / TotalWeight;
-					CumulativeWeight += Weight;
-					PrimitiveSelectionCDF.Add(CumulativeWeight);
-					PrimitiveStringKeys.Add(InBinding->GetStringTable().IndexOfByKey(Entry.Descriptor.StaticMesh.ToString()));
-
-					FPCGProceduralISMComponentDescriptor Descriptor;
-					Descriptor = Entry.Descriptor;
-					Descriptor.LocalBounds = LocalBounds;
-					Descriptor.NumCustomFloats = CustomFloatCount;
-					Descriptor.StaticMesh = StaticMesh;
-
-					uint32 InstanceCount = FMath::CeilToInt(TotalInputPointCount * Weight);
-
-					if (SelectorWeighted->MeshEntries.Num() > 1)
-					{
-						// Since we'll be selecting meshes based on random draws using the point random seeds which we don't have on CPU,
-						// we may pick more or less than the expected number of instances for each mesh. Use binomial variance to calculate
-						// the overallocation that gives 99.7% confidence (3 sigma).
-						const double Variance = double(TotalInputPointCount) * Weight * (1.0f - Weight);
-						const double Sigma = FMath::Sqrt(Variance);
-						const uint32 AdditionalAllocation = FMath::CeilToInt(3.0f * Sigma);
-
-						// Useful for debugging instance counts.
-						//UE_LOG(LogPCG, Warning, TEXT("Over allocation, N = %d, P = %.2f, NP = %d, 3sig = %f, Extra = %d (%.2f%%), Calculated = %d, Final = %d"), TotalInputPointCount, Weight, InstanceCount, float(3.0f * Sigma), AdditionalAllocation, float(AdditionalAllocation) / TotalInputPointCount, InstanceCount + AdditionalAllocation, FMath::Min(InstanceCount + AdditionalAllocation, TotalInputPointCount));
-
-						InstanceCount += AdditionalAllocation;
-					}
-
-					Descriptor.NumInstances = FMath::Min(InstanceCount, TotalInputPointCount);
-
-					ComponentsToCreate.Add(MoveTemp(Descriptor));
-
-					if (SpawnerSettings->bApplyMeshBoundsToPoints)
-					{
-						PrimitiveMeshBounds.Add(Descriptor.StaticMesh->GetBoundingBox());
-					}
-				}
-				else
-				{
-					PCG_KERNEL_VALIDATION_ERR(Context, SpawnerSettings, /*bQuiet=*/false,
-						FText::Format(PCGComputeGraphElement::CouldNotLoadStaticMeshFormat, FText::FromString(Entry.Descriptor.StaticMesh.ToString())));
-				}
-			}
-		}
-
-		// Validate instance count is not too large.
-		{
-			uint64 TotalInstanceCount = 0;
-			for (const FPCGProceduralISMComponentDescriptor& Desc : ComponentsToCreate)
-			{
-				TotalInstanceCount += Desc.NumInstances;
-			}
-
-			if (TotalInstanceCount >= MAX_INSTANCE_ID)
-			{
-				PCG_KERNEL_VALIDATION_ERR(Context, SpawnerSettings, /*bQuiet=*/false, FText::Format(
-					LOCTEXT("TooManyInstances", "Tried to spawn too many instances ({0}), procedural ISM component creation skipped and instances will not be rendered."), 
-					FText::FromString(FString::FromInt(TotalInstanceCount))));
-
-				return false;
-			}
-		}
-
-		FPCGSpawnerPrimitives& Primitives = InBinding->MeshSpawnersToPrimitives.FindOrAdd(SpawnerSettings);
-		Primitives.NumCustomFloats = CustomFloatCount;
-		Primitives.AttributeIdOffsetStrides = MoveTemp(AttributeIdOffsetStrides);
-		Primitives.SelectorAttributeId = SelectorAttributeId;
-		Primitives.SelectionCDF = MoveTemp(PrimitiveSelectionCDF);
-		Primitives.PrimitiveStringKeys = MoveTemp(PrimitiveStringKeys);
-		Primitives.PrimitiveMeshBounds = MoveTemp(PrimitiveMeshBounds);
-
-		for (const FPCGProceduralISMComponentDescriptor& Desc : ComponentsToCreate)
-		{
-			FPCGProceduralISMCBuilderParameters Params;
-			Params.Descriptor = Desc;
-			Params.bAllowDescriptorChanges = false;
-
-			UPCGManagedProceduralISMComponent* MISMC = PCGManagedProceduralISMComponent::GetOrCreateManagedProceduralISMC(TargetActor, InContext->SourceComponent.Get(), SpawnerSettings->UID, Params);
-
-			check(MISMC);
-			MISMC->SetCrc(InContext->DependenciesCrc);
-
-			// Don't bother registering the resource change as we're transient anyway.
-			//InContext->TouchedResources.Emplace(MISMC);
-
-			Primitives.Primitives.Add(MISMC->GetComponent());
-
-			Primitives.NumInstancesAllPrimitives += Desc.NumInstances;
-
-			bAnyComponentsSetup = true;
-		}
-
-		// Useful for debugging instance counts.
-		//UE_LOG(LogPCG, Log, TEXT("Input points %d, total instance count %d (+ %.2f%%)"), TotalInputPointCount, Primitives.NumInstancesAllPrimitives, (float(Primitives.NumInstancesAllPrimitives - TotalInputPointCount) / TotalInputPointCount) * 100.0f);
+		bAnyComponentsSetup |= Kernel->SetupProceduralISMComponents(InContext, InBinding);
 	}
 
 	return bAnyComponentsSetup;
diff --git a/Engine/Plugins/PCG/Source/PCG/Private/Compute/Elements/PCGCustomHLSL.cpp b/Engine/Plugins/PCG/Source/PCG/Private/Compute/Elements/PCGCustomHLSL.cpp
--- a/Engine/Plugins/PCG/Source/PCG/Private/Compute/Elements/PCGCustomHLSL.cpp
+++ b/Engine/Plugins/PCG/Source/PCG/Private/Compute/Elements/PCGCustomHLSL.cpp
@@ -2,17 +2,15 @@
 
 #include "Compute/Elements/PCGCustomHLSL.h"
 
-#include "PCGComputeGraphElement.h"
 #include "PCGContext.h"
 #include "PCGComponent.h"
 #include "PCGEdge.h"
 #include "PCGModule.h"
 #include "PCGPoint.h"
 #include "PCGSubsystem.h"
-#include "Compute/PCGComputeCommon.h"
-#include "Compute/PCGComputeGraph.h"
-#include "Compute/PCGDataBinding.h"
+#include "Compute/Elements/PCGCustomHLSLKernel.h"
 #include "Data/PCGPointData.h"
+#include "Graph/PCGGPUGraphCompilationContext.h"
 
 #include "Internationalization/Regex.h"
 #include "Containers/StaticArray.h"
@@ -25,51 +23,6 @@
 
 #define LOCTEXT_NAMESPACE "PCGCustomHLSLElement"
 
-namespace PCGHLSLElement
-{
-	const FString PinDeclTemplateStr = TEXT("{pin}");
-
-	/** First capture: Pin name (supports a - z, A - Z, and 0 - 9) */
-	constexpr int AttributePinCaptureGroup = 1;
-
-	/** Second capture: Function name (Get or Set) */
-	constexpr int AttributeFunctionCaptureGroup = 2;
-
-	/** Third capture: Attribute type (e.g. Int, Float, Rotator, etc.) */
-	constexpr int AttributeTypeCaptureGroup = 3;
-
-	/** Fourth capture: Attribute name (supports a-z, A-Z, 0-9, ' ', '-', '_', and '/') */
-	constexpr int AttributeNameCaptureGroup = 4;
-
-	/** Regex pattern used to detect and parse attribute function usage in kernels. */
-	constexpr TCHAR AttributeFunctionPattern[] = { TEXT("([a-zA-Z0-9]+)_(Get|Set)(.*?)[\\s]*?\\(.*?,.*?,.*?'([a-zA-Z0-9 -_\\/]*?)'") };
-
-	constexpr TCHAR AttributeFunctionGet[] = { TEXT("Get") };
-	constexpr TCHAR AttributeFunctionSet[] = { TEXT("Set") };
-
-	void ConvertObjectPathToShaderFilePath(FString& InOutPath)
-	{
-		// Shader compiler recognizes "/Engine/Generated/..." path as special. 
-		// It doesn't validate file suffix etc.
-		InOutPath = FString::Printf(TEXT("/Engine/Generated/UObject%s.ush"), *InOutPath);
-		// Shader compilation result parsing will break if it finds ':' where it doesn't expect.
-		InOutPath.ReplaceCharInline(TEXT(':'), TEXT('@'));
-	}
-
-	FString GetKernelAttributeKeyAsString(const FPCGKernelAttributeKey& Key)
-	{
-		return FString::Format(TEXT("'{0}'"), { Key.Name.ToString() });
-	}
-
-	FString GetDataTypeString(EPCGDataType Type)
-	{
-		const UEnum* DataTypeEnum = StaticEnum<EPCGDataType>();
-		check(DataTypeEnum);
-
-		return DataTypeEnum->GetValueOrBitfieldAsString(static_cast<int64>(Type));
-	}
-}
-
 UPCGCustomHLSLSettings::UPCGCustomHLSLSettings()
 {
 	bExecuteOnGPU = true;
@@ -186,284 +139,6 @@ FPCGElementPtr UPCGCustomHLSLSettings::CreateElement() const
 	return MakeShared<FPCGCustomHLSLElement>();
 }
 
-int UPCGCustomHLSLSettings::GetProcessingElemCountForInputPin(const UPCGPin* InputPin, const UPCGDataBinding* Binding) const
-{
-	check(InputPin && Binding);
-
-	// Upper bound estimate of total number of data elements expected to arrive at this pin.
-	int ProcessingElemCount = 0;
-
-	if (Binding->InputPins.Contains(InputPin))
-	{
-		FName PinLabel = InputPin->Properties.Label;
-		if (const FName* PinLabelAlias = Binding->InputPinLabelAliases.Find(InputPin))
-		{
-			PinLabel = *PinLabelAlias;
-		}
-
-		// Data coming straight from CPU, know its size now so count it.
-		TArray<FPCGTaggedData> DataForPin = Binding->InputDataCollection.GetInputsByPin(PinLabel);
-		for (const FPCGTaggedData& Data : DataForPin)
-		{
-			ProcessingElemCount += PCGComputeHelpers::GetElementCount(Data.Data);
-		}
-	}
-	else
-	{
-		// Estimate (upper bound) element count by looking at incident connections.
-		for (const UPCGEdge* Edge : InputPin->Edges)
-		{
-			// InputPin is upstream output pin.
-			const UPCGPin* UpstreamOutputPin = Edge->InputPin;
-			if (!UpstreamOutputPin)
-			{
-				continue;
-			}
-
-			const UPCGSettings* UpstreamSettings = UpstreamOutputPin->Node ? UpstreamOutputPin->Node->GetSettings() : nullptr;
-			check(UpstreamSettings);
-
-			FPCGDataCollectionDesc PinDesc;
-			if (ensure(Binding->GetOutputPinDataDesc_GameThread(UpstreamOutputPin, UpstreamSettings, PinDesc)))
-			{
-				ProcessingElemCount += PinDesc.ComputeDataElementCount(InputPin->Properties.AllowedTypes);
-			}
-		}
-	}
-
-	return ProcessingElemCount;
-}
-
-int UPCGCustomHLSLSettings::ComputeKernelThreadCount(const UPCGDataBinding* Binding) const
-{
-	int ThreadCount = 0;
-
-	if (KernelType == EPCGKernelType::PointGenerator)
-	{
-		// Point generator has fixed thread count.
-		ThreadCount = PointCount;
-	}
-	else if (KernelType == EPCGKernelType::PointProcessor)
-	{
-		// Processing volume depends on data arriving on primary pin.
-		if (const UPCGPin* InputPin = GetPointProcessingInputPin())
-		{
-			ThreadCount = GetProcessingElemCountForInputPin(InputPin, Binding);
-		}
-	}
-	else if (KernelType == EPCGKernelType::Custom)
-	{
-		if (DispatchThreadCount == EPCGDispatchThreadCount::FromFirstOutputPin)
-		{
-			if (const UPCGPin* OutputPin = GetFirstOutputPin())
-			{
-				FPCGDataCollectionDesc Desc;
-				if (ensure(Binding->GetOutputPinDataDesc_GameThread(OutputPin, this, Desc)))
-				{
-					ThreadCount = Desc.ComputeDataElementCount(EPCGDataType::Any);
-				}
-			}
-		}
-		else if (DispatchThreadCount == EPCGDispatchThreadCount::FromProductOfInputPins)
-		{
-			for (const FName& PinLabel : ThreadCountInputPinLabels)
-			{
-				if (const UPCGPin* InputPin = GetInputPin(PinLabel))
-				{
-					ThreadCount = GetProcessingElemCountForInputPin(InputPin, Binding);
-				}
-			}
-		}
-		else if (DispatchThreadCount == EPCGDispatchThreadCount::Fixed)
-		{
-			ThreadCount = FixedThreadCount;
-		}
-	}
-	else
-	{
-		checkNoEntry();
-	}
-
-	if (IsThreadCountMultiplierInUse())
-	{
-		ThreadCount *= ThreadCountMultiplier;
-	}
-
-	return ThreadCount;
-}
-
-bool UPCGCustomHLSLSettings::ComputeOutputPinDataDesc(const UPCGPin* OutputPin, const UPCGDataBinding* Binding, FPCGDataCollectionDesc& OutDesc) const
-{
-	check(OutputPin);
-	check(Binding);
-
-	// First inspect data coming from CPU. Will return description if this node is a CPU node or in a different compute graph.
-	if (Binding->ComputeCPUOutputPinDataDesc(OutputPin, OutDesc))
-	{
-		return true;
-	}
-
-	FPCGDataCollectionDesc PinDesc;
-
-	const FPCGPinPropertiesGPU* Properties = GetOutputPinPropertiesGPU(OutputPin->Properties.Label);
-	const UPCGPin* FirstOutputPin = GetFirstOutputPin();
-
-	// The primary output pin follows any rules prescribed by kernel type.
-	if (OutputPin == FirstOutputPin && KernelType == EPCGKernelType::PointProcessor)
-	{
-		// First output pin passes through first input pin.
-		if (const UPCGPin* PointProcessingInputPin = GetPointProcessingInputPin())
-		{
-			PinDesc = PCGDataForGPUHelpers::ComputeInputPinDataDesc(PointProcessingInputPin, Binding);
-		}
-	}
-	else if (OutputPin == FirstOutputPin && KernelType == EPCGKernelType::PointGenerator)
-	{
-		// Generators always produce a single point data with known point count.
-		PinDesc.DataDescs.Emplace(EPCGDataType::Point, PointCount);
-	}
-	else if (ensure(Properties))
-	{
-		PCGComputeHelpers::ComputeOutputPinDataDesc(*Properties, this, Binding, PinDesc);
-	}
-
-	for (const FPCGKernelAttributeKey& AttributeKey : KernelAttributeKeys)
-	{
-		// Add attributes that will be created for this pin on the GPU.
-		if (const TArray<TTuple<FPCGKernelAttributeKey, bool>>* Keys = PinToAttributeKeys.Find(OutputPin->Properties.Label))
-		{
-			const TTuple<FPCGKernelAttributeKey, bool>* Pair = Keys->FindByPredicate([AttributeKey](const TTuple<FPCGKernelAttributeKey, bool>& Pair) { return Pair.Key == AttributeKey; });
-			if (!Pair || !Pair->Value)
-			{
-				// Not created on GPU.
-				continue;
-			}
-
-			for (FPCGDataDesc& DataDesc : PinDesc.DataDescs)
-			{
-				const int32 AttributeId = Binding->GetAttributeId(AttributeKey);
-
-				if (AttributeId != INDEX_NONE)
-				{
-					const FPCGKernelAttributeDesc AttributeDesc(AttributeId, AttributeKey.Type, AttributeKey.Name);
-					DataDesc.AttributeDescs.AddUnique(AttributeDesc);
-				}
-			}
-		}
-	}
-
-	// Try to propagate string keys across node. Not trivial because there could be one or more string key attributes on input pins and on output pins,
-	// and it is in general hard to determine from source which string keys from input are being written to outputs. Try first collecting all string keys
-	// from matching attribute names (across all input pins), and then fall back to collecting keys from all string key attributes across all inputs.
-	bool bOutputHasStringKeys = false;
-
-	for (const FPCGDataDesc& DataDesc : PinDesc.DataDescs)
-	{
-		auto HasStringKey = [](const FPCGKernelAttributeDesc& InAttributeDesc)
-		{
-			return InAttributeDesc.Type == EPCGKernelAttributeType::StringKey || InAttributeDesc.Type == EPCGKernelAttributeType::Name;
-		};
-
-		if (DataDesc.AttributeDescs.FindByPredicate(HasStringKey))
-		{
-			bOutputHasStringKeys = true;
-			break;
-		}
-	}
-
-	const UPCGNode* Node = Cast<UPCGNode>(GetOuter());
-	if (bOutputHasStringKeys && Node)
-	{
-		TArray<FPCGDataCollectionDesc> RelevantInputDataDescs;
-
-		// Collect descriptions of input data items that have string key attributes.
-		for (const UPCGPin* InputPin : Node->GetInputPins())
-		{
-			const FPCGDataCollectionDesc InputPinDesc = PCGDataForGPUHelpers::ComputeInputPinDataDesc(InputPin, Binding);
-
-			bool bFoundStringKeyAttribute = false;
-
-			for (FPCGDataDesc& DataDesc : PinDesc.DataDescs)
-			{
-				for (FPCGKernelAttributeDesc& AttributeDesc : DataDesc.AttributeDescs)
-				{
-					bFoundStringKeyAttribute |= (AttributeDesc.Type == EPCGKernelAttributeType::StringKey || AttributeDesc.Type == EPCGKernelAttributeType::Name);
-				}
-			}
-
-			if (bFoundStringKeyAttribute)
-			{
-				RelevantInputDataDescs.Add(InputPinDesc);
-			}
-		}
-
-		if (!RelevantInputDataDescs.IsEmpty())
-		{
-			for (FPCGDataDesc& DataDesc : PinDesc.DataDescs)
-			{
-				for (FPCGKernelAttributeDesc& AttributeDesc : DataDesc.AttributeDescs)
-				{
-					if (AttributeDesc.Type != EPCGKernelAttributeType::StringKey && AttributeDesc.Type != EPCGKernelAttributeType::Name)
-					{
-						continue;
-					}
-
-					bool bFoundMatchingAttribute = false;
-
-					for (const FPCGDataCollectionDesc& InputPinDataDesc : RelevantInputDataDescs)
-					{
-						// Try to find string keys for matching attributes on inputs. E.g. if we are processing an output attribute named 'MeshPath',
-						// look at data on all input pins for an attribute named MeshPath and assume we could use any of its values - copy the string keys.
-						for (const FPCGDataDesc& InputDataDesc : InputPinDataDesc.DataDescs)
-						{
-							for (const FPCGKernelAttributeDesc& InputAttributeDesc : InputDataDesc.AttributeDescs)
-							{
-								const bool bIsString = InputAttributeDesc.Type == EPCGKernelAttributeType::StringKey || InputAttributeDesc.Type == EPCGKernelAttributeType::Name;
-								if (bIsString && InputAttributeDesc.Name == AttributeDesc.Name)
-								{
-									AttributeDesc.AddUniqueStringKeys(InputAttributeDesc.GetUniqueStringKeys());
-									bFoundMatchingAttribute = true;
-									break;
-								}
-							}
-						}
-					}
-
-					if (!bFoundMatchingAttribute)
-					{
-						// We didn't find an exact attribute. Fall back to finding any and all string keys. This is concerning and perhaps we can
-						// have additional hinting mechanisms in the kernel source or in the node UI.
-						for (const FPCGDataCollectionDesc& InputPinDataDesc : RelevantInputDataDescs)
-						{
-							// Try to find string keys for matching attributes on inputs. E.g. if we are processing an output attribute named 'MeshPath',
-							// look at data on all input pins for an attribute named MeshPath and assume we could use any of its values - copy the string keys.
-							for (const FPCGDataDesc& InputDataDesc : InputPinDataDesc.DataDescs)
-							{
-								for (const FPCGKernelAttributeDesc& InputAttributeDesc : InputDataDesc.AttributeDescs)
-								{
-									if (InputAttributeDesc.Type == EPCGKernelAttributeType::StringKey || InputAttributeDesc.Type == EPCGKernelAttributeType::Name)
-									{
-										AttributeDesc.AddUniqueStringKeys(InputAttributeDesc.GetUniqueStringKeys());
-									}
-								}
-							}
-						}
-					}
-				}
-			}
-		}
-		else
-		{
-			// If there were no string keys found on any input pin then we are in a bad place. String values cannot be built on the GPU, they must
-			// come in through an input.
-			UE_LOG(LogPCG, Warning, TEXT("No incoming attributes to obtain string keys from."));
-		}
-	}
-
-	OutDesc = PinDesc;
-	return true;
-}
-
 #if WITH_EDITOR
 TArray<FPCGPreConfiguredSettingsInfo> UPCGCustomHLSLSettings::GetPreconfiguredInfo() const
 {
@@ -649,22 +324,14 @@ bool UPCGCustomHLSLSettings::IsShaderTextReadOnly() const
 	return false;
 }
 
-void UPCGCustomHLSLSettings::GetDataLabels(FName InPinLabel, TArray<FString>& OutDataLabels) const
+UPCGComputeKernel* UPCGCustomHLSLSettings::CreateKernel(FPCGGPUCompilationContext& InOutContext, UObject* InObjectOuter) const
 {
-	// Matches against {PinName}_AnyFunction('{DataLabel}'...
-	const FString Pattern = FString::Format(TEXT("{0}_.*?[\\s]*?\\([\\s]*?'([a-zA-Z0-9_].*?)'"), { *InPinLabel.ToString() });
-	const FString Source = ShaderFunctions + TEXT("\n") + ShaderSource;
-
-	// First capture: Data label (supports a - z, A - Z, 0 - 9, and underscores).
-	constexpr int LabelCaptureGroup = 1;
+	UPCGCustomHLSLKernel* Kernel = InOutContext.NewObject_AnyThread<UPCGCustomHLSLKernel>(InObjectOuter);
+	Kernel->Settings = this;
 
-	FRegexMatcher ModuleMatcher(FRegexPattern(Pattern), Source);
-	while (ModuleMatcher.FindNext())
-	{
-		OutDataLabels.AddUnique(ModuleMatcher.GetCaptureGroup(LabelCaptureGroup));
-	}
+	return Kernel;
 }
-#endif //WITH_EDITOR
+#endif // WITH_EDITOR
 
 const FPCGPinPropertiesGPU* UPCGCustomHLSLSettings::GetOutputPinPropertiesGPU(const FName& InPinLabel) const
 {
@@ -1238,782 +905,6 @@ TArray<FName> UPCGCustomHLSLSettings::GetInputPinNamesAndNone() const
 }
 #endif // WITH_EDITOR
 
-bool UPCGCustomHLSLSettings::IsKernelValid(FPCGContext* InContext, bool bQuiet) const
-{
-	TRACE_CPUPROFILER_EVENT_SCOPE(UPCGCustomHLSLSettings::IsKernelValid);
-
-	if (OutputPins.IsEmpty())
-	{
-		PCG_KERNEL_VALIDATION_ERR(InContext, this, bQuiet, LOCTEXT("NoOutputs", "Custom HLSL nodes must have at least one output."));
-		return false;
-	}
-
-	auto CheckPinLabel = [This = this, InContext, bQuiet, &InPins = InputPins, OutPins = OutputPinProperties()](FName PinLabel)
-	{
-		if (PinLabel == NAME_None)
-		{
-			PCG_KERNEL_VALIDATION_ERR(InContext, This, bQuiet, LOCTEXT("InvalidPinLabelNone", "Pin label 'None' is not a valid pin label."));
-			return false;
-		}
-
-		bool bFoundPinLabel = false;
-
-		auto IsAlreadyFound = [This, InContext, bQuiet, PinLabel, &bFoundPinLabel](const FPCGPinProperties PinProps)
-		{
-			if (PinProps.Label == PinLabel)
-			{
-				if (bFoundPinLabel)
-				{
-					PCG_KERNEL_VALIDATION_ERR(InContext, This, bQuiet, FText::Format(LOCTEXT("DuplicatedPinLabels", "Duplicate pin label '{0}', all labels must be unique."), FText::FromName(PinLabel)));
-					return true;
-				}
-
-				bFoundPinLabel = true;
-			}
-
-			return false;
-		};
-
-		for (const FPCGPinProperties& PinProps : InPins)
-		{
-			if (IsAlreadyFound(PinProps))
-			{
-				return false;
-			}
-		}
-
-		for (const FPCGPinProperties& PinProps : OutPins)
-		{
-			if (IsAlreadyFound(PinProps))
-			{
-				return false;
-			}
-		}
-
-		return true;
-	};
-
-	// Validate input pins
-	bool bIsFirstInputPin = true;
-	for (const FPCGPinProperties& Properties : InputPins)
-	{
-		if (!CheckPinLabel(Properties.Label))
-		{
-			return false;
-		}
-
-		if (bIsFirstInputPin && KernelType == EPCGKernelType::PointProcessor)
-		{
-			if (Properties.AllowedTypes != EPCGDataType::Point)
-			{
-				PCG_KERNEL_VALIDATION_ERR(InContext, this, bQuiet, FText::Format(
-					LOCTEXT("InvalidNonPointPrimaryInput", "'Point Processor' nodes require primary input pin to be of type 'Point', but found '{0}'."),
-					FText::FromString(PCGHLSLElement::GetDataTypeString(Properties.AllowedTypes))));
-
-				return false;
-			}
-		}
-
-		if (!PCGComputeHelpers::IsTypeAllowedAsInput(Properties.AllowedTypes))
-		{
-			PCG_KERNEL_VALIDATION_ERR(InContext, this, bQuiet, FText::Format(
-				LOCTEXT("InvalidInputType", "Unsupported input type '{0}', found on pin '{1}'."),
-				FText::FromString(PCGHLSLElement::GetDataTypeString(Properties.AllowedTypes)),
-				FText::FromName(Properties.Label)));
-
-			return false;
-		}
-
-		bIsFirstInputPin = false;
-	}
-
-	// Validate output pins
-	bool bIsFirstOutputPin = true;
-	for (const FPCGPinPropertiesGPU& Properties : OutputPins)
-	{
-		if (!CheckPinLabel(Properties.Label))
-		{
-			return false;
-		}
-
-		const bool bPinIsDefinedByKernel = bIsFirstOutputPin && (KernelType == EPCGKernelType::PointGenerator || KernelType == EPCGKernelType::PointProcessor);
-
-		if (bPinIsDefinedByKernel)
-		{
-			if (Properties.AllowedTypes != EPCGDataType::Point)
-			{
-				PCG_KERNEL_VALIDATION_ERR(InContext, this, bQuiet, FText::Format(
-					LOCTEXT("InvalidNonPointPrimaryOutput", "'Point Processor' and 'Point Generator' nodes require primary output pin to be of type 'Point', but found '{0}'."),
-					FText::FromString(PCGHLSLElement::GetDataTypeString(Properties.AllowedTypes))));
-
-				return false;
-			}
-		}
-
-		if (!PCGComputeHelpers::IsTypeAllowedAsOutput(Properties.AllowedTypes))
-		{
-			PCG_KERNEL_VALIDATION_ERR(InContext, this, bQuiet, FText::Format(
-				LOCTEXT("InvalidOutputType", "Unsupported output type '{0}', found on pin '{1}'."),
-				FText::FromString(PCGHLSLElement::GetDataTypeString(Properties.AllowedTypes)),
-				FText::FromName(Properties.Label)));
-
-			return false;
-		}
-
-		if (!bPinIsDefinedByKernel)
-		{
-			const FPCGPinPropertiesGPUStruct& Props = Properties.PropertiesGPU;
-
-			if (Props.InitializationMode == EPCGPinInitMode::FromInputPins)
-			{
-				if (Props.PinsToInititalizeFrom.IsEmpty())
-				{
-					PCG_KERNEL_VALIDATION_ERR(InContext, this, bQuiet, FText::Format(
-						LOCTEXT("InitFromEmptyPins", "Output pin '{0}' tried to initialize from input pins, but no pins were specified."),
-						FText::FromName(Properties.Label)));
-
-					return false;
-				}
-
-				for (const FName InitPinName : Props.PinsToInititalizeFrom)
-				{
-					const FPCGPinProperties* InitPinProps = InputPins.FindByPredicate([InitPinName](const FPCGPinProperties& InPinProps)
-					{
-						return InPinProps.Label == InitPinName;
-					});
-
-					if (!InitPinProps)
-					{
-						PCG_KERNEL_VALIDATION_ERR(InContext, this, bQuiet, FText::Format(
-							LOCTEXT("InitFromNonExistentPin", "Output pin '{0}' tried to initialize from non-existent input pin '{1}'."),
-							FText::FromName(Properties.Label),
-							FText::FromName(InitPinName)));
-
-						return false;
-					}
-
-					if (!PCGComputeHelpers::IsTypeAllowedInDataCollection(InitPinProps->AllowedTypes))
-					{
-						PCG_KERNEL_VALIDATION_ERR(InContext, this, bQuiet, FText::Format(
-							LOCTEXT("InitFromInvalidPinType", "Output pin '{0}' tried to initialize from input pin '{1}', but pin '{1}' has an invalid type."),
-							FText::FromName(Properties.Label),
-							FText::FromName(InitPinName)));
-
-						return false;
-					}
-				}
-
-				// TODO: Could do validation on data multiplicity for Pairwise, checking that data counts are 1 or N, but maybe that should be a runtime error instead.
-			}
-
-			const bool bUsingFixedDataCount = Props.InitializationMode == EPCGPinInitMode::Custom || Props.DataCountMode == EPCGDataCountMode::Fixed;
-
-			if (bUsingFixedDataCount)
-			{
-				if (Props.DataCount < 1)
-				{
-					PCG_KERNEL_VALIDATION_ERR(InContext, this, bQuiet, FText::Format(
-						LOCTEXT("InvalidDataCount", "Invalid fixed data count {0} on output pin '{1}'. Must be greater than 0."),
-						FText::AsNumber(Props.DataCount),
-						FText::FromName(Properties.Label)));
-
-					return false;
-				}
-
-			}
-
-			const bool bUsingFixedElemCount = Props.InitializationMode == EPCGPinInitMode::Custom || Props.ElementCountMode == EPCGElementCountMode::Fixed;
-
-			if (bUsingFixedElemCount)
-			{
-				if (Props.ElementCount < 1)
-				{
-					PCG_KERNEL_VALIDATION_ERR(InContext, this, bQuiet, FText::Format(
-						LOCTEXT("InvalidElementCount", "Invalid fixed element count {0} on output pin '{1}'. Must be greater than 0."),
-						FText::AsNumber(Props.ElementCount),
-						FText::FromName(Properties.Label)));
-
-					return false;
-				}
-			}
-		}
-
-		bIsFirstOutputPin = false;
-	}
-
-	if (KernelType == EPCGKernelType::Custom && DispatchThreadCount == EPCGDispatchThreadCount::FromProductOfInputPins)
-	{
-		if (ThreadCountInputPinLabels.IsEmpty())
-		{
-			PCG_KERNEL_VALIDATION_ERR(InContext, this, bQuiet, LOCTEXT("MissingThreadCountPins", "Dispatch thread count is based on input pins but no labels have been set in Input Pins array."));
-			return false;
-		}
-
-		for (const FName& Label : ThreadCountInputPinLabels)
-		{
-			if (!GetInputPin(Label))
-			{
-				PCG_KERNEL_VALIDATION_ERR(InContext, this, bQuiet, FText::Format(LOCTEXT("MissingThreadCountPin", "Invalid pin specified in Input Pins array: '{0}'."), FText::FromName(Label)));
-				return false;
-			}
-		}
-	}
-
-	if (IsThreadCountMultiplierInUse())
-	{
-		if (ThreadCountMultiplier < 1)
-		{
-			PCG_KERNEL_VALIDATION_ERR(InContext, this, bQuiet, FText::Format(LOCTEXT("InvalidThreadCountMultiplier", "Thread Count Multiplier has invalid value ({0}). Must be greater than 0."), ThreadCountMultiplier));
-			return false;
-		}
-	}
-
-	// Validate attributes
-	for (const FPCGKernelAttributeKey& AttributeKey : KernelAttributeKeys)
-	{
-		if (AttributeKey.Type == EPCGKernelAttributeType::Invalid)
-		{
-			const UEnum* AttributeTypeEnum = StaticEnum<EPCGKernelAttributeType>();
-			check(AttributeTypeEnum);
-
-			PCG_KERNEL_VALIDATION_ERR(InContext, this, bQuiet, FText::Format(
-				LOCTEXT("InvalidAttributeTypeNone", "Attribute '{0}' has invalid GPU attribute type '{1}', check the 'Attributes to Create' array on your pins."),
-				FText::FromName(AttributeKey.Name),
-				FText::FromString(AttributeTypeEnum->GetNameStringByValue(static_cast<int64>(AttributeKey.Type)))));
-			return false;
-		}
-	}
-
-	if (InContext)
-	{
-		FText* ErrorTextPtr = nullptr;
-#if PCG_KERNEL_LOGGING_ENABLED
-		FText ErrorText;
-		ErrorTextPtr = &ErrorText;
-#endif
-
-		if (!AreKernelAttributesValid(InContext, ErrorTextPtr))
-		{
-			if (ErrorTextPtr)
-			{
-				PCG_KERNEL_VALIDATION_ERR(InContext, this, bQuiet, *ErrorTextPtr);
-			}
-			return false;
-		}
-
-		if (!ProcessShaderSource(/*InOutSource=*/nullptr, ErrorTextPtr))
-		{
-			if (ErrorTextPtr)
-			{
-				PCG_KERNEL_VALIDATION_ERR(InContext, this, bQuiet, *ErrorTextPtr);
-			}
-			return false;
-		}
-
-		if (!bMuteUnwrittenPinDataErrors && !AreAllOutputPinsWritten(ErrorTextPtr))
-		{
-			if (ErrorTextPtr)
-			{
-				PCG_KERNEL_VALIDATION_ERR(InContext, this, bQuiet, *ErrorTextPtr);
-			}
-			return false;
-		}
-	}
-
-	return true;
-}
-
-bool UPCGCustomHLSLSettings::AreKernelAttributesValid(FPCGContext* InContext, FText* OutErrorText) const
-{
-	TRACE_CPUPROFILER_EVENT_SCOPE(UPCGCustomHLSLSettings::AreKernelAttributesValid);
-
-	// The context can either be a compute graph element context (if the compute graph was successfully created), otherwise
-	// it will be the original CPU node context. We need the former to run the following validation.
-	if (!InContext || !InContext->IsComputeContext())
-	{
-		return true;
-	}
-
-	const FPCGComputeGraphContext* Context = static_cast<FPCGComputeGraphContext*>(InContext);
-	const UPCGDataBinding* DataBinding = Context->DataBinding.Get();
-	const UPCGNode* Node = Cast<UPCGNode>(GetOuter());
-
-	if (DataBinding && Node)
-	{
-		const TArray<TObjectPtr<UPCGPin>>& InPins = Node->GetInputPins();
-		const TArray<TObjectPtr<UPCGPin>>& OutPins = Node->GetOutputPins();
-
-		TMap<FName, FPCGDataCollectionDesc> InputPinDescs;
-		TMap<FName, FPCGDataCollectionDesc> OutputPinDescs;
-
-		for (const UPCGPin* InputPin : InPins)
-		{
-			check(InputPin);
-			InputPinDescs.Add(InputPin->Properties.Label, PCGDataForGPUHelpers::ComputeInputPinDataDesc(InputPin, DataBinding));
-		}
-
-		for (const UPCGPin* OutputPin : OutPins)
-		{
-			check(OutputPin);
-
-			FPCGDataCollectionDesc PinDesc;
-			if (ensure(DataBinding->GetOutputPinDataDesc_GameThread(OutputPin, this, PinDesc)))
-			{
-				OutputPinDescs.Add(OutputPin->Properties.Label, PinDesc);
-			}
-		}
-
-		const FString Source = ShaderFunctions + ShaderSource;
-		FRegexMatcher ModuleMatcher(FRegexPattern(PCGHLSLElement::AttributeFunctionPattern), Source);
-
-		while (ModuleMatcher.FindNext())
-		{
-			const FString PinStr = ModuleMatcher.GetCaptureGroup(PCGHLSLElement::AttributePinCaptureGroup);
-			const FString FuncStr = ModuleMatcher.GetCaptureGroup(PCGHLSLElement::AttributeFunctionCaptureGroup);
-			const FString TypeStr = ModuleMatcher.GetCaptureGroup(PCGHLSLElement::AttributeTypeCaptureGroup);
-			const FString NameStr = ModuleMatcher.GetCaptureGroup(PCGHLSLElement::AttributeNameCaptureGroup);
-
-			const FName PinName = FName(*PinStr);
-			const FPCGDataCollectionDesc* PinDesc = nullptr;
-
-			auto ConstructFunctionText = [&PinStr, &FuncStr, &TypeStr]()
-			{
-				return FText::FromString(PinStr + TEXT("_") + FuncStr + TypeStr);
-			};
-
-			if (FuncStr == PCGHLSLElement::AttributeFunctionSet)
-			{
-				PinDesc = OutputPinDescs.Find(FName(*PinStr));
-
-				if (!PinDesc && InputPinDescs.Find(PinName))
-				{
-#if PCG_KERNEL_LOGGING_ENABLED
-					if (OutErrorText)
-					{
-						*OutErrorText = FText::Format(
-							LOCTEXT("InvalidSetAttributeUsage", "Tried to call attribute function '{0}' on read-only input pin '{1}'."),
-							ConstructFunctionText(),
-							FText::FromName(PinName));
-					}
-#endif
-
-					return false;
-				}
-			}
-			else if (ensure(FuncStr == PCGHLSLElement::AttributeFunctionGet))
-			{
-				PinDesc = InputPinDescs.Find(FName(*PinStr));
-			}
-
-			if (!PinDesc)
-			{
-#if PCG_KERNEL_LOGGING_ENABLED
-				if (OutErrorText)
-				{
-					*OutErrorText = FText::Format(
-						LOCTEXT("InvalidAttributePinName", "Tried to call attribute function '{0}' on non-existent pin '{1}'."),
-						ConstructFunctionText(),
-						FText::FromName(PinName));
-				}
-#endif
-
-				return false;
-			}
-
-			const UEnum* AttributeTypeEnum = StaticEnum<EPCGKernelAttributeType>();
-			check(AttributeTypeEnum);
-
-			const int64 AttributeType = AttributeTypeEnum->GetValueByName(FName(*TypeStr));
-
-			if (AttributeType == INDEX_NONE)
-			{
-#if PCG_KERNEL_LOGGING_ENABLED
-				if (OutErrorText)
-				{
-					*OutErrorText = FText::Format(
-						LOCTEXT("InvalidAttributePinType", "Tried to call attribute function '{0}' on non-existent type '{1}'."),
-						ConstructFunctionText(),
-						FText::FromString(TypeStr));
-				}
-#endif
-
-				return false;
-			}
-
-			const FPCGKernelAttributeDesc* AttrDesc = nullptr;
-			const FName AttrName = FName(*NameStr);
-
-			if (!PinDesc->DataDescs.IsEmpty())
-			{
-				// Note: This assumes attributes are the same on all data on a pin, which is true for now
-				const FPCGDataDesc& DataDesc = PinDesc->DataDescs[0];
-				AttrDesc = DataDesc.AttributeDescs.FindByPredicate([AttrName](const FPCGKernelAttributeDesc& Desc) { return Desc.Name == AttrName; });
-			}
-
-			if (!AttrDesc)
-			{
-#if PCG_KERNEL_LOGGING_ENABLED
-				if (OutErrorText)
-				{
-					*OutErrorText = FText::Format(
-						LOCTEXT("InvalidAttributeName", "Tried to call attribute function '{0}' on attribute '{1}' which does not exist."),
-						ConstructFunctionText(),
-						FText::FromName(AttrName));
-				}
-#endif
-
-				return false;
-			}
-
-			if (AttrDesc->Type != static_cast<EPCGKernelAttributeType>(AttributeType))
-			{
-#if PCG_KERNEL_LOGGING_ENABLED
-				if (OutErrorText)
-				{
-					const FString ActualTypeStr = AttributeTypeEnum->GetNameStringByIndex(static_cast<int64>(AttrDesc->Type));
-
-					*OutErrorText = FText::Format(
-						LOCTEXT("AttributeTypeMismatch", "Type mismatch for call to attribute function '{0}' on attribute '{1}'. Expected '{2}' but received '{3}'."),
-						ConstructFunctionText(),
-						FText::FromName(AttrName),
-						FText::FromString(TypeStr),
-						FText::FromString(ActualTypeStr));
-				}
-#endif
-
-				return false;
-			}
-		}
-	}
-
-	return true;
-}
-
-bool UPCGCustomHLSLSettings::AreAllOutputPinsWritten(FText* OutWarningText) const
-{
-	TRACE_CPUPROFILER_EVENT_SCOPE(UPCGCustomHLSLSettings::AreAllOutputPinsWritten);
-
-	const UPCGNode* Node = Cast<UPCGNode>(GetOuter());
-	if (!Node)
-	{
-		return true;
-	}
-
-	const FString AllSource = ShaderFunctions + ShaderSource;
-
-	// Non-custom kernels initialize the first output pin data automatically.
-	const UPCGPin* AutomaticallyInitializedPin = (KernelType != EPCGKernelType::Custom) ? GetFirstOutputPin() : nullptr;
-
-	for (const UPCGPin* Pin : Node->GetOutputPins())
-	{
-		if (!Pin || Pin == AutomaticallyInitializedPin || Pin->Properties.Label == NAME_None)
-		{
-			continue;
-		}
-
-		const TStaticArray<FString, 3> WriteStrings = { 
-			Pin->Properties.Label.ToString() + TEXT("_Set"), 
-			Pin->Properties.Label.ToString() + TEXT("_CopyElementFrom_"),
-			Pin->Properties.Label.ToString() + TEXT("_Initialize") };
-
-		bool bFoundDataWrite = false;
-
-		for (const FString& WriteString : WriteStrings)
-		{
-			if (AllSource.Contains(WriteString))
-			{
-				bFoundDataWrite = true;
-				break;
-			}
-		}
-
-		if (!bFoundDataWrite)
-		{
-			if (OutWarningText)
-			{
-				*OutWarningText = FText::Format(
-					LOCTEXT("PinMayNotBeWritten", "Data on pin '{0}' may be uninitialized. Add code to write to this data, or mute this error in the node settings."),
-					{ FText::FromName(Pin->Properties.Label) });
-			}
-
-			return false;
-		}
-	}
-
-	return true;
-}
-
-bool UPCGCustomHLSLSettings::ProcessShaderSource(FString* InOutSource, FText* OutErrorText) const
-{
-	TSet<TPair<FString, FString>> PinsRequiringElementCopy;
-
-	const FString Source = ShaderFunctions + TEXT("\n") + ShaderSource;
-
-	/** Regex for parsing usages of CopyElementFrom function. */
-	constexpr TCHAR CopyElementPattern[] = { TEXT(/*TargetPin*/"([a-zA-Z0-9]+)" "_CopyElementFrom_" /*SourcePin*/"([a-zA-Z0-9]+)\\s*\\(") };
-	constexpr int CopyElementTargetPinCaptureGroup = 1;
-	constexpr int CopyElementSourcePinCaptureGroup = 2;
-
-	FRegexMatcher ModuleMatcher(FRegexPattern(CopyElementPattern), Source);
-
-	while (ModuleMatcher.FindNext())
-	{
-		const FString SourcePinStr = ModuleMatcher.GetCaptureGroup(CopyElementSourcePinCaptureGroup);
-		const UPCGPin* SourcePin = GetInputPin(FName(*SourcePinStr));
-		if (!SourcePin)
-		{
-#if PCG_KERNEL_LOGGING_ENABLED
-			if (OutErrorText)
-			{
-				*OutErrorText = FText::Format(LOCTEXT("InvalidCopySourcePin", "Invalid input pin '{0}' used in copy element expression."), FText::FromString(SourcePinStr));
-			}
-#endif
-			return false;
-		}
-
-		const FString TargetPinStr = ModuleMatcher.GetCaptureGroup(CopyElementTargetPinCaptureGroup);
-		const UPCGPin* TargetPin = GetOutputPin(FName(*TargetPinStr));
-		if (!TargetPin)
-		{
-#if PCG_KERNEL_LOGGING_ENABLED
-			if (OutErrorText)
-			{
-				*OutErrorText = FText::Format(LOCTEXT("InvalidCopyTargetPin", "Invalid output pin '{0}' used in copy element expression."), FText::FromString(TargetPinStr));
-			}
-#endif
-			return false;
-		}
-
-		if (SourcePin->Properties.AllowedTypes != TargetPin->Properties.AllowedTypes)
-		{
-#if PCG_KERNEL_LOGGING_ENABLED
-			if (OutErrorText)
-			{
-				*OutErrorText = LOCTEXT("InvalidCopyPinTypes", "CopyFromElement only allowed across pins of matching type.");
-			}
-#endif
-			return false;
-		}
-
-		PinsRequiringElementCopy.Emplace({ SourcePinStr, TargetPinStr });
-	}
-
-	if (!InOutSource)
-	{
-		return true;
-	}
-
-	// Replace function calls like Out_CopyElementFrom_In(...) with macro PCG_COPY_ALL_ATTRIBUTES_TO_OUTPUT(Out, In, ...).
-	for (const TPair<FString, FString>& SourcePinTargetPin : PinsRequiringElementCopy)
-	{
-		InOutSource->ReplaceInline(
-			*FString::Format(TEXT("{1}_CopyElementFrom_{0}("), { SourcePinTargetPin.Key, SourcePinTargetPin.Value }),
-			*FString::Format(TEXT("PCG_COPY_ALL_ATTRIBUTES_TO_OUTPUT({1}, {0}, "), { SourcePinTargetPin.Key, SourcePinTargetPin.Value })
-		);
-	}
-
-	return true;
-}
-
-FString UPCGCustomHLSLSettings::GetCookedKernelSource(const UPCGComputeGraph* InComputeGraph) const
-{
-	check(InComputeGraph);
-
-	const FIntVector GroupSize(PCGComputeConstants::THREAD_GROUP_SIZE, 1, 1);
-
-	FString Source = ShaderSource;
-	FString Functions = TEXT("\n\n// User kernel functions\n") + ShaderFunctions;
-
-#if PLATFORM_WINDOWS
-	// Remove old-school stuff.
-	Source.ReplaceInline(TEXT("\r"), TEXT(""));
-	Functions.ReplaceInline(TEXT("\r"), TEXT(""));
-#endif
-
-	ProcessShaderSource(&Source);
-	ProcessShaderSource(&Functions);
-
-	// We inject attribute IDs directly into the source. This is most efficient and saves us passing them into the kernel. However
-	// the trade off is that a shader with a shared source will generate multiple variants if used in different compute graphs with
-	// different attribute IDs.
-	for (const FPCGKernelAttributeKey& AttributeKey : KernelAttributeKeys)
-	{
-		const FString SourceDefinition = PCGHLSLElement::GetKernelAttributeKeyAsString(AttributeKey);
-		const FString AttributeIndexAsString = FString::FromInt(InComputeGraph->GetAttributeStaticTable().GetAttributeId(AttributeKey));
-
-		Source.ReplaceInline(*SourceDefinition, *AttributeIndexAsString);
-		Functions.ReplaceInline(*SourceDefinition, *AttributeIndexAsString);
-	}
-
-	// We inject data IDs directly into the source. They will get remapped to data indices using a label resolver data interface.
-	if (const FPCGPinDataLabels* PinToDataLabels = InComputeGraph->GetDataLabelsStaticTable().Find(this))
-	{
-		for (const TPair<FName, FPCGDataLabels>& Pair : PinToDataLabels->PinToDataLabels)
-		{
-			const FName PinLabel = Pair.Key;
-			const TArray<FString>& DataLabels = Pair.Value.Labels;
-
-			for (int DataId = 0; DataId < DataLabels.Num(); ++DataId)
-			{
-				const FString& DataLabel = DataLabels[DataId];
-				const FString ReplacementStr = FString::Format(TEXT("{0}_GetDataIndexFromIdInternal(/*DataId=*/{1})"),
-					{ 
-						PCGComputeHelpers::GetDataLabelResolverName(PinLabel),
-						FString::FromInt(DataId)
-					});
-
-				// Matches against {PinName}_AnyFunction('{DataLabel}'...
-				// First capture group is the data label, so that we can find & replace it by index.
-				const FString Pattern = FString::Format(TEXT("{0}_.*?[\\s]*?\\([\\s]*?('{1}')"), { *PinLabel.ToString(), DataLabel });
-
-				auto ReplaceDataLabels = [&Pattern, &ReplacementStr](FString& InOutSource)
-				{
-					FRegexMatcher ModuleMatcher(FRegexPattern(Pattern), InOutSource);
-					while (ModuleMatcher.FindNext())
-					{
-						const int32 DataLabelCharIndexStart = ModuleMatcher.GetCaptureGroupBeginning(1);
-						const int32 DataLabelCharIndexEnd = ModuleMatcher.GetCaptureGroupEnding(1);
-
-						InOutSource = InOutSource.Left(DataLabelCharIndexStart) + ReplacementStr + InOutSource.RightChop(DataLabelCharIndexEnd);
-					}
-				};
-
-				ReplaceDataLabels(Source);
-				ReplaceDataLabels(Functions);
-			}
-		}
-	}
-
-	FString ShaderPathName = GetPathName();
-	PCGHLSLElement::ConvertObjectPathToShaderFilePath(ShaderPathName);
-
-	//const bool bHasKernelKeyword = Source.Contains(TEXT("KERNEL"), ESearchCase::CaseSensitive);
-
-	FString Includes;
-	{
-		// Add with caution: Pulling in external includes has the danger that 1) we may pull in more than we expect
-		// if they include additional things in the future, and 2) if their functions change it could break user's
-		// kernel source. The latter may be mitigated by branching our own PCG version of these if needed in the future.
-
-		// TODO bring these in via additional sources so that their contents get hashed. Also these could generate declarations.
-		Includes += TEXT("#include \"/Engine/Private/ComputeShaderUtils.ush\"\n");
-		Includes += TEXT("#include \"/Engine/Private/Quaternion.ush\"\n");
-		Includes += TEXT("#include \"/Plugin/PCG/Private/PCGShaderUtils.ush\"\n");
-		Includes += TEXT("#include \"/Plugin/PCG/Private/PCGShaderUtilsInternal.ush\"\n");
-	}
-
-	const FString KernelFunc = FString::Printf(
-		TEXT("[numthreads(%d, %d, %d)]\nvoid %s(uint3 GroupId : SV_GroupID, uint GroupIndex : SV_GroupIndex)"),
-		GroupSize.X, GroupSize.Y, GroupSize.Z, TEXT("Main"));
-
-	const FString UnWrappedDispatchThreadId = FString::Printf(
-		TEXT("GetUnWrappedDispatchThreadId(GroupId, GroupIndex, %d)"),
-		GroupSize.X * GroupSize.Y * GroupSize.Z
-	);
-
-	// Used to signal that a kernel has executed. Set the most significant bit in NumData.
-	FString SetAsExecuted = TEXT("    // Signal kernel executed by setting the most significant bit of NumData.\n");
-
-	for (const FPCGPinPropertiesGPU& PinProps : OutputPins)
-	{
-		SetAsExecuted += FString::Format(TEXT("    if (all(GroupId == 0u) && GroupIndex == 0) {0}_SetAsExecutedInternal();\n"), { PinProps.Label.ToString() });
-	}
-
-	// Per-kernel-type preamble. Set up shader inputs and initialize output data.
-	FString KernelSpecificPreamble = TEXT("    // Kernel preamble\n");
-	
-	auto AddThreadInfoForPin = [&KernelSpecificPreamble](FName PinLabel)
-	{
-		KernelSpecificPreamble += FString::Format(TEXT(
-			"    uint {0}_DataIndex;\n"
-			"    if (!{0}_GetThreadData(ThreadIndex, {0}_DataIndex, ElementIndex)) return;\n"),
-			{ PinLabel.ToString() });
-	};
-
-	if (KernelType == EPCGKernelType::PointProcessor)
-	{
-		const UPCGPin* InputPin = GetPointProcessingInputPin();
-		const UPCGPin* OutputPin = GetFirstPointOutputPin();
-
-		if (InputPin && OutputPin)
-		{
-			KernelSpecificPreamble += TEXT("    uint ElementIndex; // Assumption - element index identical in input and output data.\n");
-
-			AddThreadInfoForPin(InputPin->Properties.Label);
-			AddThreadInfoForPin(OutputPin->Properties.Label);
-
-			// If input point is invalid, mark output point as invalid and abort.
-			KernelSpecificPreamble += FString::Format(TEXT(
-				"    if ({0}_IsPointRemoved({0}_DataIndex, ElementIndex))\n"
-				"    {\n"
-				"        {1}_RemovePoint({1}_DataIndex, ElementIndex);\n"
-				"        return;\n"
-				"    }\n"),
-				{ InputPin->Properties.Label.ToString(), OutputPin->Properties.Label.ToString() });
-
-			// Automatically copy value of all attributes for this element.
-			// TODO pass in IDs of attributes that are actually present.
-			KernelSpecificPreamble += FString::Format(TEXT(
-				"\n"
-				"// Point processor always initializes outputs by copying input data elements.\n"
-				"PCG_COPY_ALL_ATTRIBUTES_TO_OUTPUT({1}, {0}, {1}_DataIndex, ElementIndex, {0}_DataIndex, ElementIndex);\n"),
-				{ InputPin->Properties.Label.ToString(), OutputPin->Properties.Label.ToString() });
-		}
-	}
-	else if (KernelType == EPCGKernelType::PointGenerator)
-	{
-		KernelSpecificPreamble += FString::Format(TEXT("    const uint NumPoints = {0};\n"), { PointCount });
-
-		if (const UPCGPin* OutputPin = GetFirstPointOutputPin())
-		{
-			KernelSpecificPreamble += TEXT("uint ElementIndex; // Assumption - element index identical in input and output data.\n");
-
-			AddThreadInfoForPin(OutputPin->Properties.Label);
-
-			KernelSpecificPreamble += FString::Format(TEXT(
-				"    \n"
-				"    // Initialize all values to defaults for output pin {0}\n"
-				"    {0}_InitializePoint({0}_DataIndex, ElementIndex);\n"),
-				{ OutputPin->Properties.Label.ToString() });
-		}
-	}
-
-	FString Result;
-
-	// TODO: Support KERNEL keyword in shader source. Could be handy for external source assets and breaking kernels into sections to
-	// support pin/attribute declarations, etc.
-	/*if (bHasKernelKeyword)
-	{
-		Source.ReplaceInline(TEXT("KERNEL"), TEXT("void __kernel_func(uint ThreadIndex)"), ESearchCase::CaseSensitive);
-
-		Result = FString::Printf(TEXT(
-			"#line 0 \"%s\"\n" // ShaderPathName
-			"%s\n" // Includes
-			"%s\n" // Functions
-			"%s\n" // Source
-			"%s { __kernel_func(%s); }\n"), // KernelFunc, UnWrappedDispatchThreadId
-			*ShaderPathName, *Includes, *Functions, *Source, *KernelFunc, *UnWrappedDispatchThreadId);
-	}
-	else*/
-	{
-		Result = FString::Printf(TEXT(
-			"%s\n\n" // Includes
-			"%s\n\n" // Functions
-			"%s\n" // KernelFunc
-			"{\n"
-			"%s\n" // SetAsExecuted
-			"	const uint ThreadIndex = %s;\n" // UnWrappedDispatchThreadId
-			"	if (ThreadIndex >= GetNumThreads().x) return;\n"
-			"%s\n" // KernelSpecificPreamble
-			"#line 0 \"%s\"\n" // ShaderPathName
-			"%s\n" // Source
-			"}\n"),
-			*Includes, *Functions, *KernelFunc, *SetAsExecuted, *UnWrappedDispatchThreadId, *KernelSpecificPreamble, *ShaderPathName, *Source);
-	}
-
-	return Result;
-}
-
 bool FPCGCustomHLSLElement::ExecuteInternal(FPCGContext* Context) const
 {
 	// This element does not support CPU execution and we are never supposed to land here.
diff --git a/Engine/Plugins/PCG/Source/PCG/Private/Compute/Elements/PCGCustomHLSL.h b/Engine/Plugins/PCG/Source/PCG/Private/Compute/Elements/PCGCustomHLSL.h
--- a/Engine/Plugins/PCG/Source/PCG/Private/Compute/Elements/PCGCustomHLSL.h
+++ b/Engine/Plugins/PCG/Source/PCG/Private/Compute/Elements/PCGCustomHLSL.h
@@ -3,12 +3,13 @@
 #pragma once
 
 #include "PCGSettings.h"
-#include "Compute/PCGPinPropertiesGPU.h"
-
 #include "Compute/IPCGNodeSourceTextProvider.h"
 
+#include "Compute/PCGPinPropertiesGPU.h"
+
 #include "PCGCustomHLSL.generated.h"
 
+class UPCGComputeKernel;
 class UPCGPin;
 
 /** Type of kernel allows us to make decisions about execution automatically, streamlining authoring. */
@@ -56,6 +57,7 @@ public:
 	virtual FText GetDefaultNodeTitle() const override { return NSLOCTEXT("PCGCustomHLSLElement", "NodeTitle", "Custom HLSL"); }
 	virtual FText GetNodeTooltipText() const override { return NSLOCTEXT("PCGCustomHLSLElement", "NodeTooltip", "[EXPERIMENTAL] Produces a HLSL compute shader which will be executed on the GPU."); }
 	virtual EPCGSettingsType GetType() const override { return EPCGSettingsType::GPU; }
+	UPCGComputeKernel* CreateKernel(FPCGGPUCompilationContext& InOutContext, UObject* InObjectOuter) const override;
 
 	virtual TArray<FPCGPreConfiguredSettingsInfo> GetPreconfiguredInfo() const override;
 	virtual bool OnlyExposePreconfiguredSettings() const override { return true; }
@@ -64,16 +66,6 @@ public:
 	virtual FString GetAdditionalTitleInformation() const override;
 	virtual void ApplyPreconfiguredSettings(const FPCGPreConfiguredSettingsInfo& PreconfigureInfo) override;
 
-	virtual bool IsKernelValid(FPCGContext* InContext = nullptr, bool bQuiet = true) const override;
-	virtual FString GetCookedKernelSource(const UPCGComputeGraph* InComputeGraph) const override;
-	virtual void GetKernelAttributeKeys(TArray<FPCGKernelAttributeKey>& OutKeys) const override { OutKeys.Append(KernelAttributeKeys); }
-	virtual int ComputeKernelThreadCount(const UPCGDataBinding* Binding) const override;
-	virtual bool ComputeOutputPinDataDesc(const UPCGPin* OutputPin, const UPCGDataBinding* Binding, FPCGDataCollectionDesc& OutDesc) const override;
-
-#if WITH_EDITOR
-	virtual void GetDataLabels(FName InPinLabel, TArray<FString>& OutDataLabels) const override;
-#endif
-
 protected:
 #if WITH_EDITOR
 	virtual void PreEditChange(FProperty* PropertyAboutToChange) override;
@@ -131,16 +123,10 @@ protected:
 	const UPCGPin* GetPointProcessingInputPin() const;
 	const UPCGPin* GetFirstOutputPin() const;
 	const UPCGPin* GetFirstPointOutputPin() const;
-	int GetProcessingElemCountForInputPin(const UPCGPin* InputPin, const UPCGDataBinding* Binding) const;
-	bool AreKernelAttributesValid(FPCGContext* InContext, FText* OutErrorText) const;
-	bool AreAllOutputPinsWritten(FText* OutErrorText) const;
 
 	/** Will the ThreadCountMultiplier value be applied when calculating the dispatch thread count. */
 	bool IsThreadCountMultiplierInUse() const { return KernelType == EPCGKernelType::Custom && DispatchThreadCount != EPCGDispatchThreadCount::Fixed; }
 
-	/** Apply any code gen or transformations or expansions to the shader source. */
-	virtual bool ProcessShaderSource(FString* InOutSource, FText* OutErrorText = nullptr) const;
-
 protected:
 	UPROPERTY(EditAnywhere, Category = "Settings")
 	EPCGKernelType KernelType = EPCGKernelType::PointProcessor;
@@ -202,6 +188,8 @@ protected:
 
 	/** Maps pins to their attribute keys and whether or not they were created on the GPU. */
 	TMap<FName, TArray<TTuple<FPCGKernelAttributeKey, bool /*bCreatedOnGPU*/>>> PinToAttributeKeys;
+
+	friend class UPCGCustomHLSLKernel;
 };
 
 class FPCGCustomHLSLElement : public IPCGElement
@@ -210,3 +198,23 @@ protected:
 	// This will only be called if the custom HLSL node is not set up correctly (valid nodes are replaced with a compute graph element).
 	virtual bool ExecuteInternal(FPCGContext* Context) const override;
 };
+
+namespace PCGHLSLElement
+{
+	const FString PinDeclTemplateStr = TEXT("{pin}");
+
+	/** First capture: Pin name (supports a - z, A - Z, and 0 - 9) */
+	constexpr int AttributePinCaptureGroup = 1;
+
+	/** Second capture: Function name (Get or Set) */
+	constexpr int AttributeFunctionCaptureGroup = 2;
+
+	/** Third capture: Attribute type (e.g. Int, Float, Rotator, etc.) */
+	constexpr int AttributeTypeCaptureGroup = 3;
+
+	/** Fourth capture: Attribute name (supports a-z, A-Z, 0-9, ' ', '-', '_', and '/') */
+	constexpr int AttributeNameCaptureGroup = 4;
+
+	/** Regex pattern used to detect and parse attribute function usage in kernels. */
+	constexpr TCHAR AttributeFunctionPattern[] = { TEXT("([a-zA-Z0-9]+)_(Get|Set)(.*?)[\\s]*?\\(.*?,.*?,.*?'([a-zA-Z0-9 -_\\/]*?)'") };
+}
diff --git a/Engine/Plugins/PCG/Source/PCG/Private/Compute/Elements/PCGCustomHLSLKernel.cpp b/Engine/Plugins/PCG/Source/PCG/Private/Compute/Elements/PCGCustomHLSLKernel.cpp
new file mode 100644
--- /dev/null
+++ b/Engine/Plugins/PCG/Source/PCG/Private/Compute/Elements/PCGCustomHLSLKernel.cpp
@@ -0,0 +1,1122 @@
+// Copyright Epic Games, Inc. All Rights Reserved.
+
+#include "Compute/Elements/PCGCustomHLSLKernel.h"
+
+#include "PCGComputeGraphElement.h"
+#include "PCGContext.h"
+#include "PCGModule.h"
+#include "Compute/PCGComputeGraph.h"
+#include "Compute/PCGDataBinding.h"
+#include "Compute/Elements/PCGCustomHLSL.h"
+
+#include "Internationalization/Regex.h"
+#include "Containers/StaticArray.h"
+
+#include UE_INLINE_GENERATED_CPP_BY_NAME(PCGCustomHLSLKernel)
+
+#define LOCTEXT_NAMESPACE "PCGCustomHLSLKernel"
+
+namespace PCGCustomHLSLKernel
+{
+	constexpr TCHAR AttributeFunctionGet[] = { TEXT("Get") };
+	constexpr TCHAR AttributeFunctionSet[] = { TEXT("Set") };
+
+	void ConvertObjectPathToShaderFilePath(FString& InOutPath)
+	{
+		// Shader compiler recognizes "/Engine/Generated/..." path as special. 
+		// It doesn't validate file suffix etc.
+		InOutPath = FString::Printf(TEXT("/Engine/Generated/UObject%s.ush"), *InOutPath);
+		// Shader compilation result parsing will break if it finds ':' where it doesn't expect.
+		InOutPath.ReplaceCharInline(TEXT(':'), TEXT('@'));
+	}
+
+	FString GetKernelAttributeKeyAsString(const FPCGKernelAttributeKey& Key)
+	{
+		return FString::Format(TEXT("'{0}'"), { Key.Name.ToString() });
+	}
+
+	FString GetDataTypeString(EPCGDataType Type)
+	{
+		const UEnum* DataTypeEnum = StaticEnum<EPCGDataType>();
+		check(DataTypeEnum);
+
+		return DataTypeEnum->GetValueOrBitfieldAsString(static_cast<int64>(Type));
+	}
+}
+
+bool UPCGCustomHLSLKernel::IsKernelValid(FPCGContext* InContext, bool bQuiet) const
+{
+	TRACE_CPUPROFILER_EVENT_SCOPE(UPCGCustomHLSLKernel::IsKernelValid);
+
+	if (!Super::IsKernelValid(InContext, bQuiet) || !ensure(Settings))
+	{
+		return false;
+	}
+
+	if (!ValidatePCGNode(InContext, bQuiet, Cast<UPCGSettings>(Settings)))
+	{
+		return false;
+	}
+
+	if (Settings->OutputPins.IsEmpty())
+	{
+		PCG_KERNEL_VALIDATION_ERR(InContext, Settings, bQuiet, LOCTEXT("NoOutputs", "Custom HLSL nodes must have at least one output."));
+		return false;
+	}
+
+	auto CheckPinLabel = [Settings=Settings, InContext, bQuiet](FName PinLabel)
+	{
+		if (PinLabel == NAME_None)
+		{
+			PCG_KERNEL_VALIDATION_ERR(InContext, Settings, bQuiet, LOCTEXT("InvalidPinLabelNone", "Pin label 'None' is not a valid pin label."));
+			return false;
+		}
+
+		bool bFoundPinLabel = false;
+
+		auto IsAlreadyFound = [Settings, InContext, bQuiet, PinLabel, &bFoundPinLabel](const FPCGPinProperties PinProps)
+		{
+			if (PinProps.Label == PinLabel)
+			{
+				if (bFoundPinLabel)
+				{
+					PCG_KERNEL_VALIDATION_ERR(InContext, Settings, bQuiet, FText::Format(LOCTEXT("DuplicatedPinLabels", "Duplicate pin label '{0}', all labels must be unique."), FText::FromName(PinLabel)));
+					return true;
+				}
+
+				bFoundPinLabel = true;
+			}
+
+			return false;
+		};
+
+		for (const FPCGPinProperties& PinProps : Settings->InputPins)
+		{
+			if (IsAlreadyFound(PinProps))
+			{
+				return false;
+			}
+		}
+
+		for (const FPCGPinProperties& PinProps : Settings->OutputPins)
+		{
+			if (IsAlreadyFound(PinProps))
+			{
+				return false;
+			}
+		}
+
+		return true;
+	};
+
+	// Validate input pins
+	bool bIsFirstInputPin = true;
+	for (const FPCGPinProperties& Properties : Settings->InputPins)
+	{
+		if (!CheckPinLabel(Properties.Label))
+		{
+			return false;
+		}
+
+		if (bIsFirstInputPin && Settings->KernelType == EPCGKernelType::PointProcessor)
+		{
+			if (Properties.AllowedTypes != EPCGDataType::Point)
+			{
+				PCG_KERNEL_VALIDATION_ERR(InContext, Settings, bQuiet, FText::Format(
+					LOCTEXT("InvalidNonPointPrimaryInput", "'Point Processor' nodes require primary input pin to be of type 'Point', but found '{0}'."),
+					FText::FromString(PCGCustomHLSLKernel::GetDataTypeString(Properties.AllowedTypes))));
+
+				return false;
+			}
+		}
+
+		if (!PCGComputeHelpers::IsTypeAllowedAsInput(Properties.AllowedTypes))
+		{
+			PCG_KERNEL_VALIDATION_ERR(InContext, Settings, bQuiet, FText::Format(
+				LOCTEXT("InvalidInputType", "Unsupported input type '{0}', found on pin '{1}'."),
+				FText::FromString(PCGCustomHLSLKernel::GetDataTypeString(Properties.AllowedTypes)),
+				FText::FromName(Properties.Label)));
+
+			return false;
+		}
+
+		bIsFirstInputPin = false;
+	}
+
+	// Validate output pins
+	bool bIsFirstOutputPin = true;
+	for (const FPCGPinPropertiesGPU& Properties : Settings->OutputPins)
+	{
+		if (!CheckPinLabel(Properties.Label))
+		{
+			return false;
+		}
+
+		const bool bPinIsDefinedByKernel = bIsFirstOutputPin && (Settings->KernelType == EPCGKernelType::PointGenerator || Settings->KernelType == EPCGKernelType::PointProcessor);
+
+		if (bPinIsDefinedByKernel)
+		{
+			if (Properties.AllowedTypes != EPCGDataType::Point)
+			{
+				PCG_KERNEL_VALIDATION_ERR(InContext, Settings, bQuiet, FText::Format(
+					LOCTEXT("InvalidNonPointPrimaryOutput", "'Point Processor' and 'Point Generator' nodes require primary output pin to be of type 'Point', but found '{0}'."),
+					FText::FromString(PCGCustomHLSLKernel::GetDataTypeString(Properties.AllowedTypes))));
+
+				return false;
+			}
+		}
+
+		if (!PCGComputeHelpers::IsTypeAllowedAsOutput(Properties.AllowedTypes))
+		{
+			PCG_KERNEL_VALIDATION_ERR(InContext, Settings, bQuiet, FText::Format(
+				LOCTEXT("InvalidOutputType", "Unsupported output type '{0}', found on pin '{1}'."),
+				FText::FromString(PCGCustomHLSLKernel::GetDataTypeString(Properties.AllowedTypes)),
+				FText::FromName(Properties.Label)));
+
+			return false;
+		}
+
+		if (!bPinIsDefinedByKernel)
+		{
+			const FPCGPinPropertiesGPUStruct& Props = Properties.PropertiesGPU;
+
+			if (Props.InitializationMode == EPCGPinInitMode::FromInputPins)
+			{
+				if (Props.PinsToInititalizeFrom.IsEmpty())
+				{
+					PCG_KERNEL_VALIDATION_ERR(InContext, Settings, bQuiet, FText::Format(
+						LOCTEXT("InitFromEmptyPins", "Output pin '{0}' tried to initialize from input pins, but no pins were specified."),
+						FText::FromName(Properties.Label)));
+
+					return false;
+				}
+
+				for (const FName InitPinName : Props.PinsToInititalizeFrom)
+				{
+					const FPCGPinProperties* InitPinProps = Settings->InputPins.FindByPredicate([InitPinName](const FPCGPinProperties& InPinProps)
+					{
+						return InPinProps.Label == InitPinName;
+					});
+
+					if (!InitPinProps)
+					{
+						PCG_KERNEL_VALIDATION_ERR(InContext, Settings, bQuiet, FText::Format(
+							LOCTEXT("InitFromNonExistentPin", "Output pin '{0}' tried to initialize from non-existent input pin '{1}'."),
+							FText::FromName(Properties.Label),
+							FText::FromName(InitPinName)));
+
+						return false;
+					}
+
+					if (!PCGComputeHelpers::IsTypeAllowedInDataCollection(InitPinProps->AllowedTypes))
+					{
+						PCG_KERNEL_VALIDATION_ERR(InContext, Settings, bQuiet, FText::Format(
+							LOCTEXT("InitFromInvalidPinType", "Output pin '{0}' tried to initialize from input pin '{1}', but pin '{1}' has an invalid type."),
+							FText::FromName(Properties.Label),
+							FText::FromName(InitPinName)));
+
+						return false;
+					}
+				}
+
+				// TODO: Could do validation on data multiplicity for Pairwise, checking that data counts are 1 or N, but maybe that should be a runtime error instead.
+			}
+
+			const bool bUsingFixedDataCount = Props.InitializationMode == EPCGPinInitMode::Custom || Props.DataCountMode == EPCGDataCountMode::Fixed;
+
+			if (bUsingFixedDataCount)
+			{
+				if (Props.DataCount < 1)
+				{
+					PCG_KERNEL_VALIDATION_ERR(InContext, Settings, bQuiet, FText::Format(
+						LOCTEXT("InvalidDataCount", "Invalid fixed data count {0} on output pin '{1}'. Must be greater than 0."),
+						FText::AsNumber(Props.DataCount),
+						FText::FromName(Properties.Label)));
+
+					return false;
+				}
+
+			}
+
+			const bool bUsingFixedElemCount = Props.InitializationMode == EPCGPinInitMode::Custom || Props.ElementCountMode == EPCGElementCountMode::Fixed;
+
+			if (bUsingFixedElemCount)
+			{
+				if (Props.ElementCount < 1)
+				{
+					PCG_KERNEL_VALIDATION_ERR(InContext, Settings, bQuiet, FText::Format(
+						LOCTEXT("InvalidElementCount", "Invalid fixed element count {0} on output pin '{1}'. Must be greater than 0."),
+						FText::AsNumber(Props.ElementCount),
+						FText::FromName(Properties.Label)));
+
+					return false;
+				}
+			}
+		}
+
+		bIsFirstOutputPin = false;
+	}
+
+	if (Settings->KernelType == EPCGKernelType::Custom && Settings->DispatchThreadCount == EPCGDispatchThreadCount::FromProductOfInputPins)
+	{
+		if (Settings->ThreadCountInputPinLabels.IsEmpty())
+		{
+			PCG_KERNEL_VALIDATION_ERR(InContext, Settings, bQuiet, LOCTEXT("MissingThreadCountPins", "Dispatch thread count is based on input pins but no labels have been set in Input Pins array."));
+			return false;
+		}
+
+		for (const FName& Label : Settings->ThreadCountInputPinLabels)
+		{
+			if (!Settings->InputPins.FindByPredicate([Label](const FPCGPinProperties& InProps) { return InProps.Label == Label; }))
+			{
+				PCG_KERNEL_VALIDATION_ERR(InContext, Settings, bQuiet, FText::Format(LOCTEXT("MissingThreadCountPin", "Invalid pin specified in Input Pins array: '{0}'."), FText::FromName(Label)));
+				return false;
+			}
+		}
+	}
+
+	if (IsThreadCountMultiplierInUse())
+	{
+		if (Settings->ThreadCountMultiplier < 1)
+		{
+			PCG_KERNEL_VALIDATION_ERR(InContext, Settings, bQuiet, FText::Format(LOCTEXT("InvalidThreadCountMultiplier", "Thread Count Multiplier has invalid value ({0}). Must be greater than 0."), Settings->ThreadCountMultiplier));
+			return false;
+		}
+	}
+
+	// Validate attributes
+	for (const FPCGKernelAttributeKey& AttributeKey : Settings->KernelAttributeKeys)
+	{
+		if (AttributeKey.Type == EPCGKernelAttributeType::Invalid)
+		{
+			const UEnum* AttributeTypeEnum = StaticEnum<EPCGKernelAttributeType>();
+			check(AttributeTypeEnum);
+
+			PCG_KERNEL_VALIDATION_ERR(InContext, Settings, bQuiet, FText::Format(
+				LOCTEXT("InvalidAttributeTypeNone", "Attribute '{0}' has invalid GPU attribute type '{1}', check the 'Attributes to Create' array on your pins."),
+				FText::FromName(AttributeKey.Name),
+				FText::FromString(AttributeTypeEnum->GetNameStringByValue(static_cast<int64>(AttributeKey.Type)))));
+			return false;
+		}
+	}
+
+	if (InContext)
+	{
+		FText* ErrorTextPtr = nullptr;
+#if PCG_KERNEL_LOGGING_ENABLED
+		FText ErrorText;
+		ErrorTextPtr = &ErrorText;
+#endif
+
+		if (!AreAttributesValid(InContext, ErrorTextPtr))
+		{
+			if (ErrorTextPtr)
+			{
+				PCG_KERNEL_VALIDATION_ERR(InContext, Settings, bQuiet, *ErrorTextPtr);
+			}
+			return false;
+		}
+
+		if (!ProcessShaderSource(/*InOutSource=*/nullptr, ErrorTextPtr))
+		{
+			if (ErrorTextPtr)
+			{
+				PCG_KERNEL_VALIDATION_ERR(InContext, Settings, bQuiet, *ErrorTextPtr);
+			}
+			return false;
+		}
+
+		if (!Settings->bMuteUnwrittenPinDataErrors && !AreAllOutputPinsWritten(ErrorTextPtr))
+		{
+			if (ErrorTextPtr)
+			{
+				PCG_KERNEL_VALIDATION_ERR(InContext, Settings, bQuiet, *ErrorTextPtr);
+			}
+			return false;
+		}
+	}
+
+	return true;
+}
+
+bool UPCGCustomHLSLKernel::ComputeOutputBindingDataDesc(const UPCGComputeGraph* InGraph, FName InOutputPinLabel, const UPCGDataBinding* InBinding, FPCGDataCollectionDesc& OutDataDesc) const
+{
+	check(InGraph);
+	check(InBinding);
+
+	FPCGDataCollectionDesc OutputDataDesc;
+
+	const FPCGPinPropertiesGPU* OutputPinProperties = Settings->OutputPins.FindByPredicate([InOutputPinLabel](const FPCGPinPropertiesGPU& InProps) {return InProps.Label == InOutputPinLabel; });
+	if (!OutputPinProperties)
+	{
+		OutDataDesc = {};
+		return true;
+	}
+
+	const FName FirstOutputPinLabel = Settings->OutputPins[0].Label;
+
+	// The primary output pin follows any rules prescribed by kernel type.
+	if (InOutputPinLabel == FirstOutputPinLabel && Settings->KernelType == EPCGKernelType::PointProcessor)
+	{
+		FPCGDataCollectionDesc InputDesc;
+
+		if (const FPCGPinProperties* FirstInputPinProps = GetFirstPointInputPin())
+		{
+			if (!InBinding->GetKernelPinDataDesc_GameThread(this, FirstInputPinProps->Label, /*bIsInputPin=*/true, InputDesc))
+			{
+				ensure(false);
+				return false;
+			}
+		}
+
+		OutputDataDesc = MoveTemp(InputDesc);
+	}
+	else if (InOutputPinLabel == FirstOutputPinLabel && Settings->KernelType == EPCGKernelType::PointGenerator)
+	{
+		// Generators always produce a single point data with known point count.
+		OutputDataDesc.DataDescs.Emplace(EPCGDataType::Point, Settings->PointCount);
+	}
+	else if (ensure(OutputPinProperties))
+	{
+		ComputeDataDescFromPinProperties(*OutputPinProperties, MakeArrayView(Settings->InputPins), InBinding, OutputDataDesc);
+	}
+
+	for (const FPCGKernelAttributeKey& AttributeKey : Settings->KernelAttributeKeys)
+	{
+		// Add attributes that will be created for this pin on the GPU.
+		if (const TArray<TTuple<FPCGKernelAttributeKey, bool>>* Keys = Settings->PinToAttributeKeys.Find(InOutputPinLabel))
+		{
+			const TTuple<FPCGKernelAttributeKey, bool>* Pair = Keys->FindByPredicate([AttributeKey](const TTuple<FPCGKernelAttributeKey, bool>& Pair) { return Pair.Key == AttributeKey; });
+			if (!Pair || !Pair->Value)
+			{
+				// Not created on GPU.
+				continue;
+			}
+
+			for (FPCGDataDesc& DataDesc : OutputDataDesc.DataDescs)
+			{
+				const int32 AttributeId = InBinding->GetAttributeId(AttributeKey);
+
+				if (AttributeId != INDEX_NONE)
+				{
+					const FPCGKernelAttributeDesc AttributeDesc(AttributeId, AttributeKey.Type, AttributeKey.Name);
+					DataDesc.AttributeDescs.AddUnique(AttributeDesc);
+				}
+			}
+		}
+	}
+
+	// Try to propagate string keys across node. Not trivial because there could be one or more string key attributes on input pins and on output pins,
+	// and it is in general hard to determine from source which string keys from input are being written to outputs. Try first collecting all string keys
+	// from matching attribute names (across all input pins), and then fall back to collecting keys from all string key attributes across all inputs.
+	bool bOutputHasStringKeys = false;
+
+	for (const FPCGDataDesc& DataDesc : OutputDataDesc.DataDescs)
+	{
+		auto HasStringKey = [](const FPCGKernelAttributeDesc& InAttributeDesc)
+			{
+				return InAttributeDesc.Type == EPCGKernelAttributeType::StringKey || InAttributeDesc.Type == EPCGKernelAttributeType::Name;
+			};
+
+		if (DataDesc.AttributeDescs.FindByPredicate(HasStringKey))
+		{
+			bOutputHasStringKeys = true;
+			break;
+		}
+	}
+
+	if (bOutputHasStringKeys)
+	{
+		TArray<FPCGDataCollectionDesc> RelevantInputDataDescs;
+
+		// Collect descriptions of input data items that have string key attributes.
+		for(const FPCGPinProperties& PinProps : Settings->InputPins)
+		{
+			FPCGDataCollectionDesc InputPinDesc;
+			ensure(InBinding->GetKernelPinDataDesc_GameThread(this, PinProps.Label, /*bIsInputPin=*/true, InputPinDesc));
+
+			bool bFoundStringKeyAttribute = false;
+
+			for (FPCGDataDesc& DataDesc : OutputDataDesc.DataDescs)
+			{
+				for (FPCGKernelAttributeDesc& AttributeDesc : DataDesc.AttributeDescs)
+				{
+					bFoundStringKeyAttribute |= (AttributeDesc.Type == EPCGKernelAttributeType::StringKey || AttributeDesc.Type == EPCGKernelAttributeType::Name);
+				}
+			}
+
+			if (bFoundStringKeyAttribute)
+			{
+				RelevantInputDataDescs.Add(InputPinDesc);
+			}
+		}
+
+		if (!RelevantInputDataDescs.IsEmpty())
+		{
+			for (FPCGDataDesc& DataDesc : OutputDataDesc.DataDescs)
+			{
+				for (FPCGKernelAttributeDesc& AttributeDesc : DataDesc.AttributeDescs)
+				{
+					if (AttributeDesc.Type != EPCGKernelAttributeType::StringKey && AttributeDesc.Type != EPCGKernelAttributeType::Name)
+					{
+						continue;
+					}
+
+					bool bFoundMatchingAttribute = false;
+
+					for (const FPCGDataCollectionDesc& InputPinDataDesc : RelevantInputDataDescs)
+					{
+						// Try to find string keys for matching attributes on inputs. E.g. if we are processing an output attribute named 'MeshPath',
+						// look at data on all input pins for an attribute named MeshPath and assume we could use any of its values - copy the string keys.
+						for (const FPCGDataDesc& InputDataDesc : InputPinDataDesc.DataDescs)
+						{
+							for (const FPCGKernelAttributeDesc& InputAttributeDesc : InputDataDesc.AttributeDescs)
+							{
+								const bool bIsString = InputAttributeDesc.Type == EPCGKernelAttributeType::StringKey || InputAttributeDesc.Type == EPCGKernelAttributeType::Name;
+								if (bIsString && InputAttributeDesc.Name == AttributeDesc.Name)
+								{
+									AttributeDesc.AddUniqueStringKeys(InputAttributeDesc.GetUniqueStringKeys());
+									bFoundMatchingAttribute = true;
+									break;
+								}
+							}
+						}
+					}
+
+					if (!bFoundMatchingAttribute)
+					{
+						// We didn't find an exact attribute. Fall back to finding any and all string keys. This is concerning and perhaps we can
+						// have additional hinting mechanisms in the kernel source or in the node UI.
+						for (const FPCGDataCollectionDesc& InputPinDataDesc : RelevantInputDataDescs)
+						{
+							// Try to find string keys for matching attributes on inputs. E.g. if we are processing an output attribute named 'MeshPath',
+							// look at data on all input pins for an attribute named MeshPath and assume we could use any of its values - copy the string keys.
+							for (const FPCGDataDesc& InputDataDesc : InputPinDataDesc.DataDescs)
+							{
+								for (const FPCGKernelAttributeDesc& InputAttributeDesc : InputDataDesc.AttributeDescs)
+								{
+									if (InputAttributeDesc.Type == EPCGKernelAttributeType::StringKey || InputAttributeDesc.Type == EPCGKernelAttributeType::Name)
+									{
+										AttributeDesc.AddUniqueStringKeys(InputAttributeDesc.GetUniqueStringKeys());
+									}
+								}
+							}
+						}
+					}
+				}
+			}
+		}
+		else
+		{
+			// If there were no string keys found on any input pin then we are in a bad place. String values cannot be built on the GPU, they must
+			// come in through an input.
+			UE_LOG(LogPCG, Warning, TEXT("No incoming attributes to obtain string keys from."));
+		}
+	}
+
+	OutDataDesc = MoveTemp(OutputDataDesc);
+	return true;
+}
+
+#if WITH_EDITOR
+FString UPCGCustomHLSLKernel::GetCookedSource(const UPCGComputeGraph* InComputeGraph) const
+{
+	check(InComputeGraph);
+
+	const FIntVector GroupSize(PCGComputeConstants::THREAD_GROUP_SIZE, 1, 1);
+
+	FString Source = Settings->ShaderSource;
+	FString Functions = TEXT("\n\n// User kernel functions\n") + Settings->ShaderFunctions;
+
+#if PLATFORM_WINDOWS
+	// Remove old-school stuff.
+	Source.ReplaceInline(TEXT("\r"), TEXT(""));
+	Functions.ReplaceInline(TEXT("\r"), TEXT(""));
+#endif
+
+	ProcessShaderSource(&Source);
+	ProcessShaderSource(&Functions);
+
+	// We inject attribute IDs directly into the source. This is most efficient and saves us passing them into the kernel. However
+	// the trade off is that a shader with a shared source will generate multiple variants if used in different compute graphs with
+	// different attribute IDs.
+	for (const FPCGKernelAttributeKey& AttributeKey : Settings->KernelAttributeKeys)
+	{
+		const FString SourceDefinition = PCGCustomHLSLKernel::GetKernelAttributeKeyAsString(AttributeKey);
+		const FString AttributeIndexAsString = FString::FromInt(InComputeGraph->GetAttributeStaticTable().GetAttributeId(AttributeKey));
+
+		Source.ReplaceInline(*SourceDefinition, *AttributeIndexAsString);
+		Functions.ReplaceInline(*SourceDefinition, *AttributeIndexAsString);
+	}
+
+	// We inject data IDs directly into the source. They will get remapped to data indices using a label resolver data interface.
+	if (const FPCGPinDataLabels* PinToDataLabels = InComputeGraph->GetDataLabelsStaticTable().Find(Settings))
+	{
+		for (const TPair<FName, FPCGDataLabels>& Pair : PinToDataLabels->PinToDataLabels)
+		{
+			const FName PinLabel = Pair.Key;
+			const TArray<FString>& DataLabels = Pair.Value.Labels;
+
+			for (int DataId = 0; DataId < DataLabels.Num(); ++DataId)
+			{
+				const FString& DataLabel = DataLabels[DataId];
+				const FString ReplacementStr = FString::Format(TEXT("{0}_GetDataIndexFromIdInternal(/*DataId=*/{1})"),
+					{
+						PCGComputeHelpers::GetDataLabelResolverName(PinLabel),
+						FString::FromInt(DataId)
+					});
+
+				// Matches against {PinName}_AnyFunction('{DataLabel}'...
+				// First capture group is the data label, so that we can find & replace it by index.
+				const FString Pattern = FString::Format(TEXT("{0}_.*?[\\s]*?\\([\\s]*?('{1}')"), { *PinLabel.ToString(), DataLabel });
+
+				auto ReplaceDataLabels = [&Pattern, &ReplacementStr](FString& InOutSource)
+				{
+					FRegexMatcher ModuleMatcher(FRegexPattern(Pattern), InOutSource);
+					while (ModuleMatcher.FindNext())
+					{
+						const int32 DataLabelCharIndexStart = ModuleMatcher.GetCaptureGroupBeginning(1);
+						const int32 DataLabelCharIndexEnd = ModuleMatcher.GetCaptureGroupEnding(1);
+
+						InOutSource = InOutSource.Left(DataLabelCharIndexStart) + ReplacementStr + InOutSource.RightChop(DataLabelCharIndexEnd);
+					}
+				};
+
+				ReplaceDataLabels(Source);
+				ReplaceDataLabels(Functions);
+			}
+		}
+	}
+
+	FString ShaderPathName = GetPathName();
+	PCGCustomHLSLKernel::ConvertObjectPathToShaderFilePath(ShaderPathName);
+
+	//const bool bHasKernelKeyword = Source.Contains(TEXT("KERNEL"), ESearchCase::CaseSensitive);
+
+	FString Includes;
+	{
+		// Add with caution: Pulling in external includes has the danger that 1) we may pull in more than we expect
+		// if they include additional things in the future, and 2) if their functions change it could break user's
+		// kernel source. The latter may be mitigated by branching our own PCG version of these if needed in the future.
+
+		// TODO bring these in via additional sources so that their contents get hashed. Also these could generate declarations.
+		Includes += TEXT("#include \"/Engine/Private/ComputeShaderUtils.ush\"\n");
+		Includes += TEXT("#include \"/Engine/Private/Quaternion.ush\"\n");
+		Includes += TEXT("#include \"/Plugin/PCG/Private/PCGShaderUtils.ush\"\n");
+		Includes += TEXT("#include \"/Plugin/PCG/Private/PCGShaderUtilsInternal.ush\"\n");
+	}
+
+	const FString KernelFunc = FString::Printf(
+		TEXT("[numthreads(%d, %d, %d)]\nvoid %s(uint3 GroupId : SV_GroupID, uint GroupIndex : SV_GroupIndex)"),
+		GroupSize.X, GroupSize.Y, GroupSize.Z, TEXT("Main"));
+
+	const FString UnWrappedDispatchThreadId = FString::Printf(
+		TEXT("GetUnWrappedDispatchThreadId(GroupId, GroupIndex, %d)"),
+		GroupSize.X * GroupSize.Y * GroupSize.Z
+	);
+
+	// Used to signal that a kernel has executed. Set the most significant bit in NumData.
+	FString SetAsExecuted = TEXT("    // Signal kernel executed by setting the most significant bit of NumData.\n");
+
+	for (const FPCGPinPropertiesGPU& PinProps : Settings->OutputPins)
+	{
+		SetAsExecuted += FString::Format(TEXT("    if (all(GroupId == 0u) && GroupIndex == 0) {0}_SetAsExecutedInternal();\n"), { PinProps.Label.ToString() });
+	}
+
+	// Per-kernel-type preamble. Set up shader inputs and initialize output data.
+	FString KernelSpecificPreamble = TEXT("    // Kernel preamble\n");
+
+	auto AddThreadInfoForPin = [&KernelSpecificPreamble](FName PinLabel)
+		{
+			KernelSpecificPreamble += FString::Format(TEXT(
+				"    uint {0}_DataIndex;\n"
+				"    if (!{0}_GetThreadData(ThreadIndex, {0}_DataIndex, ElementIndex)) return;\n"),
+				{ PinLabel.ToString() });
+		};
+
+	if (Settings->KernelType == EPCGKernelType::PointProcessor)
+	{
+		const FPCGPinProperties* InputPin = GetFirstPointInputPin();
+		const FPCGPinPropertiesGPU* OutputPin = GetFirstPointOutputPin();
+
+		if (InputPin && OutputPin)
+		{
+			KernelSpecificPreamble += TEXT("    uint ElementIndex; // Assumption - element index identical in input and output data.\n");
+
+			AddThreadInfoForPin(InputPin->Label);
+			AddThreadInfoForPin(OutputPin->Label);
+
+			// If input point is invalid, mark output point as invalid and abort.
+			KernelSpecificPreamble += FString::Format(TEXT(
+				"    if ({0}_IsPointRemoved({0}_DataIndex, ElementIndex))\n"
+				"    {\n"
+				"        {1}_RemovePoint({1}_DataIndex, ElementIndex);\n"
+				"        return;\n"
+				"    }\n"),
+				{ InputPin->Label.ToString(), OutputPin->Label.ToString() });
+
+			// Automatically copy value of all attributes for this element.
+			// TODO pass in IDs of attributes that are actually present.
+			KernelSpecificPreamble += FString::Format(TEXT(
+				"\n"
+				"// Point processor always initializes outputs by copying input data elements.\n"
+				"PCG_COPY_ALL_ATTRIBUTES_TO_OUTPUT({1}, {0}, {1}_DataIndex, ElementIndex, {0}_DataIndex, ElementIndex);\n"),
+				{ InputPin->Label.ToString(), OutputPin->Label.ToString() });
+		}
+	}
+	else if (Settings->KernelType == EPCGKernelType::PointGenerator)
+	{
+		KernelSpecificPreamble += FString::Format(TEXT("    const uint NumPoints = {0};\n"), { Settings->PointCount });
+
+		if (const FPCGPinPropertiesGPU* OutputPin = GetFirstPointOutputPin())
+		{
+			KernelSpecificPreamble += TEXT("uint ElementIndex; // Assumption - element index identical in input and output data.\n");
+
+			AddThreadInfoForPin(OutputPin->Label);
+
+			KernelSpecificPreamble += FString::Format(TEXT(
+				"    \n"
+				"    // Initialize all values to defaults for output pin {0}\n"
+				"    {0}_InitializePoint({0}_DataIndex, ElementIndex);\n"),
+				{ OutputPin->Label.ToString() });
+		}
+	}
+
+	FString Result;
+
+	// TODO: Support KERNEL keyword in shader source. Could be handy for external source assets and breaking kernels into sections to
+	// support pin/attribute declarations, etc.
+	/*if (bHasKernelKeyword)
+	{
+		Source.ReplaceInline(TEXT("KERNEL"), TEXT("void __kernel_func(uint ThreadIndex)"), ESearchCase::CaseSensitive);
+
+		Result = FString::Printf(TEXT(
+			"#line 0 \"%s\"\n" // ShaderPathName
+			"%s\n" // Includes
+			"%s\n" // Functions
+			"%s\n" // Source
+			"%s { __kernel_func(%s); }\n"), // KernelFunc, UnWrappedDispatchThreadId
+			*ShaderPathName, *Includes, *Functions, *Source, *KernelFunc, *UnWrappedDispatchThreadId);
+	}
+	else*/
+	{
+		Result = FString::Printf(TEXT(
+			"%s\n\n" // Includes
+			"%s\n\n" // Functions
+			"%s\n" // KernelFunc
+			"{\n"
+			"%s\n" // SetAsExecuted
+			"	const uint ThreadIndex = %s;\n" // UnWrappedDispatchThreadId
+			"	if (ThreadIndex >= GetNumThreads().x) return;\n"
+			"%s\n" // KernelSpecificPreamble
+			"#line 0 \"%s\"\n" // ShaderPathName
+			"%s\n" // Source
+			"}\n"),
+			*Includes, *Functions, *KernelFunc, *SetAsExecuted, *UnWrappedDispatchThreadId, *KernelSpecificPreamble, *ShaderPathName, *Source);
+	}
+
+	return Result;
+}
+#endif // WITH_EDITOR
+
+int UPCGCustomHLSLKernel::ComputeThreadCount(const UPCGDataBinding* InBinding) const
+{
+	int ThreadCount = 0;
+
+	if (Settings->KernelType == EPCGKernelType::PointGenerator)
+	{
+		// Point generator has fixed thread count.
+		ThreadCount = Settings->PointCount;
+	}
+	else if (Settings->KernelType == EPCGKernelType::PointProcessor)
+	{
+		// Processing volume depends on data arriving on primary pin.
+		if (const FPCGPinProperties* InputPin = GetFirstPointInputPin())
+		{
+			ThreadCount = GetElementCountForInputPin(*InputPin, InBinding);
+		}
+	}
+	else if (Settings->KernelType == EPCGKernelType::Custom)
+	{
+		if (Settings->DispatchThreadCount == EPCGDispatchThreadCount::FromFirstOutputPin)
+		{
+			if (const FPCGPinPropertiesGPU* OutputPin = Settings->OutputPins.IsEmpty() ? nullptr : &Settings->OutputPins[0])
+			{
+				FPCGDataCollectionDesc Desc;
+				if (ensure(ComputeOutputBindingDataDesc(InBinding->Graph, OutputPin->Label, InBinding, Desc)))
+				{
+					ThreadCount = Desc.ComputeDataElementCount(EPCGDataType::Any);
+				}
+			}
+		}
+		else if (Settings->DispatchThreadCount == EPCGDispatchThreadCount::FromProductOfInputPins)
+		{
+			for (const FName& PinLabel : Settings->ThreadCountInputPinLabels)
+			{
+				if (const FPCGPinProperties* InputPin = Settings->InputPins.FindByPredicate([PinLabel](const FPCGPinProperties& InProps) { return InProps.Label == PinLabel; }))
+				{
+					ThreadCount = GetElementCountForInputPin(*InputPin, InBinding);
+				}
+			}
+		}
+		else if (Settings->DispatchThreadCount == EPCGDispatchThreadCount::Fixed)
+		{
+			ThreadCount = Settings->FixedThreadCount;
+		}
+	}
+	else
+	{
+		checkNoEntry();
+	}
+
+	if (IsThreadCountMultiplierInUse())
+	{
+		ThreadCount *= Settings->ThreadCountMultiplier;
+	}
+
+	return ThreadCount;
+}
+
+void UPCGCustomHLSLKernel::GetDataLabels(FName InPinLabel, TArray<FString>& OutDataLabels) const
+{
+	// Matches against {PinName}_AnyFunction('{DataLabel}'...
+	const FString Pattern = FString::Format(TEXT("{0}_.*?[\\s]*?\\([\\s]*?'([a-zA-Z0-9_].*?)'"), { *InPinLabel.ToString() });
+	const FString Source = Settings->ShaderFunctions + TEXT("\n") + Settings->ShaderSource;
+
+	// First capture: Data label (supports a - z, A - Z, 0 - 9, and underscores).
+	constexpr int LabelCaptureGroup = 1;
+
+	FRegexMatcher ModuleMatcher(FRegexPattern(Pattern), Source);
+	while (ModuleMatcher.FindNext())
+	{
+		OutDataLabels.AddUnique(ModuleMatcher.GetCaptureGroup(LabelCaptureGroup));
+	}
+}
+
+void UPCGCustomHLSLKernel::GetKernelAttributeKeys(TArray<FPCGKernelAttributeKey>& OutKeys) const
+{
+	OutKeys.Append(Settings->KernelAttributeKeys);
+}
+
+uint32 UPCGCustomHLSLKernel::GetThreadCountMultiplier() const
+{
+	return Settings->ThreadCountMultiplier;
+}
+
+bool UPCGCustomHLSLKernel::IsThreadCountMultiplierInUse() const
+{
+	return Settings->KernelType == EPCGKernelType::Custom && Settings->DispatchThreadCount != EPCGDispatchThreadCount::Fixed;
+}
+
+bool UPCGCustomHLSLKernel::AreAttributesValid(FPCGContext* InContext, FText* OutErrorText) const
+{
+	TRACE_CPUPROFILER_EVENT_SCOPE(UPCGCustomHLSLKernel::AreAttributesValid);
+
+	// The context can either be a compute graph element context (if the compute graph was successfully created), otherwise
+	// it will be the original CPU node context. We need the former to run the following validation.
+	if (!InContext || !InContext->IsComputeContext())
+	{
+		return true;
+	}
+
+	const FPCGComputeGraphContext* Context = static_cast<FPCGComputeGraphContext*>(InContext);
+	const UPCGDataBinding* DataBinding = Context->DataBinding.Get();
+
+	if (DataBinding)
+	{
+		TMap<FName, FPCGDataCollectionDesc> InputPinDescs;
+
+		for (const FPCGPinProperties& Pin : Settings->InputPins)
+		{
+			FPCGDataCollectionDesc DataDesc;
+			ensure(DataBinding->GetKernelPinDataDesc_GameThread(this, Pin.Label, /*bIsInputPin=*/true, DataDesc));
+			InputPinDescs.Add(Pin.Label, MoveTemp(DataDesc));
+		}
+
+		TMap<FName, FPCGDataCollectionDesc> OutputPinDescs;
+
+		for (const FPCGPinProperties& Pin : Settings->OutputPins)
+		{
+			FPCGDataCollectionDesc DataDesc;
+			ensure(DataBinding->GetKernelPinDataDesc_GameThread(this, Pin.Label, /*bIsInputPin=*/false, DataDesc));
+			OutputPinDescs.Add(Pin.Label, MoveTemp(DataDesc));
+		}
+
+		const FString Source = Settings->ShaderFunctions + Settings->ShaderSource;
+		FRegexMatcher ModuleMatcher(FRegexPattern(PCGHLSLElement::AttributeFunctionPattern), Source);
+
+		while (ModuleMatcher.FindNext())
+		{
+			const FString PinStr = ModuleMatcher.GetCaptureGroup(PCGHLSLElement::AttributePinCaptureGroup);
+			const FString FuncStr = ModuleMatcher.GetCaptureGroup(PCGHLSLElement::AttributeFunctionCaptureGroup);
+			const FString TypeStr = ModuleMatcher.GetCaptureGroup(PCGHLSLElement::AttributeTypeCaptureGroup);
+			const FString NameStr = ModuleMatcher.GetCaptureGroup(PCGHLSLElement::AttributeNameCaptureGroup);
+
+			const FName PinName = FName(*PinStr);
+			const FPCGDataCollectionDesc* PinDesc = nullptr;
+
+			auto ConstructFunctionText = [&PinStr, &FuncStr, &TypeStr]()
+				{
+					return FText::FromString(PinStr + TEXT("_") + FuncStr + TypeStr);
+				};
+
+			if (FuncStr == PCGCustomHLSLKernel::AttributeFunctionSet)
+			{
+				PinDesc = OutputPinDescs.Find(FName(*PinStr));
+
+				if (!PinDesc && InputPinDescs.Find(PinName))
+				{
+#if PCG_KERNEL_LOGGING_ENABLED
+					if (OutErrorText)
+					{
+						*OutErrorText = FText::Format(
+							LOCTEXT("InvalidSetAttributeUsage", "Tried to call attribute function '{0}' on read-only input pin '{1}'."),
+							ConstructFunctionText(),
+							FText::FromName(PinName));
+					}
+#endif
+
+					return false;
+				}
+			}
+			else if (ensure(FuncStr == PCGCustomHLSLKernel::AttributeFunctionGet))
+			{
+				PinDesc = InputPinDescs.Find(FName(*PinStr));
+			}
+
+			if (!PinDesc)
+			{
+#if PCG_KERNEL_LOGGING_ENABLED
+				if (OutErrorText)
+				{
+					*OutErrorText = FText::Format(
+						LOCTEXT("InvalidAttributePinName", "Tried to call attribute function '{0}' on non-existent pin '{1}'."),
+						ConstructFunctionText(),
+						FText::FromName(PinName));
+				}
+#endif
+
+				return false;
+			}
+
+			const UEnum* AttributeTypeEnum = StaticEnum<EPCGKernelAttributeType>();
+			check(AttributeTypeEnum);
+
+			const int64 AttributeType = AttributeTypeEnum->GetValueByName(FName(*TypeStr));
+
+			if (AttributeType == INDEX_NONE)
+			{
+#if PCG_KERNEL_LOGGING_ENABLED
+				if (OutErrorText)
+				{
+					*OutErrorText = FText::Format(
+						LOCTEXT("InvalidAttributePinType", "Tried to call attribute function '{0}' on non-existent type '{1}'."),
+						ConstructFunctionText(),
+						FText::FromString(TypeStr));
+				}
+#endif
+
+				return false;
+			}
+
+			const FPCGKernelAttributeDesc* AttrDesc = nullptr;
+			const FName AttrName = FName(*NameStr);
+
+			if (!PinDesc->DataDescs.IsEmpty())
+			{
+				// Note: This assumes attributes are the same on all data on a pin, which is true for now
+				const FPCGDataDesc& DataDesc = PinDesc->DataDescs[0];
+				AttrDesc = DataDesc.AttributeDescs.FindByPredicate([AttrName](const FPCGKernelAttributeDesc& Desc) { return Desc.Name == AttrName; });
+			}
+
+			if (!AttrDesc)
+			{
+#if PCG_KERNEL_LOGGING_ENABLED
+				if (OutErrorText)
+				{
+					*OutErrorText = FText::Format(
+						LOCTEXT("InvalidAttributeName", "Tried to call attribute function '{0}' on attribute '{1}' which does not exist."),
+						ConstructFunctionText(),
+						FText::FromName(AttrName));
+				}
+#endif
+
+				return false;
+			}
+
+			if (AttrDesc->Type != static_cast<EPCGKernelAttributeType>(AttributeType))
+			{
+#if PCG_KERNEL_LOGGING_ENABLED
+				if (OutErrorText)
+				{
+					const FString ActualTypeStr = AttributeTypeEnum->GetNameStringByIndex(static_cast<int64>(AttrDesc->Type));
+
+					*OutErrorText = FText::Format(
+						LOCTEXT("AttributeTypeMismatch", "Type mismatch for call to attribute function '{0}' on attribute '{1}'. Expected '{2}' but received '{3}'."),
+						ConstructFunctionText(),
+						FText::FromName(AttrName),
+						FText::FromString(TypeStr),
+						FText::FromString(ActualTypeStr));
+				}
+#endif
+
+				return false;
+			}
+		}
+	}
+
+	return true;
+}
+
+bool UPCGCustomHLSLKernel::AreAllOutputPinsWritten(FText* OutWarningText) const
+{
+	TRACE_CPUPROFILER_EVENT_SCOPE(UPCGCustomHLSLKernel::AreAllOutputPinsWritten);
+
+	const FString AllSource = Settings->ShaderFunctions + Settings->ShaderSource;
+
+	// Non-custom kernels initialize the first output pin data automatically.
+	if (Settings->KernelType != EPCGKernelType::Custom && !Settings->OutputPins.IsEmpty())
+	{
+		const FName AutomaticallyInitializedPin = Settings->OutputPins[0].Label;
+
+		for (const FPCGPinPropertiesGPU& Pin : Settings->OutputPins)
+		{
+			if (Pin.Label == AutomaticallyInitializedPin || Pin.Label == NAME_None)
+			{
+				continue;
+			}
+
+			const TStaticArray<FString, 3> WriteStrings = {
+				Pin.Label.ToString() + TEXT("_Set"),
+				Pin.Label.ToString() + TEXT("_CopyElementFrom_"),
+				Pin.Label.ToString() + TEXT("_Initialize"),
+			};
+
+			bool bFoundDataWrite = false;
+
+			for (const FString& WriteString : WriteStrings)
+			{
+				if (AllSource.Contains(WriteString))
+				{
+					bFoundDataWrite = true;
+					break;
+				}
+			}
+
+			if (!bFoundDataWrite)
+			{
+				if (OutWarningText)
+				{
+					*OutWarningText = FText::Format(
+						LOCTEXT("PinMayNotBeWritten", "Data on pin '{0}' may be uninitialized. Add code to write to this data, or mute this error in the node settings."),
+						{ FText::FromName(Pin.Label) });
+				}
+
+				return false;
+			}
+		}
+	}
+
+	return true;
+}
+
+bool UPCGCustomHLSLKernel::ProcessShaderSource(FString* InOutSource, FText* OutErrorText) const
+{
+	TSet<TPair<FString, FString>> PinsRequiringElementCopy;
+
+	const FString Source = Settings->ShaderFunctions + TEXT("\n") + Settings->ShaderSource;
+
+	/** Regex for parsing usages of CopyElementFrom function. */
+	constexpr TCHAR CopyElementPattern[] = { TEXT(/*TargetPin*/"([a-zA-Z0-9]+)" "_CopyElementFrom_" /*SourcePin*/"([a-zA-Z0-9]+)\\s*\\(") };
+	constexpr int CopyElementTargetPinCaptureGroup = 1;
+	constexpr int CopyElementSourcePinCaptureGroup = 2;
+
+	FRegexMatcher ModuleMatcher(FRegexPattern(CopyElementPattern), Source);
+
+	while (ModuleMatcher.FindNext())
+	{
+		const FString SourcePinStr = ModuleMatcher.GetCaptureGroup(CopyElementSourcePinCaptureGroup);
+		const FPCGPinProperties* SourcePinProps = Settings->InputPins.FindByPredicate([Label=FName(*SourcePinStr)](const FPCGPinProperties& InProps) { return InProps.Label == Label; });
+		if (!SourcePinProps)
+		{
+#if PCG_KERNEL_LOGGING_ENABLED
+			if (OutErrorText)
+			{
+				*OutErrorText = FText::Format(LOCTEXT("InvalidCopySourcePin", "Invalid input pin '{0}' used in copy element expression."), FText::FromString(SourcePinStr));
+			}
+#endif
+			return false;
+		}
+
+		const FString TargetPinStr = ModuleMatcher.GetCaptureGroup(CopyElementTargetPinCaptureGroup);
+		const FPCGPinPropertiesGPU* TargetPinProps = Settings->OutputPins.FindByPredicate([Label=FName(*TargetPinStr)](const FPCGPinPropertiesGPU& InProps) { return InProps.Label == Label; });
+		if (!TargetPinProps)
+		{
+#if PCG_KERNEL_LOGGING_ENABLED
+			if (OutErrorText)
+			{
+				*OutErrorText = FText::Format(LOCTEXT("InvalidCopyTargetPin", "Invalid output pin '{0}' used in copy element expression."), FText::FromString(TargetPinStr));
+			}
+#endif
+			return false;
+		}
+
+		if (SourcePinProps->AllowedTypes != TargetPinProps->AllowedTypes)
+		{
+#if PCG_KERNEL_LOGGING_ENABLED
+			if (OutErrorText)
+			{
+				*OutErrorText = LOCTEXT("InvalidCopyPinTypes", "CopyFromElement only allowed across pins of matching type.");
+			}
+#endif
+			return false;
+		}
+
+		PinsRequiringElementCopy.Emplace({ SourcePinStr, TargetPinStr });
+	}
+
+	if (!InOutSource)
+	{
+		return true;
+	}
+
+	// Replace function calls like Out_CopyElementFrom_In(...) with macro PCG_COPY_ALL_ATTRIBUTES_TO_OUTPUT(Out, In, ...).
+	for (const TPair<FString, FString>& SourcePinTargetPin : PinsRequiringElementCopy)
+	{
+		InOutSource->ReplaceInline(
+			*FString::Format(TEXT("{1}_CopyElementFrom_{0}("), { SourcePinTargetPin.Key, SourcePinTargetPin.Value }),
+			*FString::Format(TEXT("PCG_COPY_ALL_ATTRIBUTES_TO_OUTPUT({1}, {0}, "), { SourcePinTargetPin.Key, SourcePinTargetPin.Value })
+		);
+	}
+
+	return true;
+}
+
+const FPCGPinProperties* UPCGCustomHLSLKernel::GetFirstPointInputPin() const
+{
+	// Point processing input pin is first pin of type Point.
+	return Settings->InputPins.FindByPredicate([](const FPCGPinProperties& InProps) { return InProps.AllowedTypes == EPCGDataType::Point; });
+}
+
+const FPCGPinPropertiesGPU* UPCGCustomHLSLKernel::GetFirstPointOutputPin() const
+{
+	// Point processing output pin is first pin of type Point.
+	return Settings->OutputPins.FindByPredicate([](const FPCGPinPropertiesGPU& InProps) { return InProps.AllowedTypes == EPCGDataType::Point; });
+}
+
+int UPCGCustomHLSLKernel::GetElementCountForInputPin(const FPCGPinProperties& InInputPinProps, const UPCGDataBinding* InBinding) const
+{
+	check(InBinding);
+	check(InBinding->Graph);
+
+	FPCGDataCollectionDesc InputDesc;
+	if (!InBinding->GetKernelPinDataDesc_GameThread(this, InInputPinProps.Label, /*bIsInputPin=*/true, InputDesc))
+	{
+		ensure(false);
+		return false;
+	}
+
+	return InputDesc.ComputeDataElementCount(InInputPinProps.AllowedTypes);
+}
+
+#undef LOCTEXT_NAMESPACE
diff --git a/Engine/Plugins/PCG/Source/PCG/Private/Compute/Elements/PCGCustomHLSLKernel.h b/Engine/Plugins/PCG/Source/PCG/Private/Compute/Elements/PCGCustomHLSLKernel.h
new file mode 100644
--- /dev/null
+++ b/Engine/Plugins/PCG/Source/PCG/Private/Compute/Elements/PCGCustomHLSLKernel.h
@@ -0,0 +1,48 @@
+// Copyright Epic Games, Inc. All Rights Reserved.
+
+#pragma once
+
+#include "Compute/PCGComputeKernel.h"
+
+#include "Compute/PCGPinPropertiesGPU.h"
+
+#include "PCGCustomHLSLKernel.generated.h"
+
+class UPCGCustomHLSLSettings;
+
+UCLASS()
+class UPCGCustomHLSLKernel : public UPCGComputeKernel
+{
+	GENERATED_BODY()
+
+public:
+	virtual bool IsKernelValid(FPCGContext* InContext = nullptr, bool bQuiet = true) const override;
+	virtual bool ComputeOutputBindingDataDesc(const UPCGComputeGraph* InGraph, FName InOutputPinLabel, const UPCGDataBinding* InBinding, FPCGDataCollectionDesc& OutDataDesc) const override;
+#if WITH_EDITOR
+	virtual FString GetCookedSource(const UPCGComputeGraph* InComputeGraph) const override;
+#endif
+	virtual int ComputeThreadCount(const UPCGDataBinding* Binding) const override;
+	virtual void GetDataLabels(FName InPinLabel, TArray<FString>& OutDataLabels) const override;
+	virtual void GetKernelAttributeKeys(TArray<FPCGKernelAttributeKey>& OutKeys) const override;
+	virtual uint32 GetThreadCountMultiplier() const override;
+
+protected:
+	/** Will the ThreadCountMultiplier value be applied when calculating the dispatch thread count. */
+	bool IsThreadCountMultiplierInUse() const;
+
+	bool AreAttributesValid(FPCGContext* InContext, FText* OutErrorText) const;
+
+	bool AreAllOutputPinsWritten(FText* OutErrorText) const;
+
+	/** Apply any code gen or transformations or expansions to the shader source. */
+	virtual bool ProcessShaderSource(FString* InOutSource, FText* OutErrorText = nullptr) const;
+
+	const FPCGPinProperties* GetFirstPointInputPin() const;
+	const FPCGPinPropertiesGPU* GetFirstPointOutputPin() const;
+
+	int GetElementCountForInputPin(const FPCGPinProperties& InInputPinProps, const UPCGDataBinding* InBinding) const;
+
+public:
+	UPROPERTY()
+	TObjectPtr<const UPCGCustomHLSLSettings> Settings;
+};
diff --git a/Engine/Plugins/PCG/Source/PCG/Private/Compute/PCGComputeCommon.cpp b/Engine/Plugins/PCG/Source/PCG/Private/Compute/PCGComputeCommon.cpp
--- a/Engine/Plugins/PCG/Source/PCG/Private/Compute/PCGComputeCommon.cpp
+++ b/Engine/Plugins/PCG/Source/PCG/Private/Compute/PCGComputeCommon.cpp
@@ -61,7 +61,7 @@ namespace PCGComputeHelpers
 	void LogKernelWarning(const FPCGContext* Context, const UPCGSettings* Settings, const FText& InText)
 	{
 #if WITH_EDITOR
-		if (Context && ensure(Context->SourceComponent.IsValid() && Context->SourceComponent.Get()))
+		if (Context && ensure(Context->SourceComponent.IsValid() && Context->SourceComponent.Get()) && Settings)
 		{
 			if (UPCGSubsystem* Subsystem = Context->SourceComponent->GetSubsystem())
 			{
@@ -78,7 +78,7 @@ namespace PCGComputeHelpers
 	void LogKernelError(const FPCGContext* Context, const UPCGSettings* Settings, const FText& InText)
 	{
 #if WITH_EDITOR
-		if (Context && ensure(Context->SourceComponent.IsValid() && Context->SourceComponent.Get()))
+		if (Context && ensure(Context->SourceComponent.IsValid() && Context->SourceComponent.Get()) && Settings)
 		{
 			if (UPCGSubsystem* Subsystem = Context->SourceComponent->GetSubsystem())
 			{
@@ -93,216 +93,6 @@ namespace PCGComputeHelpers
 	}
 #endif
 
-	void ComputeOutputPinDataDesc(const FPCGPinPropertiesGPU& PinProperties, const UPCGSettings* Settings, const UPCGDataBinding* Binding, FPCGDataCollectionDesc& OutPinDesc)
-	{
-		check(Settings);
-		check(Binding);
-		
-		const UPCGNode* Node = CastChecked<UPCGNode>(Settings->GetOuter());
-
-		// No size set by kernel, fall back to pin settings.
-		const FPCGPinPropertiesGPUStruct& Props = PinProperties.PropertiesGPU;
-
-		if (Props.InitializationMode == EPCGPinInitMode::FromInputPins)
-		{
-			TArray<FPCGDataCollectionDesc> InputDescs;
-
-			TArray<FName> InitPins;
-
-			for (const FName PinToInitFrom : Props.PinsToInititalizeFrom)
-			{
-				if (const UPCGPin* InitPin = Node->GetInputPin(PinToInitFrom))
-				{
-					InitPins.Emplace(PinToInitFrom);
-					InputDescs.Emplace(PCGDataForGPUHelpers::ComputeInputPinDataDesc(InitPin, Binding));
-				}
-			}
-
-			const int NumInitPins = InitPins.Num();
-			check(NumInitPins == InputDescs.Num());
-
-			// Copies unique (non-reserved) attribute descriptions from 'InDataDesc' to 'OutDataDesc'.
-			auto AddAttributesFromData = [](const FPCGDataDesc& InDataDesc, FPCGDataDesc& OutDataDesc)
-			{
-				for (const FPCGKernelAttributeDesc& InAttrDesc : InDataDesc.AttributeDescs)
-				{
-					if (InAttrDesc.Index >= PCGComputeConstants::NUM_RESERVED_ATTRS
-						&& !OutDataDesc.AttributeDescs.ContainsByPredicate([InAttrDesc](const FPCGKernelAttributeDesc& OutAttrDesc)
-							{
-								// Note: We shouldn't need to check for uniqueness of attr Index, since attributes should all have unique index via
-								// the GlobalAttributeLookupTable
-								return InAttrDesc.Name == OutAttrDesc.Name;
-							}))
-					{
-						OutDataDesc.AttributeDescs.Emplace(InAttrDesc);
-					}
-				}
-			};
-
-			// Combines the data for index i of each pin into one data. Creates exactly 'MaxDataCount' datas.
-			auto AddDataPairwise = [&](int MaxDataCount)
-			{
-				for (int DataIndex = 0; DataIndex < MaxDataCount; ++DataIndex)
-				{
-					// Set element count to 0 for now, but we will overwrite it.
-					FPCGDataDesc& DataDesc = OutPinDesc.DataDescs.Emplace_GetRef(PinProperties.AllowedTypes, /*SetNumElements=*/0);
-					int TotalNumElements = 0; // Total number of elements computed for this data index.
-
-					// For each data index, loop over all the pins and create the uber-data.
-					for (int InputPinIndex = 0; InputPinIndex < NumInitPins; ++InputPinIndex)
-					{
-						const FPCGDataCollectionDesc& InputDesc = InputDescs[InputPinIndex];
-
-						int ClampedDataIndex = DataIndex;
-
-						// If this pin does not have the same number of data, clamp it to the first data.
-						if (DataIndex != InputDesc.DataDescs.Num())
-						{
-							ClampedDataIndex = 0;
-						}
-
-						const FPCGDataDesc& InputDataDesc = InputDesc.DataDescs.IsValidIndex(ClampedDataIndex) ? InputDesc.DataDescs[ClampedDataIndex] : FPCGDataDesc(EPCGDataType::Any, 0);
-
-						if (Props.ElementCountMode == EPCGElementCountMode::FromInputData)
-						{
-							if (Props.ElementMultiplicity == EPCGElementMultiplicity::Product)
-							{
-								TotalNumElements = FMath::Max(TotalNumElements, 1) * InputDataDesc.ElementCount;
-							}
-							else if (Props.ElementMultiplicity == EPCGElementMultiplicity::Sum)
-							{
-								TotalNumElements += InputDataDesc.ElementCount;
-							}
-							else
-							{
-								checkNoEntry();
-							}
-						}
-						else if (Props.ElementCountMode == EPCGElementCountMode::Fixed)
-						{
-							TotalNumElements += Props.ElementCount;
-						}
-						else
-						{
-							checkNoEntry();
-						}
-
-						if (Props.AttributeInheritanceMode == EPCGAttributeInheritanceMode::CopyAttributeSetup)
-						{
-							AddAttributesFromData(InputDataDesc, DataDesc);
-						}
-
-						DataDesc.ElementCount = TotalNumElements;
-					}
-				}
-			};
-
-			if (Props.DataCountMode == EPCGDataCountMode::FromInputData)
-			{
-				// If this is the only input pin, we can just copy it.
-				if (NumInitPins == 1)
-				{
-					for (const FPCGDataDesc& InputDataDesc : InputDescs[0].DataDescs)
-					{
-						FPCGDataDesc& DataDesc = OutPinDesc.DataDescs.Emplace_GetRef(PinProperties.AllowedTypes, InputDataDesc.ElementCount);
-
-						if (Props.AttributeInheritanceMode == EPCGAttributeInheritanceMode::CopyAttributeSetup)
-						{
-							AddAttributesFromData(InputDataDesc, DataDesc);
-						}
-					}
-				}
-				// Take pairs of datas, where the pairs are given by each data of each pin to each data of every other pin.
-				else if (Props.DataMultiplicity == EPCGDataMultiplicity::CartesianProduct)
-				{
-					for (int InputPinIndex = 0; InputPinIndex < NumInitPins; ++InputPinIndex)
-					{
-						const FPCGDataCollectionDesc& InputDesc = InputDescs[InputPinIndex];
-
-						for (int OtherInputPinIndex = InputPinIndex + 1; OtherInputPinIndex < NumInitPins; ++OtherInputPinIndex)
-						{
-							const FPCGDataCollectionDesc& OtherInputDesc = InputDescs[OtherInputPinIndex];
-
-							for (const FPCGDataDesc& InputDataDesc : InputDesc.DataDescs)
-							{
-								for (const FPCGDataDesc& OtherInputDataDesc : OtherInputDesc.DataDescs)
-								{
-									FPCGDataDesc* DataDesc = nullptr;
-
-									if (Props.ElementCountMode == EPCGElementCountMode::FromInputData)
-									{
-										if (Props.ElementMultiplicity == EPCGElementMultiplicity::Product)
-										{
-											DataDesc = &OutPinDesc.DataDescs.Emplace_GetRef(PinProperties.AllowedTypes, InputDataDesc.ElementCount * OtherInputDataDesc.ElementCount);
-										}
-										else if (Props.ElementMultiplicity == EPCGElementMultiplicity::Sum)
-										{
-											DataDesc = &OutPinDesc.DataDescs.Emplace_GetRef(PinProperties.AllowedTypes, InputDataDesc.ElementCount + OtherInputDataDesc.ElementCount);
-										}
-										else
-										{
-											checkNoEntry();
-										}
-									}
-									else if (Props.ElementCountMode == EPCGElementCountMode::Fixed)
-									{
-										DataDesc = &OutPinDesc.DataDescs.Emplace_GetRef(PinProperties.AllowedTypes, Props.ElementCount);
-									}
-									else
-									{
-										checkNoEntry();
-									}
-
-									if (ensure(DataDesc) && Props.AttributeInheritanceMode == EPCGAttributeInheritanceMode::CopyAttributeSetup)
-									{
-										AddAttributesFromData(InputDataDesc, *DataDesc);
-										AddAttributesFromData(OtherInputDataDesc, *DataDesc);
-									}
-								}
-							}
-						}
-					}
-				}
-				// Combine elements for each set of datas, where the sets are given by the Nth datas on each pin (or the first data if there is only one data).
-				else if (Props.DataMultiplicity == EPCGDataMultiplicity::Pairwise)
-				{
-					int MaxDataCount = 0;
-
-					// Find the maximum number of data among the init pins. Note, they should all be the same number of data, or only one data.
-					for (int I = 0; I < NumInitPins; ++I)
-					{
-						MaxDataCount = FMath::Max(MaxDataCount, InputDescs[I].DataDescs.Num());
-					}
-
-					AddDataPairwise(MaxDataCount);
-				}
-				else
-				{
-					checkNoEntry();
-				}
-			}
-			else if (Props.DataCountMode == EPCGDataCountMode::Fixed)
-			{
-				AddDataPairwise(Props.DataCount);
-			}
-			else
-			{
-				checkNoEntry();
-			}
-		}
-		else if (Props.InitializationMode == EPCGPinInitMode::Custom)
-		{
-			for (int I = 0; I < Props.DataCount; ++I)
-			{
-				OutPinDesc.DataDescs.Emplace(PinProperties.AllowedTypes, Props.ElementCount);
-			}
-		}
-		else
-		{
-			checkNoEntry();
-		}
-	}
-
 	bool IsBufferSizeTooLarge(uint64 InBufferSizeBytes, bool bInLogError)
 	{
 		FTextureMemoryStats TextureMemStats;
diff --git a/Engine/Plugins/PCG/Source/PCG/Private/Compute/PCGComputeGraph.cpp b/Engine/Plugins/PCG/Source/PCG/Private/Compute/PCGComputeGraph.cpp
--- a/Engine/Plugins/PCG/Source/PCG/Private/Compute/PCGComputeGraph.cpp
+++ b/Engine/Plugins/PCG/Source/PCG/Private/Compute/PCGComputeGraph.cpp
@@ -6,12 +6,21 @@
 #include "PCGModule.h"
 #include "PCGNode.h"
 #include "Compute/PCGCompilerDiagnostic.h"
+#include "Compute/PCGComputeKernel.h"
+#include "Compute/PCGDataBinding.h"
 #include "Metadata/PCGMetadata.h"
 
 #include "ComputeFramework/ComputeKernelCompileResult.h"
 
 #include UE_INLINE_GENERATED_CPP_BY_NAME(PCGComputeGraph)
 
+#define PCG_DATA_DESCRIPTION_LOGGING 0
+
+uint32 GetTypeHash(const FPCGKernelPin& In)
+{
+	return HashCombine(GetTypeHash(In.KernelIndex), GetTypeHash(In.PinLabel), GetTypeHash(In.bIsInput));
+}
+
 static FPCGCompilerDiagnostic ProcessCompilationMessage(const FComputeKernelCompileMessage& InMessage)
 {
 	FPCGCompilerDiagnostic Diagnostic;
@@ -86,3 +95,97 @@ void UPCGComputeGraph::OnKernelCompilationComplete(int32 InKernelIndex, FCompute
 		UE_LOG(LogPCG, Verbose, TEXT("Compilation message ignored for kernel index %d which has no associated node."), InKernelIndex);
 	}
 }
+
+bool UPCGComputeGraph::ComputeKernelBindingDataDesc(int32 InKernelIndex, FName InPinLabel, bool bInIsInputPin, const UPCGDataBinding* InBinding, FPCGDataCollectionDesc& OutDataDesc) const
+{
+	const int32 BindingIndex = GetBindingIndex(InKernelIndex, InPinLabel, bInIsInputPin);
+	if (BindingIndex != INDEX_NONE)
+	{
+		return ComputeKernelBindingDataDesc(BindingIndex, InBinding, OutDataDesc);
+	}
+	else
+	{
+		ensure(false);
+		return false;
+	}
+}
+
+bool UPCGComputeGraph::ComputeKernelBindingDataDesc(int32 InBindingIndex, const UPCGDataBinding* InBinding, FPCGDataCollectionDesc& OutDataDesc) const
+{
+	TRACE_CPUPROFILER_EVENT_SCOPE(UPCGComputeGraph::ComputeKernelBindingDataDesc);
+
+	const FComputeGraphEdge& GraphEdge = GraphEdges[InBindingIndex];
+
+#if PCG_DATA_DESCRIPTION_LOGGING
+	UE_LOG(LogPCG, Warning, TEXT("Computing data desc for kernel %d, DI '%s', edge index %d (function '%s'), %s"),
+		GraphEdge.KernelIndex,
+		*DataInterfaces[GraphEdge.DataInterfaceIndex]->GetName(),
+		InBindingIndex,
+		*GraphEdge.BindingFunctionNameOverride,
+		GraphEdge.bKernelInput ? TEXT("INPUT") : TEXT("OUTPUT"));
+#endif
+
+	// Find out if this data is coming from CPU (either from CPU nodes or from a separate compute graph).
+	if (const FName* VirtualLabel = CPUDataBindingToVirtualPinLabel.Find(InBindingIndex))
+	{
+		// Create description from compute graph element input data.
+		OutDataDesc = FPCGDataCollectionDesc::BuildFromInputDataCollectionAndInputPinLabel(InBinding->InputDataCollection, *VirtualLabel, InBinding);
+		return true;
+	}
+	else if (GraphEdge.bKernelInput)
+	{
+		if (const int32* UpstreamBindingIndex = DownstreamToUpstreamBinding.Find(InBindingIndex))
+		{
+			return ComputeKernelBindingDataDesc(*UpstreamBindingIndex, InBinding, OutDataDesc);
+		}
+		else
+		{
+			ensure(false);
+			return {};
+		}
+	}
+	else
+	{
+		// Compute description of data from output binding. Data doesn't come from CPU but via a kernel. Consult the kernel for its data description.
+		check(KernelInvocations.IsValidIndex(GraphEdge.KernelIndex));
+		const UPCGComputeKernel* Kernel = CastChecked<UPCGComputeKernel>(KernelInvocations[GraphEdge.KernelIndex]);
+
+		if (const FName* OutputPinLabel = KernelBindingToPinLabel.Find(InBindingIndex))
+		{
+			return Kernel->ComputeOutputBindingDataDesc(this, *OutputPinLabel, InBinding, OutDataDesc);
+		}
+		else
+		{
+			ensure(false);
+			return {};
+		}
+	}
+}
+
+int32 UPCGComputeGraph::GetBindingIndex(int32 InKernelIndex, FName InPinLabel, bool bInIsInput) const
+{
+	if (const int32* BindingIndex = KernelPinToFirstBinding.Find(FPCGKernelPin(InKernelIndex, InPinLabel, bInIsInput)))
+	{
+		return *BindingIndex;
+	}
+	else
+	{
+		ensure(false);
+		return INDEX_NONE;
+	}
+}
+
+bool UPCGComputeGraph::IsGraphValid(FPCGContext* InContext, bool bQuiet) const
+{
+	bool bAllKernelsValid = true;
+
+	// Run validation on all kernels.
+	for (const UComputeKernel* Kernel : KernelInvocations)
+	{
+		bAllKernelsValid &= CastChecked<UPCGComputeKernel>(Kernel)->IsKernelValid(InContext, bQuiet);
+	}
+
+	return bAllKernelsValid;
+}
+
+#undef PCG_DATA_DESCRIPTION_LOGGING 
diff --git a/Engine/Plugins/PCG/Source/PCG/Private/Compute/PCGComputeGraph.h b/Engine/Plugins/PCG/Source/PCG/Private/Compute/PCGComputeGraph.h
--- a/Engine/Plugins/PCG/Source/PCG/Private/Compute/PCGComputeGraph.h
+++ b/Engine/Plugins/PCG/Source/PCG/Private/Compute/PCGComputeGraph.h
@@ -5,14 +5,17 @@
 #include "Compute/PCGDataForGPU.h"
 
 #include "ComputeFramework/ComputeGraph.h"
+#include "ComputeFramework/ComputeKernel.h"
 #include "ComputeFramework/ComputeKernelCompileResult.h"
 
 #include "UObject/ObjectKey.h"
 
 #include "PCGComputeGraph.generated.h"
 
+class UPCGDataBinding;
 class UPCGNode;
 class UPCGPin;
+class UPCGStaticMeshSpawnerKernel;
 
 /** Helper struct for serializing data labels. */
 USTRUCT()
@@ -34,6 +37,37 @@ struct FPCGPinDataLabels
 	TMap</*PinLabel*/FName, FPCGDataLabels> PinToDataLabels;
 };
 
+/** An input or output pin of a kernel. Compute graph does not internally have 'pins' so this is useful for mapping between kernel data and PCG pins. */
+USTRUCT()
+struct FPCGKernelPin
+{
+	GENERATED_BODY()
+
+public:
+	FPCGKernelPin() = default;
+
+	explicit FPCGKernelPin(int32 InKernelIndex, FName InPinLabel, bool bInIsInput)
+		: KernelIndex(InKernelIndex)
+		, PinLabel(InPinLabel)
+		, bIsInput(bInIsInput)
+	{
+	}
+
+	bool operator==(const FPCGKernelPin& Other) const = default;
+
+public:
+	UPROPERTY()
+	int32 KernelIndex = INDEX_NONE;
+
+	UPROPERTY()
+	FName PinLabel = NAME_None;
+
+	UPROPERTY()
+	bool bIsInput = false;
+
+	friend uint32 GetTypeHash(const FPCGKernelPin& In);
+};
+
 UCLASS()
 class UPCGComputeGraph : public UComputeGraph
 {
@@ -47,38 +81,43 @@ public:
 	/** Get the attribute table collated statically during compilation. */
 	const FPCGKernelAttributeTable& GetAttributeStaticTable() const { return AttributeStaticTable; }
 
+	// TODO: Should be map from BindingIndex->Labels, but needs GPU compiler flow tweaked so that kernel index are available at DI creation time.
 	/** Get the data labels collated statically during compilation. */
 	const TMap<TSoftObjectPtr<const UPCGSettings>, FPCGPinDataLabels>& GetDataLabelsStaticTable() const { return DataLabelsStaticTable; }
 
 	const TArray<FString>& GetStringTable() const { return StringTable; }
 
-public:
-	TMap<TObjectKey<const UPCGNode>, TArray<FComputeKernelCompileMessage>> KernelToCompileMessages;
-	
-	/** Set of input pins at the CPU -> GPU border. */
-	UPROPERTY()
-	TArray<TSoftObjectPtr<const UPCGPin>> PinsReceivingDataFromCPU;
+	// in PCG a pin is a data interface (multiple incident edges on an input pin will be merged using a DI).
+	// Possibly we shouldn't directly talk to UPCGDataBinding here
+	bool ComputeKernelBindingDataDesc(int32 InKernelIndex, FName InPinLabel, bool bInIsInputPin, const UPCGDataBinding* InBinding, FPCGDataCollectionDesc& OutDataDesc) const;
+	bool ComputeKernelBindingDataDesc(int32 InEdgeIndex, const UPCGDataBinding* InBinding, FPCGDataCollectionDesc& OutDataDesc) const;
 
-	/** Pin label aliases, used for selecting data items corresponding to an input pin from the input data collection. */
-	UPROPERTY()
-	TMap<TSoftObjectPtr<const UPCGPin>, FName> InputPinLabelAliases;
+	/** Bindings in compute graphs roughly map to a single API on a Data Interface, like GetNumData(). There can be multiple bindings per PCG edge. */
+	int32 GetBindingIndex(int32 InKernelIndex, FName InPinLabel, bool bInIsInput) const;
 
-	/** Mapping from upstream output pin to downstream pin alias, used to select data items originating from upstream pin from the input data collection. */
-	UPROPERTY()
-	TMap<TSoftObjectPtr<const UPCGPin>, FName> OutputCPUPinToInputGPUPinAlias;
+	bool IsGraphValid(FPCGContext* InContext = nullptr, bool bQuiet = true) const;
 
+public:
+	TMap<TObjectKey<const UPCGNode>, TArray<FComputeKernelCompileMessage>> KernelToCompileMessages;
+	
 	/** Node corresponding to each kernel, useful for compilation feedback. */
 	UPROPERTY()
 	TArray<TSoftObjectPtr<const UPCGNode>> KernelToNode;
 
 	/** List of settings for all nodes that spawn static meshes, so we can do required primitive & DI setup when compute graph element executes. */
 	UPROPERTY()
-	TArray<TSoftObjectPtr<const UPCGSettings>> StaticMeshSpawners;
+	TArray<TObjectPtr<const UPCGStaticMeshSpawnerKernel>> StaticMeshSpawners;
 
 	UPROPERTY()
 	bool bLogDataDescriptions = false;
 
 protected:
+	UPROPERTY()
+	TMap</*BindingIndex*/int32, /*PinLabel*/FName> KernelBindingToPinLabel;
+
+	UPROPERTY()
+	TMap</*DownstreamBindingIndex*/int32, /*UpstreamBindingIndex*/int32> DownstreamToUpstreamBinding;
+
 	/** Global attribute information collated during compilation. */
 	UPROPERTY()
 	FPCGKernelAttributeTable AttributeStaticTable;
@@ -91,5 +130,12 @@ protected:
 	UPROPERTY()
 	TArray<FString> StringTable;
 
+	UPROPERTY()
+	TMap<FPCGKernelPin, /*FirstBindingIndex*/int32> KernelPinToFirstBinding;
+
+	/** Binding index to virtual label for bindings that are that receive data from nodes executing on the CPU or in separate compute graphs. */
+	UPROPERTY()
+	TMap</*BindingIndex*/int32, /*VirtualPinLabel*/FName> CPUDataBindingToVirtualPinLabel;
+
 	friend class FPCGGraphCompilerGPU;
 };
diff --git a/Engine/Plugins/PCG/Source/PCG/Private/Compute/PCGComputeKernel.cpp b/Engine/Plugins/PCG/Source/PCG/Private/Compute/PCGComputeKernel.cpp
new file mode 100644
--- /dev/null
+++ b/Engine/Plugins/PCG/Source/PCG/Private/Compute/PCGComputeKernel.cpp
@@ -0,0 +1,256 @@
+// Copyright Epic Games, Inc. All Rights Reserved.
+
+#include "Compute/PCGComputeKernel.h"
+
+#include "PCGModule.h"
+#include "PCGSettings.h"
+#include "Compute/PCGComputeGraph.h"
+#include "Compute/PCGDataBinding.h"
+#include "Compute/PCGPinPropertiesGPU.h"
+
+#include UE_INLINE_GENERATED_CPP_BY_NAME(PCGComputeKernel)
+
+#define LOCTEXT_NAMESPACE "PCGComputeKernel"
+
+namespace PCGComputeKernel
+{
+	static TAutoConsoleVariable<bool> CVarWarnOnOverridePinUsage(
+		TEXT("pcg.Graph.GPU.WarnOnOverridePinUsage"),
+		true,
+		TEXT("Enables warnings when parameters are overidden on GPU nodes."));
+}
+
+void UPCGComputeKernel::ComputeDataDescFromPinProperties(const FPCGPinPropertiesGPU& OutputPinProps, const TArrayView<const FPCGPinProperties>& InInputPinProps, const UPCGDataBinding* InBinding, FPCGDataCollectionDesc& OutPinDesc) const
+{
+	check(InBinding);
+
+	const FPCGPinPropertiesGPUStruct& Props = OutputPinProps.PropertiesGPU;
+
+	if (Props.InitializationMode == EPCGPinInitMode::FromInputPins)
+	{
+		TArray<FName> InitPins;
+		TArray<FPCGDataCollectionDesc> InputDescs;
+
+		for (FName PinToInitFrom : Props.PinsToInititalizeFrom)
+		{
+			if (const FPCGPinProperties* InputPinProps = InInputPinProps.FindByPredicate([PinToInitFrom](const FPCGPinProperties& InProps) { return InProps.Label == PinToInitFrom; }))
+			{
+				InitPins.Emplace(PinToInitFrom);
+
+				FPCGDataCollectionDesc& InputDesc = InputDescs.AddDefaulted_GetRef();
+				ensure(InBinding->GetKernelPinDataDesc_GameThread(this, PinToInitFrom, /*bIsInputPin=*/true, InputDesc));
+			}
+		}
+
+		const int NumInitPins = InitPins.Num();
+		check(NumInitPins == InputDescs.Num());
+
+		// Copies unique (non-reserved) attribute descriptions from 'InDataDesc' to 'OutDataDesc'.
+		auto AddAttributesFromData = [](const FPCGDataDesc& InDataDesc, FPCGDataDesc& OutDataDesc)
+		{
+			for (const FPCGKernelAttributeDesc& InAttrDesc : InDataDesc.AttributeDescs)
+			{
+				if (InAttrDesc.Index >= PCGComputeConstants::NUM_RESERVED_ATTRS
+					&& !OutDataDesc.AttributeDescs.ContainsByPredicate([InAttrDesc](const FPCGKernelAttributeDesc& OutAttrDesc)
+						{
+							// Note: We shouldn't need to check for uniqueness of attr Index, since attributes should all have unique index via
+							// the GlobalAttributeLookupTable
+							return InAttrDesc.Name == OutAttrDesc.Name;
+						}))
+				{
+					OutDataDesc.AttributeDescs.Emplace(InAttrDesc);
+				}
+			}
+		};
+
+		// Combines the data for index i of each pin into one data. Creates exactly 'MaxDataCount' datas.
+		auto AddDataPairwise = [&](int MaxDataCount)
+		{
+			for (int DataIndex = 0; DataIndex < MaxDataCount; ++DataIndex)
+			{
+				// Set element count to 0 for now, but we will overwrite it.
+				FPCGDataDesc& DataDesc = OutPinDesc.DataDescs.Emplace_GetRef(OutputPinProps.AllowedTypes, /*SetNumElements=*/0);
+				int TotalNumElements = 0; // Total number of elements computed for this data index.
+
+				// For each data index, loop over all the pins and create the uber-data.
+				for (int InputPinIndex = 0; InputPinIndex < NumInitPins; ++InputPinIndex)
+				{
+					const FPCGDataCollectionDesc& InputDesc = InputDescs[InputPinIndex];
+
+					int ClampedDataIndex = DataIndex;
+
+					// If this pin does not have the same number of data, clamp it to the first data.
+					if (DataIndex != InputDesc.DataDescs.Num())
+					{
+						ClampedDataIndex = 0;
+					}
+
+					const FPCGDataDesc& InputDataDesc = InputDesc.DataDescs.IsValidIndex(ClampedDataIndex) ? InputDesc.DataDescs[ClampedDataIndex] : FPCGDataDesc(EPCGDataType::Any, 0);
+
+					if (Props.ElementCountMode == EPCGElementCountMode::FromInputData)
+					{
+						if (Props.ElementMultiplicity == EPCGElementMultiplicity::Product)
+						{
+							TotalNumElements = FMath::Max(TotalNumElements, 1) * InputDataDesc.ElementCount;
+						}
+						else if (Props.ElementMultiplicity == EPCGElementMultiplicity::Sum)
+						{
+							TotalNumElements += InputDataDesc.ElementCount;
+						}
+						else
+						{
+							checkNoEntry();
+						}
+					}
+					else if (Props.ElementCountMode == EPCGElementCountMode::Fixed)
+					{
+						TotalNumElements += Props.ElementCount;
+					}
+					else
+					{
+						checkNoEntry();
+					}
+
+					if (Props.AttributeInheritanceMode == EPCGAttributeInheritanceMode::CopyAttributeSetup)
+					{
+						AddAttributesFromData(InputDataDesc, DataDesc);
+					}
+
+					DataDesc.ElementCount = TotalNumElements;
+				}
+			}
+		};
+
+		if (Props.DataCountMode == EPCGDataCountMode::FromInputData)
+		{
+			// If this is the only input pin, we can just copy it.
+			if (NumInitPins == 1)
+			{
+				for (const FPCGDataDesc& InputDataDesc : InputDescs[0].DataDescs)
+				{
+					FPCGDataDesc& DataDesc = OutPinDesc.DataDescs.Emplace_GetRef(OutputPinProps.AllowedTypes, InputDataDesc.ElementCount);
+
+					if (Props.AttributeInheritanceMode == EPCGAttributeInheritanceMode::CopyAttributeSetup)
+					{
+						AddAttributesFromData(InputDataDesc, DataDesc);
+					}
+				}
+			}
+			// Take pairs of datas, where the pairs are given by each data of each pin to each data of every other pin.
+			else if (Props.DataMultiplicity == EPCGDataMultiplicity::CartesianProduct)
+			{
+				for (int InputPinIndex = 0; InputPinIndex < NumInitPins; ++InputPinIndex)
+				{
+					const FPCGDataCollectionDesc& InputDesc = InputDescs[InputPinIndex];
+
+					for (int OtherInputPinIndex = InputPinIndex + 1; OtherInputPinIndex < NumInitPins; ++OtherInputPinIndex)
+					{
+						const FPCGDataCollectionDesc& OtherInputDesc = InputDescs[OtherInputPinIndex];
+
+						for (const FPCGDataDesc& InputDataDesc : InputDesc.DataDescs)
+						{
+							for (const FPCGDataDesc& OtherInputDataDesc : OtherInputDesc.DataDescs)
+							{
+								FPCGDataDesc* DataDesc = nullptr;
+
+								if (Props.ElementCountMode == EPCGElementCountMode::FromInputData)
+								{
+									if (Props.ElementMultiplicity == EPCGElementMultiplicity::Product)
+									{
+										DataDesc = &OutPinDesc.DataDescs.Emplace_GetRef(OutputPinProps.AllowedTypes, InputDataDesc.ElementCount * OtherInputDataDesc.ElementCount);
+									}
+									else if (Props.ElementMultiplicity == EPCGElementMultiplicity::Sum)
+									{
+										DataDesc = &OutPinDesc.DataDescs.Emplace_GetRef(OutputPinProps.AllowedTypes, InputDataDesc.ElementCount + OtherInputDataDesc.ElementCount);
+									}
+									else
+									{
+										checkNoEntry();
+									}
+								}
+								else if (Props.ElementCountMode == EPCGElementCountMode::Fixed)
+								{
+									DataDesc = &OutPinDesc.DataDescs.Emplace_GetRef(OutputPinProps.AllowedTypes, Props.ElementCount);
+								}
+								else
+								{
+									checkNoEntry();
+								}
+
+								if (ensure(DataDesc) && Props.AttributeInheritanceMode == EPCGAttributeInheritanceMode::CopyAttributeSetup)
+								{
+									AddAttributesFromData(InputDataDesc, *DataDesc);
+									AddAttributesFromData(OtherInputDataDesc, *DataDesc);
+								}
+							}
+						}
+					}
+				}
+			}
+			// Combine elements for each set of datas, where the sets are given by the Nth datas on each pin (or the first data if there is only one data).
+			else if (Props.DataMultiplicity == EPCGDataMultiplicity::Pairwise)
+			{
+				int MaxDataCount = 0;
+
+				// Find the maximum number of data among the init pins. Note, they should all be the same number of data, or only one data.
+				for (int I = 0; I < NumInitPins; ++I)
+				{
+					MaxDataCount = FMath::Max(MaxDataCount, InputDescs[I].DataDescs.Num());
+				}
+
+				AddDataPairwise(MaxDataCount);
+			}
+			else
+			{
+				checkNoEntry();
+			}
+		}
+		else if (Props.DataCountMode == EPCGDataCountMode::Fixed)
+		{
+			AddDataPairwise(Props.DataCount);
+		}
+		else
+		{
+			checkNoEntry();
+		}
+	}
+	else if (Props.InitializationMode == EPCGPinInitMode::Custom)
+	{
+		for (int I = 0; I < Props.DataCount; ++I)
+		{
+			OutPinDesc.DataDescs.Emplace(OutputPinProps.AllowedTypes, Props.ElementCount);
+		}
+	}
+	else
+	{
+		checkNoEntry();
+	}
+}
+
+bool UPCGComputeKernel::ValidatePCGNode(FPCGContext* InContext, bool bQuiet, const UPCGSettings* InSettings) const
+{
+	check(InSettings);
+
+	if (PCGComputeKernel::CVarWarnOnOverridePinUsage.GetValueOnAnyThread() && !bQuiet)
+	{
+		for (const FPCGSettingsOverridableParam& Param : InSettings->OverridableParams())
+		{
+			if (ensure(!Param.PropertiesNames.IsEmpty()))
+			{
+				const FName PropertyName = Param.PropertiesNames[0];
+
+				if (InSettings->IsPropertyOverriddenByPin(PropertyName))
+				{
+					PCG_KERNEL_VALIDATION_WARN(InContext, InSettings, bQuiet, FText::Format(
+						LOCTEXT("ParamOverrideGPU", "Tried to override pin '{0}', but overrides are not supported on GPU nodes."),
+						FText::FromName(PropertyName)));
+				}
+			}
+		}
+	}
+
+	// Validate types of incident edges to make sure we catch invalid cases like Spatial -> Point.
+	return InSettings->AreInputEdgesValid(InContext, bQuiet);
+}
+
+#undef LOCTEXT_NAMESPACE
diff --git a/Engine/Plugins/PCG/Source/PCG/Private/Compute/PCGComputeKernel.h b/Engine/Plugins/PCG/Source/PCG/Private/Compute/PCGComputeKernel.h
new file mode 100644
--- /dev/null
+++ b/Engine/Plugins/PCG/Source/PCG/Private/Compute/PCGComputeKernel.h
@@ -0,0 +1,80 @@
+// Copyright Epic Games, Inc. All Rights Reserved.
+
+#pragma once
+
+#include "ComputeFramework/ComputeKernel.h"
+
+#include "ComputeFramework/ComputeGraph.h"
+#include "UObject/ObjectKey.h"
+
+#include "PCGComputeKernel.generated.h"
+
+class UPCGComputeGraph;
+class UPCGDataBinding;
+class UPCGSettings;
+struct FPCGContext;
+struct FPCGDataCollectionDesc;
+struct FPCGGPUCompilationContext;
+struct FPCGKernelAttributeKey;
+struct FPCGPinProperties;
+struct FPCGPinPropertiesGPU;
+
+UCLASS(Abstract)
+class UPCGComputeKernel : public UComputeKernel
+{
+	GENERATED_BODY()
+
+public:
+	/** Get/set index of kernel in kernels array on owning compute graph. */
+	int32 GetKernelIndex() const { return KernelIndex; }
+	void SetKernelIndex(int32 InKernelIndex) { KernelIndex = InKernelIndex; }
+
+	/** Performs validation and returns true if this node is suitable for deployment to the GPU. */
+	virtual bool IsKernelValid(FPCGContext* InContext = nullptr, bool bQuiet = true) const { return true; }
+
+#if WITH_EDITOR
+	/** Produces the node specific portion of kernel shader source text, including the main entry point. */
+	virtual FString GetCookedSource(const UPCGComputeGraph* InComputeGraph) const PURE_VIRTUAL(UPCGComputeKernel::GetCookedSource, return {};);
+
+	/** Create additional input data interfaces to marshal any required input data. */
+	virtual void CreateAdditionalInputDataInterfaces(FPCGGPUCompilationContext& InOutContext, UObject* InObjectOuter, TArray<TObjectPtr<UComputeDataInterface>>& OutDataInterfaces) const {}
+
+	/** Create additional output data interfaces to marshal any required output data. */
+	virtual void CreateAdditionalOutputDataInterfaces(FPCGGPUCompilationContext& InOutContext, UObject* InObjectOuter, TArray<TObjectPtr<UComputeDataInterface>>& OutDataInterfaces) const {}
+#endif
+
+	/** Compute a description of data that will be output from pin InOutputPinLabel. */
+	virtual bool ComputeOutputBindingDataDesc(const UPCGComputeGraph* InGraph, FName InOutputPinLabel, const UPCGDataBinding* InBinding, FPCGDataCollectionDesc& OutDataDesc) const PURE_VIRTUAL(UPCGComputeKernel::ComputeOutputBindingDataDesc, return {};);
+
+	/** Compute how many threads should be dispatched to execute this node on the GPU. */
+	virtual int ComputeThreadCount(const UPCGDataBinding* Binding) const PURE_VIRTUAL(UPCGComputeKernel::GetCookedSource, return 0;);
+
+	/** Get any data labels that are used statically by this node. */
+	virtual void GetDataLabels(FName InPinLabel, TArray<FString>& OutDataLabels) const {}
+
+	/** Whether the output pin should have a buffer of counters attached, useful for compaction and other counting requirements. */
+	virtual bool DoesOutputPinRequireElementCounters(FName InOutputPinLabel) const { return false; }
+
+	/** Add any strings emitted by this node that are known statically at compile time. */
+	virtual void AddStaticCreatedStrings(TArray<FString>& InOutStringTable) const {}
+
+	/** Get all the attributes read or written by this node for which we know the name and type statically. In cases that only the
+	* attribute name used by the node is known (example: user specifies SM Spawner instance data attributes by name only, not type), the attribute
+	* is omitted from this list and must be resolved at execution time in the data provider.
+	*/
+	virtual void GetKernelAttributeKeys(TArray<FPCGKernelAttributeKey>& OutKeys) const {}
+
+	/** Get multiplier to be applied to the dispatched thread count. */
+	virtual uint32 GetThreadCountMultiplier() const { return 1u; }
+
+	/** Helper to compute data descriptions when data is entirely determined by pin properties. */
+	void ComputeDataDescFromPinProperties(const FPCGPinPropertiesGPU& OutputPinProps, const TArrayView<const FPCGPinProperties>& InInputPinProps, const UPCGDataBinding* InBinding, FPCGDataCollectionDesc& OutPinDesc) const;
+
+	/** Do any validation of the PCG node. Returns true if node is valid and GPU execution can proceed. */
+	bool ValidatePCGNode(FPCGContext* InContext, bool bQuiet, const UPCGSettings* InSettings) const;
+
+protected:
+	/** Index into kernels array on owning compute graph. */
+	UPROPERTY()
+	int32 KernelIndex = INDEX_NONE;
+};
diff --git a/Engine/Plugins/PCG/Source/PCG/Private/Compute/PCGDataBinding.cpp b/Engine/Plugins/PCG/Source/PCG/Private/Compute/PCGDataBinding.cpp
--- a/Engine/Plugins/PCG/Source/PCG/Private/Compute/PCGDataBinding.cpp
+++ b/Engine/Plugins/PCG/Source/PCG/Private/Compute/PCGDataBinding.cpp
@@ -7,6 +7,7 @@
 #include "PCGSettings.h"
 #include "Compute/PCGComputeCommon.h"
 #include "Compute/PCGComputeGraph.h"
+#include "Compute/PCGComputeKernel.h"
 #include "Compute/Data/PCGProxyForGPUData.h"
 #include "Data/PCGPointData.h"
 #include "Metadata/PCGMetadata.h"
@@ -35,62 +36,36 @@ void UPCGDataBinding::Initialize(
 #endif
 }
 
-bool UPCGDataBinding::GetOutputPinDataDesc_GameThread(FName InProducerPinLabel, const UPCGSettings* InProducerSettings, FPCGDataCollectionDesc& OutDescription) const
+bool UPCGDataBinding::GetKernelPinDataDesc_GameThread(const UPCGComputeKernel* InKernel, FName InPinLabel, bool bIsInput, FPCGDataCollectionDesc& OutDescription) const
 {
-	check(InProducerSettings);
-	const UPCGNode* Node = CastChecked<UPCGNode>(InProducerSettings->GetOuter());
-
-	return GetOutputPinDataDesc_GameThread(Node->GetOutputPin(InProducerPinLabel), InProducerSettings, OutDescription);
+	return ensure(Graph && InKernel) && GetKernelPinDataDesc_GameThread(Graph->GetBindingIndex(InKernel->GetKernelIndex(), InPinLabel, bIsInput), OutDescription);
 }
 
-bool UPCGDataBinding::GetOutputPinDataDesc_GameThread(const UPCGPin* InProducerPin, const UPCGSettings* InProducerSettings, FPCGDataCollectionDesc& OutDescription) const
+bool UPCGDataBinding::GetKernelPinDataDesc_GameThread(int32 InGraphBindingIndex, FPCGDataCollectionDesc& OutDescription) const
 {
 	ensure(IsInGameThread());
 
-	if (const FPCGDataCollectionDesc* FoundDescription = OutputPinDataDescriptionCache.Find(InProducerPin))
+	if (InGraphBindingIndex == INDEX_NONE)
 	{
-		OutDescription = *FoundDescription;
-		return true;
+		ensure(false);
+		return false;
 	}
 
-	FPCGDataCollectionDesc Description;
-	if (InProducerSettings->ComputeOutputPinDataDesc(InProducerPin, this, Description))
+	if (const FPCGDataCollectionDesc* FoundDescription = OutputPinDataDescriptionCache.Find(InGraphBindingIndex))
 	{
-		OutputPinDataDescriptionCache.Add(InProducerPin, Description);
-		OutDescription = MoveTemp(Description);
+		OutDescription = *FoundDescription;
 		return true;
 	}
-	else
-	{
-		return false;
-	}
-}
 
-bool UPCGDataBinding::ComputeCPUOutputPinDataDesc(const UPCGPin* OutputPin, FPCGDataCollectionDesc& OutDesc) const
-{
-	check(OutputPin);
-	check(Graph);
-
-	if (const FPCGDataCollectionDesc* FoundDescription = OutputPinDataDescriptionCache.Find(OutputPin))
+	if (Graph->ComputeKernelBindingDataDesc(InGraphBindingIndex, this, OutDescription))
 	{
-		OutDesc = *FoundDescription;
+		OutputPinDataDescriptionCache.Add(InGraphBindingIndex, OutDescription);
 		return true;
 	}
-
-	// Find the unique alias for the virtual pin created during compilation to receive the data from the upstream OutputPin.
-	// OutputPin can either be on a CPU node or on a GPU node in a different compute graph.
-	const FName* FoundPinAlias = Graph->OutputCPUPinToInputGPUPinAlias.Find(OutputPin);
-	if (!FoundPinAlias)
+	else
 	{
 		return false;
 	}
-
-	// Create description from compute graph element input data.
-	OutDesc = FPCGDataCollectionDesc::BuildFromInputDataCollectionAndInputPinLabel(InputDataCollection, *FoundPinAlias, this);
-
-	OutputPinDataDescriptionCache.Add(OutputPin, OutDesc);
-
-	return true;
 }
 
 void UPCGDataBinding::ReceiveDataFromGPU_GameThread(
@@ -161,17 +136,6 @@ void UPCGDataBinding::InitializeInputData(const FPCGDataCollection& InComputeGra
 		}
 	}
 
-	// Link each input pin to the data collection, so that data providers can find the data.
-	for (const TSoftObjectPtr<const UPCGPin>& InputPinPtr : Graph->PinsReceivingDataFromCPU)
-	{
-		if (const UPCGPin* InputPin = InputPinPtr.Get())
-		{
-			InputPins.Add(InputPin);
-		}
-	}
-
-	InputPinLabelAliases = Graph->InputPinLabelAliases;
-
 	// Augment static string table with attributes incoming from CPU.
 	AddInputDataAttributesToTable(InComputeGraphElementInputData);
 
@@ -368,35 +332,36 @@ void UPCGDataBinding::DebugLogDataDescriptions()
 		}
 	};
 
-	UE_LOG(LogPCG, Display, TEXT("\n### INPUT PIN DATA DESCRIPTIONS ###"));
-	for (const TSoftObjectPtr<const UPCGNode>& Node : Graph->KernelToNode)
-	{
-		const UPCGSettings* Settings = Node.IsValid() ? Node->GetSettings() : nullptr;
-		if (Settings && Settings->bDumpDataDescriptions)
-		{
-			UE_LOG(LogPCG, Display, TEXT("Node: %s"), *Node->GetNodeTitle(EPCGNodeTitleType::ListView).ToString());
-
-			for (const UPCGPin* Pin : Node->GetInputPins())
-			{
-				const FPCGDataCollectionDesc DataDesc = PCGDataForGPUHelpers::ComputeInputPinDataDesc(Pin, this);
-				UE_LOG(LogPCG, Display, TEXT("\tInput Pin: %s (%d data)"), *Pin->Properties.Label.ToString(), DataDesc.DataDescs.Num());
-				LogDataDescription(DataDesc);
-			}
-
-			for (const UPCGPin* Pin : Node->GetOutputPins())
-			{
-				FPCGDataCollectionDesc DataDesc;
-				if (ensure(GetOutputPinDataDesc_GameThread(Pin, Settings, DataDesc)))
-				{
-					UE_LOG(LogPCG, Display, TEXT("\tOutput Pin: %s (%d data)"), *Pin->Properties.Label.ToString(), DataDesc.DataDescs.Num());
-					LogDataDescription(DataDesc);
-				}
-				else
-				{
-					UE_LOG(LogPCG, Display, TEXT("\tOutput Pin: %s MISSING!"), *Pin->Properties.Label.ToString());
-				}
-			}
-		}
-	}
+	// TODO - we don't have explicit pin lists for kernels because we don't need them. Decide how to bring this back or pivot the feature.
+	//UE_LOG(LogPCG, Display, TEXT("\n### INPUT PIN DATA DESCRIPTIONS ###"));
+	//for (const TSoftObjectPtr<const UPCGNode>& Node : Graph->KernelToNode)
+	//{
+	//	const UPCGSettings* Settings = Node.IsValid() ? Node->GetSettings() : nullptr;
+	//	if (Settings && Settings->bDumpDataDescriptions)
+	//	{
+	//		UE_LOG(LogPCG, Display, TEXT("Node: %s"), *Node->GetNodeTitle(EPCGNodeTitleType::ListView).ToString());
+
+	//		for (const UPCGPin* Pin : Node->GetInputPins())
+	//		{
+	//			const FPCGDataCollectionDesc DataDesc = PCGDataForGPUHelpers::ComputeInputPinDataDesc(Pin, this);
+	//			UE_LOG(LogPCG, Display, TEXT("\tInput Pin: %s (%d data)"), *Pin->Properties.Label.ToString(), DataDesc.DataDescs.Num());
+	//			LogDataDescription(DataDesc);
+	//		}
+
+	//		for (const UPCGPin* Pin : Node->GetOutputPins())
+	//		{
+	//			FPCGDataCollectionDesc DataDesc;
+	//			if (ensure(GetOutputPinDataDesc_GameThread(Pin, Settings, DataDesc)))
+	//			{
+	//				UE_LOG(LogPCG, Display, TEXT("\tOutput Pin: %s (%d data)"), *Pin->Properties.Label.ToString(), DataDesc.DataDescs.Num());
+	//				LogDataDescription(DataDesc);
+	//			}
+	//			else
+	//			{
+	//				UE_LOG(LogPCG, Display, TEXT("\tOutput Pin: %s MISSING!"), *Pin->Properties.Label.ToString());
+	//			}
+	//		}
+	//	}
+	//}
 #endif
 }
diff --git a/Engine/Plugins/PCG/Source/PCG/Private/Compute/PCGDataBinding.h b/Engine/Plugins/PCG/Source/PCG/Private/Compute/PCGDataBinding.h
--- a/Engine/Plugins/PCG/Source/PCG/Private/Compute/PCGDataBinding.h
+++ b/Engine/Plugins/PCG/Source/PCG/Private/Compute/PCGDataBinding.h
@@ -9,7 +9,9 @@
 
 class UPCGComponent;
 class UPCGComputeGraph;
+class UPCGComputeKernel;
 class UPCGProxyForGPUData;
+class UPCGStaticMeshSpawnerKernel;
 
 USTRUCT()
 struct FPCGSpawnerPrimitives
@@ -66,11 +68,8 @@ public:
 	const TArray<FString>& GetStringTable() const { return StringTable; }
 
 	/** Get description of data produced by an output pin. */
-	bool GetOutputPinDataDesc_GameThread(FName InProducerPinLabel, const UPCGSettings* InProducerSettings, FPCGDataCollectionDesc& OutDescription) const;
-	bool GetOutputPinDataDesc_GameThread(const UPCGPin* InProducerPin, const UPCGSettings* InProducerSettings, FPCGDataCollectionDesc& OutDescription) const;
-
-	/** If the output pin is a connected upstream pin of the compute graph (on a CPU node or a GPU node from a different compute graph), returns a description of this data. */
-	bool ComputeCPUOutputPinDataDesc(const UPCGPin* OutputPin, FPCGDataCollectionDesc& OutDesc) const;
+	bool GetKernelPinDataDesc_GameThread(const UPCGComputeKernel* InKernel, FName InPinLabel, bool bIsInput, FPCGDataCollectionDesc& OutDescription) const;
+	bool GetKernelPinDataDesc_GameThread(int32 InGraphBindingIndex, FPCGDataCollectionDesc& OutDescription) const;
 
 	void ReceiveDataFromGPU_GameThread(const UPCGProxyForGPUData* InData, const UPCGSettings* InProducerSettings, EPCGExportMode InExportMode, FName InPinLabel, FName InPinLabelAlias);
 
@@ -104,21 +103,12 @@ public:
 	UPROPERTY()
 	TObjectPtr<const UPCGComputeGraph> Graph = nullptr;
 
-	/** All the node input pins that have edges that cross from CPU to GPU. */
-	UPROPERTY()
-	TSet<TSoftObjectPtr<const UPCGPin>> InputPins;
-
-	/** Input pin label aliases generated during compilation so that every GPU input pin receiving data from CPU has a unique label which can be used
-	 * to retrieve the data for this pin from the compute graph element data collection. */
-	UPROPERTY()
-	TMap<TSoftObjectPtr<const UPCGPin>, FName> InputPinLabelAliases;
-
 	/** Data arriving on compute graph element. Since the compute graph is collapsed to a single element, all data crossing from CPU to GPU is in a single collection. */
 	UPROPERTY()
 	FPCGDataCollection InputDataCollection;
 
 	UPROPERTY()
-	TMap<TObjectPtr<const UPCGSettings>, FPCGSpawnerPrimitives> MeshSpawnersToPrimitives;
+	TMap<TObjectPtr<const UPCGStaticMeshSpawnerKernel>, FPCGSpawnerPrimitives> MeshSpawnersToPrimitives;
 
 	TArray<FPCGDataToDebug> DataToDebug;
 	TArray<FPCGDataToDebug> DataToInspect;
@@ -139,5 +129,5 @@ protected:
 	TMap<FSourceBufferAndAttributeIndex, /*Attribute ID*/int32> AttributeIdRemap;
 
 	/** Cache of data descriptions computed for output pins to amortized cost of building data descriptions. */
-	mutable TMap<const UPCGPin*, FPCGDataCollectionDesc> OutputPinDataDescriptionCache;
+	mutable TMap</*Binding index*/int32, FPCGDataCollectionDesc> OutputPinDataDescriptionCache;
 };
diff --git a/Engine/Plugins/PCG/Source/PCG/Private/Compute/PCGDataForGPU.cpp b/Engine/Plugins/PCG/Source/PCG/Private/Compute/PCGDataForGPU.cpp
--- a/Engine/Plugins/PCG/Source/PCG/Private/Compute/PCGDataForGPU.cpp
+++ b/Engine/Plugins/PCG/Source/PCG/Private/Compute/PCGDataForGPU.cpp
@@ -594,54 +594,6 @@ namespace PCGDataForGPUHelpers
 
 		OutCustomFloatCount = OffsetFloats;
 	}
-
-	FPCGDataCollectionDesc ComputeInputPinDataDesc(const UPCGSettings* Settings, const FName& InputPinLabel, const UPCGDataBinding* Binding)
-	{
-		check(Settings);
-
-		const UPCGNode* Node = Cast<UPCGNode>(Settings->GetOuter());
-		const UPCGPin* InputPin = Node ? Node->GetInputPin(InputPinLabel) : nullptr;
-
-		if (ensure(InputPin))
-		{
-			return ComputeInputPinDataDesc(InputPin, Binding);
-		}
-		else
-		{
-			return {};
-		}
-	}
-
-	FPCGDataCollectionDesc ComputeInputPinDataDesc(const UPCGPin* InputPin, const UPCGDataBinding* Binding)
-	{
-		check(InputPin);
-		check(Binding);
-
-		FPCGDataCollectionDesc PinDesc;
-
-		// Grab data from all incident edges.
-		for (const UPCGEdge* Edge : InputPin->Edges)
-		{
-			// InputPin is upstream output pin.
-			const UPCGPin* UpstreamOutputPin = Edge->InputPin;
-			if (!UpstreamOutputPin)
-			{
-				continue;
-			}
-
-			const UPCGSettings* UpstreamSettings = UpstreamOutputPin->Node ? UpstreamOutputPin->Node->GetSettings() : nullptr;
-			check(UpstreamSettings);
-
-			// Add data from connected upstream output pin.
-			FPCGDataCollectionDesc EdgeDesc;
-			if (ensure(Binding->GetOutputPinDataDesc_GameThread(UpstreamOutputPin, UpstreamSettings, EdgeDesc)))
-			{
-				PinDesc.Combine(EdgeDesc);
-			}
-		}
-
-		return PinDesc;
-	}
 }
 
 bool FPCGKernelAttributeKey::operator==(const FPCGKernelAttributeKey& Other) const
diff --git a/Engine/Plugins/PCG/Source/PCG/Private/Elements/Metadata/PCGMetadataPartition.cpp b/Engine/Plugins/PCG/Source/PCG/Private/Elements/Metadata/PCGMetadataPartition.cpp
--- a/Engine/Plugins/PCG/Source/PCG/Private/Elements/Metadata/PCGMetadataPartition.cpp
+++ b/Engine/Plugins/PCG/Source/PCG/Private/Elements/Metadata/PCGMetadataPartition.cpp
@@ -4,10 +4,7 @@
 
 #include "PCGContext.h"
 #include "PCGCustomVersion.h"
-#include "Compute/PCGComputeCommon.h"
-#include "Compute/PCGDataBinding.h"
-#include "Compute/Elements/PCGComputeGraphElement.h"
-#include "Compute/DataInterfaces/Elements/PCGMetadataPartitionDataInterface.h"
+#include "Elements/Metadata/PCGMetadataPartitionKernel.h"
 #include "Graph/PCGGPUGraphCompilationContext.h"
 #include "Helpers/PCGHelpers.h"
 #include "Metadata/PCGMetadata.h"
@@ -18,17 +15,10 @@
 #include "Metadata/Accessors/PCGAttributeAccessorKeys.h"
 #include "Utils/PCGLogErrors.h"
 
-#include "ShaderCompilerCore.h"
-
 #include UE_INLINE_GENERATED_CPP_BY_NAME(PCGMetadataPartition)
 
 #define LOCTEXT_NAMESPACE "PCGMetadataPartitionElement"
 
-namespace PCGMetadataPartitionSettings
-{
-	static TCHAR const* TemplateFilePath = TEXT("/Plugin/PCG/Private/Elements/PCGMetadataPartition.usf");
-}
-
 TArray<FPCGPinProperties> UPCGMetadataPartitionSettings::OutputPinProperties() const
 {
 	TArray<FPCGPinProperties> PinProperties;
@@ -65,176 +55,7 @@ void UPCGMetadataPartitionSettings::PostLoad()
 #endif // WITH_EDITOR
 }
 
-bool UPCGMetadataPartitionSettings::IsKernelValid(FPCGContext* InContext, bool bQuiet) const
-{
-	TRACE_CPUPROFILER_EVENT_SCOPE(UPCGMetadataPartitionSettings::IsKernelValid);
-	check(InContext);
-
-	if (!Super::IsKernelValid(InContext, bQuiet))
-	{
-		return false;
-	}
-
-	if (PartitionAttributeSelectors.Num() != 1)
-	{
-		PCG_KERNEL_VALIDATION_ERR(InContext, this, bQuiet, LOCTEXT("MustProvideOnePartitionAttribute", "GPU implementation currently only supports a single partition attribute."));
-		return false;
-	}
-
-	if (!PartitionAttributeSelectors[0].IsBasicAttribute())
-	{
-		PCG_KERNEL_VALIDATION_ERR(InContext, this, bQuiet, LOCTEXT("OnlyBasicAttributesSupported", "GPU implementation currently only supports basic attributes."));
-		return false;
-	}
-
-	FPCGComputeGraphContext* Context = InContext->IsComputeContext() ? static_cast<FPCGComputeGraphContext*>(InContext) : nullptr;
-
-	if (Context && Context->DataBinding)
-	{
-		const UPCGNode* Node = CastChecked<UPCGNode>(GetOuter());
-		const UPCGPin* InputPin = Node->GetInputPin(PCGPinConstants::DefaultInputLabel);
-
-		FName AttributeName = PartitionAttributeSelectors[0].GetAttributeName();
-
-		if (ensure(InputPin))
-		{
-			FPCGDataCollectionDesc CollectionDesc = PCGDataForGPUHelpers::ComputeInputPinDataDesc(InputPin, Context->DataBinding.Get());
-
-			FPCGKernelAttributeDesc AttributeDesc;
-			bool bConflictingTypesInData = false;
-			const bool bAttributePresent = CollectionDesc.GetAttributeDesc(AttributeName, AttributeDesc, bConflictingTypesInData);
-
-			if (!bAttributePresent)
-			{
-				PCG_KERNEL_VALIDATION_ERR(InContext, this, bQuiet, LOCTEXT("MustProvideOnePartitionAttribute", "GPU implementation currently only supports a single partition attribute."));
-				return false;
-			}
-
-			if (bConflictingTypesInData)
-			{
-				PCG_KERNEL_VALIDATION_ERR(InContext, this, bQuiet,
-					FText::Format(LOCTEXT("PartitionAttributeTypeConflict", "Attribute '{0}' encountered with multiple different types in input data."), FText::FromName(AttributeName)));
-				return false;
-			}
-
-			if (AttributeDesc.Type != EPCGKernelAttributeType::StringKey)
-			{
-				PCG_KERNEL_VALIDATION_ERR(Context, this, bQuiet, FText::Format(
-					LOCTEXT("PartitionAttributeTypeInvalid", "Attribute '{0}' not usable for partitioning, only attributes of type String Key are currently supported."),
-					FText::FromName(AttributeName)));
-				return false;
-			}
-		}
-	}
-
-	return true;
-}
-
-FString UPCGMetadataPartitionSettings::GetCookedKernelSource(const UPCGComputeGraph* InComputeGraph) const
-{
-	FString TemplateFile;
-	ensure(LoadShaderSourceFile(PCGMetadataPartitionSettings::TemplateFilePath, EShaderPlatform::SP_PCD3D_SM5, &TemplateFile, nullptr));
-
-	return TemplateFile;
-}
-
-void UPCGMetadataPartitionSettings::GetKernelAttributeKeys(TArray<FPCGKernelAttributeKey>& OutKeys) const
-{
-	// Register the attributes this node reads or writes for which we know the attribute type. Currently
-	// only StringKey attributes can be partitioned when executing on GPU, so declare the attribute here
-	// rather than doing work to resolve at runtime.
-	for (const FPCGAttributePropertyInputSelector& Selector : PartitionAttributeSelectors)
-	{
-		if (Selector.IsBasicAttribute())
-		{
-			OutKeys.Add(FPCGKernelAttributeKey(Selector.GetAttributeName(), EPCGKernelAttributeType::StringKey));
-		}
-	}
-
-	// Register output attribute created by this node.
-	if (bAssignIndexPartition)
-	{
-		OutKeys.AddUnique(FPCGKernelAttributeKey(PartitionIndexAttributeName, EPCGKernelAttributeType::Int));
-	}
-}
-
-int UPCGMetadataPartitionSettings::ComputeKernelThreadCount(const UPCGDataBinding* Binding) const
-{
-	const UPCGNode* Node = CastChecked<UPCGNode>(GetOuter());
-	const UPCGPin* Pin = Node->GetInputPin(PCGPinConstants::DefaultInputLabel);
-
-	return PCGDataForGPUHelpers::ComputeInputPinDataDesc(Pin, Binding).ComputeDataElementCount(EPCGDataType::Point);
-}
-
-bool UPCGMetadataPartitionSettings::ComputeOutputPinDataDesc(const UPCGPin* OutputPin, const UPCGDataBinding* Binding, FPCGDataCollectionDesc& OutDesc) const
-{
-	check(OutputPin);
-	check(Binding);
-
-	if (!ensure(OutputPin->Properties.Label == PCGPinConstants::DefaultOutputLabel))
-	{
-		return false;
-	}
-
-	// First inspect data coming from CPU. Will return description if this node is a CPU node or in a different compute graph.
-	if (Binding->ComputeCPUOutputPinDataDesc(OutputPin, OutDesc))
-	{
-		return true;
-	}
-
-	const FPCGDataCollectionDesc InputPinDesc = PCGDataForGPUHelpers::ComputeInputPinDataDesc(this, PCGPinConstants::DefaultInputLabel, Binding);
-
-	const int32 PartitionAttributeId = Binding->GetAttributeId(PartitionAttributeSelectors[0].GetAttributeName(), EPCGKernelAttributeType::StringKey);
-
-	FPCGDataCollectionDesc OutputPinDesc;
-
-	// Emit a point data for each potential partition.
-	if (PartitionAttributeId != INDEX_NONE)
-	{
-		TArray<int32> UniqueStringKeyValues;
-		InputPinDesc.GetUniqueStringKeyValues(PartitionAttributeId, UniqueStringKeyValues);
-
-		for (const FPCGDataDesc& InputData : InputPinDesc.DataDescs)
-		{
-			// Before execution we don't know what string key value each attribute has, so we have to assume worst case - could be up to N points in M partitions.
-			for (int32 StringKeyValue : UniqueStringKeyValues)
-			{
-				FPCGDataDesc& Desc = OutputPinDesc.DataDescs.Add_GetRef(InputData);
-
-				// Set the unique string keys to reflect the partition operation.
-				for (FPCGKernelAttributeDesc& AttrDesc : Desc.AttributeDescs)
-				{
-					if (AttrDesc.Index == PartitionAttributeId)
-					{
-						AttrDesc.SetStringKeys(MakeArrayView(&StringKeyValue, 1));
-						break;
-					}
-				}
-			}
-		}
-	}
-
-	// Add output attribute (partition index).
-	if (bAssignIndexPartition)
-	{
-		OutputPinDesc.AddAttributeToAllData(FPCGKernelAttributeKey(PartitionIndexAttributeName, EPCGKernelAttributeType::Int), Binding);
-	}
-
-	OutDesc = OutputPinDesc;
-
-	return true;
-}
-
 #if WITH_EDITOR
-void UPCGMetadataPartitionSettings::CreateAdditionalInputDataInterfaces(FPCGGPUCompilationContext& InOutContext, UObject* InObjectOuter, TArray<TObjectPtr<UComputeDataInterface>>& OutDataInterfaces) const
-{
-	Super::CreateAdditionalInputDataInterfaces(InOutContext, InObjectOuter, OutDataInterfaces);
-
-	TObjectPtr<UPCGMetadataPartitionDataInterface> NodeDI = InOutContext.NewObject_AnyThread<UPCGMetadataPartitionDataInterface>(InObjectOuter);
-	NodeDI->Settings = this;
-	OutDataInterfaces.Add(NodeDI);
-}
-
 void UPCGMetadataPartitionSettings::ApplyDeprecation(UPCGNode* InOutNode)
 {
 	PRAGMA_DISABLE_DEPRECATION_WARNINGS
@@ -246,6 +67,14 @@ void UPCGMetadataPartitionSettings::ApplyDeprecation(UPCGNode* InOutNode)
 
 	Super::ApplyDeprecation(InOutNode);
 }
+
+UPCGComputeKernel* UPCGMetadataPartitionSettings::CreateKernel(FPCGGPUCompilationContext& InOutContext, UObject* InObjectOuter) const
+{
+	UPCGMetadataPartitionKernel* Kernel = InOutContext.NewObject_AnyThread<UPCGMetadataPartitionKernel>(InObjectOuter);
+	Kernel->Settings = this;
+
+	return Kernel;
+}
 #endif // WITH_EDITOR
 
 FString UPCGMetadataPartitionSettings::GetAdditionalTitleInformation() const
diff --git a/Engine/Plugins/PCG/Source/PCG/Private/Elements/Metadata/PCGMetadataPartitionKernel.cpp b/Engine/Plugins/PCG/Source/PCG/Private/Elements/Metadata/PCGMetadataPartitionKernel.cpp
new file mode 100644
--- /dev/null
+++ b/Engine/Plugins/PCG/Source/PCG/Private/Elements/Metadata/PCGMetadataPartitionKernel.cpp
@@ -0,0 +1,184 @@
+// Copyright Epic Games, Inc. All Rights Reserved.
+
+#include "Elements/Metadata/PCGMetadataPartitionKernel.h"
+
+#include "PCGContext.h"
+#include "PCGModule.h"
+#include "Compute/PCGComputeGraph.h"
+#include "Compute/PCGDataBinding.h"
+#include "Compute/DataInterfaces/Elements/PCGMetadataPartitionDataInterface.h"
+#include "Compute/Elements/PCGComputeGraphElement.h"
+#include "Elements/Metadata/PCGMetadataPartition.h"
+#include "Graph/PCGGPUGraphCompilationContext.h"
+
+#include "ShaderCompilerCore.h"
+
+#include UE_INLINE_GENERATED_CPP_BY_NAME(PCGMetadataPartitionKernel)
+
+#define LOCTEXT_NAMESPACE "PCGMetadataPartitionKernel"
+
+bool UPCGMetadataPartitionKernel::IsKernelValid(FPCGContext* InContext, bool bQuiet) const
+{
+	TRACE_CPUPROFILER_EVENT_SCOPE(UPCGMetadataPartitionKernel::IsKernelValid);
+
+	if (!Super::IsKernelValid(InContext, bQuiet) || !ensure(Settings))
+	{
+		return false;
+	}
+
+	if (!ValidatePCGNode(InContext, bQuiet, Cast<UPCGSettings>(Settings)))
+	{
+		return false;
+	}
+
+	if (Settings->PartitionAttributeSelectors.Num() != 1)
+	{
+		PCG_KERNEL_VALIDATION_ERR(InContext, Settings, bQuiet, LOCTEXT("MustProvideOnePartitionAttribute", "GPU implementation currently only supports a single partition attribute."));
+		return false;
+	}
+
+	if (!Settings->PartitionAttributeSelectors[0].IsBasicAttribute())
+	{
+		PCG_KERNEL_VALIDATION_ERR(InContext, Settings, bQuiet, LOCTEXT("OnlyBasicAttributesSupported", "GPU implementation currently only supports basic attributes."));
+		return false;
+	}
+
+	FPCGComputeGraphContext* Context = InContext->IsComputeContext() ? static_cast<FPCGComputeGraphContext*>(InContext) : nullptr;
+
+	if (Context && Context->DataBinding)
+	{
+		FName AttributeName = Settings->PartitionAttributeSelectors[0].GetAttributeName();
+
+		FPCGDataCollectionDesc InputDesc;
+		ensure(Context->DataBinding->GetKernelPinDataDesc_GameThread(this, PCGPinConstants::DefaultInputLabel, /*bIsInput=*/true, InputDesc));
+
+		FPCGKernelAttributeDesc AttributeDesc;
+		bool bConflictingTypesInData = false;
+		const bool bAttributePresent = InputDesc.GetAttributeDesc(AttributeName, AttributeDesc, bConflictingTypesInData);
+
+		if (!bAttributePresent)
+		{
+			PCG_KERNEL_VALIDATION_ERR(InContext, Settings, bQuiet, LOCTEXT("MustProvideOnePartitionAttribute", "GPU implementation currently only supports a single partition attribute."));
+			return false;
+		}
+
+		if (bConflictingTypesInData)
+		{
+			PCG_KERNEL_VALIDATION_ERR(InContext, Settings, bQuiet,
+				FText::Format(LOCTEXT("PartitionAttributeTypeConflict", "Attribute '{0}' encountered with multiple different types in input data."), FText::FromName(AttributeName)));
+			return false;
+		}
+
+		if (AttributeDesc.Type != EPCGKernelAttributeType::StringKey)
+		{
+			PCG_KERNEL_VALIDATION_ERR(Context, Settings, bQuiet, FText::Format(
+				LOCTEXT("PartitionAttributeTypeInvalid", "Attribute '{0}' not usable for partitioning, only attributes of type String Key are currently supported."),
+				FText::FromName(AttributeName)));
+			return false;
+		}
+	}
+
+	return true;
+}
+
+bool UPCGMetadataPartitionKernel::ComputeOutputBindingDataDesc(const UPCGComputeGraph* InGraph, FName InOutputPinLabel, const UPCGDataBinding* InBinding, FPCGDataCollectionDesc& OutDataDesc) const
+{
+	check(InGraph);
+	check(InBinding);
+
+	// Code assumes single output pin.
+	if (!ensure(InOutputPinLabel == PCGPinConstants::DefaultOutputLabel))
+	{
+		return false;
+	}
+
+	FPCGDataCollectionDesc InputPinDesc;
+	ensure(InBinding->GetKernelPinDataDesc_GameThread(this, PCGPinConstants::DefaultInputLabel, /*bIsInput=*/true, InputPinDesc));
+
+	const int32 PartitionAttributeId = InBinding->GetAttributeId(Settings->PartitionAttributeSelectors[0].GetAttributeName(), EPCGKernelAttributeType::StringKey);
+
+	FPCGDataCollectionDesc PinDesc;
+
+	// Emit a point data for each potential partition.
+	if (PartitionAttributeId != INDEX_NONE)
+	{
+		TArray<int32> UniqueStringKeyValues;
+		InputPinDesc.GetUniqueStringKeyValues(PartitionAttributeId, UniqueStringKeyValues);
+
+		for (const FPCGDataDesc& InputData : InputPinDesc.DataDescs)
+		{
+			// Before execution we don't know what string key value each attribute has, so we have to assume worst case - could be up to N points in M partitions.
+			for (int32 StringKeyValue : UniqueStringKeyValues)
+			{
+				FPCGDataDesc& Desc = PinDesc.DataDescs.Add_GetRef(InputData);
+
+				// Set the unique string keys to reflect the partition operation.
+				for (FPCGKernelAttributeDesc& AttrDesc : Desc.AttributeDescs)
+				{
+					if (AttrDesc.Index == PartitionAttributeId)
+					{
+						AttrDesc.SetStringKeys(MakeArrayView(&StringKeyValue, 1));
+						break;
+					}
+				}
+			}
+		}
+	}
+
+	// Add output attribute (partition index).
+	if (Settings->bAssignIndexPartition)
+	{
+		PinDesc.AddAttributeToAllData(FPCGKernelAttributeKey(Settings->PartitionIndexAttributeName, EPCGKernelAttributeType::Int), InBinding);
+	}
+
+	OutDataDesc = MoveTemp(PinDesc);
+	return true;
+}
+
+int UPCGMetadataPartitionKernel::ComputeThreadCount(const UPCGDataBinding* InBinding) const
+{
+	FPCGDataCollectionDesc InputPinDesc;
+	ensure(InBinding->GetKernelPinDataDesc_GameThread(this, PCGPinConstants::DefaultInputLabel, /*bIsInput=*/true, InputPinDesc));
+
+	return InputPinDesc.ComputeDataElementCount(EPCGDataType::Point);
+}
+
+#if WITH_EDITOR
+FString UPCGMetadataPartitionKernel::GetCookedSource(const UPCGComputeGraph* InComputeGraph) const
+{
+	FString TemplateFile;
+	ensure(LoadShaderSourceFile(TEXT("/Plugin/PCG/Private/Elements/PCGMetadataPartition.usf"), EShaderPlatform::SP_PCD3D_SM5, &TemplateFile, nullptr));
+	return TemplateFile;
+}
+
+void UPCGMetadataPartitionKernel::CreateAdditionalInputDataInterfaces(FPCGGPUCompilationContext& InOutContext, UObject* InObjectOuter, TArray<TObjectPtr<UComputeDataInterface>>& OutDataInterfaces) const
+{
+	Super::CreateAdditionalInputDataInterfaces(InOutContext, InObjectOuter, OutDataInterfaces);
+
+	TObjectPtr<UPCGMetadataPartitionDataInterface> NodeDI = InOutContext.NewObject_AnyThread<UPCGMetadataPartitionDataInterface>(InObjectOuter);
+	NodeDI->ProducerKernel = this;
+	OutDataInterfaces.Add(NodeDI);
+}
+#endif
+
+void UPCGMetadataPartitionKernel::GetKernelAttributeKeys(TArray<FPCGKernelAttributeKey>& OutKeys) const
+{
+	// Register the attributes this node reads or writes for which we know the attribute type. Currently
+	// only StringKey attributes can be partitioned when executing on GPU, so declare the attribute here
+	// rather than doing work to resolve at runtime.
+	for (const FPCGAttributePropertyInputSelector& Selector : Settings->PartitionAttributeSelectors)
+	{
+		if (Selector.IsBasicAttribute())
+		{
+			OutKeys.Add(FPCGKernelAttributeKey(Selector.GetAttributeName(), EPCGKernelAttributeType::StringKey));
+		}
+	}
+
+	// Register output attribute created by this node.
+	if (Settings->bAssignIndexPartition)
+	{
+		OutKeys.AddUnique(FPCGKernelAttributeKey(Settings->PartitionIndexAttributeName, EPCGKernelAttributeType::Int));
+	}
+}
+
+#undef LOCTEXT_NAMESPACE
diff --git a/Engine/Plugins/PCG/Source/PCG/Private/Elements/PCGCopyPoints.cpp b/Engine/Plugins/PCG/Source/PCG/Private/Elements/PCGCopyPoints.cpp
--- a/Engine/Plugins/PCG/Source/PCG/Private/Elements/PCGCopyPoints.cpp
+++ b/Engine/Plugins/PCG/Source/PCG/Private/Elements/PCGCopyPoints.cpp
@@ -3,17 +3,13 @@
 #include "Elements/PCGCopyPoints.h"
 
 #include "PCGContext.h"
-#include "Compute/PCGComputeCommon.h"
-#include "Compute/PCGDataBinding.h"
-#include "Compute/DataInterfaces/Elements/PCGCopyPointsDataInterface.h"
 #include "Data/PCGPointData.h"
 #include "Data/PCGSpatialData.h"
+#include "Elements/PCGCopyPointsKernel.h"
 #include "Graph/PCGGPUGraphCompilationContext.h"
 #include "Helpers/PCGAsync.h"
 #include "Helpers/PCGHelpers.h"
 
-#include "RHIShaderPlatform.h"
-#include "ShaderCompilerCore.h"
 #include "Async/ParallelFor.h"
 
 #include UE_INLINE_GENERATED_CPP_BY_NAME(PCGCopyPoints)
@@ -25,114 +21,13 @@ FText UPCGCopyPointsSettings::GetNodeTooltipText() const
 {
 	return LOCTEXT("NodeTooltip", "For each point pair from the source and the target, create a copy, inheriting properties & attributes depending on the node settings.");
 }
-#endif
-
-FString UPCGCopyPointsSettings::GetCookedKernelSource(const UPCGComputeGraph* InComputeGraph) const
-{
-	FString TemplateFile;
-	ensure(LoadShaderSourceFile(TEXT("/Plugin/PCG/Private/Elements/PCGCopyPoints.usf"), EShaderPlatform::SP_PCD3D_SM5, &TemplateFile, nullptr));
-	return TemplateFile;
-}
-
-bool UPCGCopyPointsSettings::ComputeOutputPinDataDesc(const UPCGPin* OutputPin, const UPCGDataBinding* Binding, FPCGDataCollectionDesc& OutDesc) const
-{
-	check(OutputPin);
-	check(Binding);
-
-	// First inspect data coming from CPU. Will return description if this node is a CPU node or in a different compute graph.
-	if (Binding->ComputeCPUOutputPinDataDesc(OutputPin, OutDesc))
-	{
-		return true;
-	}
-
-	const UPCGNode* Node = CastChecked<UPCGNode>(GetOuter());
-	const UPCGPin* SourcePin = Node->GetInputPin(PCGCopyPointsConstants::SourcePointsLabel);
-	const UPCGPin* TargetPin = Node->GetInputPin(PCGCopyPointsConstants::TargetPointsLabel);
-
-	FPCGDataCollectionDesc PinDesc;
-
-	if (ensure(SourcePin && TargetPin))
-	{
-		const FPCGDataCollectionDesc SourcePinDesc = PCGDataForGPUHelpers::ComputeInputPinDataDesc(SourcePin, Binding);
-		const FPCGDataCollectionDesc TargetPinDesc = PCGDataForGPUHelpers::ComputeInputPinDataDesc(TargetPin, Binding);
-
-		const int32 NumSources = SourcePinDesc.DataDescs.Num();
-		const int32 NumTargets = TargetPinDesc.DataDescs.Num();
-		const int32 NumIterations = bCopyEachSourceOnEveryTarget ? NumSources * NumTargets : FMath::Max(NumSources, NumTargets);
-
-		if (NumSources > 0 && NumTargets > 0 && (bCopyEachSourceOnEveryTarget || NumSources == NumTargets || NumSources == 1 || NumTargets == 1))
-		{
-			for (int32 I = 0; I < NumIterations; ++I)
-			{
-				const int32 SourceIndex = bCopyEachSourceOnEveryTarget ? (I / NumTargets) : FMath::Min(I, NumSources - 1);
-				const int32 TargetIndex = bCopyEachSourceOnEveryTarget ? (I % NumTargets) : FMath::Min(I, NumTargets - 1);
-
-				const FPCGDataDesc& SourceDesc = SourcePinDesc.DataDescs[SourceIndex];
-				const FPCGDataDesc& TargetDesc = TargetPinDesc.DataDescs[TargetIndex];
-
-				FPCGDataDesc& ResultDataDesc = PinDesc.DataDescs.Emplace_GetRef(EPCGDataType::Point, SourceDesc.ElementCount * TargetDesc.ElementCount);
-
-				if (AttributeInheritance == EPCGCopyPointsMetadataInheritanceMode::SourceFirst)
-				{
-					ResultDataDesc.AttributeDescs = SourceDesc.AttributeDescs;
-
-					for (const FPCGKernelAttributeDesc& AttrDesc : TargetDesc.AttributeDescs)
-					{
-						ResultDataDesc.AttributeDescs.AddUnique(AttrDesc);
-					}
-				}
-				else if (AttributeInheritance == EPCGCopyPointsMetadataInheritanceMode::TargetFirst)
-				{
-					ResultDataDesc.AttributeDescs = TargetDesc.AttributeDescs;
-
-					for (const FPCGKernelAttributeDesc& AttrDesc : SourceDesc.AttributeDescs)
-					{
-						ResultDataDesc.AttributeDescs.AddUnique(AttrDesc);
-					}
-				}
-				else if (AttributeInheritance == EPCGCopyPointsMetadataInheritanceMode::SourceOnly)
-				{
-					ResultDataDesc.AttributeDescs = SourceDesc.AttributeDescs;
-				}
-				else if (AttributeInheritance == EPCGCopyPointsMetadataInheritanceMode::TargetOnly)
-				{
-					ResultDataDesc.AttributeDescs = TargetDesc.AttributeDescs;
-				}
-			}
-		}
-	}
-
-	OutDesc = PinDesc;
-	return true;
-}
-
-int UPCGCopyPointsSettings::ComputeKernelThreadCount(const UPCGDataBinding* Binding) const
-{
-	const UPCGNode* Node = CastChecked<UPCGNode>(GetOuter());
-	const UPCGPin* OutPin = Node->GetOutputPin(PCGPinConstants::DefaultOutputLabel);
 
-	int ThreadCount = 0;
-
-	FPCGDataCollectionDesc OutputPinDesc;
-	if (ensure(Binding->GetOutputPinDataDesc_GameThread(OutPin, this, OutputPinDesc)))
-	{
-		for (const FPCGDataDesc& DataDesc : OutputPinDesc.DataDescs)
-		{
-			ThreadCount += DataDesc.ElementCount;
-		}
-	}
-
-	return ThreadCount;
-}
-
-#if WITH_EDITOR
-void UPCGCopyPointsSettings::CreateAdditionalInputDataInterfaces(FPCGGPUCompilationContext& InOutContext, UObject* InObjectOuter, TArray<TObjectPtr<UComputeDataInterface>>& OutDataInterfaces) const
+UPCGComputeKernel* UPCGCopyPointsSettings::CreateKernel(FPCGGPUCompilationContext& InOutContext, UObject* InObjectOuter) const
 {
-	Super::CreateAdditionalInputDataInterfaces(InOutContext, InObjectOuter, OutDataInterfaces);
+	UPCGCopyPointsKernel* Kernel = InOutContext.NewObject_AnyThread<UPCGCopyPointsKernel>(InObjectOuter);
+	Kernel->Settings = this;
 
-	TObjectPtr<UPCGCopyPointsDataInterface> NodeDI = InOutContext.NewObject_AnyThread<UPCGCopyPointsDataInterface>(InObjectOuter);
-	NodeDI->Settings = this;
-	OutDataInterfaces.Add(NodeDI);
+	return Kernel;
 }
 #endif
 
diff --git a/Engine/Plugins/PCG/Source/PCG/Private/Elements/PCGCopyPointsKernel.cpp b/Engine/Plugins/PCG/Source/PCG/Private/Elements/PCGCopyPointsKernel.cpp
new file mode 100644
--- /dev/null
+++ b/Engine/Plugins/PCG/Source/PCG/Private/Elements/PCGCopyPointsKernel.cpp
@@ -0,0 +1,135 @@
+// Copyright Epic Games, Inc. All Rights Reserved.
+
+#include "Elements/PCGCopyPointsKernel.h"
+
+#include "PCGContext.h"
+#include "PCGModule.h"
+#include "Compute/PCGComputeGraph.h"
+#include "Compute/PCGDataBinding.h"
+#include "Compute/DataInterfaces/Elements/PCGCopyPointsDataInterface.h"
+#include "Elements/PCGCopyPoints.h"
+#include "Graph/PCGGPUGraphCompilationContext.h"
+
+#include "ShaderCompilerCore.h"
+#include "Containers/StaticArray.h"
+#include "Internationalization/Regex.h"
+
+#include UE_INLINE_GENERATED_CPP_BY_NAME(PCGCopyPointsKernel)
+
+#define LOCTEXT_NAMESPACE "PCGCopyPointsKernel"
+
+bool UPCGCopyPointsKernel::IsKernelValid(FPCGContext* InContext, bool bQuiet) const
+{
+	TRACE_CPUPROFILER_EVENT_SCOPE(UPCGCopyPointsKernel::IsKernelValid);
+
+	if (!Super::IsKernelValid(InContext, bQuiet) || !ensure(Settings))
+	{
+		return false;
+	}
+
+	return ValidatePCGNode(InContext, bQuiet, Cast<UPCGSettings>(Settings));
+}
+
+bool UPCGCopyPointsKernel::ComputeOutputBindingDataDesc(const UPCGComputeGraph* InGraph, FName InOutputPinLabel, const UPCGDataBinding* InBinding, FPCGDataCollectionDesc& OutDataDesc) const
+{
+	check(InGraph);
+	check(InBinding);
+
+	// Code assumes single output pin.
+	if (!ensure(InOutputPinLabel == PCGPinConstants::DefaultOutputLabel))
+	{
+		return false;
+	}
+
+	FPCGDataCollectionDesc SourcePinDesc;
+	FPCGDataCollectionDesc TargetPinDesc;
+	ensure(InBinding->GetKernelPinDataDesc_GameThread(this, PCGCopyPointsConstants::SourcePointsLabel, /*bIsInput=*/true, SourcePinDesc));
+	ensure(InBinding->GetKernelPinDataDesc_GameThread(this, PCGCopyPointsConstants::TargetPointsLabel, /*bIsInput=*/true, TargetPinDesc));
+
+	FPCGDataCollectionDesc PinDesc;
+
+	const int32 NumSources = SourcePinDesc.DataDescs.Num();
+	const int32 NumTargets = TargetPinDesc.DataDescs.Num();
+	const int32 NumIterations = Settings->bCopyEachSourceOnEveryTarget ? NumSources * NumTargets : FMath::Max(NumSources, NumTargets);
+
+	if (NumSources > 0 && NumTargets > 0 && (Settings->bCopyEachSourceOnEveryTarget || NumSources == NumTargets || NumSources == 1 || NumTargets == 1))
+	{
+		for (int32 I = 0; I < NumIterations; ++I)
+		{
+			const int32 SourceIndex = Settings->bCopyEachSourceOnEveryTarget ? (I / NumTargets) : FMath::Min(I, NumSources - 1);
+			const int32 TargetIndex = Settings->bCopyEachSourceOnEveryTarget ? (I % NumTargets) : FMath::Min(I, NumTargets - 1);
+
+			const FPCGDataDesc& SourceDesc = SourcePinDesc.DataDescs[SourceIndex];
+			const FPCGDataDesc& TargetDesc = TargetPinDesc.DataDescs[TargetIndex];
+
+			FPCGDataDesc& ResultDataDesc = PinDesc.DataDescs.Emplace_GetRef(EPCGDataType::Point, SourceDesc.ElementCount * TargetDesc.ElementCount);
+
+			if (Settings->AttributeInheritance == EPCGCopyPointsMetadataInheritanceMode::SourceFirst)
+			{
+				ResultDataDesc.AttributeDescs = SourceDesc.AttributeDescs;
+
+				for (const FPCGKernelAttributeDesc& AttrDesc : TargetDesc.AttributeDescs)
+				{
+					ResultDataDesc.AttributeDescs.AddUnique(AttrDesc);
+				}
+			}
+			else if (Settings->AttributeInheritance == EPCGCopyPointsMetadataInheritanceMode::TargetFirst)
+			{
+				ResultDataDesc.AttributeDescs = TargetDesc.AttributeDescs;
+
+				for (const FPCGKernelAttributeDesc& AttrDesc : SourceDesc.AttributeDescs)
+				{
+					ResultDataDesc.AttributeDescs.AddUnique(AttrDesc);
+				}
+			}
+			else if (Settings->AttributeInheritance == EPCGCopyPointsMetadataInheritanceMode::SourceOnly)
+			{
+				ResultDataDesc.AttributeDescs = SourceDesc.AttributeDescs;
+			}
+			else if (Settings->AttributeInheritance == EPCGCopyPointsMetadataInheritanceMode::TargetOnly)
+			{
+				ResultDataDesc.AttributeDescs = TargetDesc.AttributeDescs;
+			}
+		}
+	}
+
+	OutDataDesc = MoveTemp(PinDesc);
+	return true;
+}
+
+int UPCGCopyPointsKernel::ComputeThreadCount(const UPCGDataBinding* InBinding) const
+{
+	int ThreadCount = 0;
+
+	FPCGDataCollectionDesc OutputPinDesc;
+	if (ensure(InBinding->GetKernelPinDataDesc_GameThread(this, PCGPinConstants::DefaultOutputLabel, /*bIsInput=*/false, OutputPinDesc)))
+	{
+		for (const FPCGDataDesc& DataDesc : OutputPinDesc.DataDescs)
+		{
+			ThreadCount += DataDesc.ElementCount;
+		}
+	}
+
+	return ThreadCount;
+}
+
+#if WITH_EDITOR
+FString UPCGCopyPointsKernel::GetCookedSource(const UPCGComputeGraph* InComputeGraph) const
+{
+	FString TemplateFile;
+	ensure(LoadShaderSourceFile(TEXT("/Plugin/PCG/Private/Elements/PCGCopyPoints.usf"), EShaderPlatform::SP_PCD3D_SM5, &TemplateFile, nullptr));
+	return TemplateFile;
+}
+
+void UPCGCopyPointsKernel::CreateAdditionalInputDataInterfaces(FPCGGPUCompilationContext& InOutContext, UObject* InObjectOuter, TArray<TObjectPtr<UComputeDataInterface>>& OutDataInterfaces) const
+{
+	Super::CreateAdditionalInputDataInterfaces(InOutContext, InObjectOuter, OutDataInterfaces);
+
+	TObjectPtr<UPCGCopyPointsDataInterface> NodeDI = InOutContext.NewObject_AnyThread<UPCGCopyPointsDataInterface>(InObjectOuter);
+	NodeDI->ProducerKernel = this;
+
+	OutDataInterfaces.Add(NodeDI);
+}
+#endif
+
+#undef LOCTEXT_NAMESPACE
diff --git a/Engine/Plugins/PCG/Source/PCG/Private/Elements/PCGReroute.cpp b/Engine/Plugins/PCG/Source/PCG/Private/Elements/PCGReroute.cpp
--- a/Engine/Plugins/PCG/Source/PCG/Private/Elements/PCGReroute.cpp
+++ b/Engine/Plugins/PCG/Source/PCG/Private/Elements/PCGReroute.cpp
@@ -87,22 +87,6 @@ FPCGElementPtr UPCGRerouteSettings::CreateElement() const
 	return MakeShared<FPCGRerouteElement>();
 }
 
-bool UPCGRerouteSettings::ComputeOutputPinDataDesc(const UPCGPin* OutputPin, const UPCGDataBinding* Binding, FPCGDataCollectionDesc& OutDesc) const
-{
-	// Reroutes always trivially forward data. Reroutes are culled and not visible to compute graph compilation and therefore we
-	// transparently forward the data here.
-	if (const UPCGNode* Node = Cast<UPCGNode>(GetOuter()))
-	{
-		if (const UPCGPin* Pin = Node->GetInputPin(PCGPinConstants::DefaultInputLabel))
-		{
-			OutDesc = PCGDataForGPUHelpers::ComputeInputPinDataDesc(Pin, Binding);
-			return true;
-		}
-	}
-
-	return false;
-}
-
 #if WITH_EDITOR
 TArray<FPCGPreconfiguredInfo> UPCGRerouteSettings::GetConversionInfo() const
 {
diff --git a/Engine/Plugins/PCG/Source/PCG/Private/Elements/PCGSkinnedMeshSpawner.cpp b/Engine/Plugins/PCG/Source/PCG/Private/Elements/PCGSkinnedMeshSpawner.cpp
--- a/Engine/Plugins/PCG/Source/PCG/Private/Elements/PCGSkinnedMeshSpawner.cpp
+++ b/Engine/Plugins/PCG/Source/PCG/Private/Elements/PCGSkinnedMeshSpawner.cpp
@@ -5,14 +5,9 @@
 #include "PCGComponent.h"
 #include "PCGCustomVersion.h"
 #include "PCGManagedResource.h"
-#include "Compute/PCGComputeCommon.h"
-#include "Compute/PCGDataBinding.h"
-#include "Compute/DataInterfaces/PCGInstanceDataInterface.h"
-#include "Compute/DataInterfaces/Elements/PCGSkinnedMeshSpawnerDataInterface.h"
 #include "Data/PCGPointData.h"
 #include "Data/PCGSpatialData.h"
 #include "Elements/PCGSkinnedMeshSpawnerContext.h"
-#include "Graph/PCGGPUGraphCompilationContext.h"
 #include "Helpers/PCGActorHelpers.h"
 #include "Helpers/PCGHelpers.h"
 #include "InstanceDataPackers/PCGSkinnedMeshInstanceDataPackerBase.h"
@@ -21,7 +16,6 @@
 #include "Components/InstancedSkinnedMeshComponent.h"
 #include "Engine/SkinnedAsset.h"
 #include "Materials/MaterialInterface.h"
-#include "ShaderCompilerCore.h"
 #include "UObject/Package.h"
 
 #include UE_INLINE_GENERATED_CPP_BY_NAME(PCGSkinnedMeshSpawner)
@@ -39,118 +33,7 @@ UPCGSkinnedMeshSpawnerSettings::UPCGSkinnedMeshSpawnerSettings(const FObjectInit
 	}
 }
 
-TCHAR const* UPCGSkinnedMeshSpawnerSettings::TemplateFilePath = TEXT("/Plugin/PCG/Private/Elements/PCGSkinnedMeshSpawner.usf");
-
-bool UPCGSkinnedMeshSpawnerSettings::IsKernelValid(FPCGContext* InContext, bool bQuiet) const
-{
-	TRACE_CPUPROFILER_EVENT_SCOPE(UPCGSkinnedMeshSpawnerSettings::IsKernelValid);
-
-	if (!Super::IsKernelValid(InContext, bQuiet))
-	{
-		return false;
-	}
-
-	if (MeshSelectorParameters == nullptr)
-	{
-		return false;
-	}
-
-	if (!MeshSelectorParameters->BankAttribute.IsValid())
-	{
-		PCG_KERNEL_VALIDATION_ERR(InContext, this, bQuiet, LOCTEXT("NoBankAttribute", "Mesh selector bank attribute is invalid."));
-		return false;
-	}
-
-	if (!MeshSelectorParameters->SequenceIndexAttribute.IsValid())
-	{
-		PCG_KERNEL_VALIDATION_ERR(InContext, this, bQuiet, LOCTEXT("NoSequenceIndexAttribute", "Mesh selector sequence index attribute is invalid."));
-		return false;
-	}
-
-	if (!MeshSelectorParameters->MeshAttribute.IsValid())
-	{
-		PCG_KERNEL_VALIDATION_ERR(InContext, this, bQuiet, LOCTEXT("NoMeshAttribute", "Mesh selector mesh attribute is invalid."));
-		return false;
-	}
-
-	// Currently instance packers must be able to specify a full list of attribute names upfront, to build the attribute table at compile time.
-	// TODO: We should be able to augment a static attribute table with new attributes at execution time, which will allow other types like regex.
-	if (InstanceDataPackerParameters && !InstanceDataPackerParameters->GetAttributeNames(/*OutNames=*/nullptr))
-	{
-		PCG_KERNEL_VALIDATION_ERR(InContext, this, bQuiet, LOCTEXT("InvalidInstancePacker", "Selected instance packer does not support GPU execution."));
-		return false;
-	}
-
-	return true;
-}
-
-FString UPCGSkinnedMeshSpawnerSettings::GetCookedKernelSource(const UPCGComputeGraph* InComputeGraph) const
-{
-	FString TemplateFile;
-	ensure(LoadShaderSourceFile(TemplateFilePath, EShaderPlatform::SP_PCD3D_SM5, &TemplateFile, nullptr));
-	return TemplateFile;
-}
-
-bool UPCGSkinnedMeshSpawnerSettings::ComputeOutputPinDataDesc(const UPCGPin* OutputPin, const UPCGDataBinding* Binding, FPCGDataCollectionDesc& OutDesc) const
-{
-	check(OutputPin);
-	check(Binding);
-
-	// Code assumes single output pin.
-	if (!ensure(OutputPin->Properties.Label == PCGPinConstants::DefaultOutputLabel))
-	{
-		return false;
-	}
-
-	check(MeshSelectorParameters);
-
-	// First inspect data coming from CPU. Will return description if this node is a CPU node or in a different compute graph.
-	if (Binding->ComputeCPUOutputPinDataDesc(OutputPin, OutDesc))
-	{
-		return true;
-	}
-
-	const FPCGDataCollectionDesc InputPinDesc = PCGDataForGPUHelpers::ComputeInputPinDataDesc(this, PCGPinConstants::DefaultInputLabel, Binding);
-	
-	// First forward input data to output.
-	FPCGDataCollectionDesc OutputPinDesc = InputPinDesc;
-
-	OutDesc = OutputPinDesc;
-	return true;
-}
-
-void UPCGSkinnedMeshSpawnerSettings::GetKernelAttributeKeys(TArray<FPCGKernelAttributeKey>& OutKeys) const
-{
-}
-
-int UPCGSkinnedMeshSpawnerSettings::ComputeKernelThreadCount(const UPCGDataBinding* Binding) const
-{
-	const UPCGNode* Node = CastChecked<UPCGNode>(GetOuter());
-	const UPCGPin* Pin = Node->GetInputPin(PCGPinConstants::DefaultInputLabel);
-
-	return PCGDataForGPUHelpers::ComputeInputPinDataDesc(Pin, Binding).ComputeDataElementCount(EPCGDataType::Point);
-}
-
 #if WITH_EDITOR
-void UPCGSkinnedMeshSpawnerSettings::CreateAdditionalInputDataInterfaces(FPCGGPUCompilationContext& InOutContext, UObject* InObjectOuter, TArray<TObjectPtr<UComputeDataInterface>>& OutDataInterfaces) const
-{
-	Super::CreateAdditionalInputDataInterfaces(InOutContext, InObjectOuter, OutDataInterfaces);
-
-	TObjectPtr<UPCGSkinnedMeshSpawnerDataInterface> NodeDI = InOutContext.NewObject_AnyThread<UPCGSkinnedMeshSpawnerDataInterface>(InObjectOuter);
-	NodeDI->Settings = this;
-	OutDataInterfaces.Add(NodeDI);
-}
-
-void UPCGSkinnedMeshSpawnerSettings::CreateAdditionalOutputDataInterfaces(FPCGGPUCompilationContext& InOutContext, UObject* InObjectOuter, TArray<TObjectPtr<UComputeDataInterface>>& OutDataInterfaces) const
-{
-	Super::CreateAdditionalOutputDataInterfaces(InOutContext, InObjectOuter, OutDataInterfaces);
-
-	UPCGInstanceDataInterface* InstanceDI = InOutContext.NewObject_AnyThread<UPCGInstanceDataInterface>(InObjectOuter);
-	InstanceDI->SetProducerSettings(this);
-	InstanceDI->InputPinProvidingData = PCGPinConstants::DefaultInputLabel;
-	OutDataInterfaces.Add(InstanceDI);
-}
-
 FText UPCGSkinnedMeshSpawnerSettings::GetDefaultNodeTitle() const
 {
 	return LOCTEXT("NodeTitle", "Instanced Skinned Mesh Spawner");
diff --git a/Engine/Plugins/PCG/Source/PCG/Private/Elements/PCGStaticMeshSpawner.cpp b/Engine/Plugins/PCG/Source/PCG/Private/Elements/PCGStaticMeshSpawner.cpp
--- a/Engine/Plugins/PCG/Source/PCG/Private/Elements/PCGStaticMeshSpawner.cpp
+++ b/Engine/Plugins/PCG/Source/PCG/Private/Elements/PCGStaticMeshSpawner.cpp
@@ -5,13 +5,10 @@
 #include "PCGComponent.h"
 #include "PCGCustomVersion.h"
 #include "PCGManagedResource.h"
-#include "Compute/PCGComputeCommon.h"
-#include "Compute/PCGDataBinding.h"
-#include "Compute/DataInterfaces/PCGInstanceDataInterface.h"
-#include "Compute/DataInterfaces/Elements/PCGStaticMeshSpawnerDataInterface.h"
 #include "Data/PCGPointData.h"
 #include "Data/PCGSpatialData.h"
 #include "Elements/PCGStaticMeshSpawnerContext.h"
+#include "Elements/PCGStaticMeshSpawnerKernel.h"
 #include "Graph/PCGGPUGraphCompilationContext.h"
 #include "Helpers/PCGActorHelpers.h"
 #include "Helpers/PCGHelpers.h"
@@ -24,7 +21,6 @@
 #include "Components/InstancedStaticMeshComponent.h"
 #include "Engine/StaticMesh.h"
 #include "Materials/MaterialInterface.h"
-#include "ShaderCompilerCore.h"
 #include "UObject/Package.h"
 
 #include UE_INLINE_GENERATED_CPP_BY_NAME(PCGStaticMeshSpawner)
@@ -48,177 +44,13 @@ UPCGStaticMeshSpawnerSettings::UPCGStaticMeshSpawnerSettings(const FObjectInitia
 	}
 }
 
-TCHAR const* UPCGStaticMeshSpawnerSettings::TemplateFilePath = TEXT("/Plugin/PCG/Private/Elements/PCGStaticMeshSpawner.usf");
-
-bool UPCGStaticMeshSpawnerSettings::IsKernelValid(FPCGContext* InContext, bool bQuiet) const
-{
-	TRACE_CPUPROFILER_EVENT_SCOPE(UPCGStaticMeshSpawnerSettings::IsKernelValid);
-
-	if (!Super::IsKernelValid(InContext, bQuiet))
-	{
-		return false;
-	}
-
-	if (UPCGMeshSelectorWeighted* SelectorWeighted = Cast<UPCGMeshSelectorWeighted>(MeshSelectorParameters))
-	{
-		if (SelectorWeighted->MeshEntries.IsEmpty())
-		{
-			PCG_KERNEL_VALIDATION_ERR(InContext, this, bQuiet, LOCTEXT("NoMeshEntries", "No meshes specified in selector."));
-			return false;
-		}
-
-		for (const FPCGMeshSelectorWeightedEntry& Entry : SelectorWeighted->MeshEntries)
-		{
-			if (Entry.Descriptor.StaticMesh.IsNull())
-			{
-				PCG_KERNEL_VALIDATION_ERR(InContext, this, bQuiet, LOCTEXT("UnassignedMesh", "Unassigned mesh."));
-				return false;
-			}
-		}
-	}
-	else if (!MeshSelectorParameters || !MeshSelectorParameters->IsA<UPCGMeshSelectorByAttribute>())
-	{
-		PCG_KERNEL_VALIDATION_ERR(InContext, this, bQuiet, LOCTEXT("InvalidMeshSelector", "Currently GPU Static Mesh Spawner nodes must use PCGMeshSelectorWeighted or UPCGMeshSelectorByAttribute as the mesh selector type."));
-		return false;
-	}
-
-	// Currently instance packers must be able to specify a full list of attribute names upfront, to build the attribute table at compile time.
-	// TODO: We should be able to augment a static attribute table with new attributes at execution time, which will allow other types like regex.
-	if (InstanceDataPackerParameters && !InstanceDataPackerParameters->GetAttributeNames(/*OutNames=*/nullptr))
-	{
-		PCG_KERNEL_VALIDATION_ERR(InContext, this, bQuiet, LOCTEXT("InvalidInstancePacker", "Selected instance packer does not support GPU execution."));
-		return false;
-	}
-
-	return true;
-}
-
-FString UPCGStaticMeshSpawnerSettings::GetCookedKernelSource(const UPCGComputeGraph* InComputeGraph) const
-{
-	FString TemplateFile;
-	ensure(LoadShaderSourceFile(TemplateFilePath, EShaderPlatform::SP_PCD3D_SM5, &TemplateFile, nullptr));
-
-	return TemplateFile;
-}
-
-bool UPCGStaticMeshSpawnerSettings::ComputeOutputPinDataDesc(const UPCGPin* OutputPin, const UPCGDataBinding* Binding, FPCGDataCollectionDesc& OutDesc) const
-{
-	check(OutputPin);
-	check(Binding);
-
-	// Code assumes single output pin.
-	if (!ensure(OutputPin->Properties.Label == PCGPinConstants::DefaultOutputLabel))
-	{
-		return false;
-	}
-
-	// First inspect data coming from CPU. Will return description if this node is a CPU node or in a different compute graph.
-	if (Binding->ComputeCPUOutputPinDataDesc(OutputPin, OutDesc))
-	{
-		return true;
-	}
-
-	const FPCGDataCollectionDesc InputPinDesc = PCGDataForGPUHelpers::ComputeInputPinDataDesc(this, PCGPinConstants::DefaultInputLabel, Binding);
-
-	// First forward input data to output.
-	FPCGDataCollectionDesc OutputPinDesc = InputPinDesc;
-
-	// Add output attribute (selected mesh).
-	{
-		// Create unique value keys for the output string.
-		if (UPCGMeshSelectorWeighted* SelectorWeighted = Cast<UPCGMeshSelectorWeighted>(MeshSelectorParameters))
-		{
-			TArray<int32> UniqueStringKeys;
-
-			// Weighted selection - add explicit strings from settings.
-			for (const FPCGMeshSelectorWeightedEntry& Entry : SelectorWeighted->MeshEntries)
-			{
-				const FString Value = Entry.Descriptor.StaticMesh.ToString();
-				const int32 StringIndex = Binding->GetStringTable().IndexOfByKey(Value);
-				if (ensureAlways(StringIndex != INDEX_NONE))
-				{
-					UniqueStringKeys.AddUnique(StringIndex);
-				}
-			}
-
-			OutputPinDesc.AddAttributeToAllData(FPCGKernelAttributeKey(OutAttributeName, EPCGKernelAttributeType::StringKey), Binding, &UniqueStringKeys);
-		}
-		else if (UPCGMeshSelectorByAttribute* SelectorByAttribute = Cast<UPCGMeshSelectorByAttribute>(MeshSelectorParameters))
-		{
-			// By-attribute selection - pass on strings from input attribute.
-			if (SelectorByAttribute->AttributeName != NAME_None)
-			{
-				for (FPCGDataDesc& DataDesc : OutputPinDesc.DataDescs)
-				{
-					const TArray<int32>* StringKeys = nullptr;
-
-					for (const FPCGKernelAttributeDesc& AttrDesc : DataDesc.AttributeDescs)
-					{
-						if (AttrDesc.Name == SelectorByAttribute->AttributeName)
-						{
-							StringKeys = &AttrDesc.GetUniqueStringKeys();
-							break;
-						}
-					}
-
-					DataDesc.AddAttribute(FPCGKernelAttributeKey(OutAttributeName, EPCGKernelAttributeType::StringKey), Binding, StringKeys);
-				}
-			}
-		}
-		else if (MeshSelectorParameters)
-		{
-			UE_LOG(LogPCG, Warning, TEXT("Mesh selector not supported by GPU Static Mesh Spawner: %s"), *MeshSelectorParameters->GetName());
-		}
-	}
-
-	OutDesc = OutputPinDesc;
-
-	return true;
-}
-
-void UPCGStaticMeshSpawnerSettings::GetKernelAttributeKeys(TArray<FPCGKernelAttributeKey>& OutKeys) const
-{
-	// Register the attribute this node creates.
-	OutKeys.AddUnique(FPCGKernelAttributeKey(OutAttributeName, EPCGKernelAttributeType::StringKey));
-}
-
-void UPCGStaticMeshSpawnerSettings::AddStaticCreatedStrings(TArray<FString>& InOutStringTable) const
-{
-	if (UPCGMeshSelectorWeighted* SelectorWeighted = Cast<UPCGMeshSelectorWeighted>(MeshSelectorParameters))
-	{
-		for (const FPCGMeshSelectorWeightedEntry& Entry : SelectorWeighted->MeshEntries)
-		{
-			InOutStringTable.AddUnique(Entry.Descriptor.StaticMesh.ToString());
-		}
-	}
-}
-
-int UPCGStaticMeshSpawnerSettings::ComputeKernelThreadCount(const UPCGDataBinding* Binding) const
-{
-	const UPCGNode* Node = CastChecked<UPCGNode>(GetOuter());
-	const UPCGPin* Pin = Node->GetInputPin(PCGPinConstants::DefaultInputLabel);
-
-	return PCGDataForGPUHelpers::ComputeInputPinDataDesc(Pin, Binding).ComputeDataElementCount(EPCGDataType::Point);
-}
-
 #if WITH_EDITOR
-void UPCGStaticMeshSpawnerSettings::CreateAdditionalInputDataInterfaces(FPCGGPUCompilationContext& InOutContext, UObject* InObjectOuter, TArray<TObjectPtr<UComputeDataInterface>>& OutDataInterfaces) const
-{
-	Super::CreateAdditionalInputDataInterfaces(InOutContext, InObjectOuter, OutDataInterfaces);
-
-	TObjectPtr<UPCGStaticMeshSpawnerDataInterface> NodeDI = InOutContext.NewObject_AnyThread<UPCGStaticMeshSpawnerDataInterface>(InObjectOuter);
-	NodeDI->Settings = this;
-	OutDataInterfaces.Add(NodeDI);
-}
-
-void UPCGStaticMeshSpawnerSettings::CreateAdditionalOutputDataInterfaces(FPCGGPUCompilationContext& InOutContext, UObject* InObjectOuter, TArray<TObjectPtr<UComputeDataInterface>>& OutDataInterfaces) const
+UPCGComputeKernel* UPCGStaticMeshSpawnerSettings::CreateKernel(FPCGGPUCompilationContext& InOutContext, UObject* InObjectOuter) const
 {
-	Super::CreateAdditionalOutputDataInterfaces(InOutContext, InObjectOuter, OutDataInterfaces);
+	UPCGStaticMeshSpawnerKernel* Kernel = InOutContext.NewObject_AnyThread<UPCGStaticMeshSpawnerKernel>(InObjectOuter);
+	Kernel->Settings = this;
 
-	UPCGInstanceDataInterface* InstanceDI = InOutContext.NewObject_AnyThread<UPCGInstanceDataInterface>(InObjectOuter);
-	InstanceDI->SetProducerSettings(this);
-	InstanceDI->InputPinProvidingData = PCGPinConstants::DefaultInputLabel;
-	OutDataInterfaces.Add(InstanceDI);
+	return Kernel;
 }
 
 FText UPCGStaticMeshSpawnerSettings::GetDefaultNodeTitle() const
diff --git a/Engine/Plugins/PCG/Source/PCG/Private/Elements/PCGStaticMeshSpawnerKernel.cpp b/Engine/Plugins/PCG/Source/PCG/Private/Elements/PCGStaticMeshSpawnerKernel.cpp
new file mode 100644
--- /dev/null
+++ b/Engine/Plugins/PCG/Source/PCG/Private/Elements/PCGStaticMeshSpawnerKernel.cpp
@@ -0,0 +1,494 @@
+// Copyright Epic Games, Inc. All Rights Reserved.
+
+#include "Elements/PCGStaticMeshSpawnerKernel.h"
+
+#include "PCGComponent.h"
+#include "PCGContext.h"
+#include "PCGModule.h"
+#include "Components/PCGProceduralISMComponent.h"
+#include "Compute/PCGComputeGraph.h"
+#include "Compute/PCGDataBinding.h"
+#include "Compute/DataInterfaces/PCGInstanceDataInterface.h"
+#include "Compute/DataInterfaces/Elements/PCGStaticMeshSpawnerDataInterface.h"
+#include "Elements/PCGStaticMeshSpawner.h"
+#include "Graph/PCGGPUGraphCompilationContext.h"
+#include "InstanceDataPackers/PCGInstanceDataPackerBase.h"
+#include "MeshSelectors/PCGMeshSelectorByAttribute.h"
+#include "MeshSelectors/PCGMeshSelectorWeighted.h"
+
+#include "ShaderCompilerCore.h"
+#include "Engine/StaticMesh.h"
+
+#include UE_INLINE_GENERATED_CPP_BY_NAME(PCGStaticMeshSpawnerKernel)
+
+#define LOCTEXT_NAMESPACE "PCGStaticMeshSpawnerKernel"
+
+namespace PCGStaticMeshSpawnerKernel
+{
+	static const FText CouldNotLoadStaticMeshFormat = LOCTEXT("CouldNotLoadStaticMesh", "Could not load static mesh from path '{0}'.");
+	static const FText TooManyPrimitivesFormat = LOCTEXT("TooManyPrimitives", "Attempted to emit too many primitive components, terminated after creating '{0}'.");
+}
+
+bool UPCGStaticMeshSpawnerKernel::IsKernelValid(FPCGContext* InContext, bool bQuiet) const
+{
+	TRACE_CPUPROFILER_EVENT_SCOPE(UPCGStaticMeshSpawnerKernel::IsKernelValid);
+
+	if (!Super::IsKernelValid(InContext, bQuiet) || !ensure(Settings))
+	{
+		return false;
+	}
+
+	if (!ValidatePCGNode(InContext, bQuiet, Cast<UPCGSettings>(Settings)))
+	{
+		return false;
+	}
+
+	if (UPCGMeshSelectorWeighted* SelectorWeighted = Cast<UPCGMeshSelectorWeighted>(Settings->MeshSelectorParameters))
+	{
+		if (SelectorWeighted->MeshEntries.IsEmpty())
+		{
+			PCG_KERNEL_VALIDATION_ERR(InContext, Settings, bQuiet, LOCTEXT("NoMeshEntries", "No meshes specified in selector."));
+			return false;
+		}
+
+		for (const FPCGMeshSelectorWeightedEntry& Entry : SelectorWeighted->MeshEntries)
+		{
+			if (Entry.Descriptor.StaticMesh.IsNull())
+			{
+				PCG_KERNEL_VALIDATION_ERR(InContext, Settings, bQuiet, LOCTEXT("UnassignedMesh", "Unassigned mesh."));
+				return false;
+			}
+		}
+	}
+	else if (!Settings->MeshSelectorParameters || !Settings->MeshSelectorParameters->IsA<UPCGMeshSelectorByAttribute>())
+	{
+		PCG_KERNEL_VALIDATION_ERR(InContext, Settings, bQuiet, LOCTEXT("InvalidMeshSelector", "Currently GPU Static Mesh Spawner nodes must use PCGMeshSelectorWeighted or UPCGMeshSelectorByAttribute as the mesh selector type."));
+		return false;
+	}
+
+	// Currently instance packers must be able to specify a full list of attribute names upfront, to build the attribute table at compile time.
+	// TODO: We should be able to augment a static attribute table with new attributes at execution time, which will allow other types like regex.
+	if (Settings->InstanceDataPackerParameters && !Settings->InstanceDataPackerParameters->GetAttributeNames(/*OutNames=*/nullptr))
+	{
+		PCG_KERNEL_VALIDATION_ERR(InContext, Settings, bQuiet, LOCTEXT("InvalidInstancePacker", "Selected instance packer does not support GPU execution."));
+		return false;
+	}
+
+	return true;
+}
+
+bool UPCGStaticMeshSpawnerKernel::ComputeOutputBindingDataDesc(const UPCGComputeGraph* InGraph, FName InOutputPinLabel, const UPCGDataBinding* InBinding, FPCGDataCollectionDesc& OutDataDesc) const
+{
+	check(InGraph);
+	check(InBinding);
+
+	// Code assumes single output pin.
+	if (!ensure(InOutputPinLabel == PCGPinConstants::DefaultOutputLabel))
+	{
+		return false;
+	}
+
+	// Forward data from In to Out.
+	FPCGDataCollectionDesc PinDesc;
+	ensure(InBinding->GetKernelPinDataDesc_GameThread(this, PCGPinConstants::DefaultInputLabel, /*bIsInput=*/true, PinDesc));
+
+	// Add output attribute (selected mesh).
+	{
+		// Create unique value keys for the output string.
+		if (UPCGMeshSelectorWeighted* SelectorWeighted = Cast<UPCGMeshSelectorWeighted>(Settings->MeshSelectorParameters))
+		{
+			TArray<int32> UniqueStringKeys;
+
+			// Weighted selection - add explicit strings from settings.
+			for (const FPCGMeshSelectorWeightedEntry& Entry : SelectorWeighted->MeshEntries)
+			{
+				const FString Value = Entry.Descriptor.StaticMesh.ToString();
+				const int32 StringIndex = InBinding->GetStringTable().IndexOfByKey(Value);
+				if (ensureAlways(StringIndex != INDEX_NONE))
+				{
+					UniqueStringKeys.AddUnique(StringIndex);
+				}
+			}
+
+			PinDesc.AddAttributeToAllData(FPCGKernelAttributeKey(Settings->OutAttributeName, EPCGKernelAttributeType::StringKey), InBinding, &UniqueStringKeys);
+		}
+		else if (UPCGMeshSelectorByAttribute* SelectorByAttribute = Cast<UPCGMeshSelectorByAttribute>(Settings->MeshSelectorParameters))
+		{
+			// By-attribute selection - pass on strings from input attribute.
+			if (SelectorByAttribute->AttributeName != NAME_None)
+			{
+				for (FPCGDataDesc& DataDesc : PinDesc.DataDescs)
+				{
+					const TArray<int32>* StringKeys = nullptr;
+
+					for (const FPCGKernelAttributeDesc& AttrDesc : DataDesc.AttributeDescs)
+					{
+						if (AttrDesc.Name == SelectorByAttribute->AttributeName)
+						{
+							StringKeys = &AttrDesc.GetUniqueStringKeys();
+							break;
+						}
+					}
+
+					DataDesc.AddAttribute(FPCGKernelAttributeKey(Settings->OutAttributeName, EPCGKernelAttributeType::StringKey), InBinding, StringKeys);
+				}
+			}
+		}
+		else if (Settings->MeshSelectorParameters)
+		{
+			UE_LOG(LogPCG, Error, TEXT("Mesh selector not supported by GPU Static Mesh Spawner: %s"), *Settings->MeshSelectorParameters->GetName());
+		}
+	}
+
+	OutDataDesc = MoveTemp(PinDesc);
+	return true;
+}
+
+int UPCGStaticMeshSpawnerKernel::ComputeThreadCount(const UPCGDataBinding* InBinding) const
+{
+	FPCGDataCollectionDesc InputPinDesc;
+	ensure(InBinding->GetKernelPinDataDesc_GameThread(this, PCGPinConstants::DefaultInputLabel, /*bIsInput=*/true, InputPinDesc));
+
+	return InputPinDesc.ComputeDataElementCount(EPCGDataType::Point);
+}
+
+#if WITH_EDITOR
+FString UPCGStaticMeshSpawnerKernel::GetCookedSource(const UPCGComputeGraph* InComputeGraph) const
+{
+	FString TemplateFile;
+	ensure(LoadShaderSourceFile(TEXT("/Plugin/PCG/Private/Elements/PCGStaticMeshSpawner.usf"), EShaderPlatform::SP_PCD3D_SM5, &TemplateFile, nullptr));
+	return TemplateFile;
+}
+
+void UPCGStaticMeshSpawnerKernel::CreateAdditionalInputDataInterfaces(FPCGGPUCompilationContext& InOutContext, UObject* InObjectOuter, TArray<TObjectPtr<UComputeDataInterface>>& OutDataInterfaces) const
+{
+	Super::CreateAdditionalInputDataInterfaces(InOutContext, InObjectOuter, OutDataInterfaces);
+
+	TObjectPtr<UPCGStaticMeshSpawnerDataInterface> NodeDI = InOutContext.NewObject_AnyThread<UPCGStaticMeshSpawnerDataInterface>(InObjectOuter);
+	NodeDI->ProducerKernel = this;
+	OutDataInterfaces.Add(NodeDI);
+}
+
+void UPCGStaticMeshSpawnerKernel::CreateAdditionalOutputDataInterfaces(FPCGGPUCompilationContext& InOutContext, UObject* InObjectOuter, TArray<TObjectPtr<UComputeDataInterface>>& OutDataInterfaces) const
+{
+	Super::CreateAdditionalOutputDataInterfaces(InOutContext, InObjectOuter, OutDataInterfaces);
+
+	UPCGInstanceDataInterface* InstanceDI = InOutContext.NewObject_AnyThread<UPCGInstanceDataInterface>(InObjectOuter);
+	InstanceDI->ProducerKernel = this;
+	InstanceDI->InputPinProvidingData = PCGPinConstants::DefaultInputLabel;
+	OutDataInterfaces.Add(InstanceDI);
+}
+#endif
+
+void UPCGStaticMeshSpawnerKernel::AddStaticCreatedStrings(TArray<FString>& InOutStringTable) const
+{
+	if (UPCGMeshSelectorWeighted* SelectorWeighted = Cast<UPCGMeshSelectorWeighted>(Settings->MeshSelectorParameters))
+	{
+		for (const FPCGMeshSelectorWeightedEntry& Entry : SelectorWeighted->MeshEntries)
+		{
+			InOutStringTable.AddUnique(Entry.Descriptor.StaticMesh.ToString());
+		}
+	}
+}
+
+void UPCGStaticMeshSpawnerKernel::GetKernelAttributeKeys(TArray<FPCGKernelAttributeKey>& OutKeys) const
+{
+	// Register the attribute this node creates.
+	OutKeys.AddUnique(FPCGKernelAttributeKey(Settings->OutAttributeName, EPCGKernelAttributeType::StringKey));
+}
+
+bool UPCGStaticMeshSpawnerKernel::SetupProceduralISMComponents(FPCGContext* InContext, UPCGDataBinding* InBinding) const
+{
+	AActor* TargetActor = Settings->TargetActor.Get() ? Settings->TargetActor.Get() : InContext->GetTargetActor(nullptr);
+	if (!ensure(Settings->MeshSelectorParameters) || !ensure(TargetActor))
+	{
+		return false;
+	}
+
+	FPCGDataCollectionDesc InputDataDesc;
+	InBinding->GetKernelPinDataDesc_GameThread(this, PCGPinConstants::DefaultInputLabel, /*bIsInput=*/true, InputDataDesc);
+
+	const uint32 TotalInputPointCount = InputDataDesc.ComputeDataElementCount(EPCGDataType::Point);
+	if (TotalInputPointCount == 0)
+	{
+		return false;
+	}
+
+	const FBox LocalBounds = InContext->SourceComponent->GetGridBounds().ShiftBy(-InContext->SourceComponent->GetOwner()->GetActorLocation());
+
+	uint32 CustomFloatCount = 0;
+
+	TArray<FUint32Vector4> AttributeIdOffsetStrides;
+	if (Settings->InstanceDataPackerParameters && Settings->InstanceDataPackerParameters->GetAttributeNames(/*OutNames=*/nullptr))
+	{
+		TArray<FName> AttributeNames;
+		Settings->InstanceDataPackerParameters->GetAttributeNames(&AttributeNames);
+
+		for (FName AttributeName : AttributeNames)
+		{
+			if (AttributeName == NAME_None)
+			{
+				PCG_KERNEL_VALIDATION_ERR(InContext, Settings, /*bQuiet=*/false,
+					FText::Format(LOCTEXT("InstanceDataAttributeInvalid", "Invalid instance data attribute specified '{0}'."), FText::FromName(AttributeName)));
+
+				continue;
+			}
+
+			if (!InputDataDesc.ContainsAttributeOnAnyData(AttributeName))
+			{
+				PCG_KERNEL_VALIDATION_ERR(InContext, Settings, /*bQuiet=*/false,
+					FText::Format(LOCTEXT("InstanceDataAttributeNotFound", "Instance data attribute '{0}' not found."), FText::FromName(AttributeName)));
+
+				continue;
+			}
+		}
+
+		PCGDataForGPUHelpers::ComputeCustomFloatPacking(InContext, Settings, AttributeNames, InBinding, InputDataDesc, CustomFloatCount, AttributeIdOffsetStrides);
+	}
+
+	TArray<float> PrimitiveSelectionCDF;
+	int32 SelectorAttributeId = INDEX_NONE;
+	TArray<int32> PrimitiveStringKeys;
+	TArray<FBox> PrimitiveMeshBounds;
+	TArray<FPCGProceduralISMComponentDescriptor> ComponentsToCreate;
+
+	if (const UPCGMeshSelectorByAttribute* SelectorByAttribute = Cast<UPCGMeshSelectorByAttribute>(Settings->MeshSelectorParameters))
+	{
+		const FName SelectorName = SelectorByAttribute->AttributeName;
+
+		// Compute how many instances we expect for each mesh.
+		TMap<int32, uint32> StringKeyToInstanceCount;
+
+		for (const FPCGDataDesc& Desc : InputDataDesc.DataDescs)
+		{
+			if (Desc.ElementCount <= 0)
+			{
+				continue;
+			}
+
+			for (const FPCGKernelAttributeDesc& AttributeDesc : Desc.AttributeDescs)
+			{
+				if (AttributeDesc.Name == SelectorName)
+				{
+					if (AttributeDesc.Type == EPCGKernelAttributeType::StringKey)
+					{
+						ensure(SelectorAttributeId == INDEX_NONE || SelectorAttributeId == AttributeDesc.Index);
+						SelectorAttributeId = AttributeDesc.Index;
+
+						for (int32 AttributeStringKeyValue : AttributeDesc.GetUniqueStringKeys())
+						{
+							if (AttributeStringKeyValue > 0)
+							{
+								StringKeyToInstanceCount.FindOrAdd(AttributeStringKeyValue, 0) += Desc.ElementCount;
+							}
+						}
+
+						// We could early out but we currently continue looping to generate the below warning for mismatched types.
+					}
+					else
+					{
+						// Currently only a single type per attribute name is supported (the name/type in the attribute table). It's possible to wire a graph
+						// for which an attribute is present with multiple types. Warn if this is encountered.
+						PCG_KERNEL_VALIDATION_WARN(InContext, Settings, /*bQuiet=*/false, FText::Format(
+							LOCTEXT("MeshSelectorAttributeNotUsable", "Attribute '{0}' not usable for mesh selection, only attributes of type String Key are supported."),
+							FText::FromName(SelectorName)));
+					}
+				}
+			}
+		}
+
+		if (SelectorAttributeId == INDEX_NONE)
+		{
+			PCG_KERNEL_VALIDATION_ERR(InContext, Settings, /*bQuiet=*/false, FText::Format(
+				LOCTEXT("MeshSelectorAttributeNotFound", "Mesh selector attribute '{0}' not found."),
+				FText::FromName(SelectorName)));
+
+			return false;
+		}
+
+		for (const TPair<int32, uint32>& StringKeyAndInstanceCount : StringKeyToInstanceCount)
+		{
+			if (!ensure(InBinding->GetStringTable().IsValidIndex(StringKeyAndInstanceCount.Key)))
+			{
+				continue;
+			}
+
+			const FString& MeshPathString = InBinding->GetStringTable()[StringKeyAndInstanceCount.Key];
+			if (MeshPathString.IsEmpty())
+			{
+				continue;
+			}
+
+			if (ComponentsToCreate.Num() >= PCGComputeConstants::MAX_PRIMITIVE_COMPONENTS_PER_SPAWNER)
+			{
+				PCG_KERNEL_VALIDATION_WARN(InContext, Settings, /*bQuiet=*/false,
+					FText::Format(PCGStaticMeshSpawnerKernel::TooManyPrimitivesFormat, FText::FromString(FString::FromInt(PCGComputeConstants::MAX_PRIMITIVE_COMPONENTS_PER_SPAWNER))));
+
+				break;
+			}
+
+			UStaticMesh* StaticMesh = Cast<UStaticMesh>(FSoftObjectPath(MeshPathString).TryLoad());
+			if (!StaticMesh)
+			{
+				PCG_KERNEL_VALIDATION_ERR(InContext, Settings, /*bQuiet=*/false,
+					FText::Format(PCGStaticMeshSpawnerKernel::CouldNotLoadStaticMeshFormat, FText::FromString(MeshPathString)));
+
+				continue;
+			}
+
+			FPCGProceduralISMComponentDescriptor Descriptor;
+			Descriptor = SelectorByAttribute->TemplateDescriptor;
+			Descriptor.NumInstances = StringKeyAndInstanceCount.Value;
+			Descriptor.LocalBounds = LocalBounds;
+			Descriptor.NumCustomFloats = CustomFloatCount;
+			Descriptor.StaticMesh = StaticMesh;
+
+			PrimitiveStringKeys.Add(StringKeyAndInstanceCount.Key);
+			ComponentsToCreate.Add(MoveTemp(Descriptor));
+
+			if (Settings->bApplyMeshBoundsToPoints)
+			{
+				PrimitiveMeshBounds.Add(Descriptor.StaticMesh->GetBoundingBox());
+			}
+		}
+
+		PrimitiveSelectionCDF.SetNumZeroed(ComponentsToCreate.Num());
+	}
+	else if (const UPCGMeshSelectorWeighted* SelectorWeighted = Cast<UPCGMeshSelectorWeighted>(Settings->MeshSelectorParameters))
+	{
+		if (SelectorWeighted->MeshEntries.IsEmpty())
+		{
+			PCG_KERNEL_VALIDATION_ERR(InContext, Settings, /*bQuiet=*/false, LOCTEXT("NoMeshEntries", "No mesh entries provided."));
+
+			return false;
+		}
+
+		float CumulativeWeight = 0.0f;
+
+		float TotalWeight = 0.0f;
+		for (const FPCGMeshSelectorWeightedEntry& Entry : SelectorWeighted->MeshEntries)
+		{
+			TotalWeight += Entry.Weight;
+		}
+
+		if (!ensure(TotalWeight > UE_SMALL_NUMBER))
+		{
+			return false;
+		}
+
+		PrimitiveSelectionCDF.Reserve(SelectorWeighted->MeshEntries.Num());
+
+		for (const FPCGMeshSelectorWeightedEntry& Entry : SelectorWeighted->MeshEntries)
+		{
+			if (Entry.Descriptor.StaticMesh.IsNull())
+			{
+				continue;
+			}
+
+			if (ComponentsToCreate.Num() >= PCGComputeConstants::MAX_PRIMITIVE_COMPONENTS_PER_SPAWNER)
+			{
+				PCG_KERNEL_VALIDATION_WARN(InContext, Settings, /*bQuiet=*/false,
+					FText::Format(PCGStaticMeshSpawnerKernel::TooManyPrimitivesFormat, FText::FromString(FString::FromInt(PCGComputeConstants::MAX_PRIMITIVE_COMPONENTS_PER_SPAWNER))));
+
+				break;
+			}
+
+			if (UStaticMesh* StaticMesh = Entry.Descriptor.StaticMesh.LoadSynchronous())
+			{
+				const float Weight = float(Entry.Weight) / TotalWeight;
+				CumulativeWeight += Weight;
+				PrimitiveSelectionCDF.Add(CumulativeWeight);
+				PrimitiveStringKeys.Add(InBinding->GetStringTable().IndexOfByKey(Entry.Descriptor.StaticMesh.ToString()));
+
+				FPCGProceduralISMComponentDescriptor Descriptor;
+				Descriptor = Entry.Descriptor;
+				Descriptor.LocalBounds = LocalBounds;
+				Descriptor.NumCustomFloats = CustomFloatCount;
+				Descriptor.StaticMesh = StaticMesh;
+
+				uint32 InstanceCount = FMath::CeilToInt(TotalInputPointCount * Weight);
+
+				if (SelectorWeighted->MeshEntries.Num() > 1)
+				{
+					// Since we'll be selecting meshes based on random draws using the point random seeds which we don't have on CPU,
+					// we may pick more or less than the expected number of instances for each mesh. Use binomial variance to calculate
+					// the overallocation that gives 99.7% confidence (3 sigma).
+					const double Variance = double(TotalInputPointCount) * Weight * (1.0f - Weight);
+					const double Sigma = FMath::Sqrt(Variance);
+					const uint32 AdditionalAllocation = FMath::CeilToInt(3.0f * Sigma);
+
+					// Useful for debugging instance counts.
+					//UE_LOG(LogPCG, Warning, TEXT("Over allocation, N = %d, P = %.2f, NP = %d, 3sig = %f, Extra = %d (%.2f%%), Calculated = %d, Final = %d"), TotalInputPointCount, Weight, InstanceCount, float(3.0f * Sigma), AdditionalAllocation, float(AdditionalAllocation) / TotalInputPointCount, InstanceCount + AdditionalAllocation, FMath::Min(InstanceCount + AdditionalAllocation, TotalInputPointCount));
+
+					InstanceCount += AdditionalAllocation;
+				}
+
+				Descriptor.NumInstances = FMath::Min(InstanceCount, TotalInputPointCount);
+
+				ComponentsToCreate.Add(MoveTemp(Descriptor));
+
+				if (Settings->bApplyMeshBoundsToPoints)
+				{
+					PrimitiveMeshBounds.Add(Descriptor.StaticMesh->GetBoundingBox());
+				}
+			}
+			else
+			{
+				PCG_KERNEL_VALIDATION_ERR(InContext, Settings, /*bQuiet=*/false,
+					FText::Format(PCGStaticMeshSpawnerKernel::CouldNotLoadStaticMeshFormat, FText::FromString(Entry.Descriptor.StaticMesh.ToString())));
+			}
+		}
+	}
+
+	// Validate instance count is not too large.
+	{
+		uint64 TotalInstanceCount = 0;
+		for (const FPCGProceduralISMComponentDescriptor& Desc : ComponentsToCreate)
+		{
+			TotalInstanceCount += Desc.NumInstances;
+		}
+
+		if (TotalInstanceCount >= MAX_INSTANCE_ID)
+		{
+			PCG_KERNEL_VALIDATION_ERR(InContext, Settings, /*bQuiet=*/false, FText::Format(
+				LOCTEXT("TooManyInstances", "Tried to spawn too many instances ({0}), procedural ISM component creation skipped and instances will not be rendered."),
+				FText::FromString(FString::FromInt(TotalInstanceCount))));
+
+			return false;
+		}
+	}
+
+	FPCGSpawnerPrimitives& Primitives = InBinding->MeshSpawnersToPrimitives.FindOrAdd(this);
+	Primitives.NumCustomFloats = CustomFloatCount;
+	Primitives.AttributeIdOffsetStrides = MoveTemp(AttributeIdOffsetStrides);
+	Primitives.SelectorAttributeId = SelectorAttributeId;
+	Primitives.SelectionCDF = MoveTemp(PrimitiveSelectionCDF);
+	Primitives.PrimitiveStringKeys = MoveTemp(PrimitiveStringKeys);
+	Primitives.PrimitiveMeshBounds = MoveTemp(PrimitiveMeshBounds);
+
+	for (const FPCGProceduralISMComponentDescriptor& Desc : ComponentsToCreate)
+	{
+		FPCGProceduralISMCBuilderParameters Params;
+		Params.Descriptor = Desc;
+		Params.bAllowDescriptorChanges = false;
+
+		UPCGManagedProceduralISMComponent* MISMC = PCGManagedProceduralISMComponent::GetOrCreateManagedProceduralISMC(TargetActor, InContext->SourceComponent.Get(), Settings->UID, Params);
+
+		check(MISMC);
+		MISMC->SetCrc(InContext->DependenciesCrc);
+
+		// Don't bother registering the resource change as we're transient anyway.
+		//InContext->TouchedResources.Emplace(MISMC);
+
+		Primitives.Primitives.Add(MISMC->GetComponent());
+
+		Primitives.NumInstancesAllPrimitives += Desc.NumInstances;
+	}
+
+	// Useful for debugging instance counts.
+	//UE_LOG(LogPCG, Log, TEXT("Input points %d, total instance count %d (+ %.2f%%)"), TotalInputPointCount, Primitives.NumInstancesAllPrimitives, (float(Primitives.NumInstancesAllPrimitives - TotalInputPointCount) / TotalInputPointCount) * 100.0f);
+
+	return !ComponentsToCreate.IsEmpty();
+}
+
+#undef LOCTEXT_NAMESPACE
diff --git a/Engine/Plugins/PCG/Source/PCG/Private/Graph/PCGGraphCompilerGPU.cpp b/Engine/Plugins/PCG/Source/PCG/Private/Graph/PCGGraphCompilerGPU.cpp
--- a/Engine/Plugins/PCG/Source/PCG/Private/Graph/PCGGraphCompilerGPU.cpp
+++ b/Engine/Plugins/PCG/Source/PCG/Private/Graph/PCGGraphCompilerGPU.cpp
@@ -6,6 +6,7 @@
 #include "PCGModule.h"
 #include "PCGPin.h"
 #include "Compute/PCGComputeCommon.h"
+#include "Compute/PCGComputeKernel.h"
 #include "Compute/PCGComputeKernelSource.h"
 #include "Compute/DataInterfaces/PCGComputeDataInterface.h"
 #include "Compute/DataInterfaces/PCGCustomKernelDataInterface.h"
@@ -16,7 +17,7 @@
 #include "Compute/DataInterfaces/PCGLandscapeDataInterface.h"
 #include "Compute/DataInterfaces/PCGTextureDataInterface.h"
 #include "Compute/Elements/PCGComputeGraphElement.h"
-#include "Elements/PCGStaticMeshSpawner.h"
+#include "Elements/PCGStaticMeshSpawnerKernel.h"
 #include "Graph/PCGGPUGraphCompilationContext.h"
 #include "Graph/PCGGraphCompiler.h"
 #include "Graph/PCGGraphExecutor.h"
@@ -36,6 +37,8 @@ namespace PCGGraphCompilerGPU
 		TEXT("Enable verbose logging of GPU compilation and execution."));
 }
 
+#define PCG_DATA_DESCRIPTION_LOGGING 0
+
 void FPCGGraphCompilerGPU::LabelConnectedGPUNodeIslands(
 	const TArray<FPCGGraphTask>& InCompiledTasks,
 	const TSet<FPCGTaskId>& InGPUCompatibleTaskIds,
@@ -554,12 +557,18 @@ void FPCGGraphCompilerGPU::BuildComputeGraphTask(
 		const FName GraphName = MakeUniqueObjectName(InGraph, UPCGComputeGraph::StaticClass(), InGraph->GetFName());
 		ComputeGraph = InOutContext.NewObject_AnyThread<UPCGComputeGraph>(ComputeGraphOuter, GraphName);
 
-		ComputeGraph->OutputCPUPinToInputGPUPinAlias = InOutputCPUPinToVirtualPin;
-
 		// Not incredibly useful for us - DG adds GetComponentSource()->GetComponentClass() object which allows it to bind at execution time by class.
 		// But execution code requires it currently.
 		ComputeGraph->Bindings.Add(UPCGDataBinding::StaticClass());
 
+		for (FPCGTaskId TaskId : InCollapsedTasks)
+		{
+			if (const UPCGSettings* Settings = InOutCompiledTasks[TaskId].Node ? InOutCompiledTasks[TaskId].Node->GetSettings() : nullptr)
+			{
+				InOutCompiledTasks[TaskId].Kernel = Settings->CreateKernel(InOutContext, ComputeGraph);
+			}
+		}
+
 		FPinToDataInterface OutputPinDataInterfaces;
 
 		CreateDataInterfaces(InOutContext, InOutCompiledTasks, InCollapsedTasks, InComputeElementTasks, InTaskSuccessors, InOriginalToVirtualPin, ComputeGraph, OutputPinDataInterfaces);
@@ -609,10 +618,13 @@ void FPCGGraphCompilerGPU::CreateDataInterfaces(
 
 	// Create data interfaces which allow kernels to read or write data. Each data interface is associated with a node output pin.
 	// For CPU->GPU edges, an upload data interface is created. For GPU->CPU edges, a readback data interface is created.
-	auto CreateDataInterface = [&InOutContext, &InCollapsedTasks, &InCompiledTasks, InOutComputeGraph](FPCGTaskId InProducerTaskId, bool bRequiresExport, const FPCGPinProperties& InOutputPinProperties) -> UPCGComputeDataInterface*
+	auto CreateOutputDataInterface = [&InOutContext, &InCollapsedTasks, &InCompiledTasks, InOutComputeGraph](FPCGTaskId InProducerTaskId, bool bRequiresExport, const FPCGPinProperties& InOutputPinProperties) -> UPCGComputeDataInterface*
 	{
 		TRACE_CPUPROFILER_EVENT_SCOPE(CreateDataInterface);
 
+		// The kernel that produced this data, if any (otherwise could be CPU node or a node from a different compute graph).
+		const UPCGComputeKernel* ProducerKernel = InCompiledTasks[InProducerTaskId].Kernel;
+
 		const bool bUpstreamIsGPUTask = InCollapsedTasks.Contains(InProducerTaskId);
 
 		EPCGDataType PinType = InOutputPinProperties.AllowedTypes;
@@ -648,7 +660,7 @@ void FPCGGraphCompilerGPU::CreateDataInterfaces(
 			check(DataInterfacePCGData);
 			DataInterfacePCGData->SetProducerSettings(ProducerSettings);
 			DataInterfacePCGData->SetRequiresExport(bRequiresExport);
-			DataInterfacePCGData->SetExposeElementCounters(ProducerSettings && ProducerPin && ProducerSettings->DoesOutputPinRequireElementCounters(ProducerPin));
+			DataInterfacePCGData->SetExposeElementCounters(ProducerKernel && ProducerPin && ProducerKernel->DoesOutputPinRequireElementCounters(ProducerPin->Properties.Label));
 
 			DataInterface = DataInterfacePCGData;
 			break;
@@ -678,6 +690,7 @@ void FPCGGraphCompilerGPU::CreateDataInterfaces(
 
 		if (DataInterface)
 		{
+			DataInterface->ProducerKernel = ProducerKernel;
 			DataInterface->SetOutputPin(InOutputPinProperties.Label);
 		}
 
@@ -687,12 +700,15 @@ void FPCGGraphCompilerGPU::CreateDataInterfaces(
 	// Create DIs for output pins (either output pins within this compute graph or outside).
 	for (FPCGTaskId TaskId : InCollapsedTasks)
 	{
+		const UPCGComputeKernel* ProducerKernel = InCompiledTasks[TaskId].Kernel;
+		const UPCGSettings* Settings = InCompiledTasks[TaskId].Node ? InCompiledTasks[TaskId].Node->GetSettings() : nullptr;
+
 		// Create DIs for all output pins regardless of outbound connections, because the kernels currently need their outputs to be bound to valid resources.
-		if (const UPCGSettings* Settings = InCompiledTasks[TaskId].Node ? InCompiledTasks[TaskId].Node->GetSettings() : nullptr)
+		if (Settings && ProducerKernel)
 		{
 			// Record attributes.
 			TArray<FPCGKernelAttributeKey> AttributeKeys;
-			Settings->GetKernelAttributeKeys(AttributeKeys);
+			ProducerKernel->GetKernelAttributeKeys(AttributeKeys);
 
 			for (const FPCGKernelAttributeKey& Key : AttributeKeys)
 			{
@@ -700,15 +716,15 @@ void FPCGGraphCompilerGPU::CreateDataInterfaces(
 			}
 
 			// Record strings.
-			Settings->AddStaticCreatedStrings(InOutComputeGraph->StringTable);
+			ProducerKernel->AddStaticCreatedStrings(InOutComputeGraph->StringTable);
 
 			// Create all the data index remaps for input pins.
-			auto GetDataLabels = [Settings, InOutComputeGraph](const TArray<FPCGPinProperties>& PinProperties)
+			auto GetDataLabels = [ProducerKernel, Settings, InOutComputeGraph](const TArray<FPCGPinProperties>& PinProperties)
 			{
 				for (const FPCGPinProperties& PinProps : PinProperties)
 				{
 					TArray<FString> DataLabels;
-					Settings->GetDataLabels(PinProps.Label, DataLabels);
+					ProducerKernel->GetDataLabels(PinProps.Label, DataLabels);
 
 					if (!DataLabels.IsEmpty())
 					{
@@ -755,7 +771,7 @@ void FPCGGraphCompilerGPU::CreateDataInterfaces(
 					}
 				}
 
-				if (UPCGComputeDataInterface* OutputDI = CreateDataInterface(TaskId, bRequiresExport, OutputPinProperties))
+				if (UPCGComputeDataInterface* OutputDI = CreateOutputDataInterface(TaskId, bRequiresExport, OutputPinProperties))
 				{
 					// Get the compute graph virtual pin corresponding to this pin, if one was created.
 					const FName* VirtualPinLabel = InOriginalToVirtualPin.Find({ TaskId, OutputPinProperties.Label, /*bIsInputPin*/false });
@@ -799,14 +815,6 @@ void FPCGGraphCompilerGPU::CreateDataInterfaces(
 				continue;
 			}
 
-			if (const UPCGNode* DownstreamNode = InCompiledTasks[TaskId].Node)
-			{
-				if (const UPCGPin* InputPin = DownstreamNode->GetInputPin(Input.DownstreamPin->Label))
-				{
-					InOutComputeGraph->PinsReceivingDataFromCPU.Add(InputPin);
-				}
-			}
-
 			if (!Input.UpstreamPin.IsSet())
 			{
 				continue;
@@ -818,7 +826,7 @@ void FPCGGraphCompilerGPU::CreateDataInterfaces(
 			{
 				OutputDI = *FoundOutputDI;
 			}
-			else if (UPCGComputeDataInterface* NewOutputDI = CreateDataInterface(Input.TaskId, /*bRequiresExport=*/false, *Input.UpstreamPin))
+			else if (UPCGComputeDataInterface* NewOutputDI = CreateOutputDataInterface(Input.TaskId, /*bRequiresExport=*/false, *Input.UpstreamPin))
 			{
 				OutputDI = NewOutputDI;
 
@@ -918,18 +926,45 @@ void FPCGGraphCompilerGPU::CompileComputeGraph(
 		const UPCGSettings* Settings = Node ? Node->GetSettings() : nullptr;
 		check(Settings && Settings->bEnabled && Settings->ShouldExecuteOnGPU());
 
-		if (Settings->IsA<UPCGStaticMeshSpawnerSettings>())
+		struct FInterfaceBinding
 		{
-			InOutComputeGraph->StaticMeshSpawners.Add(Settings);
+			const UComputeDataInterface* DataInterface = nullptr;
+			int32 DataInterfaceBindingIndex = INDEX_NONE;
+			FName Prefix = NAME_None;
+			FString BindingFunctionName;
+		};
+
+		struct FKernelWithDataBindings
+		{
+			UPCGComputeKernel* Kernel = nullptr;
+			TArray<FInterfaceBinding> InputDataBindings;
+			TArray<FInterfaceBinding> OutputDataBindings;
+		};
+
+		FKernelWithDataBindings KernelWithBindings;
+
+		const int32 KernelIndex = InOutComputeGraph->KernelInvocations.Num();
+		KernelWithBindings.Kernel = Settings->CreateKernel(InOutContext, /*InObjectOuter=*/InOutComputeGraph);
+		check(KernelWithBindings.Kernel);
+		KernelWithBindings.Kernel->SetKernelIndex(KernelIndex);
+
+		InOutComputeGraph->KernelInvocations.Add(Cast<UComputeKernel>(KernelWithBindings.Kernel));
+
+		if (const UPCGStaticMeshSpawnerKernel* SMSpawnerKernel = Cast<UPCGStaticMeshSpawnerKernel>(KernelWithBindings.Kernel))
+		{
+			InOutComputeGraph->StaticMeshSpawners.Add(SMSpawnerKernel);
 		}
 
 		// A given data interface may be connected to multiple downstream pins. Because each entry in these arrays produces its own bindings and graph edges,
 		// we also need to specify the prefix to uniquely describe each edge.
-		TArray<TPair</*DataInterfaceIndex=*/int, /*DataInterfacePrefix=*/FName>> InputDataInterfaceIndexAndPrefix;
-		TArray<TPair</*DataInterfaceIndex=*/int, /*DataInterfacePrefix=*/FName>> OutputDataInterfaceIndexAndPrefix;
+		TArray<TPair</*DataInterfaceIndex*/int, /*DataInterfacePrefix*/FName>> InputDataInterfaceIndexAndPrefix;
+		TArray<TPair</*DataInterfaceIndex*/int, /*DataInterfacePrefix*/FName>> OutputDataInterfaceIndexAndPrefix;
 		InputDataInterfaceIndexAndPrefix.Reserve(Settings->InputPinProperties().Num());
 		OutputDataInterfaceIndexAndPrefix.Reserve(Settings->OutputPinProperties().Num());
 
+		TMap</*OriginalPinLabel*/FName, /*VirtualPinLabel*/FName> OriginalToVirtualInputPin;
+		OriginalToVirtualInputPin.Reserve(InputDataInterfaceIndexAndPrefix.Num());
+
 		auto CreateDataResolverDataInterface = [&InOutContext, Settings, InOutComputeGraph](TArray<TPair</*DataInterfaceIndex=*/int, /*DataInterfacePrefix=*/FName>>& InOutDataInterfaceIndexAndPrefix, FName PinLabel)
 		{
 			// Add data label resolver DI if this pin uses any data labels.
@@ -982,6 +1017,11 @@ void FPCGGraphCompilerGPU::CompileComputeGraph(
 
 			InputDataInterfaceIndexAndPrefix.Emplace(Index, Input.DownstreamPin->Label);
 			CreateDataResolverDataInterface(InputDataInterfaceIndexAndPrefix, Input.DownstreamPin->Label);
+
+			if (const FName* VirtualPin = InOriginalToVirtualPin.Find(FNodePin(TaskId, Input.DownstreamPin->Label, /*IsInput*/true)))
+			{
+				OriginalToVirtualInputPin.Add(Input.DownstreamPin->Label, *VirtualPin);
+			}
 		}
 
 		// Setup output pin DIs. Always bind a DI to every output pin, so kernel always has something to write to.
@@ -1014,6 +1054,7 @@ void FPCGGraphCompilerGPU::CompileComputeGraph(
 		{
 			UPCGCustomKernelDataInterface* KernelDI = InOutContext.NewObject_AnyThread<UPCGCustomKernelDataInterface>(InOutComputeGraph);
 			KernelDI->Settings = Settings;
+			KernelDI->Kernel = KernelWithBindings.Kernel;
 
 			const int32 KernelDIIndex = InOutComputeGraph->DataInterfaces.Num();
 			InOutComputeGraph->DataInterfaces.Add(KernelDI);
@@ -1025,8 +1066,8 @@ void FPCGGraphCompilerGPU::CompileComputeGraph(
 		// Additional DIs created by settings.
 		{
 			TArray<TObjectPtr<UComputeDataInterface>> AdditionalInputDIs, AdditionalOutputDIs;
-			Settings->CreateAdditionalInputDataInterfaces(InOutContext, InOutComputeGraph, AdditionalInputDIs);
-			Settings->CreateAdditionalOutputDataInterfaces(InOutContext, InOutComputeGraph, AdditionalOutputDIs);
+			KernelWithBindings.Kernel->CreateAdditionalInputDataInterfaces(InOutContext, InOutComputeGraph, AdditionalInputDIs);
+			KernelWithBindings.Kernel->CreateAdditionalOutputDataInterfaces(InOutContext, InOutComputeGraph, AdditionalOutputDIs);
 
 			auto AddAdditionalDataInterface = [InOutComputeGraph](TObjectPtr<UComputeDataInterface> DataInterface) -> int32
 			{
@@ -1083,26 +1124,6 @@ void FPCGGraphCompilerGPU::CompileComputeGraph(
 		// Now that all data interfaces added, create the (trivial) binding mapping. All map to primary binding, index 0.
 		InOutComputeGraph->DataInterfaceToBinding.SetNumZeroed(InOutComputeGraph->DataInterfaces.Num());
 
-		struct FInterfaceBinding
-		{
-			const UComputeDataInterface* DataInterface = nullptr;
-			int32 DataInterfaceBindingIndex = INDEX_NONE;
-			FName Prefix = NAME_None;
-			FString BindingFunctionName = TEXT("");
-		};
-
-		struct FKernelWithDataBindings
-		{
-			UComputeKernel* Kernel = nullptr;
-			TArray<FInterfaceBinding> InputDataBindings;
-			TArray<FInterfaceBinding> OutputDataBindings;
-		};
-
-		FKernelWithDataBindings KernelWithBindings;
-
-		KernelWithBindings.Kernel = InOutContext.NewObject_AnyThread<UComputeKernel>(InOutComputeGraph);
-		const int KernelIndex = InOutComputeGraph->KernelInvocations.Num();
-		InOutComputeGraph->KernelInvocations.Add(KernelWithBindings.Kernel);
 		InOutComputeGraph->KernelToNode.Add(Node);
 
 		auto SetupAllInputBindings = [&KernelWithBindings, InOutComputeGraph](const TPair<int, FName>& InDataInterfaceIndex)
@@ -1157,7 +1178,8 @@ void FPCGGraphCompilerGPU::CompileComputeGraph(
 			KernelSource->GroupSize = FIntVector(PCGComputeConstants::THREAD_GROUP_SIZE, 1, 1);
 
 			// TODO: Instead of passing a half-baked compute graph into this virtual call, we should pass in a compilation context per compute graph.
-			KernelSource->SetSource(Settings->GetCookedKernelSource(InOutComputeGraph));
+			FString Source = KernelWithBindings.Kernel->GetCookedSource(InOutComputeGraph);
+			KernelSource->SetSource(MoveTemp(Source));
 
 			if (Settings->bDumpCookedHLSL)
 			{
@@ -1206,18 +1228,22 @@ void FPCGGraphCompilerGPU::CompileComputeGraph(
 			}
 		}
 
-		auto AddAllEdgesForKernel = [&KernelWithBindings, InOutComputeGraph](int32 InKernelIndex, bool bInEdgesAreInputs)
+		auto AddAllEdgesForKernel = [&KernelWithBindings, InOutComputeGraph, &OriginalToVirtualInputPin](int32 InKernelIndex, bool bInEdgesAreInputs)
 		{
 			TArray<FInterfaceBinding>& Bindings = bInEdgesAreInputs ? KernelWithBindings.InputDataBindings : KernelWithBindings.OutputDataBindings;
 
 			// Add all graph edges for bindings, which means include all functions that the data interfaces expose.
-			for (int BindingIndex = 0; BindingIndex < Bindings.Num(); ++BindingIndex)
+			for (int KernelBindingIndex = 0; KernelBindingIndex < Bindings.Num(); ++KernelBindingIndex)
 			{
-				FInterfaceBinding& Binding = Bindings[BindingIndex];
+				FInterfaceBinding& Binding = Bindings[KernelBindingIndex];
+
+				// Edge in compute graph is more correlated with a binding (it more describes a kernel input or output rather than a connection
+				// between two kernels) so we use "binding" terminology.
+				const int32 GraphBindingIndex = InOutComputeGraph->GraphEdges.Num();
+				FComputeGraphEdge& Edge = InOutComputeGraph->GraphEdges.AddDefaulted_GetRef();
 
-				FComputeGraphEdge& Edge = InOutComputeGraph->GraphEdges.Emplace_GetRef();
 				Edge.KernelIndex = InKernelIndex;
-				Edge.KernelBindingIndex = BindingIndex;
+				Edge.KernelBindingIndex = KernelBindingIndex;
 				Edge.DataInterfaceIndex = InOutComputeGraph->DataInterfaces.IndexOfByPredicate([&Binding](const UComputeDataInterface* In) { return Binding.DataInterface == In; });
 				check(Edge.DataInterfaceIndex != INDEX_NONE);
 				Edge.DataInterfaceBindingIndex = Binding.DataInterfaceBindingIndex;
@@ -1242,6 +1268,31 @@ void FPCGGraphCompilerGPU::CompileComputeGraph(
 						TEXT("{0}_{1}"),
 						{ Binding.Prefix.ToString(), DataInterfaceFunctions[Edge.DataInterfaceBindingIndex].Name }
 					);
+
+					// A binding corresponds to a single API within a data interface, like GetNumData() for example. There are multiple bindings
+					// per PCG graph edge and we only need to create our mappings for the first binding.
+					if (Edge.DataInterfaceBindingIndex == 0)
+					{
+						const FPCGKernelPin KernelPin(InKernelIndex, Binding.Prefix, bInEdgesAreInputs);
+
+						InOutComputeGraph->KernelBindingToPinLabel.Add(GraphBindingIndex, Binding.Prefix);
+
+						if (!InOutComputeGraph->KernelPinToFirstBinding.Contains(KernelPin))
+						{
+							InOutComputeGraph->KernelPinToFirstBinding.Add(KernelPin, GraphBindingIndex);
+						}
+
+						if (const FName* VirtualPin = bInEdgesAreInputs ? OriginalToVirtualInputPin.Find(Binding.Prefix) : nullptr)
+						{
+							InOutComputeGraph->CPUDataBindingToVirtualPinLabel.Add(GraphBindingIndex, *VirtualPin);
+						}
+
+						UPCGComputeDataInterface* PCGDataInterface = Cast<UPCGComputeDataInterface>(DataInterface);
+						if (PCGDataInterface && PCGDataInterface->GraphBindingIndex == INDEX_NONE)
+						{
+							PCGDataInterface->GraphBindingIndex = GraphBindingIndex;
+						}
+					}
 				}
 			}
 		};
@@ -1252,22 +1303,46 @@ void FPCGGraphCompilerGPU::CompileComputeGraph(
 
 	ensureMsgf(RemainingTasks.IsEmpty(), TEXT("PCG GPU graph compiler did not consume all tasks, %d/%d remaining."), RemainingTasks.Num(), InCollapsedTasks.Num());
 
-	// Register all virtual pin aliases with the corresponding pins for downstream usage.
-	for (const TPair<TTuple<FPCGTaskId, FName, bool>, FName>& Mapping : InOriginalToVirtualPin)
+	// Setup DownstreamToUpstreamBinding to assist in traversing up kernel->kernel connections.
+	// TODO: optimize quadratic loop.
+	for (int32 DownstreamEdgeIndex = 0; DownstreamEdgeIndex < InOutComputeGraph->GraphEdges.Num(); ++DownstreamEdgeIndex)
 	{
-		const bool bIsInputPin = Mapping.Get<0>().Get<2>();
-		if (bIsInputPin)
-		{
-			const FPCGTaskId TaskId = Mapping.Get<0>().Get<0>();
+		const FComputeGraphEdge& DownstreamEdge = InOutComputeGraph->GraphEdges[DownstreamEdgeIndex];
+		const bool bFirstEdgeOfDI = DownstreamEdge.DataInterfaceBindingIndex == 0;
 
-			if (const UPCGNode* Node = InCompiledTasks[TaskId].Node)
+		// Need only one direction (only need to follow edges upstream).
+		if (DownstreamEdge.bKernelInput && bFirstEdgeOfDI)
+		{
+			// Find corresponding matching edge - different kernel, same data interface, and bound as output.
+			for (int32 UpstreamEdgeIndex = 0; UpstreamEdgeIndex < InOutComputeGraph->GraphEdges.Num(); ++UpstreamEdgeIndex)
 			{
-				const FName& OriginalLabel = Mapping.Get<0>().Get<1>();
+				const FComputeGraphEdge& UpstreamEdge = InOutComputeGraph->GraphEdges[UpstreamEdgeIndex];
+				const bool bFirstEdgeOfUpstreamDI = UpstreamEdge.DataInterfaceBindingIndex == 0;
 
-				if (const UPCGPin* Pin = Node->GetInputPin(OriginalLabel))
+				if (bFirstEdgeOfUpstreamDI && UpstreamEdge.KernelIndex != DownstreamEdge.KernelIndex && UpstreamEdge.DataInterfaceIndex == DownstreamEdge.DataInterfaceIndex && !UpstreamEdge.bKernelInput)
 				{
-					const FName& VirtualLabel = Mapping.Get<1>();
-					InOutComputeGraph->InputPinLabelAliases.FindOrAdd(Pin) = VirtualLabel;
+#if PCG_DATA_DESCRIPTION_LOGGING
+					UE_LOG(LogPCG, Warning, TEXT("Connecting edge %d %s (KBD: %d) to edge %d %s (KBD: %d)"),
+						DownstreamEdge.KernelIndex,
+						*DownstreamEdge.BindingFunctionNameOverride,
+						DownstreamEdgeIndex,
+						UpstreamEdge.KernelIndex,
+						*UpstreamEdge.BindingFunctionNameOverride,
+						UpstreamEdge.KernelBindingIndex);
+
+					UE_LOG(LogPCG, Warning, TEXT("\tUpstream is kernel %d, DI '%s', edge index %d ('%s')"),
+						UpstreamEdge.KernelIndex,
+						*InOutComputeGraph->DataInterfaces[UpstreamEdge.DataInterfaceIndex]->GetName(),
+						UpstreamEdgeIndex,
+						*InOutComputeGraph->GraphEdges[UpstreamEdge.KernelBindingIndex].BindingFunctionNameOverride);
+					UE_LOG(LogPCG, Warning, TEXT("\tDownstream is kernel %d, DI '%s', edge index %d ('%s')"),
+						DownstreamEdge.KernelIndex,
+						*InOutComputeGraph->DataInterfaces[DownstreamEdge.DataInterfaceIndex]->GetName(),
+						DownstreamEdgeIndex,
+						*InOutComputeGraph->GraphEdges[DownstreamEdge.KernelBindingIndex].BindingFunctionNameOverride);
+#endif
+
+					InOutComputeGraph->DownstreamToUpstreamBinding.Add(DownstreamEdgeIndex, UpstreamEdgeIndex);
 				}
 			}
 		}
@@ -1421,4 +1496,7 @@ void FPCGGraphCompilerGPU::CreateGPUNodes(FPCGGraphCompiler& InOutCompiler, UPCG
 		return false;
 	});
 }
+
+#undef PCG_DATA_DESCRIPTION_LOGGING 
+
 #endif // WITH_EDITOR
diff --git a/Engine/Plugins/PCG/Source/PCG/Private/Graph/PCGGraphTask.h b/Engine/Plugins/PCG/Source/PCG/Private/Graph/PCGGraphTask.h
--- a/Engine/Plugins/PCG/Source/PCG/Private/Graph/PCGGraphTask.h
+++ b/Engine/Plugins/PCG/Source/PCG/Private/Graph/PCGGraphTask.h
@@ -12,6 +12,8 @@
 
 #include "PCGGraphTask.generated.h"
 
+class UPCGComputeKernel;
+
 /** Dictates how to recreate an element if it does not exist. */
 UENUM()
 enum class EPCGElementSource
@@ -138,6 +140,9 @@ struct FPCGGraphTask
 	FPCGContext* Context = nullptr;
 	TSharedPtr<const FPCGStackContext> StackContext;
 
+	/** Kernel for GPU execution if this task supports it. */
+	const UPCGComputeKernel* Kernel = nullptr;
+
 	// Whether SetupTask has been called on this task
 	bool bHasDoneSetup = false;
 	// BuildTaskInput will initialize this Collection which will later be used by PrepareForExecute
diff --git a/Engine/Plugins/PCG/Source/PCG/Private/PCGElement.cpp b/Engine/Plugins/PCG/Source/PCG/Private/PCGElement.cpp
--- a/Engine/Plugins/PCG/Source/PCG/Private/PCGElement.cpp
+++ b/Engine/Plugins/PCG/Source/PCG/Private/PCGElement.cpp
@@ -277,7 +277,6 @@ void IPCGElement::PreExecute(FPCGContext* Context) const
 	// to generate runtime errors/graph errors, then pass through.
 	if (Settings && Settings->ShouldExecuteOnGPU())
 	{
-		Settings->IsKernelValid(Context, /*bQuiet=*/false);
 		DisabledPassThroughData(Context);
 		Context->CurrentPhase = EPCGExecutionPhase::PostExecute;
 		return;
diff --git a/Engine/Plugins/PCG/Source/PCG/Private/PCGSettings.cpp b/Engine/Plugins/PCG/Source/PCG/Private/PCGSettings.cpp
--- a/Engine/Plugins/PCG/Source/PCG/Private/PCGSettings.cpp
+++ b/Engine/Plugins/PCG/Source/PCG/Private/PCGSettings.cpp
@@ -47,11 +47,6 @@ namespace PCGSettings
 		TEXT("pcg.Graph.WarnPinNameCollisions"),
 		false,
 		TEXT("Enables warnings when there are name collision between pin names and overrides/user parameters."));
-
-	static TAutoConsoleVariable<bool> CVarWarnOnOverridePinUsage(
-		TEXT("pcg.Graph.GPU.WarnOnOverridePinUsage"),
-		true,
-		TEXT("Enables warnings when parameters are overidden on GPU nodes."));
 }
 
 /** Custom Crc computation that ignores properties that will not affect the computed result of a node. */
@@ -220,82 +215,6 @@ void UPCGSettingsInterface::SetEnabled(bool bInEnabled)
 	}
 }
 
-bool UPCGSettings::IsKernelValid(FPCGContext* InContext, bool bQuiet) const
-{
-	if (PCGSettings::CVarWarnOnOverridePinUsage.GetValueOnAnyThread() && !bQuiet)
-	{
-		for (const FPCGSettingsOverridableParam& Param : OverridableParams())
-		{
-			if (ensure(!Param.PropertiesNames.IsEmpty()))
-			{
-				const FName PropertyName = Param.PropertiesNames[0];
-
-				if (IsPropertyOverriddenByPin(PropertyName))
-				{
-					PCG_KERNEL_VALIDATION_WARN(InContext, this, bQuiet, FText::Format(
-						LOCTEXT("ParamOverrideGPU", "Tried to override pin '{0}', but overrides are not supported on GPU nodes."),
-						FText::FromName(PropertyName)));
-				}
-			}
-		}
-	}
-
-	// Validate types of incident edges to make sure we catch invalid cases like Spatial -> Point.
-	bool bAllEdgesValid = true;
-
-	if (const UPCGNode* Node = Cast<UPCGNode>(GetOuter()))
-	{
-		for (const UPCGPin* InputPin : Node->GetInputPins())
-		{
-			if (!InputPin)
-			{
-				continue;
-			}
-
-			for (const UPCGEdge* InputEdge : InputPin->Edges)
-			{
-				const UPCGPin* UpstreamPin = InputEdge ? InputEdge->GetOtherPin(InputPin) : nullptr;
-				if (UpstreamPin && InputPin->GetRequiredTypeConversion(UpstreamPin) != EPCGTypeConversion::NoConversionRequired)
-				{
-					PCG_KERNEL_VALIDATION_ERR(InContext, this, bQuiet, FText::Format(
-						LOCTEXT("InvalidInputPinEdge", "Unsupported connected upstream pin '{0}' on node '{1}' with type {2}. Recreate the edge to add required conversion nodes."),
-						FText::FromName(UpstreamPin->Properties.Label),
-						Node->GetNodeTitle(EPCGNodeTitleType::ListView),
-						StaticEnum<EPCGDataType>() ? StaticEnum<EPCGDataType>()->GetDisplayNameTextByValue(static_cast<int64>(UpstreamPin->Properties.AllowedTypes)) : FText::FromString(TEXT("MISSING"))));
-
-					bAllEdgesValid = false;
-				}
-			}
-		}
-	}
-
-	return bAllEdgesValid;
-}
-
-bool UPCGSettings::ComputeOutputPinDataDesc(const FName& OutputPinLabel, const UPCGDataBinding* InBinding, FPCGDataCollectionDesc& OutDesc) const
-{
-	const UPCGNode* Node = CastChecked<UPCGNode>(GetOuter());
-	return ComputeOutputPinDataDesc(Node->GetOutputPin(OutputPinLabel), InBinding, OutDesc);
-}
-
-bool UPCGSettings::ComputeOutputPinDataDesc(const UPCGPin* OutputPin, const UPCGDataBinding* InBinding, FPCGDataCollectionDesc& OutDesc) const
-{
-	check(OutputPin);
-	check(InBinding);
-
-	const bool bSuccess = InBinding->ComputeCPUOutputPinDataDesc(OutputPin, OutDesc);
-
-	if (!bSuccess)
-	{
-		ensureMsgf(false, TEXT("Gathering data from CPU output pin '%s' on node '%s' failed. Pin was not present in OutputCPUPinToInputGPUPinAlias map (%d entries)."),
-			*OutputPin->Properties.Label.ToString(),
-			OutputPin->Node ? *OutputPin->Node->GetNodeTitle(EPCGNodeTitleType::ListView).ToString() : TEXT("MISSING"),
-			InBinding->Graph->OutputCPUPinToInputGPUPinAlias.Num());
-	}
-
-	return bSuccess;
-}
-
 uint32 UPCGSettings::GetTypeNameHash() const
 {
 	return GetTypeHash(GetClass()->GetFName());
@@ -641,6 +560,39 @@ EPCGDataType UPCGSettings::GetTypeUnionOfIncidentEdges(const FName& PinLabel) co
 	return Result;
 }
 
+bool UPCGSettings::AreInputEdgesValid(FPCGContext* InContext, bool bQuiet) const
+{
+	bool bAllEdgesValid = true;
+
+	if (const UPCGNode* Node = Cast<UPCGNode>(GetOuter()))
+	{
+		for (const UPCGPin* InputPin : Node->GetInputPins())
+		{
+			if (!InputPin)
+			{
+				continue;
+			}
+
+			for (const UPCGEdge* InputEdge : InputPin->Edges)
+			{
+				const UPCGPin* UpstreamPin = InputEdge ? InputEdge->GetOtherPin(InputPin) : nullptr;
+				if (UpstreamPin && InputPin->GetRequiredTypeConversion(UpstreamPin) != EPCGTypeConversion::NoConversionRequired)
+				{
+					PCG_KERNEL_VALIDATION_ERR(InContext, this, bQuiet, FText::Format(
+						LOCTEXT("InvalidInputPinEdge", "Unsupported connected upstream pin '{0}' on node '{1}' with type {2}. Recreate the edge to add required conversion nodes."),
+						FText::FromName(UpstreamPin->Properties.Label),
+						Node->GetNodeTitle(EPCGNodeTitleType::ListView),
+						StaticEnum<EPCGDataType>() ? StaticEnum<EPCGDataType>()->GetDisplayNameTextByValue(static_cast<int64>(UpstreamPin->Properties.AllowedTypes)) : FText::FromString(TEXT("MISSING"))));
+
+					bAllEdgesValid = false;
+				}
+			}
+		}
+	}
+
+	return bAllEdgesValid;
+}
+
 TArray<FPCGPinProperties> UPCGSettings::AllOutputPinProperties() const
 {
 	return OutputPinProperties();
diff --git a/Engine/Plugins/PCG/Source/PCG/Public/Elements/Metadata/PCGMetadataPartition.h b/Engine/Plugins/PCG/Source/PCG/Public/Elements/Metadata/PCGMetadataPartition.h
--- a/Engine/Plugins/PCG/Source/PCG/Public/Elements/Metadata/PCGMetadataPartition.h
+++ b/Engine/Plugins/PCG/Source/PCG/Public/Elements/Metadata/PCGMetadataPartition.h
@@ -18,20 +18,14 @@ public:
 	//~End UObject interface
 
 	//~Begin UPCGSettings interface
-	virtual bool IsKernelValid(FPCGContext* InContext = nullptr, bool bQuiet = true) const;
-	virtual FString GetCookedKernelSource(const UPCGComputeGraph* InComputeGraph) const override;
-	virtual void GetKernelAttributeKeys(TArray<FPCGKernelAttributeKey>& OutKeys) const override;
-	virtual int ComputeKernelThreadCount(const UPCGDataBinding* Binding) const override;
-	virtual bool ComputeOutputPinDataDesc(const UPCGPin* OutputPin, const UPCGDataBinding* Binding, FPCGDataCollectionDesc& OutDesc) const override;
-	virtual bool DoesOutputPinRequireElementCounters(const UPCGPin* OutputPin) const override { return true; }
 #if WITH_EDITOR
-	virtual bool SplitGraphAtOutput() const { return true; }
-	virtual void CreateAdditionalInputDataInterfaces(FPCGGPUCompilationContext& InOutContext, UObject* InObjectOuter, TArray<TObjectPtr<UComputeDataInterface>>& OutDataInterfaces) const override;
+	virtual bool SplitGraphAtOutput() const override { return true; }
 	virtual bool DisplayExecuteOnGPUSetting() const override { return true; }
 	virtual void ApplyDeprecation(UPCGNode* InOutNode) override;
 	virtual FName GetDefaultNodeName() const override { return FName(TEXT("AttributePartition")); }
 	virtual FText GetDefaultNodeTitle() const override { return NSLOCTEXT("PCGMetadataPartitionSettings", "NodeTitle", "Attribute Partition"); }
 	virtual EPCGSettingsType GetType() const override { return EPCGSettingsType::Metadata; }
+	UPCGComputeKernel* CreateKernel(FPCGGPUCompilationContext& InOutContext, UObject* InObjectOuter) const override;
 #endif
 	virtual bool HasDynamicPins() const override { return true; }
 	virtual FString GetAdditionalTitleInformation() const override;
diff --git a/Engine/Plugins/PCG/Source/PCG/Public/Elements/Metadata/PCGMetadataPartitionKernel.h b/Engine/Plugins/PCG/Source/PCG/Public/Elements/Metadata/PCGMetadataPartitionKernel.h
new file mode 100644
--- /dev/null
+++ b/Engine/Plugins/PCG/Source/PCG/Public/Elements/Metadata/PCGMetadataPartitionKernel.h
@@ -0,0 +1,30 @@
+// Copyright Epic Games, Inc. All Rights Reserved.
+
+#pragma once
+
+#include "Compute/PCGComputeKernel.h"
+
+#include "PCGMetadataPartitionKernel.generated.h"
+
+class UPCGMetadataPartitionSettings;
+
+UCLASS()
+class UPCGMetadataPartitionKernel : public UPCGComputeKernel
+{
+	GENERATED_BODY()
+
+public:
+	virtual bool IsKernelValid(FPCGContext* InContext = nullptr, bool bQuiet = true) const override;
+	virtual bool ComputeOutputBindingDataDesc(const UPCGComputeGraph* InGraph, FName InOutputPinLabel, const UPCGDataBinding* InBinding, FPCGDataCollectionDesc& OutDataDesc) const override;
+	virtual int ComputeThreadCount(const UPCGDataBinding* Binding) const override;
+#if WITH_EDITOR
+	virtual FString GetCookedSource(const UPCGComputeGraph* InComputeGraph) const override;
+	void CreateAdditionalInputDataInterfaces(FPCGGPUCompilationContext& InOutContext, UObject* InObjectOuter, TArray<TObjectPtr<UComputeDataInterface>>& OutDataInterfaces) const override;
+#endif
+	virtual void GetKernelAttributeKeys(TArray<FPCGKernelAttributeKey>& OutKeys) const override;
+	virtual bool DoesOutputPinRequireElementCounters(FName InOutputPinLabel) const override { return true; }
+
+public:
+	UPROPERTY()
+	TObjectPtr<const UPCGMetadataPartitionSettings> Settings;
+};
diff --git a/Engine/Plugins/PCG/Source/PCG/Public/Elements/PCGCopyPoints.h b/Engine/Plugins/PCG/Source/PCG/Public/Elements/PCGCopyPoints.h
--- a/Engine/Plugins/PCG/Source/PCG/Public/Elements/PCGCopyPoints.h
+++ b/Engine/Plugins/PCG/Source/PCG/Public/Elements/PCGCopyPoints.h
@@ -51,13 +51,9 @@ public:
 	virtual FText GetNodeTooltipText() const override;
 	virtual EPCGSettingsType GetType() const override { return EPCGSettingsType::Sampler; }
 	virtual bool DisplayExecuteOnGPUSetting() const override { return true; }
-	virtual void CreateAdditionalInputDataInterfaces(FPCGGPUCompilationContext& InOutContext, UObject* InObjectOuter, TArray<TObjectPtr<UComputeDataInterface>>& OutDataInterfaces) const override;
+	UPCGComputeKernel* CreateKernel(FPCGGPUCompilationContext& InOutContext, UObject* InObjectOuter) const override;
 #endif
 
-	virtual FString GetCookedKernelSource(const UPCGComputeGraph* InComputeGraph) const override;
-	virtual int ComputeKernelThreadCount(const UPCGDataBinding* Binding) const override;
-	virtual bool ComputeOutputPinDataDesc(const UPCGPin* OutputPin, const UPCGDataBinding* Binding, FPCGDataCollectionDesc& OutDesc) const override;
-
 protected:
 	virtual TArray<FPCGPinProperties> InputPinProperties() const override;
 	virtual TArray<FPCGPinProperties> OutputPinProperties() const override { return Super::DefaultPointOutputPinProperties(); }
diff --git a/Engine/Plugins/PCG/Source/PCG/Public/Elements/PCGCopyPointsKernel.h b/Engine/Plugins/PCG/Source/PCG/Public/Elements/PCGCopyPointsKernel.h
new file mode 100644
--- /dev/null
+++ b/Engine/Plugins/PCG/Source/PCG/Public/Elements/PCGCopyPointsKernel.h
@@ -0,0 +1,30 @@
+// Copyright Epic Games, Inc. All Rights Reserved.
+
+#pragma once
+
+#include "Compute/PCGComputeKernel.h"
+
+#include "Compute/PCGPinPropertiesGPU.h"
+
+#include "PCGCopyPointsKernel.generated.h"
+
+class UPCGCopyPointsSettings;
+
+UCLASS()
+class UPCGCopyPointsKernel : public UPCGComputeKernel
+{
+	GENERATED_BODY()
+
+public:
+	virtual bool IsKernelValid(FPCGContext* InContext = nullptr, bool bQuiet = true) const override;
+	virtual bool ComputeOutputBindingDataDesc(const UPCGComputeGraph* InGraph, FName InOutputPinLabel, const UPCGDataBinding* InBinding, FPCGDataCollectionDesc& OutDataDesc) const override;
+	virtual int ComputeThreadCount(const UPCGDataBinding* Binding) const override;
+#if WITH_EDITOR
+	virtual FString GetCookedSource(const UPCGComputeGraph* InComputeGraph) const override;
+	void CreateAdditionalInputDataInterfaces(FPCGGPUCompilationContext& InOutContext, UObject* InObjectOuter, TArray<TObjectPtr<UComputeDataInterface>>& OutDataInterfaces) const override;
+#endif
+
+public:
+	UPROPERTY()
+	TObjectPtr<const UPCGCopyPointsSettings> Settings;
+};
diff --git a/Engine/Plugins/PCG/Source/PCG/Public/Elements/PCGReroute.h b/Engine/Plugins/PCG/Source/PCG/Public/Elements/PCGReroute.h
--- a/Engine/Plugins/PCG/Source/PCG/Public/Elements/PCGReroute.h
+++ b/Engine/Plugins/PCG/Source/PCG/Public/Elements/PCGReroute.h
@@ -42,7 +42,6 @@ protected:
 	virtual TArray<FPCGPinProperties> InputPinProperties() const override;
 	virtual TArray<FPCGPinProperties> OutputPinProperties() const override;
 	virtual FPCGElementPtr CreateElement() const override;
-	virtual bool ComputeOutputPinDataDesc(const UPCGPin* OutputPin, const UPCGDataBinding* Binding, FPCGDataCollectionDesc& OutDesc) const override;
 	virtual bool HasExecutionDependencyPin() const override { return false; }
 #if WITH_EDITOR
 	virtual TArray<FPCGPreconfiguredInfo> GetConversionInfo() const override;
diff --git a/Engine/Plugins/PCG/Source/PCG/Public/Elements/PCGSkinnedMeshSpawner.h b/Engine/Plugins/PCG/Source/PCG/Public/Elements/PCGSkinnedMeshSpawner.h
--- a/Engine/Plugins/PCG/Source/PCG/Public/Elements/PCGSkinnedMeshSpawner.h
+++ b/Engine/Plugins/PCG/Source/PCG/Public/Elements/PCGSkinnedMeshSpawner.h
@@ -26,15 +26,7 @@ public:
 	UPCGSkinnedMeshSpawnerSettings(const FObjectInitializer& ObjectInitializer);
 
 	// ~Begin UPCGSettings interface
-	virtual bool IsKernelValid(FPCGContext* InContext = nullptr, bool bQuiet = true) const;
-	virtual FString GetCookedKernelSource(const UPCGComputeGraph* InComputeGraph) const override;
-	virtual void GetKernelAttributeKeys(TArray<FPCGKernelAttributeKey>& OutKeys) const override;
-	virtual int ComputeKernelThreadCount(const UPCGDataBinding* Binding) const override;
-	virtual bool ComputeOutputPinDataDesc(const UPCGPin* OutputPin, const UPCGDataBinding* Binding, FPCGDataCollectionDesc& OutDesc) const override;
 #if WITH_EDITOR
-	virtual void CreateAdditionalInputDataInterfaces(FPCGGPUCompilationContext& InOutContext, UObject* InObjectOuter, TArray<TObjectPtr<UComputeDataInterface>>& OutDataInterfaces) const override;
-	virtual void CreateAdditionalOutputDataInterfaces(FPCGGPUCompilationContext& InOutContext, UObject* InObjectOuter, TArray<TObjectPtr<UComputeDataInterface>>& OutDataInterfaces) const override;
-	virtual bool DisplayExecuteOnGPUSetting() const override { return true; }
 	virtual FName GetDefaultNodeName() const override { return FName(TEXT("SkinnedMeshSpawner")); }
 	virtual FText GetDefaultNodeTitle() const override;
 	virtual EPCGSettingsType GetType() const override { return EPCGSettingsType::Spawner; }
diff --git a/Engine/Plugins/PCG/Source/PCG/Public/Elements/PCGStaticMeshSpawner.h b/Engine/Plugins/PCG/Source/PCG/Public/Elements/PCGStaticMeshSpawner.h
--- a/Engine/Plugins/PCG/Source/PCG/Public/Elements/PCGStaticMeshSpawner.h
+++ b/Engine/Plugins/PCG/Source/PCG/Public/Elements/PCGStaticMeshSpawner.h
@@ -17,7 +17,6 @@ struct FPCGContext;
 struct FPCGMeshInstanceList;
 struct FPCGPackedCustomData;
 struct FPCGStaticMeshSpawnerContext;
-
 class UStaticMesh;
 
 UCLASS(BlueprintType, ClassGroup = (Procedural), meta = (PrioritizeCategories = "Settings"))
@@ -29,16 +28,9 @@ public:
 	UPCGStaticMeshSpawnerSettings(const FObjectInitializer &ObjectInitializer);
 
 	// ~Begin UPCGSettings interface
-	virtual bool IsKernelValid(FPCGContext* InContext = nullptr, bool bQuiet = true) const;
-	virtual FString GetCookedKernelSource(const UPCGComputeGraph* InComputeGraph) const override;
-	virtual void GetKernelAttributeKeys(TArray<FPCGKernelAttributeKey>& OutKeys) const override;
-	virtual void AddStaticCreatedStrings(TArray<FString>& InOutStringTable) const override;
-	virtual int ComputeKernelThreadCount(const UPCGDataBinding* Binding) const override;
-	virtual bool ComputeOutputPinDataDesc(const UPCGPin* OutputPin, const UPCGDataBinding* Binding, FPCGDataCollectionDesc& OutDesc) const override;
 #if WITH_EDITOR
-	virtual void CreateAdditionalInputDataInterfaces(FPCGGPUCompilationContext& InOutContext, UObject* InObjectOuter, TArray<TObjectPtr<UComputeDataInterface>>& OutDataInterfaces) const override;
-	virtual void CreateAdditionalOutputDataInterfaces(FPCGGPUCompilationContext& InOutContext, UObject* InObjectOuter, TArray<TObjectPtr<UComputeDataInterface>>& OutDataInterfaces) const override;
 	virtual bool DisplayExecuteOnGPUSetting() const override { return true; }
+	UPCGComputeKernel* CreateKernel(FPCGGPUCompilationContext& InOutContext, UObject* InObjectOuter) const override;
 	virtual FName GetDefaultNodeName() const override { return FName(TEXT("StaticMeshSpawner")); }
 	virtual FText GetDefaultNodeTitle() const override;
 	virtual EPCGSettingsType GetType() const override { return EPCGSettingsType::Spawner; }
@@ -125,9 +117,6 @@ public:
 protected:
 	void RefreshMeshSelector();
 	void RefreshInstancePacker();
-
-private:
-	static TCHAR const* TemplateFilePath;
 };
 
 class PCG_API FPCGStaticMeshSpawnerElement : public IPCGElement
diff --git a/Engine/Plugins/PCG/Source/PCG/Public/Elements/PCGStaticMeshSpawnerKernel.h b/Engine/Plugins/PCG/Source/PCG/Public/Elements/PCGStaticMeshSpawnerKernel.h
new file mode 100644
--- /dev/null
+++ b/Engine/Plugins/PCG/Source/PCG/Public/Elements/PCGStaticMeshSpawnerKernel.h
@@ -0,0 +1,35 @@
+// Copyright Epic Games, Inc. All Rights Reserved.
+
+#pragma once
+
+#include "Compute/PCGComputeKernel.h"
+
+#include "PCGStaticMeshSpawnerKernel.generated.h"
+
+class UPCGStaticMeshSpawnerSettings;
+
+UCLASS()
+class UPCGStaticMeshSpawnerKernel : public UPCGComputeKernel
+{
+	GENERATED_BODY()
+
+public:
+	virtual bool IsKernelValid(FPCGContext* InContext = nullptr, bool bQuiet = true) const override;
+	virtual bool ComputeOutputBindingDataDesc(const UPCGComputeGraph* InGraph, FName InOutputPinLabel, const UPCGDataBinding* InBinding, FPCGDataCollectionDesc& OutDataDesc) const override;
+	virtual int ComputeThreadCount(const UPCGDataBinding* Binding) const override;
+#if WITH_EDITOR
+	virtual FString GetCookedSource(const UPCGComputeGraph* InComputeGraph) const override;
+	void CreateAdditionalInputDataInterfaces(FPCGGPUCompilationContext& InOutContext, UObject* InObjectOuter, TArray<TObjectPtr<UComputeDataInterface>>& OutDataInterfaces) const override;
+	void CreateAdditionalOutputDataInterfaces(FPCGGPUCompilationContext& InOutContext, UObject* InObjectOuter, TArray<TObjectPtr<UComputeDataInterface>>& OutDataInterfaces) const override;
+#endif
+	virtual void AddStaticCreatedStrings(TArray<FString>& InOutStringTable) const override;
+	virtual void GetKernelAttributeKeys(TArray<FPCGKernelAttributeKey>& OutKeys) const override;
+	virtual bool DoesOutputPinRequireElementCounters(FName InOutputPinLabel) const override { return true; }
+
+	/** Setup Procedural ISMCs. Returns true if at least one component was set up. */
+	virtual bool SetupProceduralISMComponents(FPCGContext* InContext, UPCGDataBinding* InBinding) const;
+
+public:
+	UPROPERTY()
+	TObjectPtr<const UPCGStaticMeshSpawnerSettings> Settings;
+};
diff --git a/Engine/Plugins/PCG/Source/PCG/Public/PCGSettings.h b/Engine/Plugins/PCG/Source/PCG/Public/PCGSettings.h
--- a/Engine/Plugins/PCG/Source/PCG/Public/PCGSettings.h
+++ b/Engine/Plugins/PCG/Source/PCG/Public/PCGSettings.h
@@ -16,16 +16,11 @@
 
 class UPCGComponent;
 class UPCGPin;
-class UPCGComputeGraph;
-class UPCGDataBinding;
+class UPCGComputeKernel;
 class UPCGGraph;
-class UComputeDataInterface;
 class UPCGNode;
 class UPCGSettings;
-struct FPCGDataCollectionDesc;
 struct FPCGGPUCompilationContext;
-struct FPCGKernelAttributeKey;
-struct FPCGKernelAttributeTable;
 struct FPCGPinProperties;
 struct FPropertyChangedEvent;
 
@@ -226,6 +221,9 @@ public:
 	/** Bitwise union of the allowed types of each incident edge on pin. Returns None if no common bits, or no edges. */
 	EPCGDataType GetTypeUnionOfIncidentEdges(const FName& PinLabel) const;
 
+	/** Validate input edges, for example checking connected pin types are compatible. */
+	bool AreInputEdgesValid(FPCGContext* InContext, bool bQuiet) const;
+
 	// Internal functions, should not be used by any user.
 	// Return a different subset for input/output pin properties, in case of a default object.
 	virtual TArray<FPCGPinProperties> DefaultInputPinProperties() const;
@@ -506,65 +504,10 @@ public:
 	*/
 	virtual bool ShouldExecuteOnGPU() const { return bExecuteOnGPU; }
 
-	/** [EXPERIMENTAL] Performs validation and returns true if this node is suitable for deployment to the GPU.
-	* Note: This API function is subject to change without deprecation.
-	*/
-	virtual bool IsKernelValid(FPCGContext* InContext = nullptr, bool bQuiet = true) const;
-
-	/** [EXPERIMENTAL] Produces the node specific portion of kernel shader source text, including the main entry point.
-	* Note: This API function is subject to change without deprecation.
-	*/
-	virtual FString GetCookedKernelSource(const UPCGComputeGraph* InComputeGraph) const { return TEXT(""); }
-
-	/** [EXPERIMENTAL] Get all the attributes read or written by this node for which we know the name and type statically. In cases that only the
-	* attribute name used by the node is known (example: user specifies SM Spawner instance data attributes by name only, not type), the attribute
-	* is omitted from this list and must be resolved at execution time in the data provider.
-	* Note: This API function is subject to change without deprecation.
-	*/
-	virtual void GetKernelAttributeKeys(TArray<FPCGKernelAttributeKey>& OutKeys) const {}
-
-	/** [EXPERIMENTAL] Add any strings emitted by this node that are known statically at compile time.
-	* Note: This API function is subject to change without deprecation.
-	*/
-	virtual void AddStaticCreatedStrings(TArray<FString>& InOutStringTable) const {};
-
-	/** [EXPERIMENTAL] Compute how many threads should be dispatched to execute this node on the GPU.
-	* Note: This API function is subject to change without deprecation.
-	*/
-	virtual int ComputeKernelThreadCount(const UPCGDataBinding* Binding) const { return 0; };
-
-	/** [EXPERIMENTAL] Compute a description of data that will be output from OutputPinLabel/OutputPin.
-	* Note: This API function is subject to change without deprecation.
-	*/
-	bool ComputeOutputPinDataDesc(const FName& OutputPinLabel, const UPCGDataBinding* InBinding, FPCGDataCollectionDesc& OutDesc) const;
-
-	/** [EXPERIMENTAL] Compute a description of data that will be output from OutputPinLabel/OutputPin.
-	* Note: This API function is subject to change without deprecation.
-	*/
-	virtual bool ComputeOutputPinDataDesc(const UPCGPin* OutputPin, const UPCGDataBinding* InBinding, FPCGDataCollectionDesc& OutDesc) const;
-
-	/** [EXPERIMENTAL] Whether the output pin should have a buffer of counters attached, useful for compaction and other counting requirements.
-	* Note: This API function is subject to change without deprecation.
-	*/
-	virtual bool DoesOutputPinRequireElementCounters(const UPCGPin* OutputPin) const { return false; }
-
 #if WITH_EDITOR
-	/** [EXPERIMENTAL] Create additional input data interfaces to marshal any required input data.
-	* Note: This API function is subject to change without deprecation.
-	*/
-	virtual void CreateAdditionalInputDataInterfaces(FPCGGPUCompilationContext& InOutContext, UObject* InObjectOuter, TArray<TObjectPtr<UComputeDataInterface>>& OutDataInterfaces) const {}
-
-	/** [EXPERIMENTAL] Create additional output data interfaces to marshal any required output data.
-	* Note: This API function is subject to change without deprecation.
-	*/
-	virtual void CreateAdditionalOutputDataInterfaces(FPCGGPUCompilationContext& InOutContext, UObject* InObjectOuter, TArray<TObjectPtr<UComputeDataInterface>>& OutDataInterfaces) const {}
-	virtual void CreateAdditionalOutputDataInterfaces(TArray<TObjectPtr<UComputeDataInterface>>& OutDataInterfaces) const {}
-
-	/** [EXPERIMENTAL] Get any data labels that are used statically by this node.
-	* Note: This API function is subject to change without deprecation.
-	*/
-	virtual void GetDataLabels(FName InPinLabel, TArray<FString>& OutDataLabels) const {}
+	virtual UPCGComputeKernel* CreateKernel(FPCGGPUCompilationContext& InOutContext, UObject* InObjectOuter) const { return nullptr; }
 
+	// TODO - should also move to kernel, but needs flow to be shuffled in the GPU graph compiler.
 	/** [EXPERIMENTAL] Whether compute graphs should be split at the output of this node. Useful for reading back runtime statistics or diagnostic info before continuing.
 	* Note: This API function is subject to change without deprecation.
 	*/
@@ -591,6 +534,7 @@ protected:
 	UPROPERTY(EditAnywhere, Category = "GPU", meta = (Tooltip = "Whether this node should be executed on the GPU.", EditCondition = "DisplayExecuteOnGPUSetting()", EditConditionHides, HideEditConditionToggle))
 	bool bExecuteOnGPU = false;
 
+public:
 	/** [EXPERIMENTAL] Note: This variable is subject to change without deprecation. */
 	UPROPERTY(EditAnywhere, Category = "GPU", AdvancedDisplay, meta = (Tooltip = "Dump the cooked HLSL into the log after it is generated.", EditCondition = "bExecuteOnGPU", EditConditionHides))
 	bool bDumpCookedHLSL = false;
@@ -616,8 +560,33 @@ private:
 	FPCGCrc CachedCrc;
 #endif
 
-	friend class FPCGGraphCompilerGPU;
-	friend class UPCGDataBinding;
+public:
+	UE_DEPRECATED(5.6, "All GPU logic now resides separate kernel objects, see implementations of UPCGSettings::CreateKernel().")
+	virtual bool IsKernelValid(FPCGContext* InContext = nullptr, bool bQuiet = true) const { return false; }
+	UE_DEPRECATED(5.6, "All GPU logic now resides separate kernel objects, see implementations of UPCGSettings::CreateKernel().")
+	virtual FString GetCookedKernelSource(const class UPCGComputeGraph* InComputeGraph) const { return FString(); }
+	UE_DEPRECATED(5.6, "All GPU logic now resides separate kernel objects, see implementations of UPCGSettings::CreateKernel().")
+	virtual void GetKernelAttributeKeys(TArray<struct FPCGKernelAttributeKey>& OutKeys) const {}
+	UE_DEPRECATED(5.6, "All GPU logic now resides separate kernel objects, see implementations of UPCGSettings::CreateKernel().")
+	virtual void AddStaticCreatedStrings(TArray<FString>& InOutStringTable) const {};
+	UE_DEPRECATED(5.6, "All GPU logic now resides separate kernel objects, see implementations of UPCGSettings::CreateKernel().")
+	virtual int ComputeKernelThreadCount(const class UPCGDataBinding* Binding) const { return 0; };
+	UE_DEPRECATED(5.6, "All GPU logic now resides separate kernel objects, see implementations of UPCGSettings::CreateKernel().")
+	bool ComputeOutputPinDataDesc(const FName& OutputPinLabel, const class UPCGDataBinding* InBinding, struct FPCGDataCollectionDesc& OutDesc) const { return false; }
+	UE_DEPRECATED(5.6, "All GPU logic now resides separate kernel objects, see implementations of UPCGSettings::CreateKernel().")
+	virtual bool ComputeOutputPinDataDesc(const UPCGPin* OutputPin, const class UPCGDataBinding* InBinding, struct FPCGDataCollectionDesc& OutDesc) const { return false; }
+	UE_DEPRECATED(5.6, "All GPU logic now resides separate kernel objects, see implementations of UPCGSettings::CreateKernel().")
+	virtual bool DoesOutputPinRequireElementCounters(const UPCGPin* OutputPin) const { return false; }
+#if WITH_EDITOR
+	UE_DEPRECATED(5.6, "All GPU logic now resides separate kernel objects, see implementations of UPCGSettings::CreateKernel().")
+	virtual void CreateAdditionalInputDataInterfaces(FPCGGPUCompilationContext& InOutContext, UObject* InObjectOuter, TArray<TObjectPtr<class UComputeDataInterface>>& OutDataInterfaces) const {}
+	UE_DEPRECATED(5.6, "All GPU logic now resides separate kernel objects, see implementations of UPCGSettings::CreateKernel().")
+	virtual void CreateAdditionalOutputDataInterfaces(FPCGGPUCompilationContext& InOutContext, UObject* InObjectOuter, TArray<TObjectPtr<class UComputeDataInterface>>& OutDataInterfaces) const {}
+	UE_DEPRECATED(5.6, "All GPU logic now resides separate kernel objects, see implementations of UPCGSettings::CreateKernel().")
+	virtual void CreateAdditionalOutputDataInterfaces(TArray<TObjectPtr<class UComputeDataInterface>>& OutDataInterfaces) const {}
+	UE_DEPRECATED(5.6, "All GPU logic now resides separate kernel objects, see implementations of UPCGSettings::CreateKernel().")
+	virtual void GetDataLabels(FName InPinLabel, TArray<FString>& OutDataLabels) const {}
+#endif
 };
 
 UCLASS(BlueprintType, ClassGroup = (Procedural))
