[PCG] Wait on texture compilation - async cache alone is insufficient.

#rnx

diff --git a/Engine/Plugins/PCG/Source/PCG/Private/Data/PCGTextureData.cpp b/Engine/Plugins/PCG/Source/PCG/Private/Data/PCGTextureData.cpp
--- a/Engine/Plugins/PCG/Source/PCG/Private/Data/PCGTextureData.cpp
+++ b/Engine/Plugins/PCG/Source/PCG/Private/Data/PCGTextureData.cpp
@@ -32,7 +32,7 @@ namespace PCGTextureSamplingHelpers
 		}
 
 #if WITH_EDITOR
-		if (!Texture->IsAsyncCacheComplete())
+		if (!Texture->IsAsyncCacheComplete() || Texture->IsCompiling())
 		{
 			return {};
 		}
@@ -52,7 +52,7 @@ namespace PCGTextureSamplingHelpers
 		}
 
 #if WITH_EDITOR
-		if (!Texture->IsAsyncCacheComplete())
+		if (!Texture->IsAsyncCacheComplete() || Texture->IsCompiling())
 		{
 			return {};
 		}
@@ -597,7 +597,7 @@ TOptional<bool> UPCGTextureData::InitializeFromCPUTexture()
 	}
 
 #if WITH_EDITOR
-	if (!Texture2D->IsAsyncCacheComplete())
+	if (!Texture2D->IsAsyncCacheComplete() || Texture->IsCompiling())
 	{
 		// Wait until texture ready before interrogating it for access options.
 		return {};
@@ -878,7 +878,7 @@ TOptional<bool> UPCGTextureData::InitializeGPUTextureFromCPU()
 	}
 
 #if WITH_EDITOR
-	if (!TextureForReadback->IsAsyncCacheComplete())
+	if (!TextureForReadback->IsAsyncCacheComplete() || Texture->IsCompiling())
 	{
 		return {};
 	}
