[PCG] Revert added check for Texture->IsCompiling as it never becomes true in the texture data unit test and freezes the editor.

#rnx

diff --git a/Engine/Plugins/PCG/Source/PCG/Private/Data/PCGTextureData.cpp b/Engine/Plugins/PCG/Source/PCG/Private/Data/PCGTextureData.cpp
--- a/Engine/Plugins/PCG/Source/PCG/Private/Data/PCGTextureData.cpp
+++ b/Engine/Plugins/PCG/Source/PCG/Private/Data/PCGTextureData.cpp
@@ -32,7 +32,7 @@ namespace PCGTextureSamplingHelpers
 		}
 
 #if WITH_EDITOR
-		if (!Texture->IsAsyncCacheComplete() || Texture->IsCompiling())
+		if (!Texture->IsAsyncCacheComplete())
 		{
 			return {};
 		}
@@ -52,7 +52,7 @@ namespace PCGTextureSamplingHelpers
 		}
 
 #if WITH_EDITOR
-		if (!Texture->IsAsyncCacheComplete() || Texture->IsCompiling())
+		if (!Texture->IsAsyncCacheComplete())
 		{
 			return {};
 		}
@@ -597,7 +597,7 @@ TOptional<bool> UPCGTextureData::InitializeFromCPUTexture()
 	}
 
 #if WITH_EDITOR
-	if (!Texture2D->IsAsyncCacheComplete() || Texture->IsCompiling())
+	if (!Texture2D->IsAsyncCacheComplete())
 	{
 		// Wait until texture ready before interrogating it for access options.
 		return {};
@@ -878,7 +878,7 @@ TOptional<bool> UPCGTextureData::InitializeGPUTextureFromCPU()
 	}
 
 #if WITH_EDITOR
-	if (!TextureForReadback->IsAsyncCacheComplete() || Texture->IsCompiling())
+	if (!TextureForReadback->IsAsyncCacheComplete())
 	{
 		return {};
 	}
