[PCG] UPCGPointArrayData Parenting support (pcg.Experimental.EnablePointArrayDataParenting)
- First example of parenting optimization with PCGBoundsModifier (CopyPoints can now be avoided)

#rnx

diff --git a/Engine/Plugins/PCG/Source/PCG/Private/Data/PCGBasePointData.cpp b/Engine/Plugins/PCG/Source/PCG/Private/Data/PCGBasePointData.cpp
--- a/Engine/Plugins/PCG/Source/PCG/Private/Data/PCGBasePointData.cpp
+++ b/Engine/Plugins/PCG/Source/PCG/Private/Data/PCGBasePointData.cpp
@@ -726,8 +726,13 @@ FBox UPCGBasePointData::GetBounds() const
 	return Bounds;
 }
 
-void UPCGBasePointData::CopyPointsTo(UPCGBasePointData* To, int32 ReadStartIndex, int32 WriteStartIndex, int32 Count) const
+void UPCGBasePointData::CopyPropertiesTo(UPCGBasePointData* To, int32 ReadStartIndex, int32 WriteStartIndex, int32 Count, EPCGPointNativeProperties Properties) const
 {
+	if (EPCGPointNativeProperties::None == Properties)
+	{
+		return;
+	}
+
 	const TConstPCGValueRange<FTransform> FromTransformRange = GetConstTransformValueRange();
 	const TConstPCGValueRange<float> FromSteepnessRange = GetConstSteepnessValueRange();
 	const TConstPCGValueRange<float> FromDensityRange = GetConstDensityValueRange();
@@ -736,8 +741,8 @@ void UPCGBasePointData::CopyPointsTo(UPCGBasePointData* To, int32 ReadStartIndex
 	const TConstPCGValueRange<FVector4> FromColorRange = GetConstColorValueRange();
 	const TConstPCGValueRange<int64> FromMetadataEntryRange = GetConstMetadataEntryValueRange();
 	const TConstPCGValueRange<int32> FromSeedRange = GetConstSeedValueRange();
-	
-	To->AllocateProperties(GetAllocatedProperties());
+
+	To->AllocateProperties(Properties);
 
 	TPCGValueRange<FTransform> ToTransformRange = To->GetTransformValueRange(/*bAllocate=*/false);
 	TPCGValueRange<float> ToSteepnessRange = To->GetSteepnessValueRange(/*bAllocate=*/false);
@@ -750,24 +755,57 @@ void UPCGBasePointData::CopyPointsTo(UPCGBasePointData* To, int32 ReadStartIndex
 
 	check(ReadStartIndex + Count <= GetNumPoints() && WriteStartIndex + Count <= To->GetNumPoints());
 
-	for (int32 Index = 0; Index < Count; ++Index)
+	if (EPCGPointNativeProperties::All == Properties)
+	{
+		for (int32 Index = 0; Index < Count; ++Index)
+		{
+			const int32 ReadIndex = ReadStartIndex + Index;
+			const int32 WriteIndex = WriteStartIndex + Index;
+
+			ToTransformRange[WriteIndex] = FromTransformRange[ReadIndex];
+			ToDensityRange[WriteIndex] = FromDensityRange[ReadIndex];
+			ToBoundsMinRange[WriteIndex] = FromBoundsMinRange[ReadIndex];
+			ToBoundsMaxRange[WriteIndex] = FromBoundsMaxRange[ReadIndex];
+			ToColorRange[WriteIndex] = FromColorRange[ReadIndex];
+			ToSteepnessRange[WriteIndex] = FromSteepnessRange[ReadIndex];
+			ToSeedRange[WriteIndex] = FromSeedRange[ReadIndex];
+			ToMetadataEntryRange[WriteIndex] = FromMetadataEntryRange[ReadIndex];
+		}
+	}
+	else
 	{
-		const int32 ReadIndex = ReadStartIndex + Index;
-		const int32 WriteIndex = WriteStartIndex + Index;
-
-		ToTransformRange[WriteIndex] = FromTransformRange[ReadIndex];
-		ToDensityRange[WriteIndex] = FromDensityRange[ReadIndex];
-		ToBoundsMinRange[WriteIndex] = FromBoundsMinRange[ReadIndex];
-		ToBoundsMaxRange[WriteIndex] = FromBoundsMaxRange[ReadIndex];
-		ToColorRange[WriteIndex] = FromColorRange[ReadIndex];
-		ToSteepnessRange[WriteIndex] = FromSteepnessRange[ReadIndex];
-		ToSeedRange[WriteIndex] = FromSeedRange[ReadIndex];
-		ToMetadataEntryRange[WriteIndex] = FromMetadataEntryRange[ReadIndex];
+		auto CopyRangeIf = [ReadStartIndex, WriteStartIndex, Count, Properties]<typename T>(TPCGValueRange<T>& ToRange, const TConstPCGValueRange<T>& FromRange, EPCGPointNativeProperties Property)
+		{
+			if (EnumHasAllFlags(Properties, Property))
+			{
+				for (int32 Index = 0; Index < Count; ++Index)
+				{
+					const int32 ReadIndex = ReadStartIndex + Index;
+					const int32 WriteIndex = WriteStartIndex + Index;
+
+					ToRange[WriteIndex] = FromRange[ReadIndex];
+				}
+			}
+		};
+
+		CopyRangeIf(ToTransformRange, FromTransformRange, EPCGPointNativeProperties::Transform);
+		CopyRangeIf(ToDensityRange, FromDensityRange, EPCGPointNativeProperties::Density);
+		CopyRangeIf(ToBoundsMinRange, FromBoundsMinRange, EPCGPointNativeProperties::BoundsMin);
+		CopyRangeIf(ToBoundsMaxRange, FromBoundsMaxRange, EPCGPointNativeProperties::BoundsMax);
+		CopyRangeIf(ToColorRange, FromColorRange, EPCGPointNativeProperties::Color);
+		CopyRangeIf(ToSteepnessRange, FromSteepnessRange, EPCGPointNativeProperties::Steepness);
+		CopyRangeIf(ToSeedRange, FromSeedRange, EPCGPointNativeProperties::Seed);
+		CopyRangeIf(ToMetadataEntryRange, FromMetadataEntryRange, EPCGPointNativeProperties::MetadataEntry);
 	}
 
 	To->DirtyCache();
 }
 
+void UPCGBasePointData::CopyPointsTo(UPCGBasePointData* To, int32 ReadStartIndex, int32 WriteStartIndex, int32 Count) const
+{
+	CopyPropertiesTo(To, ReadStartIndex, WriteStartIndex, Count, EPCGPointNativeProperties::All);
+}
+
 void UPCGBasePointData::CopyPointsFrom(const UPCGBasePointData* InData, const TArray<int32>& InDataIndices)
 {
 	return CopyPoints(InData, this, InDataIndices, false);
@@ -804,14 +842,14 @@ void UPCGBasePointData::CopyPoints(const UPCGBasePointData* From, UPCGBasePointD
 	{
 		const int32 PointIndex = bCopyAll ? Index : InDataIndices[Index];
 
-		ToTransformRange[PointIndex] = FromTransformRange[PointIndex];
-		ToSteepnessRange[PointIndex] = FromSteepnessRange[PointIndex];
-		ToDensityRange[PointIndex] = FromDensityRange[PointIndex];
-		ToBoundsMinRange[PointIndex] = FromBoundsMinRange[PointIndex];
-		ToBoundsMaxRange[PointIndex] = FromBoundsMaxRange[PointIndex];
-		ToColorRange[PointIndex] = FromColorRange[PointIndex];
-		ToMetadataEntryRange[PointIndex] = FromMetadataEntryRange[PointIndex];
-		ToSeedRange[PointIndex] = FromSeedRange[PointIndex];
+		ToTransformRange[Index] = FromTransformRange[PointIndex];
+		ToSteepnessRange[Index] = FromSteepnessRange[PointIndex];
+		ToDensityRange[Index] = FromDensityRange[PointIndex];
+		ToBoundsMinRange[Index] = FromBoundsMinRange[PointIndex];
+		ToBoundsMaxRange[Index] = FromBoundsMaxRange[PointIndex];
+		ToColorRange[Index] = FromColorRange[PointIndex];
+		ToMetadataEntryRange[Index] = FromMetadataEntryRange[PointIndex];
+		ToSeedRange[Index] = FromSeedRange[PointIndex];
 	}
 
 	To->DirtyCache();
diff --git a/Engine/Plugins/PCG/Source/PCG/Private/Data/PCGPointArrayData.cpp b/Engine/Plugins/PCG/Source/PCG/Private/Data/PCGPointArrayData.cpp
--- a/Engine/Plugins/PCG/Source/PCG/Private/Data/PCGPointArrayData.cpp
+++ b/Engine/Plugins/PCG/Source/PCG/Private/Data/PCGPointArrayData.cpp
@@ -6,6 +6,11 @@
 
 #include "Data/PCGPointData.h"
 
+TAutoConsoleVariable<bool> CVarPCGEnablePointArrayDataParenting(
+	TEXT("pcg.Experimental.EnablePointArrayDataParenting"),
+	true,
+	TEXT("Whether to enable inheritance of data on PointArrayData (memory savings)"));
+
 void UPCGPointArrayData::GetResourceSizeEx(FResourceSizeEx& CumulativeResourceSize)
 {
 	Super::GetResourceSizeEx(CumulativeResourceSize);
@@ -13,6 +18,15 @@ void UPCGPointArrayData::GetResourceSizeEx(FResourceSizeEx& CumulativeResourceSi
 	CumulativeResourceSize.AddDedicatedSystemMemoryBytes(PointArray.GetSizeBytes());
 }
 
+void UPCGPointArrayData::VisitDataNetwork(TFunctionRef<void(const UPCGData*)> Action) const
+{
+	Super::VisitDataNetwork(Action);
+	if (ParentData)
+	{
+		ParentData->VisitDataNetwork(Action);
+	}
+}
+
 const UPCGPointData* UPCGPointArrayData::ToPointData(FPCGContext* Context, const FBox& InBounds) const
 {
 	UPCGPointData* PointData = FPCGContext::NewObject_AnyThread<UPCGPointData>(Context);
@@ -27,19 +41,262 @@ UPCGSpatialData* UPCGPointArrayData::CopyInternal(FPCGContext* Context) const
 {
 	UPCGPointArrayData* NewPointData = FPCGContext::NewObject_AnyThread<UPCGPointArrayData>(Context);
 	
-	NewPointData->PointArray = PointArray;
+	// If inheritance is supported we are going to inherit from this data in InitializeSpatialDataInternal
+	if (!SupportsSpatialDataInheritance())
+	{
+		NewPointData->PointArray = PointArray;
+	}
 
 	return NewPointData;
 }
 
-void UPCGPointArrayData::CopyPointsTo(UPCGBasePointData* OutData, int32 ReadStartIndex, int32 WriteStartIndex, int32 Count) const
+void UPCGPointArrayData::CopyPropertiesTo(UPCGBasePointData* To, int32 ReadStartIndex, int32 WriteStartIndex, int32 Count, EPCGPointNativeProperties Properties) const
 {
-	if (UPCGPointArrayData* PointArrayData = Cast<UPCGPointArrayData>(OutData))
+	if (UPCGPointArrayData* PointArrayData = Cast<UPCGPointArrayData>(To))
 	{
-		PointArray.CopyPointsTo(PointArrayData->PointArray, ReadStartIndex, WriteStartIndex, Count);
+		PointArrayData->AllocateProperties(GetAllocatedProperties());
+
+		if (EnumHasAllFlags(Properties, EPCGPointNativeProperties::Transform))
+		{
+			GetProperty<FTransform>(EPCGPointNativeProperties::Transform)->CopyTo(*PointArrayData->GetProperty<FTransform>(EPCGPointNativeProperties::Transform, /*bWithInheritance=*/false), ReadStartIndex, WriteStartIndex, Count);
+		}
+
+		if (EnumHasAllFlags(Properties, EPCGPointNativeProperties::Density))
+		{
+			GetProperty<float>(EPCGPointNativeProperties::Density)->CopyTo(*PointArrayData->GetProperty<float>(EPCGPointNativeProperties::Density, /*bWithInheritance=*/false), ReadStartIndex, WriteStartIndex, Count);
+		}
+
+		if (EnumHasAllFlags(Properties, EPCGPointNativeProperties::BoundsMin))
+		{
+			GetProperty<FVector>(EPCGPointNativeProperties::BoundsMin)->CopyTo(*PointArrayData->GetProperty<FVector>(EPCGPointNativeProperties::BoundsMin, /*bWithInheritance=*/false), ReadStartIndex, WriteStartIndex, Count);
+		}
+
+		if (EnumHasAllFlags(Properties, EPCGPointNativeProperties::BoundsMax))
+		{
+			GetProperty<FVector>(EPCGPointNativeProperties::BoundsMax)->CopyTo(*PointArrayData->GetProperty<FVector>(EPCGPointNativeProperties::BoundsMax, /*bWithInheritance=*/false), ReadStartIndex, WriteStartIndex, Count);
+		}
+
+		if (EnumHasAllFlags(Properties, EPCGPointNativeProperties::Color))
+		{
+			GetProperty<FVector4>(EPCGPointNativeProperties::Color)->CopyTo(*PointArrayData->GetProperty<FVector4>(EPCGPointNativeProperties::Color, /*bWithInheritance=*/false), ReadStartIndex, WriteStartIndex, Count);
+		}
+
+		if (EnumHasAllFlags(Properties, EPCGPointNativeProperties::Steepness))
+		{
+			GetProperty<float>(EPCGPointNativeProperties::Steepness)->CopyTo(*PointArrayData->GetProperty<float>(EPCGPointNativeProperties::Steepness, /*bWithInheritance=*/false), ReadStartIndex, WriteStartIndex, Count);
+		}
+
+		if (EnumHasAllFlags(Properties, EPCGPointNativeProperties::Seed))
+		{
+			GetProperty<int32>(EPCGPointNativeProperties::Seed)->CopyTo(*PointArrayData->GetProperty<int32>(EPCGPointNativeProperties::Seed, /*bWithInheritance=*/false), ReadStartIndex, WriteStartIndex, Count);
+		}
+
+		if (EnumHasAllFlags(Properties, EPCGPointNativeProperties::MetadataEntry))
+		{
+			GetProperty<int64>(EPCGPointNativeProperties::MetadataEntry)->CopyTo(*PointArrayData->GetProperty<int64>(EPCGPointNativeProperties::MetadataEntry, /*bWithInheritance=*/false), ReadStartIndex, WriteStartIndex, Count);
+		}
 	}
 	else
 	{
-		Super::CopyPointsTo(OutData, ReadStartIndex, WriteStartIndex, Count);
+		Super::CopyPropertiesTo(To, ReadStartIndex, WriteStartIndex, Count, Properties);
+	}
+}
+
+void UPCGPointArrayData::InitializeSpatialDataInternal(const FPCGInitializeFromDataParams& InParams)
+{
+	Super::InitializeSpatialDataInternal(InParams);
+		
+	if (const UPCGPointArrayData* SourceParentData = Cast<UPCGPointArrayData>(InParams.Source); InParams.bInheritSpatialData && SourceParentData && SupportsSpatialDataInheritance())
+	{
+		// Some nodes will call DuplicateData and then call InitializeFromData so it is possible we already have set the parent
+		check(ParentData == nullptr || ParentData == InParams.Source);
+
+		if (ParentData != InParams.Source)
+		{
+			SetNumPoints(SourceParentData->GetNumPoints());
+			InheritedProperties = EPCGPointNativeProperties::All;
+			ParentData = const_cast<UPCGPointArrayData*>(SourceParentData);
+		}
+	}
+}
+
+EPCGPointNativeProperties UPCGPointArrayData::GetAllocatedProperties(bool bWithInheritance) const
+{
+	EPCGPointNativeProperties AllocatedProperties = PointArray.GetAllocatedProperties();
+	if (bWithInheritance && ParentData)
+	{
+		AllocatedProperties |= ParentData->GetAllocatedProperties(bWithInheritance);
+	}
+	return AllocatedProperties;
+}
+
+bool UPCGPointArrayData::SupportsSpatialDataInheritance() const
+{
+	return CVarPCGEnablePointArrayDataParenting.GetValueOnAnyThread();
+}
+
+void UPCGPointArrayData::Flatten()
+{
+	Super::Flatten();
+
+	FlattenPropertiesIfNeeded();
+	
+	check(!ParentData);
+}
+
+void UPCGPointArrayData::FlattenPropertiesIfNeeded(EPCGPointNativeProperties Properties)
+{	
+	if (EnumHasAnyFlags(Properties, EPCGPointNativeProperties::Transform))
+	{
+		FlattenPropertyIfNeeded<FTransform>(EPCGPointNativeProperties::Transform);
+	}
+	
+	if (EnumHasAnyFlags(Properties, EPCGPointNativeProperties::Steepness))
+	{
+		FlattenPropertyIfNeeded<float>(EPCGPointNativeProperties::Steepness);
+	}
+
+	if (EnumHasAnyFlags(Properties, EPCGPointNativeProperties::BoundsMin))
+	{
+		FlattenPropertyIfNeeded<FVector>(EPCGPointNativeProperties::BoundsMin);
+	}
+
+	if (EnumHasAnyFlags(Properties, EPCGPointNativeProperties::BoundsMax))
+	{
+		FlattenPropertyIfNeeded<FVector>(EPCGPointNativeProperties::BoundsMax);
+	}
+
+	if (EnumHasAnyFlags(Properties, EPCGPointNativeProperties::Color))
+	{
+		FlattenPropertyIfNeeded<FVector4>(EPCGPointNativeProperties::Color);
+	}
+
+	if (EnumHasAnyFlags(Properties, EPCGPointNativeProperties::Density))
+	{
+		FlattenPropertyIfNeeded<float>(EPCGPointNativeProperties::Density);
+	}
+
+	if (EnumHasAnyFlags(Properties, EPCGPointNativeProperties::Seed))
+	{
+		FlattenPropertyIfNeeded<int32>(EPCGPointNativeProperties::Seed);
+	}
+
+	if (EnumHasAnyFlags(Properties, EPCGPointNativeProperties::MetadataEntry))
+	{
+		FlattenPropertyIfNeeded<int64>(EPCGPointNativeProperties::MetadataEntry);
 	}
+}
+
+void UPCGPointArrayData::SetNumPoints(int32 InNumPoints)
+{
+	if (ParentData && ParentData->GetNumPoints() != InNumPoints)
+	{
+		FlattenPropertiesIfNeeded();
+	}
+	PointArray.SetNumPoints(InNumPoints);
+}
+
+void UPCGPointArrayData::AllocateProperties(EPCGPointNativeProperties Properties)
+{
+	FlattenPropertiesIfNeeded(Properties);
+	PointArray.Allocate(Properties);
+}
+
+void UPCGPointArrayData::FreeProperties(EPCGPointNativeProperties Properties)
+{
+	FlattenPropertiesIfNeeded(Properties);
+	PointArray.Free(Properties);
+}
+
+void UPCGPointArrayData::MoveRange(int32 RangeStartIndex, int32 MoveToIndex, int32 NumElements)
+{
+	FlattenPropertiesIfNeeded();
+	PointArray.MoveRange(RangeStartIndex, MoveToIndex, NumElements);
+}
+
+TPCGValueRange<FTransform> UPCGPointArrayData::GetTransformValueRange(bool bAllocate)
+{ 
+	FlattenPropertiesIfNeeded(EPCGPointNativeProperties::Transform);
+	return PointArray.GetTransformValueRange(bAllocate); 
+}
+
+TPCGValueRange<float> UPCGPointArrayData::GetDensityValueRange(bool bAllocate)
+{ 
+	FlattenPropertiesIfNeeded(EPCGPointNativeProperties::Density);
+	return PointArray.GetDensityValueRange(bAllocate); 
+}
+
+TPCGValueRange<FVector> UPCGPointArrayData::GetBoundsMinValueRange(bool bAllocate)
+{ 
+	FlattenPropertiesIfNeeded(EPCGPointNativeProperties::BoundsMin);
+	return PointArray.GetBoundsMinValueRange(bAllocate); 
+}
+
+TPCGValueRange<FVector> UPCGPointArrayData::GetBoundsMaxValueRange(bool bAllocate)
+{
+	FlattenPropertiesIfNeeded(EPCGPointNativeProperties::BoundsMax);
+	return PointArray.GetBoundsMaxValueRange(bAllocate); 
+}
+
+TPCGValueRange<FVector4> UPCGPointArrayData::GetColorValueRange(bool bAllocate)
+{ 
+	FlattenPropertiesIfNeeded(EPCGPointNativeProperties::Color);
+	return PointArray.GetColorValueRange(bAllocate); 
+}
+
+TPCGValueRange<float> UPCGPointArrayData::GetSteepnessValueRange(bool bAllocate)
+{ 
+	FlattenPropertiesIfNeeded(EPCGPointNativeProperties::Steepness);
+	return PointArray.GetSteepnessValueRange(bAllocate); 
+}
+
+TPCGValueRange<int32> UPCGPointArrayData::GetSeedValueRange(bool bAllocate)
+{ 
+	FlattenPropertiesIfNeeded(EPCGPointNativeProperties::Seed);
+	return PointArray.GetSeedValueRange(bAllocate); 
+}
+TPCGValueRange<int64> UPCGPointArrayData::GetMetadataEntryValueRange(bool bAllocate)
+{ 
+	FlattenPropertiesIfNeeded(EPCGPointNativeProperties::MetadataEntry);
+	return PointArray.GetMetadataEntryValueRange(bAllocate); 
+}
+
+TConstPCGValueRange<FTransform> UPCGPointArrayData::GetConstTransformValueRange() const
+{ 
+	return GetProperty<FTransform>(EPCGPointNativeProperties::Transform)->GetConstValueRange();
+}
+
+TConstPCGValueRange<float> UPCGPointArrayData::GetConstDensityValueRange() const
+{ 
+	return GetProperty<float>(EPCGPointNativeProperties::Density)->GetConstValueRange();
+}
+
+TConstPCGValueRange<FVector> UPCGPointArrayData::GetConstBoundsMinValueRange() const
+{ 
+	return GetProperty<FVector>(EPCGPointNativeProperties::BoundsMin)->GetConstValueRange();
+}
+
+TConstPCGValueRange<FVector> UPCGPointArrayData::GetConstBoundsMaxValueRange() const
+{ 
+	return GetProperty<FVector>(EPCGPointNativeProperties::BoundsMax)->GetConstValueRange();
+}
+
+TConstPCGValueRange<FVector4> UPCGPointArrayData::GetConstColorValueRange() const
+{ 
+	return GetProperty<FVector4>(EPCGPointNativeProperties::Color)->GetConstValueRange();
+}
+
+TConstPCGValueRange<float> UPCGPointArrayData::GetConstSteepnessValueRange() const
+{ 
+	return GetProperty<float>(EPCGPointNativeProperties::Steepness)->GetConstValueRange();
+}
+
+TConstPCGValueRange<int32> UPCGPointArrayData::GetConstSeedValueRange() const
+{ 
+	return GetProperty<int32>(EPCGPointNativeProperties::Seed)->GetConstValueRange();
+}
+
+TConstPCGValueRange<int64> UPCGPointArrayData::GetConstMetadataEntryValueRange() const
+{ 
+	return GetProperty<int64>(EPCGPointNativeProperties::MetadataEntry)->GetConstValueRange();
 }
\ No newline at end of file
diff --git a/Engine/Plugins/PCG/Source/PCG/Private/Data/PCGSpatialData.cpp b/Engine/Plugins/PCG/Source/PCG/Private/Data/PCGSpatialData.cpp
--- a/Engine/Plugins/PCG/Source/PCG/Private/Data/PCGSpatialData.cpp
+++ b/Engine/Plugins/PCG/Source/PCG/Private/Data/PCGSpatialData.cpp
@@ -293,22 +293,17 @@ UPCGMetadata* UPCGSpatialData::CreateEmptyMetadata()
 	return Metadata;
 }
 
-void UPCGSpatialData::InitializeFromData(const UPCGSpatialData* InSource, const UPCGMetadata* InMetadataParentOverride, bool bInheritMetadata, bool bInheritAttributes)
+void UPCGSpatialData::InitializeMetadataInternal(const FPCGInitializeFromDataParams& InParams)
 {
-	if (InSource && TargetActor.IsExplicitlyNull())
-	{
-		TargetActor = InSource->TargetActor;
-	}
-
 	if (!Metadata)
 	{
 		Metadata = NewObject<UPCGMetadata>(this);
 	}
 
-	if (!bInheritMetadata || InMetadataParentOverride || InSource)
+	if (!InParams.bInheritMetadata || InParams.MetadataParentOverride || InParams.Source)
 	{
-		const UPCGMetadata* ParentMetadata = bInheritMetadata ? (InMetadataParentOverride ? InMetadataParentOverride : (InSource ? InSource->Metadata : nullptr)) : nullptr;
-		Metadata->Initialize(ParentMetadata, bInheritAttributes);
+		const UPCGMetadata* ParentMetadata = InParams.bInheritMetadata ? (InParams.MetadataParentOverride ? InParams.MetadataParentOverride : (InParams.Source ? InParams.Source->Metadata : nullptr)) : nullptr;
+		Metadata->Initialize(ParentMetadata, InParams.bInheritAttributes);
 	}
 	else
 	{
@@ -316,14 +311,40 @@ void UPCGSpatialData::InitializeFromData(const UPCGSpatialData* InSource, const
 	}
 }
 
+void UPCGSpatialData::InitializeSpatialDataInternal(const FPCGInitializeFromDataParams& InParams)
+{
+	if (InParams.Source && TargetActor.IsExplicitlyNull())
+	{
+		TargetActor = InParams.Source->TargetActor;
+	}
+}
+
+void UPCGSpatialData::InitializeFromData(const UPCGSpatialData* InSource, const UPCGMetadata* InMetadataParentOverride, bool bInheritMetadata, bool bInheritAttributes)
+{
+	FPCGInitializeFromDataParams Params(InSource);
+	Params.MetadataParentOverride = InMetadataParentOverride;
+	Params.bInheritMetadata = bInheritMetadata;
+	Params.bInheritAttributes = bInheritAttributes;
+	Params.bInheritSpatialData = true;
+	InitializeFromDataWithParams(MoveTemp(Params));
+}
+
+void UPCGSpatialData::InitializeFromDataWithParams(const FPCGInitializeFromDataParams& InParams)
+{	
+	InitializeSpatialDataInternal(InParams);
+	InitializeMetadataInternal(InParams);
+}
+
 UPCGSpatialData* UPCGSpatialData::DuplicateData(FPCGContext* Context, bool bInitializeMetadata) const
 {
 	UPCGSpatialData* NewSpatialData = CopyInternal(Context);
 	check(NewSpatialData);
 
+	NewSpatialData->InitializeSpatialDataInternal(this);
+
 	if (bInitializeMetadata)
 	{
-		NewSpatialData->InitializeFromData(this);
+		NewSpatialData->InitializeMetadataInternal(this);
 	}
 
 	if (bHasCachedLastSelector)
diff --git a/Engine/Plugins/PCG/Source/PCG/Private/Elements/PCGBoundsModifier.cpp b/Engine/Plugins/PCG/Source/PCG/Private/Elements/PCGBoundsModifier.cpp
--- a/Engine/Plugins/PCG/Source/PCG/Private/Elements/PCGBoundsModifier.cpp
+++ b/Engine/Plugins/PCG/Source/PCG/Private/Elements/PCGBoundsModifier.cpp
@@ -28,15 +28,9 @@ FText UPCGBoundsModifierSettings::GetNodeTooltipText() const
 }
 #endif // WITH_EDITOR
 
-EPCGPointNativeProperties FPCGBoundsModifier::GetPropertiesToAllocate(const UPCGBasePointData* Input) const
+EPCGPointNativeProperties FPCGBoundsModifier::GetPropertiesToAllocate(FPCGContext* InContext) const
 {
-	EPCGPointNativeProperties PropertiesToAllocate = Input->GetAllocatedProperties();
-	
-	PropertiesToAllocate |= EPCGPointNativeProperties::BoundsMin;
-	PropertiesToAllocate |= EPCGPointNativeProperties::BoundsMax;
-	PropertiesToAllocate |= EPCGPointNativeProperties::Steepness;
-
-	return PropertiesToAllocate;
+	return EPCGPointNativeProperties::BoundsMin | EPCGPointNativeProperties::BoundsMax | EPCGPointNativeProperties::Steepness;
 }
 
 bool FPCGBoundsModifier::ExecuteInternal(FPCGContext* Context) const
@@ -54,8 +48,6 @@ bool FPCGBoundsModifier::ExecuteInternal(FPCGContext* Context) const
 		case EPCGBoundsModifierMode::Intersect:
 			return ExecutePointOperation(BoundsModifierContext, [Settings](const UPCGBasePointData* InputData, UPCGBasePointData* OutputData, int32 StartIndex, int32 Count)
 			{
-				InputData->CopyPointsTo(OutputData, StartIndex, StartIndex, Count);
-
 				TPCGValueRange<FVector> BoundsMinRange = OutputData->GetBoundsMinValueRange();
 				TPCGValueRange<FVector> BoundsMaxRange = OutputData->GetBoundsMaxValueRange();
 				TPCGValueRange<float> SteepnessRange = OutputData->GetSteepnessValueRange();
@@ -81,8 +73,6 @@ bool FPCGBoundsModifier::ExecuteInternal(FPCGContext* Context) const
 		case EPCGBoundsModifierMode::Include:
 			return ExecutePointOperation(BoundsModifierContext, [Settings](const UPCGBasePointData* InputData, UPCGBasePointData* OutputData, int32 StartIndex, int32 Count)
 			{
-				InputData->CopyPointsTo(OutputData, StartIndex, StartIndex, Count);
-
 				TPCGValueRange<FVector> BoundsMinRange = OutputData->GetBoundsMinValueRange();
 				TPCGValueRange<FVector> BoundsMaxRange = OutputData->GetBoundsMaxValueRange();
 				TPCGValueRange<float> SteepnessRange = OutputData->GetSteepnessValueRange();
@@ -108,8 +98,6 @@ bool FPCGBoundsModifier::ExecuteInternal(FPCGContext* Context) const
 		case EPCGBoundsModifierMode::Translate:
 			return ExecutePointOperation(BoundsModifierContext, [Settings](const UPCGBasePointData* InputData, UPCGBasePointData* OutputData, int32 StartIndex, int32 Count)
 			{
-				InputData->CopyPointsTo(OutputData, StartIndex, StartIndex, Count);
-
 				TPCGValueRange<FVector> BoundsMinRange = OutputData->GetBoundsMinValueRange();
 				TPCGValueRange<FVector> BoundsMaxRange = OutputData->GetBoundsMaxValueRange();
 				TPCGValueRange<float> SteepnessRange = OutputData->GetSteepnessValueRange();
@@ -131,8 +119,6 @@ bool FPCGBoundsModifier::ExecuteInternal(FPCGContext* Context) const
 		case EPCGBoundsModifierMode::Scale:
 			return ExecutePointOperation(BoundsModifierContext, [Settings](const UPCGBasePointData* InputData, UPCGBasePointData* OutputData, int32 StartIndex, int32 Count)
 			{
-				InputData->CopyPointsTo(OutputData, StartIndex, StartIndex, Count);
-							
 				TPCGValueRange<FVector> BoundsMinRange = OutputData->GetBoundsMinValueRange();
 				TPCGValueRange<FVector> BoundsMaxRange = OutputData->GetBoundsMaxValueRange();
 				TPCGValueRange<float> SteepnessRange = OutputData->GetSteepnessValueRange();
@@ -154,8 +140,6 @@ bool FPCGBoundsModifier::ExecuteInternal(FPCGContext* Context) const
 		case EPCGBoundsModifierMode::Set:
 			return ExecutePointOperation(BoundsModifierContext, [Settings](const UPCGBasePointData* InputData, UPCGBasePointData* OutputData, int32 StartIndex, int32 Count)
 			{
-				InputData->CopyPointsTo(OutputData, StartIndex, StartIndex, Count);
-
 				TPCGValueRange<FVector> BoundsMinRange = OutputData->GetBoundsMinValueRange();
 				TPCGValueRange<FVector> BoundsMaxRange = OutputData->GetBoundsMaxValueRange();
 				TPCGValueRange<float> SteepnessRange = OutputData->GetSteepnessValueRange();
diff --git a/Engine/Plugins/PCG/Source/PCG/Private/Elements/PCGPointOperationElementBase.cpp b/Engine/Plugins/PCG/Source/PCG/Private/Elements/PCGPointOperationElementBase.cpp
--- a/Engine/Plugins/PCG/Source/PCG/Private/Elements/PCGPointOperationElementBase.cpp
+++ b/Engine/Plugins/PCG/Source/PCG/Private/Elements/PCGPointOperationElementBase.cpp
@@ -58,7 +58,17 @@ bool FPCGPointOperationElementBase::PreparePointOperationData(ContextType* InCon
 			OutState.OutputData = SupportsBasePointDataInputs(InContext) ? FPCGContext::NewPointData_AnyThread(InContext) : FPCGContext::NewObject_AnyThread<UPCGPointData>(InContext);
 			OutState.OutputData->InitializeFromData(OutState.InputData);
 			OutState.OutputData->SetNumPoints(OutState.NumPoints);
-			OutState.OutputData->AllocateProperties(GetPropertiesToAllocate(OutState.InputData));
+
+			// Allocate properties that we are going to modify
+			EPCGPointNativeProperties PropertiesToAllocate = GetPropertiesToAllocate(InContext);
+			
+			// If data doesn't support parenting also allocate properties we are going to copy from input
+			if (!OutState.OutputData->HasSpatialDataParent())
+			{
+				PropertiesToAllocate |= OutState.InputData->GetAllocatedProperties();
+			}
+
+			OutState.OutputData->AllocateProperties(PropertiesToAllocate);
 
 			Output.Data = OutState.OutputData;
 
@@ -75,7 +85,7 @@ PRAGMA_ENABLE_DEPRECATION_WARNINGS
 
 namespace PCGPointOperationElementBase
 {
-	bool ExecuteSliceOneToOne(FPCGPointOperationElementBase::ContextType* Context, const FPCGPointOperationElementBase::IterStateType& IterState, TFunctionRef<bool(const FPCGPoint& InPoint, FPCGPoint& OutPoint)> Callback, int32 PointsPerChunk)
+	bool ExecuteSliceOneToOne(FPCGPointOperationElementBase::ContextType* Context, const FPCGPointOperationElementBase::IterStateType& IterState, TFunctionRef<bool(const FPCGPoint& InPoint, FPCGPoint& OutPoint)> Callback, int32 PointsPerChunk, bool bShouldCopyPoints)
 	{
 PRAGMA_DISABLE_DEPRECATION_WARNINGS
 		check(IterState.InputPointData && IterState.OutputPointData);
@@ -84,8 +94,13 @@ PRAGMA_DISABLE_DEPRECATION_WARNINGS
 PRAGMA_ENABLE_DEPRECATION_WARNINGS
 
 		// Conversion lambda from index to point ref for ease of use
-		auto InternalPointFunction = [&Callback, &InPoints, &OutPoints](int32 ReadIndex, int32 WriteIndex) -> bool
+		auto InternalPointFunction = [&Callback, &InPoints, &OutPoints, bShouldCopyPoints](int32 ReadIndex, int32 WriteIndex) -> bool
 		{
+			if(bShouldCopyPoints)
+			{
+				OutPoints[WriteIndex] = InPoints[ReadIndex];
+			}
+
 			return Callback(InPoints[ReadIndex], OutPoints[WriteIndex]);
 		};
 
@@ -98,14 +113,19 @@ PRAGMA_ENABLE_DEPRECATION_WARNINGS
 			PointsPerChunk);
 	}
 
-	bool ExecuteSliceRange(FPCGPointOperationElementBase::ContextType* Context, const FPCGPointOperationElementBase::IterStateType& IterState, TFunctionRef<bool(const UPCGBasePointData* InPointData, UPCGBasePointData* OutPointData, int32 InStartIndex, int32 InCount)> Callback, int32 PointsPerChunk)
+	bool ExecuteSliceRange(FPCGPointOperationElementBase::ContextType* Context, const FPCGPointOperationElementBase::IterStateType& IterState, TFunctionRef<bool(const UPCGBasePointData* InPointData, UPCGBasePointData* OutPointData, int32 InStartIndex, int32 InCount)> Callback, int32 PointsPerChunk, bool bShouldCopyPoints)
 	{
 		const UPCGBasePointData* InputData = IterState.InputData;
 		UPCGBasePointData* OutputData = IterState.OutputData;
 
-		auto ProcessRangeFunc = [&Callback, &InputData, &OutputData](int32 StartReadIndex, int32 StartWriteIndex, int32 Count)
+		auto ProcessRangeFunc = [&Callback, &InputData, &OutputData, bShouldCopyPoints](int32 StartReadIndex, int32 StartWriteIndex, int32 Count)
 		{
 			check(StartReadIndex == StartWriteIndex);
+			if(bShouldCopyPoints && !OutputData->HasSpatialDataParent())
+			{
+				InputData->CopyPointsTo(OutputData, StartReadIndex, StartReadIndex, Count);
+			}
+						
 			Callback(InputData, OutputData, StartReadIndex, Count);
 			return Count;
 		};
@@ -122,7 +142,7 @@ PRAGMA_ENABLE_DEPRECATION_WARNINGS
 	}
 
 	template <typename Func>
-	bool ExecutePointOperation(const FPCGPointOperationElementBase* Element, FPCGPointOperationElementBase::ContextType* Context, Func Callback, int32 PointsPerChunk)
+	bool ExecutePointOperation(const FPCGPointOperationElementBase* Element, FPCGPointOperationElementBase::ContextType* Context, Func Callback, int32 PointsPerChunk, bool bShouldCopyPoints)
 	{
 		TRACE_CPUPROFILER_EVENT_SCOPE(FPCGPointOperationElementBase::ExecutePointProcessing);
 		check(Context);
@@ -134,7 +154,7 @@ PRAGMA_ENABLE_DEPRECATION_WARNINGS
 			return true;
 		}
 
-		return Element->ExecuteSlice(Context, [&Callback, &PointsPerChunk](FPCGPointOperationElementBase::ContextType* Context, const FPCGPointOperationElementBase::ExecStateType& ExecState, const FPCGPointOperationElementBase::IterStateType& IterState, const uint32 IterIndex)
+		return Element->ExecuteSlice(Context, [&Callback, &PointsPerChunk, &bShouldCopyPoints](FPCGPointOperationElementBase::ContextType* Context, const FPCGPointOperationElementBase::ExecStateType& ExecState, const FPCGPointOperationElementBase::IterStateType& IterState, const uint32 IterIndex)
 		{
 			// If this input created an error, result in no output
 			if (Context->GetIterationStateResult(IterIndex) != EPCGTimeSliceInitResult::Success)
@@ -146,11 +166,11 @@ PRAGMA_ENABLE_DEPRECATION_WARNINGS
 			bool bAsyncDone = false;
 			if constexpr (std::is_invocable_v<Func, const FPCGPoint&, FPCGPoint&>)
 			{
-				bAsyncDone = ExecuteSliceOneToOne(Context, IterState, Callback, PointsPerChunk);
+				bAsyncDone = ExecuteSliceOneToOne(Context, IterState, Callback, PointsPerChunk, bShouldCopyPoints);
 			}
 			else
 			{
-				bAsyncDone = ExecuteSliceRange(Context, IterState, Callback, PointsPerChunk);
+				bAsyncDone = ExecuteSliceRange(Context, IterState, Callback, PointsPerChunk, bShouldCopyPoints);
 			}
 
 			if (bAsyncDone)
@@ -165,12 +185,12 @@ PRAGMA_ENABLE_DEPRECATION_WARNINGS
 
 bool FPCGPointOperationElementBase::ExecutePointOperationWithPoints(ContextType* Context, TFunctionRef<bool(const FPCGPoint& InPoint, FPCGPoint& OutPoint)> Callback, int32 PointsPerChunk) const
 {
-	return PCGPointOperationElementBase::ExecutePointOperation(this, Context, Callback, PointsPerChunk);
+	return PCGPointOperationElementBase::ExecutePointOperation(this, Context, Callback, PointsPerChunk, ShouldCopyPoints());
 }
 
 bool FPCGPointOperationElementBase::ExecutePointOperationWithIndices(ContextType* Context, TFunctionRef<bool(const UPCGBasePointData* InPointData, UPCGBasePointData* OutPointData, int32 InStartIndex, int32 InCount)> Callback, int32 PointsPerChunk) const
 {
-	return PCGPointOperationElementBase::ExecutePointOperation(this, Context, Callback, PointsPerChunk);
+	return PCGPointOperationElementBase::ExecutePointOperation(this, Context, Callback, PointsPerChunk, ShouldCopyPoints());
 }
 
 #undef LOCTEXT_NAMESPACE
diff --git a/Engine/Plugins/PCG/Source/PCG/Private/Graph/PCGGraphExecutor.cpp b/Engine/Plugins/PCG/Source/PCG/Private/Graph/PCGGraphExecutor.cpp
--- a/Engine/Plugins/PCG/Source/PCG/Private/Graph/PCGGraphExecutor.cpp
+++ b/Engine/Plugins/PCG/Source/PCG/Private/Graph/PCGGraphExecutor.cpp
@@ -683,17 +683,18 @@ FPCGTaskId FPCGGraphExecutor::ScheduleGeneric(TFunction<bool()> InOperation, TFu
 		/*TaskDataDependencies=*/{});
 }
 
-FPCGTaskId FPCGGraphExecutor::ScheduleGenericWithContext(TFunction<bool(FPCGContext*)> InOperation, UPCGComponent* InSourceComponent, const TArray<FPCGTaskId>& TaskExecutionDependencies, const TArray<FPCGTaskId>& TaskDataDependencies)
+FPCGTaskId FPCGGraphExecutor::ScheduleGenericWithContext(TFunction<bool(FPCGContext*)> InOperation, UPCGComponent* InSourceComponent, const TArray<FPCGTaskId>& TaskExecutionDependencies, const TArray<FPCGTaskId>& TaskDataDependencies, bool bSupportBasePointDataInput)
 {
 	return ScheduleGenericWithContext(
 		InOperation,
 		TFunction<void(FPCGContext*)>(),
 		InSourceComponent,
 		TaskExecutionDependencies,
-		TaskDataDependencies);
+		TaskDataDependencies,
+		bSupportBasePointDataInput);
 }
 
-FPCGTaskId FPCGGraphExecutor::ScheduleGenericWithContext(TFunction<bool(FPCGContext*)> InOperation, TFunction<void(FPCGContext*)> InAbortOperation, UPCGComponent* InSourceComponent, const TArray<FPCGTaskId>& TaskExecutionDependencies, const TArray<FPCGTaskId>& TaskDataDependencies)
+FPCGTaskId FPCGGraphExecutor::ScheduleGenericWithContext(TFunction<bool(FPCGContext*)> InOperation, TFunction<void(FPCGContext*)> InAbortOperation, UPCGComponent* InSourceComponent, const TArray<FPCGTaskId>& TaskExecutionDependencies, const TArray<FPCGTaskId>& TaskDataDependencies, bool bSupportBasePointDataInput)
 {
 	// Build task & element to hold the operation to perform
 	FPCGGraphTask Task;
@@ -711,7 +712,7 @@ FPCGTaskId FPCGGraphExecutor::ScheduleGenericWithContext(TFunction<bool(FPCGCont
 	}
 
 	Task.SourceComponent = InSourceComponent;
-	Task.Element = MakeShared<FPCGGenericElement>(InOperation, InAbortOperation);
+	Task.Element = MakeShared<FPCGGenericElement>(InOperation, InAbortOperation, bSupportBasePointDataInput);
 
 	{
 		PCGGraphExecutor::TScopeLock ScopeLock(ScheduleLock);
@@ -2591,6 +2592,14 @@ FPCGGenericElement::FPCGGenericElement(TFunction<bool(FPCGContext*)> InOperation
 {
 }
 
+FPCGGenericElement::FPCGGenericElement(TFunction<bool(FPCGContext*)> InOperation, TFunction<void(FPCGContext*)> InAbortOperation, bool bInSupportBasePointDataInput, const FContextAllocator& InContextAllocator)
+	: Operation(InOperation)
+	, AbortOperation(InAbortOperation)
+	, bSupportBasePointDataInput(bInSupportBasePointDataInput)
+	, ContextAllocator(InContextAllocator)
+{
+}
+
 FPCGContext* FPCGGenericElement::Initialize(const FPCGDataCollection& InputData, TWeakObjectPtr<UPCGComponent> SourceComponent, const UPCGNode* Node)
 {
 	check(ContextAllocator);
diff --git a/Engine/Plugins/PCG/Source/PCG/Private/Graph/PCGGraphExecutor.h b/Engine/Plugins/PCG/Source/PCG/Private/Graph/PCGGraphExecutor.h
--- a/Engine/Plugins/PCG/Source/PCG/Private/Graph/PCGGraphExecutor.h
+++ b/Engine/Plugins/PCG/Source/PCG/Private/Graph/PCGGraphExecutor.h
@@ -86,21 +86,24 @@ public:
 	FPCGTaskId ScheduleGeneric(TFunction<bool()> InOperation, TFunction<void()> InAbortOperation, UPCGComponent* InSourceComponent, const TArray<FPCGTaskId>& TaskExecutionDependencies);
 
 	/** General job scheduling
-	*  @param InOperation:               Callback that takes a Context as argument and returns true if the task is done, false otherwise
-	*  @param InSourceComponent:         PCG component associated with this task. Can be null.
-	*  @param TaskExecutionDependencies: Task will wait on these tasks to execute and won't take their output data as input.
-	*  @param TaskDataDependencies:      Task will wait on these tasks to execute and will take their output data as input.
+	*  @param InOperation:                Callback that takes a Context as argument and returns true if the task is done, false otherwise
+	*  @param InSourceComponent:          PCG component associated with this task. Can be null.
+	*  @param TaskExecutionDependencies:  Task will wait on these tasks to execute and won't take their output data as input.
+	*  @param TaskDataDependencies:       Task will wait on these tasks to execute and will take their output data as input.
+	*  @param bSupportBasePointDataInput: When true, generic element will not convert input to UPCGPointData, this is false by default to preserver backward compatibility.
+	* 
 	*/
-	FPCGTaskId ScheduleGenericWithContext(TFunction<bool(FPCGContext*)> InOperation, UPCGComponent* InSourceComponent, const TArray<FPCGTaskId>& TaskExecutionDependencies, const TArray<FPCGTaskId>& TaskDataDependencies);
+	FPCGTaskId ScheduleGenericWithContext(TFunction<bool(FPCGContext*)> InOperation, UPCGComponent* InSourceComponent, const TArray<FPCGTaskId>& TaskExecutionDependencies, const TArray<FPCGTaskId>& TaskDataDependencies, bool bSupportBasePointDataInput = false);
 
 	/** General job scheduling
-	*  @param InOperation:               Callback that takes a Context as argument and returns true if the task is done, false otherwise
-	*  @param InAbortOperation:          Callback that is called if the task is aborted (cancelled) before fully executed.
-	*  @param InSourceComponent:         PCG component associated with this task. Can be null.
-	*  @param TaskExecutionDependencies: Task will wait on these tasks to execute and won't take their output data as input.
-	*  @param TaskDataDependencies:      Task will wait on these tasks to execute and will take their output data as input.
+	*  @param InOperation:                Callback that takes a Context as argument and returns true if the task is done, false otherwise
+	*  @param InAbortOperation:           Callback that is called if the task is aborted (cancelled) before fully executed.
+	*  @param InSourceComponent:          PCG component associated with this task. Can be null.
+	*  @param TaskExecutionDependencies:  Task will wait on these tasks to execute and won't take their output data as input.
+	*  @param TaskDataDependencies:       Task will wait on these tasks to execute and will take their output data as input.
+	*  @param bSupportBasePointDataInput: When true, generic element will not convert input to UPCGPointData, this is false by default to preserver backward compatibility.
 	*/
-	FPCGTaskId ScheduleGenericWithContext(TFunction<bool(FPCGContext*)> InOperation, TFunction<void(FPCGContext*)> InAbortOperation, UPCGComponent* InSourceComponent, const TArray<FPCGTaskId>& TaskExecutionDependencies, const TArray<FPCGTaskId>& TaskDataDependencies);
+	FPCGTaskId ScheduleGenericWithContext(TFunction<bool(FPCGContext*)> InOperation, TFunction<void(FPCGContext*)> InAbortOperation, UPCGComponent* InSourceComponent, const TArray<FPCGTaskId>& TaskExecutionDependencies, const TArray<FPCGTaskId>& TaskDataDependencies, bool bSupportBasePointDataInput = false);
 
 	/** Gets data in the output results. Returns false if data is not ready. */
 	bool GetOutputData(FPCGTaskId InTaskId, FPCGDataCollection& OutData);
@@ -358,6 +361,7 @@ public:
 protected:
 	virtual bool ExecuteInternal(FPCGContext* Context) const override;
 	virtual bool IsPassthrough(const UPCGSettings* InSettings) const override { return true; }
+	virtual bool SupportsBasePointDataInputs(FPCGContext* InContext) const override { return true; }
 };
 
 class FPCGPreGraphElement : public IPCGElement
@@ -368,6 +372,7 @@ public:
 protected:
 	virtual bool ExecuteInternal(FPCGContext* Context) const override;
 	virtual bool IsPassthrough(const UPCGSettings* InSettings) const override { return true; }
+	virtual bool SupportsBasePointDataInputs(FPCGContext* InContext) const override { return true; }
 };
 
 class FPCGGenericElement : public IPCGElement
@@ -389,6 +394,15 @@ public:
 	{
 		return new FPCGContext();
 	});
+
+	FPCGGenericElement(
+		TFunction<bool(FPCGContext*)> InOperation,
+		TFunction<void(FPCGContext*)> InAbortOperation,
+		bool bInSupportBasePointDataInput,
+		const FContextAllocator& InContextAllocator = (FContextAllocator)[](const FPCGDataCollection&, TWeakObjectPtr<UPCGComponent>, const UPCGNode*)
+	{
+		return new FPCGContext();
+	});
 	
 	virtual FPCGContext* Initialize(const FPCGDataCollection& InputData, TWeakObjectPtr<UPCGComponent> SourceComponent, const UPCGNode* Node) override;
 
@@ -401,6 +415,7 @@ protected:
 	virtual bool ExecuteInternal(FPCGContext* Context) const override;
 	virtual void AbortInternal(FPCGContext* Context) const override;
 	virtual bool IsCancellable() const override { return false; }
+	virtual bool SupportsBasePointDataInputs(FPCGContext* InContext) const override { return bSupportBasePointDataInput; }
 
 #if WITH_EDITOR
 	virtual bool ShouldLog() const override { return false; }
@@ -409,6 +424,7 @@ protected:
 private:
 	TFunction<bool(FPCGContext*)> Operation;
 	TFunction<void(FPCGContext*)> AbortOperation;
+	bool bSupportBasePointDataInput = false;
 
 	/** Creates a context object for this element. */
 	FContextAllocator ContextAllocator;
diff --git a/Engine/Plugins/PCG/Source/PCG/Private/PCGPointArray.cpp b/Engine/Plugins/PCG/Source/PCG/Private/PCGPointArray.cpp
--- a/Engine/Plugins/PCG/Source/PCG/Private/PCGPointArray.cpp
+++ b/Engine/Plugins/PCG/Source/PCG/Private/PCGPointArray.cpp
@@ -151,16 +151,47 @@ bool FPCGPointArray::Serialize(FArchive& Ar)
 	return true;
 }
 
-void FPCGPointArray::CopyPointsTo(FPCGPointArray& OutPointArray, int32 StartReadIndex, int32 StartWriteIndex, int32 Count) const
+void FPCGPointArray::CopyPropertiesTo(FPCGPointArray& OutPointArray, int32 StartReadIndex, int32 StartWriteIndex, int32 Count, EPCGPointNativeProperties Properties) const
 {
-	Transform.CopyTo(OutPointArray.Transform, StartReadIndex, StartWriteIndex, Count);
-	Density.CopyTo(OutPointArray.Density, StartReadIndex, StartWriteIndex, Count);
-	BoundsMin.CopyTo(OutPointArray.BoundsMin, StartReadIndex, StartWriteIndex, Count);
-	BoundsMax.CopyTo(OutPointArray.BoundsMax, StartReadIndex, StartWriteIndex, Count);
-	Color.CopyTo(OutPointArray.Color, StartReadIndex, StartWriteIndex, Count);
-	Steepness.CopyTo(OutPointArray.Steepness, StartReadIndex, StartWriteIndex, Count);
-	Seed.CopyTo(OutPointArray.Seed, StartReadIndex, StartWriteIndex, Count);
-	MetadataEntry.CopyTo(OutPointArray.MetadataEntry, StartReadIndex, StartWriteIndex, Count);
+	if (EnumHasAllFlags(Properties, EPCGPointNativeProperties::Transform))
+	{
+		Transform.CopyTo(OutPointArray.Transform, StartReadIndex, StartWriteIndex, Count);
+	}
+
+	if (EnumHasAllFlags(Properties, EPCGPointNativeProperties::Density))
+	{
+		Density.CopyTo(OutPointArray.Density, StartReadIndex, StartWriteIndex, Count);
+	}
+
+	if (EnumHasAllFlags(Properties, EPCGPointNativeProperties::BoundsMin))
+	{
+		BoundsMin.CopyTo(OutPointArray.BoundsMin, StartReadIndex, StartWriteIndex, Count);
+	}
+
+	if (EnumHasAllFlags(Properties, EPCGPointNativeProperties::BoundsMax))
+	{
+		BoundsMax.CopyTo(OutPointArray.BoundsMax, StartReadIndex, StartWriteIndex, Count);
+	}
+
+	if (EnumHasAllFlags(Properties, EPCGPointNativeProperties::Color))
+	{
+		Color.CopyTo(OutPointArray.Color, StartReadIndex, StartWriteIndex, Count);
+	}
+
+	if (EnumHasAllFlags(Properties, EPCGPointNativeProperties::Steepness))
+	{
+		Steepness.CopyTo(OutPointArray.Steepness, StartReadIndex, StartWriteIndex, Count);
+	}
+
+	if (EnumHasAllFlags(Properties, EPCGPointNativeProperties::Seed))
+	{
+		Seed.CopyTo(OutPointArray.Seed, StartReadIndex, StartWriteIndex, Count);
+	}
+
+	if (EnumHasAllFlags(Properties, EPCGPointNativeProperties::MetadataEntry))
+	{
+		MetadataEntry.CopyTo(OutPointArray.MetadataEntry, StartReadIndex, StartWriteIndex, Count);
+	}
 }
 
 EPCGPointNativeProperties FPCGPointArray::GetAllocatedProperties() const
diff --git a/Engine/Plugins/PCG/Source/PCG/Private/PCGSubsystem.cpp b/Engine/Plugins/PCG/Source/PCG/Private/PCGSubsystem.cpp
--- a/Engine/Plugins/PCG/Source/PCG/Private/PCGSubsystem.cpp
+++ b/Engine/Plugins/PCG/Source/PCG/Private/PCGSubsystem.cpp
@@ -693,7 +693,7 @@ FPCGTaskId UPCGSubsystem::ScheduleComponent(UPCGComponent* PCGComponent, EPCGHiG
 			}
 
 			return true;
-		}, PCGComponent, ExecutionDependencyTasks, DataDependencyTasks);
+		}, PCGComponent, ExecutionDependencyTasks, DataDependencyTasks, /*bSupportBasePointDataInput=*/true);
 	}
 	else
 	{
diff --git a/Engine/Plugins/PCG/Source/PCG/Private/Tests/Data/PCGPointArrayDataTest.cpp b/Engine/Plugins/PCG/Source/PCG/Private/Tests/Data/PCGPointArrayDataTest.cpp
new file mode 100644
--- /dev/null
+++ b/Engine/Plugins/PCG/Source/PCG/Private/Tests/Data/PCGPointArrayDataTest.cpp
@@ -0,0 +1,399 @@
+// Copyright Epic Games, Inc. All Rights Reserved.
+
+#include "Tests/PCGTestsCommon.h"
+#include "Data/PCGPointArrayData.h"
+
+#define PCG_FMT(Text, ...) FString::Printf(TEXT(Text), __VA_ARGS__)
+
+class FPCGPointArrayDataTestBaseClass : public FPCGTestBaseClass
+{
+public:
+	using FPCGTestBaseClass::FPCGTestBaseClass;
+
+	template <typename T>
+	bool TestEqualRange(const UPCGBasePointData* PointDataA, const UPCGBasePointData* PointDataB, const EPCGPointNativeProperties& NativeProperty)
+	{		
+		const TConstPCGValueRange<T> RangeA = PointDataA->GetConstValueRange<T>(NativeProperty);
+		const TConstPCGValueRange<T> RangeB = PointDataB->GetConstValueRange<T>(NativeProperty);
+
+		const FString PropertyName = StaticEnum<EPCGPointNativeProperties>()->GetNameStringByValue((int64)NativeProperty);
+		UTEST_EQUAL(PCG_FMT("Property %s : RangeA.Num() == RangeB.Num()", *PropertyName), RangeA.Num(), RangeB.Num());
+
+		for (int32 i = 0; i < RangeA.Num(); ++i)
+		{
+			UTEST_EQUAL(PCG_FMT("Property % s : RangeA[%d] == RangeB[%d]", *PropertyName, i, i), RangeA[i], RangeB[i]);
+		}
+
+		return true;
+	};
+
+	bool TestEqualData(const UPCGBasePointData* PointDataA, const UPCGBasePointData* PointDataB)
+	{
+		if (!TestEqualRange<FTransform>(PointDataA, PointDataB, EPCGPointNativeProperties::Transform))
+		{
+			return false;
+		}
+
+		if (!TestEqualRange<float>(PointDataA, PointDataB, EPCGPointNativeProperties::Density))
+		{
+			return false;
+		}
+
+		if (!TestEqualRange<FVector>(PointDataA, PointDataB, EPCGPointNativeProperties::BoundsMin))
+		{
+			return false;
+		}
+
+		if (!TestEqualRange<FVector>(PointDataA, PointDataB, EPCGPointNativeProperties::BoundsMax))
+		{
+			return false;
+		}
+
+		if (!TestEqualRange<FVector4>(PointDataA, PointDataB, EPCGPointNativeProperties::Color))
+		{
+			return false;
+		}
+
+		if (!TestEqualRange<float>(PointDataA, PointDataB, EPCGPointNativeProperties::Steepness))
+		{
+			return false;
+		}
+
+		if (!TestEqualRange<int32>(PointDataA, PointDataB, EPCGPointNativeProperties::Seed))
+		{
+			return false;
+		}
+
+		if (!TestEqualRange<int64>(PointDataA, PointDataB, EPCGPointNativeProperties::MetadataEntry))
+		{
+			return false;
+		}
+
+		return true;
+	}
+};
+
+IMPLEMENT_CUSTOM_SIMPLE_AUTOMATION_TEST(FPCGPointArrayDataInitializeFromDataTest, FPCGPointArrayDataTestBaseClass, "Plugins.PCG.PointArrayData.InitializeFromData", PCGTestsCommon::TestFlags)
+IMPLEMENT_CUSTOM_SIMPLE_AUTOMATION_TEST(FPCGPointArrayDataDuplicateDataTest, FPCGPointArrayDataTestBaseClass, "Plugins.PCG.PointArrayData.DuplicateData", PCGTestsCommon::TestFlags)
+IMPLEMENT_CUSTOM_SIMPLE_AUTOMATION_TEST(FPCGPointArrayDataFlattenTest, FPCGPointArrayDataTestBaseClass, "Plugins.PCG.PointArrayData.Flatten", PCGTestsCommon::TestFlags)
+IMPLEMENT_CUSTOM_SIMPLE_AUTOMATION_TEST(FPCGPointArrayDataMultiLevelInheritanceTest, FPCGPointArrayDataTestBaseClass, "Plugins.PCG.PointArrayData.MultiLevelInheritance", PCGTestsCommon::TestFlags)
+IMPLEMENT_CUSTOM_SIMPLE_AUTOMATION_TEST(FPCGPointArrayDataToPointDataTest, FPCGPointArrayDataTestBaseClass, "Plugins.PCG.PointArrayData.ToPointData", PCGTestsCommon::TestFlags)
+IMPLEMENT_CUSTOM_SIMPLE_AUTOMATION_TEST(FPCGPointArrayDataToPointArrayDataTest, FPCGPointArrayDataTestBaseClass, "Plugins.PCG.PointArrayData.ToPointArray", PCGTestsCommon::TestFlags)
+IMPLEMENT_CUSTOM_SIMPLE_AUTOMATION_TEST(FPCGPointArrayDataDisabledParentingTest, FPCGPointArrayDataTestBaseClass, "Plugins.PCG.PointArrayData.DisabledParenting", PCGTestsCommon::TestFlags)
+IMPLEMENT_CUSTOM_SIMPLE_AUTOMATION_TEST(FPCGPointArrayDataCopyPointsToFromPointArrayDataToPointDataTest, FPCGPointArrayDataTestBaseClass, "Plugins.PCG.PointArrayData.CopyPointsToFromPointArrayDataToPointData", PCGTestsCommon::TestFlags)
+IMPLEMENT_CUSTOM_SIMPLE_AUTOMATION_TEST(FPCGPointArrayDataCopyPointsToFromPointArrayDataToPointArrayDataTest, FPCGPointArrayDataTestBaseClass, "Plugins.PCG.PointArrayData.CopyPointsFromPointArrayDataToPointArrayData", PCGTestsCommon::TestFlags)
+IMPLEMENT_CUSTOM_SIMPLE_AUTOMATION_TEST(FPCGPointArrayDataCopyPointsToFromPointDataToPointDataTest, FPCGPointArrayDataTestBaseClass, "Plugins.PCG.PointArrayData.CopyPointsToFromPointDataToPointData", PCGTestsCommon::TestFlags)
+IMPLEMENT_CUSTOM_SIMPLE_AUTOMATION_TEST(FPCGPointArrayDataCopyPointsToFromPointDataToPointArrayDataTest, FPCGPointArrayDataTestBaseClass, "Plugins.PCG.PointArrayData.CopyPointsFromPointDataToPointArrayData", PCGTestsCommon::TestFlags)
+IMPLEMENT_CUSTOM_SIMPLE_AUTOMATION_TEST(FPCGPointArrayDataCopyPropertyOverInheritedPropertyTest, FPCGPointArrayDataTestBaseClass, "Plugins.PCG.PointArrayData.CopyPropertyOverInheritedProperty", PCGTestsCommon::TestFlags)
+
+bool FPCGPointArrayDataInitializeFromDataTest::RunTest(const FString& Parameters)
+{
+	// Enable Parenting
+	const bool bPreviousCVarValue = CVarPCGEnablePointArrayDataParenting.GetValueOnAnyThread();
+	CVarPCGEnablePointArrayDataParenting->Set(true, ECVF_SetByCode);
+	ON_SCOPE_EXIT
+	{
+		CVarPCGEnablePointArrayDataParenting->Set(bPreviousCVarValue, ECVF_SetByCode);
+	};
+
+	UPCGPointArrayData* ParentData = PCGTestsCommon::CreateRandomPointData<UPCGPointArrayData>(100, 42);
+	EPCGPointNativeProperties ParentAllocatedProperties = EPCGPointNativeProperties::Transform | EPCGPointNativeProperties::Seed | EPCGPointNativeProperties::Color | EPCGPointNativeProperties::Density;
+
+	// CreateRandomPointData only allocates those properties
+	TestEqual("Valid ParentAllocatedProperties", (int32)ParentAllocatedProperties, (int32)ParentData->GetAllocatedProperties());
+		
+	UPCGPointArrayData* ChildData = PCGTestsCommon::CreateEmptyPointData<UPCGPointArrayData>();
+	EPCGPointNativeProperties ChildAllocatedProperties = EPCGPointNativeProperties::None;
+
+	// CreateEmptyPointData should not allocate any properties
+	TestEqual("Valid ChildAllocatedProperties", (int32)ChildAllocatedProperties, (int32)ChildData->GetAllocatedProperties());
+	TestTrue("ChildData has 0 points", ChildData->GetNumPoints() == 0);
+
+	ChildData->InitializeFromData(ParentData);
+	TestEqual("ChildData->GetAllocatedProperties() == ParentData->GetAllocatedProperties()", (int32)ChildData->GetAllocatedProperties(), (int32)ParentData->GetAllocatedProperties());
+	TestTrue("ChildData->GetNumPoints() == ParentData->GetNumPoints()", ChildData->GetNumPoints() == ParentData->GetNumPoints());
+	TestTrue("ChildData has parent", ChildData->HasSpatialDataParent());
+	TestFalse("ParentData has no parent", ParentData->HasSpatialDataParent());
+
+	// Compare data
+	return TestEqualData(ParentData, ChildData);
+}
+
+bool FPCGPointArrayDataDuplicateDataTest::RunTest(const FString& Parameters)
+{
+	// Enable Parenting
+	const bool bPreviousCVarValue = CVarPCGEnablePointArrayDataParenting.GetValueOnAnyThread();
+	CVarPCGEnablePointArrayDataParenting->Set(true, ECVF_SetByCode);
+	ON_SCOPE_EXIT
+	{
+		CVarPCGEnablePointArrayDataParenting->Set(bPreviousCVarValue, ECVF_SetByCode);
+	};
+
+	UPCGPointArrayData* ParentData = PCGTestsCommon::CreateRandomPointData<UPCGPointArrayData>(100, 42);
+	EPCGPointNativeProperties ParentAllocatedProperties = EPCGPointNativeProperties::Transform | EPCGPointNativeProperties::Seed | EPCGPointNativeProperties::Color | EPCGPointNativeProperties::Density;
+
+	// CreateRandomPointData only allocates those properties
+	TestEqual("Valid ParentAllocatedProperties", (int32)ParentAllocatedProperties, (int32)ParentData->GetAllocatedProperties());
+
+	UPCGPointArrayData* ChildData = CastChecked<UPCGPointArrayData>(ParentData->DuplicateData(nullptr));
+		
+	TestEqual("ChildData->GetAllocatedProperties() == ParentData->GetAllocatedProperties()", (int32)ChildData->GetAllocatedProperties(), (int32)ParentData->GetAllocatedProperties());
+	TestTrue("ChildData->GetNumPoints() == ParentData->GetNumPoints()", ChildData->GetNumPoints() == ParentData->GetNumPoints());
+	TestTrue("ChildData has parent", ChildData->HasSpatialDataParent());
+	TestFalse("ParentData has no parent", ParentData->HasSpatialDataParent());
+
+	// Compare data
+	return TestEqualData(ParentData, ChildData);
+}
+
+bool FPCGPointArrayDataFlattenTest::RunTest(const FString& Parameters)
+{
+	// Enable Parenting
+	const bool bPreviousCVarValue = CVarPCGEnablePointArrayDataParenting.GetValueOnAnyThread();
+	CVarPCGEnablePointArrayDataParenting->Set(true, ECVF_SetByCode);
+	ON_SCOPE_EXIT
+	{
+		CVarPCGEnablePointArrayDataParenting->Set(bPreviousCVarValue, ECVF_SetByCode);
+	};
+
+	UPCGPointArrayData* ParentData = PCGTestsCommon::CreateRandomPointData<UPCGPointArrayData>(100, 42);
+	EPCGPointNativeProperties ParentAllocatedProperties = EPCGPointNativeProperties::Transform | EPCGPointNativeProperties::Seed | EPCGPointNativeProperties::Color | EPCGPointNativeProperties::Density;
+
+	// CreateRandomPointData only allocates those properties
+	TestEqual("Valid ParentAllocatedProperties", (int32)ParentAllocatedProperties, (int32)ParentData->GetAllocatedProperties());
+
+	UPCGPointArrayData* ChildData = PCGTestsCommon::CreateEmptyPointData<UPCGPointArrayData>();
+	EPCGPointNativeProperties ChildAllocatedProperties = EPCGPointNativeProperties::None;
+
+	// CreateEmptyPointData should not allocate any properties
+	TestEqual("Valid ChildAllocatedProperties", (int32)ChildAllocatedProperties, (int32)ChildData->GetAllocatedProperties());
+	TestTrue("ChildData has 0 points", ChildData->GetNumPoints() == 0);
+
+	ChildData->InitializeFromData(ParentData);
+	TestEqual("ChildData->GetAllocatedProperties() == ParentData->GetAllocatedProperties()", (int32)ChildData->GetAllocatedProperties(), (int32)ParentData->GetAllocatedProperties());
+	TestTrue("ChildData->GetNumPoints() == ParentData->GetNumPoints()", ChildData->GetNumPoints() == ParentData->GetNumPoints());
+	TestTrue("ChildData has parent", ChildData->HasSpatialDataParent());
+	TestFalse("ParentData has no parent", ParentData->HasSpatialDataParent());
+
+	// This will copy parent properties into child memory, at this point we should no longer inherit from parent
+	ChildData->Flatten();
+	TestFalse("ChildData has parent", ChildData->HasSpatialDataParent());
+
+	// Compare data
+	return TestEqualData(ParentData, ChildData);
+}
+
+bool FPCGPointArrayDataMultiLevelInheritanceTest::RunTest(const FString& Parameters)
+{
+	// Enable Parenting
+	const bool bPreviousCVarValue = CVarPCGEnablePointArrayDataParenting.GetValueOnAnyThread();
+	CVarPCGEnablePointArrayDataParenting->Set(true, ECVF_SetByCode);
+	ON_SCOPE_EXIT
+	{
+		CVarPCGEnablePointArrayDataParenting->Set(bPreviousCVarValue, ECVF_SetByCode);
+	};
+
+	UPCGPointArrayData* ParentData = PCGTestsCommon::CreateRandomPointData<UPCGPointArrayData>(100, 42);
+	EPCGPointNativeProperties ParentAllocatedProperties = EPCGPointNativeProperties::Transform | EPCGPointNativeProperties::Seed | EPCGPointNativeProperties::Color | EPCGPointNativeProperties::Density;
+
+	// CreateRandomPointData only allocates those properties
+	TestEqual("Valid ParentAllocatedProperties", (int32)ParentAllocatedProperties, (int32)ParentData->GetAllocatedProperties());
+
+	UPCGPointArrayData* ChildData = PCGTestsCommon::CreateEmptyPointData<UPCGPointArrayData>();
+	EPCGPointNativeProperties ChildAllocatedProperties = EPCGPointNativeProperties::None;
+
+	// CreateEmptyPointData should not allocate any properties
+	TestEqual("Valid ChildAllocatedProperties", (int32)ChildAllocatedProperties, (int32)ChildData->GetAllocatedProperties());
+	TestTrue("ChildData has 0 points", ChildData->GetNumPoints() == 0);
+
+	ChildData->InitializeFromData(ParentData);
+	TestEqual("ChildData->GetAllocatedProperties() == ParentData->GetAllocatedProperties()", (int32)ChildData->GetAllocatedProperties(), (int32)ParentData->GetAllocatedProperties());
+	TestTrue("ChildData->GetNumPoints() == ParentData->GetNumPoints()", ChildData->GetNumPoints() == ParentData->GetNumPoints());
+	TestTrue("ChildData has parent", ChildData->HasSpatialDataParent());
+	TestFalse("ParentData has no parent", ParentData->HasSpatialDataParent());
+	
+	UPCGPointArrayData* GrandChildData = CastChecked<UPCGPointArrayData>(ChildData->DuplicateData(nullptr));
+	
+	TestEqual("GrandChildData->GetAllocatedProperties() == ChildData->GetAllocatedProperties()", (int32)GrandChildData->GetAllocatedProperties(), (int32)ChildData->GetAllocatedProperties());
+	TestTrue("GrandChildData->GetNumPoints() == ChildData->GetNumPoints()", GrandChildData->GetNumPoints() == ChildData->GetNumPoints());
+	TestTrue("GrandChildData has parent", GrandChildData->HasSpatialDataParent());
+	TestTrue("ChildData has parent", ChildData->HasSpatialDataParent());
+
+	// Compare data
+	if (!TestEqualData(ParentData, ChildData))
+	{
+		return false;
+	}
+
+	if (!TestEqualData(ParentData, GrandChildData))
+	{
+		return false;
+	}
+
+	return TestEqualData(ChildData, GrandChildData);
+}
+
+bool FPCGPointArrayDataToPointDataTest::RunTest(const FString& Parameters)
+{
+	UPCGPointArrayData* PointArrayData = PCGTestsCommon::CreateRandomPointData<UPCGPointArrayData>(100, 42);
+	const UPCGPointData* PointData = PointArrayData->ToPointData(nullptr);
+
+	return TestEqualData(PointArrayData, PointData);
+}
+
+bool FPCGPointArrayDataToPointArrayDataTest::RunTest(const FString& Parameters)
+{
+	UPCGPointData* PointData = PCGTestsCommon::CreateRandomPointData<UPCGPointData>(100, 42);
+	const UPCGPointArrayData* PointArrayData = PointData->ToPointArrayData(nullptr);
+
+	return TestEqualData(PointArrayData, PointData);
+}
+
+bool FPCGPointArrayDataDisabledParentingTest::RunTest(const FString& Parameters)
+{
+	// Disable Parenting
+	const bool bPreviousCVarValue = CVarPCGEnablePointArrayDataParenting.GetValueOnAnyThread();
+	CVarPCGEnablePointArrayDataParenting->Set(false, ECVF_SetByCode);
+	ON_SCOPE_EXIT
+	{
+		CVarPCGEnablePointArrayDataParenting->Set(bPreviousCVarValue, ECVF_SetByCode);
+	};
+
+	UPCGPointArrayData* ParentData = PCGTestsCommon::CreateRandomPointData<UPCGPointArrayData>(100, 42);
+	EPCGPointNativeProperties ParentAllocatedProperties = EPCGPointNativeProperties::Transform | EPCGPointNativeProperties::Seed | EPCGPointNativeProperties::Color | EPCGPointNativeProperties::Density;
+
+	// CreateRandomPointData only allocates those properties
+	TestEqual("Valid ParentAllocatedProperties", (int32)ParentAllocatedProperties, (int32)ParentData->GetAllocatedProperties());
+
+	UPCGPointArrayData* ChildDataA = PCGTestsCommon::CreateEmptyPointData<UPCGPointArrayData>();
+	EPCGPointNativeProperties ChildAllocatedProperties = EPCGPointNativeProperties::None;
+
+	// CreateEmptyPointData should not allocate any properties
+	TestEqual("Valid ChildDataA ChildAllocatedProperties", (int32)ChildAllocatedProperties, (int32)ChildDataA->GetAllocatedProperties());
+	TestTrue("ChildDataA has 0 points", ChildDataA->GetNumPoints() == 0);
+
+	ChildDataA->InitializeFromData(ParentData);
+	TestTrue("ChildDataA has 0 points (After InitializeFromData)", ChildDataA->GetNumPoints() == 0);
+	TestFalse("ChildDataA has no parent", ChildDataA->HasSpatialDataParent());
+	
+	UPCGPointArrayData* ChildDataB = CastChecked<UPCGPointArrayData>(ParentData->DuplicateData(nullptr));
+	
+	TestEqual("ChildDataB->GetAllocatedProperties() == ParentData->GetAllocatedProperties()", (int32)ChildDataB->GetAllocatedProperties(), (int32)ParentData->GetAllocatedProperties());
+	TestTrue("ChildDataB->GetNumPoints() == ParentData->GetNumPoints()", ChildDataB->GetNumPoints() == ParentData->GetNumPoints());
+	TestFalse("ChildDataB has no parent", ChildDataB->HasSpatialDataParent());
+	
+	// Compare data
+	return TestEqualData(ParentData, ChildDataB);
+}
+
+bool FPCGPointArrayDataCopyPointsToFromPointArrayDataToPointDataTest::RunTest(const FString& Parameters)
+{
+	UPCGPointArrayData* PointArrayData = PCGTestsCommon::CreateRandomPointData<UPCGPointArrayData>(100, 42);
+	UPCGPointData* PointData = PCGTestsCommon::CreateEmptyPointData<UPCGPointData>();
+
+	PointData->SetNumPoints(PointArrayData->GetNumPoints());
+	PointArrayData->CopyPointsTo(PointData, 0, 0, PointArrayData->GetNumPoints());
+
+	return TestEqualData(PointArrayData, PointData);
+}
+
+bool FPCGPointArrayDataCopyPointsToFromPointArrayDataToPointArrayDataTest::RunTest(const FString& Parameters)
+{
+	UPCGPointArrayData* PointArrayData = PCGTestsCommon::CreateRandomPointData<UPCGPointArrayData>(100, 42);
+	UPCGPointArrayData* PointArrayData2 = PCGTestsCommon::CreateEmptyPointData<UPCGPointArrayData>();
+
+	PointArrayData2->SetNumPoints(PointArrayData->GetNumPoints());
+	PointArrayData->CopyPointsTo(PointArrayData2, 0, 0, PointArrayData->GetNumPoints());
+
+	return TestEqualData(PointArrayData, PointArrayData2);
+}
+
+bool FPCGPointArrayDataCopyPointsToFromPointDataToPointArrayDataTest::RunTest(const FString& Parameters)
+{
+	UPCGPointData* PointData = PCGTestsCommon::CreateRandomPointData<UPCGPointData>(100, 42);
+	UPCGPointArrayData* PointArrayData = PCGTestsCommon::CreateEmptyPointData<UPCGPointArrayData>();
+
+	PointArrayData->SetNumPoints(PointData->GetNumPoints());
+	PointData->CopyPointsTo(PointArrayData, 0, 0, PointData->GetNumPoints());
+
+	return TestEqualData(PointArrayData, PointData);
+}
+
+bool FPCGPointArrayDataCopyPointsToFromPointDataToPointDataTest::RunTest(const FString& Parameters)
+{
+	UPCGPointData* PointData = PCGTestsCommon::CreateRandomPointData<UPCGPointData>(100, 42);
+	UPCGPointData* PointData2 = PCGTestsCommon::CreateEmptyPointData<UPCGPointData>();
+
+	PointData2->SetNumPoints(PointData->GetNumPoints());
+	PointData->CopyPointsTo(PointData2, 0, 0, PointData->GetNumPoints());
+
+	return TestEqualData(PointData2, PointData);
+}
+
+bool FPCGPointArrayDataCopyPropertyOverInheritedPropertyTest::RunTest(const FString& Parameters)
+{
+	// Enable Parenting
+	const bool bPreviousCVarValue = CVarPCGEnablePointArrayDataParenting.GetValueOnAnyThread();
+	CVarPCGEnablePointArrayDataParenting->Set(true, ECVF_SetByCode);
+	ON_SCOPE_EXIT
+	{
+		CVarPCGEnablePointArrayDataParenting->Set(bPreviousCVarValue, ECVF_SetByCode);
+	};
+
+	UPCGPointArrayData* ParentData = PCGTestsCommon::CreateRandomPointData<UPCGPointArrayData>(100, 42);
+	EPCGPointNativeProperties ParentAllocatedProperties = EPCGPointNativeProperties::Transform | EPCGPointNativeProperties::Seed | EPCGPointNativeProperties::Color | EPCGPointNativeProperties::Density;
+
+	// CreateRandomPointData only allocates those properties
+	TestEqual("Valid ParentAllocatedProperties", (int32)ParentAllocatedProperties, (int32)ParentData->GetAllocatedProperties());
+
+	UPCGPointArrayData* ChildData = PCGTestsCommon::CreateEmptyPointData<UPCGPointArrayData>();
+	EPCGPointNativeProperties ChildAllocatedProperties = EPCGPointNativeProperties::None;
+
+	// CreateEmptyPointData should not allocate any properties
+	TestEqual("Valid ChildAllocatedProperties", (int32)ChildAllocatedProperties, (int32)ChildData->GetAllocatedProperties());
+	TestTrue("ChildData has 0 points", ChildData->GetNumPoints() == 0);
+
+	ChildData->InitializeFromData(ParentData);
+	TestEqual("ChildData->GetAllocatedProperties() == ParentData->GetAllocatedProperties()", (int32)ChildData->GetAllocatedProperties(), (int32)ParentData->GetAllocatedProperties());
+	TestTrue("ChildData->GetNumPoints() == ParentData->GetNumPoints()", ChildData->GetNumPoints() == ParentData->GetNumPoints());
+	TestTrue("ChildData has parent", ChildData->HasSpatialDataParent());
+	TestFalse("ParentData has no parent", ParentData->HasSpatialDataParent());
+
+	if (!TestEqualData(ParentData, ChildData))
+	{
+		return false;
+	}
+
+	UPCGPointArrayData* OtherData = PCGTestsCommon::CreateRandomPointData<UPCGPointArrayData>(100, 57);
+
+	OtherData->CopyPropertiesTo(ChildData, 0, 0, OtherData->GetNumPoints(), EPCGPointNativeProperties::Color | EPCGPointNativeProperties::Seed);
+	TestTrue("ChildData has parent (after CopyPropertiesTo)", ChildData->HasSpatialDataParent());
+
+	if (!TestEqualRange<FVector4>(ChildData, OtherData, EPCGPointNativeProperties::Color))
+	{
+		return false;
+	}
+
+	if (!TestEqualRange<int32>(ChildData, OtherData, EPCGPointNativeProperties::Seed))
+	{
+		return false;
+	}
+
+	if (!TestEqualRange<FTransform>(ChildData, ParentData, EPCGPointNativeProperties::Transform))
+	{
+		return false;
+	}
+
+	if (!TestEqualRange<float>(ChildData, ParentData, EPCGPointNativeProperties::Density))
+	{
+		return false;
+	}
+
+	OtherData->CopyPropertiesTo(ChildData, 0, 0, OtherData->GetNumPoints(), EPCGPointNativeProperties::Transform | EPCGPointNativeProperties::Density);
+	ChildData->Flatten(); // Make sure we don't have inherited properties anymore (non allocated ones)
+	TestFalse("ChildData has no parent (after CopyPropertiesTo)", ChildData->HasSpatialDataParent());
+		
+	return TestEqualData(ChildData, OtherData);
+}
+
+#undef PCG_FMT
\ No newline at end of file
diff --git a/Engine/Plugins/PCG/Source/PCG/Private/Tests/PCGValueRangeTest.cpp b/Engine/Plugins/PCG/Source/PCG/Private/Tests/PCGValueRangeTest.cpp
--- a/Engine/Plugins/PCG/Source/PCG/Private/Tests/PCGValueRangeTest.cpp
+++ b/Engine/Plugins/PCG/Source/PCG/Private/Tests/PCGValueRangeTest.cpp
@@ -26,7 +26,7 @@ public:
 		const TConstPCGValueRange<T> RangeA = PointDataA->GetConstValueRange<T>(NativeProperty);
 		const TConstPCGValueRange<T> RangeB = PointDataB->GetConstValueRange<T>(NativeProperty);
 
-		UTEST_EQUAL(PCG_FMT("RangeA.Num() == RangeB.Num()", *StaticEnum<EPCGPointNativeProperties>()->GetNameStringByValue((int64)NativeProperty)), RangeA.Num(), RangeB.Num());
+		UTEST_EQUAL("RangeA.Num() == RangeB.Num()", RangeA.Num(), RangeB.Num());
 
 		for (int32 i = 0; i < RangeA.Num(); ++i)
 		{
diff --git a/Engine/Plugins/PCG/Source/PCG/Public/Data/PCGBasePointData.h b/Engine/Plugins/PCG/Source/PCG/Public/Data/PCGBasePointData.h
--- a/Engine/Plugins/PCG/Source/PCG/Public/Data/PCGBasePointData.h
+++ b/Engine/Plugins/PCG/Source/PCG/Public/Data/PCGBasePointData.h
@@ -136,8 +136,9 @@ public:
 	UFUNCTION(BlueprintCallable, Category = SpatialData)
 	void CopyPointsFrom(const UPCGBasePointData* InData, const TArray<int32>& InDataIndices);
 
-	virtual void CopyPointsTo(UPCGBasePointData* OutData, int32 ReadStartIndex, int32 WriteStartIndex, int32 Count) const;
-	virtual EPCGPointNativeProperties GetAllocatedProperties() const { return EPCGPointNativeProperties::All; }
+	virtual void CopyPropertiesTo(UPCGBasePointData* To, int32 ReadStartIndex, int32 WriteStartIndex, int32 Count, EPCGPointNativeProperties Properties) const;
+	virtual void CopyPointsTo(UPCGBasePointData* To, int32 ReadStartIndex, int32 WriteStartIndex, int32 Count) const;
+	virtual EPCGPointNativeProperties GetAllocatedProperties(bool bWithInheritance = true) const { return EPCGPointNativeProperties::All; }
 
 	/** Metadata specific */
 	virtual FPCGMetadataDomainID GetDefaultMetadataDomainID() const override { return PCGMetadataDomainID::Elements; }
diff --git a/Engine/Plugins/PCG/Source/PCG/Public/Data/PCGPointArrayData.h b/Engine/Plugins/PCG/Source/PCG/Public/Data/PCGPointArrayData.h
--- a/Engine/Plugins/PCG/Source/PCG/Public/Data/PCGPointArrayData.h
+++ b/Engine/Plugins/PCG/Source/PCG/Public/Data/PCGPointArrayData.h
@@ -7,6 +7,8 @@
 
 #include "PCGPointArrayData.generated.h"
 
+extern PCG_API TAutoConsoleVariable<bool> CVarPCGEnablePointArrayDataParenting;
+
 UCLASS(BlueprintType, ClassGroup = (Procedural))
 class PCG_API UPCGPointArrayData : public UPCGBasePointData
 {
@@ -16,45 +18,223 @@ public:
 	virtual void GetResourceSizeEx(FResourceSizeEx& CumulativeResourceSize) override;
 	//~End UObject Interface
 
+	virtual void Flatten() override;
+
+	virtual void VisitDataNetwork(TFunctionRef<void(const UPCGData*)> Action) const override;
 	virtual const UPCGPointData* ToPointData(FPCGContext* Context, const FBox& InBounds = FBox(EForceInit::ForceInit)) const override;
 	virtual const UPCGPointArrayData* ToPointArrayData(FPCGContext* Context) const override { return this; }
 	virtual UPCGSpatialData* CopyInternal(FPCGContext* Context) const override;
-	virtual void CopyPointsTo(UPCGBasePointData* OutData, int32 ReadStartIndex, int32 WriteStartIndex, int32 Count) const;
-	virtual EPCGPointNativeProperties GetAllocatedProperties() const override { return PointArray.GetAllocatedProperties(); }
+	virtual void CopyPropertiesTo(UPCGBasePointData* To, int32 ReadStartIndex, int32 WriteStartIndex, int32 Count, EPCGPointNativeProperties Properties) const override;
+	virtual EPCGPointNativeProperties GetAllocatedProperties(bool bWithInheritance = true) const override;
+	virtual bool SupportsSpatialDataInheritance() const override;
+	virtual bool HasSpatialDataParent() const override { return ParentData != nullptr; }
 
 	//~Begin UPCGBasePointData Interface
 	virtual bool IsValidRef(const PCGPointOctree::FPointRef& InPointRef) const override { return InPointRef.Index >= 0 && InPointRef.Index < GetNumPoints(); }
 	
 	virtual int32 GetNumPoints() const override { return PointArray.GetNumPoints(); }
-	virtual void SetNumPoints(int32 InNumPoints) override { PointArray.SetNumPoints(InNumPoints); }
+	virtual void SetNumPoints(int32 InNumPoints) override;
 		
-	virtual void AllocateProperties(EPCGPointNativeProperties Properties) override { PointArray.Allocate(Properties); }
-	virtual void FreeProperties(EPCGPointNativeProperties Properties) override { PointArray.Free(Properties); }
-	virtual void MoveRange(int32 RangeStartIndex, int32 MoveToIndex, int32 NumElements) override { PointArray.MoveRange(RangeStartIndex, MoveToIndex, NumElements); }
+	virtual void AllocateProperties(EPCGPointNativeProperties Properties) override;
+	virtual void FreeProperties(EPCGPointNativeProperties Properties) override;
+	virtual void MoveRange(int32 RangeStartIndex, int32 MoveToIndex, int32 NumElements) override;
 
 	virtual TArray<FTransform> GetTransformsCopy() const { return PointArray.GetTransformCopy(); }
 
-	virtual TPCGValueRange<FTransform> GetTransformValueRange(bool bAllocate = true) override { return PointArray.GetTransformValueRange(bAllocate); }
-	virtual TPCGValueRange<float> GetDensityValueRange(bool bAllocate = true) override { return PointArray.GetDensityValueRange(bAllocate); }
-	virtual TPCGValueRange<FVector> GetBoundsMinValueRange(bool bAllocate = true) override { return PointArray.GetBoundsMinValueRange(bAllocate); }
-	virtual TPCGValueRange<FVector> GetBoundsMaxValueRange(bool bAllocate = true) override { return PointArray.GetBoundsMaxValueRange(bAllocate); }
-	virtual TPCGValueRange<FVector4> GetColorValueRange(bool bAllocate = true) override { return PointArray.GetColorValueRange(bAllocate); }
-	virtual TPCGValueRange<float> GetSteepnessValueRange(bool bAllocate = true) override { return PointArray.GetSteepnessValueRange(bAllocate); }
-	virtual TPCGValueRange<int32> GetSeedValueRange(bool bAllocate = true) override { return PointArray.GetSeedValueRange(bAllocate); }
-	virtual TPCGValueRange<int64> GetMetadataEntryValueRange(bool bAllocate = true) override { return PointArray.GetMetadataEntryValueRange(bAllocate); }
-
-	virtual TConstPCGValueRange<FTransform> GetConstTransformValueRange() const override { return PointArray.GetConstTransformValueRange(); }
-	virtual TConstPCGValueRange<float> GetConstDensityValueRange() const override { return PointArray.GetConstDensityValueRange(); }
-	virtual TConstPCGValueRange<FVector> GetConstBoundsMinValueRange() const override { return PointArray.GetConstBoundsMinValueRange(); }
-	virtual TConstPCGValueRange<FVector> GetConstBoundsMaxValueRange() const override { return PointArray.GetConstBoundsMaxValueRange(); }
-	virtual TConstPCGValueRange<FVector4> GetConstColorValueRange() const override { return PointArray.GetConstColorValueRange(); }
-	virtual TConstPCGValueRange<float> GetConstSteepnessValueRange() const override { return PointArray.GetConstSteepnessValueRange(); }
-	virtual TConstPCGValueRange<int32> GetConstSeedValueRange() const override { return PointArray.GetConstSeedValueRange(); }
-	virtual TConstPCGValueRange<int64> GetConstMetadataEntryValueRange() const override { return PointArray.GetConstMetadataEntryValueRange(); }
+	virtual TPCGValueRange<FTransform> GetTransformValueRange(bool bAllocate = true) override;
+	virtual TPCGValueRange<float> GetDensityValueRange(bool bAllocate = true) override;
+	virtual TPCGValueRange<FVector> GetBoundsMinValueRange(bool bAllocate = true) override;
+	virtual TPCGValueRange<FVector> GetBoundsMaxValueRange(bool bAllocate = true) override;
+	virtual TPCGValueRange<FVector4> GetColorValueRange(bool bAllocate = true) override;
+	virtual TPCGValueRange<float> GetSteepnessValueRange(bool bAllocate = true) override;
+	virtual TPCGValueRange<int32> GetSeedValueRange(bool bAllocate = true) override;
+	virtual TPCGValueRange<int64> GetMetadataEntryValueRange(bool bAllocate = true) override;
+
+	virtual TConstPCGValueRange<FTransform> GetConstTransformValueRange() const override;
+	virtual TConstPCGValueRange<float> GetConstDensityValueRange() const override;
+	virtual TConstPCGValueRange<FVector> GetConstBoundsMinValueRange() const override;
+	virtual TConstPCGValueRange<FVector> GetConstBoundsMaxValueRange() const override;
+	virtual TConstPCGValueRange<FVector4> GetConstColorValueRange() const override;
+	virtual TConstPCGValueRange<float> GetConstSteepnessValueRange() const override;
+	virtual TConstPCGValueRange<int32> GetConstSeedValueRange() const override;
+	virtual TConstPCGValueRange<int64> GetConstMetadataEntryValueRange() const override;
 	//~End UPCGBasePointData Interface
 		
+protected:
+	virtual void InitializeSpatialDataInternal(const FPCGInitializeFromDataParams& InParams) override;
+
 private:
+	void FlattenPropertiesIfNeeded(EPCGPointNativeProperties Property = EPCGPointNativeProperties::All);
+
+	template<class T>
+	FPCGPointArrayProperty<T>* GetProperty(EPCGPointNativeProperties Property, bool bWithInheritance = true)
+	{
+		if (bWithInheritance)
+		{
+			if (EnumHasAnyFlags(InheritedProperties, Property))
+			{
+				check(ParentData);
+				return ParentData->GetProperty<T>(Property, bWithInheritance);
+			}
+		}
+
+		if constexpr (std::is_same_v<T, FTransform>)
+		{
+			if (Property == EPCGPointNativeProperties::Transform)
+			{
+				return &PointArray.Transform;
+			}
+		}
+		else if constexpr (std::is_same_v<T, float>)
+		{
+			if (Property == EPCGPointNativeProperties::Density)
+			{
+				return &PointArray.Density;
+			}
+			else if (Property == EPCGPointNativeProperties::Steepness)
+			{
+				return &PointArray.Steepness;
+			}
+		}
+		else if constexpr (std::is_same_v<T, FVector>)
+		{
+			if (Property == EPCGPointNativeProperties::BoundsMin)
+			{
+				return &PointArray.BoundsMin;
+			}
+			else if (Property == EPCGPointNativeProperties::BoundsMax)
+			{
+				return &PointArray.BoundsMax;
+			}
+		}
+		else if constexpr (std::is_same_v<T, FVector4>)
+		{
+			if (Property == EPCGPointNativeProperties::Color)
+			{
+				return &PointArray.Color;
+			}
+		}
+		else if constexpr (std::is_same_v<T, int32>)
+		{
+			if (Property == EPCGPointNativeProperties::Seed)
+			{
+				return &PointArray.Seed;
+			}
+		}
+		else if constexpr (std::is_same_v<T, int64>)
+		{
+			if (Property == EPCGPointNativeProperties::MetadataEntry)
+			{
+				return &PointArray.MetadataEntry;
+			}
+		}
+
+		check(false);
+		return nullptr;
+	}
+
+	template<class T>
+	const FPCGPointArrayProperty<T>* GetProperty(EPCGPointNativeProperties Property, bool bWithInheritance = true) const
+	{
+		if (bWithInheritance)
+		{
+			if (EnumHasAnyFlags(InheritedProperties, Property))
+			{
+				check(ParentData);
+				return ParentData->GetProperty<T>(Property, bWithInheritance);
+			}
+		}
+
+		if constexpr (std::is_same_v<T, FTransform>)
+		{
+			if (Property == EPCGPointNativeProperties::Transform)
+			{
+				return &PointArray.Transform;
+			}
+		}
+		else if constexpr (std::is_same_v<T, float>)
+		{
+			if (Property == EPCGPointNativeProperties::Density)
+			{
+				return &PointArray.Density;
+			}
+			else if (Property == EPCGPointNativeProperties::Steepness)
+			{
+				return &PointArray.Steepness;
+			}
+		}
+		else if constexpr (std::is_same_v<T, FVector>)
+		{
+			if (Property == EPCGPointNativeProperties::BoundsMin)
+			{
+				return &PointArray.BoundsMin;
+			}
+			else if (Property == EPCGPointNativeProperties::BoundsMax)
+			{
+				return &PointArray.BoundsMax;
+			}
+		}
+		else if constexpr (std::is_same_v<T, FVector4>)
+		{
+			if (Property == EPCGPointNativeProperties::Color)
+			{
+				return &PointArray.Color;
+			}
+		}
+		else if constexpr (std::is_same_v<T, int32>)
+		{
+			if (Property == EPCGPointNativeProperties::Seed)
+			{
+				return &PointArray.Seed;
+			}
+		}
+		else if constexpr (std::is_same_v<T, int64>)
+		{
+			if (Property == EPCGPointNativeProperties::MetadataEntry)
+			{
+				return &PointArray.MetadataEntry;
+			}
+		}
 
+		check(false);
+		return nullptr;
+	}
+
+	template<class T>
+	bool FlattenPropertyIfNeeded(EPCGPointNativeProperties NativeProperty)
+	{
+		if (!EnumHasAnyFlags(InheritedProperties, NativeProperty))
+		{
+			return false;
+		}
+
+		const FPCGPointArrayProperty<T>* InheritedProperty = GetProperty<T>(NativeProperty, /*bWithInheritance=*/true);
+		FPCGPointArrayProperty<T>* Property = GetProperty<T>(NativeProperty, /*bWithInheritance=*/false);
+
+		// Should not be equal since InheritedProperties contains that property
+		if (ensure(InheritedProperty != Property))
+		{
+			Property->Allocate();
+			check(InheritedProperty->Num() == Property->Num());
+			InheritedProperty->CopyTo(*Property, 0, 0, Property->Num());
+		}
+
+		EnumRemoveFlags(InheritedProperties, NativeProperty);
+
+		if (InheritedProperties == EPCGPointNativeProperties::None)
+		{
+			ParentData = nullptr;
+		}
+
+		return true;
+	};
+	
 	UPROPERTY()
 	FPCGPointArray PointArray;
+
+	UPROPERTY(VisibleAnywhere, Category = ParentData)
+	TObjectPtr<UPCGPointArrayData> ParentData = nullptr;
+
+	UPROPERTY(VisibleAnywhere, Category = ParentData)
+	EPCGPointNativeProperties InheritedProperties = EPCGPointNativeProperties::None;
 };
\ No newline at end of file
diff --git a/Engine/Plugins/PCG/Source/PCG/Public/Data/PCGSpatialData.h b/Engine/Plugins/PCG/Source/PCG/Public/Data/PCGSpatialData.h
--- a/Engine/Plugins/PCG/Source/PCG/Public/Data/PCGSpatialData.h
+++ b/Engine/Plugins/PCG/Source/PCG/Public/Data/PCGSpatialData.h
@@ -12,6 +12,7 @@
 
 class AActor;
 struct FPCGContext;
+class UPCGSpatialData;
 class UPCGBasePointData;
 class UPCGPointData;
 class UPCGPointArrayData;
@@ -20,6 +21,34 @@ class UPCGUnionData;
 class UPCGDifferenceData;
 class UPCGProjectionData;
 
+USTRUCT(BlueprintType)
+struct PCG_API FPCGInitializeFromDataParams
+{
+	GENERATED_BODY()
+
+	FPCGInitializeFromDataParams() = default;
+
+	FPCGInitializeFromDataParams(const UPCGSpatialData* InSource)
+		: Source(InSource)
+	{
+	}
+	
+	UPROPERTY(BlueprintReadWrite, EditAnywhere, Category = Default)
+	TObjectPtr<const UPCGSpatialData> Source = nullptr;
+	
+	UPROPERTY(BlueprintReadWrite, EditAnywhere, Category = Default)
+	TObjectPtr<const UPCGMetadata> MetadataParentOverride = nullptr;
+		
+	UPROPERTY(BlueprintReadWrite, EditAnywhere, Category = Default)
+	bool bInheritMetadata = true;
+
+	UPROPERTY(BlueprintReadWrite, EditAnywhere, Category = Default)
+	bool bInheritAttributes = true;
+
+	UPROPERTY(BlueprintReadWrite, EditAnywhere, Category = Default)
+	bool bInheritSpatialData = true;
+};
+
 /**
 * "Concrete" data base class for PCG generation
 * This will be the base class for data classes that actually represent
@@ -166,12 +195,21 @@ public:
 	UFUNCTION(BlueprintCallable, Category = SpatialData)
 	void InitializeFromData(const UPCGSpatialData* InSource, const UPCGMetadata* InMetadataParentOverride = nullptr, bool bInheritMetadata = true, bool bInheritAttributes = true);
 
+	UFUNCTION(BlueprintCallable, Category = SpatialData)
+	void InitializeFromDataWithParams(const FPCGInitializeFromDataParams& InParams);
+
 	/** True if this operation does not have an inverse and cannot be queried analytically/implicitly, and therefore must be collapsed to an explicit point representation. */
 	virtual bool RequiresCollapseToSample() const { return false; }
 
 	/** Find the first concrete (non-composite) shape in the network. Depth first search. */
 	virtual const UPCGSpatialData* FindFirstConcreteShapeFromNetwork() const { return !!(GetDataType() & EPCGDataType::Concrete) ? this : nullptr; }
 
+	/** True if subclass UPCGData instances support inheriting data from parent UPCGData */
+	virtual bool SupportsSpatialDataInheritance() const { return false; }
+	
+	/** True if data inherits from parent */
+	virtual bool HasSpatialDataParent() const { return false; }
+
 	/** Recipient of any artifacts generated using this data. */
 	UPROPERTY(Transient, BlueprintReadWrite, EditAnywhere, Category = Data)
 	TWeakObjectPtr<AActor> TargetActor = nullptr;
@@ -185,6 +223,9 @@ protected:
 
 	virtual UPCGSpatialData* CopyInternal(FPCGContext* Context) const PURE_VIRTUAL(UPCGSpatialData::CopyInternal, return nullptr;);
 
+	void InitializeMetadataInternal(const FPCGInitializeFromDataParams& InParams);
+	virtual void InitializeSpatialDataInternal(const FPCGInitializeFromDataParams& InParams);
+
 private:
 	/** Cache to keep track of the latest attribute manipulated on this data. */
 	UPROPERTY()
diff --git a/Engine/Plugins/PCG/Source/PCG/Public/Elements/PCGBoundsModifier.h b/Engine/Plugins/PCG/Source/PCG/Public/Elements/PCGBoundsModifier.h
--- a/Engine/Plugins/PCG/Source/PCG/Public/Elements/PCGBoundsModifier.h
+++ b/Engine/Plugins/PCG/Source/PCG/Public/Elements/PCGBoundsModifier.h
@@ -62,8 +62,9 @@ public:
 class FPCGBoundsModifier : public FPCGPointOperationElementBase
 {
 public:
-	virtual bool SupportsBasePointDataInputs(FPCGContext* InContext) const { return true; }
+	virtual bool SupportsBasePointDataInputs(FPCGContext* InContext) const override { return true; }
 protected:
 	virtual bool ExecuteInternal(FPCGContext* Context) const override;
-	virtual EPCGPointNativeProperties GetPropertiesToAllocate(const UPCGBasePointData* Input) const override;
+	virtual EPCGPointNativeProperties GetPropertiesToAllocate(FPCGContext* InContext) const override;
+	virtual bool ShouldCopyPoints() const override { return true; }
 };
\ No newline at end of file
diff --git a/Engine/Plugins/PCG/Source/PCG/Public/Elements/PCGHiGenGridSize.h b/Engine/Plugins/PCG/Source/PCG/Public/Elements/PCGHiGenGridSize.h
--- a/Engine/Plugins/PCG/Source/PCG/Public/Elements/PCGHiGenGridSize.h
+++ b/Engine/Plugins/PCG/Source/PCG/Public/Elements/PCGHiGenGridSize.h
@@ -48,4 +48,5 @@ public:
 
 protected:
 	virtual bool ExecuteInternal(FPCGContext* Context) const override;
+	virtual bool SupportsBasePointDataInputs(FPCGContext* InContext) const override { return true; }
 };
diff --git a/Engine/Plugins/PCG/Source/PCG/Public/Elements/PCGPointOperationElementBase.h b/Engine/Plugins/PCG/Source/PCG/Public/Elements/PCGPointOperationElementBase.h
--- a/Engine/Plugins/PCG/Source/PCG/Public/Elements/PCGPointOperationElementBase.h
+++ b/Engine/Plugins/PCG/Source/PCG/Public/Elements/PCGPointOperationElementBase.h
@@ -63,7 +63,10 @@ protected:
 	/** Mandatory call. Using the context, prepares the state data for time slice execution */
 	bool PreparePointOperationData(ContextType* Context, FName InputPinLabel = PCGPinConstants::DefaultInputLabel) const;
 
-	virtual EPCGPointNativeProperties GetPropertiesToAllocate(const UPCGBasePointData* Input) const { return EPCGPointNativeProperties::All; }
+	virtual EPCGPointNativeProperties GetPropertiesToAllocate(FPCGContext* InContext) const { return EPCGPointNativeProperties::All; }
+	
+	/** If True, input points will be copied into output points before each point operation. The copy will only happen if output doesn't support inheritance */
+	virtual bool ShouldCopyPoints() const { return false; }
 
 private:
 	bool ExecutePointOperationWithPoints(ContextType* Context, TFunctionRef<bool(const FPCGPoint& InPoint, FPCGPoint& OutPoint)> Callback, int32 PointsPerChunk) const;
diff --git a/Engine/Plugins/PCG/Source/PCG/Public/Elements/PCGReroute.h b/Engine/Plugins/PCG/Source/PCG/Public/Elements/PCGReroute.h
--- a/Engine/Plugins/PCG/Source/PCG/Public/Elements/PCGReroute.h
+++ b/Engine/Plugins/PCG/Source/PCG/Public/Elements/PCGReroute.h
@@ -115,4 +115,5 @@ class FPCGRerouteElement : public IPCGElement
 {
 protected:
 	virtual bool ExecuteInternal(FPCGContext* Context) const override;
+	virtual bool SupportsBasePointDataInputs(FPCGContext* InContext) const override { return true; }
 };
diff --git a/Engine/Plugins/PCG/Source/PCG/Public/Elements/PCGSurfaceSampler.h b/Engine/Plugins/PCG/Source/PCG/Public/Elements/PCGSurfaceSampler.h
--- a/Engine/Plugins/PCG/Source/PCG/Public/Elements/PCGSurfaceSampler.h
+++ b/Engine/Plugins/PCG/Source/PCG/Public/Elements/PCGSurfaceSampler.h
@@ -169,7 +169,7 @@ public:
 	virtual void GetDependenciesCrc(const FPCGDataCollection& InInput, const UPCGSettings* InSettings, UPCGComponent* InComponent, FPCGCrc& OutCrc) const override;
 	// Might be sampling landscape or other external data, worth computing a full CRC in case we can halt change propagation/re-executions
 	virtual bool ShouldComputeFullOutputDataCrc(FPCGContext* Context) const override { return true; }
-	virtual bool SupportsBasePointDataInputs(FPCGContext* InContext) const { return true; }
+	virtual bool SupportsBasePointDataInputs(FPCGContext* InContext) const override { return true; }
 
 protected:
 	virtual bool PrepareDataInternal(FPCGContext* InContext) const override;
diff --git a/Engine/Plugins/PCG/Source/PCG/Public/PCGInputOutputSettings.h b/Engine/Plugins/PCG/Source/PCG/Public/PCGInputOutputSettings.h
--- a/Engine/Plugins/PCG/Source/PCG/Public/PCGInputOutputSettings.h
+++ b/Engine/Plugins/PCG/Source/PCG/Public/PCGInputOutputSettings.h
@@ -22,6 +22,7 @@ protected:
 	virtual bool ExecuteInternal(FPCGContext* Context) const override;
 	virtual bool IsCacheable(const UPCGSettings* InSettings) const override { return false; }
 	virtual bool SupportsGPUResidentData(FPCGContext* InContext) const override;
+	virtual bool SupportsBasePointDataInputs(FPCGContext* InContext) const override { return true; }
 };
 
 UCLASS(MinimalAPI, NotBlueprintable, Hidden, ClassGroup = (Procedural))
diff --git a/Engine/Plugins/PCG/Source/PCG/Public/PCGPointArray.h b/Engine/Plugins/PCG/Source/PCG/Public/PCGPointArray.h
--- a/Engine/Plugins/PCG/Source/PCG/Public/PCGPointArray.h
+++ b/Engine/Plugins/PCG/Source/PCG/Public/PCGPointArray.h
@@ -27,7 +27,7 @@ struct FPCGPointArrayProperty
 {
 public:
 	/** Defined for serialization purposes(FPCGPointArray trait WithIdenticalViaEquality) */
-	bool operator==(const FPCGPointArrayProperty<T>& Other) const
+	[[nodiscard]] FORCEINLINE bool operator==(const FPCGPointArrayProperty<T>& Other) const
 	{
 		if (NumValues != Other.NumValues || Values.Num() != Other.Values.Num())
 		{
@@ -63,13 +63,13 @@ public:
 	}
 	
 	/** If 'Values' isn't currently allocated, allocate it and copy 'Value' at every index of 'Values[0, NumValues - 1]' */
-	void Allocate()
+	FORCEINLINE void Allocate()
 	{
 		SetNum(NumValues, /*bAllocate=*/true);
 	}
 
 	/** Set the 'NumValues', optionally allocating 'Values' */
-	void SetNum(int32 NewNum, bool bAllocate)
+	FORCEINLINE void SetNum(int32 NewNum, bool bAllocate)
 	{
 		NumValues = NewNum;
 
@@ -89,14 +89,19 @@ public:
 		}
 	};
 
+	FORCEINLINE int32 Num() const
+	{
+		return NumValues;
+	}
+
 	/** Free 'Values', effectively resets all values to 'Value' */
-	void Free()
+	FORCEINLINE void Free()
 	{
 		Values.Reset();
 	}
 
 	/** Set all values('Value/Values') to InValue */
-	void SetValue(const T& InValue)
+	FORCEINLINE void SetValue(const T& InValue)
 	{
 		for (int32 i = 0; i < Values.Num(); ++i)
 		{
@@ -106,7 +111,7 @@ public:
 	}
 
 	/** Return value at Index */
-	const T& GetValue(int32 Index) const
+	[[nodiscard]] FORCEINLINE const T& GetValue(int32 Index) const
 	{
 		ensure(Index >= 0 && Index < NumValues);
 		if (Values.Num() == 0)
@@ -118,7 +123,7 @@ public:
 	}
 
 	/** Move ranges of values(moves them back into the array).Used by the PCGAsync api */
-	void MoveRange(int32 RangeStartIndex, int32 MoveToIndex, int32 NumElements)
+	FORCEINLINE void MoveRange(int32 RangeStartIndex, int32 MoveToIndex, int32 NumElements)
 	{
 		if (IsAllocated())
 		{
@@ -128,14 +133,14 @@ public:
 	}
 
 	/* Returns true if 'Values' is allocated */
-	bool IsAllocated() const
+	[[nodiscard]] FORCEINLINE bool IsAllocated() const
 	{
 		ensure(Values.Num() == 0 || Values.Num() == NumValues);
 		return Values.Num() != 0;
 	}
 
 	/** Returns a TArray copy of size NumValues */
-	TArray<T> GetCopy() const
+	[[nodiscard]] FORCEINLINE TArray<T> GetCopy() const
 	{
 		if (IsAllocated())
 		{
@@ -153,7 +158,7 @@ public:
 	}
 
 	/** Returns a TPCGValueRange<T> of NumValues and optionally allocate 'Values' */
-	TPCGValueRange<T> GetValueRange(bool bAllocate)
+	[[nodiscard]] FORCEINLINE TPCGValueRange<T> GetValueRange(bool bAllocate)
 	{
 		if (bAllocate)
 		{
@@ -165,17 +170,17 @@ public:
 	}
 
 	/** Returns a TConstPCGValueRange<T> of NumValues */
-	TConstPCGValueRange<T> GetConstValueRange() const
+	[[nodiscard]] FORCEINLINE TConstPCGValueRange<T> GetConstValueRange() const
 	{
 		return TConstPCGValueRange<T>(IsAllocated() ? MakeConstStridedView(Values) : MakeConstStridedView(sizeof(Value), &Value, 1), NumValues);
 	}
 		
-	SIZE_T GetSizeBytes() const
+	[[nodiscard]] FORCEINLINE SIZE_T GetSizeBytes() const
 	{
 		return Values.GetAllocatedSize() + sizeof(Value) + sizeof(NumValues) + sizeof(Values);
 	}
 
-	void Serialize(FArchive& Ar)
+	FORCEINLINE void Serialize(FArchive& Ar)
 	{
 		Ar << NumValues;
 		Ar << Value;
@@ -325,8 +330,11 @@ public:
 	// WithSerializer & WithIdenticalViaEquality needed for this struct to properly serialize
 	bool Serialize(FArchive& Ar);
 
-	void CopyPointsTo(FPCGPointArray& OutPointArray, int32 InStartReadIndex, int32 InStartWriteIndex, int32 Count) const;
+	void CopyPropertiesTo(FPCGPointArray& OutPointArray, int32 InStartReadIndex, int32 InStartWriteIndex, int32 Count, EPCGPointNativeProperties Properties) const;
+
 private:
+	friend class UPCGPointArrayData;
+	
 	FPCGPointArrayProperty<FTransform>	Transform;
 	FPCGPointArrayProperty<float>		Density;
 	FPCGPointArrayProperty<FVector>		BoundsMin;
diff --git a/Engine/Plugins/PCG/Source/PCG/Public/PCGSettings.h b/Engine/Plugins/PCG/Source/PCG/Public/PCGSettings.h
--- a/Engine/Plugins/PCG/Source/PCG/Public/PCGSettings.h
+++ b/Engine/Plugins/PCG/Source/PCG/Public/PCGSettings.h
@@ -649,4 +649,5 @@ protected:
 	virtual bool ExecuteInternal(FPCGContext* Context) const override;
 	virtual bool IsPassthrough(const UPCGSettings* InSettings) const override { return true; }
 	virtual bool SupportsGPUResidentData(FPCGContext* InContext) const override { return true; }
+	virtual bool SupportsBasePointDataInputs(FPCGContext* InContext) const override { return true; }
 };
