[PCG] Add Element Count Multiplier setting to GPU output pins for 1:N expansions.

#rnx

diff --git a/Engine/Plugins/PCG/Source/PCG/Private/Compute/DataInterfaces/PCGDataCollectionDataInterface.cpp b/Engine/Plugins/PCG/Source/PCG/Private/Compute/DataInterfaces/PCGDataCollectionDataInterface.cpp
--- a/Engine/Plugins/PCG/Source/PCG/Private/Compute/DataInterfaces/PCGDataCollectionDataInterface.cpp
+++ b/Engine/Plugins/PCG/Source/PCG/Private/Compute/DataInterfaces/PCGDataCollectionDataInterface.cpp
@@ -379,7 +379,10 @@ void UPCGDataCollectionDataInterface::GetSupportedInputs(TArray<FShaderFunctionD
 			.AddParam(EShaderFundamentalType::Uint)
 			.AddParam(EShaderFundamentalType::Uint)
 			.AddParam(EShaderFundamentalType::Float, 4, 4);
+	}
 
+	// Misc
+	{
 		OutFunctions.AddDefaulted_GetRef()
 			.SetName(TEXT("RemovePoint"))
 			.AddParam(EShaderFundamentalType::Uint) // DataIndex
@@ -389,11 +392,16 @@ void UPCGDataCollectionDataInterface::GetSupportedInputs(TArray<FShaderFunctionD
 			.SetName(TEXT("InitializePoint"))
 			.AddParam(EShaderFundamentalType::Uint)
 			.AddParam(EShaderFundamentalType::Uint);
+
+		OutFunctions.AddDefaulted_GetRef()
+			.SetName(TEXT("GetElementCountMultiplier"))
+			.AddReturnType(EShaderFundamentalType::Uint);
 	}
 }
 
 BEGIN_SHADER_PARAMETER_STRUCT(FPCGDataCollectionDataInterfaceParameters, )
 	SHADER_PARAMETER_RDG_BUFFER_UAV(RWByteAddressBuffer, DataCollectionBuffer)
+	SHADER_PARAMETER(uint32, ElementCountMultiplier)
 	SHADER_PARAMETER_RDG_BUFFER_UAV(RWStructuredBuffer<uint32>, ElementCountersBuffer)
 	SHADER_PARAMETER(uint32, ElementCountersPresent)
 END_SHADER_PARAMETER_STRUCT()
@@ -455,6 +463,8 @@ UComputeDataProvider* UPCGDataCollectionDataInterface::CreateDataProvider(TObjec
 
 	Provider->ExportMode = bRequiresExport ? EPCGExportMode::ComputeGraphOutput : EPCGExportMode::NoExport;
 
+	Provider->ElementCountMultiplier = ElementCountMultiplier;
+
 	Provider->bExposeElementCounters = bExposeElementCounters;
 
 #if WITH_EDITOR
@@ -482,7 +492,7 @@ FComputeDataProviderRenderProxy* UPCGDataCollectionDataProvider::GetRenderProxy(
 {
 	TRACE_CPUPROFILER_EVENT_SCOPE(UPCGDataCollectionDataProvider::GetRenderProxy);
 
-	FPCGDataCollectionDataProviderProxy* Proxy = new FPCGDataCollectionDataProviderProxy(this, PinDesc, ExportMode, bExposeElementCounters, OutputPinLabel, OutputPinLabelAlias);
+	FPCGDataCollectionDataProviderProxy* Proxy = new FPCGDataCollectionDataProviderProxy(this, PinDesc, ExportMode, ElementCountMultiplier, bExposeElementCounters, OutputPinLabel, OutputPinLabelAlias);
 
 #if WITH_EDITOR
 	if (const UPCGNode* Node = Cast<UPCGNode>(ProducerSettings ? ProducerSettings->GetOuter() : nullptr))
@@ -544,10 +554,12 @@ FPCGDataCollectionDataProviderProxy::FPCGDataCollectionDataProviderProxy(
 	TWeakObjectPtr<UPCGDataCollectionDataProvider> InDataProvider,
 	const FPCGDataCollectionDesc& InPinDesc,
 	EPCGExportMode InExportMode,
+	uint32 InElementCountMultiplier,
 	bool bInExposeElementCounters,
 	FName InOutputPinLabel,
 	FName InOutputPinLabelAlias)
 	: ExportMode(InExportMode)
+	, ElementCountMultiplier(InElementCountMultiplier)
 	, bExposeElementCounters(bInExposeElementCounters)
 	, PinDesc(InPinDesc)
 	, DataProviderWeakPtr(InDataProvider)
@@ -583,6 +595,7 @@ void FPCGDataCollectionDataProviderProxy::GatherDispatchData(FDispatchData const
 	{
 		FParameters& Parameters = ParameterArray[InvocationIndex];
 		Parameters.DataCollectionBuffer = BufferUAV;
+		Parameters.ElementCountMultiplier = ElementCountMultiplier;
 		Parameters.ElementCountersBuffer = CountersUAV;
 		Parameters.ElementCountersPresent = bExposeElementCounters ? 1 : 0;
 	}
diff --git a/Engine/Plugins/PCG/Source/PCG/Private/Compute/DataInterfaces/PCGDataCollectionDataInterface.h b/Engine/Plugins/PCG/Source/PCG/Private/Compute/DataInterfaces/PCGDataCollectionDataInterface.h
--- a/Engine/Plugins/PCG/Source/PCG/Private/Compute/DataInterfaces/PCGDataCollectionDataInterface.h
+++ b/Engine/Plugins/PCG/Source/PCG/Private/Compute/DataInterfaces/PCGDataCollectionDataInterface.h
@@ -42,6 +42,8 @@ public:
 
 	void SetRequiresExport(bool bInRequiresExport) { bRequiresExport = bInRequiresExport; }
 
+	void SetElementCountMultiplier(uint32 InElementCountMultiplier) { ElementCountMultiplier = InElementCountMultiplier; }
+
 	void SetExposeElementCounters(bool bInExposeElementCounters) { bExposeElementCounters = bInExposeElementCounters; }
 
 protected:
@@ -53,6 +55,9 @@ protected:
 	UPROPERTY()
 	bool bRequiresExport = false;
 
+	UPROPERTY()
+	uint32 ElementCountMultiplier = 0;
+
 	UPROPERTY()
 	bool bExposeElementCounters = false;
 
@@ -91,6 +96,9 @@ public:
 	UPROPERTY()
 	EPCGExportMode ExportMode = EPCGExportMode::NoExport;
 
+	UPROPERTY()
+	uint32 ElementCountMultiplier = 0;
+
 	UPROPERTY()
 	bool bExposeElementCounters = false;
 
@@ -117,6 +125,7 @@ public:
 		TWeakObjectPtr<UPCGDataCollectionDataProvider> InDataProvider,
 		const FPCGDataCollectionDesc& InPinDesc,
 		EPCGExportMode InExportMode,
+		uint32 InElementCountMultiplier,
 		bool bInExposeElementCounters,
 		FName InOutputPinLabel,
 		FName InOutputPinLabelAlias);
@@ -141,6 +150,8 @@ protected:
 
 	EPCGExportMode ExportMode = EPCGExportMode::NoExport;
 
+	uint32 ElementCountMultiplier = 0;
+
 	bool bExposeElementCounters = false;
 
 	using FParameters = FPCGDataCollectionDataInterfaceParameters;
diff --git a/Engine/Plugins/PCG/Source/PCG/Private/Compute/DataInterfaces/PCGDataCollectionUploadDataInterface.cpp b/Engine/Plugins/PCG/Source/PCG/Private/Compute/DataInterfaces/PCGDataCollectionUploadDataInterface.cpp
--- a/Engine/Plugins/PCG/Source/PCG/Private/Compute/DataInterfaces/PCGDataCollectionUploadDataInterface.cpp
+++ b/Engine/Plugins/PCG/Source/PCG/Private/Compute/DataInterfaces/PCGDataCollectionUploadDataInterface.cpp
@@ -49,7 +49,16 @@ FComputeDataProviderRenderProxy* UPCGDataProviderDataCollectionUpload::GetRender
 {
 	TRACE_CPUPROFILER_EVENT_SCOPE(UPCGDataProviderDataCollectionUpload::GetRenderProxy);
 
-	FPCGDataProviderDataCollectionUploadProxy* Proxy = new FPCGDataProviderDataCollectionUploadProxy(DataAdaptor, this, PinDesc, DownstreamInputPinLabels, OutputPinLabel, OutputPinLabelAlias, ExportMode, bExposeElementCounters);
+	FPCGDataProviderDataCollectionUploadProxy* Proxy = new FPCGDataProviderDataCollectionUploadProxy(
+		DataAdaptor,
+		this,
+		PinDesc,
+		DownstreamInputPinLabels,
+		OutputPinLabel,
+		OutputPinLabelAlias,
+		ExportMode,
+		ElementCountMultiplier,
+		bExposeElementCounters);
 
 #if WITH_EDITOR
 	if (const UPCGNode* Node = Cast<UPCGNode>(ProducerSettings ? ProducerSettings->GetOuter() : nullptr))
@@ -73,8 +82,9 @@ FPCGDataProviderDataCollectionUploadProxy::FPCGDataProviderDataCollectionUploadP
 	FName InOutputPinLabel,
 	FName InOutputPinLabelAlias,
 	EPCGExportMode InExportMode,
+	uint32 InElementCountMultiplier,
 	bool bInExposeElementCounters)
-	: FPCGDataCollectionDataProviderProxy(InDataProvider, InPinDesc, InExportMode, bInExposeElementCounters, InOutputPinLabel, InOutputPinLabelAlias)
+	: FPCGDataCollectionDataProviderProxy(InDataProvider, InPinDesc, InExportMode, InElementCountMultiplier, bInExposeElementCounters, InOutputPinLabel, InOutputPinLabelAlias)
 	, DataAdaptor(InDataAdaptor)
 	, DownstreamInputPinLabels(InDownstreamInputPinLabels)
 {
diff --git a/Engine/Plugins/PCG/Source/PCG/Private/Compute/DataInterfaces/PCGDataCollectionUploadDataInterface.h b/Engine/Plugins/PCG/Source/PCG/Private/Compute/DataInterfaces/PCGDataCollectionUploadDataInterface.h
--- a/Engine/Plugins/PCG/Source/PCG/Private/Compute/DataInterfaces/PCGDataCollectionUploadDataInterface.h
+++ b/Engine/Plugins/PCG/Source/PCG/Private/Compute/DataInterfaces/PCGDataCollectionUploadDataInterface.h
@@ -53,6 +53,7 @@ public:
 		FName InOutputPinLabel,
 		FName InOutputPinLabelAlias,
 		EPCGExportMode InExportMode,
+		uint32 InElementCountMultiplier,
 		bool bInExposeElementCounters);
 
 	//~ Begin FComputeDataProviderRenderProxy Interface
diff --git a/Engine/Plugins/PCG/Source/PCG/Private/Compute/Elements/PCGCustomHLSL.cpp b/Engine/Plugins/PCG/Source/PCG/Private/Compute/Elements/PCGCustomHLSL.cpp
--- a/Engine/Plugins/PCG/Source/PCG/Private/Compute/Elements/PCGCustomHLSL.cpp
+++ b/Engine/Plugins/PCG/Source/PCG/Private/Compute/Elements/PCGCustomHLSL.cpp
@@ -580,9 +580,12 @@ void UPCGCustomHLSLSettings::UpdateOutputDeclarations()
 
 		if (bMultiPin)
 		{
-			OutputDeclarations += TEXT("// Valid pins: ") + FString::Join(DataCollectionDataPins, TEXT(", ")) + TEXT("\n");
+			OutputDeclarations += TEXT("// Valid pins: ") + FString::Join(DataCollectionDataPins, TEXT(", ")) + TEXT("\n\n");
 		}
 
+		OutputDeclarations += FString::Format(TEXT("void {0}_GetElementCountMultiplier();\n\n"),
+			{ bMultiPin ? PCGHLSLElement::PinDeclTemplateStr : DataCollectionDataPins[0] });
+
 		OutputDeclarations += FString::Format(TEXT(
 			"// Valid types: bool, int, float, float2, float3, float4, Rotator (float3), Quat (float4), Transform (float4x4), StringKey (int), Name (uint2)\n"
 			"\n"
diff --git a/Engine/Plugins/PCG/Source/PCG/Private/Compute/Elements/PCGCustomHLSLKernel.cpp b/Engine/Plugins/PCG/Source/PCG/Private/Compute/Elements/PCGCustomHLSLKernel.cpp
--- a/Engine/Plugins/PCG/Source/PCG/Private/Compute/Elements/PCGCustomHLSLKernel.cpp
+++ b/Engine/Plugins/PCG/Source/PCG/Private/Compute/Elements/PCGCustomHLSLKernel.cpp
@@ -252,6 +252,16 @@ bool UPCGCustomHLSLKernel::IsKernelValid(FPCGContext* InContext, bool bQuiet) co
 					return false;
 				}
 			}
+
+			if (Props.ElementCountMultiplier < 1)
+			{
+				PCG_KERNEL_VALIDATION_ERR(InContext, Settings, bQuiet, FText::Format(
+					LOCTEXT("InvalidElementCountMultiplier", "Invalid element count multiplier {0} on output pin '{1}'. Must be greater than 0."),
+					FText::AsNumber(Props.ElementCountMultiplier),
+					FText::FromName(Properties.Label)));
+
+				return false;
+			}
 		}
 
 		bIsFirstOutputPin = false;
@@ -745,7 +755,7 @@ int UPCGCustomHLSLKernel::ComputeThreadCount(const UPCGDataBinding* InBinding) c
 				FPCGDataCollectionDesc Desc;
 				if (ensure(ComputeOutputBindingDataDesc(InBinding->Graph, OutputPin->Label, InBinding, Desc)))
 				{
-					ThreadCount = Desc.ComputeDataElementCount(EPCGDataType::Any);
+					ThreadCount = OutputPin->GetElementCountMultiplier() * Desc.ComputeDataElementCount(EPCGDataType::Any);
 				}
 			}
 		}
@@ -755,7 +765,7 @@ int UPCGCustomHLSLKernel::ComputeThreadCount(const UPCGDataBinding* InBinding) c
 			{
 				if (const FPCGPinProperties* InputPin = Settings->InputPins.FindByPredicate([PinLabel](const FPCGPinProperties& InProps) { return InProps.Label == PinLabel; }))
 				{
-					ThreadCount = GetElementCountForInputPin(*InputPin, InBinding);
+					ThreadCount = FMath::Max(ThreadCount, 1) * GetElementCountForInputPin(*InputPin, InBinding);
 				}
 			}
 		}
@@ -803,6 +813,19 @@ uint32 UPCGCustomHLSLKernel::GetThreadCountMultiplier() const
 	return Settings->ThreadCountMultiplier;
 }
 
+uint32 UPCGCustomHLSLKernel::GetElementCountMultiplier(FName InOutputPinLabel) const
+{
+	if (const FPCGPinPropertiesGPU* PinProps = Settings->OutputPins.FindByPredicate([InOutputPinLabel](const FPCGPinPropertiesGPU& InProps) { return InProps.Label == InOutputPinLabel; }))
+	{
+		return PinProps->GetElementCountMultiplier();
+	}
+	else
+	{
+		ensure(false);
+		return 1u;
+	}
+}
+
 bool UPCGCustomHLSLKernel::IsThreadCountMultiplierInUse() const
 {
 	return Settings->KernelType == EPCGKernelType::Custom && Settings->DispatchThreadCount != EPCGDispatchThreadCount::Fixed;
diff --git a/Engine/Plugins/PCG/Source/PCG/Private/Compute/Elements/PCGCustomHLSLKernel.h b/Engine/Plugins/PCG/Source/PCG/Private/Compute/Elements/PCGCustomHLSLKernel.h
--- a/Engine/Plugins/PCG/Source/PCG/Private/Compute/Elements/PCGCustomHLSLKernel.h
+++ b/Engine/Plugins/PCG/Source/PCG/Private/Compute/Elements/PCGCustomHLSLKernel.h
@@ -25,6 +25,7 @@ public:
 	virtual void GetDataLabels(FName InPinLabel, TArray<FString>& OutDataLabels) const override;
 	virtual void GetKernelAttributeKeys(TArray<FPCGKernelAttributeKey>& OutKeys) const override;
 	virtual uint32 GetThreadCountMultiplier() const override;
+	virtual uint32 GetElementCountMultiplier(FName InOutputPinLabel) const override;
 
 protected:
 	/** Will the ThreadCountMultiplier value be applied when calculating the dispatch thread count. */
diff --git a/Engine/Plugins/PCG/Source/PCG/Private/Compute/PCGComputeKernel.cpp b/Engine/Plugins/PCG/Source/PCG/Private/Compute/PCGComputeKernel.cpp
--- a/Engine/Plugins/PCG/Source/PCG/Private/Compute/PCGComputeKernel.cpp
+++ b/Engine/Plugins/PCG/Source/PCG/Private/Compute/PCGComputeKernel.cpp
@@ -213,6 +213,13 @@ void UPCGComputeKernel::ComputeDataDescFromPinProperties(const FPCGPinProperties
 		{
 			checkNoEntry();
 		}
+
+		// Apply element count multiplier.
+		const uint32 Multiplier = OutputPinProps.GetElementCountMultiplier();
+		for (FPCGDataDesc& Desc : OutPinDesc.DataDescs)
+		{
+			Desc.ElementCount *= Multiplier;
+		}
 	}
 	else if (Props.InitializationMode == EPCGPinInitMode::Custom)
 	{
diff --git a/Engine/Plugins/PCG/Source/PCG/Private/Compute/PCGComputeKernel.h b/Engine/Plugins/PCG/Source/PCG/Private/Compute/PCGComputeKernel.h
--- a/Engine/Plugins/PCG/Source/PCG/Private/Compute/PCGComputeKernel.h
+++ b/Engine/Plugins/PCG/Source/PCG/Private/Compute/PCGComputeKernel.h
@@ -67,6 +67,9 @@ public:
 	/** Get multiplier to be applied to the dispatched thread count. */
 	virtual uint32 GetThreadCountMultiplier() const { return 1u; }
 
+	/** Multiplier to be applied to the element count of each output data. */
+	virtual uint32 GetElementCountMultiplier(FName InOutputPinLabel) const { return 1u; }
+
 	/** Helper to compute data descriptions when data is entirely determined by pin properties. */
 	void ComputeDataDescFromPinProperties(const FPCGPinPropertiesGPU& OutputPinProps, const TArrayView<const FPCGPinProperties>& InInputPinProps, const UPCGDataBinding* InBinding, FPCGDataCollectionDesc& OutPinDesc) const;
 
diff --git a/Engine/Plugins/PCG/Source/PCG/Private/Compute/PCGPinPropertiesGPU.cpp b/Engine/Plugins/PCG/Source/PCG/Private/Compute/PCGPinPropertiesGPU.cpp
--- a/Engine/Plugins/PCG/Source/PCG/Private/Compute/PCGPinPropertiesGPU.cpp
+++ b/Engine/Plugins/PCG/Source/PCG/Private/Compute/PCGPinPropertiesGPU.cpp
@@ -4,6 +4,18 @@
 
 #include UE_INLINE_GENERATED_CPP_BY_NAME(PCGPinPropertiesGPU)
 
+uint32 FPCGPinPropertiesGPU::GetElementCountMultiplier() const
+{
+	if (PropertiesGPU.ElementCountMode != EPCGElementCountMode::Fixed)
+	{
+		return static_cast<uint32>(FMath::Max(PropertiesGPU.ElementCountMultiplier, 1));
+	}
+	else
+	{
+		return 1u;
+	}
+}
+
 #if WITH_EDITOR
 bool FPCGPinPropertiesGPU::CanEditChange(const FEditPropertyChain& PropertyChain) const
 {
diff --git a/Engine/Plugins/PCG/Source/PCG/Private/Compute/PCGPinPropertiesGPU.h b/Engine/Plugins/PCG/Source/PCG/Private/Compute/PCGPinPropertiesGPU.h
--- a/Engine/Plugins/PCG/Source/PCG/Private/Compute/PCGPinPropertiesGPU.h
+++ b/Engine/Plugins/PCG/Source/PCG/Private/Compute/PCGPinPropertiesGPU.h
@@ -95,6 +95,10 @@ struct FPCGPinPropertiesGPUStruct
 	UPROPERTY(BlueprintReadWrite, EditAnywhere, Category = Settings, meta = (EditCondition = "bAllowEditInitMode && (InitializationMode == EPCGPinInitMode::Custom || ElementCountMode == EPCGElementCountMode::Fixed)", EditConditionHides, ClampMin = 1))
 	int ElementCount = 1;
 
+	/** Fixed number of elements to create in each output data. */
+	UPROPERTY(BlueprintReadWrite, EditAnywhere, Category = Settings, meta = (EditCondition = "bAllowEditInitMode && (InitializationMode == EPCGPinInitMode::Custom || ElementCountMode == EPCGElementCountMode::FromInputData)", EditConditionHides, ClampMin = 1))
+	int ElementCountMultiplier = 1;
+
 	/** How to inherit attribute names, types, and values. */
 	UPROPERTY(BlueprintReadWrite, EditAnywhere, Category = Settings, meta = (EditCondition = "bAllowEditInitMode && InitializationMode == EPCGPinInitMode::FromInputPins", EditConditionHides))
 	EPCGAttributeInheritanceMode AttributeInheritanceMode = EPCGAttributeInheritanceMode::CopyAttributeSetup;
@@ -133,6 +137,8 @@ public:
 		: FPCGPinProperties(InPinProps.Label, InPinProps.AllowedTypes)
 	{}
 
+	uint32 GetElementCountMultiplier() const;
+
 #if WITH_EDITOR
 	bool CanEditChange(const FEditPropertyChain& PropertyChain) const;
 #endif
diff --git a/Engine/Plugins/PCG/Source/PCG/Private/Graph/PCGGraphCompilerGPU.cpp b/Engine/Plugins/PCG/Source/PCG/Private/Graph/PCGGraphCompilerGPU.cpp
--- a/Engine/Plugins/PCG/Source/PCG/Private/Graph/PCGGraphCompilerGPU.cpp
+++ b/Engine/Plugins/PCG/Source/PCG/Private/Graph/PCGGraphCompilerGPU.cpp
@@ -660,6 +660,7 @@ void FPCGGraphCompilerGPU::CreateDataInterfaces(
 			check(DataInterfacePCGData);
 			DataInterfacePCGData->SetProducerSettings(ProducerSettings);
 			DataInterfacePCGData->SetRequiresExport(bRequiresExport);
+			DataInterfacePCGData->SetElementCountMultiplier(ProducerKernel ? ProducerKernel->GetElementCountMultiplier(ProducerPin->Properties.Label) : 1);
 			DataInterfacePCGData->SetExposeElementCounters(ProducerKernel && ProducerPin && ProducerKernel->DoesOutputPinRequireElementCounters(ProducerPin->Properties.Label));
 
 			DataInterface = DataInterfacePCGData;
