[PCG] Fix static analysis - missing pointer guard

#rnx

diff --git a/Engine/Plugins/PCG/Source/PCG/Private/Graph/PCGGraphCompilerGPU.cpp b/Engine/Plugins/PCG/Source/PCG/Private/Graph/PCGGraphCompilerGPU.cpp
--- a/Engine/Plugins/PCG/Source/PCG/Private/Graph/PCGGraphCompilerGPU.cpp
+++ b/Engine/Plugins/PCG/Source/PCG/Private/Graph/PCGGraphCompilerGPU.cpp
@@ -660,7 +660,7 @@ void FPCGGraphCompilerGPU::CreateDataInterfaces(
 			check(DataInterfacePCGData);
 			DataInterfacePCGData->SetProducerSettings(ProducerSettings);
 			DataInterfacePCGData->SetRequiresExport(bRequiresExport);
-			DataInterfacePCGData->SetElementCountMultiplier(ProducerKernel ? ProducerKernel->GetElementCountMultiplier(ProducerPin->Properties.Label) : 1);
+			DataInterfacePCGData->SetElementCountMultiplier((ProducerKernel && ProducerPin) ? ProducerKernel->GetElementCountMultiplier(ProducerPin->Properties.Label) : 1);
 			DataInterfacePCGData->SetExposeElementCounters(ProducerKernel && ProducerPin && ProducerKernel->DoesOutputPinRequireElementCounters(ProducerPin->Properties.Label));
 
 			DataInterface = DataInterfacePCGData;
