[PCG] Fix ProjectPoint metadata

* Metadata entry key was discarded and not saved on the out point

#rnx

diff --git a/Engine/Plugins/PCG/Source/PCG/Private/Data/PCGBasePointData.cpp b/Engine/Plugins/PCG/Source/PCG/Private/Data/PCGBasePointData.cpp
--- a/Engine/Plugins/PCG/Source/PCG/Private/Data/PCGBasePointData.cpp
+++ b/Engine/Plugins/PCG/Source/PCG/Private/Data/PCGBasePointData.cpp
@@ -1409,11 +1409,11 @@ bool UPCGBasePointData::ProjectPoint(const FTransform& InTransform, const FBox&
 	{
 		//TRACE_CPUPROFILER_EVENT_SCOPE(UPCGPointData::SamplePoint::SetupMetadata);
 		// Initialise metadata entry for this temporary point
-		PCGMetadataEntryKey Key = OutMetadata ? (OutMetadata->HasParent(Metadata) ? OutMetadata->AddEntry(MetadataEntryRange[MaxContributor]) : OutMetadata->AddEntry()) : PCGInvalidEntryKey;
+		OutPoint.MetadataEntry = OutMetadata ? (OutMetadata->HasParent(Metadata) ? OutMetadata->AddEntry(MetadataEntryRange[MaxContributor]) : OutMetadata->AddEntry()) : PCGInvalidEntryKey;
 		
 		if (ContributionsForMetadata.Num() > 1)
 		{
-			OutMetadata->ComputeWeightedAttribute(Key, MakeArrayView(ContributionsForMetadata), Metadata);
+			OutMetadata->ComputeWeightedAttribute(OutPoint.MetadataEntry, MakeArrayView(ContributionsForMetadata), Metadata);
 		}
 	}
 
