[PCG] Convert all FPCGPointOperationElementBase sub classes to support PCGPointArrayData

#rnx

diff --git a/Engine/Plugins/PCG/Source/PCG/Private/Elements/PCGApplyScaleToBounds.cpp b/Engine/Plugins/PCG/Source/PCG/Private/Elements/PCGApplyScaleToBounds.cpp
--- a/Engine/Plugins/PCG/Source/PCG/Private/Elements/PCGApplyScaleToBounds.cpp
+++ b/Engine/Plugins/PCG/Source/PCG/Private/Elements/PCGApplyScaleToBounds.cpp
@@ -4,7 +4,7 @@
 
 #include "PCGContext.h"
 #include "PCGPin.h"
-#include "Data/PCGPointData.h"
+#include "Data/PCGBasePointData.h"
 
 #define LOCTEXT_NAMESPACE "PCGApplyScaleToBoundsElement"
 
@@ -13,6 +13,11 @@ FPCGElementPtr UPCGApplyScaleToBoundsSettings::CreateElement() const
 	return MakeShared<FPCGApplyScaleToBoundsElement>();
 }
 
+EPCGPointNativeProperties FPCGApplyScaleToBoundsElement::GetPropertiesToAllocate(FPCGContext* Context) const
+{
+	return EPCGPointNativeProperties::Transform | EPCGPointNativeProperties::BoundsMin | EPCGPointNativeProperties::BoundsMax;
+}
+
 bool FPCGApplyScaleToBoundsElement::ExecuteInternal(FPCGContext* Context) const
 {
 	TRACE_CPUPROFILER_EVENT_SCOPE(FPCGApplyScaleToBoundsElement::Execute);
@@ -20,11 +25,19 @@ bool FPCGApplyScaleToBoundsElement::ExecuteInternal(FPCGContext* Context) const
 
 	FPCGApplyScaleToBoundsElement::ContextType* ApplyScaleContext = static_cast<FPCGApplyScaleToBoundsElement::ContextType*>(Context);
 
-	return ExecutePointOperation(ApplyScaleContext, [](const FPCGPoint& InPoint, FPCGPoint& OutPoint)
+	return ExecutePointOperation(ApplyScaleContext, [](const UPCGBasePointData* InputData, UPCGBasePointData* OutputData, int32 StartIndex, int32 Count)
 	{
-		OutPoint = InPoint;
-
-		OutPoint.ApplyScaleToBounds();
+		TPCGValueRange<FVector> BoundsMinRange = OutputData->GetBoundsMinValueRange();
+		TPCGValueRange<FVector> BoundsMaxRange = OutputData->GetBoundsMaxValueRange();
+		TPCGValueRange<FTransform> TransformRange = OutputData->GetTransformValueRange();
+
+		for (int32 Index = StartIndex; Index < (StartIndex + Count); ++Index)
+		{
+			FVector& BoundsMin = BoundsMinRange[Index];
+			FVector& BoundsMax = BoundsMaxRange[Index];
+			FTransform& Transform = TransformRange[Index];
+			PCGPointHelpers::ApplyScaleToBounds(Transform, BoundsMin, BoundsMax);
+		}
 		
 		return true;
 	});
diff --git a/Engine/Plugins/PCG/Source/PCG/Private/Elements/PCGBoundsModifier.cpp b/Engine/Plugins/PCG/Source/PCG/Private/Elements/PCGBoundsModifier.cpp
--- a/Engine/Plugins/PCG/Source/PCG/Private/Elements/PCGBoundsModifier.cpp
+++ b/Engine/Plugins/PCG/Source/PCG/Private/Elements/PCGBoundsModifier.cpp
@@ -41,122 +41,101 @@ bool FPCGBoundsModifier::ExecuteInternal(FPCGContext* Context) const
 	FPCGBoundsModifier::ContextType* BoundsModifierContext = static_cast<FPCGBoundsModifier::ContextType*>(Context);
 
 	const UPCGBoundsModifierSettings* Settings = Context->GetInputSettings<UPCGBoundsModifierSettings>();
+	const FBox SettingsBounds(Settings->BoundsMin, Settings->BoundsMax);
 	check(Settings);
 
+	auto SetBoundsLoop = [Settings, &SettingsBounds](UPCGBasePointData* OutputData, int32 StartIndex, int32 Count, TFunctionRef<void(const FBox&, const FBox&, FVector&, FVector&, float&)> BoundsModifierLoop)
+	{
+		TPCGValueRange<FVector> BoundsMinRange = OutputData->GetBoundsMinValueRange();
+		TPCGValueRange<FVector> BoundsMaxRange = OutputData->GetBoundsMaxValueRange();
+		TPCGValueRange<float> SteepnessRange = OutputData->GetSteepnessValueRange();
+				
+		for (int32 Index = StartIndex; Index < (StartIndex + Count); ++Index)
+		{
+			FVector& BoundsMin = BoundsMinRange[Index];
+			FVector& BoundsMax = BoundsMaxRange[Index];
+			float& Steepness = SteepnessRange[Index];
+
+			const FBox LocalBounds = PCGPointHelpers::GetLocalBounds(BoundsMin, BoundsMax);
+
+			BoundsModifierLoop(SettingsBounds, LocalBounds, BoundsMin, BoundsMax, Steepness);
+		}
+
+		return true;
+	};
+
 	switch (Settings->Mode)
 	{
 		case EPCGBoundsModifierMode::Intersect:
-			return ExecutePointOperation(BoundsModifierContext, [Settings](const UPCGBasePointData* InputData, UPCGBasePointData* OutputData, int32 StartIndex, int32 Count)
+			return ExecutePointOperation(BoundsModifierContext, [SetBoundsLoop, Settings](const UPCGBasePointData* InputData, UPCGBasePointData* OutputData, int32 StartIndex, int32 Count)
 			{
-				TPCGValueRange<FVector> BoundsMinRange = OutputData->GetBoundsMinValueRange();
-				TPCGValueRange<FVector> BoundsMaxRange = OutputData->GetBoundsMaxValueRange();
-				TPCGValueRange<float> SteepnessRange = OutputData->GetSteepnessValueRange();
-
-				const FBox SettingsBounds(Settings->BoundsMin, Settings->BoundsMax);
-
-				for (int32 Index = StartIndex; Index < (StartIndex + Count); ++Index)
+				return SetBoundsLoop(OutputData, StartIndex, Count, [Settings](const FBox& SettingsBounds, const FBox& LocalBounds, FVector& OutBoundsMin, FVector& OutBoundsMax, float& OutSteepness)
 				{
-					FVector& BoundsMin = BoundsMinRange[Index];
-					FVector& BoundsMax = BoundsMaxRange[Index];
-					const FBox LocalBounds = PCGPointHelpers::GetLocalBounds(BoundsMin, BoundsMax);
-					PCGPointHelpers::SetLocalBounds(LocalBounds.Overlap(SettingsBounds), BoundsMin, BoundsMax);
-				
+					PCGPointHelpers::SetLocalBounds(LocalBounds.Overlap(SettingsBounds), OutBoundsMin, OutBoundsMax);
+
 					if (Settings->bAffectSteepness)
 					{
-						SteepnessRange[Index] = FMath::Min(SteepnessRange[Index], Settings->Steepness);
+						OutSteepness = FMath::Min(OutSteepness, Settings->Steepness);
 					}
-				}
-
-				return true;
+				});
 			}, PCGBoundsModifier::PointsPerChunk);
 
 		case EPCGBoundsModifierMode::Include:
-			return ExecutePointOperation(BoundsModifierContext, [Settings](const UPCGBasePointData* InputData, UPCGBasePointData* OutputData, int32 StartIndex, int32 Count)
+			return ExecutePointOperation(BoundsModifierContext, [SetBoundsLoop, Settings](const UPCGBasePointData* InputData, UPCGBasePointData* OutputData, int32 StartIndex, int32 Count)
 			{
-				TPCGValueRange<FVector> BoundsMinRange = OutputData->GetBoundsMinValueRange();
-				TPCGValueRange<FVector> BoundsMaxRange = OutputData->GetBoundsMaxValueRange();
-				TPCGValueRange<float> SteepnessRange = OutputData->GetSteepnessValueRange();
-
-				const FBox SettingsBounds(Settings->BoundsMin, Settings->BoundsMax);
-
-				for (int32 Index = StartIndex; Index < (StartIndex + Count); ++Index)
+				return SetBoundsLoop(OutputData, StartIndex, Count, [Settings](const FBox& SettingsBounds, const FBox& LocalBounds, FVector& OutBoundsMin, FVector& OutBoundsMax, float& OutSteepness)
 				{
-					FVector& BoundsMin = BoundsMinRange[Index];
-					FVector& BoundsMax = BoundsMaxRange[Index];
-					const FBox LocalBounds = PCGPointHelpers::GetLocalBounds(BoundsMin, BoundsMax);
-					PCGPointHelpers::SetLocalBounds(LocalBounds + SettingsBounds, BoundsMin, BoundsMax);
-				
+					PCGPointHelpers::SetLocalBounds(LocalBounds + SettingsBounds, OutBoundsMin, OutBoundsMax);
+
 					if (Settings->bAffectSteepness)
 					{
-						SteepnessRange[Index] = FMath::Max(SteepnessRange[Index], Settings->Steepness);
+						OutSteepness = FMath::Max(OutSteepness, Settings->Steepness);
 					}
-				}
-
-				return true;
+				});
 			}, PCGBoundsModifier::PointsPerChunk);
 
 		case EPCGBoundsModifierMode::Translate:
-			return ExecutePointOperation(BoundsModifierContext, [Settings](const UPCGBasePointData* InputData, UPCGBasePointData* OutputData, int32 StartIndex, int32 Count)
+			return ExecutePointOperation(BoundsModifierContext, [SetBoundsLoop, Settings](const UPCGBasePointData* InputData, UPCGBasePointData* OutputData, int32 StartIndex, int32 Count)
 			{
-				TPCGValueRange<FVector> BoundsMinRange = OutputData->GetBoundsMinValueRange();
-				TPCGValueRange<FVector> BoundsMaxRange = OutputData->GetBoundsMaxValueRange();
-				TPCGValueRange<float> SteepnessRange = OutputData->GetSteepnessValueRange();
-
-				for (int32 Index = StartIndex; Index < (StartIndex + Count); ++Index)
+				return SetBoundsLoop(OutputData, StartIndex, Count, [Settings](const FBox& SettingsBounds, const FBox& LocalBounds, FVector& OutBoundsMin, FVector& OutBoundsMax, float& OutSteepness)
 				{
-					BoundsMinRange[Index] += Settings->BoundsMin;
-					BoundsMaxRange[Index] += Settings->BoundsMax;
+					OutBoundsMin += Settings->BoundsMin;
+					OutBoundsMax += Settings->BoundsMax;
 
 					if (Settings->bAffectSteepness)
 					{
-						SteepnessRange[Index] = FMath::Clamp(SteepnessRange[Index] + Settings->Steepness, 0.0f, 1.0f);
+						OutSteepness = FMath::Clamp(OutSteepness + Settings->Steepness, 0.0f, 1.0f);
 					}
-				}
-						
-				return true;
+				});
 			}, PCGBoundsModifier::PointsPerChunk);
 
 		case EPCGBoundsModifierMode::Scale:
-			return ExecutePointOperation(BoundsModifierContext, [Settings](const UPCGBasePointData* InputData, UPCGBasePointData* OutputData, int32 StartIndex, int32 Count)
+			return ExecutePointOperation(BoundsModifierContext, [SetBoundsLoop, Settings](const UPCGBasePointData* InputData, UPCGBasePointData* OutputData, int32 StartIndex, int32 Count)
 			{
-				TPCGValueRange<FVector> BoundsMinRange = OutputData->GetBoundsMinValueRange();
-				TPCGValueRange<FVector> BoundsMaxRange = OutputData->GetBoundsMaxValueRange();
-				TPCGValueRange<float> SteepnessRange = OutputData->GetSteepnessValueRange();
-
-				for (int32 Index = StartIndex; Index < (StartIndex + Count); ++Index)
+				return SetBoundsLoop(OutputData, StartIndex, Count, [Settings](const FBox& SettingsBounds, const FBox& LocalBounds, FVector& OutBoundsMin, FVector& OutBoundsMax, float& OutSteepness)
 				{
-					BoundsMinRange[Index] *= Settings->BoundsMin;
-					BoundsMaxRange[Index] *= Settings->BoundsMax;
+					OutBoundsMin *= Settings->BoundsMin;
+					OutBoundsMax *= Settings->BoundsMax;
 
 					if (Settings->bAffectSteepness)
 					{
-						SteepnessRange[Index] = FMath::Clamp(SteepnessRange[Index] * Settings->Steepness, 0.0f, 1.0f);
+						OutSteepness = FMath::Clamp(OutSteepness * Settings->Steepness, 0.0f, 1.0f);
 					}
-				}
-			
-				return true;
+				});
 			}, PCGBoundsModifier::PointsPerChunk);
 
 		case EPCGBoundsModifierMode::Set:
-			return ExecutePointOperation(BoundsModifierContext, [Settings](const UPCGBasePointData* InputData, UPCGBasePointData* OutputData, int32 StartIndex, int32 Count)
+			return ExecutePointOperation(BoundsModifierContext, [SetBoundsLoop, Settings](const UPCGBasePointData* InputData, UPCGBasePointData* OutputData, int32 StartIndex, int32 Count)
 			{
-				TPCGValueRange<FVector> BoundsMinRange = OutputData->GetBoundsMinValueRange();
-				TPCGValueRange<FVector> BoundsMaxRange = OutputData->GetBoundsMaxValueRange();
-				TPCGValueRange<float> SteepnessRange = OutputData->GetSteepnessValueRange();
-
-				const FBox SettingsBounds(Settings->BoundsMin, Settings->BoundsMax);
-
-				for (int32 Index = StartIndex; Index < (StartIndex + Count); ++Index)
+				return SetBoundsLoop(OutputData, StartIndex, Count, [Settings](const FBox& SettingsBounds, const FBox& LocalBounds, FVector& OutBoundsMin, FVector& OutBoundsMax, float& OutSteepness)
 				{
-					PCGPointHelpers::SetLocalBounds(SettingsBounds, BoundsMinRange[Index], BoundsMaxRange[Index]);
-				
+					PCGPointHelpers::SetLocalBounds(SettingsBounds, OutBoundsMin, OutBoundsMax);
+
 					if (Settings->bAffectSteepness)
 					{
-						SteepnessRange[Index] = Settings->Steepness;
+						OutSteepness = Settings->Steepness;
 					}
-				}
-
-				return true;
+				});
 			}, PCGBoundsModifier::PointsPerChunk);
 
 		default:
diff --git a/Engine/Plugins/PCG/Source/PCG/Private/Elements/PCGDensityRemapElement.cpp b/Engine/Plugins/PCG/Source/PCG/Private/Elements/PCGDensityRemapElement.cpp
--- a/Engine/Plugins/PCG/Source/PCG/Private/Elements/PCGDensityRemapElement.cpp
+++ b/Engine/Plugins/PCG/Source/PCG/Private/Elements/PCGDensityRemapElement.cpp
@@ -3,8 +3,8 @@
 #include "Elements/PCGDensityRemapElement.h"
 
 #include "PCGContext.h"
-
 #include "PCGPoint.h"
+#include "Data/PCGBasePointData.h"
 
 #include UE_INLINE_GENERATED_CPP_BY_NAME(PCGDensityRemapElement)
 
@@ -20,6 +20,11 @@ FPCGElementPtr UPCGDensityRemapSettings::CreateElement() const
 	return MakeShared<FPCGDensityRemapElement>();
 }
 
+EPCGPointNativeProperties FPCGDensityRemapElement::GetPropertiesToAllocate(FPCGContext* InContext) const
+{
+	return EPCGPointNativeProperties::Density;
+}
+
 bool FPCGDensityRemapElement::ExecuteInternal(FPCGContext* InContext) const
 {
 	TRACE_CPUPROFILER_EVENT_SCOPE(FPCGDensityRemapElement::Execute);
@@ -59,15 +64,19 @@ PRAGMA_ENABLE_DEPRECATION_WARNINGS
 		Intercept = OutRangeMin;
 	}
 
-	return ExecutePointOperation(Context, [bExcludeValuesOutsideInputRange, InRangeTrueMin, InRangeTrueMax, Slope, InRangeMin, Intercept](const FPCGPoint& InPoint, FPCGPoint& OutPoint)
+	return ExecutePointOperation(Context, [bExcludeValuesOutsideInputRange, InRangeTrueMin, InRangeTrueMax, Slope, InRangeMin, Intercept](const UPCGBasePointData* InputData, UPCGBasePointData* OutputData, int32 StartIndex, int32 Count)
 	{
-		OutPoint = InPoint;
-		const float SourceDensity = InPoint.Density;
+		const TConstPCGValueRange<float> SourceDensityRange = InputData->GetConstDensityValueRange();
+		TPCGValueRange<float> DensityRange = OutputData->GetDensityValueRange();
 
-		if (!bExcludeValuesOutsideInputRange || (SourceDensity >= InRangeTrueMin && SourceDensity <= InRangeTrueMax))
+		for (int32 Index = StartIndex; Index < (StartIndex + Count); ++Index)
 		{
-			const float UnclampedDensity = Slope * (SourceDensity - InRangeMin) + Intercept;
-			OutPoint.Density = FMath::Clamp(UnclampedDensity, 0.f, 1.f);
+			const float SourceDensity = SourceDensityRange[Index];
+			if (!bExcludeValuesOutsideInputRange || (SourceDensity >= InRangeTrueMin && SourceDensity <= InRangeTrueMax))
+			{
+				const float UnclampedDensity = Slope * (SourceDensity - InRangeMin) + Intercept;
+				DensityRange[Index] = FMath::Clamp(UnclampedDensity, 0.f, 1.f);
+			}
 		}
 
 		return true;
diff --git a/Engine/Plugins/PCG/Source/PCG/Private/Elements/PCGMutateSeed.cpp b/Engine/Plugins/PCG/Source/PCG/Private/Elements/PCGMutateSeed.cpp
--- a/Engine/Plugins/PCG/Source/PCG/Private/Elements/PCGMutateSeed.cpp
+++ b/Engine/Plugins/PCG/Source/PCG/Private/Elements/PCGMutateSeed.cpp
@@ -3,6 +3,7 @@
 #include "Elements/PCGMutateSeed.h"
 
 #include "PCGContext.h"
+#include "Data/PCGBasePointData.h"
 #include "Helpers/PCGHelpers.h"
 
 #define LOCTEXT_NAMESPACE "PCGMutateSeedSettings"
@@ -19,6 +20,11 @@ FPCGElementPtr UPCGMutateSeedSettings::CreateElement() const
 	return MakeShared<FPCGMutateSeedElement>();
 }
 
+EPCGPointNativeProperties FPCGMutateSeedElement::GetPropertiesToAllocate(FPCGContext* Context) const
+{
+	return EPCGPointNativeProperties::Seed;
+}
+
 bool FPCGMutateSeedElement::ExecuteInternal(FPCGContext* Context) const
 {
 	TRACE_CPUPROFILER_EVENT_SCOPE(FPCGMutateSeedElement::Execute);
@@ -26,10 +32,16 @@ bool FPCGMutateSeedElement::ExecuteInternal(FPCGContext* Context) const
 	check(Context);
 	ContextType* MutateSeedContext = static_cast<ContextType*>(Context);
 
-	return ExecutePointOperation(MutateSeedContext, [Seed = Context->GetSeed()](const FPCGPoint& InPoint, FPCGPoint& OutPoint)
+	return ExecutePointOperation(MutateSeedContext, [Seed = Context->GetSeed()](const UPCGBasePointData* InputData, UPCGBasePointData* OutputData, int32 StartIndex, int32 Count)
 	{
-		OutPoint = InPoint;
-		OutPoint.Seed = PCGHelpers::ComputeSeed(PCGHelpers::ComputeSeedFromPosition(OutPoint.Transform.GetLocation()), Seed, OutPoint.Seed);
+		TConstPCGValueRange<FTransform> TransformRange = OutputData->GetConstTransformValueRange();
+		TPCGValueRange<int32> SeedRange = OutputData->GetSeedValueRange();
+
+		for (int32 Index = StartIndex; Index < (StartIndex + Count); ++Index)
+		{
+			SeedRange[Index] = PCGHelpers::ComputeSeed(PCGHelpers::ComputeSeedFromPosition(TransformRange[Index].GetLocation()), Seed, SeedRange[Index]);
+		}
+
 		return true;
 	}, PCGMutateSeedConstants::PointsPerChunk);
 }
diff --git a/Engine/Plugins/PCG/Source/PCG/Private/Elements/PCGNormalToDensity.cpp b/Engine/Plugins/PCG/Source/PCG/Private/Elements/PCGNormalToDensity.cpp
--- a/Engine/Plugins/PCG/Source/PCG/Private/Elements/PCGNormalToDensity.cpp
+++ b/Engine/Plugins/PCG/Source/PCG/Private/Elements/PCGNormalToDensity.cpp
@@ -11,6 +11,11 @@ FPCGElementPtr UPCGNormalToDensitySettings::CreateElement() const
 	return MakeShared<FPCGNormalToDensityElement>();
 }
 
+EPCGPointNativeProperties FPCGNormalToDensityElement::GetPropertiesToAllocate(FPCGContext* Context) const
+{
+	return EPCGPointNativeProperties::Density;
+}
+
 bool FPCGNormalToDensityElement::ExecuteInternal(FPCGContext* InContext) const
 {
 	TRACE_CPUPROFILER_EVENT_SCOPE(FPCGNormalToDensityElement::Execute);
@@ -27,76 +32,88 @@ bool FPCGNormalToDensityElement::ExecuteInternal(FPCGContext* InContext) const
 
 	const double InvStrength = 1.0/FMath::Max(0.0001, Strength);
 
-	const auto CalcValue = [Normal, Offset, InvStrength](const FPCGPoint& InPoint)
+	const auto CalcValue = [Normal, Offset, InvStrength](const FTransform& InPointTransform)
 	{
-		const FVector Up = InPoint.Transform.GetUnitAxis(EAxis::Z);
+		const FVector Up = InPointTransform.GetUnitAxis(EAxis::Z);
 		return FMath::Pow(FMath::Clamp(Up.Dot(Normal) + Offset, 0.0, 1.0), InvStrength);
 	};
 
+	const auto SetDensityLoop = [CalcValue](UPCGBasePointData* OutputData, int32 StartIndex, int32 Count, TFunctionRef<void(float, float&)> SetDensityFunc)
+	{
+		TConstPCGValueRange<FTransform> TransformRange = OutputData->GetConstTransformValueRange();
+		TPCGValueRange<float> DensityRange = OutputData->GetDensityValueRange();
+
+		for (int32 Index = StartIndex; Index < (StartIndex + Count); ++Index)
+		{
+			SetDensityFunc(CalcValue(TransformRange[Index]), DensityRange[Index]);
+		}
+
+		return true;
+	};
+
 	switch (DensityMode)
 	{
 		case PCGNormalToDensityMode::Set:
-			return ExecutePointOperation(Context, [CalcValue](const FPCGPoint& InPoint, FPCGPoint& OutPoint)
+			return ExecutePointOperation(Context, [SetDensityLoop](const UPCGBasePointData* InputData, UPCGBasePointData* OutputData, int32 StartIndex, int32 Count)
 			{
-				OutPoint = InPoint;
-
-				OutPoint.Density = CalcValue(InPoint);
-
-				return true;
+				return SetDensityLoop(OutputData, StartIndex, Count, [](float InNormalDensity, float& OutDensity)
+				{
+					OutDensity = InNormalDensity;
+				});
 			});
-
+			
 		case PCGNormalToDensityMode::Minimum:
-			return ExecutePointOperation(Context, [CalcValue](const FPCGPoint& InPoint, FPCGPoint& OutPoint)->bool{
-				OutPoint = InPoint;
-
-				OutPoint.Density = FMath::Min(InPoint.Density, CalcValue(InPoint));
-
-				return true;
+			return ExecutePointOperation(Context, [SetDensityLoop](const UPCGBasePointData* InputData, UPCGBasePointData* OutputData, int32 StartIndex, int32 Count)
+			{
+				return SetDensityLoop(OutputData, StartIndex, Count, [](float InNormalDensity, float& OutDensity)
+				{
+					OutDensity = FMath::Min(OutDensity, InNormalDensity);
+				});
 			});
 
 		case PCGNormalToDensityMode::Maximum:
-			return ExecutePointOperation(Context, [CalcValue](const FPCGPoint& InPoint, FPCGPoint& OutPoint)->bool{
-				OutPoint = InPoint;
-
-				OutPoint.Density = FMath::Max(InPoint.Density, CalcValue(InPoint));
-
-				return true;
+			return ExecutePointOperation(Context, [SetDensityLoop](const UPCGBasePointData* InputData, UPCGBasePointData* OutputData, int32 StartIndex, int32 Count)
+			{
+				return SetDensityLoop(OutputData, StartIndex, Count, [](float InNormalDensity, float& OutDensity)
+				{
+					OutDensity = FMath::Max(OutDensity, InNormalDensity);
+				});
 			});
 
 		case PCGNormalToDensityMode::Add:
-			return ExecutePointOperation(Context, [CalcValue](const FPCGPoint& InPoint, FPCGPoint& OutPoint)->bool{
-				OutPoint = InPoint;
-
-				OutPoint.Density += CalcValue(InPoint);
-
-				return true;
+			return ExecutePointOperation(Context, [SetDensityLoop](const UPCGBasePointData* InputData, UPCGBasePointData* OutputData, int32 StartIndex, int32 Count)
+			{
+				return SetDensityLoop(OutputData, StartIndex, Count, [](float InNormalDensity, float& OutDensity)
+				{
+					OutDensity += InNormalDensity;
+				});
 			});
 
 		case PCGNormalToDensityMode::Subtract:
-			return ExecutePointOperation(Context, [CalcValue](const FPCGPoint& InPoint, FPCGPoint& OutPoint)->bool{
-				OutPoint = InPoint;
-
-				OutPoint.Density -= CalcValue(InPoint);
-
-				return true;
+			return ExecutePointOperation(Context, [SetDensityLoop](const UPCGBasePointData* InputData, UPCGBasePointData* OutputData, int32 StartIndex, int32 Count)
+			{
+				return SetDensityLoop(OutputData, StartIndex, Count, [](float InNormalDensity, float& OutDensity)
+				{
+					OutDensity -= InNormalDensity;
+				});
 			});
 
 		case PCGNormalToDensityMode::Multiply:
-			return ExecutePointOperation(Context, [CalcValue](const FPCGPoint& InPoint, FPCGPoint& OutPoint)->bool{
-				OutPoint = InPoint;
-
-				OutPoint.Density *= CalcValue(InPoint);
-
-				return true;
+			return ExecutePointOperation(Context, [SetDensityLoop](const UPCGBasePointData* InputData, UPCGBasePointData* OutputData, int32 StartIndex, int32 Count)
+			{
+				return SetDensityLoop(OutputData, StartIndex, Count, [](float InNormalDensity, float& OutDensity)
+				{
+					OutDensity *= InNormalDensity;
+				});
 			});
 
 		case PCGNormalToDensityMode::Divide:
-			return ExecutePointOperation(Context, [CalcValue](const FPCGPoint& InPoint, FPCGPoint& OutPoint)->bool{
-				OutPoint = InPoint;
-
-				OutPoint.Density = UKismetMathLibrary::SafeDivide(OutPoint.Density, CalcValue(InPoint));
-
-				return true;
+			return ExecutePointOperation(Context, [SetDensityLoop](const UPCGBasePointData* InputData, UPCGBasePointData* OutputData, int32 StartIndex, int32 Count)
+			{
+				return SetDensityLoop(OutputData, StartIndex, Count, [](float InNormalDensity, float& OutDensity)
+				{
+					OutDensity = UKismetMathLibrary::SafeDivide(OutDensity, InNormalDensity);
+				});
 			});
 	}
 
diff --git a/Engine/Plugins/PCG/Source/PCG/Private/Elements/PCGPointExtentsModifier.cpp b/Engine/Plugins/PCG/Source/PCG/Private/Elements/PCGPointExtentsModifier.cpp
--- a/Engine/Plugins/PCG/Source/PCG/Private/Elements/PCGPointExtentsModifier.cpp
+++ b/Engine/Plugins/PCG/Source/PCG/Private/Elements/PCGPointExtentsModifier.cpp
@@ -4,6 +4,7 @@
 
 #include "PCGContext.h"
 #include "PCGPoint.h"
+#include "Data/PCGBasePointData.h"
 
 #include UE_INLINE_GENERATED_CPP_BY_NAME(PCGPointExtentsModifier)
 
@@ -12,6 +13,11 @@ FPCGElementPtr UPCGPointExtentsModifierSettings::CreateElement() const
 	return MakeShared<FPCGPointExtentsModifier>();
 }
 
+EPCGPointNativeProperties FPCGPointExtentsModifier::GetPropertiesToAllocate(FPCGContext* Context) const
+{
+	return EPCGPointNativeProperties::BoundsMin | EPCGPointNativeProperties::BoundsMax;
+}
+
 bool FPCGPointExtentsModifier::ExecuteInternal(FPCGContext* InContext) const
 {
 	TRACE_CPUPROFILER_EVENT_SCOPE(FPCGPointExtentsModifier::Execute);
@@ -23,46 +29,69 @@ bool FPCGPointExtentsModifier::ExecuteInternal(FPCGContext* InContext) const
 	const EPCGPointExtentsModifierMode Mode = Settings->Mode;
 	const FVector& Extents = Settings->Extents;
 
+	auto SetExtentsLoop = [](UPCGBasePointData* OutputData, int32 StartIndex, int32 Count, TFunctionRef<void(const FVector&, FVector&, FVector&)> SetExtentsFunc)
+	{
+		TPCGValueRange<FVector> BoundsMinRange = OutputData->GetBoundsMinValueRange();
+		TPCGValueRange<FVector> BoundsMaxRange = OutputData->GetBoundsMaxValueRange();
+
+		for (int32 Index = StartIndex; Index < (StartIndex + Count); ++Index)
+		{
+			FVector& BoundsMin = BoundsMinRange[Index];
+			FVector& BoundsMax = BoundsMaxRange[Index];
+
+			FVector CurrentExtents = PCGPointHelpers::GetExtents(BoundsMin, BoundsMax);
+
+			SetExtentsFunc(CurrentExtents, BoundsMin, BoundsMax);
+		}
+
+		return true;
+	};
+
 	switch (Mode)
 	{
 		case EPCGPointExtentsModifierMode::Minimum:
-			return ExecutePointOperation(Context, [Extents](const FPCGPoint& InPoint, FPCGPoint& OutPoint)
+			return ExecutePointOperation(Context, [SetExtentsLoop, Extents](const UPCGBasePointData* InputData, UPCGBasePointData* OutputData, int32 StartIndex, int32 Count)
 			{
-				OutPoint = InPoint;
-				OutPoint.SetExtents(FVector::Min(Extents, InPoint.GetExtents()));
-				return true;
+				return SetExtentsLoop(OutputData, StartIndex, Count, [Extents](const FVector& CurrentExtents, FVector& OutBoundsMin, FVector& OutBoundsMax)
+				{
+					PCGPointHelpers::SetExtents(FVector::Min(CurrentExtents, Extents), OutBoundsMin, OutBoundsMax);
+				});
 			});
 
 		case EPCGPointExtentsModifierMode::Maximum:
-			return ExecutePointOperation(Context, [Extents](const FPCGPoint& InPoint, FPCGPoint& OutPoint)
+			return ExecutePointOperation(Context, [SetExtentsLoop, Extents](const UPCGBasePointData* InputData, UPCGBasePointData* OutputData, int32 StartIndex, int32 Count)
 			{
-				OutPoint = InPoint;
-				OutPoint.SetExtents(FVector::Max(Extents, InPoint.GetExtents()));
-				return true;
+				return SetExtentsLoop(OutputData, StartIndex, Count, [Extents](const FVector& CurrentExtents, FVector& OutBoundsMin, FVector& OutBoundsMax)
+				{
+					PCGPointHelpers::SetExtents(FVector::Max(CurrentExtents, Extents), OutBoundsMin, OutBoundsMax);
+				});
 			});
 
-	case EPCGPointExtentsModifierMode::Add:
-			return ExecutePointOperation(Context, [Extents](const FPCGPoint& InPoint, FPCGPoint& OutPoint)
+		case EPCGPointExtentsModifierMode::Add:
+			return ExecutePointOperation(Context, [SetExtentsLoop, Extents](const UPCGBasePointData* InputData, UPCGBasePointData* OutputData, int32 StartIndex, int32 Count)
 			{
-				OutPoint = InPoint;
-				OutPoint.SetExtents(Extents + InPoint.GetExtents());
-				return true;
+				return SetExtentsLoop(OutputData, StartIndex, Count, [Extents](const FVector& CurrentExtents, FVector& OutBoundsMin, FVector& OutBoundsMax)
+				{
+					PCGPointHelpers::SetExtents(CurrentExtents + Extents, OutBoundsMin, OutBoundsMax);
+				});
 			});
 
 		case EPCGPointExtentsModifierMode::Multiply:
-			return ExecutePointOperation(Context, [Extents](const FPCGPoint& InPoint, FPCGPoint& OutPoint)
+			return ExecutePointOperation(Context, [SetExtentsLoop, Extents](const UPCGBasePointData* InputData, UPCGBasePointData* OutputData, int32 StartIndex, int32 Count)
 			{
-				OutPoint = InPoint;
-				OutPoint.SetExtents(Extents * InPoint.GetExtents());
-				return true;
+				return SetExtentsLoop(OutputData, StartIndex, Count, [Extents](const FVector& CurrentExtents, FVector& OutBoundsMin, FVector& OutBoundsMax)
+				{
+					PCGPointHelpers::SetExtents(CurrentExtents * Extents, OutBoundsMin, OutBoundsMax);
+				});
 			});
 
 		case EPCGPointExtentsModifierMode::Set:
-			return ExecutePointOperation(Context, [Extents](const FPCGPoint& InPoint, FPCGPoint& OutPoint)
+			return ExecutePointOperation(Context, [SetExtentsLoop, Extents](const UPCGBasePointData* InputData, UPCGBasePointData* OutputData, int32 StartIndex, int32 Count)
 			{
-				OutPoint = InPoint;
-				OutPoint.SetExtents(Extents);
-				return true;
+				return SetExtentsLoop(OutputData, StartIndex, Count, [Extents](const FVector& CurrentExtents, FVector& OutBoundsMin, FVector& OutBoundsMax)
+				{
+					PCGPointHelpers::SetExtents(Extents, OutBoundsMin, OutBoundsMax);
+				});
 			});
 	}
 
diff --git a/Engine/Plugins/PCG/Source/PCG/Private/Elements/PCGResetPointCenter.cpp b/Engine/Plugins/PCG/Source/PCG/Private/Elements/PCGResetPointCenter.cpp
--- a/Engine/Plugins/PCG/Source/PCG/Private/Elements/PCGResetPointCenter.cpp
+++ b/Engine/Plugins/PCG/Source/PCG/Private/Elements/PCGResetPointCenter.cpp
@@ -4,7 +4,7 @@
 
 #include "PCGContext.h"
 #include "PCGPin.h"
-#include "Data/PCGPointData.h"
+#include "Data/PCGBasePointData.h"
 
 #define LOCTEXT_NAMESPACE "PCGResetPointCenterElement"
 
@@ -13,6 +13,11 @@ FPCGElementPtr UPCGResetPointCenterSettings::CreateElement() const
 	return MakeShared<FPCGResetPointCenterElement>();
 }
 
+EPCGPointNativeProperties FPCGResetPointCenterElement::GetPropertiesToAllocate(FPCGContext* Context) const
+{
+	return EPCGPointNativeProperties::Transform | EPCGPointNativeProperties::BoundsMin | EPCGPointNativeProperties::BoundsMax;
+}
+
 bool FPCGResetPointCenterElement::ExecuteInternal(FPCGContext* Context) const
 {
 	TRACE_CPUPROFILER_EVENT_SCOPE(FPCGResetPointCenterElement::Execute);
@@ -23,11 +28,16 @@ bool FPCGResetPointCenterElement::ExecuteInternal(FPCGContext* Context) const
 	const UPCGResetPointCenterSettings* Settings = Context->GetInputSettings<UPCGResetPointCenterSettings>();
 	check(Settings);
 
-	return ExecutePointOperation(PointCenterContext, [&PointCenterLocation = Settings->PointCenterLocation](const FPCGPoint& InPoint, FPCGPoint& OutPoint)
+	return ExecutePointOperation(PointCenterContext, [&PointCenterLocation = Settings->PointCenterLocation](const UPCGBasePointData* InputData, UPCGBasePointData* OutputData, int32 StartIndex, int32 Count)
 	{
-		OutPoint = InPoint;
-
-		OutPoint.ResetPointCenter(PointCenterLocation);
+		TPCGValueRange<FTransform> TransformRange = OutputData->GetTransformValueRange();
+		TPCGValueRange<FVector> BoundsMinRange = OutputData->GetBoundsMinValueRange();
+		TPCGValueRange<FVector> BoundsMaxRange = OutputData->GetBoundsMaxValueRange();
+		
+		for (int32 Index = StartIndex; Index < (StartIndex + Count); ++Index)
+		{
+			PCGPointHelpers::ResetPointCenter(PointCenterLocation, TransformRange[Index], BoundsMinRange[Index], BoundsMaxRange[Index]);
+		}
 
 		return true;
 	});
diff --git a/Engine/Plugins/PCG/Source/PCG/Private/Tests/Elements/PCGApplyScaleToBoundsTest.cpp b/Engine/Plugins/PCG/Source/PCG/Private/Tests/Elements/PCGApplyScaleToBoundsTest.cpp
--- a/Engine/Plugins/PCG/Source/PCG/Private/Tests/Elements/PCGApplyScaleToBoundsTest.cpp
+++ b/Engine/Plugins/PCG/Source/PCG/Private/Tests/Elements/PCGApplyScaleToBoundsTest.cpp
@@ -38,18 +38,19 @@ bool FPCGApplyScaleToBoundsTest_Basic::RunTest(const FString& Parameters)
 
 	UTEST_EQUAL("Output count", Outputs.Num(), 1);
 
-	const UPCGPointData* OutPointData = Cast<UPCGPointData>(Outputs[0].Data);
+	const UPCGBasePointData* OutPointData = Cast<UPCGBasePointData>(Outputs[0].Data);
 
 	UTEST_NOT_NULL("Output point data", OutPointData);
+		
+	UTEST_EQUAL("Output point count", OutPointData->GetNumPoints(), 10);
 
-	const TArray<FPCGPoint>& OutPoints = OutPointData->GetPoints();
+	const TConstPCGValueRange<FVector> BoundsMinRange = OutPointData->GetConstBoundsMinValueRange();
+	const TConstPCGValueRange<FVector> BoundsMaxRange = OutPointData->GetConstBoundsMaxValueRange();
 
-	UTEST_EQUAL("Output point count", OutPoints.Num(), 10);
-
-	for (int i = 0; i < OutPoints.Num(); ++i)
+	for (int i = 0; i < OutPointData->GetNumPoints(); ++i)
 	{
-		UTEST_EQUAL("Output points should have a Bounds Min of -3", OutPoints[i].BoundsMin, FVector(-3));
-		UTEST_EQUAL("Output points should have a Bounds Max of 3", OutPoints[i].BoundsMax, FVector(3));
+		UTEST_EQUAL("Output points should have a Bounds Min of -3", BoundsMinRange[i], FVector(-3));
+		UTEST_EQUAL("Output points should have a Bounds Max of 3", BoundsMaxRange[i], FVector(3));
 	}
 
 	return true;
@@ -86,19 +87,21 @@ bool FPCGApplyScaleToBoundsTest_NegativeScale::RunTest(const FString& Parameters
 
 	UTEST_EQUAL("Output count", Outputs.Num(), 1);
 
-	const UPCGPointData* OutPointData = Cast<UPCGPointData>(Outputs[0].Data);
+	const UPCGBasePointData* OutPointData = Cast<UPCGBasePointData>(Outputs[0].Data);
 
 	UTEST_NOT_NULL("Output point data", OutPointData);
 
-	const TArray<FPCGPoint>& OutPoints = OutPointData->GetPoints();
+	UTEST_EQUAL("Output point count", OutPointData->GetNumPoints(), 10);
 
-	UTEST_EQUAL("Output point count", OutPoints.Num(), 10);
+	const TConstPCGValueRange<FTransform> TransformRange = OutPointData->GetConstTransformValueRange();
+	const TConstPCGValueRange<FVector> BoundsMinRange = OutPointData->GetConstBoundsMinValueRange();
+	const TConstPCGValueRange<FVector> BoundsMaxRange = OutPointData->GetConstBoundsMaxValueRange();
 
-	for (int i = 0; i < OutPoints.Num(); ++i)
+	for (int i = 0; i < OutPointData->GetNumPoints(); ++i)
 	{
-		UTEST_EQUAL("Output points should have a Bounds Min of -2", OutPoints[i].BoundsMin, FVector(-2.0));
-		UTEST_EQUAL("Output points should have a Bounds Min of 2", OutPoints[i].BoundsMax, FVector(2.0));
-		UTEST_EQUAL("Output points Bounds Max should be negative -1", OutPoints[i].Transform.GetScale3D(), FVector(-1.0));
+		UTEST_EQUAL("Output points should have a Bounds Min of -2", BoundsMinRange[i], FVector(-2.0));
+		UTEST_EQUAL("Output points should have a Bounds Max of 2", BoundsMaxRange[i], FVector(2.0));
+		UTEST_EQUAL("Output points Bounds Max should be negative -1", TransformRange[i].GetScale3D(), FVector(-1.0));
 	}
 
 	return true;
diff --git a/Engine/Plugins/PCG/Source/PCG/Private/Tests/Elements/PCGResetPointCenterTest.cpp b/Engine/Plugins/PCG/Source/PCG/Private/Tests/Elements/PCGResetPointCenterTest.cpp
--- a/Engine/Plugins/PCG/Source/PCG/Private/Tests/Elements/PCGResetPointCenterTest.cpp
+++ b/Engine/Plugins/PCG/Source/PCG/Private/Tests/Elements/PCGResetPointCenterTest.cpp
@@ -51,17 +51,18 @@ bool FPCGResetPointCenterTest_Center::RunTest(const FString& Parameters)
 
 	UTEST_EQUAL("Output count", Outputs.Num(), 1);
 
-	const UPCGPointData* OutPointData = Cast<UPCGPointData>(Outputs[0].Data);
+	const UPCGBasePointData* OutPointData = Cast<UPCGBasePointData>(Outputs[0].Data);
 
 	UTEST_NOT_NULL("Output point data", OutPointData);
+	UTEST_EQUAL("Output point count", OutPointData->GetNumPoints(), 1);
 
-	const TArray<FPCGPoint>& OutPoints = OutPointData->GetPoints();
+	const TConstPCGValueRange<FTransform> TransformRange = OutPointData->GetConstTransformValueRange();
+	const TConstPCGValueRange<FVector> BoundsMinRange = OutPointData->GetConstBoundsMinValueRange();
+	const TConstPCGValueRange<FVector> BoundsMaxRange = OutPointData->GetConstBoundsMaxValueRange();
 
-	UTEST_EQUAL("Output point count", OutPoints.Num(), 1);
-
-	UTEST_EQUAL("Output has the correct transform,", OutPoints[0].Transform, FTransform(FRotator::ZeroRotator, FVector(0.0), FVector::One()));
-	UTEST_EQUAL("Output has the correct minimum bounds,", OutPoints[0].BoundsMin, FVector(-1.0));
-	UTEST_EQUAL("Output has the correct maximum bounds,", OutPoints[0].BoundsMax, FVector(1.0));
+	UTEST_EQUAL("Output has the correct transform,", TransformRange[0], FTransform(FRotator::ZeroRotator, FVector(0.0), FVector::One()));
+	UTEST_EQUAL("Output has the correct minimum bounds,", BoundsMinRange[0], FVector(-1.0));
+	UTEST_EQUAL("Output has the correct maximum bounds,", BoundsMaxRange[0], FVector(1.0));
 
 	return true;
 }
@@ -77,17 +78,18 @@ bool FPCGResetPointCenterTest_Zero::RunTest(const FString& Parameters)
 
 	UTEST_EQUAL("Output count", Outputs.Num(), 1);
 
-	const UPCGPointData* OutPointData = Cast<UPCGPointData>(Outputs[0].Data);
+	const UPCGBasePointData* OutPointData = Cast<UPCGBasePointData>(Outputs[0].Data);
 
 	UTEST_NOT_NULL("Output point data", OutPointData);
+	UTEST_EQUAL("Output point count", OutPointData->GetNumPoints(), 1);
 
-	const TArray<FPCGPoint>& OutPoints = OutPointData->GetPoints();
-
-	UTEST_EQUAL("Output point count", OutPoints.Num(), 1);
+	const TConstPCGValueRange<FTransform> TransformRange = OutPointData->GetConstTransformValueRange();
+	const TConstPCGValueRange<FVector> BoundsMinRange = OutPointData->GetConstBoundsMinValueRange();
+	const TConstPCGValueRange<FVector> BoundsMaxRange = OutPointData->GetConstBoundsMaxValueRange();
 
-	UTEST_EQUAL("Output has the correct transform,", OutPoints[0].Transform, FTransform(FRotator::ZeroRotator, FVector(-1.0), FVector::One()));
-	UTEST_EQUAL("Output has the correct minimum bounds,", OutPoints[0].BoundsMin, FVector(0.0));
-	UTEST_EQUAL("Output has the correct maximum bounds,", OutPoints[0].BoundsMax, FVector(2.0));
+	UTEST_EQUAL("Output has the correct transform,", TransformRange[0], FTransform(FRotator::ZeroRotator, FVector(-1.0), FVector::One()));
+	UTEST_EQUAL("Output has the correct minimum bounds,", BoundsMinRange[0], FVector(0.0));
+	UTEST_EQUAL("Output has the correct maximum bounds,", BoundsMaxRange[0], FVector(2.0));
 
 	return true;
 }
@@ -101,19 +103,18 @@ bool FPCGResetPointCenterTest_One::RunTest(const FString& Parameters)
 	// Point center should be in the top right corner
 	const TArray<FPCGTaggedData>& Outputs = Context->OutputData.TaggedData;
 
-	UTEST_EQUAL("Output count", Outputs.Num(), 1);
-
-	const UPCGPointData* OutPointData = Cast<UPCGPointData>(Outputs[0].Data);
+	const UPCGBasePointData* OutPointData = Cast<UPCGBasePointData>(Outputs[0].Data);
 
 	UTEST_NOT_NULL("Output point data", OutPointData);
+	UTEST_EQUAL("Output point count", OutPointData->GetNumPoints(), 1);
 
-	const TArray<FPCGPoint>& OutPoints = OutPointData->GetPoints();
-
-	UTEST_EQUAL("Output point count", OutPoints.Num(), 1);
+	const TConstPCGValueRange<FTransform> TransformRange = OutPointData->GetConstTransformValueRange();
+	const TConstPCGValueRange<FVector> BoundsMinRange = OutPointData->GetConstBoundsMinValueRange();
+	const TConstPCGValueRange<FVector> BoundsMaxRange = OutPointData->GetConstBoundsMaxValueRange();
 
-	UTEST_EQUAL("Output has the correct transform,", OutPoints[0].Transform, FTransform(FRotator::ZeroRotator, FVector(1.0), FVector::One()));
-	UTEST_EQUAL("Output has the correct minimum bounds,", OutPoints[0].BoundsMin, FVector(-2.0));
-	UTEST_EQUAL("Output has the correct maximum bounds,", OutPoints[0].BoundsMax, FVector(0.0));
+	UTEST_EQUAL("Output has the correct transform,", TransformRange[0], FTransform(FRotator::ZeroRotator, FVector(1.0), FVector::One()));
+	UTEST_EQUAL("Output has the correct minimum bounds,", BoundsMinRange[0], FVector(-2.0));
+	UTEST_EQUAL("Output has the correct maximum bounds,", BoundsMaxRange[0], FVector(0.0));
 
 	return true;
 }
@@ -129,17 +130,18 @@ bool FPCGResetPointCenterTest_NegativeOutside::RunTest(const FString& Parameters
 
 	UTEST_EQUAL("Output count", Outputs.Num(), 1);
 
-	const UPCGPointData* OutPointData = Cast<UPCGPointData>(Outputs[0].Data);
+	const UPCGBasePointData* OutPointData = Cast<UPCGBasePointData>(Outputs[0].Data);
 
 	UTEST_NOT_NULL("Output point data", OutPointData);
+	UTEST_EQUAL("Output point count", OutPointData->GetNumPoints(), 1);
 
-	const TArray<FPCGPoint>& OutPoints = OutPointData->GetPoints();
+	const TConstPCGValueRange<FTransform> TransformRange = OutPointData->GetConstTransformValueRange();
+	const TConstPCGValueRange<FVector> BoundsMinRange = OutPointData->GetConstBoundsMinValueRange();
+	const TConstPCGValueRange<FVector> BoundsMaxRange = OutPointData->GetConstBoundsMaxValueRange();
 
-	UTEST_EQUAL("Output point count", OutPoints.Num(), 1);
-
-	UTEST_EQUAL("Output has the correct transform,", OutPoints[0].Transform, FTransform(FRotator::ZeroRotator, FVector(-7.0), FVector::One()));
-	UTEST_EQUAL("Output has the correct minimum bounds,", OutPoints[0].BoundsMin, FVector(6.0));
-	UTEST_EQUAL("Output has the correct maximum bounds,", OutPoints[0].BoundsMax, FVector(8.0));
+	UTEST_EQUAL("Output has the correct transform,", TransformRange[0], FTransform(FRotator::ZeroRotator, FVector(-7.0), FVector::One()));
+	UTEST_EQUAL("Output has the correct minimum bounds,", BoundsMinRange[0], FVector(6.0));
+	UTEST_EQUAL("Output has the correct maximum bounds,", BoundsMaxRange[0], FVector(8.0));
 
 	return true;
 }
@@ -155,17 +157,18 @@ bool FPCGResetPointCenterTest_PositiveOutside::RunTest(const FString& Parameters
 
 	UTEST_EQUAL("Output count", Outputs.Num(), 1);
 
-	const UPCGPointData* OutPointData = Cast<UPCGPointData>(Outputs[0].Data);
+	const UPCGBasePointData* OutPointData = Cast<UPCGBasePointData>(Outputs[0].Data);
 
 	UTEST_NOT_NULL("Output point data", OutPointData);
+	UTEST_EQUAL("Output point count", OutPointData->GetNumPoints(), 1);
 
-	const TArray<FPCGPoint>& OutPoints = OutPointData->GetPoints();
-
-	UTEST_EQUAL("Output point count", OutPoints.Num(), 1);
+	const TConstPCGValueRange<FTransform> TransformRange = OutPointData->GetConstTransformValueRange();
+	const TConstPCGValueRange<FVector> BoundsMinRange = OutPointData->GetConstBoundsMinValueRange();
+	const TConstPCGValueRange<FVector> BoundsMaxRange = OutPointData->GetConstBoundsMaxValueRange();
 
-	UTEST_EQUAL("Output has the correct transform,", OutPoints[0].Transform, FTransform(FRotator::ZeroRotator, FVector(7.0), FVector::One()));
-	UTEST_EQUAL("Output has the correct minimum bounds,", OutPoints[0].BoundsMin, FVector(-8.0));
-	UTEST_EQUAL("Output has the correct maximum bounds,", OutPoints[0].BoundsMax, FVector(-6.0));
+	UTEST_EQUAL("Output has the correct transform,", TransformRange[0], FTransform(FRotator::ZeroRotator, FVector(7.0), FVector::One()));
+	UTEST_EQUAL("Output has the correct minimum bounds,", BoundsMinRange[0], FVector(-8.0));
+	UTEST_EQUAL("Output has the correct maximum bounds,", BoundsMaxRange[0], FVector(-6.0));
 
 	return true;
 }
@@ -180,17 +183,18 @@ bool FPCGResetPointCenterTest_NonUniform::RunTest(const FString& Parameters)
 
 	UTEST_EQUAL("Output count", Outputs.Num(), 1);
 
-	const UPCGPointData* OutPointData = Cast<UPCGPointData>(Outputs[0].Data);
+	const UPCGBasePointData* OutPointData = Cast<UPCGBasePointData>(Outputs[0].Data);
 
 	UTEST_NOT_NULL("Output point data", OutPointData);
+	UTEST_EQUAL("Output point count", OutPointData->GetNumPoints(), 1);
 
-	const TArray<FPCGPoint>& OutPoints = OutPointData->GetPoints();
-
-	UTEST_EQUAL("Output point count", OutPoints.Num(), 1);
+	const TConstPCGValueRange<FTransform> TransformRange = OutPointData->GetConstTransformValueRange();
+	const TConstPCGValueRange<FVector> BoundsMinRange = OutPointData->GetConstBoundsMinValueRange();
+	const TConstPCGValueRange<FVector> BoundsMaxRange = OutPointData->GetConstBoundsMaxValueRange();
 
-	UTEST_EQUAL("Output has the correct transform,", OutPoints[0].Transform, FTransform(FRotator::ZeroRotator, FVector(-3.0, 0.0, 3.0), FVector::One()));
-	UTEST_EQUAL("Output has the correct minimum bounds,", OutPoints[0].BoundsMin, FVector(2.0, -1.0, -4.0));
-	UTEST_EQUAL("Output has the correct maximum bounds,", OutPoints[0].BoundsMax, FVector(4.0, 1.0, -2.0));
+	UTEST_EQUAL("Output has the correct transform,", TransformRange[0], FTransform(FRotator::ZeroRotator, FVector(-3.0, 0.0, 3.0), FVector::One()));
+	UTEST_EQUAL("Output has the correct minimum bounds,", BoundsMinRange[0], FVector(2.0, -1.0, -4.0));
+	UTEST_EQUAL("Output has the correct maximum bounds,", BoundsMaxRange[0], FVector(4.0, 1.0, -2.0));
 
 	return true;
 }
@@ -206,17 +210,18 @@ bool FPCGResetPointCenterTest_CenterRotation::RunTest(const FString& Parameters)
 
 	UTEST_EQUAL("Output count", Outputs.Num(), 1);
 
-	const UPCGPointData* OutPointData = Cast<UPCGPointData>(Outputs[0].Data);
+	const UPCGBasePointData* OutPointData = Cast<UPCGBasePointData>(Outputs[0].Data);
 
 	UTEST_NOT_NULL("Output point data", OutPointData);
+	UTEST_EQUAL("Output point count", OutPointData->GetNumPoints(), 1);
 
-	const TArray<FPCGPoint>& OutPoints = OutPointData->GetPoints();
+	const TConstPCGValueRange<FTransform> TransformRange = OutPointData->GetConstTransformValueRange();
+	const TConstPCGValueRange<FVector> BoundsMinRange = OutPointData->GetConstBoundsMinValueRange();
+	const TConstPCGValueRange<FVector> BoundsMaxRange = OutPointData->GetConstBoundsMaxValueRange();
 
-	UTEST_EQUAL("Output point count", OutPoints.Num(), 1);
-
-	UTEST_EQUAL("Output has the correct transform,", OutPoints[0].Transform, FTransform(FRotator(20.0), FVector(0.0), FVector::One()));
-	UTEST_EQUAL("Output has the correct minimum bounds,", OutPoints[0].BoundsMin, FVector(-1.0));
-	UTEST_EQUAL("Output has the correct maximum bounds,", OutPoints[0].BoundsMax, FVector(1.0));
+	UTEST_EQUAL("Output has the correct transform,", TransformRange[0], FTransform(FRotator(20.0), FVector(0.0), FVector::One()));
+	UTEST_EQUAL("Output has the correct minimum bounds,", BoundsMinRange[0], FVector(-1.0));
+	UTEST_EQUAL("Output has the correct maximum bounds,", BoundsMaxRange[0], FVector(1.0));
 
 	return true;
 }
@@ -232,17 +237,18 @@ bool FPCGResetPointCenterTest_CenterScale::RunTest(const FString& Parameters)
 
 	UTEST_EQUAL("Output count", Outputs.Num(), 1);
 
-	const UPCGPointData* OutPointData = Cast<UPCGPointData>(Outputs[0].Data);
+	const UPCGBasePointData* OutPointData = Cast<UPCGBasePointData>(Outputs[0].Data);
 
 	UTEST_NOT_NULL("Output point data", OutPointData);
+	UTEST_EQUAL("Output point count", OutPointData->GetNumPoints(), 1);
 
-	const TArray<FPCGPoint>& OutPoints = OutPointData->GetPoints();
-
-	UTEST_EQUAL("Output point count", OutPoints.Num(), 1);
+	const TConstPCGValueRange<FTransform> TransformRange = OutPointData->GetConstTransformValueRange();
+	const TConstPCGValueRange<FVector> BoundsMinRange = OutPointData->GetConstBoundsMinValueRange();
+	const TConstPCGValueRange<FVector> BoundsMaxRange = OutPointData->GetConstBoundsMaxValueRange();
 
-	UTEST_EQUAL("Output has the correct transform,", OutPoints[0].Transform, FTransform(FRotator::ZeroRotator, FVector(0.0), FVector(3.0)));
-	UTEST_EQUAL("Output has the correct minimum bounds,", OutPoints[0].BoundsMin, FVector(-1.0));
-	UTEST_EQUAL("Output has the correct maximum bounds,", OutPoints[0].BoundsMax, FVector(1.0));
+	UTEST_EQUAL("Output has the correct transform,", TransformRange[0], FTransform(FRotator::ZeroRotator, FVector(0.0), FVector(3.0)));
+	UTEST_EQUAL("Output has the correct minimum bounds,", BoundsMinRange[0], FVector(-1.0));
+	UTEST_EQUAL("Output has the correct maximum bounds,", BoundsMaxRange[0], FVector(1.0));
 
 	return true;
 }
\ No newline at end of file
diff --git a/Engine/Plugins/PCG/Source/PCG/Public/Elements/PCGApplyScaleToBounds.h b/Engine/Plugins/PCG/Source/PCG/Public/Elements/PCGApplyScaleToBounds.h
--- a/Engine/Plugins/PCG/Source/PCG/Public/Elements/PCGApplyScaleToBounds.h
+++ b/Engine/Plugins/PCG/Source/PCG/Public/Elements/PCGApplyScaleToBounds.h
@@ -35,4 +35,7 @@ class FPCGApplyScaleToBoundsElement : public FPCGPointOperationElementBase
 {
 protected:
 	virtual bool ExecuteInternal(FPCGContext* Context) const override;
+	virtual EPCGPointNativeProperties GetPropertiesToAllocate(FPCGContext* InContext) const override;
+	virtual bool ShouldCopyPoints() const override { return true; }
+	virtual bool SupportsBasePointDataInputs(FPCGContext* InContext) const override { return true; }
 };
\ No newline at end of file
diff --git a/Engine/Plugins/PCG/Source/PCG/Public/Elements/PCGDensityRemapElement.h b/Engine/Plugins/PCG/Source/PCG/Public/Elements/PCGDensityRemapElement.h
--- a/Engine/Plugins/PCG/Source/PCG/Public/Elements/PCGDensityRemapElement.h
+++ b/Engine/Plugins/PCG/Source/PCG/Public/Elements/PCGDensityRemapElement.h
@@ -51,5 +51,8 @@ public:
 class FPCGDensityRemapElement : public FPCGPointOperationElementBase
 {
 protected:
-	virtual bool ExecuteInternal(FPCGContext* InContext) const override;
+	virtual bool ExecuteInternal(FPCGContext* Context) const override;
+	virtual EPCGPointNativeProperties GetPropertiesToAllocate(FPCGContext* InContext) const override;
+	virtual bool ShouldCopyPoints() const override { return true; }
+	virtual bool SupportsBasePointDataInputs(FPCGContext* InContext) const override { return true; }
 };
diff --git a/Engine/Plugins/PCG/Source/PCG/Public/Elements/PCGMutateSeed.h b/Engine/Plugins/PCG/Source/PCG/Public/Elements/PCGMutateSeed.h
--- a/Engine/Plugins/PCG/Source/PCG/Public/Elements/PCGMutateSeed.h
+++ b/Engine/Plugins/PCG/Source/PCG/Public/Elements/PCGMutateSeed.h
@@ -36,4 +36,7 @@ class FPCGMutateSeedElement : public FPCGPointOperationElementBase
 {
 protected:
 	virtual bool ExecuteInternal(FPCGContext* Context) const override;
+	virtual EPCGPointNativeProperties GetPropertiesToAllocate(FPCGContext* InContext) const override;
+	virtual bool ShouldCopyPoints() const override { return true; }
+	virtual bool SupportsBasePointDataInputs(FPCGContext* InContext) const override { return true; }
 };
diff --git a/Engine/Plugins/PCG/Source/PCG/Public/Elements/PCGNormalToDensity.h b/Engine/Plugins/PCG/Source/PCG/Public/Elements/PCGNormalToDensity.h
--- a/Engine/Plugins/PCG/Source/PCG/Public/Elements/PCGNormalToDensity.h
+++ b/Engine/Plugins/PCG/Source/PCG/Public/Elements/PCGNormalToDensity.h
@@ -63,6 +63,9 @@ public:
 class FPCGNormalToDensityElement : public FPCGPointOperationElementBase
 {
 protected:
-	virtual bool ExecuteInternal(FPCGContext* InContext) const override;
+	virtual bool ExecuteInternal(FPCGContext* Context) const override;
+	virtual EPCGPointNativeProperties GetPropertiesToAllocate(FPCGContext* InContext) const override;
+	virtual bool ShouldCopyPoints() const override { return true; }
+	virtual bool SupportsBasePointDataInputs(FPCGContext* InContext) const override { return true; }
 };
 
diff --git a/Engine/Plugins/PCG/Source/PCG/Public/Elements/PCGPointExtentsModifier.h b/Engine/Plugins/PCG/Source/PCG/Public/Elements/PCGPointExtentsModifier.h
--- a/Engine/Plugins/PCG/Source/PCG/Public/Elements/PCGPointExtentsModifier.h
+++ b/Engine/Plugins/PCG/Source/PCG/Public/Elements/PCGPointExtentsModifier.h
@@ -49,5 +49,8 @@ public:
 class FPCGPointExtentsModifier : public FPCGPointOperationElementBase
 {
 protected:
-	virtual bool ExecuteInternal(FPCGContext* InContext) const override;
+	virtual bool ExecuteInternal(FPCGContext* Context) const override;
+	virtual EPCGPointNativeProperties GetPropertiesToAllocate(FPCGContext* InContext) const override;
+	virtual bool ShouldCopyPoints() const override { return true; }
+	virtual bool SupportsBasePointDataInputs(FPCGContext* InContext) const override { return true; }
 };
\ No newline at end of file
diff --git a/Engine/Plugins/PCG/Source/PCG/Public/Elements/PCGResetPointCenter.h b/Engine/Plugins/PCG/Source/PCG/Public/Elements/PCGResetPointCenter.h
--- a/Engine/Plugins/PCG/Source/PCG/Public/Elements/PCGResetPointCenter.h
+++ b/Engine/Plugins/PCG/Source/PCG/Public/Elements/PCGResetPointCenter.h
@@ -41,4 +41,7 @@ class FPCGResetPointCenterElement : public FPCGPointOperationElementBase
 {
 protected:
 	virtual bool ExecuteInternal(FPCGContext* Context) const override;
+	virtual EPCGPointNativeProperties GetPropertiesToAllocate(FPCGContext* InContext) const override;
+	virtual bool ShouldCopyPoints() const override { return true; }
+	virtual bool SupportsBasePointDataInputs(FPCGContext* InContext) const override { return true; }
 };
\ No newline at end of file
