[PCG] Fix up texture init attempt 2 - wait for full texture init only for GPU readback path, which could remove the blocking wait there.

#rnx

diff --git a/Engine/Plugins/PCG/Source/PCG/Private/Data/PCGTextureData.cpp b/Engine/Plugins/PCG/Source/PCG/Private/Data/PCGTextureData.cpp
--- a/Engine/Plugins/PCG/Source/PCG/Private/Data/PCGTextureData.cpp
+++ b/Engine/Plugins/PCG/Source/PCG/Private/Data/PCGTextureData.cpp
@@ -721,7 +721,6 @@ bool UPCGTextureData::InitializeFromGPUTexture()
 	{
 		return false;
 	}
-	bReadbackFromGPUInitiated = true;
 
 	if (ResourceType == EPCGTextureResourceType::TextureObject)
 	{
@@ -729,15 +728,18 @@ bool UPCGTextureData::InitializeFromGPUTexture()
 		{
 			return true;
 		}
-#if WITH_EDITOR
-		// Force a wait on any incomplete async texture compilation and caching operations
-		Texture->FinishCachePlatformData();
-#endif
 
-		// TODO: Forcing a resource update is problematic. Clears cooked platform data and may break other usages of this texture
-		// But it is necessary to force the texture memory to be present on the first dispatch
-		Texture->UpdateResource();
-		Texture->WaitForPendingInitOrStreaming();
+		if (!bUpdatedReadbackTextureResource)
+		{
+			Texture->UpdateResource();
+
+			bUpdatedReadbackTextureResource = true;
+		}
+
+		if (Texture->HasPendingInitOrStreaming(/*bWaitForLODTransition=*/true))
+		{
+			return false;
+		}
 	}
 
 	if (FTextureRHIRef TextureRHI = GetTextureRHI())
@@ -814,6 +816,8 @@ bool UPCGTextureData::InitializeFromGPUTexture()
 				This->PostInitializeCallback();
 			}
 		});
+
+		bReadbackFromGPUInitiated = true;
 	}
 	else
 	{
diff --git a/Engine/Plugins/PCG/Source/PCG/Public/Data/PCGTextureData.h b/Engine/Plugins/PCG/Source/PCG/Public/Data/PCGTextureData.h
--- a/Engine/Plugins/PCG/Source/PCG/Public/Data/PCGTextureData.h
+++ b/Engine/Plugins/PCG/Source/PCG/Public/Data/PCGTextureData.h
@@ -227,10 +227,9 @@ public:
 	UPROPERTY(BlueprintReadOnly, VisibleAnywhere, Category = Properties)
 	int TextureIndex = 0;
 
-	UPROPERTY(BlueprintReadOnly, VisibleAnywhere, Category = Properties)
+	UPROPERTY(Transient, BlueprintReadOnly, VisibleAnywhere, Category = Properties)
 	bool bSuccessfullyInitialized = false;
 
-	UPROPERTY(BlueprintReadOnly, VisibleAnywhere, Category = Properties)
 	bool bReadbackFromGPUInitiated = false;
 
 	// Added to help deprecation in 5.5. To be removed when the deprecated Initialized function is removed.
@@ -244,7 +243,13 @@ protected:
 	/** If initialized from an exported texture this holds a reference to the resource. */
 	TRefCountPtr<IPooledRenderTarget> TextureHandle = nullptr;
 
+	bool bUpdatedReadbackTextureResource = false;
+
 public:
 	UE_DEPRECATED(5.5, "Will be removed. Poll the alternate Initialize API until it returns true instead of passing in a callback.")
 	PCG_API void Initialize(UTexture* InTexture, uint32 InTextureIndex, const FTransform& InTransform, const TFunction<void()>& InPostInitializeCallback, bool bCreateCPUDuplicateEditorOnly = false);
+
+	UE_DEPRECATED(5.6, "Internal object state removed from blueprint.")
+	UPROPERTY()
+	bool bReadbackFromGPUInitiated_DEPRECATED = false;
 };
