[PCG] Fix localization keys

#rnx

diff --git a/Engine/Plugins/PCG/Source/PCG/Private/Elements/PCGStaticMeshSpawnerKernel.cpp b/Engine/Plugins/PCG/Source/PCG/Private/Elements/PCGStaticMeshSpawnerKernel.cpp
--- a/Engine/Plugins/PCG/Source/PCG/Private/Elements/PCGStaticMeshSpawnerKernel.cpp
+++ b/Engine/Plugins/PCG/Source/PCG/Private/Elements/PCGStaticMeshSpawnerKernel.cpp
@@ -27,6 +27,7 @@ namespace PCGStaticMeshSpawnerKernel
 {
 	static const FText CouldNotLoadStaticMeshFormat = LOCTEXT("CouldNotLoadStaticMesh", "Could not load static mesh from path '{0}'.");
 	static const FText TooManyPrimitivesFormat = LOCTEXT("TooManyPrimitives", "Attempted to emit too many primitive components, terminated after creating '{0}'.");
+	static const FText NoMeshEntriesFormat = LOCTEXT("NoMeshEntries", "No mesh entries provided in weighted mesh selector.");
 }
 
 bool UPCGStaticMeshSpawnerKernel::IsKernelValid(FPCGContext* InContext, bool bQuiet) const
@@ -47,7 +48,7 @@ bool UPCGStaticMeshSpawnerKernel::IsKernelValid(FPCGContext* InContext, bool bQu
 	{
 		if (SelectorWeighted->MeshEntries.IsEmpty())
 		{
-			PCG_KERNEL_VALIDATION_ERR(InContext, Settings, bQuiet, LOCTEXT("NoMeshEntries", "No meshes specified in selector."));
+			PCG_KERNEL_VALIDATION_ERR(InContext, Settings, bQuiet, PCGStaticMeshSpawnerKernel::NoMeshEntriesFormat);
 			return false;
 		}
 
@@ -358,7 +359,7 @@ bool UPCGStaticMeshSpawnerKernel::SetupProceduralISMComponents(FPCGContext* InCo
 	{
 		if (SelectorWeighted->MeshEntries.IsEmpty())
 		{
-			PCG_KERNEL_VALIDATION_ERR(InContext, Settings, /*bQuiet=*/false, LOCTEXT("NoMeshEntries", "No mesh entries provided."));
+			PCG_KERNEL_VALIDATION_ERR(InContext, Settings, /*bQuiet=*/false, PCGStaticMeshSpawnerKernel::NoMeshEntriesFormat);
 
 			return false;
 		}
