[PCG] Fix Undo on graph parameters in PCG Graph

* Undo/Redo graph parameter changes in a PCG graph was not updating the instances on this graph.
* Now on Undo/Redo, we trigger the same change than on PostLoad for the instances to compare their struct with the original struct of the graph.

#rnx

diff --git a/Engine/Plugins/PCG/Source/PCG/Private/PCGGraph.cpp b/Engine/Plugins/PCG/Source/PCG/Private/PCGGraph.cpp
--- a/Engine/Plugins/PCG/Source/PCG/Private/PCGGraph.cpp
+++ b/Engine/Plugins/PCG/Source/PCG/Private/PCGGraph.cpp
@@ -1723,6 +1723,15 @@ void UPCGGraph::PostEditChangeChainProperty(FPropertyChangedChainEvent& Property
 	PreviousPropertyBag = nullptr;
 }
 
+void UPCGGraph::PostEditUndo()
+{
+	// If we have parameters, they might have changed with this Undo/Redo. So trigger the same mechanism as PostLoad or add multiple properties to update the graph instances that would depend on it.
+	if (UserParameters.GetNumPropertiesInBag() > 0)
+	{
+		OnGraphParametersChanged(EPCGGraphParameterEvent::UndoRedo, NAME_None);
+	}
+}
+
 void UPCGGraph::FixInvalidEdges()
 {
 	auto ValidatePins = [this](const TArray<TObjectPtr<UPCGPin>>& Pins, bool bPinsAreInputs)
@@ -2441,10 +2450,11 @@ bool FPCGOverrideInstancedPropertyBag::RefreshParameters(const FInstancedPropert
 		}
 		break;
 	}
-	// Do the same thing in case of post load and multiple properties added.
-	// We have 2 enums to avoid puzzling someone that wonders why we would call GraphPostLoad when we add multiple properties.
+	// Do the same thing in case of post load, multiple properties added and undo/redo.
+	// We have 3 enums to avoid puzzling someone that wonders why we would call GraphPostLoad when we add multiple properties or undo/redo.
 	case EPCGGraphParameterEvent::GraphPostLoad: // fall-through
-	case EPCGGraphParameterEvent::MultiplePropertiesAdded:
+	case EPCGGraphParameterEvent::MultiplePropertiesAdded: // fall-through
+	case EPCGGraphParameterEvent::UndoRedo:
 	{
 		// Check if the property struct mismatch. If so, do the migration
 		if (Parameters.GetPropertyBagStruct() != ParentUserParameters->GetPropertyBagStruct())
diff --git a/Engine/Plugins/PCG/Source/PCG/Public/PCGGraph.h b/Engine/Plugins/PCG/Source/PCG/Public/PCGGraph.h
--- a/Engine/Plugins/PCG/Source/PCG/Public/PCGGraph.h
+++ b/Engine/Plugins/PCG/Source/PCG/Public/PCGGraph.h
@@ -36,6 +36,7 @@ enum class EPCGGraphParameterEvent
 	ValueModifiedLocally,
 	ValueModifiedByParent,
 	MultiplePropertiesAdded,
+	UndoRedo,
 	None
 };
 
@@ -262,6 +263,7 @@ public:
 #if WITH_EDITOR
 	virtual void PreEditChange(FProperty* InProperty) override;
 	virtual void PostEditChangeChainProperty(FPropertyChangedChainEvent& PropertyChangedEvent);
+	virtual void PostEditUndo() override;
 #endif
 	/** ~End UObject interface */
 
