Update jump to in find window to preserve debug object selection

diff --git a/Engine/Plugins/PCG/Source/PCGEditor/Private/SPCGEditorGraphDebugObjectTree.cpp b/Engine/Plugins/PCG/Source/PCGEditor/Private/SPCGEditorGraphDebugObjectTree.cpp
--- a/Engine/Plugins/PCG/Source/PCGEditor/Private/SPCGEditorGraphDebugObjectTree.cpp
+++ b/Engine/Plugins/PCG/Source/PCGEditor/Private/SPCGEditorGraphDebugObjectTree.cpp
@@ -555,7 +555,27 @@ void SPCGEditorGraphDebugObjectTree::SetDebugObjectFromSelection_OnClicked()
 bool SPCGEditorGraphDebugObjectTree::SetDebugObjectFromStackFromAnotherEditor(const FPCGStack& InStack)
 {
 	RefreshTree();
-	return SetDebugObjectFromStack(InStack);
+
+	const UPCGGraph* CurrentGraph = PCGEditor.Pin()->GetPCGGraph();
+	if (InStack.GetGraphForCurrentFrame() == CurrentGraph)
+	{
+		return SetDebugObjectFromStack(InStack);
+	}
+	else
+	{
+		// Find first instance that starts with the given stack but has the current graph at its end.
+		for (const FPCGEditorGraphDebugObjectItemPtr& Item : AllGraphItems)
+		{
+			const FPCGStack* ItemStack = Item->GetPCGStack();
+
+			if (ItemStack && ItemStack->BeginsWith(InStack) && ItemStack->GetGraphForCurrentFrame() == CurrentGraph)
+			{
+				return SetDebugObjectFromStack(*ItemStack);
+			}
+		}
+
+		return false;
+	}
 }
 
 bool SPCGEditorGraphDebugObjectTree::SetDebugObjectFromStack(const FPCGStack& InStack)
diff --git a/Engine/Plugins/PCG/Source/PCGEditor/Private/SPCGEditorGraphFind.cpp b/Engine/Plugins/PCG/Source/PCGEditor/Private/SPCGEditorGraphFind.cpp
--- a/Engine/Plugins/PCG/Source/PCGEditor/Private/SPCGEditorGraphFind.cpp
+++ b/Engine/Plugins/PCG/Source/PCGEditor/Private/SPCGEditorGraphFind.cpp
@@ -77,7 +77,8 @@ FReply FPCGEditorGraphFindResult::OnClick(TWeakPtr<FPCGEditor> InPCGEditorPtr)
 
 FReply FPCGEditorGraphFindResult::OnDoubleClick(TWeakPtr<FPCGEditor> InPCGEditorPtr)
 {
-	TSharedPtr<FPCGEditor> Editor = InPCGEditorPtr.Pin();
+	TSharedPtr<FPCGEditor> CurrentEditor = InPCGEditorPtr.Pin();
+	TSharedPtr<FPCGEditor> Editor = CurrentEditor;
 
 	if (ParentGraph)
 	{
@@ -87,6 +88,11 @@ FReply FPCGEditorGraphFindResult::OnDoubleClick(TWeakPtr<FPCGEditor> InPCGEditor
 
 	if (Editor)
 	{
+		if (Editor != CurrentEditor && CurrentEditor->GetStackBeingInspected())
+		{
+			Editor->SetStackBeingInspectedFromAnotherEditor(*CurrentEditor->GetStackBeingInspected());
+		}
+
 		if (UEdGraphPin* ResolvedPin = Pin.Get())
 		{
 			Editor->JumpToNode(ResolvedPin->GetOwningNode());
