[PCG] Fix nullptr check in Metadata operation if the accessor failed to be constructed.

#rnx

diff --git a/Engine/Plugins/PCG/Source/PCG/Private/Elements/Metadata/PCGMetadataOpElementBase.cpp b/Engine/Plugins/PCG/Source/PCG/Private/Elements/Metadata/PCGMetadataOpElementBase.cpp
--- a/Engine/Plugins/PCG/Source/PCG/Private/Elements/Metadata/PCGMetadataOpElementBase.cpp
+++ b/Engine/Plugins/PCG/Source/PCG/Private/Elements/Metadata/PCGMetadataOpElementBase.cpp
@@ -622,11 +622,15 @@ bool FPCGMetadataElementBase::PrepareDataInternal(FPCGContext* Context) const
 			{
 				const TUniquePtr<const IPCGAttributeAccessor> TempConstAccessor = PCGAttributeAccessorHelpers::CreateConstAccessor(OutputTaggedData.Data.Get(), OutputTarget);
 
-				const bool bIsValid = TempConstAccessor.IsValid() && PCG::Private::IsBroadcastable(PCG::Private::MetadataTypes<AttributeType>::Id, TempConstAccessor->GetUnderlyingType());
+				if (!TempConstAccessor.IsValid())
+				{
+					PCGLog::Metadata::LogFailToCreateAccessorError(OutputTarget, Context);
+					return false;
+				}
 
-				if (!bIsValid)
+				if (!PCG::Private::IsBroadcastable(PCG::Private::MetadataTypes<AttributeType>::Id, TempConstAccessor->GetUnderlyingType()))
 				{
-					PCGE_LOG(Error, GraphAndLog, FText::Format(LOCTEXT("AttributeTypeBroadcastFailed_Updated", "Attribute/Property '{0}' ({1}) is not compatible with operation output type ({2})."),
+					PCGE_LOG(Error, GraphAndLog, FText::Format(LOCTEXT("AttributeTypeBroadcastFailed_Updated", "Output Attribute/Property '{0}' ({1}) is not compatible with operation output type ({2})."),
 						OutputTarget.GetDisplayText(),
 						PCG::Private::GetTypeNameText(TempConstAccessor->GetUnderlyingType()),
 						PCG::Private::GetTypeNameText<AttributeType>()));
