[PCG] Fix assert when running the GC at tasks end.

* We can't run the GC within a World tick, so we need to force it but not doing it immediately.

#rnx

diff --git a/Engine/Plugins/PCG/Source/PCG/Private/Graph/PCGGraphExecutor.cpp b/Engine/Plugins/PCG/Source/PCG/Private/Graph/PCGGraphExecutor.cpp
--- a/Engine/Plugins/PCG/Source/PCG/Private/Graph/PCGGraphExecutor.cpp
+++ b/Engine/Plugins/PCG/Source/PCG/Private/Graph/PCGGraphExecutor.cpp
@@ -22,6 +22,7 @@
 
 #include "Algo/AnyOf.h"
 #include "Algo/ForEach.h"
+#include "Engine/Engine.h"
 #include "GameFramework/Actor.h"
 #include "Misc/ScopeExit.h"
 #include "Tasks/Task.h"
@@ -890,9 +891,9 @@ void FPCGGraphExecutor::ExecuteTasksEnded()
 	}
 
 	// Purge things from cache if memory usage is too high
-	if (GraphCache.EnforceMemoryBudget())
+	if (GraphCache.EnforceMemoryBudget() && GEngine)
 	{
-		CollectGarbage(GARBAGE_COLLECTION_KEEPFLAGS, true);
+		GEngine->ForceGarbageCollection(/*bFullPurge=*/true);
 	}
 }
 
