[PCG] Add missing shader template file from commit that added the element count multiplier

#rnx

diff --git a/Engine/Plugins/PCG/Shaders/Private/PCGDataCollectionDataInterface.ush b/Engine/Plugins/PCG/Shaders/Private/PCGDataCollectionDataInterface.ush
--- a/Engine/Plugins/PCG/Shaders/Private/PCGDataCollectionDataInterface.ush
+++ b/Engine/Plugins/PCG/Shaders/Private/PCGDataCollectionDataInterface.ush
@@ -45,6 +45,9 @@
 
 RWByteAddressBuffer {DataInterfaceName}_DataCollectionBuffer;
 
+// Multiplier that can be configured on output pins.
+uint {DataInterfaceName}_ElementCountMultiplier;
+
 RWStructuredBuffer<uint> {DataInterfaceName}_ElementCountersBuffer;
 uint {DataInterfaceName}_ElementCountersPresent;
 
@@ -258,6 +261,11 @@ uint2 GetName_{DataInterfaceName}(uint DataIndex, uint ElementIndex, int Attribu
 	return uint2((uint)GetInt_{DataInterfaceName}(DataIndex, ElementIndex, AttributeId), 0);
 }
 
+uint GetElementCountMultiplier_{DataInterfaceName}()
+{
+	return {DataInterfaceName}_ElementCountMultiplier;
+}
+
 // #################### ATTRIBUTE SETTERS ##########################
 
 void SetBool_{DataInterfaceName}(uint DataIndex, uint ElementIndex, int AttributeId, bool Value)
