[PCG Test] Add Missing Root component to resolve "Failed to attach component" warnings when running tests

#rnx

diff --git a/Engine/Plugins/PCG/Source/PCG/Private/Tests/PCGTestsCommon.cpp b/Engine/Plugins/PCG/Source/PCG/Private/Tests/PCGTestsCommon.cpp
--- a/Engine/Plugins/PCG/Source/PCG/Private/Tests/PCGTestsCommon.cpp
+++ b/Engine/Plugins/PCG/Source/PCG/Private/Tests/PCGTestsCommon.cpp
@@ -17,6 +17,7 @@
 
 #if WITH_EDITOR
 #include "Editor.h"
+#include "Components/SceneComponent.h"
 #endif
 
 namespace PCGTestsCommon
@@ -58,6 +59,15 @@ namespace PCGTestsCommon
 		check(TestGraph);
 		TestPCGComponent->SetGraphLocal(TestGraph);
 
+		// Add Root Component to actor if none exists
+		if (!TestActor->GetRootComponent())
+		{
+			USceneComponent* NewRootComponent = NewObject<USceneComponent>(TestActor, FName(TEXT("DefaultSceneRoot")), RF_Transient);
+			TestActor->SetRootComponent(NewRootComponent);
+			TestActor->AddInstanceComponent(NewRootComponent);
+			NewRootComponent->RegisterComponent();
+		}
+
 		// Initialize CRC to avoid asserts.
 		InputData.ComputeCrcs(/*bFullDataCrc=*/false);
 #else
