[PCG] 1/3 Inline Pin Constants (Default Values) Frontend

* Reorganization of editor graph nodes into folder
* Organizational precursor to new Inline Pin Constants

NOTE: Organizational changes only. No functional changes intended in this shelf.

diff --git a/Engine/Plugins/PCG/Source/PCGEditor/Private/Nodes/PCGEditorGraphNode.cpp b/Engine/Plugins/PCG/Source/PCGEditor/Private/Nodes/PCGEditorGraphNode.cpp
new file mode 100644
--- /dev/null
+++ b/Engine/Plugins/PCG/Source/PCGEditor/Private/Nodes/PCGEditorGraphNode.cpp
@@ -0,0 +1,110 @@
+// Copyright Epic Games, Inc. All Rights Reserved.
+
+#include "PCGEditorGraphNode.h"
+
+#include "PCGNode.h"
+#include "PCGPin.h"
+#include "PCGSettings.h"
+#include "PCGSubsystem.h"
+
+#include "PCGEditorModule.h"
+#include "PCGEditorSettings.h"
+
+#include "GraphEditorActions.h"
+#include "ToolMenu.h"
+#include "ToolMenuSection.h"
+#include "Framework/Commands/GenericCommands.h"
+
+#define LOCTEXT_NAMESPACE "PCGEditorGraphNode"
+
+UPCGEditorGraphNode::UPCGEditorGraphNode(const FObjectInitializer& ObjectInitializer)
+	: Super(ObjectInitializer)
+{
+	bCanRenameNode = true;
+}
+
+FText UPCGEditorGraphNode::GetNodeTitle(ENodeTitleType::Type TitleType) const
+{
+	if (!PCGNode)
+	{
+		return NSLOCTEXT("PCGEditorGraphNode", "UnnamedNodeTitle", "Unnamed Node");
+	}
+
+	if (TitleType == ENodeTitleType::FullTitle)
+	{
+		return PCGNode->GetNodeTitle(EPCGNodeTitleType::FullTitle);
+	}
+	else
+	{
+		return PCGNode->GetNodeTitle(EPCGNodeTitleType::ListView);
+	}
+}
+
+void UPCGEditorGraphNode::GetNodeContextMenuActions(UToolMenu* Menu, class UGraphNodeContextMenuContext* Context) const
+{
+	if (!Context->Node)
+	{
+		return;
+	}
+
+	Super::GetNodeContextMenuActions(Menu, Context);
+
+	// General node actions
+	if (GetDefault<UPCGEditorSettings>()->bShowNodeGeneralActionsRightClickContextMenu)
+	{
+		FToolMenuSection& Section = Menu->AddSection("EdGraphSchemaGeneral", LOCTEXT("GeneralHeader", "General"));
+		Section.AddMenuEntry(FGenericCommands::Get().Delete);
+		Section.AddMenuEntry(FGenericCommands::Get().Cut);
+		Section.AddMenuEntry(FGenericCommands::Get().Copy);
+		Section.AddMenuEntry(FGenericCommands::Get().Duplicate);
+	}
+}
+
+void UPCGEditorGraphNode::AllocateDefaultPins()
+{
+	if (PCGNode)
+	{
+		CreatePins(PCGNode->GetInputPins(), PCGNode->GetOutputPins());
+	}
+}
+
+void UPCGEditorGraphNode::OnRenameNode(const FString& NewName)
+{
+	if (CanRenameNode(NewName))
+	{
+		CommitNodeRename(NewName);
+	}
+}
+
+bool UPCGEditorGraphNode::CanRenameNode(const FString& NewName) const
+{
+	if (!GetCanRenameNode())
+	{
+		return false;
+	}
+
+	if (NewName.Len() >= NAME_SIZE)
+	{
+		UE_LOG(LogPCGEditor, Error, TEXT("New name for PCG node is too long."));
+		return false;
+	}
+
+	if (!PCGNode)
+	{
+		return false;
+	}
+
+	return true;
+}
+
+void UPCGEditorGraphNode::CommitNodeRename(const FString& NewName)
+{
+	if (PCGNode->GetAuthoredTitleLine().ToString() != NewName)
+	{
+		Modify();
+		PCGNode->Modify();
+		PCGNode->NodeTitle = FName(*NewName);
+	}
+}
+
+#undef LOCTEXT_NAMESPACE
\ No newline at end of file
diff --git a/Engine/Plugins/PCG/Source/PCGEditor/Private/Nodes/PCGEditorGraphNode.h b/Engine/Plugins/PCG/Source/PCGEditor/Private/Nodes/PCGEditorGraphNode.h
new file mode 100644
--- /dev/null
+++ b/Engine/Plugins/PCG/Source/PCGEditor/Private/Nodes/PCGEditorGraphNode.h
@@ -0,0 +1,30 @@
+// Copyright Epic Games, Inc. All Rights Reserved.
+
+#pragma once
+
+#include "PCGEditorGraphNodeBase.h"
+
+#include "PCGEditorGraphNode.generated.h"
+
+namespace ENodeTitleType { enum Type : int; }
+
+class UToolMenu;
+
+UCLASS()
+class UPCGEditorGraphNode : public UPCGEditorGraphNodeBase
+{
+	GENERATED_BODY()
+
+public:
+	UPCGEditorGraphNode(const FObjectInitializer& ObjectInitializer);
+	
+	// ~Begin UEdGraphNode interface
+	virtual FText GetNodeTitle(ENodeTitleType::Type TitleType) const override;
+	virtual void GetNodeContextMenuActions(UToolMenu* Menu, class UGraphNodeContextMenuContext* Context) const override;
+	virtual void AllocateDefaultPins() override;
+	virtual void OnRenameNode(const FString& NewName) override;
+	// ~End UEdGraphNode interface
+
+	virtual bool CanRenameNode(const FString& NewName) const;
+	virtual void CommitNodeRename(const FString& NewName);
+};
\ No newline at end of file
diff --git a/Engine/Plugins/PCG/Source/PCGEditor/Private/Nodes/PCGEditorGraphNodeBase.cpp b/Engine/Plugins/PCG/Source/PCGEditor/Private/Nodes/PCGEditorGraphNodeBase.cpp
new file mode 100644
--- /dev/null
+++ b/Engine/Plugins/PCG/Source/PCGEditor/Private/Nodes/PCGEditorGraphNodeBase.cpp
@@ -0,0 +1,1320 @@
+// Copyright Epic Games, Inc. All Rights Reserved.
+
+#include "PCGEditorGraphNodeBase.h"
+
+#include "PCGComponent.h"
+#include "PCGEdge.h"
+#include "PCGEngineSettings.h"
+#include "PCGGraph.h"
+#include "PCGPin.h"
+#include "PCGSettingsWithDynamicInputs.h"
+#include "PCGSubsystem.h"
+#include "PCGWorldActor.h"
+#include "Elements/PCGHiGenGridSize.h"
+#include "Elements/PCGReroute.h"
+
+#include "PCGEditor.h"
+#include "PCGEditorCommands.h"
+#include "PCGEditorCommon.h"
+#include "PCGEditorGraph.h"
+#include "PCGEditorGraphSchema.h"
+#include "PCGEditorSettings.h"
+#include "PCGEditorStyle.h"
+
+#include "CoreGlobals.h"
+#include "GraphEditorActions.h"
+#include "ScopedTransaction.h"
+#include "ToolMenu.h"
+#include "ToolMenuSection.h"
+#include "Logging/TokenizedMessage.h"
+#include "Misc/TransactionObjectEvent.h"
+#include "Widgets/Colors/SColorPicker.h"
+
+#define LOCTEXT_NAMESPACE "PCGEditorGraphNodeBase"
+
+namespace PCGEditorGraphSwitches
+{
+	TAutoConsoleVariable<bool> CVarCheckConnectionCycles{
+		TEXT("pcg.Editor.CheckConnectionCycles"),
+		true,
+		TEXT("Prevents user from creating cycles in graph")
+	};
+}
+
+namespace PCGEditorGraphNodeBase
+{
+	/** Whether this node was culled during graph compilation or during graph execution. */
+	bool ShouldDisplayAsActive(const UPCGEditorGraphNodeBase* InNode, const UPCGComponent* InComponentBeingDebugged, const FPCGStack* InStackBeingInspected)
+	{
+		if (!InNode)
+		{
+			return true;
+		}
+
+		const UPCGNode* PCGNode = InNode->GetPCGNode();
+		if (!PCGNode)
+		{
+			return true;
+		}
+
+		// Don't display as culled while component is executing or about to refresh as nodes will flash to culled state and back
+		// which looks disturbing.
+		if (!InComponentBeingDebugged || InComponentBeingDebugged->IsGenerating() || InComponentBeingDebugged->IsRefreshInProgress())
+		{
+			return true;
+		}
+
+		const UPCGEngineSettings* EngineSettings = GetDefault<UPCGEngineSettings>();
+		const bool bActiveVisualizationEnabled = !ensure(EngineSettings) || EngineSettings->bDisplayCullingStateWhenDebugging;
+		const UPCGSettings* Settings = PCGNode ? PCGNode->GetSettings() : nullptr;
+
+		// Display whether node was culled dynamically or statically.
+		if (InStackBeingInspected && bActiveVisualizationEnabled)
+		{
+			if (!Settings || !Settings->IsA<UPCGRerouteSettings>())
+			{
+				// Task will be displayed as active if it was executed.
+				return InComponentBeingDebugged->WasNodeExecuted(PCGNode, *InStackBeingInspected);
+			}
+			else
+			{
+				// Named reroute usages mirror the enabled state of the upstream declaration.
+				if (const UPCGNamedRerouteUsageSettings* RerouteUsageSettings = Cast<UPCGNamedRerouteUsageSettings>(Settings))
+				{
+					const UPCGNode* DeclarationPCGNode = RerouteUsageSettings->Declaration ? Cast<UPCGNode>(RerouteUsageSettings->Declaration->GetOuter()) : nullptr;
+					const UPCGEditorGraph* EditorGraph = Cast<UPCGEditorGraph>(InNode->GetOuter());
+					const UPCGEditorGraphNodeBase* DeclarationNode = EditorGraph ? EditorGraph->GetEditorNodeFromPCGNode(DeclarationPCGNode) : nullptr;
+					return !DeclarationNode || ShouldDisplayAsActive(DeclarationNode, InComponentBeingDebugged, InStackBeingInspected);
+				}
+
+				// Special case - reroute culled state is evaluated here based on upstream connections. Reroutes are always culled/never executed, but still need
+				// to reflect the active/inactive state to not look wrong/confusing.
+				for (const UEdGraphPin* Pin : InNode->Pins)
+				{
+					if (Pin && Pin->Direction == EEdGraphPinDirection::EGPD_Input)
+					{
+						for (const UEdGraphPin* LinkedPin : Pin->LinkedTo)
+						{
+							if (const UPCGEditorGraphNodeBase* UpstreamNode = LinkedPin ? Cast<UPCGEditorGraphNodeBase>(LinkedPin->GetOwningNode()) : nullptr)
+							{
+								const bool bUpstreamNodeActive = ShouldDisplayAsActive(UpstreamNode, InComponentBeingDebugged, InStackBeingInspected);
+								const bool bUpstreamPinActive = UpstreamNode->IsOutputPinActive(LinkedPin);
+
+								if (bUpstreamNodeActive && bUpstreamPinActive)
+								{
+									// Active if any input is active.
+									return true;
+								}
+							}
+						}
+					}
+				}
+
+				return false;
+			}
+		}
+
+		return true;
+	}
+
+	namespace Constants
+	{
+		FName NodeConversionActionName = TEXT("NodeConversion");
+
+		FText NodeConversionActionLabel = LOCTEXT("NodeConversion", "Convert Node");
+		FText NodeConversionActionTooltip = LOCTEXT("NodeConversionTooltip", "Convert a single node into a different node, or otherwise compatible output.");
+		FTextFormat ConvertToFormat = LOCTEXT("ConvertTo", "Convert to {0}");
+
+		FText ConversionHeaderLabel = LOCTEXT("ConvertHeader", "Convert To");
+		FText OrganizationHeaderLabel = LOCTEXT("OrganizationHeader", "Organization");
+		FText DeterminismHeaderLabel = LOCTEXT("DeterminismHeader", "Determinism");
+	}
+}
+
+void UPCGEditorGraphNodeBase::Construct(UPCGNode* InPCGNode)
+{
+	check(InPCGNode);
+	PCGNode = InPCGNode;
+	InPCGNode->OnNodeChangedDelegate.AddUObject(this, &UPCGEditorGraphNodeBase::OnNodeChanged);
+
+	NodePosX = InPCGNode->PositionX;
+	NodePosY = InPCGNode->PositionY;
+	NodeComment = InPCGNode->NodeComment;
+	bCommentBubblePinned = InPCGNode->bCommentBubblePinned;
+	bCommentBubbleVisible = InPCGNode->bCommentBubbleVisible;
+	bCanRenameNode = false;
+
+	if (const UPCGSettingsInterface* PCGSettingsInterface = InPCGNode->GetSettingsInterface())
+	{
+		const ENodeEnabledState NewEnabledState = !PCGSettingsInterface->bEnabled ? ENodeEnabledState::Disabled : ENodeEnabledState::Enabled;
+		SetEnabledState(NewEnabledState);
+	}
+
+	// Update to current graph/inspection state.
+	const UPCGEditorGraph* Graph = Cast<UPCGEditorGraph>(GetOuter());
+	const FPCGEditor* Editor = Graph ? Graph->GetEditor().Pin().Get() : nullptr;
+	UpdateStructuralVisualization(Editor ? Editor->GetPCGComponentBeingInspected() : nullptr, Editor ? Editor->GetStackBeingInspected() : nullptr, /*bNewlyPlaced=*/true);
+}
+
+void UPCGEditorGraphNodeBase::BeginDestroy()
+{
+	if (PCGNode)
+	{
+		PCGNode->OnNodeChangedDelegate.RemoveAll(this);
+	}
+
+	Super::BeginDestroy();
+}
+
+void UPCGEditorGraphNodeBase::PostTransacted(const FTransactionObjectEvent& TransactionEvent)
+{
+	Super::PostTransacted(TransactionEvent);
+
+	TArray<FName> PropertiesChanged = TransactionEvent.GetChangedProperties();
+
+	if (PropertiesChanged.Contains(TEXT("bCommentBubblePinned")))
+	{
+		UpdateCommentBubblePinned();
+	}
+
+	if (PropertiesChanged.Contains(TEXT("NodePosX")) || PropertiesChanged.Contains(TEXT("NodePosY")))
+	{
+		UpdatePosition();
+	}
+}
+
+void UPCGEditorGraphNodeBase::GetNodeContextMenuActions(UToolMenu* Menu, class UGraphNodeContextMenuContext* Context) const
+{
+	using namespace PCGEditorGraphNodeBase::Constants;
+
+	if (!Context->Node)
+	{
+		return;
+	}
+
+	// Local pin special actions should come first.
+	{
+		FToolMenuSection& Section = Menu->AddSection("EdGraphSchemaPinActions", LOCTEXT("PinActionsMenuHeader", "Pin Actions"));
+		if (Context->Pin && CanUserAddRemoveDynamicInputPins())
+		{
+			Section.AddMenuEntry(FPCGEditorCommands::Get().AddSourcePin);
+			Section.AddMenuEntry("RemovePin",
+				LOCTEXT("RemovePin", "Remove Source Pin"),
+				LOCTEXT("RemovePinTooltip", "Remove this source pin from the current node"),
+				FSlateIcon(),
+				FUIAction(
+					FExecuteAction::CreateLambda([Pin = Context->Pin, this]
+					{
+						const_cast<UPCGEditorGraphNodeBase*>(this)->OnUserRemoveDynamicInputPin(const_cast<UEdGraphPin*>(Pin));
+					}),
+					FCanExecuteAction::CreateLambda([Pin = Context->Pin, this]
+					{
+						return const_cast<UPCGEditorGraphNodeBase*>(this)->CanUserRemoveDynamicInputPin(const_cast<UEdGraphPin*>(Pin));
+					})));
+		}
+	}
+
+	// Node special actions should come after pin actions
+	{
+		FToolMenuSection& Section = Menu->AddSection("EdGraphSchemaNodeActions", LOCTEXT("NodeActionsHeader", "Node Actions"));
+
+		const UPCGSettings* Settings = PCGNode->GetSettings();
+
+		// Special nodes operations
+		if (PCGNode && Settings)
+		{
+			if (Settings->IsA<UPCGNamedRerouteDeclarationSettings>())
+			{
+				Section.AddMenuEntry(FPCGEditorCommands::Get().SelectNamedRerouteUsages);
+			}
+			else if (Settings->IsA<UPCGNamedRerouteUsageSettings>())
+			{
+				Section.AddMenuEntry(FPCGEditorCommands::Get().SelectNamedRerouteDeclaration);
+			}
+
+			// Operation to convert a node in place into something else, such as a different node.
+			const TArray<FPCGPreconfiguredInfo> ConversionInfo = Settings->GetConversionInfo();
+			if (!ConversionInfo.IsEmpty())
+			{
+				const int32 NumConversions = ConversionInfo.Num();
+
+				auto AddEntries = [this, &Section, ConversionInfo = std::move(ConversionInfo)](UToolMenu* AlignmentMenu)
+				{
+					FToolMenuSection& SubSection = (ConversionInfo.Num() > 1) ? AlignmentMenu->AddSection("EdGraphSchemaConversion", ConversionHeaderLabel) : Section;
+
+					for (const FPCGPreconfiguredInfo& Conversion : ConversionInfo)
+					{
+						const FText Label = (ConversionInfo.Num() > 1) ? Conversion.Label : FText::Format(ConvertToFormat, Conversion.Label);
+						const FText Tooltip = Conversion.Tooltip;
+						FUIAction Action(FExecuteAction::CreateLambda([this, Conversion = std::move(Conversion)] { const_cast<UPCGEditorGraphNodeBase*>(this)->OnConvertNode(Conversion); }));
+						// TODO: UX feedback (read-only or greyed out font, etc) if the conversion can't happen.
+						SubSection.AddMenuEntry(FName(Label.ToString()), Label, std::move(Tooltip), FSlateIcon(), std::move(Action));
+					}
+				};
+
+				if (NumConversions == 1) // Single conversion, just add it to the action list
+				{
+					AddEntries(Menu);
+				}
+				else // Add the entries to a submenu
+				{
+					Section.AddSubMenu(NodeConversionActionName, NodeConversionActionLabel, NodeConversionActionTooltip, FNewToolMenuDelegate::CreateLambda(AddEntries));
+				}
+			}
+		}
+
+		// General PCG node actions
+		Section.AddMenuEntry(FPCGEditorCommands::Get().ToggleEnabled, LOCTEXT("ToggleEnabledLabel", "Enable"));
+		Section.AddMenuEntry(FPCGEditorCommands::Get().RenameNode, LOCTEXT("RenameNode", "Rename"));
+		Section.AddMenuEntry(FPCGEditorCommands::Get().ToggleDebug, LOCTEXT("ToggleDebugLabel", "Debug"));
+		Section.AddMenuEntry(FPCGEditorCommands::Get().DebugOnlySelected);
+		Section.AddMenuEntry(FPCGEditorCommands::Get().DisableDebugOnAllNodes);
+		Section.AddMenuEntry(FPCGEditorCommands::Get().ToggleInspect, LOCTEXT("ToggleinspectionLabel", "Inspect"));
+		Section.AddMenuEntry(FGraphEditorCommands::Get().BreakNodeLinks);
+		Section.AddMenuEntry(FPCGEditorCommands::Get().ExportNodes, TAttribute<FText>(), TAttribute<FText>(), FSlateIcon(FPCGEditorStyle::Get().GetStyleSetName(), "ClassIcon.PCGSettings"));
+		Section.AddMenuEntry(FPCGEditorCommands::Get().CollapseNodes);
+		Section.AddMenuEntry(FPCGEditorCommands::Get().ConvertToStandaloneNodes);
+	}
+
+	// Organizational actions
+	if (GetDefault<UPCGEditorSettings>()->bShowNodeOrganizationalActionsRightClickContextMenu)
+	{
+		FToolMenuSection& Section = Menu->AddSection("EdGraphSchemaOrganization", OrganizationHeaderLabel);
+		Section.AddMenuEntry(
+			"PCGNode_SetColor",
+			LOCTEXT("PCGNode_SetColor", "Set Node Color"),
+			LOCTEXT("PCGNode_SetColorTooltip", "Sets a specific color on the given node. Note that white maps to the default value"),
+			FSlateIcon(FAppStyle::GetAppStyleSetName(), "ColorPicker.Mode"),
+			FUIAction(FExecuteAction::CreateUObject(const_cast<UPCGEditorGraphNodeBase*>(this), &UPCGEditorGraphNodeBase::OnPickColor),
+				FCanExecuteAction::CreateUObject(const_cast<UPCGEditorGraphNodeBase*>(this), &UPCGEditorGraphNodeBase::CanPickColor)));
+
+		Section.AddSubMenu("Alignment", LOCTEXT("AlignmentHeader", "Alignment"), FText(), FNewToolMenuDelegate::CreateLambda([](UToolMenu* AlignmentMenu)
+		{
+			{
+				FToolMenuSection& SubSection = AlignmentMenu->AddSection("EdGraphSchemaAlignment", LOCTEXT("AlignHeader", "Align"));
+				SubSection.AddMenuEntry(FGraphEditorCommands::Get().AlignNodesTop);
+				SubSection.AddMenuEntry(FGraphEditorCommands::Get().AlignNodesMiddle);
+				SubSection.AddMenuEntry(FGraphEditorCommands::Get().AlignNodesBottom);
+				SubSection.AddMenuEntry(FGraphEditorCommands::Get().AlignNodesLeft);
+				SubSection.AddMenuEntry(FGraphEditorCommands::Get().AlignNodesCenter);
+				SubSection.AddMenuEntry(FGraphEditorCommands::Get().AlignNodesRight);
+				SubSection.AddMenuEntry(FGraphEditorCommands::Get().StraightenConnections);
+			}
+
+			{
+				FToolMenuSection& SubSection = AlignmentMenu->AddSection("EdGraphSchemaDistribution", LOCTEXT("DistributionHeader", "Distribution"));
+				SubSection.AddMenuEntry(FGraphEditorCommands::Get().DistributeNodesHorizontally);
+				SubSection.AddMenuEntry(FGraphEditorCommands::Get().DistributeNodesVertically);
+			}
+		}));
+	}
+
+	// PCG Determinism actions
+	if (GetDefault<UPCGEditorSettings>()->bShowNodeDeterminismActionsRightClickContext)
+	{
+		FToolMenuSection& Section = Menu->AddSection("EdGraphSchemaDeterminism", DeterminismHeaderLabel);
+		Section.AddMenuEntry(FPCGEditorCommands::Get().RunDeterminismNodeTest,
+			LOCTEXT("Determinism_RunTest", "Validate Determinism on Selection"),
+			LOCTEXT("Determinism_RunTestToolTip", "Run a test to validate the selected nodes for determinism."));
+	}
+
+	// Comment Group is the final section
+	{
+		FToolMenuSection& Section = Menu->AddSection("EdGraphSchemaCommentGroup", LOCTEXT("CommentGroupHeader", "Comment Group"));
+		Section.AddMenuEntry(FGraphEditorCommands::Get().CreateComment,
+			LOCTEXT("MultiCommentDesc", "Create Comment from Selection"),
+			LOCTEXT("CommentToolTip", "Create a resizable comment box around selection."));
+	}
+}
+
+void UPCGEditorGraphNodeBase::AutowireNewNode(UEdGraphPin* FromPin)
+{
+	if (PCGNode == nullptr || FromPin == nullptr)
+	{
+		return;
+	}
+
+	const bool bFromPinIsInput = FromPin->Direction == EEdGraphPinDirection::EGPD_Input;
+	const TArray<TObjectPtr<UPCGPin>>& OtherPinsList = bFromPinIsInput ? PCGNode->GetOutputPins() : PCGNode->GetInputPins();
+
+	// Try to connect to the first compatible pin
+	for (const TObjectPtr<UPCGPin>& OtherPin : OtherPinsList)
+	{
+		check(OtherPin);
+
+		// TODO: Allow autoconnecting output dependency pins to input dependency pins.
+		if (OtherPin->Properties.IsAdvancedPin() || OtherPin->Properties.IsDatalessPin())
+		{
+			continue;
+		}
+
+		const FName& OtherPinName = OtherPin->Properties.Label;
+		UEdGraphPin* ToPin = FindPinChecked(OtherPinName, bFromPinIsInput ? EEdGraphPinDirection::EGPD_Output : EEdGraphPinDirection::EGPD_Input);
+		if (ToPin && GetSchema()->TryCreateConnection(FromPin, ToPin))
+		{
+			// Connection succeeded
+			break;
+		}
+	}
+
+	NodeConnectionListChanged();
+}
+
+void UPCGEditorGraphNodeBase::PrepareForCopying()
+{
+	if (PCGNode)
+	{
+		// Temporarily take ownership of the MaterialExpression, so that it is not deleted when cutting
+		PCGNode->Rename(nullptr, this, REN_DontCreateRedirectors | REN_DoNotDirty);
+	}
+}
+
+bool UPCGEditorGraphNodeBase::CanCreateUnderSpecifiedSchema(const UEdGraphSchema* Schema) const
+{
+	return Schema->IsA(UPCGEditorGraphSchema::StaticClass());
+}
+
+void UPCGEditorGraphNodeBase::PostCopy()
+{
+	if (PCGNode)
+	{
+		UPCGEditorGraph* PCGEditorGraph = CastChecked<UPCGEditorGraph>(GetGraph());
+		UPCGGraph* PCGGraph = PCGEditorGraph->GetPCGGraph();
+		check(PCGGraph);
+		PCGNode->Rename(nullptr, PCGGraph, REN_DontCreateRedirectors | REN_DoNotDirty);
+	}
+}
+
+void UPCGEditorGraphNodeBase::PostPasteNode()
+{
+	bDisableReconstructFromNode = true;
+}
+
+void UPCGEditorGraphNodeBase::RebuildAfterPaste()
+{
+	if (PCGNode)
+	{
+		PCGNode->RebuildAfterPaste();
+
+		RebuildEdgesFromPins();
+
+		PCGNode->OnNodeChangedDelegate.AddUObject(this, &UPCGEditorGraphNodeBase::OnNodeChanged);
+		PCGNode->PositionX = NodePosX;
+		PCGNode->PositionY = NodePosY;
+
+		// Refresh the node if it has dynamic pins
+		if (const UPCGSettings* Settings = PCGNode->GetSettings())
+		{
+			if (Settings->HasDynamicPins())
+			{
+				PCGNode->OnNodeChangedDelegate.Broadcast(PCGNode, EPCGChangeType::Node);
+			}
+		}
+	}
+}
+
+void UPCGEditorGraphNodeBase::PostPaste()
+{
+	bDisableReconstructFromNode = false;
+}
+
+void UPCGEditorGraphNodeBase::EnableDeferredReconstruct()
+{
+	ensure(DeferredReconstructCounter >= 0);
+	++DeferredReconstructCounter;
+}
+
+void UPCGEditorGraphNodeBase::DisableDeferredReconstruct()
+{
+	ensure(DeferredReconstructCounter > 0);
+	--DeferredReconstructCounter;
+
+	if (DeferredReconstructCounter == 0 && bDeferredReconstruct)
+	{
+		ReconstructNode();
+		bDeferredReconstruct = false;
+	}
+}
+
+void UPCGEditorGraphNodeBase::RebuildEdgesFromPins()
+{
+	check(PCGNode);
+	check(bDisableReconstructFromNode);
+
+	if (PCGNode->GetGraph())
+	{
+		PCGNode->GetGraph()->DisableNotificationsForEditor();
+	}
+
+	RebuildEdgesFromPins_Internal();
+	
+	if (PCGNode->GetGraph())
+	{
+		PCGNode->GetGraph()->EnableNotificationsForEditor();
+	}
+}
+
+void UPCGEditorGraphNodeBase::RebuildEdgesFromPins_Internal()
+{
+	check(PCGNode);
+	check(bDisableReconstructFromNode);
+
+	for (UEdGraphPin* Pin : Pins)
+	{
+		if (Pin->Direction == EEdGraphPinDirection::EGPD_Output)
+		{
+			for (UEdGraphPin* ConnectedPin : Pin->LinkedTo)
+			{
+				UEdGraphNode* ConnectedGraphNode = ConnectedPin->GetOwningNode();
+				UPCGEditorGraphNodeBase* ConnectedPCGGraphNode = CastChecked<UPCGEditorGraphNodeBase>(ConnectedGraphNode);
+
+				if (UPCGNode* ConnectedPCGNode = ConnectedPCGGraphNode->GetPCGNode())
+				{
+					PCGNode->AddEdgeTo(Pin->PinName, ConnectedPCGNode, ConnectedPin->PinName);
+				}
+			}
+		}
+	}
+}
+
+void UPCGEditorGraphNodeBase::OnNodeChanged(UPCGNode* InNode, EPCGChangeType ChangeType)
+{
+	if (InNode == PCGNode)
+	{
+		if (!!(ChangeType & EPCGChangeType::Settings))
+		{
+			if (const UPCGSettingsInterface* PCGSettingsInterface = InNode->GetSettingsInterface())
+			{
+				const ENodeEnabledState NewEnabledState = PCGSettingsInterface->bEnabled ? ENodeEnabledState::Enabled : ENodeEnabledState::Disabled;
+				if (NewEnabledState != GetDesiredEnabledState())
+				{
+					SetEnabledState(NewEnabledState);
+				}
+			}
+		}
+
+		ChangeType |= UpdateErrorsAndWarnings();
+
+		UPCGComponent* ComponentBeingDebugged = nullptr;
+		const FPCGStack* StackBeingDebugged = nullptr;
+
+		{
+			const UPCGEditorGraph* EditorGraph = CastChecked<UPCGEditorGraph>(GetGraph());
+			const FPCGEditor* Editor = EditorGraph ? EditorGraph->GetEditor().Pin().Get() : nullptr;
+			ComponentBeingDebugged = Editor ? Editor->GetPCGComponentBeingInspected() : nullptr;
+			StackBeingDebugged = Editor ? Editor->GetStackBeingInspected() : nullptr;
+		}
+
+		if (!!(ChangeType & (EPCGChangeType::Structural | EPCGChangeType::Node | EPCGChangeType::Edge | EPCGChangeType::Cosmetic)))
+		{
+			ReconstructNodeOnChange();
+		}
+	}
+}
+
+EPCGChangeType UPCGEditorGraphNodeBase::UpdateErrorsAndWarnings()
+{
+	const UPCGSubsystem* Subsystem = UPCGSubsystem::GetActiveEditorInstance();
+	if (!PCGNode || !Subsystem)
+	{
+		return EPCGChangeType::None;
+	}
+	
+	const FPCGStack* InspectedStack = nullptr;
+	{
+		const UPCGEditorGraph* EditorGraph = CastChecked<UPCGEditorGraph>(GetGraph());
+		const FPCGEditor* Editor = (EditorGraph && EditorGraph->GetEditor().IsValid()) ? EditorGraph->GetEditor().Pin().Get() : nullptr;
+		InspectedStack = Editor ? Editor->GetStackBeingInspected() : nullptr;
+	}
+
+	const bool bOldHasCompilerMessage = bHasCompilerMessage;
+	const int32 OldErrorType = ErrorType;
+	const FString OldErrorMsg = ErrorMsg;
+		
+	if (InspectedStack)
+	{
+		// Get errors/warnings for the inspected stack.
+		FPCGStack StackWithNode = *InspectedStack;
+		StackWithNode.PushFrame(PCGNode);
+		bHasCompilerMessage = Subsystem->GetNodeVisualLogs().HasLogs(StackWithNode);
+
+		if (bHasCompilerMessage)
+		{
+			ErrorMsg = Subsystem->GetNodeVisualLogs().GetLogsSummaryText(StackWithNode).ToString();
+
+			const bool bHasErrors = Subsystem->GetNodeVisualLogs().HasLogsOfVerbosity(StackWithNode, ELogVerbosity::Error);
+			ErrorType = bHasErrors ? EMessageSeverity::Error : EMessageSeverity::Warning;
+		}
+		else
+		{
+			ErrorMsg.Empty();
+			ErrorType = 0;
+		}
+	}
+	else
+	{
+		// Collect all errors/warnings for this node.
+		ELogVerbosity::Type MinimumVerbosity;
+		ErrorMsg = Subsystem->GetNodeVisualLogs().GetLogsSummaryText(PCGNode.Get(), MinimumVerbosity).ToString();
+
+		bHasCompilerMessage = !ErrorMsg.IsEmpty();
+
+		if (bHasCompilerMessage)
+		{
+			ErrorType = MinimumVerbosity < ELogVerbosity::Warning ? EMessageSeverity::Error : EMessageSeverity::Warning;
+		}
+		else
+		{
+			ErrorType = 0;
+		}
+	}
+
+	const bool bStateChanged = (bHasCompilerMessage != bOldHasCompilerMessage) || (ErrorType != OldErrorType) || (ErrorMsg != OldErrorMsg);
+	return bStateChanged ? EPCGChangeType::Cosmetic : EPCGChangeType::None;
+}
+
+EPCGChangeType UPCGEditorGraphNodeBase::UpdateStructuralVisualization(UPCGComponent* InComponentBeingDebugged, const FPCGStack* InStackBeingInspected, bool bNewlyPlaced)
+{
+	const UPCGGraph* Graph = PCGNode ? PCGNode->GetGraph() : nullptr;
+	if (!Graph)
+	{
+		return EPCGChangeType::None;
+	}
+
+	const bool bInspecting = InComponentBeingDebugged && InStackBeingInspected && !InStackBeingInspected->GetStackFrames().IsEmpty();
+
+	EPCGChangeType ChangeType = EPCGChangeType::None;
+
+	const uint64 NewInactiveMask = bInspecting ? InComponentBeingDebugged->GetNodeInactivePinMask(PCGNode, *InStackBeingInspected) : 0;
+	if (NewInactiveMask != InactiveOutputPinMask)
+	{
+		InactiveOutputPinMask = NewInactiveMask;
+		ChangeType |= EPCGChangeType::Cosmetic;
+	}
+
+	// Check top graph for higen enable - subgraphs always inherit higen state from the top graph.
+	const UPCGGraph* TopGraph = bInspecting ? InStackBeingInspected->GetRootGraph() : Graph;
+	const bool HiGenEnabled = TopGraph && TopGraph->IsHierarchicalGenerationEnabled();
+
+	// Set the inspected grid size - this is used for grid size visualization.
+	uint32 InspectingGridSize = PCGHiGenGrid::UninitializedGridSize();
+	EPCGHiGenGrid InspectingGrid = EPCGHiGenGrid::Uninitialized;
+	if (TopGraph && TopGraph->IsHierarchicalGenerationEnabled() && InComponentBeingDebugged && (InComponentBeingDebugged->IsPartitioned() || InComponentBeingDebugged->IsLocalComponent()))
+	{
+		InspectingGridSize = InComponentBeingDebugged->GetGenerationGridSize();
+		InspectingGrid = InComponentBeingDebugged->GetGenerationGrid();
+	}
+
+	if (InspectedGenerationGrid != InspectingGrid)
+	{
+		InspectedGenerationGrid = InspectingGrid;
+		ChangeType |= EPCGChangeType::Cosmetic;
+	}
+
+	bool bShouldDisplayAsDisabled = false;
+
+	// Special treatment for higen grid sizes nodes which do nothing if higen is disabled.
+	// TODO: Drive this from an API on settings as we add more higen-specific functionality.
+	if (PCGNode && Cast<UPCGHiGenGridSizeSettings>(PCGNode->GetSettings()))
+	{
+		// Higen must be enabled on graph, and we must be editing top graph.
+		bShouldDisplayAsDisabled = (Graph != TopGraph) || !TopGraph->IsHierarchicalGenerationEnabled();
+
+		// If we're inspecting a component, it must either be a partitioned OC or an LC (because higen requires partitioning).
+		if (!bShouldDisplayAsDisabled && InComponentBeingDebugged)
+		{
+			bShouldDisplayAsDisabled = !InComponentBeingDebugged->IsPartitioned() && !InComponentBeingDebugged->IsLocalComponent();
+		}
+	}
+
+	// Don't do culling visualization on newly placed nodes. Let the execution complete notification update that.
+	bool bIsCulled = !bNewlyPlaced && !PCGEditorGraphNodeBase::ShouldDisplayAsActive(this, InComponentBeingDebugged, InStackBeingInspected);
+
+	EPCGHiGenGrid ThisGrid = EPCGHiGenGrid::Uninitialized;
+
+	// Show grid size visualization if higen is enabled and if we're inspecting a specific grid, and we're inspecting a subgraph since subgraphs
+	// execute at the invoked grid level.
+	if (HiGenEnabled && InspectingGridSize != PCGHiGenGrid::UninitializedGridSize())
+	{
+		if (InStackBeingInspected && InStackBeingInspected->IsCurrentFrameInRootGraph())
+		{
+			const uint32 DefaultGridSize = TopGraph->GetDefaultGridSize();
+			const uint32 NodeGridSize = Graph->GetNodeGenerationGridSize(PCGNode, DefaultGridSize);
+
+			if (NodeGridSize < InspectingGridSize)
+			{
+				// Disable nodes that are on a smaller grid
+				bShouldDisplayAsDisabled |= NodeGridSize < InspectingGridSize;
+
+				// We don't know if the node was culled or not on that grid, disable visualization.
+				bIsCulled = false;
+			}
+			else if (NodeGridSize > InspectingGridSize)
+			{
+				// We don't know if the node was culled or not on that grid, disable visualization.
+				bIsCulled = false;
+			}
+
+			ThisGrid = PCGHiGenGrid::GridSizeToGrid(NodeGridSize);
+		}
+		else
+		{
+			// If higen is enabled then we are inspecting an invoked subgraph. Display the inspected grid size so that the user still
+			// gets the execution grid information.
+			ThisGrid = PCGHiGenGrid::GridSizeToGrid(InspectingGridSize);
+		}
+	}
+
+	if (GenerationGrid != ThisGrid)
+	{
+		GenerationGrid = ThisGrid;
+		ChangeType |= EPCGChangeType::Cosmetic;
+	}
+
+	SetIsCulledFromExecution(bIsCulled);
+
+	if (bIsCulled)
+	{
+		bShouldDisplayAsDisabled = true;
+	}
+
+	if (IsDisplayAsDisabledForced() != bShouldDisplayAsDisabled)
+	{
+		SetForceDisplayAsDisabled(bShouldDisplayAsDisabled);
+		ChangeType |= EPCGChangeType::Cosmetic;
+	}
+
+	return ChangeType;
+}
+
+bool UPCGEditorGraphNodeBase::ShouldDrawCompact() const
+{
+	UPCGSettings* Settings = PCGNode ? PCGNode->GetSettings() : nullptr;
+	return Settings && Settings->ShouldDrawNodeCompact();
+}
+
+bool UPCGEditorGraphNodeBase::GetCompactNodeIcon(FName& OutCompactNodeIcon) const
+{
+	UPCGSettings* Settings = PCGNode ? PCGNode->GetSettings() : nullptr;
+	return Settings && Settings->GetCompactNodeIcon(OutCompactNodeIcon);
+}
+
+bool UPCGEditorGraphNodeBase::CanUserAddRemoveDynamicInputPins() const
+{
+	const UPCGSettings* Settings = PCGNode ? PCGNode->GetSettings() : nullptr;
+	return Settings && Settings->IsA<UPCGSettingsWithDynamicInputs>();
+}
+
+bool UPCGEditorGraphNodeBase::HasFlippedTitleLines() const
+{
+	return PCGNode ? PCGNode->HasFlippedTitleLines() : false;
+}
+
+FText UPCGEditorGraphNodeBase::GetAuthoredTitleLine() const
+{
+	return PCGNode ? PCGNode->GetAuthoredTitleLine() : FText();
+}
+
+FText UPCGEditorGraphNodeBase::GetGeneratedTitleLine() const
+{
+	return PCGNode ? PCGNode->GetGeneratedTitleLine() : FText();
+}
+
+bool UPCGEditorGraphNodeBase::IsOutputPinActive(const UEdGraphPin* InOutputPin) const
+{
+	bool bPinActive = true;
+
+	if (InactiveOutputPinMask != 0)
+	{
+		bool bFoundPin = false;
+		int OutputPinIndex = 0;
+
+		for (const UEdGraphPin* NodePin : Pins)
+		{
+			if (NodePin == InOutputPin)
+			{
+				bFoundPin = true;
+				break;
+			}
+
+			if (NodePin->Direction == EEdGraphPinDirection::EGPD_Output)
+			{
+				++OutputPinIndex;
+			}
+		}
+
+		if (bFoundPin)
+		{
+			bPinActive = !((1ULL << OutputPinIndex) & InactiveOutputPinMask);
+		}
+	}
+
+	return bPinActive;
+}
+
+void UPCGEditorGraphNodeBase::EnterRenamingMode()
+{
+	bCanRenameNode = true;
+
+	// Notify the SPCGEditorGraphNode that the user is trying to rename it.
+	OnNodeRenameInitiatedDelegate.ExecuteIfBound();
+}
+
+void UPCGEditorGraphNodeBase::ExitRenamingMode()
+{
+	bCanRenameNode = false;
+
+	// Update so that the node renders with the new node title.
+	OnNodeChangedDelegate.ExecuteIfBound();
+}
+
+bool UPCGEditorGraphNodeBase::IsCompatible(const UPCGPin* InputPin, const UPCGPin* OutputPin, FText& OutReason) const
+{
+	if (PCGEditorGraphSwitches::CVarCheckConnectionCycles.GetValueOnAnyThread() && InputPin && OutputPin && InputPin->Node == PCGNode)
+	{
+		// Upstream Visitor
+		TSet<const UPCGNode*> VisitedNodes;
+		auto Visitor = [&VisitedNodes, ThisPCGNode = PCGNode](const UPCGNode* InNode, auto VisitorLambda) -> bool
+		{
+			if (InNode)
+			{
+				if (InNode == ThisPCGNode)
+				{
+					return false;
+				}
+				else if (VisitedNodes.Contains(InNode))
+				{
+					return true;
+				}
+
+				VisitedNodes.Add(InNode);
+
+				for (const TObjectPtr<UPCGPin>& InputPin : InNode->GetInputPins())
+				{
+					if (InputPin)
+					{
+						for (const TObjectPtr<UPCGEdge>& Edge : InputPin->Edges)
+						{
+							if (Edge)
+							{
+								if (const UPCGPin* OtherPin = Edge->GetOtherPin(InputPin.Get()))
+								{
+									if (!VisitorLambda(OtherPin->Node, VisitorLambda))
+									{
+										return false;
+									}
+								}
+							}
+						}
+					}
+				}
+			}
+
+			return true;
+		};
+
+		// OutputPin is trying to connect to this nodes InputPin so visit the OutputPin upstream and try to find
+		// a existing connection to this UPCGEditorGraphNodeNamedRerouteDeclaration's PCGNode. If we do deny connection which would create cycle.
+		if (!Visitor(OutputPin->Node, Visitor))
+		{
+			OutReason = LOCTEXT("ConnectionFailedCyclic", "Connection would create cycle");
+			return false;
+		}
+	}
+
+	return true;
+}
+
+void UPCGEditorGraphNodeBase::OnPickColor()
+{
+	FColorPickerArgs PickerArgs;
+	PickerArgs.bIsModal = true;
+	PickerArgs.bUseAlpha = false;
+	PickerArgs.InitialColor = GetNodeTitleColor();
+	PickerArgs.OnColorCommitted = FOnLinearColorValueChanged::CreateUObject(this, &UPCGEditorGraphNodeBase::OnColorPicked);
+
+	OpenColorPicker(PickerArgs);
+}
+
+void UPCGEditorGraphNodeBase::OnColorPicked(FLinearColor NewColor)
+{
+	if (PCGNode && GetNodeTitleColor() != NewColor)
+	{
+		PCGNode->Modify();
+		PCGNode->NodeTitleColor = NewColor;
+	}
+}
+
+void UPCGEditorGraphNodeBase::ReconstructNode()
+{
+	// In copy-paste cases, we don't want to remove the pins
+	if (bDisableReconstructFromNode)
+	{
+		return;
+	}
+
+	if (DeferredReconstructCounter > 0)
+	{
+		bDeferredReconstruct = true;
+		return;
+	}
+
+	// While in an Undo/Redo a call to ReconstructNode should not be needed as the transaction object records
+	// should be enough to serialize the nodes back into their proper state 
+	if (GIsTransacting)
+	{
+		return;
+	}
+	
+	Modify();
+
+	// Store copy of old pins
+	TArray<UEdGraphPin*> OldPins = MoveTemp(Pins);
+	Pins.Reset();
+
+	// Generate new pins
+	AllocateDefaultPins();
+
+	// Transfer persistent data from old to new pins
+	for (UEdGraphPin* OldPin : OldPins)
+	{
+		const FName& OldPinName = OldPin->PinName;
+		if (UEdGraphPin** NewPin = Pins.FindByPredicate([&OldPinName](UEdGraphPin* InPin) { return InPin->PinName == OldPinName; }))
+		{
+			(*NewPin)->MovePersistentDataFromOldPin(*OldPin);
+		}
+	}
+
+	// Remove old pins
+	for (UEdGraphPin* OldPin : OldPins)
+	{
+		RemovePin(OldPin);
+	}
+
+	// Generate new links
+	// TODO: we should either keep a map in the PCGEditorGraph or do this elsewhere
+	// TODO: this will not work if we have non-PCG nodes in the graph
+	if (PCGNode)
+	{
+		for (UEdGraphPin* Pin : Pins)
+		{
+			Pin->BreakAllPinLinks();
+		}
+		
+		UPCGEditorGraph* PCGEditorGraph = CastChecked<UPCGEditorGraph>(GetGraph());
+		PCGEditorGraph->CreateLinks(this, /*bCreateInbound=*/true, /*bCreateOutbound=*/true);
+	}
+
+	// Notify editor
+	OnNodeChangedDelegate.ExecuteIfBound();
+}
+
+FLinearColor UPCGEditorGraphNodeBase::GetNodeTitleColor() const
+{
+	if (PCGNode)
+	{
+		const UPCGSettingsInterface* PCGSettingsInterface = PCGNode->GetSettingsInterface();
+		const UPCGSettings* PCGSettings = PCGSettingsInterface ? PCGSettingsInterface->GetSettings() : nullptr;
+
+		if (PCGNode->NodeTitleColor != FLinearColor::White)
+		{
+			return PCGNode->NodeTitleColor;
+		}
+		else if (PCGSettings)
+		{
+			FLinearColor SettingsColor = PCGNode->GetSettings()->GetNodeTitleColor();
+			if (SettingsColor == FLinearColor::White)
+			{
+				SettingsColor = GetDefault<UPCGEditorSettings>()->GetColor(PCGNode->GetSettings());
+			}
+
+			if (SettingsColor != FLinearColor::White)
+			{
+				return SettingsColor;
+			}
+		}
+	}
+
+	return GetDefault<UPCGEditorSettings>()->DefaultNodeColor;
+}
+
+FLinearColor UPCGEditorGraphNodeBase::GetNodeBodyTintColor() const
+{
+	if (PCGNode)
+	{
+		const UPCGSettingsInterface* PCGSettingsInterface = PCGNode->GetSettingsInterface();
+		if (PCGSettingsInterface)
+		{
+			if (PCGSettingsInterface->IsInstance())
+			{
+				return GetDefault<UPCGEditorSettings>()->InstancedNodeBodyTintColor;
+			}
+		}
+	}
+
+	return Super::GetNodeBodyTintColor();
+}
+
+FEdGraphPinType UPCGEditorGraphNodeBase::GetPinType(const UPCGPin* InPin)
+{
+	FEdGraphPinType EdPinType;
+	EdPinType.ResetToDefaults();
+	EdPinType.PinCategory = NAME_None;
+	EdPinType.PinSubCategory = NAME_None;
+	EdPinType.ContainerType = EPinContainerType::None;
+
+	check(InPin);
+	const EPCGDataType PinType = InPin->GetCurrentTypes();
+
+	auto CheckType = [PinType](EPCGDataType AllowedType)
+	{
+		return !!(PinType & AllowedType) && !(PinType & ~AllowedType);
+	};
+
+	if (CheckType(EPCGDataType::Concrete))
+	{
+		EdPinType.PinCategory = FPCGEditorCommon::ConcreteDataType;
+
+		// Assign subcategory if we have precise information
+		if (CheckType(EPCGDataType::Point))
+		{
+			EdPinType.PinSubCategory = FPCGEditorCommon::PointDataType;
+		}
+		else if (CheckType(EPCGDataType::PolyLine))
+		{
+			EdPinType.PinSubCategory = FPCGEditorCommon::PolyLineDataType;
+		}
+		else if (CheckType(EPCGDataType::Landscape))
+		{
+			EdPinType.PinSubCategory = FPCGEditorCommon::LandscapeDataType;
+		}
+		else if (CheckType(EPCGDataType::BaseTexture))
+		{
+			EdPinType.PinSubCategory = FPCGEditorCommon::BaseTextureDataType;
+		}
+		else if (CheckType(EPCGDataType::Texture))
+		{
+			EdPinType.PinSubCategory = FPCGEditorCommon::TextureDataType;
+		}
+		else if (CheckType(EPCGDataType::RenderTarget))
+		{
+			EdPinType.PinSubCategory = FPCGEditorCommon::RenderTargetDataType;
+		}
+		else if (CheckType(EPCGDataType::Surface))
+		{
+			EdPinType.PinSubCategory = FPCGEditorCommon::SurfaceDataType;
+		}
+		else if (CheckType(EPCGDataType::Volume))
+		{
+			EdPinType.PinSubCategory = FPCGEditorCommon::VolumeDataType;
+		}
+		else if (CheckType(EPCGDataType::DynamicMesh))
+		{
+			EdPinType.PinSubCategory = FPCGEditorCommon::DynamicMeshDataType;
+		}
+		else if (CheckType(EPCGDataType::Primitive))
+		{
+			EdPinType.PinSubCategory = FPCGEditorCommon::PrimitiveDataType;
+		}
+	}
+	else if (CheckType(EPCGDataType::Spatial))
+	{
+		EdPinType.PinCategory = FPCGEditorCommon::SpatialDataType;
+	}
+	else if (CheckType(EPCGDataType::Param))
+	{
+		EdPinType.PinCategory = FPCGEditorCommon::ParamDataType;
+	}
+	else if (CheckType(EPCGDataType::Settings))
+	{
+		EdPinType.PinCategory = FPCGEditorCommon::SettingsDataType;
+	}
+	else if (CheckType(EPCGDataType::Other))
+	{
+		EdPinType.PinCategory = FPCGEditorCommon::OtherDataType;
+	}
+
+	return EdPinType;
+}
+
+FText UPCGEditorGraphNodeBase::GetTooltipText() const
+{
+	// Either use specified tooltip for description, or fall back to node name if none given.
+	const FText Description = (PCGNode && !PCGNode->GetNodeTooltipText().IsEmpty()) ? PCGNode->GetNodeTooltipText() : GetNodeTitle(ENodeTitleType::FullTitle);
+
+	return FText::Format(LOCTEXT("NodeTooltip", "{0}\n\n{1} - Node index {2}"),
+		Description,
+		PCGNode ? FText::FromName(PCGNode->GetFName()) : LOCTEXT("InvalidNodeName", "Unbound node"),
+		PCGNode ? FText::AsNumber(PCGNode->GetGraph()->GetNodes().IndexOfByKey(PCGNode)) : LOCTEXT("InvalidNodeIndex", "Invalid index"));
+}
+
+void UPCGEditorGraphNodeBase::GetPinHoverText(const UEdGraphPin& Pin, FString& HoverTextOut) const
+{
+	const bool bIsInputPin = (Pin.Direction == EGPD_Input);
+	UPCGPin* MatchingPin = (PCGNode ? (bIsInputPin ? PCGNode->GetInputPin(Pin.PinName) : PCGNode->GetOutputPin(Pin.PinName)) : nullptr);
+
+	auto PCGDataTypeToText = [](EPCGDataType DataType)
+	{
+		TArray<FText> BitFlags;
+
+		for (uint64 BitIndex = 1; BitIndex < 8 * sizeof(EPCGDataType); ++BitIndex)
+		{
+			const int64 BitValue = static_cast<__underlying_type(EPCGDataType)>(DataType) & (1 << BitIndex);
+			if (BitValue)
+			{
+				BitFlags.Add(StaticEnum<EPCGDataType>()->GetDisplayNameTextByValue(BitValue));
+			}
+		}
+
+		return FText::Join(FText(LOCTEXT("Delimiter", " | ")), BitFlags);
+	};
+
+	auto PinTypeToText = [&PCGDataTypeToText](const FName& Category, UPCGPin* MatchingPin)
+	{
+		if (Category != NAME_None)
+		{
+			return FText::FromName(Category);
+		}
+		else if (MatchingPin && MatchingPin->GetCurrentTypes() == EPCGDataType::Any)
+		{
+			return FText::FromName(FName("Any"));
+		}
+		else if (MatchingPin)
+		{
+			return PCGDataTypeToText(MatchingPin->GetCurrentTypes());
+		}
+		else
+		{
+			return FText(LOCTEXT("Unknown data type", "Unknown data type"));
+		}
+	};
+
+	const FText DataTypeText = PinTypeToText(Pin.PinType.PinCategory, MatchingPin);
+	const FText DataSubtypeText = PinTypeToText(Pin.PinType.PinSubCategory, MatchingPin);
+
+	FText Description;
+	if (MatchingPin)
+	{
+		Description = MatchingPin->Properties.Tooltip.IsEmpty() ? FText::FromName(MatchingPin->Properties.Label) : MatchingPin->Properties.Tooltip;
+	}
+
+	FText Required;
+	FText MultiDataSupport;
+	FText MultiConnectionSupport;	
+
+	if (MatchingPin)
+	{
+		if (bIsInputPin)
+		{
+			if (PCGNode && PCGNode->IsInputPinRequiredByExecution(MatchingPin))
+			{
+				Required = LOCTEXT("InputIsRequired", "Required input. ");
+			}
+
+			MultiDataSupport = MatchingPin->Properties.bAllowMultipleData ? LOCTEXT("InputSupportsMultiData", "Supports multiple data in input(s). ") : LOCTEXT("InputSingleDataOnly", "Supports only single data in input(s). ");
+
+			MultiConnectionSupport = MatchingPin->Properties.AllowsMultipleConnections() ? LOCTEXT("SupportsMultiInput", "Supports multiple inputs.") : LOCTEXT("SingleInputOnly", "Supports only one input.");
+		}
+		else
+		{
+			MultiDataSupport = MatchingPin->Properties.bAllowMultipleData ? LOCTEXT("OutputSupportsMultiData", "Can generate multiple data.") : LOCTEXT("OutputSingleDataOnly", "Generates only single data.");
+		}
+	}
+
+	HoverTextOut = FText::Format(LOCTEXT("PinHoverToolTipFull", "{0}\n\nType: {1}\nSubtype: {2}\nAdditional information: {3}{4}{5}"),
+		Description,
+		DataTypeText,
+		DataSubtypeText,
+		Required,
+		MultiDataSupport,
+		MultiConnectionSupport).ToString();
+}
+
+UObject* UPCGEditorGraphNodeBase::GetJumpTargetForDoubleClick() const
+{
+	if (PCGNode)
+	{
+		if (UPCGSettings* Settings = PCGNode->GetSettings())
+		{
+			return Settings->GetJumpTargetForDoubleClick();
+		}
+	}
+
+	return nullptr;
+}
+
+void UPCGEditorGraphNodeBase::OnUpdateCommentText(const FString& NewComment)
+{
+	Super::OnUpdateCommentText(NewComment);
+
+	if (PCGNode && PCGNode->NodeComment != NewComment)
+	{
+		PCGNode->Modify();
+		PCGNode->NodeComment = NewComment;
+	}
+}
+
+void UPCGEditorGraphNodeBase::OnCommentBubbleToggled(bool bInCommentBubbleVisible)
+{
+	Super::OnCommentBubbleToggled(bInCommentBubbleVisible);
+
+	if (PCGNode && PCGNode->bCommentBubbleVisible != bInCommentBubbleVisible)
+	{
+		PCGNode->Modify();
+		PCGNode->bCommentBubbleVisible = bInCommentBubbleVisible;
+	}
+}
+
+void UPCGEditorGraphNodeBase::OnUserAddDynamicInputPin()
+{
+	check(PCGNode);
+	
+	if (UPCGSettingsWithDynamicInputs* DynamicNodeSettings = Cast<UPCGSettingsWithDynamicInputs>(PCGNode->GetSettings()))
+	{
+		const FScopedTransaction Transaction(*FPCGEditorCommon::ContextIdentifier, LOCTEXT("PCGEditorUserAddDynamicInputPin", "Add Source Pin"), DynamicNodeSettings);
+		DynamicNodeSettings->Modify();
+		DynamicNodeSettings->OnUserAddDynamicInputPin();
+	}
+}
+
+bool UPCGEditorGraphNodeBase::CanUserRemoveDynamicInputPin(UEdGraphPin* InPinToRemove) const
+{
+	check(PCGNode && InPinToRemove);
+
+	if (UPCGSettingsWithDynamicInputs* DynamicNodeSettings = Cast<UPCGSettingsWithDynamicInputs>(PCGNode->GetSettings()))
+	{
+		const UPCGEditorGraphNodeBase* PCGGraphNode = CastChecked<const UPCGEditorGraphNodeBase>(InPinToRemove->GetOwningNode());
+		return DynamicNodeSettings->CanUserRemoveDynamicInputPin(PCGGraphNode->GetPinIndex(InPinToRemove));
+	}
+	
+	return false;
+}
+
+void UPCGEditorGraphNodeBase::OnUserRemoveDynamicInputPin(UEdGraphPin* InRemovedPin)
+{
+	check(PCGNode && InRemovedPin);
+
+	if (UPCGSettingsWithDynamicInputs* DynamicNodeSettings = Cast<UPCGSettingsWithDynamicInputs>(PCGNode->GetSettings()))
+	{
+		if (UPCGEditorGraphNodeBase* PCGGraphNode = Cast<UPCGEditorGraphNodeBase>(InRemovedPin->GetOwningNode()))
+		{
+			const FScopedTransaction Transaction(*FPCGEditorCommon::ContextIdentifier, LOCTEXT("PCGEditorUserRemoveDynamicInputPin", "Remove Source Pin"), DynamicNodeSettings);
+			DynamicNodeSettings->Modify();
+			DynamicNodeSettings->OnUserRemoveDynamicInputPin(PCGGraphNode->GetPCGNode(), PCGGraphNode->GetPinIndex(InRemovedPin));
+		}
+	}
+}
+
+void UPCGEditorGraphNodeBase::OnConvertNode(const FPCGPreconfiguredInfo& ConversionInfo)
+{
+	check(PCGNode && PCGNode->GetSettings());
+
+	if (UPCGEditorGraph* EditorGraph = Cast<UPCGEditorGraph>(GetGraph()))
+	{
+		if (UPCGGraph* PCGGraph = EditorGraph->GetPCGGraph())
+		{
+			FText ConversionTransactionName = FText::Format(INVTEXT("{0}: {1}"), LOCTEXT("PCGEditorConvertNode", "Convert Node"), PCGNode->GetDefaultTitle());
+			FScopedTransaction Transaction(*FPCGEditorCommon::ContextIdentifier, std::move(ConversionTransactionName), PCGGraph);
+
+			if (!PCGNode->GetSettings()->ConvertNode(ConversionInfo))
+			{
+				// TODO: It would be useful to have an error/feedback process for it can not be converted.
+				Transaction.Cancel();
+				return;
+			}
+
+			EditorGraph->ReconstructGraph();
+		}
+	}
+}
+
+void UPCGEditorGraphNodeBase::UpdateCommentBubblePinned()
+{
+	if (PCGNode)
+	{
+		PCGNode->Modify();
+		PCGNode->bCommentBubblePinned = bCommentBubblePinned;
+	}
+}
+
+void UPCGEditorGraphNodeBase::UpdatePosition()
+{
+	if (PCGNode)
+	{
+		PCGNode->Modify();
+		PCGNode->PositionX = NodePosX;
+		PCGNode->PositionY = NodePosY;
+	}
+}
+
+bool UPCGEditorGraphNodeBase::ShouldCreatePin(const UPCGPin* InPin) const
+{
+	return InPin && !InPin->Properties.bInvisiblePin;
+}
+
+void UPCGEditorGraphNodeBase::CreatePins(const TArray<UPCGPin*>& InInputPins, const TArray<UPCGPin*>& InOutputPins)
+{
+	bool bHasAdvancedPin = false;
+
+	for (const UPCGPin* InputPin : InInputPins)
+	{
+		if (!ShouldCreatePin(InputPin))
+		{
+			continue;
+		}
+
+		UEdGraphPin* Pin = CreatePin(EEdGraphPinDirection::EGPD_Input, GetPinType(InputPin), InputPin->Properties.Label);
+		Pin->PinFriendlyName = GetPinFriendlyName(InputPin);
+		Pin->bAdvancedView = InputPin->Properties.IsAdvancedPin();
+		bHasAdvancedPin |= Pin->bAdvancedView;
+	}
+
+	for (const UPCGPin* OutputPin : InOutputPins)
+	{
+		if (!ShouldCreatePin(OutputPin))
+		{
+			continue;
+		}
+
+		UEdGraphPin* Pin = CreatePin(EEdGraphPinDirection::EGPD_Output, GetPinType(OutputPin), OutputPin->Properties.Label);
+		Pin->PinFriendlyName = GetPinFriendlyName(OutputPin);
+		Pin->bAdvancedView = OutputPin->Properties.IsAdvancedPin();
+		bHasAdvancedPin |= Pin->bAdvancedView;
+	}
+
+	if (bHasAdvancedPin && AdvancedPinDisplay == ENodeAdvancedPins::NoPins)
+	{
+		AdvancedPinDisplay = ENodeAdvancedPins::Hidden;
+	}
+	else if (!bHasAdvancedPin)
+	{
+		AdvancedPinDisplay = ENodeAdvancedPins::NoPins;
+	}
+}
+
+FText UPCGEditorGraphNodeBase::GetPinFriendlyName(const UPCGPin* InPin) const
+{
+	check(InPin);
+
+	// For overridable params, use the display name of properties (for localized version or overridden display name in metadata).
+	if (InPin->Properties.IsAdvancedPin() && InPin->Properties.AllowedTypes == EPCGDataType::Param)
+	{
+		const UPCGSettings* Settings = InPin->Node ? InPin->Node->GetSettings() : nullptr;
+		if (Settings)
+		{
+			const FPCGSettingsOverridableParam* Param = Settings->OverridableParams().FindByPredicate([Label = InPin->Properties.Label](const FPCGSettingsOverridableParam& ParamToCheck)
+			{
+				return ParamToCheck.Label == Label;
+			});
+
+			if (Param)
+			{
+				return Param->GetDisplayPropertyPathText();
+			}
+		}
+	}
+
+	return FText::FromString(FName::NameToDisplayString(InPin->Properties.Label.ToString(), /*bIsBool=*/false));
+}
+
+#undef LOCTEXT_NAMESPACE
diff --git a/Engine/Plugins/PCG/Source/PCGEditor/Private/Nodes/PCGEditorGraphNodeBase.h b/Engine/Plugins/PCG/Source/PCGEditor/Private/Nodes/PCGEditorGraphNodeBase.h
new file mode 100644
--- /dev/null
+++ b/Engine/Plugins/PCG/Source/PCGEditor/Private/Nodes/PCGEditorGraphNodeBase.h
@@ -0,0 +1,220 @@
+// Copyright Epic Games, Inc. All Rights Reserved.
+
+#pragma once
+
+#include "EdGraph/EdGraphNode.h"
+#include "EdGraph/EdGraphPin.h"
+
+#include "PCGCommon.h"
+#include "Utils/PCGPreconfiguration.h"
+
+#include "PCGEditorGraphNodeBase.generated.h"
+
+struct FPCGPreconfiguredInfo;
+enum class EPCGChangeType : uint32;
+
+class UPCGComponent;
+class UPCGGraph;
+class UPCGNode;
+class UPCGPin;
+class UToolMenu;
+struct FPCGStack;
+
+UCLASS()
+class UPCGEditorGraphNodeBase : public UEdGraphNode
+{
+	GENERATED_BODY()
+
+public:
+	void Construct(UPCGNode* InPCGNode);
+
+	// ~Begin UObject interface
+	virtual void BeginDestroy() override;
+	virtual void PostTransacted(const FTransactionObjectEvent& TransactionEvent) override;
+	// ~End UObject interface
+
+	// ~Begin UEdGraphNode interface
+	virtual void GetNodeContextMenuActions(UToolMenu* Menu, class UGraphNodeContextMenuContext* Context) const override;
+	virtual void AutowireNewNode(UEdGraphPin* FromPin) override;
+	virtual void PrepareForCopying() override;
+	virtual bool CanCreateUnderSpecifiedSchema(const UEdGraphSchema* Schema) const override;
+	virtual void ReconstructNode() override;
+	virtual FLinearColor GetNodeTitleColor() const override;
+	virtual FLinearColor GetNodeBodyTintColor() const override;
+	virtual void PostPasteNode() override;
+	virtual FText GetTooltipText() const override;
+	virtual void GetPinHoverText(const UEdGraphPin& Pin, FString& HoverTextOut) const override;
+	virtual UObject* GetJumpTargetForDoubleClick() const override;
+	virtual void OnUpdateCommentText(const FString& NewComment) override;
+	virtual void OnCommentBubbleToggled(bool bInCommentBubbleVisible) override;
+	virtual bool GetCanRenameNode() const override { return bCanRenameNode; }
+	// ~End UEdGraphNode interface
+
+	void OnUserAddDynamicInputPin();
+	bool CanUserRemoveDynamicInputPin(UEdGraphPin* InPinToRemove) const;
+	void OnUserRemoveDynamicInputPin(UEdGraphPin* InRemovedPin);
+
+	void OnConvertNode(const FPCGPreconfiguredInfo& ConversionInfo);
+
+	UPCGNode* GetPCGNode() { return PCGNode; }
+	const UPCGNode* GetPCGNode() const { return PCGNode; }
+	void PostCopy();
+	void RebuildAfterPaste();
+	virtual void PostPaste();
+
+	void SetInspected(bool InIsInspecting) { bIsInspected = InIsInspecting; }
+	bool GetInspected() const { return bIsInspected; }
+
+	/** Whether node was culled either during compilation or at execution time. */
+	void SetIsCulledFromExecution(bool bInIsCulledFromExecution) { bIsCulledFromExecution = bInIsCulledFromExecution; }
+	bool IsCulledFromExecution() const { return bIsCulledFromExecution; }
+
+	/** Increase deferred reconstruct counter, calls to ReconstructNode will flag reconstruct to happen when count hits zero */
+	void EnableDeferredReconstruct();
+	/** Decrease deferred reconstruct counter, ReconstructNode will be called if counter hits zero and the node is flagged for reconstruction  */
+	void DisableDeferredReconstruct();
+
+	/** Pulls current errors/warnings state from PCG subsystem. */
+	EPCGChangeType UpdateErrorsAndWarnings();
+
+	/** If the currently inspected grid size is smaller than the grid size of this node, display transparent. */
+	EPCGChangeType UpdateStructuralVisualization(class UPCGComponent* InComponentBeingDebugged, const FPCGStack* InStackBeingInspected, bool bNewlyPlaced = false);
+
+	/** Puts node title on node body, reducing overall node size */
+	bool ShouldDrawCompact() const;
+
+	/** Returns custom compact node icon if available */
+	bool GetCompactNodeIcon(FName& OutCompactNodeIcon) const;
+
+	bool CanUserAddRemoveDynamicInputPins() const;
+
+	/** Whether to flip the order of the title lines - display generated title first and authored second. */
+	bool HasFlippedTitleLines() const;
+
+	/** Authored part of node title (like "Create Attribute X"). */
+	FText GetAuthoredTitleLine() const;
+
+	/** Generated part of node title, not user editable (like "X = 5.0"). */
+	FText GetGeneratedTitleLine() const;
+
+	/** Bitmask of inactive output pins. Bit N will be set if output pin index N is inactive. */
+	uint64 GetInactiveOutputPinMask() const { return InactiveOutputPinMask; }
+
+	/** The grid that this node executes on if higen is enabled, otherwise Unitialized. */
+	EPCGHiGenGrid GetGenerationGrid() const { return GenerationGrid; }
+
+	/** The higen grid currently being inspected if any, otherwise Uninitialized. */
+	EPCGHiGenGrid GetInspectedGenerationGrid() const { return InspectedGenerationGrid; }
+
+	/** Whether the given output pin was active in the previous execution. */
+	bool IsOutputPinActive(const UEdGraphPin* InOutputPin) const;
+
+	/** Marks the node as re-nameable and provokes the node to update, placing the user in rename mode. */
+	void EnterRenamingMode();
+
+	/** Disables renaming, preventing the node from entering editing mode on subsequent updates. */
+	void ExitRenamingMode();
+
+	/** Whether this OutputPin can be connected to this InputPin */
+	virtual bool IsCompatible(const UPCGPin* InputPin, const UPCGPin* OutputPin, FText& OutReason) const;
+
+	DECLARE_DELEGATE(FOnPCGEditorGraphNodeChanged);
+	FOnPCGEditorGraphNodeChanged OnNodeChangedDelegate;
+
+	DECLARE_DELEGATE(FOnPCGEditorGraphNodeRenameInitiated);
+	FOnPCGEditorGraphNodeRenameInitiated OnNodeRenameInitiatedDelegate;
+
+protected:
+	static FEdGraphPinType GetPinType(const UPCGPin* InPin);
+
+	/** Create PCG-side edges from editor pins/edges. */
+	void RebuildEdgesFromPins();
+	virtual void RebuildEdgesFromPins_Internal();
+
+	void OnNodeChanged(UPCGNode* InNode, EPCGChangeType ChangeType);
+	void OnPickColor();
+	virtual bool CanPickColor() const { return true; }
+	virtual void OnColorPicked(FLinearColor NewColor);
+	void UpdateCommentBubblePinned();
+	void UpdatePosition();
+
+	void CreatePins(const TArray<UPCGPin*>& InInputPins, const TArray<UPCGPin*>& InOutputPins);
+
+	// Performs potentially custom logic when there's a change that would trigger a reconstruct (needed for linked nodes like the named reroutes)
+	virtual void ReconstructNodeOnChange() { ReconstructNode(); }
+
+	// Custom logic to hide some pins to the user (by not creating a UI pin, even if the model pin exists).
+	// Useful for deprecation
+	virtual bool ShouldCreatePin(const UPCGPin* InPin) const;
+
+	// Returns the appropriate pin name to allow for some flexibility
+	virtual FText GetPinFriendlyName(const UPCGPin* InPin) const;
+
+	UPROPERTY()
+	TObjectPtr<UPCGNode> PCGNode = nullptr;
+
+	int32 DeferredReconstructCounter = 0;
+	bool bDeferredReconstruct = false;
+	bool bDisableReconstructFromNode = false;
+	bool bIsInspected = false;
+
+	/** Whether this node was culled in the last execution. */
+	bool bIsCulledFromExecution = false;
+
+	/** Bitmask of inactive output pins. Bit N will be set if output pin index N is inactive. */
+	uint64 InactiveOutputPinMask = 0;
+
+	/** The grid that this node executes on if higen is enabled, otherwise Unitialized. */
+	EPCGHiGenGrid GenerationGrid = EPCGHiGenGrid::Uninitialized;
+
+	/** The higen grid currently being inspected if any, otherwise Uninitialized. */
+	EPCGHiGenGrid InspectedGenerationGrid = EPCGHiGenGrid::Uninitialized;
+};
+
+/** Disables reconstruct on nodes (or from a pin) and re-enables in destructor. */
+struct FPCGDeferNodeReconstructScope
+{
+	explicit FPCGDeferNodeReconstructScope(UEdGraphPin* FromPin)
+		: Node(FromPin ? Cast<UPCGEditorGraphNodeBase>(FromPin->GetOwningNode()) : nullptr)
+	{
+		if (Node)
+		{
+			Node->EnableDeferredReconstruct();
+		}
+	}
+
+	explicit FPCGDeferNodeReconstructScope(UPCGEditorGraphNodeBase* InNode)
+		: Node(InNode)
+	{
+		if (Node)
+		{
+			Node->EnableDeferredReconstruct();
+		}
+	}
+
+	FPCGDeferNodeReconstructScope(const FPCGDeferNodeReconstructScope&) = delete;
+	FPCGDeferNodeReconstructScope(FPCGDeferNodeReconstructScope&& Other)
+	{
+		Node = Other.Node;
+		Other.Node = nullptr;
+	}
+
+	FPCGDeferNodeReconstructScope& operator=(const FPCGDeferNodeReconstructScope&) = delete;
+	FPCGDeferNodeReconstructScope& operator=(FPCGDeferNodeReconstructScope&& Other)
+	{
+		Swap(Node, Other.Node);
+		return *this;
+	}
+
+	~FPCGDeferNodeReconstructScope()
+	{
+		if (Node)
+		{
+			Node->DisableDeferredReconstruct();
+			Node = nullptr;
+		}
+	}
+
+private:
+	UPCGEditorGraphNodeBase* Node = nullptr;
+};
\ No newline at end of file
diff --git a/Engine/Plugins/PCG/Source/PCGEditor/Private/Nodes/PCGEditorGraphNodeFactory.cpp b/Engine/Plugins/PCG/Source/PCGEditor/Private/Nodes/PCGEditorGraphNodeFactory.cpp
new file mode 100644
--- /dev/null
+++ b/Engine/Plugins/PCG/Source/PCGEditor/Private/Nodes/PCGEditorGraphNodeFactory.cpp
@@ -0,0 +1,65 @@
+// Copyright Epic Games, Inc. All Rights Reserved.
+
+#include "PCGEditorGraphNodeFactory.h"
+
+#include "PCGEditorGraphNodeBase.h"
+#include "PCGEditorGraphNodeReroute.h"
+#include "PCGNode.h"
+#include "SPCGEditorGraphNode.h"
+#include "SPCGEditorGraphVarNode.h"
+#include "Elements/PCGReroute.h"
+#include "Elements/PCGUserParameterGet.h"
+
+#include "SGraphNodeKnot.h"
+
+class SPCGEditorGraphNodeKnot : public SGraphNodeKnot
+{
+public:
+	SLATE_BEGIN_ARGS(SPCGEditorGraphNodeKnot) {}
+	SLATE_END_ARGS();
+
+	void Construct(const FArguments& InArgs, UPCGEditorGraphNodeBase* InPCGGraphNode)
+	{
+		SGraphNodeKnot::Construct(SGraphNodeKnot::FArguments(), InPCGGraphNode);
+		InPCGGraphNode->OnNodeChangedDelegate.BindSP(this, &SPCGEditorGraphNodeKnot::OnNodeChanged);
+	}
+
+private:
+	void OnNodeChanged()
+	{
+		UpdateGraphNode();
+	}
+};
+
+
+TSharedPtr<SGraphNode> FPCGEditorGraphNodeFactory::CreateNode(UEdGraphNode* InNode) const
+{
+	if (UPCGEditorGraphNodeBase* GraphNode = Cast<UPCGEditorGraphNodeBase>(InNode))
+	{
+		TSharedPtr<SGraphNode> VisualNode;
+		
+		const UPCGNode* PCGNode = GraphNode->GetPCGNode();
+		if (PCGNode && Cast<UPCGNamedRerouteBaseSettings>(PCGNode->GetSettings()))
+		{
+			SAssignNew(VisualNode, SPCGEditorGraphVarNode, GraphNode);
+		}
+		else if (PCGNode && Cast<UPCGRerouteSettings>(PCGNode->GetSettings()))
+		{
+			SAssignNew(VisualNode, SPCGEditorGraphNodeKnot, GraphNode);
+		}
+		else if (PCGNode && Cast<UPCGUserParameterGetSettings>(PCGNode->GetSettings()))
+		{
+			SAssignNew(VisualNode, SPCGEditorGraphVarNode, GraphNode);
+		}
+		else
+		{
+			SAssignNew(VisualNode, SPCGEditorGraphNode, GraphNode);
+		}
+		
+		VisualNode->SlatePrepass();
+
+		return VisualNode.ToSharedRef();
+	}
+
+	return nullptr;
+}
diff --git a/Engine/Plugins/PCG/Source/PCGEditor/Private/Nodes/PCGEditorGraphNodeFactory.h b/Engine/Plugins/PCG/Source/PCGEditor/Private/Nodes/PCGEditorGraphNodeFactory.h
new file mode 100644
--- /dev/null
+++ b/Engine/Plugins/PCG/Source/PCGEditor/Private/Nodes/PCGEditorGraphNodeFactory.h
@@ -0,0 +1,13 @@
+// Copyright Epic Games, Inc. All Rights Reserved.
+
+#pragma once
+
+#include "EdGraphUtilities.h"
+
+class SGraphNode;
+class UEdGraphNode;
+
+class FPCGEditorGraphNodeFactory : public FGraphPanelNodeFactory
+{
+	virtual TSharedPtr<SGraphNode> CreateNode(UEdGraphNode* InNode) const override;
+};
diff --git a/Engine/Plugins/PCG/Source/PCGEditor/Private/Nodes/PCGEditorGraphNodeInput.cpp b/Engine/Plugins/PCG/Source/PCGEditor/Private/Nodes/PCGEditorGraphNodeInput.cpp
new file mode 100644
--- /dev/null
+++ b/Engine/Plugins/PCG/Source/PCGEditor/Private/Nodes/PCGEditorGraphNodeInput.cpp
@@ -0,0 +1,29 @@
+// Copyright Epic Games, Inc. All Rights Reserved.
+
+#include "PCGEditorGraphNodeInput.h"
+
+#include "PCGNode.h"
+#include "PCGPin.h"
+
+#define LOCTEXT_NAMESPACE "PCGEditorGraphNodeInput"
+
+FText UPCGEditorGraphNodeInput::GetNodeTitle(ENodeTitleType::Type TitleType) const
+{	
+	return LOCTEXT("NodeTitle", "Input");
+}
+
+void UPCGEditorGraphNodeInput::AllocateDefaultPins()
+{
+	if (PCGNode)
+	{
+		CreatePins(/*InInputPins=*/{}, PCGNode->GetOutputPins());
+	}
+}
+
+void UPCGEditorGraphNodeInput::ReconstructNode()
+{
+	Super::ReconstructNode();
+	//TODO: Implement special version for input to avoid the enum type
+}
+
+#undef LOCTEXT_NAMESPACE
diff --git a/Engine/Plugins/PCG/Source/PCGEditor/Private/Nodes/PCGEditorGraphNodeInput.h b/Engine/Plugins/PCG/Source/PCGEditor/Private/Nodes/PCGEditorGraphNodeInput.h
new file mode 100644
--- /dev/null
+++ b/Engine/Plugins/PCG/Source/PCGEditor/Private/Nodes/PCGEditorGraphNodeInput.h
@@ -0,0 +1,24 @@
+// Copyright Epic Games, Inc. All Rights Reserved.
+
+#pragma once
+
+#include "PCGEditorGraphNodeBase.h"
+
+#include "PCGEditorGraphNodeInput.generated.h"
+
+namespace ENodeTitleType { enum Type : int; }
+
+UCLASS()
+class UPCGEditorGraphNodeInput : public UPCGEditorGraphNodeBase
+{
+	GENERATED_BODY()
+
+public:
+	// ~Begin UEdGraphNode interface
+	virtual FText GetNodeTitle(ENodeTitleType::Type TitleType) const override;
+	virtual void AllocateDefaultPins() override;
+	virtual bool CanUserDeleteNode() const override { return false; }
+	virtual bool CanDuplicateNode() const override { return false; }
+	virtual void ReconstructNode() override;
+	// ~End UEdGraphNode interface
+};
diff --git a/Engine/Plugins/PCG/Source/PCGEditor/Private/Nodes/PCGEditorGraphNodeOutput.cpp b/Engine/Plugins/PCG/Source/PCGEditor/Private/Nodes/PCGEditorGraphNodeOutput.cpp
new file mode 100644
--- /dev/null
+++ b/Engine/Plugins/PCG/Source/PCGEditor/Private/Nodes/PCGEditorGraphNodeOutput.cpp
@@ -0,0 +1,29 @@
+// Copyright Epic Games, Inc. All Rights Reserved.
+
+#include "PCGEditorGraphNodeOutput.h"
+
+#include "PCGNode.h"
+#include "PCGPin.h"
+
+#define LOCTEXT_NAMESPACE "PCGEditorGraphNodeOutput"
+
+FText UPCGEditorGraphNodeOutput::GetNodeTitle(ENodeTitleType::Type TitleType) const
+{
+	return LOCTEXT("NodeTitle", "Output");
+}
+
+void UPCGEditorGraphNodeOutput::AllocateDefaultPins()
+{
+	if (PCGNode)
+	{
+		CreatePins(PCGNode->GetInputPins(), /*InOutputPins=*/{});
+	}
+}
+
+void UPCGEditorGraphNodeOutput::ReconstructNode()
+{
+	Super::ReconstructNode();
+	//TODO: Implement special version for output to avoid the enum type
+}
+
+#undef LOCTEXT_NAMESPACE
diff --git a/Engine/Plugins/PCG/Source/PCGEditor/Private/Nodes/PCGEditorGraphNodeOutput.h b/Engine/Plugins/PCG/Source/PCGEditor/Private/Nodes/PCGEditorGraphNodeOutput.h
new file mode 100644
--- /dev/null
+++ b/Engine/Plugins/PCG/Source/PCGEditor/Private/Nodes/PCGEditorGraphNodeOutput.h
@@ -0,0 +1,24 @@
+// Copyright Epic Games, Inc. All Rights Reserved.
+
+#pragma once
+
+#include "PCGEditorGraphNodeBase.h"
+
+#include "PCGEditorGraphNodeOutput.generated.h"
+
+namespace ENodeTitleType { enum Type : int; }
+
+UCLASS()
+class UPCGEditorGraphNodeOutput : public UPCGEditorGraphNodeBase
+{
+	GENERATED_BODY()
+
+public:
+	// ~Begin UEdGraphNode interface
+	virtual FText GetNodeTitle(ENodeTitleType::Type TitleType) const override;
+	virtual void AllocateDefaultPins() override;
+	virtual bool CanUserDeleteNode() const override { return false; }
+	virtual bool CanDuplicateNode() const override { return false; }
+	virtual void ReconstructNode() override;
+	// ~End UEdGraphNode interface
+};
diff --git a/Engine/Plugins/PCG/Source/PCGEditor/Private/Nodes/PCGEditorGraphNodeReroute.cpp b/Engine/Plugins/PCG/Source/PCGEditor/Private/Nodes/PCGEditorGraphNodeReroute.cpp
new file mode 100644
--- /dev/null
+++ b/Engine/Plugins/PCG/Source/PCGEditor/Private/Nodes/PCGEditorGraphNodeReroute.cpp
@@ -0,0 +1,238 @@
+// Copyright Epic Games, Inc. All Rights Reserved.
+
+#include "PCGEditorGraphNodeReroute.h"
+
+#include "PCGEditorGraph.h"
+
+#include "PCGNode.h"
+#include "Elements/PCGReroute.h"
+
+#include "EdGraph/EdGraphPin.h"
+
+FText UPCGEditorGraphNodeReroute::GetNodeTitle(ENodeTitleType::Type TitleType) const
+{
+	return NSLOCTEXT("PCGEditorGraphNodeReroute", "NodeTitle", "Reroute");
+}
+
+bool UPCGEditorGraphNodeReroute::ShouldOverridePinNames() const
+{
+	return true;
+}
+
+FText UPCGEditorGraphNodeReroute::GetPinNameOverride(const UEdGraphPin& Pin) const
+{
+	return FText::GetEmpty();
+}
+
+bool UPCGEditorGraphNodeReroute::CanSplitPin(const UEdGraphPin* Pin) const
+{
+	return false;
+}
+
+bool UPCGEditorGraphNodeReroute::ShouldDrawNodeAsControlPointOnly(int32& OutInputPinIndex, int32& OutOutputPinIndex) const
+{
+	OutInputPinIndex = 0;
+	OutOutputPinIndex = 1;
+	return true;
+}
+
+FText UPCGEditorGraphNodeReroute::GetTooltipText() const
+{
+	return FText::GetEmpty();
+}
+
+UEdGraphPin* UPCGEditorGraphNodeReroute::GetPassThroughPin(const UEdGraphPin* FromPin) const
+{
+	if (FromPin == GetInputPin())
+	{
+		return GetOutputPin();
+	}
+	else
+	{
+		return GetInputPin();
+	}
+}
+
+UEdGraphPin* UPCGEditorGraphNodeReroute::GetInputPin() const
+{
+	for (UEdGraphPin* Pin : Pins)
+	{
+		if (Pin->Direction == EGPD_Input)
+		{
+			return Pin;
+		}
+	}
+
+	return nullptr;
+}
+
+UEdGraphPin* UPCGEditorGraphNodeReroute::GetOutputPin() const
+{
+	for (UEdGraphPin* Pin : Pins)
+	{
+		if (Pin->Direction == EGPD_Output)
+		{
+			return Pin;
+		}
+	}
+
+	return nullptr;
+}
+
+FText UPCGEditorGraphNodeNamedRerouteBase::GetNodeTitle(ENodeTitleType::Type TitleType) const
+{
+	// Never show full title, only list view
+	return Super::GetNodeTitle(ENodeTitleType::ListView);
+}
+
+FText UPCGEditorGraphNodeNamedRerouteUsage::GetPinFriendlyName(const UPCGPin* InPin) const
+{
+	return FText::FromString(" ");
+}
+
+void UPCGEditorGraphNodeNamedRerouteUsage::RebuildEdgesFromPins_Internal()
+{
+	Super::RebuildEdgesFromPins_Internal();
+
+	check(PCGNode);
+
+	if (PCGNode->HasInboundEdges())
+	{
+		return;
+	}
+
+	UPCGGraph* Graph = PCGNode->GetGraph();
+
+	if (!Graph)
+	{
+		return;
+	}
+
+	UPCGNamedRerouteUsageSettings* Usage = Cast<UPCGNamedRerouteUsageSettings>(PCGNode->GetSettings());
+
+	if (!Usage)
+	{
+		return;
+	}
+
+	// Make sure we're hooked to the declaration if it's not already the case
+	if (UPCGNode* DeclarationNode = Graph->FindNodeWithSettings(Usage->Declaration))
+	{
+		DeclarationNode->AddEdgeTo(PCGNamedRerouteConstants::InvisiblePinLabel, PCGNode, PCGPinConstants::DefaultInputLabel);
+	}
+}
+
+FText UPCGEditorGraphNodeNamedRerouteDeclaration::GetPinFriendlyName(const UPCGPin* InPin) const
+{
+	return FText::FromString(" ");
+}
+
+void UPCGEditorGraphNodeNamedRerouteDeclaration::ApplyToUsageNodes(TFunctionRef<void(UPCGEditorGraphNodeNamedRerouteUsage*)> Action)
+{
+	if (!GetPCGNode())
+	{
+		return;
+	}
+
+	const UPCGNamedRerouteDeclarationSettings* Declaration = Cast<UPCGNamedRerouteDeclarationSettings>(GetPCGNode()->GetSettings());
+
+	if (!Declaration)
+	{
+		return;
+	}
+
+	UPCGEditorGraph* EditorGraph = Cast<UPCGEditorGraph>(GetGraph());
+
+	if (!EditorGraph)
+	{
+		return;
+	}
+
+	for (const TObjectPtr<UEdGraphNode>& EdGraphNode : EditorGraph->Nodes)
+	{
+		if (UPCGEditorGraphNodeNamedRerouteUsage* RerouteNode = Cast<UPCGEditorGraphNodeNamedRerouteUsage>(EdGraphNode))
+		{
+			if (RerouteNode->GetPCGNode() &&
+				Cast<UPCGNamedRerouteUsageSettings>(RerouteNode->GetPCGNode()->GetSettings()) &&
+				Cast<UPCGNamedRerouteUsageSettings>(RerouteNode->GetPCGNode()->GetSettings())->Declaration == Declaration)
+			{
+				Action(RerouteNode);
+			}
+		}
+	}
+}
+
+void UPCGEditorGraphNodeNamedRerouteDeclaration::ReconstructNodeOnChange()
+{
+	Super::ReconstructNodeOnChange();
+
+	// We must make sure to trigger a notify node changed on all editor nodes that are usages of that declaration
+	ApplyToUsageNodes([](UPCGEditorGraphNodeNamedRerouteUsage* RerouteNode)
+	{
+		RerouteNode->ReconstructNodeOnChange();
+	});
+}
+
+FString UPCGEditorGraphNodeNamedRerouteDeclaration::GenerateNodeName(const UPCGNode* FromNode, FName FromPinName)
+{
+	if (FromNode)
+	{
+		return FromNode->GetNodeTitle(EPCGNodeTitleType::ListView).ToString() + " " + FromPinName.ToString();
+	}
+	else if (FromPinName != NAME_None)
+	{
+		return FromPinName.ToString();
+	}
+	else
+	{
+		return TEXT("Reroute");
+	}
+}
+
+void UPCGEditorGraphNodeNamedRerouteDeclaration::OnRenameNode(const FString& NewName)
+{
+	check(PCGNode);
+	if (TOptional<FName> ModifiedName = UPCGRerouteSettings::GetCollisionFreeNodeName(PCGNode->GetGraph(), FName(NewName)))
+	{
+		FString NameString = ModifiedName->ToString();
+		if (CanRenameNode(NameString))
+		{
+			CommitNodeRename(NameString);
+			// Propagate the name change to downstream usage nodes
+			ApplyToUsageNodes([&NameString](UPCGEditorGraphNodeNamedRerouteUsage* RerouteNode)
+			{
+				RerouteNode->bCanRenameNode = true;
+				RerouteNode->OnRenameNode(NameString);
+				RerouteNode->bCanRenameNode = false;
+			});
+
+			ReconstructNodeOnChange();
+		}
+	}
+}
+
+void UPCGEditorGraphNodeNamedRerouteDeclaration::OnColorPicked(FLinearColor NewColor)
+{
+	Super::OnColorPicked(NewColor);
+
+	// Propagate change to downstream usage nodes
+	ApplyToUsageNodes([&NewColor](UPCGEditorGraphNodeNamedRerouteUsage* RerouteNode)
+	{
+		RerouteNode->OnColorPicked(NewColor);
+	});
+}
+
+void UPCGEditorGraphNodeNamedRerouteDeclaration::PostPaste()
+{
+	Super::PostPaste();
+	FixNodeNameCollision();
+}
+
+void UPCGEditorGraphNodeNamedRerouteDeclaration::FixNodeNameCollision()
+{
+	const FString BaseName = GetNodeTitle(ENodeTitleType::ListView).ToString();
+
+	bCanRenameNode = true;
+	OnRenameNode(BaseName);
+	bCanRenameNode = false;
+}
\ No newline at end of file
diff --git a/Engine/Plugins/PCG/Source/PCGEditor/Private/Nodes/PCGEditorGraphNodeReroute.h b/Engine/Plugins/PCG/Source/PCGEditor/Private/Nodes/PCGEditorGraphNodeReroute.h
new file mode 100644
--- /dev/null
+++ b/Engine/Plugins/PCG/Source/PCGEditor/Private/Nodes/PCGEditorGraphNodeReroute.h
@@ -0,0 +1,72 @@
+// Copyright Epic Games, Inc. All Rights Reserved.
+
+#pragma once
+
+#include "PCGEditorGraphNode.h"
+
+class UPCGNode;
+
+#include "PCGEditorGraphNodeReroute.generated.h"
+
+UCLASS()
+class UPCGEditorGraphNodeReroute : public UPCGEditorGraphNode
+{
+	GENERATED_BODY()
+
+public:
+	// ~Begin UEdGraphNode interface
+	virtual FText GetNodeTitle(ENodeTitleType::Type TitleType) const override;
+	virtual bool ShouldOverridePinNames() const override;
+	virtual FText GetPinNameOverride(const UEdGraphPin& Pin) const override;
+	virtual bool CanSplitPin(const UEdGraphPin* Pin) const override;
+	virtual bool ShouldDrawNodeAsControlPointOnly(int32& OutInputPinIndex, int32& OutOutputPinIndex) const override;
+	virtual FText GetTooltipText() const override;
+	virtual UEdGraphPin* GetPassThroughPin(const UEdGraphPin* FromPin) const override;
+	// ~End UEdGraphNode interface
+
+	UEdGraphPin* GetInputPin() const;
+	UEdGraphPin* GetOutputPin() const;
+};
+
+UCLASS()
+class UPCGEditorGraphNodeNamedRerouteBase : public UPCGEditorGraphNode
+{
+	GENERATED_BODY()
+
+public:
+	// ~Begin UEdGraphNode interface
+	virtual FText GetNodeTitle(ENodeTitleType::Type TitleType) const override;
+	// ~End UEdGraphNode interface
+};
+
+UCLASS()
+class UPCGEditorGraphNodeNamedRerouteUsage : public UPCGEditorGraphNodeNamedRerouteBase
+{
+	GENERATED_BODY()
+	friend class UPCGEditorGraphNodeNamedRerouteDeclaration;
+
+protected:
+	virtual void RebuildEdgesFromPins_Internal();
+	virtual bool CanPickColor() const override { return false; }
+	virtual FText GetPinFriendlyName(const UPCGPin* InPin) const override;
+};
+
+UCLASS()
+class UPCGEditorGraphNodeNamedRerouteDeclaration : public UPCGEditorGraphNodeNamedRerouteBase
+{
+	GENERATED_BODY()
+
+public:
+	virtual void OnRenameNode(const FString& NewName) override;
+	virtual void PostPaste() override;
+
+	FString GenerateNodeName(const UPCGNode* FromNode, FName FromPinName);
+	void FixNodeNameCollision();
+
+protected:
+	virtual FText GetPinFriendlyName(const UPCGPin* InPin) const override;
+	virtual void OnColorPicked(FLinearColor NewColor) override;
+	virtual void ReconstructNodeOnChange() override;
+
+	void ApplyToUsageNodes(TFunctionRef<void(UPCGEditorGraphNodeNamedRerouteUsage*)> Action);
+};
\ No newline at end of file
diff --git a/Engine/Plugins/PCG/Source/PCGEditor/Private/Nodes/SPCGEditorGraphNode.cpp b/Engine/Plugins/PCG/Source/PCGEditor/Private/Nodes/SPCGEditorGraphNode.cpp
new file mode 100644
--- /dev/null
+++ b/Engine/Plugins/PCG/Source/PCGEditor/Private/Nodes/SPCGEditorGraphNode.cpp
@@ -0,0 +1,1074 @@
+// Copyright Epic Games, Inc. All Rights Reserved.
+
+#include "SPCGEditorGraphNode.h"
+
+#include "PCGNode.h"
+#include "PCGPin.h"
+#include "PCGSettings.h"
+#include "PCGSettingsWithDynamicInputs.h"
+
+#include "PCGEditorGraphNodeBase.h"
+#include "PCGEditorStyle.h"
+
+#include "GraphEditorSettings.h"
+#include "IDocumentation.h"
+#include "SCommentBubble.h"
+#include "SGraphPin.h"
+#include "SLevelOfDetailBranchNode.h"
+#include "SPinTypeSelector.h"
+#include "TutorialMetaData.h"
+#include "Algo/AnyOf.h"
+#include "Widgets/Images/SImage.h"
+#include "Widgets/SToolTip.h"
+#include "Widgets/Input/SButton.h"
+#include "Widgets/Layout/SBox.h"
+#include "Widgets/Text/SInlineEditableTextBlock.h"
+#include "Widgets/SBoxPanel.h"
+
+#define LOCTEXT_NAMESPACE "SPCGEditorGraphNode"
+
+/** PCG pin primarily to give more control over pin coloring. */
+class SPCGEditorGraphNodePin : public SGraphPin
+{
+public:
+	SLATE_BEGIN_ARGS(SPCGEditorGraphNodePin)
+		: _PinLabelStyle(NAME_DefaultPinLabelStyle)
+		, _UsePinColorForText(false)
+		, _SideToSideMargin(5.0f)
+		{}
+		SLATE_ARGUMENT(FName, PinLabelStyle)
+		SLATE_ARGUMENT(bool, UsePinColorForText)
+		SLATE_ARGUMENT(float, SideToSideMargin)
+	SLATE_END_ARGS()
+
+	void Construct(const FArguments& InArgs, UEdGraphPin* InPin);
+
+	virtual FSlateColor GetPinColor() const override;
+	virtual FSlateColor GetPinTextColor() const override;
+	FName GetLabelStyle(FName DefaultLabelStyle) const;
+	bool GetExtraIcon(FName& OutExtraIcon, FText& OutTooltip) const;
+
+	/** Whether pin is required to be connected for execution. */
+	bool ShouldDisplayAsRequiredForExecution() const;
+
+private:
+	void ApplyUnusedPinStyle(FSlateColor& InOutColor) const;
+	void GetPCGNodeAndPin(const UPCGNode*& OutNode, const UPCGPin*& OutPin) const;
+};
+
+void SPCGEditorGraphNodePin::Construct(const FArguments& InArgs, UEdGraphPin* InPin)
+{
+	// IMPLEMENTATION NOTE: this is the code from SGraphPin::Construct with additional padding exposed,
+	// with an optional extra icon shown before the pin label, and with a marker icon to show pins
+	// that are required for execution.
+
+	bUsePinColorForText = InArgs._UsePinColorForText;
+	this->SetCursor(EMouseCursor::Default);
+
+	SetVisibility(MakeAttributeSP(this, &SPCGEditorGraphNodePin::GetPinVisiblity));
+
+	GraphPinObj = InPin;
+	check(GraphPinObj != NULL);
+
+	const UEdGraphSchema* Schema = GraphPinObj->GetSchema();
+	checkf(
+		Schema, 
+		TEXT("Missing schema for pin: %s with outer: %s of type %s"), 
+		*(GraphPinObj->GetName()),
+		GraphPinObj->GetOuter() ? *(GraphPinObj->GetOuter()->GetName()) : TEXT("NULL OUTER"), 
+		GraphPinObj->GetOuter() ? *(GraphPinObj->GetOuter()->GetClass()->GetName()) : TEXT("NULL OUTER")
+	);
+
+	const bool bIsInput = (GetDirection() == EGPD_Input);
+
+	// A small marker to indicate the pin is required for the node to be executed.
+	TSharedPtr<SImage> RequiredPinIconWidget;
+	float RequiredPinMarkerWidth = 0.0f;
+	bool bDisplayPinMarker = false;
+	if (bIsInput)
+	{
+		const FSlateBrush* RequiredPinMarkerIcon = FPCGEditorStyle::Get().GetBrush(PCGEditorStyleConstants::Pin_Required);
+		RequiredPinMarkerWidth = RequiredPinMarkerIcon ? RequiredPinMarkerIcon->GetImageSize().X : 8.0f;
+		bDisplayPinMarker = ShouldDisplayAsRequiredForExecution();
+
+		RequiredPinIconWidget =
+			SNew(SImage)
+			.Image(bDisplayPinMarker ? RequiredPinMarkerIcon : FAppStyle::GetNoBrush())
+			.ColorAndOpacity(this, &SPCGEditorGraphNodePin::GetPinColor);
+	}
+
+	// Create the pin icon widget
+	TSharedRef<SWidget> PinWidgetRef = SPinTypeSelector::ConstructPinTypeImage(
+		MakeAttributeSP(this, &SPCGEditorGraphNodePin::GetPinIcon ),
+		MakeAttributeSP(this, &SPCGEditorGraphNodePin::GetPinColor),
+		MakeAttributeSP(this, &SPCGEditorGraphNodePin::GetSecondaryPinIcon),
+		MakeAttributeSP(this, &SPCGEditorGraphNodePin::GetSecondaryPinColor));
+	PinImage = PinWidgetRef;
+
+	PinWidgetRef->SetCursor( 
+		TAttribute<TOptional<EMouseCursor::Type> >::Create (
+			TAttribute<TOptional<EMouseCursor::Type> >::FGetter::CreateRaw( this, &SPCGEditorGraphNodePin::GetPinCursor )
+		)
+	);
+
+	// Create the pin indicator widget (used for watched values)
+	static const FName NAME_NoBorder("NoBorder");
+	TSharedRef<SWidget> PinStatusIndicator =
+		SNew(SButton)
+		.ButtonStyle(FAppStyle::Get(), NAME_NoBorder)
+		.Visibility(this, &SPCGEditorGraphNodePin::GetPinStatusIconVisibility)
+		.ContentPadding(0)
+		.OnClicked(this, &SPCGEditorGraphNodePin::ClickedOnPinStatusIcon)
+		[
+			SNew(SImage)
+			.Image(this, &SPCGEditorGraphNodePin::GetPinStatusIcon)
+		];
+
+	TSharedRef<SWidget> LabelWidget = GetLabelWidget(GetLabelStyle(InArgs._PinLabelStyle));
+
+	// Create the widget used for the pin body (status indicator, label, and value)
+	LabelAndValue =
+		SNew(SWrapBox)
+		.PreferredSize(150.f);
+
+	TSharedPtr<SImage> ExtraPinIconWidget;
+	FName ExtraPinIcon;
+	FText ExtraPinIconTooltip;
+	if (GetExtraIcon(ExtraPinIcon, ExtraPinIconTooltip))
+	{
+		ExtraPinIconWidget = SNew(SImage)
+			.Image(FAppStyle::GetBrush(ExtraPinIcon))
+			.ColorAndOpacity(this, &SPCGEditorGraphNodePin::GetPinTextColor);
+
+		if (!ExtraPinIconTooltip.IsEmpty())
+		{
+			ExtraPinIconWidget->SetToolTipText(ExtraPinIconTooltip);
+		}
+	}
+
+	if (!bIsInput)
+	{
+		LabelAndValue->AddSlot()
+			.VAlign(VAlign_Center)
+			[
+				PinStatusIndicator
+			];
+
+		LabelAndValue->AddSlot()
+			.VAlign(VAlign_Center)
+			[
+				LabelWidget
+			];
+
+		if (ExtraPinIconWidget.IsValid())
+		{
+			LabelAndValue->AddSlot()
+				.Padding(5, 0, 0, 0)
+				.VAlign(VAlign_Center)
+				[
+					ExtraPinIconWidget.ToSharedRef()
+				];
+		}
+	}
+	else
+	{
+		if (ExtraPinIconWidget.IsValid())
+		{
+			LabelAndValue->AddSlot()
+				.Padding(0, 0, 5, 0)
+				.VAlign(VAlign_Center)
+				[
+					ExtraPinIconWidget.ToSharedRef()
+				];
+		}
+
+		LabelAndValue->AddSlot()
+			.VAlign(VAlign_Center)
+			[
+				LabelWidget
+			];
+
+		ValueWidget = GetDefaultValueWidget();
+
+		if (ValueWidget != SNullWidget::NullWidget)
+		{
+			TSharedPtr<SBox> ValueBox;
+			LabelAndValue->AddSlot()
+				.Padding(bIsInput ? FMargin(InArgs._SideToSideMargin, 0, 0, 0) : FMargin(0, 0, InArgs._SideToSideMargin, 0))
+				.VAlign(VAlign_Center)
+				[
+					SAssignNew(ValueBox, SBox)
+					.Padding(0.0f)
+					[
+						ValueWidget.ToSharedRef()
+					]
+				];
+
+			if (!DoesWidgetHandleSettingEditingEnabled())
+			{
+				ValueBox->SetEnabled(TAttribute<bool>(this, &SPCGEditorGraphNodePin::IsEditingEnabled));
+			}
+		}
+
+		LabelAndValue->AddSlot()
+			.VAlign(VAlign_Center)
+			[
+				PinStatusIndicator
+			];
+	}
+
+	TSharedPtr<SHorizontalBox> PinContent;
+	if (bIsInput)
+	{
+		// Input pin
+		FullPinHorizontalRowWidget = PinContent = 
+			SNew(SHorizontalBox)
+			+SHorizontalBox::Slot()
+			.AutoWidth()
+			.VAlign(VAlign_Center)
+			[
+				RequiredPinIconWidget.ToSharedRef()
+			]
+			+SHorizontalBox::Slot()
+			.AutoWidth()
+			.VAlign(VAlign_Center)
+			.Padding(bDisplayPinMarker ? 0.0f : RequiredPinMarkerWidth, 0, InArgs._SideToSideMargin, 0)
+			[
+				PinWidgetRef
+			]
+			+SHorizontalBox::Slot()
+			.AutoWidth()
+			.VAlign(VAlign_Center)
+			[
+				LabelAndValue.ToSharedRef()
+			];
+	}
+	else
+	{
+		// Output pin
+		FullPinHorizontalRowWidget = PinContent = SNew(SHorizontalBox)
+			+SHorizontalBox::Slot()
+			.AutoWidth()
+			.VAlign(VAlign_Center)
+			[
+				LabelAndValue.ToSharedRef()
+			]
+			+SHorizontalBox::Slot()
+			.AutoWidth()
+			.VAlign(VAlign_Center)
+			.Padding(InArgs._SideToSideMargin, 0, 0, 0)
+			[
+				PinWidgetRef
+			];
+	}
+
+	// Set up a hover for pins that is tinted the color of the pin.
+	
+	SBorder::Construct(SBorder::FArguments()
+		.BorderImage(this, &SPCGEditorGraphNodePin::GetPinBorder)
+		.BorderBackgroundColor(this, &SPCGEditorGraphNodePin::GetHighlightColor)
+		.OnMouseButtonDown(this, &SPCGEditorGraphNodePin::OnPinNameMouseDown)
+		.Padding(0) // NOTE: This is different from base class implementation
+		[
+			SNew(SBorder)
+			.BorderImage(CachedImg_Pin_DiffOutline)
+			.BorderBackgroundColor(this, &SPCGEditorGraphNodePin::GetPinDiffColor)
+			.Padding(0) // NOTE: This is different from base class implementation
+			[
+				SNew(SLevelOfDetailBranchNode)
+				.UseLowDetailSlot(this, &SPCGEditorGraphNodePin::UseLowDetailPinNames)
+				.LowDetail()
+				[
+					//@TODO: Try creating a pin-colored line replacement that doesn't measure text / call delegates but still renders
+					PinWidgetRef
+				]
+				.HighDetail()
+				[
+					PinContent.ToSharedRef()
+				]
+			]
+		]
+	);
+
+	TSharedPtr<IToolTip> TooltipWidget = SNew(SToolTip)
+		.Text(this, &SPCGEditorGraphNodePin::GetTooltipText);
+
+	SetToolTip(TooltipWidget);
+
+}
+
+void SPCGEditorGraphNodePin::GetPCGNodeAndPin(const UPCGNode*& OutNode, const UPCGPin*& OutPin) const
+{
+	UEdGraphPin* GraphPin = GetPinObj();
+	if (GraphPin && !GraphPin->IsPendingKill())
+	{
+		const UPCGEditorGraphNodeBase* EditorNode = CastChecked<const UPCGEditorGraphNodeBase>(GraphPinObj->GetOwningNode());
+		
+		OutNode = EditorNode ? EditorNode->GetPCGNode() : nullptr;
+
+		OutPin = OutNode ? OutNode->GetInputPin(GraphPin->GetFName()) : nullptr;
+		if (!OutPin)
+		{
+			OutPin = OutNode ? OutNode->GetOutputPin(GraphPin->GetFName()) : nullptr;
+		}
+	}
+	else
+	{
+		OutNode = nullptr;
+		OutPin = nullptr;
+	}
+}
+
+void SPCGEditorGraphNodePin::ApplyUnusedPinStyle(FSlateColor& InOutColor) const
+{
+	const UPCGPin* PCGPin = nullptr;
+	const UPCGNode* PCGNode = nullptr;
+	GetPCGNodeAndPin(PCGNode, PCGPin);
+
+	bool bPinDisabled = false;
+
+	// Check if the pin was deactivated in the previous execution.
+	const UEdGraphPin* Pin = GetPinObj();
+	if (Pin && Pin->Direction == EEdGraphPinDirection::EGPD_Output)
+	{
+		if (const UPCGEditorGraphNodeBase* Node = Cast<UPCGEditorGraphNodeBase>(Pin ? Pin->GetOwningNode() : nullptr))
+		{
+			// If node is already disabled, don't bother disabling pin on top of that, does not look nice to disable both and may
+			// not be meaningful to do so in any case.
+			if (!Node->IsDisplayAsDisabledForced())
+			{
+				bPinDisabled = !Node->IsOutputPinActive(Pin);
+			}
+		}
+	}
+
+	// Halve opacity if pin is unused - intended to happen whether disabled or not
+	if (bPinDisabled || (PCGPin && PCGNode && !PCGNode->IsPinUsedByNodeExecution(PCGPin)))
+	{
+		FLinearColor Color = InOutColor.GetSpecifiedColor();
+		Color.A *= 0.5;
+		InOutColor = Color;
+	}
+}
+
+// Adapted from SGraphPin::GetPinColor
+FSlateColor SPCGEditorGraphNodePin::GetPinColor() const
+{
+	FSlateColor Color = SGraphPin::GetPinColor();
+
+	ApplyUnusedPinStyle(Color);
+
+	return Color;
+}
+
+FSlateColor SPCGEditorGraphNodePin::GetPinTextColor() const
+{
+	FSlateColor Color = SGraphPin::GetPinTextColor();
+
+	ApplyUnusedPinStyle(Color);
+
+	return Color;
+}
+
+FName SPCGEditorGraphNodePin::GetLabelStyle(FName DefaultLabelStyle) const
+{
+	const UPCGPin* PCGPin = nullptr;
+	const UPCGNode* PCGNode = nullptr;
+	FName LabelStyle = NAME_None;
+
+	GetPCGNodeAndPin(PCGNode, PCGPin);
+
+	const UPCGSettings* Settings = PCGNode ? PCGNode->GetSettings() : nullptr;
+
+	if (!PCGPin || !Settings || !Settings->GetPinLabelStyle(PCGPin, LabelStyle))
+	{
+		LabelStyle = DefaultLabelStyle;
+	}
+
+	return LabelStyle;
+}
+
+bool SPCGEditorGraphNodePin::GetExtraIcon(FName& OutExtraIcon, FText& OutTooltip) const
+{
+	const UPCGPin* PCGPin = nullptr;
+	const UPCGNode* PCGNode = nullptr;
+
+	GetPCGNodeAndPin(PCGNode, PCGPin);
+
+	const UPCGSettings* Settings = PCGNode ? PCGNode->GetSettings() : nullptr;
+	return (PCGPin && Settings) ? Settings->GetPinExtraIcon(PCGPin, OutExtraIcon, OutTooltip) : false;
+}
+
+bool SPCGEditorGraphNodePin::ShouldDisplayAsRequiredForExecution() const
+{
+	const UPCGPin* PCGPin = nullptr;
+	const UPCGNode* PCGNode = nullptr;
+	GetPCGNodeAndPin(PCGNode, PCGPin);
+
+	// Trivial early out tests, and advanced pins should never display as required.
+	if (!PCGPin || !PCGNode || PCGPin->Properties.IsAdvancedPin())
+	{
+		return false;
+	}
+
+	if (PCGNode->IsInputPinRequiredByExecution(PCGPin))
+	{
+		return true;
+	}
+
+	const UPCGSettings* Settings = PCGNode->GetSettings();
+	if (Settings && Settings->CanCullTaskIfUnwired())
+	{
+		// If the node will cull if unwired, and if it only has a single normal pin (and no required pins), then display the pin
+		// as required, because it effectively is. So return false if there are other pins which are not advanced.
+		return !Algo::AnyOf(PCGNode->GetInputPins(), [PCGPin](UPCGPin* InOtherPin)
+		{
+			return InOtherPin && InOtherPin != PCGPin && !InOtherPin->Properties.IsAdvancedPin();
+		});
+	}
+
+	return false;
+}
+
+void SPCGEditorGraphNode::Construct(const FArguments& InArgs, UPCGEditorGraphNodeBase* InNode)
+{
+	GraphNode = InNode;
+	PCGEditorGraphNode = InNode;
+
+	if (InNode)
+	{
+		InNode->OnNodeChangedDelegate.BindSP(this, &SPCGEditorGraphNode::OnNodeChanged);
+		InNode->OnNodeRenameInitiatedDelegate.BindSP(this, &SPCGEditorGraphNode::OnNodeRenameInitiated);
+	}
+
+	UpdateGraphNode();
+}
+
+void SPCGEditorGraphNode::CreateAddPinButtonWidget()
+{
+	// Add Pin Button (+) Source Reference: Engine\Source\Editor\GraphEditor\Private\KismetNodes\SGraphNodeK2Sequence.cpp
+	const TSharedPtr<SWidget> AddPinButton = AddPinButtonContent(LOCTEXT("AddSourcePin", "Add Pin"), LOCTEXT("AddSourcePinTooltip", "Add a dynamic source input pin"));
+
+	FMargin AddPinPadding = Settings->GetInputPinPadding();
+	AddPinPadding.Top += 6.0f;
+
+	check(PCGEditorGraphNode->GetPCGNode());
+	const UPCGSettingsWithDynamicInputs* NodeSettings = CastChecked<UPCGSettingsWithDynamicInputs>(PCGEditorGraphNode->GetPCGNode()->GetSettings());
+
+	const int32 Index = NodeSettings->GetStaticInputPinNum() + NodeSettings->GetDynamicInputPinNum();
+	LeftNodeBox->InsertSlot(Index)
+		.AutoHeight()
+		.VAlign(VAlign_Bottom)
+		.Padding(AddPinPadding)
+		[
+			AddPinButton.ToSharedRef()
+		];
+}
+
+void SPCGEditorGraphNode::UpdateGraphNode()
+{
+	if (PCGEditorGraphNode && PCGEditorGraphNode->ShouldDrawCompact())
+	{
+		UpdateCompactNode();
+	}
+	else
+	{
+		SGraphNode::UpdateGraphNode();
+
+		if (PCGEditorGraphNode->CanUserAddRemoveDynamicInputPins())
+		{
+			CreateAddPinButtonWidget();
+		}
+	}
+}
+
+const FSlateBrush* SPCGEditorGraphNode::GetNodeBodyBrush() const
+{
+	const bool bNeedsTint = PCGEditorGraphNode && PCGEditorGraphNode->GetPCGNode() && PCGEditorGraphNode->GetPCGNode()->IsInstance();
+	if (bNeedsTint)
+	{
+		return FAppStyle::GetBrush("Graph.Node.TintedBody");
+	}
+	else
+	{
+		return FAppStyle::GetBrush("Graph.Node.Body");
+	}
+}
+
+TSharedRef<SWidget> SPCGEditorGraphNode::CreateTitleWidget(TSharedPtr<SNodeTitle> InNodeTitle)
+{
+	// Reimplementation of the SGraphNode::CreateTitleWidget so we can control the style
+	const bool bIsInstanceNode = (PCGEditorGraphNode && PCGEditorGraphNode->GetPCGNode() && PCGEditorGraphNode->GetPCGNode()->IsInstance());
+
+	SAssignNew(InlineEditableText, SInlineEditableTextBlock)
+		.Style(FPCGEditorStyle::Get(), bIsInstanceNode ? "PCG.Node.InstancedNodeTitleInlineEditableText" : "PCG.Node.NodeTitleInlineEditableText")
+		.Text(InNodeTitle.Get(), &SNodeTitle::GetHeadTitle)
+		.OnVerifyTextChanged(this, &SPCGEditorGraphNode::OnVerifyNameTextChanged)
+		.OnTextCommitted(this, &SPCGEditorGraphNode::OnNameTextCommited)
+		.IsReadOnly(this, &SPCGEditorGraphNode::IsNameReadOnly)
+		.IsSelected(this, &SPCGEditorGraphNode::IsSelectedExclusively);
+	InlineEditableText->SetColorAndOpacity(TAttribute<FLinearColor>::Create(TAttribute<FLinearColor>::FGetter::CreateSP(this, &SPCGEditorGraphNode::GetNodeTitleTextColor)));
+
+	return InlineEditableText.ToSharedRef();
+}
+
+TSharedPtr<SGraphPin> SPCGEditorGraphNode::CreatePinWidget(UEdGraphPin* Pin) const
+{
+	return SNew(SPCGEditorGraphNodePin, Pin);
+}
+
+EVisibility SPCGEditorGraphNode::IsAddPinButtonVisible() const
+{
+	if (PCGEditorGraphNode && PCGEditorGraphNode->IsNodeEnabled() && PCGEditorGraphNode->CanUserAddRemoveDynamicInputPins() && SGraphNode::IsAddPinButtonVisible() == EVisibility::Visible)
+	{
+		return EVisibility::Visible;
+	}
+	
+	return EVisibility::Hidden;
+}
+
+FReply SPCGEditorGraphNode::OnAddPin()
+{
+	check(PCGEditorGraphNode);
+
+	PCGEditorGraphNode->OnUserAddDynamicInputPin();
+	
+	return FReply::Handled();
+}
+
+void SPCGEditorGraphNode::AddPin(const TSharedRef<SGraphPin>& PinToAdd)
+{
+	check(PCGEditorGraphNode);
+	UPCGNode* PCGNode = PCGEditorGraphNode->GetPCGNode();
+
+	if (PCGNode && PinToAdd->GetPinObj())
+	{
+		const bool bIsInPin = PinToAdd->GetPinObj()->Direction == EEdGraphPinDirection::EGPD_Input;
+		const FName& PinName = PinToAdd->GetPinObj()->PinName;
+
+		if(UPCGPin* Pin = (bIsInPin ? PCGNode->GetInputPin(PinName) : PCGNode->GetOutputPin(PinName)))
+		{
+			const bool bIsMultiData = Pin->Properties.bAllowMultipleData;
+			const bool bIsMultiConnections = Pin->AllowsMultipleConnections();
+
+			// Check for special types
+			if (Pin->GetCurrentTypes() == EPCGDataType::Param)
+			{
+				const FSlateBrush* ConnectedBrush = FPCGEditorStyle::Get().GetBrush(bIsInPin ? PCGEditorStyleConstants::Pin_Param_IN_C : PCGEditorStyleConstants::Pin_Param_OUT_C);
+				const FSlateBrush* DisconnectedBrush = FPCGEditorStyle::Get().GetBrush(bIsInPin ? PCGEditorStyleConstants::Pin_Param_IN_DC : PCGEditorStyleConstants::Pin_Param_OUT_DC);
+
+				PinToAdd->SetCustomPinIcon(ConnectedBrush, DisconnectedBrush);
+			}
+			else if (Pin->GetCurrentTypes() == EPCGDataType::Spatial)
+			{
+				const FSlateBrush* ConnectedBrush = FPCGEditorStyle::Get().GetBrush(bIsInPin ? PCGEditorStyleConstants::Pin_Composite_IN_C : PCGEditorStyleConstants::Pin_Composite_OUT_C);
+				const FSlateBrush* DisconnectedBrush = FPCGEditorStyle::Get().GetBrush(bIsInPin ? PCGEditorStyleConstants::Pin_Composite_IN_DC : PCGEditorStyleConstants::Pin_Composite_OUT_DC);
+
+				PinToAdd->SetCustomPinIcon(ConnectedBrush, DisconnectedBrush);
+			}
+			else if (Pin->Properties.Usage == EPCGPinUsage::DependencyOnly)
+			{
+				const FSlateBrush* ConnectedBrush = FPCGEditorStyle::Get().GetBrush(PCGEditorStyleConstants::Pin_Graph_Dependency_C);
+				const FSlateBrush* DisconnectedBrush = FPCGEditorStyle::Get().GetBrush(PCGEditorStyleConstants::Pin_Graph_Dependency_DC);
+
+				PinToAdd->SetCustomPinIcon(ConnectedBrush, DisconnectedBrush);
+			}
+			else
+			{
+				// Node outputs are always single collection (SC).
+				static const FName* PinBrushes[] =
+				{
+					&PCGEditorStyleConstants::Pin_SD_SC_IN_C,
+					&PCGEditorStyleConstants::Pin_SD_SC_IN_DC,
+					&PCGEditorStyleConstants::Pin_SD_MC_IN_C,
+					&PCGEditorStyleConstants::Pin_SD_MC_IN_DC,
+					&PCGEditorStyleConstants::Pin_MD_SC_IN_C,
+					&PCGEditorStyleConstants::Pin_MD_SC_IN_DC,
+					&PCGEditorStyleConstants::Pin_MD_MC_IN_C,
+					&PCGEditorStyleConstants::Pin_MD_MC_IN_DC,
+					&PCGEditorStyleConstants::Pin_SD_SC_OUT_C,
+					&PCGEditorStyleConstants::Pin_SD_SC_OUT_DC,
+					&PCGEditorStyleConstants::Pin_SD_SC_OUT_C,
+					&PCGEditorStyleConstants::Pin_SD_SC_OUT_DC,
+					&PCGEditorStyleConstants::Pin_MD_SC_OUT_C,
+					&PCGEditorStyleConstants::Pin_MD_SC_OUT_DC,
+					&PCGEditorStyleConstants::Pin_MD_SC_OUT_C,
+					&PCGEditorStyleConstants::Pin_MD_SC_OUT_DC
+				};
+
+				const int32 ConnectedIndex = (bIsInPin ? 0 : 8) + (bIsMultiData ? 4 : 0) + (bIsMultiConnections ? 2 : 0);
+				const int32 DisconnectedIndex = ConnectedIndex + 1;
+
+				const FSlateBrush* ConnectedBrush = FPCGEditorStyle::Get().GetBrush(*PinBrushes[ConnectedIndex]);
+				const FSlateBrush* DisconnectedBrush = FPCGEditorStyle::Get().GetBrush(*PinBrushes[DisconnectedIndex]);
+
+				PinToAdd->SetCustomPinIcon(ConnectedBrush, DisconnectedBrush);
+			}
+		}
+	}
+
+	SGraphNode::AddPin(PinToAdd);
+
+	// The base class does not give an override to change the padding of the pin widgets, so do it here. Our input pins widgets include
+	// a small marker to indicate the pin is required, which need to display at the left edge of the node, so remove left padding.
+	if (PinToAdd->GetDirection() == EEdGraphPinDirection::EGPD_Input)
+	{
+		const int LastIndex = LeftNodeBox->GetChildren()->Num() - 1;
+		check(LastIndex >= 0);
+
+		SVerticalBox::FSlot& PinSlot = LeftNodeBox->GetSlot(LastIndex);
+
+		FMargin Margin = Settings->GetInputPinPadding();
+		Margin.Left = 0;
+		PinSlot.SetPadding(Margin);
+	}
+}
+
+TArray<FOverlayWidgetInfo> SPCGEditorGraphNode::GetOverlayWidgets(bool bSelected, const FVector2D& WidgetSize) const
+{
+	TArray<FOverlayWidgetInfo> OverlayWidgets = SGraphNode::GetOverlayWidgets(bSelected, WidgetSize);
+
+	if (PCGEditorGraphNode && !PCGEditorGraphNode->ShouldDrawCompact())
+	{
+		// Higen grid size overlay widget. All magic numbers below hand tweaked to match UI mockup.
+		const EPCGHiGenGrid InspectedGrid = PCGEditorGraphNode->GetInspectedGenerationGrid();
+		UPCGNode* PCGNode = PCGEditorGraphNode->GetPCGNode();
+
+		const bool bInspectingHigen = InspectedGrid != EPCGHiGenGrid::Uninitialized;
+		if (bInspectingHigen && PCGEditorGraphNode->IsNodeEnabled())
+		{
+			FText GenerationGridText;
+			const EPCGHiGenGrid Grid = PCGEditorGraphNode->GetGenerationGrid();
+
+			if (Grid == EPCGHiGenGrid::Unbounded)
+			{
+				GenerationGridText = FText::FromString(TEXT("UB"));
+			}
+			else
+			{
+				// Meters are easier on the eyes.
+				const uint32 GridSize = PCGHiGenGrid::GridToGridSize(Grid) / 100;
+				GenerationGridText = FText::AsNumber(GridSize, &FNumberFormattingOptions::DefaultNoGrouping());
+			}
+
+			FLinearColor Tint = FLinearColor::White;
+			if (Grid != EPCGHiGenGrid::Uninitialized)
+			{
+				Tint = GetGridLabelColor(Grid);
+			}
+			else if (PCGEditorGraphNode->IsDisplayAsDisabledForced())
+			{
+				Tint.A *= 0.35f;
+			}
+
+			// Create a border brush for each combination of grids, to workaround issue where the tint does not apply
+			// to the border element.
+			const FSlateBrush* BorderBrush = GetBorderBrush(InspectedGrid, Grid);
+
+			FLinearColor TextColor = FColor::White;
+			FLinearColor BackgroundColor = FColor::Black;
+			if (InspectedGrid == Grid)
+			{
+				// Flip colors for active grid to highlight them.
+				Swap(TextColor, BackgroundColor);
+			}
+
+			TSharedPtr<SWidget> GridSizeLabel =
+				SNew(SHorizontalBox)
+				.Visibility(EVisibility::Visible)
+				+SHorizontalBox::Slot()
+				[
+					SNew(SBorder)
+					.BorderImage(BorderBrush)
+					.Padding(FMargin(12, 3))
+					.ColorAndOpacity(Tint)
+					[
+						SNew(STextBlock)
+						.TextStyle(FAppStyle::Get(), "Graph.Node.NodeTitle")
+						.Text(GenerationGridText)
+						.Justification(ETextJustify::Center)
+						.ColorAndOpacity(TextColor)
+					]
+				];
+
+			FOverlayWidgetInfo GridSizeLabelInfo(GridSizeLabel);
+			GridSizeLabelInfo.OverlayOffset = FVector2D(GetDesiredSize().X - 30.0f, -9.0f);
+
+			OverlayWidgets.Add(GridSizeLabelInfo);
+		}
+
+		if (PCGEditorGraphNode->GetPCGNode() && PCGEditorGraphNode->GetPCGNode()->GetSettings() && PCGEditorGraphNode->GetPCGNode()->GetSettings()->ShouldExecuteOnGPU())
+		{
+			const float BorderRadius = 7.0f;
+			const float BorderStroke = 1.0f;
+			FText GPUText = FText::FromString(TEXT("GPU"));
+			const FLinearColor BorderColor(0.5f, 0.5f, 0.5f, 0.5f);
+			FLinearColor TextColor(0.5f, 0.5f, 0.5f, 0.8f);
+			
+			const FSlateBrush* BorderBrush = new FSlateRoundedBoxBrush(
+				FLinearColor::Transparent,
+				BorderRadius,
+				BorderColor,
+				BorderStroke);
+
+			TSharedPtr<SWidget> GPUUsageLabel =
+				SNew(SHorizontalBox)
+				.Visibility(EVisibility::Visible)
+				+SHorizontalBox::Slot()
+				[
+					SNew(SBorder)
+					.BorderImage(BorderBrush)
+					.Padding(FMargin(4, 3))
+					[
+						SNew(STextBlock)
+						.TextStyle(FPCGEditorStyle::Get(), "PCG.Node.AdditionalOverlayWidgetText")
+						.Text(GPUText)
+						.Justification(ETextJustify::Center)
+						.ColorAndOpacity(TextColor)
+					]
+				];
+
+			FOverlayWidgetInfo GPUUsageLabelInfo(GPUUsageLabel);
+			GPUUsageLabelInfo.OverlayOffset = FVector2D(GetDesiredSize().X - 34.0f, GetDesiredSize().Y + 5.0f);
+
+			OverlayWidgets.Add(GPUUsageLabelInfo);
+		}
+	}
+
+	return OverlayWidgets;
+}
+
+void SPCGEditorGraphNode::GetOverlayBrushes(bool bSelected, const FVector2D WidgetSize, TArray<FOverlayBrushInfo>& Brushes) const
+{
+	check(PCGEditorGraphNode);
+	
+	FVector2D OverlayOffset(0.0, 0.0);
+
+	auto AddOverlayBrush = [&OverlayOffset, &Brushes](const FName& BrushName)
+	{
+		const FSlateBrush* Brush = FPCGEditorStyle::Get().GetBrush(BrushName);
+
+		if (Brush)
+		{
+			FOverlayBrushInfo BrushInfo;
+			BrushInfo.Brush = Brush;
+			BrushInfo.OverlayOffset = OverlayOffset - Brush->GetImageSize() / 2.0;
+			Brushes.Add(BrushInfo);
+
+			OverlayOffset.Y += Brush->GetImageSize().Y;
+		}
+	};
+
+	if (PCGEditorGraphNode->IsCulledFromExecution())
+	{
+		AddOverlayBrush(PCGEditorStyleConstants::Node_Overlay_Inactive);
+	}
+
+	if (const UPCGNode* PCGNode = PCGEditorGraphNode->GetPCGNode())
+	{
+		if (PCGNode->GetSettingsInterface() && PCGNode->GetSettingsInterface()->bDebug)
+		{
+			AddOverlayBrush(TEXT("PCG.NodeOverlay.Debug"));
+		}
+	}
+
+	if (PCGEditorGraphNode->GetInspected())
+	{
+		AddOverlayBrush(TEXT("PCG.NodeOverlay.Inspect"));
+	}
+}
+
+void SPCGEditorGraphNode::OnNodeChanged()
+{
+	UpdateGraphNode();
+}
+
+void SPCGEditorGraphNode::OnNodeRenameInitiated()
+{
+	if (InlineEditableText.IsValid())
+	{
+		InlineEditableText->EnterEditingMode();
+	}
+}
+
+void SPCGEditorGraphNode::UpdateCompactNode()
+{
+	// Based on SGraphNodeK2Base::UpdateCompactNode. Changes:
+	// * Removed creation of tooltip widget, it did not port across trivially and the usage is fairly obvious for the current
+	//   compact nodes, but this could be re-added - TODO
+	// * Changed title style - reduced font size substantially
+	// * Layout differentiation for "pure" vs "impure" K2 nodes removed
+
+	InputPins.Empty();
+	OutputPins.Empty();
+
+	// Error handling set-up
+	SetupErrorReporting();
+
+	// Reset variables that are going to be exposed, in case we are refreshing an already setup node.
+	RightNodeBox.Reset();
+	LeftNodeBox.Reset();
+
+	if (!SWidget::GetToolTip().IsValid())
+	{
+		TSharedRef<SToolTip> DefaultToolTip = IDocumentation::Get()->CreateToolTip(TAttribute< FText >(this, &SGraphNode::GetNodeTooltip), NULL, GraphNode->GetDocumentationLink(), GraphNode->GetDocumentationExcerptName());
+		SetToolTip(DefaultToolTip);
+	}
+
+	// Setup a meta tag for this node
+	FGraphNodeMetaData TagMeta(TEXT("Graphnode"));
+	PopulateMetaTag(&TagMeta);
+	
+	TSharedPtr<SNodeTitle> NodeTitle = SNew(SNodeTitle, GraphNode);
+	TSharedRef<SOverlay> NodeOverlay = SNew(SOverlay);
+	
+	IconColor = FLinearColor::White;
+
+	// Add optional node specific widget to the overlay:
+	TSharedPtr<SWidget> OverlayWidget = GraphNode->CreateNodeImage();
+	if(OverlayWidget.IsValid())
+	{
+		NodeOverlay->AddSlot()
+		.HAlign(HAlign_Center)
+		.VAlign(VAlign_Center)
+		[
+			SNew( SBox )
+			.WidthOverride( 70.f )
+			.HeightOverride( 70.f )
+			[
+				OverlayWidget.ToSharedRef()
+			]
+		];
+	}
+
+	FName CompactBodyIcon = NAME_None;
+	check(PCGEditorGraphNode);
+	if (PCGEditorGraphNode->GetCompactNodeIcon(CompactBodyIcon))
+	{
+		NodeOverlay->AddSlot()
+			.HAlign(HAlign_Center)
+			.VAlign(VAlign_Center)
+			.Padding(45.f, 0.f, 45.f, 0.f)
+			[
+				SNew(SImage)
+					.Image(FPCGEditorStyle::Get().GetBrush(CompactBodyIcon))
+					.ColorAndOpacity(this, &SGraphNode::GetNodeTitleIconColor)
+			];
+	}
+	else
+	{
+		NodeOverlay->AddSlot()
+			.HAlign(HAlign_Center)
+			.VAlign(VAlign_Center)
+			.Padding(45.f, 0.f, 45.f, 0.f)
+			[
+				// MIDDLE
+				SNew(SVerticalBox)
+				+SVerticalBox::Slot()
+				.HAlign(HAlign_Center)
+				.AutoHeight()
+				[
+					SNew(STextBlock)
+						.TextStyle(FAppStyle::Get(), "Graph.Node.NodeTitle" )
+						.Text( NodeTitle.Get(), &SNodeTitle::GetHeadTitle )
+						.WrapTextAt(128.0f)
+						.ColorAndOpacity(this, &SGraphNode::GetNodeTitleIconColor)
+				]
+				+SVerticalBox::Slot()
+				.AutoHeight()
+				[
+					NodeTitle.ToSharedRef()
+				]
+			];
+	}
+	
+	NodeOverlay->AddSlot()
+		.HAlign(HAlign_Left)
+		.VAlign(VAlign_Center)
+		.Padding(0.f, 0.f, 55.f, 0.f)
+		[
+			// LEFT
+			SAssignNew(LeftNodeBox, SVerticalBox)
+		];
+
+	NodeOverlay->AddSlot()
+		.HAlign(HAlign_Right)
+		.VAlign(VAlign_Center)
+		.Padding(55.f, 0.f, 0.f, 0.f)
+		[
+			// RIGHT
+			SAssignNew(RightNodeBox, SVerticalBox)
+		];
+
+	//
+	//             ______________________
+	//            | (>) L |      | R (>) |
+	//            | (>) E |      | I (>) |
+	//            | (>) F |   +  | G (>) |
+	//            | (>) T |      | H (>) |
+	//            |       |      | T (>) |
+	//            |_______|______|_______|
+	//
+	this->ContentScale.Bind( this, &SGraphNode::GetContentScale );
+	
+	TSharedRef<SVerticalBox> InnerVerticalBox =
+		SNew(SVerticalBox)
+		+SVerticalBox::Slot()
+		[
+			// NODE CONTENT AREA
+			SNew( SOverlay)
+			+SOverlay::Slot()
+			[
+				SNew(SImage)
+				.Image( FAppStyle::GetBrush("Graph.VarNode.Body") )
+				.ColorAndOpacity(this, &SGraphNode::GetNodeBodyColor)
+			]
+			+ SOverlay::Slot()
+			[
+				SNew(SImage)
+				.Image( FAppStyle::GetBrush("Graph.VarNode.Gloss") )
+				.ColorAndOpacity(this, &SGraphNode::GetNodeBodyColor)
+			]
+			+SOverlay::Slot()
+			.Padding( FMargin(0,3) )
+			[
+				NodeOverlay
+			]
+		];
+	
+	TSharedPtr<SWidget> EnabledStateWidget = GetEnabledStateWidget();
+	if (EnabledStateWidget.IsValid())
+	{
+		InnerVerticalBox->AddSlot()
+			.AutoHeight()
+			.HAlign(HAlign_Fill)
+			.VAlign(VAlign_Top)
+			.Padding(FMargin(2, 0))
+			[
+				EnabledStateWidget.ToSharedRef()
+			];
+	}
+
+	InnerVerticalBox->AddSlot()
+		.AutoHeight()
+		.Padding( FMargin(5.0f, 1.0f) )
+		[
+			ErrorReporting->AsWidget()
+		];
+
+	this->GetOrAddSlot( ENodeZone::Center )
+		.HAlign(HAlign_Center)
+		.VAlign(VAlign_Center)
+		[
+			InnerVerticalBox
+		];
+
+	CreatePinWidgets();
+
+	// Hide pin labels
+	for (auto InputPin: this->InputPins)
+	{
+		if (InputPin->GetPinObj()->ParentPin == nullptr)
+		{
+			InputPin->SetShowLabel(false);
+		}
+	}
+
+	for (auto OutputPin : this->OutputPins)
+	{
+		if (OutputPin->GetPinObj()->ParentPin == nullptr)
+		{
+			OutputPin->SetShowLabel(false);
+		}
+	}
+
+	// Create comment bubble
+	TSharedPtr<SCommentBubble> CommentBubble;
+	const FSlateColor CommentColor = GetDefault<UGraphEditorSettings>()->DefaultCommentNodeTitleColor;
+
+	SAssignNew( CommentBubble, SCommentBubble )
+		.GraphNode( GraphNode )
+		.Text( this, &SGraphNode::GetNodeComment )
+		.OnTextCommitted( this, &SGraphNode::OnCommentTextCommitted )
+		.ColorAndOpacity( CommentColor )
+		.AllowPinning( true )
+		.EnableTitleBarBubble( true )
+		.EnableBubbleCtrls( true )
+		.GraphLOD( this, &SGraphNode::GetCurrentLOD )
+		.IsGraphNodeHovered( this, &SGraphNode::IsHovered );
+
+	GetOrAddSlot( ENodeZone::TopCenter )
+		.SlotOffset(TAttribute<FVector2D>(CommentBubble.Get(), &SCommentBubble::GetOffset))
+		.SlotSize(TAttribute<FVector2D>(CommentBubble.Get(), &SCommentBubble::GetSize))
+		.AllowScaling(TAttribute<bool>(CommentBubble.Get(), &SCommentBubble::IsScalingAllowed))
+		.VAlign(VAlign_Top)
+		[
+			CommentBubble.ToSharedRef()
+		];
+
+	CreateInputSideAddButton(LeftNodeBox);
+	CreateOutputSideAddButton(RightNodeBox);
+}
+
+FLinearColor SPCGEditorGraphNode::GetGridLabelColor(EPCGHiGenGrid NodeGrid)
+{
+	// All colours hand tweaked to give a kind of "temperature scale" for the hierarchy.
+	switch (NodeGrid)
+	{
+	case EPCGHiGenGrid::Unbounded:
+		return FColor(255, 255, 255, 255);
+	case EPCGHiGenGrid::Grid4194304: // fall-through
+	case EPCGHiGenGrid::Grid2097152: // fall-through
+	case EPCGHiGenGrid::Grid1048576: // fall-through
+	case EPCGHiGenGrid::Grid524288: // fall-through
+	case EPCGHiGenGrid::Grid262144: // fall-through
+	case EPCGHiGenGrid::Grid131072: // fall-through
+	case EPCGHiGenGrid::Grid65536: // fall-through
+	case EPCGHiGenGrid::Grid32768: // fall-through
+	case EPCGHiGenGrid::Grid16384: // fall-through
+	case EPCGHiGenGrid::Grid8192: // fall-through
+	case EPCGHiGenGrid::Grid4096: // fall-through
+	case EPCGHiGenGrid::Grid2048:
+		return FColor(53, 60, 171, 255);
+	case EPCGHiGenGrid::Grid1024:
+		return FColor(31, 82, 210, 255);
+	case EPCGHiGenGrid::Grid512:
+		return FColor(16, 120, 217, 255);
+	case EPCGHiGenGrid::Grid256:
+		return FColor(8, 151, 208, 255);
+	case EPCGHiGenGrid::Grid128:
+		return FColor(9, 170, 188, 255);
+	case EPCGHiGenGrid::Grid64:
+		return FColor(64, 185, 150, 255);
+	case EPCGHiGenGrid::Grid32:
+		return FColor(144, 189, 114, 255);
+	case EPCGHiGenGrid::Grid16:
+		return FColor(207, 185, 89, 255);
+	case EPCGHiGenGrid::Grid8:
+		return FColor(252, 189, 61, 255);
+	case EPCGHiGenGrid::Grid4:
+		return FColor(243, 227, 28, 255);
+	default:
+		ensure(false);
+		return FLinearColor::White;
+	}
+}
+
+const FSlateBrush* SPCGEditorGraphNode::GetBorderBrush(EPCGHiGenGrid InspectedGrid, EPCGHiGenGrid NodeGrid) const
+{
+	if (InspectedGrid == NodeGrid)
+	{
+		return FPCGEditorStyle::Get().GetBrush(PCGEditorStyleConstants::Node_Overlay_GridSizeLabel_Active_Border);
+	}
+
+	// Hand tweaked multiplier to fade child node grid size labels.
+	const float Opacity = (InspectedGrid < NodeGrid) ? 1.0f : 0.5f;
+
+	return new FSlateRoundedBoxBrush(
+		FLinearColor::Black * Opacity,
+		PCGEditorStyleConstants::Node_Overlay_GridSizeLabel_BorderRadius,
+		GetGridLabelColor(NodeGrid) * Opacity,
+		PCGEditorStyleConstants::Node_Overlay_GridSizeLabel_BorderStroke);
+}
+
+#undef LOCTEXT_NAMESPACE
diff --git a/Engine/Plugins/PCG/Source/PCGEditor/Private/Nodes/SPCGEditorGraphNode.h b/Engine/Plugins/PCG/Source/PCGEditor/Private/Nodes/SPCGEditorGraphNode.h
new file mode 100644
--- /dev/null
+++ b/Engine/Plugins/PCG/Source/PCGEditor/Private/Nodes/SPCGEditorGraphNode.h
@@ -0,0 +1,55 @@
+// Copyright Epic Games, Inc. All Rights Reserved.
+
+#pragma once
+
+#include "SGraphNode.h"
+
+#include "PCGCommon.h"
+
+struct FOverlayBrushInfo;
+
+class UPCGEditorGraphNodeBase;
+
+class SPCGEditorGraphNode : public SGraphNode
+{
+public:
+	SLATE_BEGIN_ARGS(SPCGEditorGraphNode){}
+	SLATE_END_ARGS()
+
+	void Construct(const FArguments& InArgs, UPCGEditorGraphNodeBase* InNode);
+	void CreateAddPinButtonWidget();
+
+	//~ Begin SGraphNode Interface
+	virtual void UpdateGraphNode() override;
+	virtual const FSlateBrush* GetNodeBodyBrush() const override;
+	virtual void RequestRenameOnSpawn() override { /* Empty to avoid the default behavior to rename on node spawn */ }
+	virtual void AddPin(const TSharedRef<SGraphPin>& PinToAdd) override;
+	virtual TSharedRef<SWidget> CreateTitleWidget(TSharedPtr<SNodeTitle> InNodeTitle) override;
+	virtual TSharedPtr<SGraphPin> CreatePinWidget(UEdGraphPin* Pin) const override;
+	virtual EVisibility IsAddPinButtonVisible() const override;
+	virtual FReply OnAddPin() override;
+	//~ End SGraphNode Interface
+
+	//~ Begin SNodePanel::SNode Interface
+	virtual TArray<FOverlayWidgetInfo> GetOverlayWidgets(bool bSelected, const FVector2D& WidgetSize) const override;
+	virtual void GetOverlayBrushes(bool bSelected, const FVector2D WidgetSize, TArray<FOverlayBrushInfo>& Brushes) const override;
+	//~ End SNodePanel::SNode Interface
+
+protected:
+	void OnNodeChanged();
+	void OnNodeRenameInitiated();
+
+	/** Set up node in 'compact' mode */
+	void UpdateCompactNode();
+
+private:
+	static FLinearColor GetGridLabelColor(EPCGHiGenGrid NodeGrid);
+
+	/**
+	* Get the border brush for the given combination of grid sizes and enabled state. All a big workaround for FSlateRoundedBoxBrush not respecting
+	* the tint colour.
+	*/
+	const FSlateBrush* GetBorderBrush(EPCGHiGenGrid InspectedGrid, EPCGHiGenGrid NodeGrid) const;
+
+	UPCGEditorGraphNodeBase* PCGEditorGraphNode = nullptr;
+};
diff --git a/Engine/Plugins/PCG/Source/PCGEditor/Private/Nodes/SPCGEditorGraphVarNode.cpp b/Engine/Plugins/PCG/Source/PCGEditor/Private/Nodes/SPCGEditorGraphVarNode.cpp
new file mode 100644
--- /dev/null
+++ b/Engine/Plugins/PCG/Source/PCGEditor/Private/Nodes/SPCGEditorGraphVarNode.cpp
@@ -0,0 +1,208 @@
+// Copyright Epic Games, Inc. All Rights Reserved.
+
+/**
+* Adapted from SGraphNodeK2Var.cpp
+*/
+
+#include "SPCGEditorGraphVarNode.h"
+
+#include "PCGEditorGraphNodeReroute.h"
+
+#include "GraphEditorSettings.h"
+#include "SCommentBubble.h"
+#include "SGraphNode.h"
+#include "SlotBase.h"
+#include "EdGraph/EdGraphNode.h"
+#include "TutorialMetaData.h"
+#include "Widgets/Images/SImage.h"
+#include "Widgets/Layout/SSpacer.h"
+#include "Widgets/SBoxPanel.h"
+
+class SWidget;
+struct FSlateBrush;
+
+void SPCGEditorGraphVarNode::Construct(const FArguments& InArgs, UPCGEditorGraphNodeBase* InNode)
+{
+	SPCGEditorGraphNode::Construct(SPCGEditorGraphNode::FArguments{}, InNode);
+}
+
+FSlateColor SPCGEditorGraphVarNode::GetVariableColor() const
+{
+	return GraphNode->GetNodeTitleColor();
+}
+
+void SPCGEditorGraphVarNode::UpdateGraphNode()
+{
+	InputPins.Empty();
+	OutputPins.Empty();
+
+	// Reset variables that are going to be exposed, in case we are refreshing an already setup node.
+	RightNodeBox.Reset();
+	LeftNodeBox.Reset();
+
+	FMargin ContentAreaMargin = FMargin(0.0f, 4.0f);
+
+	SetupErrorReporting();
+
+	const bool bNeedsTitle = (Cast<UPCGEditorGraphNodeNamedRerouteBase>(GraphNode) != nullptr);
+
+	TSharedPtr<SWidget> TitleArea = SNullWidget::NullWidget.ToSharedPtr();
+
+	if (bNeedsTitle)
+	{
+		const bool bNeedsPadding = (Cast<UPCGEditorGraphNodeNamedRerouteUsage>(GraphNode) != nullptr);
+		FMargin CustomPadding = FMargin(Settings->PaddingLeftOfOutput, Settings->PaddingAbovePin, 0, Settings->PaddingBelowPin);
+
+		TSharedPtr<SNodeTitle> NodeTitle = SNew(SNodeTitle, GraphNode);
+		SAssignNew(TitleArea, SVerticalBox)
+			+ SVerticalBox::Slot()
+			.AutoHeight()
+			.VAlign(VAlign_Center)
+			.Padding(bNeedsPadding ? CustomPadding : FMargin())
+			[
+				CreateTitleWidget(NodeTitle)
+			]
+			+ SVerticalBox::Slot()
+			.AutoHeight()
+			[
+				NodeTitle.ToSharedRef()
+			];
+	}
+
+	// Setup a meta tag for this node
+	FGraphNodeMetaData TagMeta(TEXT("Graphnode"));
+	PopulateMetaTag(&TagMeta);
+
+	//             ________________
+	//            | (>) L |  R (>) |
+	//            | (>) E |  I (>) |
+	//            | (>) F |  G (>) |
+	//            | (>) T |  H (>) |
+	//            |       |  T (>) |
+	//            |_______|________|
+	//
+	this->ContentScale.Bind( this, &SGraphNode::GetContentScale );
+	this->GetOrAddSlot( ENodeZone::Center )
+	.HAlign(HAlign_Center)
+	.VAlign(VAlign_Center)
+	[
+		SNew(SVerticalBox)
+		+SVerticalBox::Slot()
+		[
+			SNew(SOverlay)
+			.AddMetaData<FGraphNodeMetaData>(TagMeta)
+			+ SOverlay::Slot()
+			[
+				SNew(SImage)
+				.Image( FAppStyle::GetBrush("Graph.VarNode.Body") )
+			]
+			+ SOverlay::Slot()
+			.VAlign(VAlign_Top)
+			[
+				SNew(SImage)
+				.Image( FAppStyle::GetBrush("Graph.VarNode.ColorSpill") )
+				.ColorAndOpacity( this, &SPCGEditorGraphVarNode::GetVariableColor )
+			]
+			+ SOverlay::Slot()
+			[
+				SNew(SImage)
+				.Image( FAppStyle::GetBrush("Graph.VarNode.Gloss") )
+			]
+			+ SOverlay::Slot()
+			.Padding( ContentAreaMargin )
+			[
+				// NODE CONTENT AREA
+				SNew(SHorizontalBox)
+				+SHorizontalBox::Slot()
+				.HAlign(HAlign_Left)
+				.FillWidth(1.0f)
+				.Padding( FMargin(2,0) )
+				[
+					// LEFT
+					SAssignNew(LeftNodeBox, SVerticalBox)
+				]
+				+ SHorizontalBox::Slot()
+				.AutoWidth()
+				.HAlign(HAlign_Center)
+				.VAlign(VAlign_Center)
+				[
+					TitleArea.ToSharedRef()
+				]
+				+ SHorizontalBox::Slot()
+				.AutoWidth()
+				.HAlign(HAlign_Right)
+				.Padding( FMargin(2,0) )
+				[
+					// RIGHT
+					SAssignNew(RightNodeBox, SVerticalBox)
+				]
+			]
+		]
+		+SVerticalBox::Slot()
+		.VAlign(VAlign_Top)
+		.AutoHeight() 
+		.Padding( FMargin(5.0f, 1.0f) )
+		[
+			ErrorReporting->AsWidget()
+		]
+	];
+
+	float VerticalPaddingAmount = 0.0f;
+
+	if (VerticalPaddingAmount > 0.0f)
+	{
+		LeftNodeBox->AddSlot()
+		.AutoHeight()
+		.HAlign(HAlign_Left)
+		.VAlign(VAlign_Center)
+		[
+			SNew(SSpacer).Size(FVector2D(0.0f, VerticalPaddingAmount))
+		];
+
+		RightNodeBox->AddSlot()
+		.AutoHeight()
+		.HAlign(HAlign_Left)
+		.VAlign(VAlign_Center)
+		[
+			SNew(SSpacer).Size(FVector2D(0.0f, VerticalPaddingAmount))
+		];
+	}
+	// Create comment bubble
+	TSharedPtr<SCommentBubble> CommentBubble;
+	const FSlateColor CommentColor = GetDefault<UGraphEditorSettings>()->DefaultCommentNodeTitleColor;
+
+	SAssignNew( CommentBubble, SCommentBubble )
+	.GraphNode( GraphNode )
+	.Text( this, &SGraphNode::GetNodeComment )
+	.OnTextCommitted( this, &SGraphNode::OnCommentTextCommitted )
+	.ColorAndOpacity( CommentColor )
+	.AllowPinning( true )
+	.EnableTitleBarBubble( true )
+	.EnableBubbleCtrls( true )
+	.GraphLOD( this, &SGraphNode::GetCurrentLOD )
+	.IsGraphNodeHovered( this, &SGraphNode::IsHovered );
+
+	GetOrAddSlot( ENodeZone::TopCenter )
+	.SlotOffset( TAttribute<FVector2D>( CommentBubble.Get(), &SCommentBubble::GetOffset ))
+	.SlotSize( TAttribute<FVector2D>( CommentBubble.Get(), &SCommentBubble::GetSize ))
+	.AllowScaling( TAttribute<bool>( CommentBubble.Get(), &SCommentBubble::IsScalingAllowed ))
+	.VAlign( VAlign_Top )
+	[
+		CommentBubble.ToSharedRef()
+	];
+
+	// Create widgets for each of the real pins
+	CreatePinWidgets();
+}
+
+const FSlateBrush* SPCGEditorGraphVarNode::GetShadowBrush(bool bSelected) const
+{
+	return bSelected ? FAppStyle::GetBrush(TEXT("Graph.VarNode.ShadowSelected")) : FAppStyle::GetBrush(TEXT("Graph.VarNode.Shadow"));
+}
+
+void SPCGEditorGraphVarNode::GetDiffHighlightBrushes(const FSlateBrush*& BackgroundOut, const FSlateBrush*& ForegroundOut) const
+{
+	BackgroundOut = FAppStyle::GetBrush(TEXT("Graph.VarNode.DiffHighlight"));
+	ForegroundOut = FAppStyle::GetBrush(TEXT("Graph.VarNode.DiffHighlightShading"));
+}
+
diff --git a/Engine/Plugins/PCG/Source/PCGEditor/Private/Nodes/SPCGEditorGraphVarNode.h b/Engine/Plugins/PCG/Source/PCGEditor/Private/Nodes/SPCGEditorGraphVarNode.h
new file mode 100644
--- /dev/null
+++ b/Engine/Plugins/PCG/Source/PCGEditor/Private/Nodes/SPCGEditorGraphVarNode.h
@@ -0,0 +1,36 @@
+// Copyright Epic Games, Inc. All Rights Reserved.
+
+#pragma once
+
+#include "CoreMinimal.h"
+#include "SPCGEditorGraphNode.h"
+#include "Styling/SlateColor.h"
+#include "Templates/SharedPointer.h"
+#include "Types/SlateEnums.h"
+#include "Widgets/DeclarativeSyntaxSupport.h"
+
+class FText;
+class SWidget;
+class UK2Node;
+struct FMargin;
+struct FSlateBrush;
+
+class SPCGEditorGraphVarNode : public SPCGEditorGraphNode
+{
+public:
+	SLATE_BEGIN_ARGS(SPCGEditorGraphVarNode){}
+	SLATE_END_ARGS()
+
+	void Construct(const FArguments& InArgs, UPCGEditorGraphNodeBase* InNode);
+
+	virtual const FSlateBrush* GetShadowBrush(bool bSelected) const override;
+
+	virtual void GetDiffHighlightBrushes(const FSlateBrush*& BackgroundOut, const FSlateBrush*& ForegroundOut) const override;
+
+	// SGraphNode interface
+	virtual void UpdateGraphNode() override;
+	// End of SGraphNode interface
+
+protected:
+	FSlateColor GetVariableColor() const;
+};
diff --git a/Engine/Plugins/PCG/Source/PCGEditor/Private/PCGEditor.cpp b/Engine/Plugins/PCG/Source/PCGEditor/Private/PCGEditor.cpp
--- a/Engine/Plugins/PCG/Source/PCGEditor/Private/PCGEditor.cpp
+++ b/Engine/Plugins/PCG/Source/PCGEditor/Private/PCGEditor.cpp
@@ -3,16 +3,12 @@
 #include "PCGEditor.h"
 
 #include "PCGComponent.h"
-#include "Framework/Application/SlateApplication.h"
 #include "PCGEdge.h"
-#include "IAssetTools.h"
-#include "PCGGraph.h"
 #include "PCGEditorModule.h"
+#include "PCGGraph.h"
 #include "PCGGraphFactory.h"
-#include "PCGInputOutputSettings.h"
 #include "PCGPin.h"
 #include "PCGSubsystem.h"
-#include "Compute/IPCGNodeSourceTextProvider.h"
 #include "Elements/PCGReroute.h"
 #include "Helpers/PCGSubgraphHelpers.h"
 #include "Rendering/SlateRenderer.h"
@@ -22,10 +18,6 @@
 
 #include "PCGEditorCommands.h"
 #include "PCGEditorGraph.h"
-#include "PCGEditorGraphNode.h"
-#include "PCGEditorGraphNodeInput.h"
-#include "PCGEditorGraphNodeOutput.h"
-#include "PCGEditorGraphNodeReroute.h"
 #include "PCGEditorGraphSchema.h"
 #include "PCGEditorGraphSchemaActions.h"
 #include "PCGEditorMenuContext.h"
@@ -40,23 +32,26 @@
 #include "SPCGEditorGraphLogView.h"
 #include "SPCGEditorGraphNodePalette.h"
 #include "SPCGEditorGraphProfilingView.h"
+#include "Nodes/PCGEditorGraphNode.h"
+#include "Nodes/PCGEditorGraphNodeInput.h"
+#include "Nodes/PCGEditorGraphNodeOutput.h"
+#include "Nodes/PCGEditorGraphNodeReroute.h"
 #include "Widgets/SPCGEditorNodeSource.h"
 
 #include "AssetToolsModule.h"
 #include "EdGraphUtilities.h"
 #include "EditorAssetLibrary.h"
 #include "GraphEditorActions.h"
-#include "IDetailsView.h"
+#include "IAssetTools.h"
 #include "LevelEditor.h"
-#include "PropertyEditorModule.h"
-#include "SNodePanel.h"
 #include "ScopedTransaction.h"
 #include "ShaderCore.h"
+#include "SNodePanel.h"
 #include "SourceCodeNavigation.h"
 #include "ToolMenu.h"
 #include "ToolMenuEntry.h"
-#include "ToolMenuSection.h"
 #include "ToolMenus.h"
+#include "ToolMenuSection.h"
 #include "UnrealEdGlobals.h"
 #include "Algo/AnyOf.h"
 #include "Editor/UnrealEdEngine.h"
diff --git a/Engine/Plugins/PCG/Source/PCGEditor/Private/PCGEditorGraph.cpp b/Engine/Plugins/PCG/Source/PCGEditor/Private/PCGEditorGraph.cpp
--- a/Engine/Plugins/PCG/Source/PCGEditor/Private/PCGEditorGraph.cpp
+++ b/Engine/Plugins/PCG/Source/PCGEditor/Private/PCGEditorGraph.cpp
@@ -3,15 +3,15 @@
 #include "PCGEditorGraph.h"
 
 #include "PCGEdge.h"
-#include "PCGEditorGraphNode.h"
-#include "PCGEditorGraphNodeInput.h"
-#include "PCGEditorGraphNodeOutput.h"
-#include "PCGEditorGraphNodeReroute.h"
 #include "PCGEditorModule.h"
 #include "PCGGraph.h"
 #include "PCGPin.h"
 #include "Elements/PCGReroute.h"
 #include "Elements/PCGUserParameterGet.h"
+#include "Nodes/PCGEditorGraphNode.h"
+#include "Nodes/PCGEditorGraphNodeInput.h"
+#include "Nodes/PCGEditorGraphNodeOutput.h"
+#include "Nodes/PCGEditorGraphNodeReroute.h"
 
 #include "EdGraph/EdGraphPin.h"
 
diff --git a/Engine/Plugins/PCG/Source/PCGEditor/Private/PCGEditorGraphNode.cpp b/Engine/Plugins/PCG/Source/PCGEditor/Private/PCGEditorGraphNode.cpp
deleted file mode 100644
diff --git a/Engine/Plugins/PCG/Source/PCGEditor/Private/PCGEditorGraphNode.h b/Engine/Plugins/PCG/Source/PCGEditor/Private/PCGEditorGraphNode.h
deleted file mode 100644
diff --git a/Engine/Plugins/PCG/Source/PCGEditor/Private/PCGEditorGraphNodeBase.cpp b/Engine/Plugins/PCG/Source/PCGEditor/Private/PCGEditorGraphNodeBase.cpp
deleted file mode 100644
diff --git a/Engine/Plugins/PCG/Source/PCGEditor/Private/PCGEditorGraphNodeBase.h b/Engine/Plugins/PCG/Source/PCGEditor/Private/PCGEditorGraphNodeBase.h
deleted file mode 100644
diff --git a/Engine/Plugins/PCG/Source/PCGEditor/Private/PCGEditorGraphNodeFactory.cpp b/Engine/Plugins/PCG/Source/PCGEditor/Private/PCGEditorGraphNodeFactory.cpp
deleted file mode 100644
diff --git a/Engine/Plugins/PCG/Source/PCGEditor/Private/PCGEditorGraphNodeFactory.h b/Engine/Plugins/PCG/Source/PCGEditor/Private/PCGEditorGraphNodeFactory.h
deleted file mode 100644
diff --git a/Engine/Plugins/PCG/Source/PCGEditor/Private/PCGEditorGraphNodeInput.cpp b/Engine/Plugins/PCG/Source/PCGEditor/Private/PCGEditorGraphNodeInput.cpp
deleted file mode 100644
diff --git a/Engine/Plugins/PCG/Source/PCGEditor/Private/PCGEditorGraphNodeInput.h b/Engine/Plugins/PCG/Source/PCGEditor/Private/PCGEditorGraphNodeInput.h
deleted file mode 100644
diff --git a/Engine/Plugins/PCG/Source/PCGEditor/Private/PCGEditorGraphNodeOutput.cpp b/Engine/Plugins/PCG/Source/PCGEditor/Private/PCGEditorGraphNodeOutput.cpp
deleted file mode 100644
diff --git a/Engine/Plugins/PCG/Source/PCGEditor/Private/PCGEditorGraphNodeOutput.h b/Engine/Plugins/PCG/Source/PCGEditor/Private/PCGEditorGraphNodeOutput.h
deleted file mode 100644
diff --git a/Engine/Plugins/PCG/Source/PCGEditor/Private/PCGEditorGraphNodeReroute.cpp b/Engine/Plugins/PCG/Source/PCGEditor/Private/PCGEditorGraphNodeReroute.cpp
deleted file mode 100644
diff --git a/Engine/Plugins/PCG/Source/PCGEditor/Private/PCGEditorGraphNodeReroute.h b/Engine/Plugins/PCG/Source/PCGEditor/Private/PCGEditorGraphNodeReroute.h
deleted file mode 100644
diff --git a/Engine/Plugins/PCG/Source/PCGEditor/Private/PCGEditorGraphSchema.cpp b/Engine/Plugins/PCG/Source/PCGEditor/Private/PCGEditorGraphSchema.cpp
--- a/Engine/Plugins/PCG/Source/PCGEditor/Private/PCGEditorGraphSchema.cpp
+++ b/Engine/Plugins/PCG/Source/PCGEditor/Private/PCGEditorGraphSchema.cpp
@@ -18,11 +18,11 @@
 #include "PCGEditor.h"
 #include "PCGEditorCommon.h"
 #include "PCGEditorGraph.h"
-#include "PCGEditorGraphNodeBase.h"
-#include "PCGEditorGraphNodeReroute.h"
 #include "PCGEditorGraphSchemaActions.h"
 #include "PCGEditorSettings.h"
 #include "PCGEditorUtils.h"
+#include "Nodes/PCGEditorGraphNodeBase.h"
+#include "Nodes/PCGEditorGraphNodeReroute.h"
 
 #include "AssetRegistry/AssetData.h"
 #include "AssetRegistry/AssetRegistryModule.h"
diff --git a/Engine/Plugins/PCG/Source/PCGEditor/Private/PCGEditorGraphSchemaActions.cpp b/Engine/Plugins/PCG/Source/PCGEditor/Private/PCGEditorGraphSchemaActions.cpp
--- a/Engine/Plugins/PCG/Source/PCGEditor/Private/PCGEditorGraphSchemaActions.cpp
+++ b/Engine/Plugins/PCG/Source/PCGEditor/Private/PCGEditorGraphSchemaActions.cpp
@@ -4,9 +4,9 @@
 
 #include "PCGEditorCommon.h"
 #include "PCGEditorGraph.h"
-#include "PCGEditorGraphNode.h"
-#include "PCGEditorGraphNodeReroute.h"
 #include "PCGEditorModule.h"
+#include "Nodes/PCGEditorGraphNode.h"
+#include "Nodes/PCGEditorGraphNodeReroute.h"
 
 #include "PCGDataAsset.h"
 #include "PCGGraph.h"
diff --git a/Engine/Plugins/PCG/Source/PCGEditor/Private/PCGEditorModule.cpp b/Engine/Plugins/PCG/Source/PCGEditor/Private/PCGEditorModule.cpp
--- a/Engine/Plugins/PCG/Source/PCGEditor/Private/PCGEditorModule.cpp
+++ b/Engine/Plugins/PCG/Source/PCGEditor/Private/PCGEditorModule.cpp
@@ -19,7 +19,6 @@
 #include "WorldPartitionBuilder/PCGWorldPartitionBuilder.h"
 
 #include "PCGEditorCommands.h"
-#include "PCGEditorGraphNodeFactory.h"
 #include "PCGEditorMenuUtils.h"
 #include "PCGEditorProgressNotification.h"
 #include "PCGEditorSettings.h"
@@ -39,6 +38,7 @@
 #include "Details/PCGGraphInstanceDetails.h"
 #include "Details/PCGInstancedPropertyBagOverrideDetails.h"
 #include "Details/PCGVolumeDetails.h"
+#include "Nodes/PCGEditorGraphNodeFactory.h"
 
 #include "ContentBrowserMenuContexts.h"
 #include "ContentBrowserModule.h"
diff --git a/Engine/Plugins/PCG/Source/PCGEditor/Private/SPCGEditorGraphAttributeListView.cpp b/Engine/Plugins/PCG/Source/PCGEditor/Private/SPCGEditorGraphAttributeListView.cpp
--- a/Engine/Plugins/PCG/Source/PCGEditor/Private/SPCGEditorGraphAttributeListView.cpp
+++ b/Engine/Plugins/PCG/Source/PCGEditor/Private/SPCGEditorGraphAttributeListView.cpp
@@ -23,8 +23,8 @@
 
 #include "PCGEditor.h"
 #include "PCGEditorGraph.h"
-#include "PCGEditorGraphNodeBase.h"
 #include "PCGEditorStyle.h"
+#include "Nodes/PCGEditorGraphNodeBase.h"
 
 #include "Fonts/FontMeasure.h"
 #include "Framework/Commands/GenericCommands.h"
diff --git a/Engine/Plugins/PCG/Source/PCGEditor/Private/SPCGEditorGraphDetailsView.cpp b/Engine/Plugins/PCG/Source/PCGEditor/Private/SPCGEditorGraphDetailsView.cpp
--- a/Engine/Plugins/PCG/Source/PCGEditor/Private/SPCGEditorGraphDetailsView.cpp
+++ b/Engine/Plugins/PCG/Source/PCGEditor/Private/SPCGEditorGraphDetailsView.cpp
@@ -7,7 +7,7 @@
 #include "PCGSettings.h"
 
 #include "PCGEditor.h"
-#include "PCGEditorGraphNode.h"
+#include "Nodes/PCGEditorGraphNode.h"
 
 #include "DetailsViewArgs.h"
 #include "IDetailsView.h"
diff --git a/Engine/Plugins/PCG/Source/PCGEditor/Private/SPCGEditorGraphFind.cpp b/Engine/Plugins/PCG/Source/PCGEditor/Private/SPCGEditorGraphFind.cpp
--- a/Engine/Plugins/PCG/Source/PCGEditor/Private/SPCGEditorGraphFind.cpp
+++ b/Engine/Plugins/PCG/Source/PCGEditor/Private/SPCGEditorGraphFind.cpp
@@ -4,16 +4,15 @@
 
 #include "PCGEditor.h"
 #include "PCGEditorGraph.h"
-#include "PCGEditorGraphNodeBase.h"
 #include "PCGEditorStyle.h"
 #include "PCGGraph.h"
 #include "PCGNode.h"
 #include "PCGSubgraph.h"
 #include "PCGSubsystem.h"
+#include "Nodes/PCGEditorGraphNodeBase.h"
 
 #include "EdGraph/EdGraphSchema.h"
 #include "Framework/Application/SlateApplication.h"
-#include "Framework/Views/TableViewMetadata.h"
 #include "Layout/WidgetPath.h"
 #include "Widgets/Input/SComboButton.h"
 #include "Widgets/Input/SSearchBox.h"
diff --git a/Engine/Plugins/PCG/Source/PCGEditor/Private/SPCGEditorGraphLogView.cpp b/Engine/Plugins/PCG/Source/PCGEditor/Private/SPCGEditorGraphLogView.cpp
--- a/Engine/Plugins/PCG/Source/PCGEditor/Private/SPCGEditorGraphLogView.cpp
+++ b/Engine/Plugins/PCG/Source/PCGEditor/Private/SPCGEditorGraphLogView.cpp
@@ -5,13 +5,11 @@
 #include "PCGComponent.h"
 #include "PCGGraph.h"
 #include "PCGSubgraph.h"
-#include "Graph/PCGStackContext.h"
 
 #include "PCGEditor.h"
 #include "PCGEditorGraph.h"
-#include "PCGEditorGraphNode.h"
+#include "Nodes/PCGEditorGraphNode.h"
 
-#include "Framework/Views/TableViewMetadata.h"
 #include "Styling/StyleColors.h"
 #include "Widgets/Input/SButton.h"
 #include "Widgets/Layout/SScrollBox.h"
diff --git a/Engine/Plugins/PCG/Source/PCGEditor/Private/SPCGEditorGraphNode.cpp b/Engine/Plugins/PCG/Source/PCGEditor/Private/SPCGEditorGraphNode.cpp
deleted file mode 100644
diff --git a/Engine/Plugins/PCG/Source/PCGEditor/Private/SPCGEditorGraphNode.h b/Engine/Plugins/PCG/Source/PCGEditor/Private/SPCGEditorGraphNode.h
deleted file mode 100644
diff --git a/Engine/Plugins/PCG/Source/PCGEditor/Private/SPCGEditorGraphProfilingView.cpp b/Engine/Plugins/PCG/Source/PCGEditor/Private/SPCGEditorGraphProfilingView.cpp
--- a/Engine/Plugins/PCG/Source/PCGEditor/Private/SPCGEditorGraphProfilingView.cpp
+++ b/Engine/Plugins/PCG/Source/PCGEditor/Private/SPCGEditorGraphProfilingView.cpp
@@ -7,9 +7,8 @@
 
 #include "PCGEditor.h"
 #include "PCGEditorGraph.h"
-#include "PCGEditorGraphNode.h"
+#include "Nodes/PCGEditorGraphNode.h"
 
-#include "Framework/Views/TableViewMetadata.h"
 #include "HAL/PlatformApplicationMisc.h"
 #include "Widgets/Input/SButton.h"
 #include "Widgets/Input/SCheckBox.h"
diff --git a/Engine/Plugins/PCG/Source/PCGEditor/Private/SPCGEditorGraphVarNode.cpp b/Engine/Plugins/PCG/Source/PCGEditor/Private/SPCGEditorGraphVarNode.cpp
deleted file mode 100644
diff --git a/Engine/Plugins/PCG/Source/PCGEditor/Private/SPCGEditorGraphVarNode.h b/Engine/Plugins/PCG/Source/PCGEditor/Private/SPCGEditorGraphVarNode.h
deleted file mode 100644
