[PCG] GPU - Fixes to read back of element count & inspection

* Fix data counters not correctly updating inspection data for nodes that have no downstream connections, leading to incorrect data counts when inspecting a node
* Fix counters being incorrectly enabled for SM Spawner, which would lead to random counts being applied if output pin is unconnected and node is inspected

#rnx

diff --git a/Engine/Plugins/PCG/Source/PCG/Private/Compute/PCGComputeGraph.h b/Engine/Plugins/PCG/Source/PCG/Private/Compute/PCGComputeGraph.h
--- a/Engine/Plugins/PCG/Source/PCG/Private/Compute/PCGComputeGraph.h
+++ b/Engine/Plugins/PCG/Source/PCG/Private/Compute/PCGComputeGraph.h
@@ -87,10 +87,12 @@ public:
 
 	const TArray<FString>& GetStringTable() const { return StringTable; }
 
-	// in PCG a pin is a data interface (multiple incident edges on an input pin will be merged using a DI).
-	// Possibly we shouldn't directly talk to UPCGDataBinding here
+	// Get description of data produced or consumed by a pin on a kernel.
 	bool ComputeKernelBindingDataDesc(int32 InKernelIndex, FName InPinLabel, bool bInIsInputPin, const UPCGDataBinding* InBinding, FPCGDataCollectionDesc& OutDataDesc) const;
-	bool ComputeKernelBindingDataDesc(int32 InEdgeIndex, const UPCGDataBinding* InBinding, FPCGDataCollectionDesc& OutDataDesc) const;
+
+	// Get description of data produced or consumed by a binding in the graph. In compute graph a binding is a single function that is exposed to a kernel.
+	// in PCG a pin is a data interface (multiple incident edges on an input pin will be merged using a DI).
+	bool ComputeKernelBindingDataDesc(int32 InBindingIndex, const UPCGDataBinding* InBinding, FPCGDataCollectionDesc& OutDataDesc) const;
 
 	/** Bindings in compute graphs roughly map to a single API on a Data Interface, like GetNumData(). There can be multiple bindings per PCG edge. */
 	int32 GetBindingIndex(int32 InKernelIndex, FName InPinLabel, bool bInIsInput) const;
diff --git a/Engine/Plugins/PCG/Source/PCG/Private/Compute/PCGDataBinding.cpp b/Engine/Plugins/PCG/Source/PCG/Private/Compute/PCGDataBinding.cpp
--- a/Engine/Plugins/PCG/Source/PCG/Private/Compute/PCGDataBinding.cpp
+++ b/Engine/Plugins/PCG/Source/PCG/Private/Compute/PCGDataBinding.cpp
@@ -69,7 +69,7 @@ bool UPCGDataBinding::GetKernelPinDataDesc_GameThread(int32 InGraphBindingIndex,
 }
 
 void UPCGDataBinding::ReceiveDataFromGPU_GameThread(
-	const UPCGProxyForGPUData* InData,
+	UPCGProxyForGPUData* InData,
 	const UPCGSettings* InProducerSettings,
 	EPCGExportMode InExportMode,
 	FName InPinLabel,
@@ -86,16 +86,16 @@ void UPCGDataBinding::ReceiveDataFromGPU_GameThread(
 
 	if (!!(InExportMode & EPCGExportMode::DebugVisualization))
 	{
-		TStrongObjectPtr<const UPCGProxyForGPUData> Proxy;
+		TStrongObjectPtr<UPCGProxyForGPUData> Proxy;
 		Proxy.Reset(InData);
-		DataToDebug.Add(FPCGDataToDebug{ .Proxy = Proxy, .ProducerSettings = InProducerSettings, .PinLabel = InPinLabel });
+		DataToDebug.Add(FPCGDataToDebug{ .Proxy = Proxy, .ProducerSettings = InProducerSettings, .PinLabel = InPinLabel, .PinLabelAlias = InPinLabelAlias });
 	}
 
 	if (!!(InExportMode & EPCGExportMode::Inspection))
 	{
-		TStrongObjectPtr<const UPCGProxyForGPUData> Proxy;
+		TStrongObjectPtr<UPCGProxyForGPUData> Proxy;
 		Proxy.Reset(InData);
-		DataToInspect.Add(FPCGDataToDebug{ .Proxy = Proxy, .ProducerSettings = InProducerSettings, .PinLabel = InPinLabel });
+		DataToInspect.Add(FPCGDataToDebug{ .Proxy = Proxy, .ProducerSettings = InProducerSettings, .PinLabel = InPinLabel, .PinLabelAlias = InPinLabelAlias });
 	}
 }
 
@@ -116,6 +116,21 @@ void UPCGDataBinding::ApplyElementCounts_GameThread()
 			}
 		}
 	}
+
+	auto UpdateCountsOnDebugData = [This=this](TArray<FPCGDataToDebug>& DebugData)
+	{
+		for (FPCGDataToDebug& Data : DebugData)
+		{
+			const TArray<uint32>* FoundElementCounts = This->PinToElementCounts.Find(Data.PinLabelAlias);
+			if (FoundElementCounts && Data.Proxy)
+			{
+				Data.Proxy->UpdateElementCountsFromReadback(*FoundElementCounts);
+			}
+		}
+	};
+
+	UpdateCountsOnDebugData(DataToInspect);
+	UpdateCountsOnDebugData(DataToDebug);
 }
 
 void UPCGDataBinding::InitializeInputData(const FPCGDataCollection& InComputeGraphElementInputData)
diff --git a/Engine/Plugins/PCG/Source/PCG/Private/Compute/PCGDataBinding.h b/Engine/Plugins/PCG/Source/PCG/Private/Compute/PCGDataBinding.h
--- a/Engine/Plugins/PCG/Source/PCG/Private/Compute/PCGDataBinding.h
+++ b/Engine/Plugins/PCG/Source/PCG/Private/Compute/PCGDataBinding.h
@@ -48,9 +48,10 @@ struct FPCGSpawnerPrimitives
 
 struct FPCGDataToDebug
 {
-	TStrongObjectPtr<const UPCGProxyForGPUData> Proxy;
+	TStrongObjectPtr<UPCGProxyForGPUData> Proxy;
 	TWeakObjectPtr<const UPCGSettings> ProducerSettings;
 	FName PinLabel;
+	FName PinLabelAlias;
 };
 
 UCLASS(Transient, Category = PCG)
@@ -71,7 +72,7 @@ public:
 	bool GetKernelPinDataDesc_GameThread(const UPCGComputeKernel* InKernel, FName InPinLabel, bool bIsInput, FPCGDataCollectionDesc& OutDescription) const;
 	bool GetKernelPinDataDesc_GameThread(int32 InGraphBindingIndex, FPCGDataCollectionDesc& OutDescription) const;
 
-	void ReceiveDataFromGPU_GameThread(const UPCGProxyForGPUData* InData, const UPCGSettings* InProducerSettings, EPCGExportMode InExportMode, FName InPinLabel, FName InPinLabelAlias);
+	void ReceiveDataFromGPU_GameThread(UPCGProxyForGPUData* InData, const UPCGSettings* InProducerSettings, EPCGExportMode InExportMode, FName InPinLabel, FName InPinLabelAlias);
 
 	void ReceiveElementCountsFromGPU_GameThread(FName InPinLabel, const TArray<uint32>& InCounters);
 
@@ -113,10 +114,10 @@ public:
 	TArray<FPCGDataToDebug> DataToDebug;
 	TArray<FPCGDataToDebug> DataToInspect;
 
+protected:
 	/** Final element counts read back from GPU. */
 	TMap<FName, TArray<uint32>> PinToElementCounts;
 
-protected:
 	/** All attributes present in graph at execution time. Coherent across all graph branches. Seeded using statically known created attributes at compile time, then
 	 * augmented with incoming attributes from input data collection at runtime. */
 	UPROPERTY()
diff --git a/Engine/Plugins/PCG/Source/PCG/Private/Graph/PCGGraphCompilerGPU.cpp b/Engine/Plugins/PCG/Source/PCG/Private/Graph/PCGGraphCompilerGPU.cpp
--- a/Engine/Plugins/PCG/Source/PCG/Private/Graph/PCGGraphCompilerGPU.cpp
+++ b/Engine/Plugins/PCG/Source/PCG/Private/Graph/PCGGraphCompilerGPU.cpp
@@ -397,40 +397,20 @@ void FPCGGraphCompilerGPU::SetupVirtualPins(
 			}
 		}
 
-		if (!InTaskSuccessors.Contains(GPUTaskId))
-		{
-			continue;
-		}
-
-		// Next consider output edges of the compute graph.
-		for (FPCGTaskId Successor : InTaskSuccessors[GPUTaskId])
+		// Create virtual pin labels for all output pins of GPU nodes. We could create them only for output pins that internal->external edges, however
+		// it is helpful for special cases like inspection-only to always have virtual labels for all external output pins.
+		// TODO update to use kernels once the GPU compiler flow is refactored.
+		if (const UPCGSettings* Settings = InCompiledTasks[GPUTaskId].Node ? InCompiledTasks[GPUTaskId].Node->GetSettings() : nullptr)
 		{
-			if (InCollapsedTasks.Contains(Successor))
-			{
-				continue;
-			}
-
-			// Rewire inputs of this downstream CPU node to the outputs of the compute graph task.
-			const FPCGGraphTask& DownstreamCPUNode = InCompiledTasks[Successor];
-
-			for (int SuccessorInputIndex = 0; SuccessorInputIndex < DownstreamCPUNode.Inputs.Num(); ++SuccessorInputIndex)
+			for (const FPCGPinProperties& PinProps : Settings->OutputPinProperties())
 			{
-				// Skip irrelevant edges.
-				if (DownstreamCPUNode.Inputs[SuccessorInputIndex].TaskId != GPUTaskId)
+				const FNodePin PinKey = { GPUTaskId, PinProps.Label, /*Pin is input*/false };
+				if (!OutOriginalToVirtualPin.Contains(PinKey))
 				{
-					continue;
-				}
+					const FName VirtualLabel = *FString::Format(TEXT("{0}-VirtualOut{1}"), { PinProps.Label.ToString(), OutputCount });
+					OutOriginalToVirtualPin.Add(PinKey, VirtualLabel);
 
-				if (DownstreamCPUNode.Inputs[SuccessorInputIndex].UpstreamPin.IsSet())
-				{
-					const FNodePin PinKey = { GPUTaskId, DownstreamCPUNode.Inputs[SuccessorInputIndex].UpstreamPin->Label, /*Pin is input*/false };
-					if (!OutOriginalToVirtualPin.Contains(PinKey))
-					{
-						const FName VirtualLabel = *FString::Format(TEXT("{0}-VirtualOut{1}"), { DownstreamCPUNode.Inputs[SuccessorInputIndex].UpstreamPin->Label.ToString(), OutputCount });
-						OutOriginalToVirtualPin.Add(PinKey, VirtualLabel);
-
-						++OutputCount;
-					}
+					++OutputCount;
 				}
 			}
 		}
diff --git a/Engine/Plugins/PCG/Source/PCG/Public/Elements/PCGStaticMeshSpawnerKernel.h b/Engine/Plugins/PCG/Source/PCG/Public/Elements/PCGStaticMeshSpawnerKernel.h
--- a/Engine/Plugins/PCG/Source/PCG/Public/Elements/PCGStaticMeshSpawnerKernel.h
+++ b/Engine/Plugins/PCG/Source/PCG/Public/Elements/PCGStaticMeshSpawnerKernel.h
@@ -24,7 +24,6 @@ public:
 #endif
 	virtual void AddStaticCreatedStrings(TArray<FString>& InOutStringTable) const override;
 	virtual void GetKernelAttributeKeys(TArray<FPCGKernelAttributeKey>& OutKeys) const override;
-	virtual bool DoesOutputPinRequireElementCounters(FName InOutputPinLabel) const override { return true; }
 
 	/** Setup Procedural ISMCs. Returns true if at least one component was set up. */
 	virtual bool SetupProceduralISMComponents(FPCGContext* InContext, UPCGDataBinding* InBinding) const;
