[PCG] Support tags in data descriptions
* Allows tags to flow through GPU graphs end to end (CPU -> GPU -> GPU -> CPU)
* When combining data descriptions, the resulting description will have all of the unique tags from the combined descriptions
* Note that tags are never uploaded to GPU buffers and are not accessible from HLSL
* Data labels are now supported on output pins and between GPU nodes
* Fixed bug where data label resolvers broke when a GPU node had multiple input pins connected

#rnx

diff --git a/Engine/Plugins/PCG/Source/PCG/Internal/Compute/Data/PCGProxyForGPUData.h b/Engine/Plugins/PCG/Source/PCG/Internal/Compute/Data/PCGProxyForGPUData.h
--- a/Engine/Plugins/PCG/Source/PCG/Internal/Compute/Data/PCGProxyForGPUData.h
+++ b/Engine/Plugins/PCG/Source/PCG/Internal/Compute/Data/PCGProxyForGPUData.h
@@ -38,7 +38,7 @@ public:
 		bool bComplete = false;
 
 		/** The data created from readback, null if read back failed. */
-		const UPCGData* Data = nullptr;
+		FPCGTaggedData TaggedData;
 	};
 
 	/** Populates a CPU data object representing the GPU data, performing a readback from GPU->CPU if required. */
diff --git a/Engine/Plugins/PCG/Source/PCG/Internal/Compute/PCGComputeCommon.h b/Engine/Plugins/PCG/Source/PCG/Internal/Compute/PCGComputeCommon.h
--- a/Engine/Plugins/PCG/Source/PCG/Internal/Compute/PCGComputeCommon.h
+++ b/Engine/Plugins/PCG/Source/PCG/Internal/Compute/PCGComputeCommon.h
@@ -13,6 +13,7 @@ struct FPCGContext;
 struct FPCGDataCollectionDesc;
 struct FPCGDataDesc;
 struct FPCGPinPropertiesGPU;
+struct FPCGTaggedData;
 
 #define PCG_KERNEL_LOGGING_ENABLED (!(UE_BUILD_SHIPPING || UE_BUILD_TEST) || USE_LOGGING_IN_SHIPPING)
 
@@ -111,7 +112,7 @@ namespace PCGComputeHelpers
 	/** Returns true if the given buffer size is dangerously large. Optionally emits error log. */
 	bool IsBufferSizeTooLarge(uint64 InBufferSizeBytes, bool bInLogError = true);
 
-	bool ComputeDataDescription(const UPCGData* InData, const UPCGDataBinding* InBinding, FPCGDataDesc& OutDescription);
+	bool ComputeDataDescription(const FPCGTaggedData& InTaggedData, const UPCGDataBinding* InBinding, FPCGDataDesc& OutDescription);
 
 	int32 GetAttributeIdFromMetadataAttributeIndex(int32 InAttributeIndex);
 	int32 GetMetadataAttributeIndexFromAttributeId(int32 InAttributeId);
diff --git a/Engine/Plugins/PCG/Source/PCG/Internal/Compute/PCGDataForGPU.h b/Engine/Plugins/PCG/Source/PCG/Internal/Compute/PCGDataForGPU.h
--- a/Engine/Plugins/PCG/Source/PCG/Internal/Compute/PCGDataForGPU.h
+++ b/Engine/Plugins/PCG/Source/PCG/Internal/Compute/PCGDataForGPU.h
@@ -147,7 +147,7 @@ struct FPCGDataDesc
 public:
 	FPCGDataDesc() = default;
 	FPCGDataDesc(EPCGDataType InType, int32 InElementCount);
-	FPCGDataDesc(const UPCGData* InData, const UPCGDataBinding* InBinding);
+	FPCGDataDesc(const FPCGTaggedData& InTaggedData, const UPCGDataBinding* InBinding);
 
 	uint64 ComputePackedSize() const;
 
@@ -161,6 +161,7 @@ public:
 	EPCGDataType Type = EPCGDataType::Point;
 	TArray<FPCGKernelAttributeDesc> AttributeDescs;
 	int32 ElementCount = 0;
+	TArray<int32, TInlineAllocator<4>> TagStringKeys;
 
 private:
 	void InitializeAttributeDescs(const UPCGData* InData, const UPCGDataBinding* InBinding);
@@ -187,7 +188,7 @@ struct FPCGDataCollectionDesc
 	void WriteHeader(TArray<uint32>& OutPackedDataCollectionHeader) const;
 
 	/** Pack a data collection into the GPU data format. DataDescs defines which attributes are packed. */
-	void PackDataCollection(const FPCGDataCollection& InDataCollection, FName InPin, const TArray<FString>& InStringTable, TArray<uint32>& OutPackedDataCollection) const;
+	void PackDataCollection(const FPCGDataCollection& InDataCollection, FName InPin, const UPCGDataBinding* InDataBinding, TArray<uint32>& OutPackedDataCollection) const;
 
 	/** Unpack a buffer of 8-bit uints to a data collection. */
 	EPCGUnpackDataCollectionResult UnpackDataCollection(FPCGContext* InContext, const TArray<uint8>& InPackedData, FName InPin, const TArray<FString>& InStringTable, FPCGDataCollection& OutDataCollection) const;
@@ -220,7 +221,7 @@ public:
 	FPCGProxyForGPUDataCollection(TRefCountPtr<FRDGPooledBuffer> InBuffer, uint32 InBufferSizeBytes, const FPCGDataCollectionDesc& InDescription, const TArray<FString>& InStringTable);
 
 	/** Populates a CPU data object representing the GPU data for the given index, performing a readback from GPU->CPU if required. */
-	bool GetCPUData(FPCGContext* InContext, int32 InDataIndex, const UPCGData*& OutData);
+	bool GetCPUData(FPCGContext* InContext, int32 InDataIndex, FPCGTaggedData& OutData);
 
 	const TRefCountPtr<FRDGPooledBuffer>& GetBuffer() const { return Buffer; }
 
@@ -244,7 +245,8 @@ protected:
 	TArray<FString> StringTable;
 
 	/** Read back data. Populated once upon first readback request. */
-	TArray<TStrongObjectPtr<const UPCGData>> ReadbackData;
+	TArray<const FPCGTaggedData> ReadbackData;
+	TArray<TStrongObjectPtr<const UPCGData>> ReadbackDataRefs;
 
 	TSharedPtr<FRHIGPUBufferReadback> ReadbackRequest;
 
diff --git a/Engine/Plugins/PCG/Source/PCG/Private/Compute/Data/PCGProxyForGPUData.cpp b/Engine/Plugins/PCG/Source/PCG/Private/Compute/Data/PCGProxyForGPUData.cpp
--- a/Engine/Plugins/PCG/Source/PCG/Private/Compute/Data/PCGProxyForGPUData.cpp
+++ b/Engine/Plugins/PCG/Source/PCG/Private/Compute/Data/PCGProxyForGPUData.cpp
@@ -69,16 +69,16 @@ UPCGProxyForGPUData::FReadbackResult UPCGProxyForGPUData::GetCPUData(FPCGContext
 	{
 		UE_LOG(LogPCG, Error, TEXT("Data collection lost!"));
 
-		return FReadbackResult{ .bComplete = true, .Data = nullptr };
+		return FReadbackResult{ .bComplete = true };
 	}
 
-	const UPCGData* ResultData = nullptr;
+	FPCGTaggedData ResultData;
 	if (!DataOnGPU->GetCPUData(InContext, DataIndexInCollection, ResultData))
 	{
-		return FReadbackResult{ .bComplete = false, .Data = nullptr };
+		return FReadbackResult{ .bComplete = false };
 	}
 
-	return FReadbackResult{ .bComplete = true, .Data = ResultData };
+	return FReadbackResult{ .bComplete = true, .TaggedData = MoveTemp(ResultData) };
 }
 
 TSharedPtr<const FPCGProxyForGPUDataCollection> UPCGProxyForGPUData::GetGPUInfo() const
diff --git a/Engine/Plugins/PCG/Source/PCG/Private/Compute/DataInterfaces/PCGDataLabelResolverDataInterface.cpp b/Engine/Plugins/PCG/Source/PCG/Private/Compute/DataInterfaces/PCGDataLabelResolverDataInterface.cpp
--- a/Engine/Plugins/PCG/Source/PCG/Private/Compute/DataInterfaces/PCGDataLabelResolverDataInterface.cpp
+++ b/Engine/Plugins/PCG/Source/PCG/Private/Compute/DataInterfaces/PCGDataLabelResolverDataInterface.cpp
@@ -5,6 +5,7 @@
 #include "PCGModule.h"
 #include "PCGSettings.h"
 #include "Compute/PCGComputeGraph.h"
+#include "Compute/PCGComputeKernel.h"
 
 #include "RenderGraphBuilder.h"
 #include "RenderGraphResources.h"
@@ -69,16 +70,22 @@ void UPCGDataLabelResolverDataInterface::GetHLSL(FString& OutHLSL, FString const
 UComputeDataProvider* UPCGDataLabelResolverDataInterface::CreateDataProvider(TObjectPtr<UObject> InBinding, uint64 InInputMask, uint64 InOutputMask) const
 {
 	TRACE_CPUPROFILER_EVENT_SCOPE(UPCGDataLabelResolverDataInterface::CreateDataProvider);
-	check(ConsumerSettings);
+
+	check(Kernel);
 
 	UPCGDataBinding* Binding = CastChecked<UPCGDataBinding>(InBinding);
-	const UPCGComputeGraph* ComputeGraph = Binding ? Binding->Graph.Get() : nullptr;
+	const UPCGComputeGraph* ComputeGraph = Binding->Graph.Get();
 	check(ComputeGraph);
 
+	const TArray<FString>& StringTable = Binding->GetStringTable();
+
+	FPCGDataCollectionDesc PinDataDesc;
+	ensure(Binding->GetKernelPinDataDesc_GameThread(Kernel, PinLabel, bIsInput, PinDataDesc));
+
 	UPCGDataLabelResolverDataProvider* Provider = NewObject<UPCGDataLabelResolverDataProvider>();
 
-	const FPCGPinDataLabels* FoundPinDataLabels = ComputeGraph->GetDataLabelsStaticTable().Find(ConsumerSettings);
-	if (const FPCGDataLabels* FoundLabels = FoundPinDataLabels ? FoundPinDataLabels->PinToDataLabels.Find(InputPinLabel) : nullptr)
+	const FPCGPinDataLabels* FoundPinDataLabels = ComputeGraph->GetStaticDataLabelsTable().Find(Kernel->GetKernelIndex());
+	if (const FPCGDataLabels* FoundLabels = FoundPinDataLabels ? FoundPinDataLabels->PinToDataLabels.Find(PinLabel) : nullptr)
 	{
 		TArray<int32> DataIdToDataIndexMap;
 		DataIdToDataIndexMap.SetNumUninitialized(FoundLabels->Labels.Num());
@@ -91,14 +98,15 @@ UComputeDataProvider* UPCGDataLabelResolverDataInterface::CreateDataProvider(TOb
 			// Initialize with invalid indices. The data IDs are not ordered.
 			DataIdToDataIndexMap[LabelIndex] = INDEX_NONE;
 
-			// TODO: In order to support data labels on GPU <-> GPU edges, we have to get the tags from data descriptions
-			// instead of from the CPU tagged data.
-			for (int32 InputDataIndex = 0; InputDataIndex < Binding->InputDataCollection.TaggedData.Num(); ++InputDataIndex)
+			for (int32 InputDataIndex = 0; InputDataIndex < PinDataDesc.DataDescs.Num(); ++InputDataIndex)
 			{
-				if (Binding->InputDataCollection.TaggedData[InputDataIndex].Tags.Contains(PrefixedLabel))
+				for (const int32 TagStringKey : PinDataDesc.DataDescs[InputDataIndex].TagStringKeys)
 				{
-					DataIdToDataIndexMap[LabelIndex] = InputDataIndex;
-					break;
+					if (StringTable[TagStringKey] == PrefixedLabel)
+					{
+						DataIdToDataIndexMap[LabelIndex] = InputDataIndex;
+						break;
+					}
 				}
 			}
 
diff --git a/Engine/Plugins/PCG/Source/PCG/Private/Compute/DataInterfaces/PCGDataLabelResolverDataInterface.h b/Engine/Plugins/PCG/Source/PCG/Private/Compute/DataInterfaces/PCGDataLabelResolverDataInterface.h
--- a/Engine/Plugins/PCG/Source/PCG/Private/Compute/DataInterfaces/PCGDataLabelResolverDataInterface.h
+++ b/Engine/Plugins/PCG/Source/PCG/Private/Compute/DataInterfaces/PCGDataLabelResolverDataInterface.h
@@ -28,11 +28,17 @@ public:
 	UComputeDataProvider* CreateDataProvider(TObjectPtr<UObject> InBinding, uint64 InInputMask, uint64 InOutputMask) const override;
 	//~ End UComputeDataInterface Interface
 
+	/** Label of the pin being resolved. */
 	UPROPERTY()
-	TObjectPtr<const UPCGSettings> ConsumerSettings = nullptr;
+	FName PinLabel;
 
+	/** Whether the pin is an input or output pin. */
 	UPROPERTY()
-	FName InputPinLabel;
+	bool bIsInput = false;
+
+	/** Kernel that owns the pin being resolved. */
+	UPROPERTY()
+	TObjectPtr<const UPCGComputeKernel> Kernel = nullptr;
 };
 
 UCLASS()
diff --git a/Engine/Plugins/PCG/Source/PCG/Private/Compute/Elements/PCGComputeGraphElement.cpp b/Engine/Plugins/PCG/Source/PCG/Private/Compute/Elements/PCGComputeGraphElement.cpp
--- a/Engine/Plugins/PCG/Source/PCG/Private/Compute/Elements/PCGComputeGraphElement.cpp
+++ b/Engine/Plugins/PCG/Source/PCG/Private/Compute/Elements/PCGComputeGraphElement.cpp
@@ -158,7 +158,7 @@ bool FPCGComputeGraphElement::ExecuteInternal(FPCGContext* InContext) const
 		UPCGDataBinding* DataBinding = FPCGContext::NewObject_AnyThread<UPCGDataBinding>(Context);
 		Context->DataBinding.Reset(DataBinding);
 
-		Context->DataBinding->Initialize(Context->ComputeGraph.Get(), Context->SourceComponent, Context->InputData, Context->ComputeGraph->GetAttributeStaticTable());
+		Context->DataBinding->Initialize(Context->ComputeGraph.Get(), Context->SourceComponent, Context->InputData);
 
 		// Perform validation after input data is initialized.
 		if (!Context->ComputeGraph->IsGraphValid(InContext, /*bQuiet=*/false))
@@ -479,11 +479,11 @@ bool FPCGComputeGraphElement::ExecuteDebugDraw(FPCGComputeGraphContext* InContex
 			UPCGProxyForGPUData::FReadbackResult Readback = Proxy->GetCPUData(InContext);
 			if (Readback.bComplete)
 			{
-				if (Readback.Data)
+				if (Readback.TaggedData.Data)
 				{
-					if (const IPCGDataVisualization* DataVis = DataVisRegistry.GetDataVisualization(Readback.Data->GetClass()))
+					if (const IPCGDataVisualization* DataVis = DataVisRegistry.GetDataVisualization(Readback.TaggedData.Data->GetClass()))
 					{
-						DataVis->ExecuteDebugDisplay(InContext, ProducerSettings, Readback.Data, InContext->GetTargetActor(nullptr));
+						DataVis->ExecuteDebugDisplay(InContext, ProducerSettings, Readback.TaggedData.Data, InContext->GetTargetActor(nullptr));
 					}
 				}
 
@@ -553,13 +553,12 @@ bool FPCGComputeGraphElement::StoreDataForInspection(FPCGComputeGraphContext* In
 		if (Proxy && ProducerSettings)
 		{
 			UPCGProxyForGPUData::FReadbackResult Readback = Proxy->GetCPUData(InContext);
-			if (ensure(Readback.bComplete && Readback.Data))
+			if (ensure(Readback.bComplete && Readback.TaggedData.Data))
 			{
 				FPCGDataCollection& DataCollection = SettingsToDataCollection.FindOrAdd(ProducerSettings);
-				FPCGTaggedData Data;
-				Data.Data = Readback.Data;
-				Data.Pin = DataToInspect.PinLabel;
-				DataCollection.TaggedData.Add(Data);
+				FPCGTaggedData TaggedData = Readback.TaggedData;
+				TaggedData.Pin = DataToInspect.PinLabel;
+				DataCollection.TaggedData.Add(TaggedData);
 			}
 		}
 	}
diff --git a/Engine/Plugins/PCG/Source/PCG/Private/Compute/Elements/PCGCustomHLSLKernel.cpp b/Engine/Plugins/PCG/Source/PCG/Private/Compute/Elements/PCGCustomHLSLKernel.cpp
--- a/Engine/Plugins/PCG/Source/PCG/Private/Compute/Elements/PCGCustomHLSLKernel.cpp
+++ b/Engine/Plugins/PCG/Source/PCG/Private/Compute/Elements/PCGCustomHLSLKernel.cpp
@@ -553,14 +553,14 @@ FString UPCGCustomHLSLKernel::GetCookedSource(const UPCGComputeGraph* InComputeG
 	for (const FPCGKernelAttributeKey& AttributeKey : Settings->KernelAttributeKeys)
 	{
 		const FString SourceDefinition = PCGCustomHLSLKernel::GetKernelAttributeKeyAsString(AttributeKey);
-		const FString AttributeIndexAsString = FString::FromInt(InComputeGraph->GetAttributeStaticTable().GetAttributeId(AttributeKey));
+		const FString AttributeIndexAsString = FString::FromInt(InComputeGraph->GetStaticAttributeTable().GetAttributeId(AttributeKey));
 
 		Source.ReplaceInline(*SourceDefinition, *AttributeIndexAsString);
 		Functions.ReplaceInline(*SourceDefinition, *AttributeIndexAsString);
 	}
 
 	// We inject data IDs directly into the source. They will get remapped to data indices using a label resolver data interface.
-	if (const FPCGPinDataLabels* PinToDataLabels = InComputeGraph->GetDataLabelsStaticTable().Find(Settings))
+	if (const FPCGPinDataLabels* PinToDataLabels = InComputeGraph->GetStaticDataLabelsTable().Find(KernelIndex))
 	{
 		for (const TPair<FName, FPCGDataLabels>& Pair : PinToDataLabels->PinToDataLabels)
 		{
diff --git a/Engine/Plugins/PCG/Source/PCG/Private/Compute/PCGComputeCommon.cpp b/Engine/Plugins/PCG/Source/PCG/Private/Compute/PCGComputeCommon.cpp
--- a/Engine/Plugins/PCG/Source/PCG/Private/Compute/PCGComputeCommon.cpp
+++ b/Engine/Plugins/PCG/Source/PCG/Private/Compute/PCGComputeCommon.cpp
@@ -119,9 +119,9 @@ namespace PCGComputeHelpers
 		return bBufferTooLarge;
 	}
 
-	bool ComputeDataDescription(const UPCGData* InData, const UPCGDataBinding* InBinding, FPCGDataDesc& OutDescription)
+	bool ComputeDataDescription(const FPCGTaggedData& InTaggedData, const UPCGDataBinding* InBinding, FPCGDataDesc& OutDescription)
 	{
-		if (const UPCGProxyForGPUData* Proxy = Cast<UPCGProxyForGPUData>(InData))
+		if (const UPCGProxyForGPUData* Proxy = Cast<UPCGProxyForGPUData>(InTaggedData.Data))
 		{
 			if (!Proxy->GetDescription(OutDescription))
 			{
@@ -149,7 +149,7 @@ namespace PCGComputeHelpers
 		}
 		else
 		{
-			OutDescription = FPCGDataDesc(InData, InBinding);
+			OutDescription = FPCGDataDesc(InTaggedData, InBinding);
 		}
 
 		return true;
diff --git a/Engine/Plugins/PCG/Source/PCG/Private/Compute/PCGComputeGraph.h b/Engine/Plugins/PCG/Source/PCG/Private/Compute/PCGComputeGraph.h
--- a/Engine/Plugins/PCG/Source/PCG/Private/Compute/PCGComputeGraph.h
+++ b/Engine/Plugins/PCG/Source/PCG/Private/Compute/PCGComputeGraph.h
@@ -12,6 +12,7 @@
 
 #include "PCGComputeGraph.generated.h"
 
+class UPCGComputeKernel;
 class UPCGDataBinding;
 class UPCGNode;
 class UPCGPin;
@@ -79,13 +80,13 @@ public:
 	//~End UComputeGraph interface
 
 	/** Get the attribute table collated statically during compilation. */
-	const FPCGKernelAttributeTable& GetAttributeStaticTable() const { return AttributeStaticTable; }
+	const FPCGKernelAttributeTable& GetStaticAttributeTable() const { return StaticAttributeTable; }
 
-	// TODO: Should be map from BindingIndex->Labels, but needs GPU compiler flow tweaked so that kernel index are available at DI creation time.
 	/** Get the data labels collated statically during compilation. */
-	const TMap<TSoftObjectPtr<const UPCGSettings>, FPCGPinDataLabels>& GetDataLabelsStaticTable() const { return DataLabelsStaticTable; }
+	const TMap</*KernelIndex*/int32, FPCGPinDataLabels>& GetStaticDataLabelsTable() const { return StaticDataLabelsTable; }
 
-	const TArray<FString>& GetStringTable() const { return StringTable; }
+	/** Get the strings collated statically during compilation. */
+	const TArray<FString>& GetStaticStringTable() const { return StringTable; }
 
 	// Get description of data produced or consumed by a pin on a kernel.
 	bool ComputeKernelBindingDataDesc(int32 InKernelIndex, FName InPinLabel, bool bInIsInputPin, const UPCGDataBinding* InBinding, FPCGDataCollectionDesc& OutDataDesc) const;
@@ -122,11 +123,11 @@ protected:
 
 	/** Global attribute information collated during compilation. */
 	UPROPERTY()
-	FPCGKernelAttributeTable AttributeStaticTable;
+	FPCGKernelAttributeTable StaticAttributeTable;
 
 	/** Global data label information collated during compilation. */
 	UPROPERTY()
-	TMap<TSoftObjectPtr<const UPCGSettings>, FPCGPinDataLabels> DataLabelsStaticTable;
+	TMap</*KernelIndex*/int32, FPCGPinDataLabels> StaticDataLabelsTable;
 
 	/** String table collated during compilation. */
 	UPROPERTY()
diff --git a/Engine/Plugins/PCG/Source/PCG/Private/Compute/PCGComputeKernel.cpp b/Engine/Plugins/PCG/Source/PCG/Private/Compute/PCGComputeKernel.cpp
--- a/Engine/Plugins/PCG/Source/PCG/Private/Compute/PCGComputeKernel.cpp
+++ b/Engine/Plugins/PCG/Source/PCG/Private/Compute/PCGComputeKernel.cpp
@@ -116,6 +116,12 @@ void UPCGComputeKernel::ComputeDataDescFromPinProperties(const FPCGPinProperties
 						AddAttributesFromData(InputDataDesc, DataDesc);
 					}
 
+					// Add unique tags from this data.
+					for (const int32 TagStringKey : InputDataDesc.TagStringKeys)
+					{
+						DataDesc.TagStringKeys.AddUnique(TagStringKey);
+					}
+
 					DataDesc.ElementCount = TotalNumElements;
 				}
 			}
@@ -134,6 +140,8 @@ void UPCGComputeKernel::ComputeDataDescFromPinProperties(const FPCGPinProperties
 					{
 						AddAttributesFromData(InputDataDesc, DataDesc);
 					}
+
+					DataDesc.TagStringKeys = InputDataDesc.TagStringKeys;
 				}
 			}
 			// Take pairs of datas, where the pairs are given by each data of each pin to each data of every other pin.
@@ -177,10 +185,24 @@ void UPCGComputeKernel::ComputeDataDescFromPinProperties(const FPCGPinProperties
 									checkNoEntry();
 								}
 
-								if (ensure(DataDesc) && Props.AttributeInheritanceMode == EPCGAttributeInheritanceMode::CopyAttributeSetup)
+								if (ensure(DataDesc))
 								{
-									AddAttributesFromData(InputDataDesc, *DataDesc);
-									AddAttributesFromData(OtherInputDataDesc, *DataDesc);
+									if (Props.AttributeInheritanceMode == EPCGAttributeInheritanceMode::CopyAttributeSetup)
+									{
+										AddAttributesFromData(InputDataDesc, *DataDesc);
+										AddAttributesFromData(OtherInputDataDesc, *DataDesc);
+									}
+
+									// Add unique tags from both input data.
+									for (const int32 TagStringKey : InputDataDesc.TagStringKeys)
+									{
+										DataDesc->TagStringKeys.AddUnique(TagStringKey);
+									}
+
+									for (const int32 TagStringKey : OtherInputDataDesc.TagStringKeys)
+									{
+										DataDesc->TagStringKeys.AddUnique(TagStringKey);
+									}
 								}
 							}
 						}
diff --git a/Engine/Plugins/PCG/Source/PCG/Private/Compute/PCGDataBinding.cpp b/Engine/Plugins/PCG/Source/PCG/Private/Compute/PCGDataBinding.cpp
--- a/Engine/Plugins/PCG/Source/PCG/Private/Compute/PCGDataBinding.cpp
+++ b/Engine/Plugins/PCG/Source/PCG/Private/Compute/PCGDataBinding.cpp
@@ -15,8 +15,7 @@
 void UPCGDataBinding::Initialize(
 	const UPCGComputeGraph* InComputeGraph,
 	TWeakObjectPtr<UPCGComponent> InSourceComponent,
-	const FPCGDataCollection& InComputeGraphElementInputData,
-	const FPCGKernelAttributeTable& InAttributeStaticTable)
+	const FPCGDataCollection& InComputeGraphElementInputData)
 {
 	TRACE_CPUPROFILER_EVENT_SCOPE(UPCGDataBinding::Initialize);
 	check(InComputeGraph);
@@ -24,7 +23,11 @@ void UPCGDataBinding::Initialize(
 	Graph = InComputeGraph;
 	SourceComponent = InSourceComponent;
 
-	AttributeTable = InAttributeStaticTable;
+	AttributeTable = InComputeGraph->GetStaticAttributeTable();
+
+	// String table always contains empty string in index 0 (and string key attributes are 0-initialized), and then any strings known statically at compilation time.
+	StringTable = { FString() };
+	StringTable.Append(Graph->GetStaticStringTable());
 
 	InitializeInputData(InComputeGraphElementInputData);
 
@@ -151,15 +154,15 @@ void UPCGDataBinding::InitializeInputData(const FPCGDataCollection& InComputeGra
 		}
 	}
 
-	// Augment static string table with attributes incoming from CPU.
-	AddInputDataAttributesToTable(InComputeGraphElementInputData);
-
-	BuildStringTable();
+	// Augment static tables with data incoming from CPU.
+	AddInputDataAttributesToTable();
+	AddInputDataStringsToTable();
+	AddInputDataTagsToTable();
 }
 
-void UPCGDataBinding::AddInputDataAttributesToTable(const FPCGDataCollection& InComputeGraphElementInputData)
+void UPCGDataBinding::AddInputDataAttributesToTable()
 {
-	for (const FPCGTaggedData& Data : InComputeGraphElementInputData.TaggedData)
+	for (const FPCGTaggedData& Data : InputDataCollection.TaggedData)
 	{
 		if (const UPCGProxyForGPUData* Proxy = Cast<UPCGProxyForGPUData>(Data.Data))
 		{
@@ -204,17 +207,11 @@ void UPCGDataBinding::AddInputDataAttributesToTable(const FPCGDataCollection& In
 	}
 }
 
-void UPCGDataBinding::BuildStringTable()
+void UPCGDataBinding::AddInputDataStringsToTable()
 {
-	TRACE_CPUPROFILER_EVENT_SCOPE(UPCGDataBinding::BuildStringTable);
+	TRACE_CPUPROFILER_EVENT_SCOPE(UPCGDataBinding::AddInputDataStringsToTable);
 	check(Graph);
 
-	// String table always contains empty string in index 0 (and string key attributes are 0-initialized).
-	StringTable = { FString() };
-
-	// Add any strings known statically at compilation time.
-	StringTable.Append(Graph->GetStringTable());
-
 	for (const FPCGTaggedData& Data : InputDataCollection.TaggedData)
 	{
 		if (const UPCGProxyForGPUData* Proxy = Cast<UPCGProxyForGPUData>(Data.Data))
@@ -300,6 +297,17 @@ void UPCGDataBinding::BuildStringTable()
 	}
 }
 
+void UPCGDataBinding::AddInputDataTagsToTable()
+{
+	for (const FPCGTaggedData& Data : InputDataCollection.TaggedData)
+	{
+		for (const FString& Tag : Data.Tags)
+		{
+			StringTable.AddUnique(Tag);
+		}
+	}
+}
+
 void UPCGDataBinding::DebugLogDataDescriptions()
 {
 #if PCG_KERNEL_LOGGING_ENABLED
@@ -325,8 +333,20 @@ void UPCGDataBinding::DebugLogDataDescriptions()
 			UE_LOG(LogPCG, Display, TEXT("\t\tData Index %d"), DataIndex);
 			UE_LOG(LogPCG, Display, TEXT("\t\t\tType: %s"), *PCGDataTypeEnum->GetDisplayNameTextByValue(static_cast<int64>(DataDesc.Type)).ToString());
 			UE_LOG(LogPCG, Display, TEXT("\t\t\tElementCount: %d"), DataDesc.ElementCount);
+
+			{
+				FString TagStringKeys;
+				for (int32 TagStringKey : DataDesc.TagStringKeys)
+				{
+					TagStringKeys += TEXT(", ");
+					TagStringKeys += FString::Format(TEXT("{0}"), { TagStringKey });
+				}
+
+				UE_LOG(LogPCG, Display, TEXT("\t\t\tTag String Keys%s"), *TagStringKeys);
+			}
+
 			UE_LOG(LogPCG, Display, TEXT("\t\t\tAttributes (%d)"), DataDesc.AttributeDescs.Num());
-				
+
 			for (int32 AttributeIndex = 0; AttributeIndex < DataDesc.AttributeDescs.Num(); ++AttributeIndex)
 			{
 				const FPCGKernelAttributeDesc& AttributeDesc = DataDesc.AttributeDescs[AttributeIndex];
diff --git a/Engine/Plugins/PCG/Source/PCG/Private/Compute/PCGDataBinding.h b/Engine/Plugins/PCG/Source/PCG/Private/Compute/PCGDataBinding.h
--- a/Engine/Plugins/PCG/Source/PCG/Private/Compute/PCGDataBinding.h
+++ b/Engine/Plugins/PCG/Source/PCG/Private/Compute/PCGDataBinding.h
@@ -60,7 +60,7 @@ class UPCGDataBinding : public UObject
 	GENERATED_BODY()
 
 public:
-	void Initialize(const UPCGComputeGraph* InComputeGraph, TWeakObjectPtr<UPCGComponent> InSourceComponent, const FPCGDataCollection& InComputeGraphElementInputData, const FPCGKernelAttributeTable& InAttributeStaticTable);
+	void Initialize(const UPCGComputeGraph* InComputeGraph, TWeakObjectPtr<UPCGComponent> InSourceComponent, const FPCGDataCollection& InComputeGraphElementInputData);
 
 	int32 GetAttributeId(const FPCGKernelAttributeKey& InAttribute) const { return AttributeTable.GetAttributeId(InAttribute); }
 	int32 GetAttributeId(FName InName, EPCGKernelAttributeType InType) const { return AttributeTable.GetAttributeId(InName, InType); }
@@ -86,9 +86,11 @@ private:
 	void InitializeInputData(const FPCGDataCollection& InComputeGraphElementInputData);
 
 	/** Loop over all metadata attributes in all input data ensure all attributes are registered in attribute table. */
-	void AddInputDataAttributesToTable(const FPCGDataCollection& InComputeGraphElementInputData);
+	void AddInputDataAttributesToTable();
 
-	void BuildStringTable();
+	void AddInputDataStringsToTable();
+
+	void AddInputDataTagsToTable();
 
 	void DebugLogDataDescriptions();
 
diff --git a/Engine/Plugins/PCG/Source/PCG/Private/Compute/PCGDataForGPU.cpp b/Engine/Plugins/PCG/Source/PCG/Private/Compute/PCGDataForGPU.cpp
--- a/Engine/Plugins/PCG/Source/PCG/Private/Compute/PCGDataForGPU.cpp
+++ b/Engine/Plugins/PCG/Source/PCG/Private/Compute/PCGDataForGPU.cpp
@@ -669,14 +669,25 @@ FPCGDataDesc::FPCGDataDesc(EPCGDataType InType, int InElementCount)
 	InitializeAttributeDescs(/*InData=*/nullptr, /*InBinding*/nullptr);
 }
 
-FPCGDataDesc::FPCGDataDesc(const UPCGData* InData, const UPCGDataBinding* InBinding)
+FPCGDataDesc::FPCGDataDesc(const FPCGTaggedData& InTaggedData, const UPCGDataBinding* InBinding)
 {
-	check(InData);
+	check(InTaggedData.Data);
 
-	Type = InData->GetDataType();
-	ElementCount = PCGComputeHelpers::GetElementCount(InData);
+	Type = InTaggedData.Data->GetDataType();
+	ElementCount = PCGComputeHelpers::GetElementCount(InTaggedData.Data);
 
-	InitializeAttributeDescs(InData, InBinding);
+	TagStringKeys.Reserve(InTaggedData.Tags.Num());
+
+	for (const FString& Tag : InTaggedData.Tags)
+	{
+		const int Index = InBinding->GetStringTable().IndexOfByKey(Tag);
+		if (Index != INDEX_NONE)
+		{
+			TagStringKeys.Add(Index);
+		}
+	}
+
+	InitializeAttributeDescs(InTaggedData.Data, InBinding);
 }
 
 uint64 FPCGDataDesc::ComputePackedSize() const
@@ -889,7 +900,7 @@ FPCGDataCollectionDesc FPCGDataCollectionDesc::BuildFromDataCollection(
 			continue;
 		}
 
-		CollectionDesc.DataDescs.Emplace(Data.Data, InBinding);
+		CollectionDesc.DataDescs.Emplace(Data, InBinding);
 	}
 
 	return CollectionDesc;
@@ -905,13 +916,13 @@ FPCGDataCollectionDesc FPCGDataCollectionDesc::BuildFromInputDataCollectionAndIn
 
 	for (const FPCGTaggedData& Data : DataForPin)
 	{
-		if (!Data.Data || !PCGComputeHelpers::IsTypeAllowedInDataCollection(Data.Data->GetDataType()))
+		if (!Data.Data || !PCGComputeHelpers::IsTypeAllowedAsInput(Data.Data->GetDataType()))
 		{
 			continue;
 		}
 
 		FPCGDataDesc DataDesc;
-		if (ensure(PCGComputeHelpers::ComputeDataDescription(Data.Data, InBinding, DataDesc)))
+		if (ensure(PCGComputeHelpers::ComputeDataDescription(Data, InBinding, DataDesc)))
 		{
 			CollectionDesc.DataDescs.Add(MoveTemp(DataDesc));
 		}
@@ -1031,7 +1042,7 @@ static uint32 GetElementDataStartAddressUints(const uint32* InPackedDataCollecti
 	return InPackedDataCollection[ReadAddressBytes >> 2] >> 2;
 }
 
-void FPCGDataCollectionDesc::PackDataCollection(const FPCGDataCollection& InDataCollection, FName InPin, const TArray<FString>& InStringTable, TArray<uint32>& OutPackedDataCollection) const
+void FPCGDataCollectionDesc::PackDataCollection(const FPCGDataCollection& InDataCollection, FName InPin, const UPCGDataBinding* InDataBinding, TArray<uint32>& OutPackedDataCollection) const
 {
 	TRACE_CPUPROFILER_EVENT_SCOPE(FPCGDataCollectionDesc::PackDataCollection);
 
@@ -1194,10 +1205,10 @@ void FPCGDataCollectionDesc::PackDataCollection(const FPCGDataCollection& InData
 				else
 				{
 					// Pack attribute. Validate first element only for perf.
-					ensure(PCGDataForGPUHelpers::PackAttributeHelper(AttributeBase, AttributeDesc, Points[0].MetadataEntry, InStringTable, OutPackedDataCollection, AddressUints));
+					ensure(PCGDataForGPUHelpers::PackAttributeHelper(AttributeBase, AttributeDesc, Points[0].MetadataEntry, InDataBinding->GetStringTable(), OutPackedDataCollection, AddressUints));
 					for (uint32 ElementIndex = 1; ElementIndex < NumElements; ++ElementIndex)
 					{
-						PCGDataForGPUHelpers::PackAttributeHelper(AttributeBase, AttributeDesc, Points[ElementIndex].MetadataEntry, InStringTable, OutPackedDataCollection, AddressUints);
+						PCGDataForGPUHelpers::PackAttributeHelper(AttributeBase, AttributeDesc, Points[ElementIndex].MetadataEntry, InDataBinding->GetStringTable(), OutPackedDataCollection, AddressUints);
 					}
 				}
 			}
@@ -1215,10 +1226,10 @@ void FPCGDataCollectionDesc::PackDataCollection(const FPCGDataCollection& InData
 				uint32 AddressUints = GetElementDataStartAddressUints(OutPackedDataCollection.GetData(), DataIndex, AttributeDesc.Index);
 
 				// Pack attribute. Validate first element only for perf.
-				ensure(PCGDataForGPUHelpers::PackAttributeHelper(AttributeBase, AttributeDesc, /*InEntryKey*/0, InStringTable, OutPackedDataCollection, AddressUints));
+				ensure(PCGDataForGPUHelpers::PackAttributeHelper(AttributeBase, AttributeDesc, /*InEntryKey*/0, InDataBinding->GetStringTable(), OutPackedDataCollection, AddressUints));
 				for (int32 ElementIndex = 1; ElementIndex < DataDesc.ElementCount; ++ElementIndex)
 				{
-					PCGDataForGPUHelpers::PackAttributeHelper(AttributeBase, AttributeDesc, /*InEntryKey*/ElementIndex, InStringTable, OutPackedDataCollection, AddressUints);
+					PCGDataForGPUHelpers::PackAttributeHelper(AttributeBase, AttributeDesc, /*InEntryKey*/ElementIndex, InDataBinding->GetStringTable(), OutPackedDataCollection, AddressUints);
 				}
 			}
 		}
@@ -1313,6 +1324,14 @@ EPCGUnpackDataCollectionResult FPCGDataCollectionDesc::UnpackDataCollection(FPCG
 			OutTaggedData.Data = OutPointData;
 			OutTaggedData.Pin = InPin;
 
+			for (const int32 TagStringKey : DataDesc.TagStringKeys)
+			{
+				if (ensure(InStringTable.IsValidIndex(TagStringKey)))
+				{
+					OutTaggedData.Tags.Add(InStringTable[TagStringKey]);
+				}
+			}
+
 			// No work to do if there are no elements to process.
 			if (NumElements == 0)
 			{
@@ -1621,13 +1640,13 @@ FPCGProxyForGPUDataCollection::FPCGProxyForGPUDataCollection(TRefCountPtr<FRDGPo
 {
 }
 
-bool FPCGProxyForGPUDataCollection::GetCPUData(FPCGContext* InContext, int32 InDataIndex, const UPCGData*& OutData)
+bool FPCGProxyForGPUDataCollection::GetCPUData(FPCGContext* InContext, int32 InDataIndex, FPCGTaggedData& OutData)
 {
 	UE::TScopeLock Lock(ReadbackLock);
 
 	if (bReadbackDataProcessed)
 	{
-		OutData = ReadbackData.IsValidIndex(InDataIndex) ? ReadbackData[InDataIndex].Get() : nullptr;
+		OutData = ReadbackData.IsValidIndex(InDataIndex) ? ReadbackData[InDataIndex] : FPCGTaggedData();
 		return true;
 	}
 
@@ -1686,7 +1705,7 @@ bool FPCGProxyForGPUDataCollection::GetCPUData(FPCGContext* InContext, int32 InD
 
 	if (!bReadbackDataArrived)
 	{
-		OutData = nullptr;
+		OutData = FPCGTaggedData();
 		return false;
 	}
 
@@ -1703,18 +1722,19 @@ bool FPCGProxyForGPUDataCollection::GetCPUData(FPCGContext* InContext, int32 InD
 
 		if (Result != EPCGUnpackDataCollectionResult::Success)
 		{
-			OutData = nullptr;
+			OutData = FPCGTaggedData();
 			return true;
 		}
 
 		ReadbackData.Reset();
 		for (const FPCGTaggedData& TaggedData : DataFromGPU.TaggedData)
 		{
-			ReadbackData.Add(TStrongObjectPtr<const UPCGData>(TaggedData.Data));
+			ReadbackData.Add(TaggedData);
+			ReadbackDataRefs.Add(TStrongObjectPtr<const UPCGData>(TaggedData.Data));
 		}
 	}
 
-	OutData = ReadbackData.IsValidIndex(InDataIndex) ? ReadbackData[InDataIndex].Get() : nullptr;
+	OutData = ReadbackData.IsValidIndex(InDataIndex) ? ReadbackData[InDataIndex] : FPCGTaggedData();
 	return true;
 }
 
diff --git a/Engine/Plugins/PCG/Source/PCG/Private/Compute/PCGDataUploadAdaptor.cpp b/Engine/Plugins/PCG/Source/PCG/Private/Compute/PCGDataUploadAdaptor.cpp
--- a/Engine/Plugins/PCG/Source/PCG/Private/Compute/PCGDataUploadAdaptor.cpp
+++ b/Engine/Plugins/PCG/Source/PCG/Private/Compute/PCGDataUploadAdaptor.cpp
@@ -122,9 +122,9 @@ bool FPCGDataUploadAdaptor::PrepareData_GameThread(const UPCGDataBinding* InBind
 
 		bool bAllDataReady = true;
 
-		for (FPCGTaggedData& Data : Binding->InputDataCollection.TaggedData)
+		for (FPCGTaggedData& TaggedData : Binding->InputDataCollection.TaggedData)
 		{
-			if (const UPCGProxyForGPUData* Proxy = Cast<UPCGProxyForGPUData>(Data.Data))
+			if (const UPCGProxyForGPUData* Proxy = Cast<UPCGProxyForGPUData>(TaggedData.Data))
 			{
 				UPCGProxyForGPUData::FReadbackResult Result = Proxy->GetCPUData(nullptr);
 				if (!Result.bComplete)
@@ -133,7 +133,7 @@ bool FPCGDataUploadAdaptor::PrepareData_GameThread(const UPCGDataBinding* InBind
 				}
 				else
 				{
-					Data.Data = Result.Data;
+					TaggedData = Result.TaggedData;
 				}
 			}
 		}
@@ -143,7 +143,7 @@ bool FPCGDataUploadAdaptor::PrepareData_GameThread(const UPCGDataBinding* InBind
 			return false;
 		}
 
-		TargetDataCollectionDesc.PackDataCollection(Binding->InputDataCollection, InputPinLabel, Binding->GetStringTable(), PackedDataCollection);
+		TargetDataCollectionDesc.PackDataCollection(Binding->InputDataCollection, InputPinLabel, Binding, PackedDataCollection);
 	}
 
 	return true;
diff --git a/Engine/Plugins/PCG/Source/PCG/Private/Elements/PCGCopyPointsKernel.cpp b/Engine/Plugins/PCG/Source/PCG/Private/Elements/PCGCopyPointsKernel.cpp
--- a/Engine/Plugins/PCG/Source/PCG/Private/Elements/PCGCopyPointsKernel.cpp
+++ b/Engine/Plugins/PCG/Source/PCG/Private/Elements/PCGCopyPointsKernel.cpp
@@ -90,6 +90,22 @@ bool UPCGCopyPointsKernel::ComputeOutputBindingDataDesc(const UPCGComputeGraph*
 			{
 				ResultDataDesc.AttributeDescs = TargetDesc.AttributeDescs;
 			}
+
+			if (Settings->TagInheritance == EPCGCopyPointsTagInheritanceMode::Source || Settings->TagInheritance == EPCGCopyPointsTagInheritanceMode::Both)
+			{
+				for (const int32 SourceTagStringKey : SourceDesc.TagStringKeys)
+				{
+					ResultDataDesc.TagStringKeys.AddUnique(SourceTagStringKey);
+				}
+			}
+			
+			if (Settings->TagInheritance == EPCGCopyPointsTagInheritanceMode::Target || Settings->TagInheritance == EPCGCopyPointsTagInheritanceMode::Both)
+			{
+				for (const int32 TargetTagStringKey : TargetDesc.TagStringKeys)
+				{
+					ResultDataDesc.TagStringKeys.AddUnique(TargetTagStringKey);
+				}
+			}
 		}
 	}
 
diff --git a/Engine/Plugins/PCG/Source/PCG/Private/Graph/PCGGraphCompilerGPU.cpp b/Engine/Plugins/PCG/Source/PCG/Private/Graph/PCGGraphCompilerGPU.cpp
--- a/Engine/Plugins/PCG/Source/PCG/Private/Graph/PCGGraphCompilerGPU.cpp
+++ b/Engine/Plugins/PCG/Source/PCG/Private/Graph/PCGGraphCompilerGPU.cpp
@@ -551,6 +551,8 @@ void FPCGGraphCompilerGPU::BuildComputeGraphTask(
 
 		FPinToDataInterface OutputPinDataInterfaces;
 
+		BuildComputeGraphStaticData(InOutCompiledTasks, InCollapsedTasks, ComputeGraph);
+
 		CreateDataInterfaces(InOutContext, InOutCompiledTasks, InCollapsedTasks, InComputeElementTasks, InTaskSuccessors, InOriginalToVirtualPin, ComputeGraph, OutputPinDataInterfaces);
 
 		CompileComputeGraph(InOutContext, InOutCompiledTasks, InCollapsedTasks, InComputeElementTasks, InTaskSuccessors, InOriginalToVirtualPin, OutputPinDataInterfaces, ComputeGraph);
@@ -584,6 +586,36 @@ void FPCGGraphCompilerGPU::BuildComputeGraphTask(
 	}
 }
 
+void FPCGGraphCompilerGPU::BuildComputeGraphStaticData(
+	const TArray<FPCGGraphTask>& InCompiledTasks,
+	const TSet<FPCGTaskId>& InCollapsedTasks,
+	UPCGComputeGraph* InOutComputeGraph)
+{
+	TRACE_CPUPROFILER_EVENT_SCOPE(FPCGGraphCompilerGPU::BuildComputeGraphStaticData);
+
+	for (FPCGTaskId TaskId : InCollapsedTasks)
+	{
+		// TODO: The kernel being used here is a duplicate of the actual kernel that will be created during compute graph creation.
+		// This is okay for now because the functions just pass through data from the Settings at the moment, which doesn't require a fully initialized kernel.
+		const UPCGComputeKernel* Kernel = InCompiledTasks[TaskId].Kernel;
+
+		if (ensure(Kernel))
+		{
+			// Record attributes.
+			TArray<FPCGKernelAttributeKey> AttributeKeys;
+			Kernel->GetKernelAttributeKeys(AttributeKeys);
+
+			for (const FPCGKernelAttributeKey& Key : AttributeKeys)
+			{
+				InOutComputeGraph->StaticAttributeTable.AttributeTable.AddUnique(Key);
+			}
+
+			// Record strings.
+			Kernel->AddStaticCreatedStrings(InOutComputeGraph->StringTable);
+		}
+	}
+}
+
 void FPCGGraphCompilerGPU::CreateDataInterfaces(
 	FPCGGPUCompilationContext& InOutContext,
 	const TArray<FPCGGraphTask>& InCompiledTasks,
@@ -681,46 +713,9 @@ void FPCGGraphCompilerGPU::CreateDataInterfaces(
 	// Create DIs for output pins (either output pins within this compute graph or outside).
 	for (FPCGTaskId TaskId : InCollapsedTasks)
 	{
-		const UPCGComputeKernel* ProducerKernel = InCompiledTasks[TaskId].Kernel;
-		const UPCGSettings* Settings = InCompiledTasks[TaskId].Node ? InCompiledTasks[TaskId].Node->GetSettings() : nullptr;
-
 		// Create DIs for all output pins regardless of outbound connections, because the kernels currently need their outputs to be bound to valid resources.
-		if (Settings && ProducerKernel)
+		if (const UPCGSettings* Settings = InCompiledTasks[TaskId].Node ? InCompiledTasks[TaskId].Node->GetSettings() : nullptr)
 		{
-			// Record attributes.
-			TArray<FPCGKernelAttributeKey> AttributeKeys;
-			ProducerKernel->GetKernelAttributeKeys(AttributeKeys);
-
-			for (const FPCGKernelAttributeKey& Key : AttributeKeys)
-			{
-				InOutComputeGraph->AttributeStaticTable.AttributeTable.AddUnique(Key);
-			}
-
-			// Record strings.
-			ProducerKernel->AddStaticCreatedStrings(InOutComputeGraph->StringTable);
-
-			// Create all the data index remaps for input pins.
-			auto GetDataLabels = [ProducerKernel, Settings, InOutComputeGraph](const TArray<FPCGPinProperties>& PinProperties)
-			{
-				for (const FPCGPinProperties& PinProps : PinProperties)
-				{
-					TArray<FString> DataLabels;
-					ProducerKernel->GetDataLabels(PinProps.Label, DataLabels);
-
-					if (!DataLabels.IsEmpty())
-					{
-						FPCGPinDataLabels& PinDataLabels = InOutComputeGraph->DataLabelsStaticTable.FindOrAdd(Settings);
-						ensure(!PinDataLabels.PinToDataLabels.Contains(PinProps.Label));
-
-						FPCGDataLabels& Labels = PinDataLabels.PinToDataLabels.Add(PinProps.Label);
-						Labels.Labels = MoveTemp(DataLabels);
-					}
-				}
-			};
-
-			GetDataLabels(Settings->AllInputPinProperties());
-			GetDataLabels(Settings->AllOutputPinProperties());
-
 			// Create all the output data interfaces.
 			for (const FPCGPinProperties& OutputPinProperties : Settings->AllOutputPinProperties())
 			{
@@ -936,6 +931,30 @@ void FPCGGraphCompilerGPU::CompileComputeGraph(
 			InOutComputeGraph->StaticMeshSpawners.Add(SMSpawnerKernel);
 		}
 
+		// Populate static data labels. Cannot happen during BuildComputeGraphStaticData() because this relies on kernel index, which is not initialized at that time.
+		{
+			auto GetDataLabels = [Kernel = KernelWithBindings.Kernel, InOutComputeGraph](const TArray<FPCGPinProperties>& PinProperties)
+			{
+				for (const FPCGPinProperties& PinProps : PinProperties)
+				{
+					TArray<FString> DataLabels;
+					Kernel->GetDataLabels(PinProps.Label, DataLabels);
+
+					if (!DataLabels.IsEmpty())
+					{
+						FPCGPinDataLabels& PinDataLabels = InOutComputeGraph->StaticDataLabelsTable.FindOrAdd(Kernel->GetKernelIndex());
+						ensure(!PinDataLabels.PinToDataLabels.Contains(PinProps.Label));
+
+						FPCGDataLabels& Labels = PinDataLabels.PinToDataLabels.Add(PinProps.Label);
+						Labels.Labels = MoveTemp(DataLabels);
+					}
+				}
+			};
+
+			GetDataLabels(Settings->AllInputPinProperties());
+			GetDataLabels(Settings->AllOutputPinProperties());
+		}
+
 		// A given data interface may be connected to multiple downstream pins. Because each entry in these arrays produces its own bindings and graph edges,
 		// we also need to specify the prefix to uniquely describe each edge.
 		TArray<TPair</*DataInterfaceIndex*/int, /*DataInterfacePrefix*/FName>> InputDataInterfaceIndexAndPrefix;
@@ -946,17 +965,21 @@ void FPCGGraphCompilerGPU::CompileComputeGraph(
 		TMap</*OriginalPinLabel*/FName, /*VirtualPinLabel*/FName> OriginalToVirtualInputPin;
 		OriginalToVirtualInputPin.Reserve(InputDataInterfaceIndexAndPrefix.Num());
 
-		auto CreateDataResolverDataInterface = [&InOutContext, Settings, InOutComputeGraph](TArray<TPair</*DataInterfaceIndex=*/int, /*DataInterfacePrefix=*/FName>>& InOutDataInterfaceIndexAndPrefix, FName PinLabel)
+		auto CreateDataResolverDataInterface = [&InOutContext, Kernel = KernelWithBindings.Kernel, InOutComputeGraph](
+			TArray<TPair</*DataInterfaceIndex=*/int, /*DataInterfacePrefix=*/FName>>& InOutDataInterfaceIndexAndPrefix,
+			FName PinLabel,
+			bool bIsInput)
 		{
 			// Add data label resolver DI if this pin uses any data labels.
-			const FPCGPinDataLabels* PinDataLabels = InOutComputeGraph->GetDataLabelsStaticTable().Find(Settings);
+			const FPCGPinDataLabels* PinDataLabels = InOutComputeGraph->GetStaticDataLabelsTable().Find(Kernel->GetKernelIndex());
 			const FPCGDataLabels* DataLabels = PinDataLabels ? PinDataLabels->PinToDataLabels.Find(PinLabel) : nullptr;
 
 			if (DataLabels && !DataLabels->Labels.IsEmpty())
 			{
 				UPCGDataLabelResolverDataInterface* Resolver = InOutContext.NewObject_AnyThread<UPCGDataLabelResolverDataInterface>(InOutComputeGraph);
-				Resolver->InputPinLabel = PinLabel;
-				Resolver->ConsumerSettings = Settings;
+				Resolver->Kernel = Kernel;
+				Resolver->PinLabel = PinLabel;
+				Resolver->bIsInput = bIsInput;
 
 				const int ResolverIndex = InOutComputeGraph->DataInterfaces.Add(Resolver);
 				InOutDataInterfaceIndexAndPrefix.Emplace(ResolverIndex, PCGComputeHelpers::GetDataLabelResolverName(PinLabel));
@@ -997,12 +1020,13 @@ void FPCGGraphCompilerGPU::CompileComputeGraph(
 			}
 
 			InputDataInterfaceIndexAndPrefix.Emplace(Index, Input.DownstreamPin->Label);
-			CreateDataResolverDataInterface(InputDataInterfaceIndexAndPrefix, Input.DownstreamPin->Label);
 
 			if (const FName* VirtualPin = InOriginalToVirtualPin.Find(FNodePin(TaskId, Input.DownstreamPin->Label, /*IsInput*/true)))
 			{
 				OriginalToVirtualInputPin.Add(Input.DownstreamPin->Label, *VirtualPin);
 			}
+
+			CreateDataResolverDataInterface(InputDataInterfaceIndexAndPrefix, Input.DownstreamPin->Label, /*bIsInput=*/true);
 		}
 
 		// Setup output pin DIs. Always bind a DI to every output pin, so kernel always has something to write to.
@@ -1028,7 +1052,7 @@ void FPCGGraphCompilerGPU::CompileComputeGraph(
 			}
 
 			OutputDataInterfaceIndexAndPrefix.Emplace(Index, OutputPinProperties.Label);
-			CreateDataResolverDataInterface(OutputDataInterfaceIndexAndPrefix, OutputPinProperties.Label);
+			CreateDataResolverDataInterface(OutputDataInterfaceIndexAndPrefix, OutputPinProperties.Label, /*bIsInput=*/false);
 		}
 
 		// Kernel data interface.
@@ -1171,7 +1195,7 @@ void FPCGGraphCompilerGPU::CompileComputeGraph(
 			if (PCGGraphCompilerGPU::CVarEnableDebugging.GetValueOnAnyThread())
 			{
 				UE_LOG(LogPCG, Display, TEXT("\n### STATIC METADATA ATTRIBUTE TABLE ###"));
-				InOutComputeGraph->AttributeStaticTable.DebugLog();
+				InOutComputeGraph->StaticAttributeTable.DebugLog();
 			}
 #endif
 
diff --git a/Engine/Plugins/PCG/Source/PCG/Private/Graph/PCGGraphCompilerGPU.h b/Engine/Plugins/PCG/Source/PCG/Private/Graph/PCGGraphCompilerGPU.h
--- a/Engine/Plugins/PCG/Source/PCG/Private/Graph/PCGGraphCompilerGPU.h
+++ b/Engine/Plugins/PCG/Source/PCG/Private/Graph/PCGGraphCompilerGPU.h
@@ -76,6 +76,12 @@ private:
 		const FOriginalToVirtualPin& InOriginalToVirtualPin,
 		const TMap<TSoftObjectPtr<const UPCGPin>, FName>& InOutputCPUPinToVirtualPin);
 
+	/** Collates data for the compute graph which is statically known. */
+	static void BuildComputeGraphStaticData(
+		const TArray<FPCGGraphTask>& InCompiledTasks,
+		const TSet<FPCGTaskId>& InCollapsedTasks,
+		UPCGComputeGraph* InOutComputeGraph);
+
 	/** Creates and configures all data interfaces for the graph. */
 	static void CreateDataInterfaces(
 		FPCGGPUCompilationContext& InOutContext,
diff --git a/Engine/Plugins/PCG/Source/PCG/Private/PCGElement.cpp b/Engine/Plugins/PCG/Source/PCG/Private/PCGElement.cpp
--- a/Engine/Plugins/PCG/Source/PCG/Private/PCGElement.cpp
+++ b/Engine/Plugins/PCG/Source/PCG/Private/PCGElement.cpp
@@ -369,9 +369,9 @@ bool IPCGElement::PrepareData(FPCGContext* Context) const
 		bool bHasPendingReadbacks = false;
 
 		// If there are any proxies in the input data, request readback to CPU.
-		for (FPCGTaggedData& Data : Context->InputData.TaggedData)
+		for (FPCGTaggedData& TaggedData : Context->InputData.TaggedData)
 		{
-			if (const UPCGProxyForGPUData* DataWithGPUSupport = Cast<UPCGProxyForGPUData>(Data.Data))
+			if (const UPCGProxyForGPUData* DataWithGPUSupport = Cast<UPCGProxyForGPUData>(TaggedData.Data))
 			{
 				// Poll until readback is done.
 				UPCGProxyForGPUData::FReadbackResult Result = DataWithGPUSupport->GetCPUData(Context);
@@ -379,8 +379,9 @@ bool IPCGElement::PrepareData(FPCGContext* Context) const
 				if (Result.bComplete)
 				{
 					// Readback done, replace the proxy with the result data.
-					ensure(Result.Data);
-					Data.Data = Result.Data;
+					ensure(Result.TaggedData.Data);
+					TaggedData.Data = Result.TaggedData.Data;
+					TaggedData.Tags.Append(Result.TaggedData.Tags);
 
 					// Will update referenced data objects.
 					Context->bInputDataModified = true;
@@ -819,11 +820,13 @@ bool IPCGElement::DebugDisplay(FPCGContext* Context) const
 			if (Result.bComplete)
 			{
 				// Readback done. Queue up data for debug (unless we've found pending readbacks, in which case we'll run again later).
-				if (ensure(Result.Data) && !bHasPendingReadbacks)
+				if (ensure(Result.TaggedData.Data) && !bHasPendingReadbacks)
 				{
-					FPCGTaggedData Data = Output;
-					Data.Data = Result.Data;
-					DataToDebug.Add(Data);
+					FPCGTaggedData TaggedData = Output;
+					TaggedData.Data = Result.TaggedData.Data;
+					TaggedData.Tags.Append(Result.TaggedData.Tags);
+
+					DataToDebug.Add(TaggedData);
 				}
 			}
 			else
diff --git a/Engine/Plugins/PCG/Source/PCGEditor/Private/SPCGEditorGraphAttributeListView.cpp b/Engine/Plugins/PCG/Source/PCGEditor/Private/SPCGEditorGraphAttributeListView.cpp
--- a/Engine/Plugins/PCG/Source/PCGEditor/Private/SPCGEditorGraphAttributeListView.cpp
+++ b/Engine/Plugins/PCG/Source/PCGEditor/Private/SPCGEditorGraphAttributeListView.cpp
@@ -730,12 +730,12 @@ void SPCGEditorGraphAttributeListView::RefreshAttributeList()
 			return;
 		}
 
-		if (!Result.Data)
+		if (!Result.TaggedData.Data)
 		{
 			return;
 		}
 
-		PCGData = Result.Data;
+		PCGData = Result.TaggedData.Data;
 	}
 
 	const FPCGDataVisualizationRegistry& DataVisRegistry = FPCGModule::GetConstPCGDataVisualizationRegistry();
