[PCG] Expose type Base Texture in PCG pin properties.

Our GPU texture data interface is able to accept either texture or RT and this seems desirable.

#rnx

diff --git a/Engine/Plugins/PCG/Source/PCG/Public/PCGCommon.h b/Engine/Plugins/PCG/Source/PCG/Public/PCGCommon.h
--- a/Engine/Plugins/PCG/Source/PCG/Public/PCGCommon.h
+++ b/Engine/Plugins/PCG/Source/PCG/Public/PCGCommon.h
@@ -110,7 +110,7 @@ enum class EPCGDataType : uint32
 	Landscape = 1 << 4,
 	Texture = 1 << 5,
 	RenderTarget = 1 << 6,
-	BaseTexture = Texture | RenderTarget UMETA(Hidden),
+	BaseTexture = Texture | RenderTarget UMETA(Tooltip = "Common base type for both textures and render targets."),
 	Surface = Landscape | BaseTexture,
 
 	Volume = 1 << 7,
