[PCG] Node conversion(s) from Graph Parameter (Specific and Generic) to Constant

#rnx

diff --git a/Engine/Plugins/PCG/Source/PCG/Private/Elements/PCGUserParameterGet.cpp b/Engine/Plugins/PCG/Source/PCG/Private/Elements/PCGUserParameterGet.cpp
--- a/Engine/Plugins/PCG/Source/PCG/Private/Elements/PCGUserParameterGet.cpp
+++ b/Engine/Plugins/PCG/Source/PCG/Private/Elements/PCGUserParameterGet.cpp
@@ -9,6 +9,7 @@
 #include "PCGPin.h"
 #include "PCGSubgraph.h"
 #include "Data/PCGUserParametersData.h"
+#include "Elements/PCGCreateAttribute.h"
 #include "Helpers/PCGConversion.h"
 #include "Helpers/PCGDynamicTrackingHelpers.h"
 #include "Helpers/PCGPropertyHelpers.h"
@@ -98,10 +99,72 @@ namespace PCGUserParameterGet
 	namespace Constants::Conversion
 	{
 		constexpr int32 SpecificToGenericIndex = 0;
+		constexpr int32 SpecificToConstantIndex = 1;
 		constexpr int32 GenericToSpecificIndex = 0;
+		constexpr int32 GenericToConstantIndex = 1;
 
 		const FText SpecificToGenericTooltip = LOCTEXT("SpecificToGenericTooltip", "Convert to a Get User Parameter (Generic) node.");
 		const FText GenericToSpecificTooltip = LOCTEXT("GenericToSpecificTooltip", "Convert to a specified Get User Parameter node.");
+		const FText ToCreateAttributeTooltip = LOCTEXT("ToConstantTooltip", "Convert to a Create Attribute node.");
+
+		const FText SpecificGetUserParameterNodeTitle = LOCTEXT("SpecificGetUserParameterNodeTitle", "Get Graph Parameter");
+		const FText GenericGetUserParameterNodeTitle = LOCTEXT("GenericGetUserParameterNodeTitle", "Get Graph Parameter (Generic)");
+	}
+
+	namespace Helpers::Conversion
+	{
+		// Assigns a value based on a graph parameter to a MetadataTypesConstantStruct, which is normally only user defined.
+		bool AssignValue(FPCGMetadataTypesConstantStruct& OutTypeStruct, const UPCGGraph* Graph, const FName GraphParameterName, FText& OutFailReason)
+		{
+			if (GraphParameterName == NAME_None || !Graph)
+			{
+				OutFailReason = LOCTEXT("ConversionAssignValueInvalidGraphOrParameter", "Invalid graph or parameter.");
+				return false;
+			}
+
+			const FInstancedPropertyBag* UserParametersStruct = Graph->GetUserParametersStruct();
+			check(UserParametersStruct);
+
+			const FPropertyBagPropertyDesc* Desc = UserParametersStruct->FindPropertyDescByName(GraphParameterName);
+			if (!Desc || !Desc->ContainerTypes.IsEmpty()) // Property doesn't exist, or it is a container, which doesn't work yet. Early out.
+			{
+				OutFailReason = LOCTEXT("ConversionAssignValueInvalidParamNameOrContainer", "Parameter not found or it is a container, which is not supported.");
+				return false;
+			}
+
+			TUniquePtr<const IPCGAttributeAccessor> PropertyAccessor = PCGAttributeAccessorHelpers::CreatePropertyAccessor(Desc->CachedProperty);
+
+			if (!PropertyAccessor.IsValid())
+			{
+				OutFailReason = LOCTEXT("ConversionAssignValueInvalidAccessor", "Could not access parameter, likely due to unsupported type.");
+				return false;
+			}
+
+			// The type must be set before the dispatch.
+			OutTypeStruct.Type = static_cast<EPCGMetadataTypes>(PropertyAccessor->GetUnderlyingType());
+			if (OutTypeStruct.Type == EPCGMetadataTypes::Unknown)
+			{
+				OutFailReason = LOCTEXT("ConversionAssignValueUnknownType", "Parameter type is not supported by Create Attribute.");
+				return false;
+			}
+
+			auto DispatchFunc = [UserParametersStruct, &PropertyAccessor, &OutFailReason]<typename T>(const T& Value)
+			{
+				// Simulating initialization by the user. Const qualifier from 'Dispatcher' function signature only, so safe to cast away.
+				T& Ref = const_cast<T&>(Value);
+
+				const FPCGAttributeAccessorKeysSingleObjectPtr Keys(UserParametersStruct->GetValue().GetMemory());
+				const bool bResult = PropertyAccessor->Get<T>(Ref, Keys, EPCGAttributeAccessorFlags::AllowBroadcastAndConstructible);
+				if (!bResult)
+				{
+					OutFailReason = LOCTEXT("ConversionAssignValueCouldNotReadAccessor", "Could not read value from user parameter accessor.");
+				}
+
+				return bResult;
+			};
+
+			return OutTypeStruct.Dispatcher(DispatchFunc);
+		}
 	}
 } // namespace PCGUserParameterGet
 
@@ -145,20 +208,26 @@ TArray<FPCGPreconfiguredInfo> UPCGUserParameterGetSettings::GetConversionInfo()
 		FText::Format(INVTEXT("{0} (Generic)"), PCGConversion::Helpers::GetDefaultNodeTitle(UPCGGenericUserParameterGetSettings::StaticClass())),
 		SpecificToGenericTooltip);
 
+	ConversionInfo.Emplace(
+		SpecificToConstantIndex,
+		PCGConversion::Helpers::GetDefaultNodeTitle(UPCGCreateAttributeSetSettings::StaticClass()),
+		ToCreateAttributeTooltip);
+
 	return ConversionInfo;
 }
 #endif // WITH_EDITOR
 
 bool UPCGUserParameterGetSettings::ConvertNode(const FPCGPreconfiguredInfo& ConversionInfo)
 {
+	using namespace PCGUserParameterGet::Constants::Conversion;
 	UPCGNode* Node = CastChecked<UPCGNode>(GetOuter());
 
-	if (ConversionInfo.PreconfiguredIndex == PCGUserParameterGet::Constants::Conversion::SpecificToGenericIndex)
+	if (ConversionInfo.PreconfiguredIndex == SpecificToGenericIndex)
 	{
 		FPCGSingleNodeConverter NodeConverter(Node, UPCGGenericUserParameterGetSettings::StaticClass());
 		if (!NodeConverter.IsGraphInitialized() || !NodeConverter.IsSourceInitialized())
 		{
-			PCGLog::Settings::LogInvalidPreconfigurationWarning(ConversionInfo.PreconfiguredIndex, LOCTEXT("GetUserParameterNodeTitle", "Get User Parameter"));
+			PCGLog::Settings::LogInvalidPreconfigurationWarning(ConversionInfo.PreconfiguredIndex, SpecificGetUserParameterNodeTitle);
 			return false;
 		}
 
@@ -173,6 +242,36 @@ bool UPCGUserParameterGetSettings::ConvertNode(const FPCGPreconfiguredInfo& Conv
 
 		return NodeConverter.IsComplete();
 	}
+	else if (ConversionInfo.PreconfiguredIndex == SpecificToConstantIndex)
+	{
+		FPCGSingleNodeConverter NodeConverter(Node, UPCGCreateAttributeSetSettings::StaticClass());
+		if (!NodeConverter.IsGraphInitialized() || !NodeConverter.IsSourceInitialized())
+		{
+			PCGLog::Settings::LogInvalidPreconfigurationWarning(ConversionInfo.PreconfiguredIndex, SpecificGetUserParameterNodeTitle);
+			return false;
+		}
+
+		NodeConverter.PrepareData();
+
+		if (UPCGGraph* Graph = NodeConverter.GetGraph())
+		{
+			UPCGCreateAttributeSetSettings* Settings = CastChecked<UPCGCreateAttributeSetSettings>(NodeConverter.GetGeneratedSettings());
+			Settings->OutputTarget.Update(PropertyName.ToString());
+
+			FText ErrorMessage;
+			if (PCGUserParameterGet::Helpers::Conversion::AssignValue(Settings->AttributeTypes, Graph, PropertyName, ErrorMessage))
+			{
+				NodeConverter.ApplyStructural();
+				NodeConverter.Finalize();
+			}
+			else
+			{
+				PCGLog::Settings::LogInvalidConversionError(ConversionInfo.PreconfiguredIndex, SpecificGetUserParameterNodeTitle, ErrorMessage);
+			}
+		}
+
+		return NodeConverter.IsComplete();
+	}
 
 	return false;
 }
@@ -220,15 +319,21 @@ TArray<FPCGPreconfiguredInfo> UPCGGenericUserParameterGetSettings::GetConversion
 		PCGConversion::Helpers::GetDefaultNodeTitle(UPCGUserParameterGetSettings::StaticClass()),
 		GenericToSpecificTooltip);
 
+	ConversionInfo.Emplace(
+		GenericToConstantIndex,
+		PCGConversion::Helpers::GetDefaultNodeTitle(UPCGCreateAttributeSetSettings::StaticClass()),
+		ToCreateAttributeTooltip);
+
 	return ConversionInfo;
 }
 #endif // WITH_EDITOR
 
 bool UPCGGenericUserParameterGetSettings::ConvertNode(const FPCGPreconfiguredInfo& ConversionInfo)
 {
+	using namespace PCGUserParameterGet::Constants::Conversion;
 	UPCGNode* Node = CastChecked<UPCGNode>(GetOuter());
 
-	if (ConversionInfo.PreconfiguredIndex == PCGUserParameterGet::Constants::Conversion::GenericToSpecificIndex)
+	if (ConversionInfo.PreconfiguredIndex == GenericToSpecificIndex)
 	{
 		// If the property is overridden or upstream, abort.
 		if (IsPropertyOverriddenByPin(
@@ -242,7 +347,7 @@ bool UPCGGenericUserParameterGetSettings::ConvertNode(const FPCGPreconfiguredInf
 		FPCGSingleNodeConverter NodeConverter(Node, UPCGUserParameterGetSettings::StaticClass());
 		if (!NodeConverter.IsGraphInitialized() || !NodeConverter.IsSourceInitialized())
 		{
-			PCGLog::Settings::LogInvalidPreconfigurationWarning(ConversionInfo.PreconfiguredIndex, LOCTEXT("GenericGetUserParameterNodeTitle", "Get User Parameter (Generic)"));
+			PCGLog::Settings::LogInvalidPreconfigurationWarning(ConversionInfo.PreconfiguredIndex, GenericGetUserParameterNodeTitle);
 			return false;
 		}
 
@@ -276,6 +381,37 @@ bool UPCGGenericUserParameterGetSettings::ConvertNode(const FPCGPreconfiguredInf
 			return NodeConverter.IsComplete();
 		}
 	}
+	else if (ConversionInfo.PreconfiguredIndex == GenericToConstantIndex)
+	{
+		FPCGSingleNodeConverter NodeConverter(Node, UPCGCreateAttributeSetSettings::StaticClass());
+		if (!NodeConverter.IsGraphInitialized() || !NodeConverter.IsSourceInitialized())
+		{
+			PCGLog::Settings::LogInvalidPreconfigurationWarning(ConversionInfo.PreconfiguredIndex, GenericGetUserParameterNodeTitle);
+			return false;
+		}
+
+		NodeConverter.PrepareData();
+
+		if (UPCGGraph* Graph = NodeConverter.GetGraph())
+		{
+			const FName PropertyName(PropertyPath);
+			UPCGCreateAttributeSetSettings* Settings = CastChecked<UPCGCreateAttributeSetSettings>(NodeConverter.GetGeneratedSettings());
+			Settings->OutputTarget.Update(OutputAttributeName.ToString());
+
+			FText ErrorMessage;
+			if (PCGUserParameterGet::Helpers::Conversion::AssignValue(Settings->AttributeTypes, Graph, PropertyName, ErrorMessage))
+			{
+				NodeConverter.ApplyStructural();
+				NodeConverter.Finalize();
+			}
+			else
+			{
+				PCGLog::Settings::LogInvalidConversionError(ConversionInfo.PreconfiguredIndex, GenericGetUserParameterNodeTitle, ErrorMessage);
+			}
+		}
+
+		return NodeConverter.IsComplete();
+	}
 
 	return false;
 }
diff --git a/Engine/Plugins/PCG/Source/PCG/Private/Helpers/PCGConversion.cpp b/Engine/Plugins/PCG/Source/PCG/Private/Helpers/PCGConversion.cpp
--- a/Engine/Plugins/PCG/Source/PCG/Private/Helpers/PCGConversion.cpp
+++ b/Engine/Plugins/PCG/Source/PCG/Private/Helpers/PCGConversion.cpp
@@ -77,6 +77,7 @@ FPCGConverterBase::FPCGConverterBase(UPCGGraph* InOutGraph)
 
 FPCGConverterBase::~FPCGConverterBase()
 {
+	// Implementation Note: Each converter should roll back any changes based on the step completed, if terminated early.
 #if WITH_EDITOR
 	if (bGraphIsDirty)
 	{
@@ -118,15 +119,29 @@ FPCGSingleNodeConverter::FPCGSingleNodeConverter(UPCGNode* InNode, TSubclassOf<U
 	}
 }
 
+FPCGSingleNodeConverter::~FPCGSingleNodeConverter()
+{
+	// Unravel the pre-structural changes, if cancelled early.
+	if (!AreStructuralChangesApplied())
+	{
+		if (SourceGraph && GeneratedNode)
+		{
+			SourceGraph->RemoveNode(GeneratedNode);
+		}
+
+		bGraphIsDirty = false;
+	}
+}
+
 bool FPCGSingleNodeConverter::ExecutePrepareData(bool& bOutMarkGraphAsDirty)
 {
 	if (ensure(IsValid()))
 	{
 		GeneratedNode = SourceGraph->AddNodeOfType(TargetSettingsClass, GeneratedSettings);
-		bOutMarkGraphAsDirty = true;
 
 		if (GeneratedNode && GeneratedSettings)
 		{
+			bOutMarkGraphAsDirty = true;
 			if (SourceNode->HasAuthoredTitle())
 			{
 				GeneratedNode->NodeTitle = SourceNode->NodeTitle;
diff --git a/Engine/Plugins/PCG/Source/PCG/Private/Helpers/PCGConversion.h b/Engine/Plugins/PCG/Source/PCG/Private/Helpers/PCGConversion.h
--- a/Engine/Plugins/PCG/Source/PCG/Private/Helpers/PCGConversion.h
+++ b/Engine/Plugins/PCG/Source/PCG/Private/Helpers/PCGConversion.h
@@ -99,6 +99,7 @@ class FPCGSingleNodeConverter : public FPCGConverterBase
 {
 public:
 	explicit FPCGSingleNodeConverter(UPCGNode* InNode, TSubclassOf<UPCGSettings> InTargetSettingsClass);
+	virtual ~FPCGSingleNodeConverter();
 
 	//~Begin FPCGConverterBase interface
 	/** The node conversion still has valid references. */
diff --git a/Engine/Plugins/PCG/Source/PCG/Private/Helpers/PCGPropertyHelpers.cpp b/Engine/Plugins/PCG/Source/PCG/Private/Helpers/PCGPropertyHelpers.cpp
--- a/Engine/Plugins/PCG/Source/PCG/Private/Helpers/PCGPropertyHelpers.cpp
+++ b/Engine/Plugins/PCG/Source/PCG/Private/Helpers/PCGPropertyHelpers.cpp
@@ -227,12 +227,6 @@ EPCGMetadataTypes PCGPropertyHelpers::GetMetadataTypeFromProperty(const FPropert
 		return EPCGMetadataTypes::Unknown;
 	}
 
-	// Object are not yet supported as accessors
-	if (const FObjectPropertyBase* ObjectProperty = CastField<FObjectPropertyBase>(InProperty))
-	{
-		return EPCGMetadataTypes::String;
-	}
-
 	TUniquePtr<IPCGAttributeAccessor> PropertyAccessor = PCGAttributeAccessorHelpers::CreatePropertyAccessor(InProperty);
 
 	return PropertyAccessor.IsValid() ? EPCGMetadataTypes(PropertyAccessor->GetUnderlyingType()) : EPCGMetadataTypes::Unknown;
diff --git a/Engine/Plugins/PCG/Source/PCG/Private/Utils/PCGLogErrors.cpp b/Engine/Plugins/PCG/Source/PCG/Private/Utils/PCGLogErrors.cpp
--- a/Engine/Plugins/PCG/Source/PCG/Private/Utils/PCGLogErrors.cpp
+++ b/Engine/Plugins/PCG/Source/PCG/Private/Utils/PCGLogErrors.cpp
@@ -126,6 +126,7 @@ namespace PCGLog::Settings
 	namespace Format
 	{
 		const FTextFormat InvalidPreconfiguration = LOCTEXT("InvalidPreconfiguration", "Invalid preconfiguration index '{0}' for node settings '{1}'. Default settings will be used.");
+		const FTextFormat InvalidConversion = LOCTEXT("InvalidConversion", "Invalid conversion for preconfiguration index '{0}' for node settings '{1}'. Reason: {2}");
 	}
 
 	// Warnings
@@ -133,6 +134,12 @@ namespace PCGLog::Settings
 	{
 		LogWarningOnGraph(FText::Format(Format::InvalidPreconfiguration, FText::AsNumber(PreconfigurationIndex), NodeTitle));
 	}
+
+	// Errors
+	void LogInvalidConversionError(const int32 PreconfigurationIndex, const FText& NodeTitle, const FText& Reason)
+	{
+		LogErrorOnGraph(FText::Format(Format::InvalidConversion, FText::AsNumber(PreconfigurationIndex), NodeTitle, Reason));
+	}
 }
 
 #undef LOCTEXT_NAMESPACE
diff --git a/Engine/Plugins/PCG/Source/PCG/Public/Utils/PCGLogErrors.h b/Engine/Plugins/PCG/Source/PCG/Public/Utils/PCGLogErrors.h
--- a/Engine/Plugins/PCG/Source/PCG/Public/Utils/PCGLogErrors.h
+++ b/Engine/Plugins/PCG/Source/PCG/Public/Utils/PCGLogErrors.h
@@ -132,6 +132,7 @@ namespace PCGLog
 
 	namespace Settings
 	{
-		PCG_API void LogInvalidPreconfigurationWarning(const int32 PreconfigurationIndex, const FText& NodeTitle);
+		PCG_API void LogInvalidPreconfigurationWarning(int32 PreconfigurationIndex, const FText& NodeTitle);
+		PCG_API void LogInvalidConversionError(int32 PreconfigurationIndex, const FText& NodeTitle, const FText& Reason);
 	}
 }
