Made match & set attributes multi-data-aware

diff --git a/Engine/Plugins/PCG/Source/PCG/Private/Elements/PCGMatchAndSetAttributes.cpp b/Engine/Plugins/PCG/Source/PCG/Private/Elements/PCGMatchAndSetAttributes.cpp
--- a/Engine/Plugins/PCG/Source/PCG/Private/Elements/PCGMatchAndSetAttributes.cpp
+++ b/Engine/Plugins/PCG/Source/PCG/Private/Elements/PCGMatchAndSetAttributes.cpp
@@ -59,7 +59,7 @@ FText UPCGMatchAndSetAttributesSettings::GetDefaultNodeTitle() const
 
 FText UPCGMatchAndSetAttributesSettings::GetNodeTooltipText() const
 {
-	return LOCTEXT("NodeTooltip", "Matches or randomly assigns values from the Attribute Set to the input data");
+	return LOCTEXT("NodeTooltip", "Matches or randomly assigns values from the Attribute Set to the input data. Supports N (input):1 (match data), or N:N configurations.");
 }
 
 void UPCGMatchAndSetAttributesSettings::ApplyDeprecationBeforeUpdatePins(UPCGNode* InOutNode, TArray<TObjectPtr<UPCGPin>>& InputPins, TArray<TObjectPtr<UPCGPin>>& OutputPins)
@@ -98,8 +98,8 @@ TArray<FPCGPinProperties> UPCGMatchAndSetAttributesSettings::InputPinProperties(
 
 	PinProperties.Emplace(PCGMatchAndSetAttributesConstants::MatchDataLabel, 
 		EPCGDataType::Param, 
-		/*bAllowMultipleConnection=*/false, 
-		/*bAllowMultipleData=*/false,
+		/*bAllowMultipleConnection=*/true,
+		/*bAllowMultipleData=*/true,
 		LOCTEXT("MatchDataTooltip", "Input containing the data to match to, then copy the accompanying attribute values")
 	);
 
@@ -107,7 +107,7 @@ TArray<FPCGPinProperties> UPCGMatchAndSetAttributesSettings::InputPinProperties(
 	{
 		PinProperties.Emplace(PCGMatchAndSetAttributesConstants::MaxDistanceLabel,
 			EPCGDataType::PointOrParam,
-			/*bAllowMultipleConnections=*/false,
+			/*bAllowMultipleConnections=*/true,
 			/*bAllowMultipleData=*/true,
 			LOCTEXT("MaxDistanceTooltip", "Input containing the maximum distance allowed for nearest search, selected by the Max Distance Attribute.")
 		);
@@ -858,6 +858,14 @@ FPCGMatchAndSetAttributesExecutionState::~FPCGMatchAndSetAttributesExecutionStat
 	}
 }
 
+FPCGMatchAndSetAttributesIterationState::~FPCGMatchAndSetAttributesIterationState()
+{
+	if (bIsPartitionOwner)
+	{
+		delete Partition;
+	}
+}
+
 bool FPCGMatchAndSetAttributesElement::PrepareDataInternal(FPCGContext* InContext) const
 {
 	TRACE_CPUPROFILER_EVENT_SCOPE(FPCGMatchAndSetAttributesElement::PrepareDataInternal);
@@ -871,16 +879,16 @@ bool FPCGMatchAndSetAttributesElement::PrepareDataInternal(FPCGContext* InContex
 	TArray<FPCGTaggedData> Inputs = InContext->InputData.GetInputsByPin(PCGPinConstants::DefaultInputLabel);
 	TArray<FPCGTaggedData>& Outputs = InContext->OutputData.TaggedData;
 	TArray<FPCGTaggedData> ParamDataInputs = InContext->InputData.GetInputsByPin(PCGMatchAndSetAttributesConstants::MatchDataLabel);
+	TArray<FPCGTaggedData> MaxDistanceDataInputs = InContext->InputData.GetInputsByPin(PCGMatchAndSetAttributesConstants::MaxDistanceLabel);
 
-	EPCGTimeSliceInitResult InitResult = TimeSlicedContext->InitializePerExecutionState([Settings, &ParamDataInputs, &Inputs](FPCGMatchAndSetAttributesElement::ContextType* Context, FPCGMatchAndSetAttributesExecutionState& OutState) -> EPCGTimeSliceInitResult
+	auto CreatePartition = [Settings, Context = TimeSlicedContext, &Inputs, &MaxDistanceDataInputs](const UPCGParamData* ParamData, const TArray<FPCGTaggedData>& PartitionInputs, int PartitionInputIndexIfSingle, FPCGMatchAndSetPartition*& OutPartition) -> EPCGTimeSliceInitResult
 	{
-		const UPCGParamData* ParamData = nullptr;
-
-		if (ParamDataInputs.Num() == 1)
+		if (PartitionInputs.IsEmpty())
 		{
-			ParamData = Cast<const UPCGParamData>(ParamDataInputs[0].Data);
+			return EPCGTimeSliceInitResult::NoOperation;
 		}
 
+		// Validate that the param data is valid
 		if (!ParamData)
 		{
 			if (Settings->bWarnIfNoMatchData)
@@ -891,45 +899,81 @@ bool FPCGMatchAndSetAttributesElement::PrepareDataInternal(FPCGContext* InContex
 			return EPCGTimeSliceInitResult::NoOperation;
 		}
 
+		FPCGMatchAndSetPartition* Partition = new FPCGMatchAndSetPartition(Context, Settings, Context->SourceComponent.Get(), ParamData);
+		check(Partition);
+
 		// If there are provided max distance entries, we should have either 1 or the same cardinality as the inputs
 		TMap<const UPCGData*, const UPCGData*> InputToMaxDistanceMapping;
 		if (Settings->MaxDistanceMode == EPCGMatchMaxDistanceMode::AttributeMaxDistance)
 		{
-			TArray<FPCGTaggedData> InputMaxDistanceData = Context->InputData.GetInputsByPin(PCGMatchAndSetAttributesConstants::MaxDistanceLabel);
+			check(PartitionInputs.Num() == 1 || MaxDistanceDataInputs.Num() == PartitionInputs.Num());
 
-			if (InputMaxDistanceData.Num() == Inputs.Num())
+			if (MaxDistanceDataInputs.Num() == 1)
 			{
-				for (int DataIndex = 0; DataIndex < Inputs.Num(); ++DataIndex)
+				for (int DataIndex = 0; DataIndex < PartitionInputs.Num(); ++DataIndex)
 				{
-					InputToMaxDistanceMapping.Add(Inputs[DataIndex].Data, InputMaxDistanceData[DataIndex].Data);
+					InputToMaxDistanceMapping.Add(PartitionInputs[DataIndex].Data, MaxDistanceDataInputs[0].Data);
 				}
 			}
-			else if (InputMaxDistanceData.Num() == 1)
+			else if (PartitionInputIndexIfSingle == INDEX_NONE && PartitionInputs.Num() == MaxDistanceDataInputs.Num())
 			{
-				for (int DataIndex = 0; DataIndex < Inputs.Num(); ++DataIndex)
+				for (int DataIndex = 0; DataIndex < PartitionInputs.Num(); ++DataIndex)
 				{
-					InputToMaxDistanceMapping.Add(Inputs[DataIndex].Data, InputMaxDistanceData[0].Data);
+					InputToMaxDistanceMapping.Add(PartitionInputs[DataIndex].Data, MaxDistanceDataInputs[DataIndex].Data);
 				}
 			}
+			else if(PartitionInputIndexIfSingle != INDEX_NONE && PartitionInputs.Num() == 1 && MaxDistanceDataInputs.IsValidIndex(PartitionInputIndexIfSingle))
+			{
+				InputToMaxDistanceMapping.Add(PartitionInputs[0].Data, MaxDistanceDataInputs[PartitionInputIndexIfSingle].Data);
+			}
 			else
 			{
-				PCGE_LOG_C(Error, GraphAndLog, Context, FText::Format(LOCTEXT("NoMatchingMaxDistanceData", "Invalid number of max distance providers for data; expected {0}, got {1}."), FText::AsNumber(Inputs.Num()), FText::AsNumber(InputMaxDistanceData.Num())));
 				return EPCGTimeSliceInitResult::AbortExecution;
 			}
 		}
 
-		OutState.Partition = new FPCGMatchAndSetPartition(Context, Settings, Context->SourceComponent.Get(), ParamData);
-		check(OutState.Partition);
-
-		if (OutState.Partition->Initialize(InputToMaxDistanceMapping))
+		if (Partition->Initialize(InputToMaxDistanceMapping))
 		{
+			OutPartition = Partition;
 			return EPCGTimeSliceInitResult::Success;
 		}
 		else
 		{
-			PCGE_LOG_C(Warning, GraphAndLog, Context, LOCTEXT("CouldNotInitializeExecutionState", "Could not initialize per-execution timeslice state data"));
+			delete Partition;
 			return EPCGTimeSliceInitResult::AbortExecution;
 		}
+	};
+
+	EPCGTimeSliceInitResult InitResult = TimeSlicedContext->InitializePerExecutionState([Settings, &ParamDataInputs, &Inputs, &MaxDistanceDataInputs, &CreatePartition](FPCGMatchAndSetAttributesElement::ContextType* Context, FPCGMatchAndSetAttributesExecutionState& OutState) -> EPCGTimeSliceInitResult
+	{
+		if(ParamDataInputs.IsEmpty() || (ParamDataInputs.Num() != 1 && ParamDataInputs.Num() != Inputs.Num()))
+		{
+			PCGLog::InputOutput::LogInvalidCardinalityError(PCGPinConstants::DefaultInputLabel, PCGMatchAndSetAttributesConstants::MatchDataLabel, Context);
+			return EPCGTimeSliceInitResult::AbortExecution;
+		}
+
+		if (Settings->MaxDistanceMode == EPCGMatchMaxDistanceMode::AttributeMaxDistance && MaxDistanceDataInputs.Num() != Inputs.Num() && MaxDistanceDataInputs.Num() != 1)
+		{
+			PCGLog::InputOutput::LogInvalidCardinalityError(PCGPinConstants::DefaultInputLabel, PCGMatchAndSetAttributesConstants::MaxDistanceLabel, Context);
+			return EPCGTimeSliceInitResult::AbortExecution;
+		}
+
+		// We will create a partition here only in the N:1 case.
+		if (ParamDataInputs.Num() == 1)
+		{
+			EPCGTimeSliceInitResult Result = CreatePartition(Cast<const UPCGParamData>(ParamDataInputs[0].Data), Inputs, -1, OutState.Partition);
+
+			if(Result == EPCGTimeSliceInitResult::AbortExecution)
+			{
+				PCGE_LOG_C(Warning, GraphAndLog, Context, LOCTEXT("CouldNotInitializeExecutionState", "Could not initialize per-execution timeslice state data"));
+			}
+
+			return Result;
+		}
+		else
+		{
+			return EPCGTimeSliceInitResult::Success;
+		}
 	});
 
 	if (InitResult == EPCGTimeSliceInitResult::AbortExecution)
@@ -943,8 +987,9 @@ bool FPCGMatchAndSetAttributesElement::PrepareDataInternal(FPCGContext* InContex
 		return true;
 	}
 
-	TimeSlicedContext->InitializePerIterationStates(Inputs.Num(), [&Inputs, &Outputs, InContext](FPCGMatchAndSetAttributesIterationState& OutState, const FPCGMatchAndSetAttributesExecutionState&, const uint32 IterationIndex) -> EPCGTimeSliceInitResult
+	TimeSlicedContext->InitializePerIterationStates(Inputs.Num(), [&Inputs, &ParamDataInputs, &Outputs, InContext, &CreatePartition](FPCGMatchAndSetAttributesIterationState& OutState, const FPCGMatchAndSetAttributesExecutionState& ExecutionState, const uint32 IterationIndex) -> EPCGTimeSliceInitResult
 	{
+		// Create output data
 		FPCGTaggedData& Output = Outputs.Add_GetRef(Inputs[IterationIndex]);
 
 		OutState.InData = Inputs[IterationIndex].Data;
@@ -954,6 +999,30 @@ bool FPCGMatchAndSetAttributesElement::PrepareDataInternal(FPCGContext* InContex
 			return EPCGTimeSliceInitResult::NoOperation;
 		}
 
+		// Create or assign partition as needed
+		if (ParamDataInputs.Num() == 1)
+		{
+			OutState.Partition = ExecutionState.Partition;
+			OutState.bIsPartitionOwner = false;
+		}
+		else
+		{
+			EPCGTimeSliceInitResult Result = CreatePartition(Cast<const UPCGParamData>(ParamDataInputs[IterationIndex].Data), { Inputs[IterationIndex] }, IterationIndex, OutState.Partition);
+			if (Result == EPCGTimeSliceInitResult::Success)
+			{
+				OutState.bIsPartitionOwner = true;
+			}
+			else
+			{
+				if (Result == EPCGTimeSliceInitResult::AbortExecution)
+				{
+					PCGE_LOG_C(Warning, GraphAndLog, InContext, LOCTEXT("CouldNotInitializeStateData", "Could not initialize timeslice state data"));
+				}
+				
+				return Result;
+			}
+		}
+
 		if (const UPCGPointData* InPointData = Cast<const UPCGPointData>(OutState.InData))
 		{
 			UPCGPointData* OutPointData = FPCGContext::NewObject_AnyThread<UPCGPointData>(InContext);
@@ -1018,7 +1087,7 @@ bool FPCGMatchAndSetAttributesElement::ExecuteInternal(FPCGContext* InContext) c
 		check(InitResult == EPCGTimeSliceInitResult::Success);
 
 		// Run the execution until the time slice is finished. We actually don't care about the max index since we won't select anything.
-		return ExecState.Partition->SelectMultiple(*Context, IterState.InData, IterState.CurrentIndex, /*MaxIndex*/ 0, IterState.OutData);
+		return IterState.Partition->SelectMultiple(*Context, IterState.InData, IterState.CurrentIndex, /*MaxIndex*/ 0, IterState.OutData);
 	});
 }
 
diff --git a/Engine/Plugins/PCG/Source/PCG/Public/Elements/PCGMatchAndSetAttributes.h b/Engine/Plugins/PCG/Source/PCG/Public/Elements/PCGMatchAndSetAttributes.h
--- a/Engine/Plugins/PCG/Source/PCG/Public/Elements/PCGMatchAndSetAttributes.h
+++ b/Engine/Plugins/PCG/Source/PCG/Public/Elements/PCGMatchAndSetAttributes.h
@@ -21,9 +21,14 @@ struct FPCGMatchAndSetAttributesExecutionState
 
 struct FPCGMatchAndSetAttributesIterationState
 {
+	~FPCGMatchAndSetAttributesIterationState();
+
 	int CurrentIndex = 0;
 	const UPCGData* InData = nullptr;
 	UPCGData* OutData = nullptr;
+
+	FPCGMatchAndSetPartition* Partition = nullptr;
+	bool bIsPartitionOwner = false;
 };
 
 UENUM()
