[PCG] Tentative fix to force the shader to be recompiled

#rnx

diff --git a/Engine/Plugins/PCG/Shaders/Private/PCGDataCollectionDataInterface.ush b/Engine/Plugins/PCG/Shaders/Private/PCGDataCollectionDataInterface.ush
--- a/Engine/Plugins/PCG/Shaders/Private/PCGDataCollectionDataInterface.ush
+++ b/Engine/Plugins/PCG/Shaders/Private/PCGDataCollectionDataInterface.ush
@@ -55,7 +55,7 @@ uint {DataInterfaceName}_ElementCountersPresent;
 
 uint GetNumData_{DataInterfaceName}()
 {
-	// We use the most significant bit of NumData to store the KernelExecuted flag
+	// We use the most significant bit of NumData to store the KernelExecuted flag.
 	return {DataInterfaceName}_DataCollectionBuffer.Load(0) & ~KERNEL_EXECUTED_FLAG;
 }
 
