Added overridable tag list on PCG ISM descriptors

diff --git a/Engine/Plugins/PCG/Source/PCG/Private/Helpers/PCGActorHelpers.cpp b/Engine/Plugins/PCG/Source/PCG/Private/Helpers/PCGActorHelpers.cpp
--- a/Engine/Plugins/PCG/Source/PCG/Private/Helpers/PCGActorHelpers.cpp
+++ b/Engine/Plugins/PCG/Source/PCG/Private/Helpers/PCGActorHelpers.cpp
@@ -103,6 +103,12 @@ UPCGManagedISMComponent* UPCGActorHelpers::GetOrCreateManagedISMC(AActor* InTarg
 		{
 			ISMC->ComponentTags.AddUnique(ComponentTag);
 		}
+
+		const TArray<FString> AdditionalComponentTags = PCGHelpers::GetStringArrayFromCommaSeparatedList(InParams.Descriptor.AdditionalCommaSeparatedTags);
+		for (const FString& AdditionalComponentTag : AdditionalComponentTags)
+		{
+			ISMC->ComponentTags.AddUnique(FName(AdditionalComponentTag));
+		}
 	};
 
 	// If the component class is invalid, default to HISM.
diff --git a/Engine/Plugins/PCG/Source/PCG/Private/MeshSelectors/PCGISMDescriptor.cpp b/Engine/Plugins/PCG/Source/PCG/Private/MeshSelectors/PCGISMDescriptor.cpp
--- a/Engine/Plugins/PCG/Source/PCG/Private/MeshSelectors/PCGISMDescriptor.cpp
+++ b/Engine/Plugins/PCG/Source/PCG/Private/MeshSelectors/PCGISMDescriptor.cpp
@@ -2,6 +2,8 @@
 
 #include "MeshSelectors/PCGISMDescriptor.h"
 
+#include "Helpers/PCGHelpers.h"
+
 #include "Components/InstancedStaticMeshComponent.h"
 
 #include UE_INLINE_GENERATED_CPP_BY_NAME(PCGISMDescriptor)
@@ -29,6 +31,13 @@ void FPCGSoftISMComponentDescriptor::InitFrom(const UStaticMeshComponent* Compon
 void FPCGSoftISMComponentDescriptor::InitComponent(UInstancedStaticMeshComponent* ISMComponent) const
 {
 	ISMComponent->ComponentTags = ComponentTags;
+
+	const TArray<FString> AdditionalTags = PCGHelpers::GetStringArrayFromCommaSeparatedList(AdditionalCommaSeparatedTags);
+	for (const FString& AdditionalTag : AdditionalTags)
+	{
+		ISMComponent->ComponentTags.Emplace(AdditionalTag);
+	}
+
 	Super::InitComponent(ISMComponent);
 }
 
diff --git a/Engine/Plugins/PCG/Source/PCG/Public/MeshSelectors/PCGISMDescriptor.h b/Engine/Plugins/PCG/Source/PCG/Public/MeshSelectors/PCGISMDescriptor.h
--- a/Engine/Plugins/PCG/Source/PCG/Public/MeshSelectors/PCGISMDescriptor.h
+++ b/Engine/Plugins/PCG/Source/PCG/Public/MeshSelectors/PCGISMDescriptor.h
@@ -23,6 +23,9 @@ public:
 	UPROPERTY(EditAnywhere, Category="Component Settings|Tags", meta = (DisplayAfter = "ComponentClass"))
 	TArray<FName> ComponentTags;
 
+	UPROPERTY(EditAnywhere, Category="Component Settings|Tags", meta = (DisplayAfter = "ComponentTags", Tooltip="Comma-separated tag list. Normally should be left empty, and used as a property override instead."))
+	FString AdditionalCommaSeparatedTags;
+
 	PCG_API bool SerializeFromMismatchedTag(const FPropertyTag& Tag, FStructuredArchive::FSlot Slot);
 };
 
