Fixed small warning vs error in match & set

#rnx

diff --git a/Engine/Plugins/PCG/Source/PCG/Private/Elements/PCGMatchAndSetAttributes.cpp b/Engine/Plugins/PCG/Source/PCG/Private/Elements/PCGMatchAndSetAttributes.cpp
--- a/Engine/Plugins/PCG/Source/PCG/Private/Elements/PCGMatchAndSetAttributes.cpp
+++ b/Engine/Plugins/PCG/Source/PCG/Private/Elements/PCGMatchAndSetAttributes.cpp
@@ -947,7 +947,17 @@ bool FPCGMatchAndSetAttributesElement::PrepareDataInternal(FPCGContext* InContex
 
 	EPCGTimeSliceInitResult InitResult = TimeSlicedContext->InitializePerExecutionState([Settings, &ParamDataInputs, &Inputs, &MaxDistanceDataInputs, &CreatePartition](FPCGMatchAndSetAttributesElement::ContextType* Context, FPCGMatchAndSetAttributesExecutionState& OutState) -> EPCGTimeSliceInitResult
 	{
-		if(ParamDataInputs.IsEmpty() || (ParamDataInputs.Num() != 1 && ParamDataInputs.Num() != Inputs.Num()))
+		if (ParamDataInputs.IsEmpty())
+		{
+			if (Settings->bWarnIfNoMatchData)
+			{
+				PCGE_LOG_C(Warning, GraphAndLog, Context, LOCTEXT("NoMatchData", "Must have exactly one Attribute Set to match against"));
+			}
+
+			return EPCGTimeSliceInitResult::NoOperation;
+		}
+
+		if(ParamDataInputs.Num() != 1 && ParamDataInputs.Num() != Inputs.Num())
 		{
 			PCGLog::InputOutput::LogInvalidCardinalityError(PCGPinConstants::DefaultInputLabel, PCGMatchAndSetAttributesConstants::MatchDataLabel, Context);
 			return EPCGTimeSliceInitResult::AbortExecution;
