Small adjustments following the debug object tree view changes.
- Selection does not focus on node anymore
- Added jump to & focus buttons in the tree view on entry highlight.
- Updated icon for jump.
- Better icon coloring to emphasize what is the current graph and also local errors.

diff --git a/Engine/Plugins/PCG/Source/PCGEditor/Private/PCGEditor.cpp b/Engine/Plugins/PCG/Source/PCGEditor/Private/PCGEditor.cpp
--- a/Engine/Plugins/PCG/Source/PCGEditor/Private/PCGEditor.cpp
+++ b/Engine/Plugins/PCG/Source/PCGEditor/Private/PCGEditor.cpp
@@ -2840,8 +2840,10 @@ void FPCGEditor::OnNodeDoubleClicked(UEdGraphNode* Node)
 			// If we're inspecting, we'll try to find a match in the stacks for subgraphs instead of relying on the static/template subgraph
 			if (GetStackBeingInspected() && DebugObjectTreeWidget)
 			{
-				DebugObjectTreeWidget->GetFirstStackFromSelection(PCGNode, /*Graph=*/nullptr, StackToInspect);
-				Object = const_cast<UPCGGraph*>(StackToInspect.GetGraphForCurrentFrame());
+				if (DebugObjectTreeWidget->GetFirstStackFromSelection(PCGNode, /*Graph=*/nullptr, StackToInspect))
+				{
+					Object = const_cast<UPCGGraph*>(StackToInspect.GetGraphForCurrentFrame());
+				}
 			}
 
 			if (Object)
diff --git a/Engine/Plugins/PCG/Source/PCGEditor/Private/PCGEditorStyle.cpp b/Engine/Plugins/PCG/Source/PCGEditor/Private/PCGEditorStyle.cpp
--- a/Engine/Plugins/PCG/Source/PCGEditor/Private/PCGEditorStyle.cpp
+++ b/Engine/Plugins/PCG/Source/PCGEditor/Private/PCGEditorStyle.cpp
@@ -140,6 +140,7 @@ FPCGEditorStyle::FPCGEditorStyle() : FSlateStyleSet("PCGEditorStyle")
 	Set("PCG.Editor.ZoomToSelection", new IMAGE_BRUSH_SVG("Style/PCG_Editor_ZoomToSelection", Icon16x16));
 	Set("PCG.Editor.Pause", new IMAGE_BRUSH_SVG("Icons/PCG_16", Icon16x16));
 	Set("PCG.Editor.AlternatePause", new IMAGE_BRUSH_SVG("Style/PCG_Editor_AlternatePause", Icon16x16));
+	Set("PCG.Editor.JumpTo", new IMAGE_BRUSH_SVG("Style/PCG_DebugTree_JumpTo", Icon16x16));
 
 	// Text editor styles
 	{
diff --git a/Engine/Plugins/PCG/Source/PCGEditor/Private/SPCGEditorGraphDebugObjectTree.cpp b/Engine/Plugins/PCG/Source/PCGEditor/Private/SPCGEditorGraphDebugObjectTree.cpp
--- a/Engine/Plugins/PCG/Source/PCGEditor/Private/SPCGEditorGraphDebugObjectTree.cpp
+++ b/Engine/Plugins/PCG/Source/PCGEditor/Private/SPCGEditorGraphDebugObjectTree.cpp
@@ -11,6 +11,7 @@
 #include "PCGEditor.h"
 #include "PCGEditorGraph.h"
 #include "PCGEditorStyle.h"
+#include "AssetDefinitions/AssetDefinition_PCGGraphInterface.h"
 
 #include "PropertyCustomizationHelpers.h"
 #include "Selection.h"
@@ -28,22 +29,11 @@ namespace PCGEditorGraphDebugObjectTree
 {
 	void GetRowIconState(FPCGEditorGraphDebugObjectItemPtr Item, const FSlateBrush*& OutBrush, FLinearColor& OutColorAndOpacity)
 	{
-		const UObject* ItemObject = Item->GetObject();
+		OutBrush = Item->GetIcon();
 
-		if (!IsValid(ItemObject))
-		{
-			OutBrush = FAppStyle::Get().GetBrush("LevelEditor.Tabs.Debug");
-		}
-		else if (ItemObject->IsA<AActor>())
-		{
-			OutBrush = FSlateIconFinder::FindIconBrushForClass(ItemObject->GetClass());
-		}
-		else
-		{
-			OutBrush = FSlateIconFinder::FindIconBrushForClass(UPCGGraph::StaticClass());
-		}
+		const FLinearColor DefaultGraphColor = CastChecked<UAssetDefinition_PCGGraphInterface>(UAssetDefinition_PCGGraphInterface::StaticClass()->ClassDefaultObject)->GetAssetColor();
 
-		OutColorAndOpacity = Item->IsDebuggable() ? FLinearColor::White : FLinearColor(1.0f, 1.0f, 1.0f, 0.5f);
+		OutColorAndOpacity = Item->IsDebuggable() ? DefaultGraphColor : FLinearColor(1.0f, 1.0f, 1.0f, 0.5f);
 	}
 
 	void GetErrorInfo(FPCGEditorGraphDebugObjectItemPtr Item, EVisibility& OutIconVisibility, FText& OutIconTooltipText, FLinearColor& OutIconColorAndOpacity, FLinearColor& OutRowTextColorAndOpacity)
@@ -142,6 +132,11 @@ void FPCGEditorGraphDebugObjectItem::AddChild(TSharedRef<FPCGEditorGraphDebugObj
 	Children.Add(MoveTemp(InChild));
 }
 
+const FSlateBrush* FPCGEditorGraphDebugObjectItem::GetIcon() const
+{
+	return FSlateIconFinder::FindIconBrushForClass(UPCGGraph::StaticClass());
+}
+
 const TSet<TSharedPtr<FPCGEditorGraphDebugObjectItem>>& FPCGEditorGraphDebugObjectItem::GetChildren() const
 {
 	return Children;
@@ -199,6 +194,12 @@ FString FPCGEditorGraphDebugObjectItem_Actor::GetLabel(bool bForSorting) const
 	return Actor.IsValid() ? Actor->GetActorNameOrLabel() : FString();
 }
 
+const FSlateBrush* FPCGEditorGraphDebugObjectItem_Actor::GetIcon() const
+{
+	const UObject* Object = GetObject();
+	return FSlateIconFinder::FindIconBrushForClass(Object ? Object->GetClass() : AActor::StaticClass());
+}
+
 FString FPCGEditorGraphDebugObjectItem_PCGComponent::GetLabel(bool bForSorting) const
 {
 	if (PCGComponent.IsValid() && PCGGraph.IsValid())
@@ -236,6 +237,18 @@ FString FPCGEditorGraphDebugObjectItem_PCGSubgraph::GetLabel(bool bForSorting) c
 	return FString();
 }
 
+const FSlateBrush* FPCGEditorGraphDebugObjectItem_PCGSubgraph::GetIcon() const
+{
+	if (SubgraphType == EPCGEditorSubgraphNodeType::LoopSubgraph)
+	{
+		return FAppStyle::Get().GetBrush("GraphEditor.Macro.Loop_16x");
+	}
+	else
+	{
+		return FPCGEditorGraphDebugObjectItem::GetIcon();
+	}
+}
+
 FString FPCGEditorGraphDebugObjectItem_PCGLoopIndex::GetLabel(bool bForSorting) const
 {
 	return FString::Format(TEXT("{0}"), { LoopIndex });
@@ -246,7 +259,11 @@ void SPCGEditorGraphDebugObjectItemRow::Construct(const FArguments& InArgs, cons
 	Item = InItem;
 
 	// This function should auto-expand the row and select the deepest entry as the debug object if it is unambiguous (the only entry at its level in the tree).
-	DoubleClickFunc = InArgs._OnDoubleClickFunc;
+	OnDoubleClick = InArgs._OnDoubleClick;
+	OnJumpTo = InArgs._OnJumpTo;
+	OnFocus = InArgs._OnFocus;
+	CanJumpTo = InArgs._CanJumpTo;
+	CanFocus = InArgs._CanFocus;
 
 	// Computed once during construction as the tree is refreshed on relevant events.
 	const FSlateBrush* RowIcon = nullptr;
@@ -259,6 +276,14 @@ void SPCGEditorGraphDebugObjectItemRow::Construct(const FArguments& InArgs, cons
 	FLinearColor ErrorIconColorAndOpacity;
 	FLinearColor RowTextColorAndOpacity;
 	PCGEditorGraphDebugObjectTree::GetErrorInfo(InItem, ErrorIconVisibility, ErrorIconTooltipText, ErrorIconColorAndOpacity, RowTextColorAndOpacity);
+
+	// Transfer warning/error color to icon if there is a local error, but preserve original opacity (otherwise it will look weird for non-local graphs).
+	if (ErrorIconVisibility != EVisibility::Hidden && RowTextColorAndOpacity != FLinearColor::White)
+	{
+		const float OriginalOpacity = RowIconColorAndOpacity.A;
+		RowIconColorAndOpacity = RowTextColorAndOpacity;
+		RowIconColorAndOpacity.A = OriginalOpacity;
+	}
 	
 	ChildSlot
 	[
@@ -273,17 +298,18 @@ void SPCGEditorGraphDebugObjectItemRow::Construct(const FArguments& InArgs, cons
 			.Image(RowIcon)
 		]
 		+SHorizontalBox::Slot()
-		.AutoWidth()
 		[
 			SNew(STextBlock)
 			.Text(FText::FromString(Item->GetLabel()))
+			.ToolTipText(FText::FromString(Item->GetLabel()))
 			.Font_Lambda([this]() { return PCGEditorGraphDebugObjectTree::GetRowFont(Item); })
 			// Highlight based on data available for currently inspected node. Computed dynamically in lambda to respond to inspection changes.
 			.ColorAndOpacity_Lambda([this, RowTextColorAndOpacity](){ return Item->IsGrayedOut() ? FColor(75, 75, 75) : RowTextColorAndOpacity; })
+			.OverflowPolicy(ETextOverflowPolicy::Ellipsis)
 		]
 		+SHorizontalBox::Slot()
 		.AutoWidth()
-		.Padding(20.0f, 0.0f)
+		.Padding(2.0f, 0.0f)
 		[
 			SNew(SImage)
 			.Visibility(ErrorIconVisibility)
@@ -291,19 +317,97 @@ void SPCGEditorGraphDebugObjectItemRow::Construct(const FArguments& InArgs, cons
 			.Image(FAppStyle::Get().GetBrush("Icons.Error"))
 			.ColorAndOpacity(ErrorIconColorAndOpacity)
 		]
+		+SHorizontalBox::Slot()
+		.AutoWidth()
+		.Padding(2.0f, 0.0f)
+		[
+			SNew(SButton)
+			.ButtonStyle(FAppStyle::Get(), "SimpleButton")
+			.Visibility_Lambda([this]() { return IsHovered() ? EVisibility::Visible : EVisibility::Hidden; })
+			.OnClicked(this, &SPCGEditorGraphDebugObjectItemRow::FocusClicked)
+			.IsEnabled(this, &SPCGEditorGraphDebugObjectItemRow::IsFocusEnabled)
+			.ToolTipText(LOCTEXT("FocusOnNode", "Show the calling node in the current graph."))
+			.ContentPadding(0.0f)
+			[
+				SNew(SImage)
+				.Image(FPCGEditorStyle::Get().GetBrush("PCG.Editor.ZoomToSelection"))
+				.ColorAndOpacity(FSlateColor::UseForeground())
+			]
+		]
+		+SHorizontalBox::Slot()
+		.AutoWidth()
+		.Padding(2.0f, 0.0f)
+		[
+			SNew(SButton)
+			.ButtonStyle(FAppStyle::Get(), "SimpleButton")
+			.Visibility_Lambda([this]() { return IsHovered() ? EVisibility::Visible : EVisibility::Hidden; })
+			.OnClicked(this, &SPCGEditorGraphDebugObjectItemRow::JumpToClicked)
+			.IsEnabled(this, &SPCGEditorGraphDebugObjectItemRow::IsJumpToEnabled)
+			.ToolTipText(LOCTEXT("JumpToGraph", "Jump to graph, with this debug object context."))
+			.ContentPadding(0.0f)
+			[
+				SNew(SImage)
+				.Image(FPCGEditorStyle::Get().GetBrush("PCG.Editor.JumpTo"))
+				.ColorAndOpacity(FSlateColor::UseForeground())
+			]
+		]
 	];
 }
 
 FReply SPCGEditorGraphDebugObjectItemRow::OnMouseButtonDoubleClick(const FGeometry& InMyGeometry, const FPointerEvent& InMouseEvent)
 {
-	if (Item.IsValid() && DoubleClickFunc)
+	if (Item.IsValid() && OnDoubleClick.IsBound())
+	{
+		OnDoubleClick.Execute(Item);
+	}
+
+	return FReply::Handled();
+}
+
+FReply SPCGEditorGraphDebugObjectItemRow::FocusClicked() const
+{
+	if (Item.IsValid() && OnFocus.IsBound())
 	{
-		DoubleClickFunc(Item);
+		OnFocus.Execute(Item);
 	}
 
 	return FReply::Handled();
 }
 
+FReply SPCGEditorGraphDebugObjectItemRow::JumpToClicked() const
+{
+	if (Item.IsValid() && OnJumpTo.IsBound())
+	{
+		OnJumpTo.Execute(Item);
+	}
+
+	return FReply::Handled();
+}
+
+bool SPCGEditorGraphDebugObjectItemRow::IsJumpToEnabled() const
+{
+	if (Item.IsValid() && CanJumpTo.IsBound())
+	{
+		return CanJumpTo.Execute(Item);
+	}
+	else
+	{
+		return true;
+	}
+}
+
+bool SPCGEditorGraphDebugObjectItemRow::IsFocusEnabled() const
+{
+	if (Item.IsValid() && CanFocus.IsBound())
+	{
+		return CanFocus.Execute(Item);
+	}
+	else
+	{
+		return true;
+	}
+}
+
 SPCGEditorGraphDebugObjectTree::~SPCGEditorGraphDebugObjectTree()
 {
 	FCoreUObjectDelegates::OnPreObjectPropertyChanged.RemoveAll(this);
@@ -325,10 +429,6 @@ void SPCGEditorGraphDebugObjectTree::Construct(const FArguments& InArgs, TShared
 	FCoreUObjectDelegates::OnObjectConstructed.AddSP(this, &SPCGEditorGraphDebugObjectTree::OnObjectConstructed);
 	USelection::SelectionChangedEvent.AddSP(this, &SPCGEditorGraphDebugObjectTree::OnEditorSelectionChanged);
 
-	const TSharedRef<SScrollBar> HorizontalScrollBar = SNew(SScrollBar)
-		.Orientation(Orient_Horizontal)
-		.Thickness(FVector2D(12.0f, 12.0f));
-
 	const TSharedRef<SScrollBar> VerticalScrollBar = SNew(SScrollBar)
 		.Orientation(Orient_Vertical)
 		.Thickness(FVector2D(12.0f, 12.0f));
@@ -391,14 +491,7 @@ void SPCGEditorGraphDebugObjectTree::Construct(const FArguments& InArgs, TShared
 			+SHorizontalBox::Slot()
 			.FillWidth(1.0f)
 			[
-				SNew(SScrollBox)
-				.Orientation(Orient_Horizontal)
-				.ExternalScrollbar(HorizontalScrollBar)
-				+SScrollBox::Slot()
-				.FillSize(1.0f)
-				[
-					DebugObjectTreeView.ToSharedRef()
-				]
+				DebugObjectTreeView.ToSharedRef()
 			]
 			+SHorizontalBox::Slot()
 			.AutoWidth()
@@ -406,15 +499,6 @@ void SPCGEditorGraphDebugObjectTree::Construct(const FArguments& InArgs, TShared
 				VerticalScrollBar
 			]
 		]
-		+SVerticalBox::Slot()
-		.AutoHeight()
-		[
-			SNew(SHorizontalBox)
-			+SHorizontalBox::Slot()
-			[
-				HorizontalScrollBar
-			]
-		]
 	];
 
 	RequestRefresh();
@@ -591,7 +675,7 @@ bool SPCGEditorGraphDebugObjectTree::SetDebugObjectFromStack(const FPCGStack& In
 	}
 }
 
-void SPCGEditorGraphDebugObjectTree::ExpandAndSelectDebugObject(FPCGEditorGraphDebugObjectItemPtr InItem)
+void SPCGEditorGraphDebugObjectTree::ExpandAndSelectDebugObject(const FPCGEditorGraphDebugObjectItemPtr& InItem)
 {
 	if (!InItem)
 	{
@@ -833,7 +917,7 @@ void SPCGEditorGraphDebugObjectTree::AddStacksToTree(const TArray<FPCGStack>& St
 		};
 
 		// Generic function to create the subgraph/loop item and hook it up to the parent properly
-		auto AddSubgraphOrLoopItemToStack = [&Stack, &InOutStackToItem, &AddToGraphItemsAndAttachItemToParent](const UPCGGraph* InGraph, const UPCGNode* InSubgraphNode, int FrameCutoff, bool bInIsDebuggable, bool bInDisplayGrayedOut, bool bIsDynamic)
+		auto AddSubgraphOrLoopItemToStack = [&Stack, &InOutStackToItem, &AddToGraphItemsAndAttachItemToParent](const UPCGGraph* InGraph, const UPCGNode* InSubgraphNode, int FrameCutoff, bool bInIsDebuggable, bool bInDisplayGrayedOut, EPCGEditorSubgraphNodeType InSubgraphType)
 		{
 			FPCGStack GraphStack = Stack;
 			GraphStack.GetStackFramesMutable().SetNum(FrameCutoff);
@@ -842,7 +926,7 @@ void SPCGEditorGraphDebugObjectTree::AddStacksToTree(const TArray<FPCGStack>& St
 			{
 				FPCGEditorGraphDebugObjectItemPtr GraphItem = InOutStackToItem.Emplace(
 					GraphStack,
-					MakeShared<FPCGEditorGraphDebugObjectItem_PCGSubgraph>(InSubgraphNode, InGraph, GraphStack, bInIsDebuggable, bInDisplayGrayedOut, bIsDynamic));
+					MakeShared<FPCGEditorGraphDebugObjectItem_PCGSubgraph>(InSubgraphNode, InGraph, GraphStack, bInIsDebuggable, bInDisplayGrayedOut, InSubgraphType));
 
 				AddToGraphItemsAndAttachItemToParent(GraphStack, GraphItem);
 			}
@@ -886,7 +970,7 @@ void SPCGEditorGraphDebugObjectTree::AddStacksToTree(const TArray<FPCGStack>& St
 				// Previous stack was node, therefore static subgraph.
 				else if (const UPCGNode* SubgraphNode = Cast<const UPCGNode>(PreviousStackFrame.Object.Get()))
 				{
-					AddSubgraphOrLoopItemToStack(StackGraph, SubgraphNode, FrameIndex + 1, bIsDebuggable, bDisplayGrayedOut, /*bIsDynamic=*/false);
+					AddSubgraphOrLoopItemToStack(StackGraph, SubgraphNode, FrameIndex + 1, bIsDebuggable, bDisplayGrayedOut, EPCGEditorSubgraphNodeType::StaticSubgraph);
 				}
 				// Previous stack was loop index, therefore loop subgraph.
 				else if (FrameIndex >= 2 && (PreviousStackFrame.LoopIndex != INDEX_NONE))
@@ -895,7 +979,7 @@ void SPCGEditorGraphDebugObjectTree::AddStacksToTree(const TArray<FPCGStack>& St
 					if (ensure(LoopSubgraphNode))
 					{
 						// Take the stack up to the looped subgraph node, add a item for the node + graph
-						AddSubgraphOrLoopItemToStack(StackGraph, LoopSubgraphNode, FrameIndex - 1, /*bIsDebuggable=*/false, /*bIsGrayedOut=*/false, /*bIsDynamic=*/true);
+						AddSubgraphOrLoopItemToStack(StackGraph, LoopSubgraphNode, FrameIndex - 1, /*bIsDebuggable=*/false, /*bIsGrayedOut=*/bDisplayGrayedOut, EPCGEditorSubgraphNodeType::LoopSubgraph);
 
 						// Take full stack up until this point which will be the unique stack for the loop iteration.
 						FPCGStack LoopIterationStack = Stack;
@@ -917,7 +1001,7 @@ void SPCGEditorGraphDebugObjectTree::AddStacksToTree(const TArray<FPCGStack>& St
 					const UPCGNode* DynamicSubgraphNode = Cast<const UPCGNode>(StackFrames[FrameIndex - 2].Object.Get());
 					if (DynamicSubgraphNode)
 					{
-						AddSubgraphOrLoopItemToStack(StackGraph, DynamicSubgraphNode, FrameIndex + 1, bIsDebuggable, bDisplayGrayedOut, /*bIsDynamic=*/true);
+						AddSubgraphOrLoopItemToStack(StackGraph, DynamicSubgraphNode, FrameIndex + 1, bIsDebuggable, bDisplayGrayedOut, EPCGEditorSubgraphNodeType::DynamicSubgraph);
 					}
 				}
 			}
@@ -1172,7 +1256,11 @@ TSharedRef<ITableRow> SPCGEditorGraphDebugObjectTree::MakeTreeRowWidget(FPCGEdit
 	return SNew( STableRow<TSharedPtr<SPCGEditorGraphDebugObjectItemRow>>, InOwnerTable)
 	[
 		SNew(SPCGEditorGraphDebugObjectItemRow, InOwnerTable, InItem)
-		.OnDoubleClickFunc([this](const FPCGEditorGraphDebugObjectItemPtr& Item) { ExpandAndSelectFirstLeafDebugObject(Item); } )
+		.OnDoubleClick(this, &SPCGEditorGraphDebugObjectTree::ExpandAndSelectDebugObject)
+		.OnJumpTo(this, &SPCGEditorGraphDebugObjectTree::JumpToGraphInTree)
+		.CanJumpTo(this, &SPCGEditorGraphDebugObjectTree::CanJumpToGraphInTree)
+		.OnFocus(this, &SPCGEditorGraphDebugObjectTree::FocusOnItem)
+		.CanFocus(this, &SPCGEditorGraphDebugObjectTree::CanFocusOnItem)
 	];
 }
 
@@ -1190,13 +1278,6 @@ void SPCGEditorGraphDebugObjectTree::OnGetChildren(FPCGEditorGraphDebugObjectIte
 
 void SPCGEditorGraphDebugObjectTree::OnSelectionChanged(FPCGEditorGraphDebugObjectItemPtr InItem, ESelectInfo::Type InSelectInfo)
 {
-	// If the user selects a new item, record the previous stack. This is helpful to give previous selection information to 
-	// context menu commands, because right clicking an item changes the selection.
-	if (InSelectInfo != ESelectInfo::Direct)
-	{
-		PreviouslySelectedStack = SelectedStack;
-	}
-
 	// Unmark the previously selected object
 	if (FPCGEditorGraphDebugObjectItemPtr PreviousItem = GetItemFromStack(SelectedStack))
 	{
@@ -1233,10 +1314,8 @@ void SPCGEditorGraphDebugObjectTree::OnSelectionChanged(FPCGEditorGraphDebugObje
 		}
 	}
 
-	// Backup selected stack if we're going to jump elsewhere
-	SelectedStackBeforeCulling = SelectedStack;
-
 	const UPCGGraph* CurrentGraph = PCGEditor.Pin()->GetPCGGraph();
+	bool bStackInspectedSet = false;
 
 	// Only attempt to inspect stacks that correspond to the edited graph. Other graphs need to be inspected in their own editor.
 	if (InItem)
@@ -1244,49 +1323,22 @@ void SPCGEditorGraphDebugObjectTree::OnSelectionChanged(FPCGEditorGraphDebugObje
 		// Basically - if the selection is something from "upstream", then we can clear the stack being inspected.
 		if (SelectedStack.HasObject(CurrentGraph))
 		{
-			const UPCGNode* LastNode = nullptr;
-
-			// Two things should happen here:
 			// If the last graph in the stack is the current graph, then we can just use that as the selected stack
-			// Otherwise, we'll cut the stack down to the last instance of this graph, and instead move the editor selection
-			// to the first node after the graph
 			TArray<FPCGStackFrame>& StackFrames = SelectedStack.GetStackFramesMutable();
-			while(!StackFrames.IsEmpty())
+			while(!StackFrames.IsEmpty() && StackFrames.Last().Object.Get() != CurrentGraph)
 			{
-				const UObject* LastObject = StackFrames.Last().Object.Get();
-				if (LastObject == CurrentGraph)
-				{
-					break;
-				}
-				else if (const UPCGNode* Node = Cast<const UPCGNode>(LastObject))
-				{
-					LastNode = Node;
-				}
-
 				StackFrames.Pop();
 			}
 
 			if (ensure(!SelectedStack.GetStackFrames().IsEmpty()))
 			{
 				PCGEditor.Pin()->SetStackBeingInspected(SelectedStack);
-
-				if (LastNode)
-				{
-					PCGEditor.Pin()->JumpToNode(LastNode);
-				}
-			}
-			else
-			{
-				PCGEditor.Pin()->ClearStackBeingInspected();
+				bStackInspectedSet = true;
 			}
 		}
-		else
-		{
-			// Selected an upstream item - let's clear the stack
-			PCGEditor.Pin()->ClearStackBeingInspected();
-		}
 	}
-	else
+
+	if (!bStackInspectedSet)
 	{
 		PCGEditor.Pin()->ClearStackBeingInspected();
 	}
@@ -1395,7 +1447,7 @@ TSharedPtr<SWidget> SPCGEditorGraphDebugObjectTree::OpenContextMenu()
 	MenuBuilder.AddMenuEntry(
 		LOCTEXT("JumpToSelectedGraph", "Jump To"),
 		LOCTEXT("JumpToSelectedGraphTooltip", "Jumps to the selected graph."),
-		FSlateIcon(FAppStyle::GetAppStyleSetName(), "SystemWideCommands.FindInContentBrowser.Small"),
+		FSlateIcon(FPCGEditorStyle::Get().GetStyleSetName(), "PCG.Editor.JumpTo"),
 		FUIAction(FExecuteAction::CreateSP(this, &SPCGEditorGraphDebugObjectTree::ContextMenu_JumpToGraphInTree), FCanExecuteAction::CreateSP(this, &SPCGEditorGraphDebugObjectTree::ContextMenu_JumpToGraphInTree_CanExecute))
 	);
 
@@ -1404,66 +1456,72 @@ TSharedPtr<SWidget> SPCGEditorGraphDebugObjectTree::OpenContextMenu()
 
 void SPCGEditorGraphDebugObjectTree::ContextMenu_JumpToGraphInTree()
 {
-	// If we were previously debugging a target and we're jumping to a parent graph, jump to the subgraph/loop node
-	// that corresponds to the previously selected stack.
-	// 
-	// Example:
-	//		Jump target: Component/TopGraph
-	//		Previous stack: Component/TopGraph/SubgraphNode/Subgraph
-	// 
-	// Here the "subject" of the jump here is SubgraphNode - obtain this from the previous stack, frame index after TopGraph.
-
 	const TArray<FPCGEditorGraphDebugObjectItemPtr> SelectedItems = DebugObjectTreeView->GetSelectedItems();
 	for (const FPCGEditorGraphDebugObjectItemPtr& SelectedItem : SelectedItems)
 	{
-		UPCGGraph* JumpToPCGGraph = nullptr;
-		UPCGNode* JumpToPCGNode = nullptr;
+		JumpToGraphInTree(SelectedItem);
+	}
+}
 
-		if (SelectedItem->IsDebuggable())
+void SPCGEditorGraphDebugObjectTree::JumpToGraphInTree(const FPCGEditorGraphDebugObjectItemPtr& Item)
+{
+	if (!Item || !Item->GetPCGStack())
+	{
+		return;
+	}
+
+	FPCGStack Stack = *Item->GetPCGStack();
+
+	// Example stack:
+	//     Component/TopGraph/SubgraphNode/Subgraph/LoopSubgraphNode/LoopIndex/LoopSubgraph
+	//                                       ^ static subgraph
+	//     Component/TopGraph/SubgraphNode/INDEX_NONE/Subgraph/...
+	//                                       ^ dynamic subgraph
+
+	UPCGGraph* JumpToPCGGraph = nullptr;
+	UPCGNode* JumpToPCGNode = nullptr;
+
+	// If the item is about this graph, then we'll jump to this instance's caller - we'll walk back until we find a graph (e.g. parent)
+	// and the next entry should be our caller node.
+	if (Item->IsDebuggable())
+	{
+		// Debuggable target graphs correspond to the currently edited graph. For this case open the parent graph and jump to the corresponding subgraph node.
+		// Search the stack for the parent graph.
+		for (int i = Stack.GetStackFrames().Num() - 2; i > 0; --i) 
 		{
-			// Debuggable target graphs correspond to the currently edited graph. For this case open the parent graph and jump to the corresponding subgraph node.
-			
-			// Search the stack for the parent graph.
-			for (int i = SelectedStack.GetStackFrames().Num() - 2; i > 0; --i) 
+			if (Stack.GetStackFrames()[i].Object != nullptr && Stack.GetStackFrames()[i].Object->IsA<UPCGGraph>())
 			{
-				if (SelectedStack.GetStackFrames()[i].Object != nullptr && SelectedStack.GetStackFrames()[i].Object->IsA<UPCGGraph>())
-				{
-					JumpToPCGGraph = const_cast<UPCGGraph*>(Cast<const UPCGGraph>(SelectedStack.GetStackFrames()[i].Object.Get()));
-					JumpToPCGNode = const_cast<UPCGNode*>(Cast<const UPCGNode>(SelectedStack.GetStackFrames()[i + 1].Object.Get()));
-					break;
-				}
-			}
-		}
-		else 
-		{
-			// Non-debuggable target graphs are graphs that are not the current edited graph.
-			JumpToPCGGraph = const_cast<UPCGGraph*>(SelectedItem->GetPCGGraph());
+				JumpToPCGGraph = const_cast<UPCGGraph*>(Cast<const UPCGGraph>(Stack.GetStackFrames()[i].Object.Get()));
+				JumpToPCGNode = const_cast<UPCGNode*>(Cast<const UPCGNode>(Stack.GetStackFrames()[i + 1].Object.Get()));
 
-			// When the user right clicks, the selection moves. PreviouslySelectedStack will hold the previous selection, and
-			// we retrieve our subject node from there, which can be found 1 frame after the jump to graph, hence the Num() below.
-			const int SubjectNodeFrameIndex = SelectedStack.GetStackFrames().Num();
-			if (PreviouslySelectedStack.GetStackFrames().IsValidIndex(SubjectNodeFrameIndex))
-			{
-				JumpToPCGNode = const_cast<UPCGNode*>(Cast<const UPCGNode>(PreviouslySelectedStack.GetStackFrames()[SubjectNodeFrameIndex].Object.Get()));
+				// Cull remaining remaining frames so the stack selection into the parent graph editor works.
+				Stack.GetStackFramesMutable().SetNum(i);
+				
+				break;
 			}
 		}
+	}
+	else 
+	{
+		// Two cases here: upstream or downstream. In both cases, there's no need to jump to a given node as there is no concept of caller here.
+		JumpToPCGGraph = const_cast<UPCGGraph*>(Item->GetPCGGraph());
+	}
+
+	if (JumpToPCGGraph)
+	{
+		GEditor->GetEditorSubsystem<UAssetEditorSubsystem>()->OpenEditorForAsset(JumpToPCGGraph);
+		IAssetEditorInstance* EditorInstance = GEditor->GetEditorSubsystem<UAssetEditorSubsystem>()->FindEditorForAsset(JumpToPCGGraph, /*bFocusIfOpen*/true);
+		FPCGEditor* OtherPCGEditor = static_cast<FPCGEditor*>(EditorInstance);
 
-		if (JumpToPCGGraph)
+		if (OtherPCGEditor)
 		{
-			GEditor->GetEditorSubsystem<UAssetEditorSubsystem>()->OpenEditorForAsset(JumpToPCGGraph);
-			IAssetEditorInstance* EditorInstance = GEditor->GetEditorSubsystem<UAssetEditorSubsystem>()->FindEditorForAsset(JumpToPCGGraph, /*bFocusIfOpen*/true);
-			FPCGEditor* OtherPCGEditor = static_cast<FPCGEditor*>(EditorInstance);
+			// Implementation note - for selection purposes we had culled the stack to be in the local referential
+			// however in the case of the jump here we want the full stack
+			OtherPCGEditor->SetStackBeingInspectedFromAnotherEditor(Stack);
 
-			if (OtherPCGEditor)
+			if (JumpToPCGNode)
 			{
-				// Implementation note - for selection purposes we had culled the stack to be in the local referential
-				// however in the case of the jump here we want the full stack
-				OtherPCGEditor->SetStackBeingInspectedFromAnotherEditor(SelectedStackBeforeCulling);
-
-				if (JumpToPCGNode)
-				{
-					OtherPCGEditor->JumpToNode(JumpToPCGNode);
-				}
+				OtherPCGEditor->JumpToNode(JumpToPCGNode);
 			}
 		}
 	}
@@ -1474,14 +1532,76 @@ bool SPCGEditorGraphDebugObjectTree::ContextMenu_JumpToGraphInTree_CanExecute()
 	const TArray<FPCGEditorGraphDebugObjectItemPtr> SelectedItems = DebugObjectTreeView->GetSelectedItems();
 	for (const FPCGEditorGraphDebugObjectItemPtr& SelectedItem : SelectedItems)
 	{
-		const UPCGGraph* PCGGraphBeingEdited = PCGEditor.IsValid() ? PCGEditor.Pin()->GetPCGGraph() : nullptr;
-		const UPCGGraph* SelectedPCGGraph = SelectedItem->GetPCGGraph();
+		if (CanJumpToGraphInTree(SelectedItem))
+		{
+			return true;
+		}
+	}
+
+	return false;
+}
+
+bool SPCGEditorGraphDebugObjectTree::CanJumpToGraphInTree(const FPCGEditorGraphDebugObjectItemPtr& Item) const
+{
+	// Offer jump-to command if any selected item is a graph or a subgraph loop iteration.
+	return Item && (Item->GetPCGGraph() || Item->IsLoopIteration());
+}
+
+void SPCGEditorGraphDebugObjectTree::FocusOnItem(const FPCGEditorGraphDebugObjectItemPtr& Item)
+{
+	if (!Item || !Item->GetPCGStack())
+	{
+		return;
+	}
+
+	const UPCGGraph* CurrentGraph = PCGEditor.Pin()->GetPCGGraph();
+	const FPCGStack& Stack = *Item->GetPCGStack();
 
-		// Offer jump-to command if any selected item is a graph or a subgraph loop iteration.
-		if (SelectedPCGGraph || SelectedItem->IsLoopIteration())
+	bool bFoundCurrentGraph = false;
+	const UPCGNode* LastNode = nullptr;
+
+	// Find first node after the current graph in the stack.
+	for (int i = Stack.GetStackFrames().Num() - 1; i >= 0; --i)
+	{
+		const FPCGStackFrame& StackFrame = Stack.GetStackFrames()[i];
+		const UObject* Object = StackFrame.Object.Get();
+
+		if (Object == CurrentGraph)
+		{
+			bFoundCurrentGraph = true;
+			break;
+		}
+		else if (const UPCGNode* Node = Cast<const UPCGNode>(Object))
+		{
+			LastNode = Node;
+		}
+	}
+
+	if (bFoundCurrentGraph && LastNode)
+	{
+		PCGEditor.Pin()->JumpToNode(LastNode);
+	}
+}
+
+bool SPCGEditorGraphDebugObjectTree::CanFocusOnItem(const FPCGEditorGraphDebugObjectItemPtr& Item) const
+{
+	if (!Item)
+	{
+		return false;
+	}
+
+	// Anything downstream from the currently editable graph (e.g. "IsDebuggable" is true) is avaiable for focus.
+	FPCGEditorGraphDebugObjectItemPtr Parent = Item->GetParent();
+	while (Parent)
+	{
+		if (Parent->IsDebuggable())
 		{
 			return true;
 		}
+		else
+		{
+			Parent = Parent->GetParent();
+		}
 	}
 
 	return false;
diff --git a/Engine/Plugins/PCG/Source/PCGEditor/Private/SPCGEditorGraphDebugObjectTree.h b/Engine/Plugins/PCG/Source/PCGEditor/Private/SPCGEditorGraphDebugObjectTree.h
--- a/Engine/Plugins/PCG/Source/PCGEditor/Private/SPCGEditorGraphDebugObjectTree.h
+++ b/Engine/Plugins/PCG/Source/PCGEditor/Private/SPCGEditorGraphDebugObjectTree.h
@@ -57,6 +57,7 @@ public:
 	virtual const UObject* GetObject() const  = 0;
 	virtual const UPCGGraph* GetPCGGraph() const { return nullptr; }
 	virtual bool IsLoopIteration() const { return false; }
+	virtual const FSlateBrush* GetIcon() const;
 
 	virtual void OnObjectsReplaced(const TMap<UObject*, UObject*>& ReplacementMap);
 
@@ -84,6 +85,7 @@ public:
 	virtual UPCGComponent* GetPCGComponent() const override { return nullptr; }
 	virtual const UObject* GetObject() const override { return Actor.Get(); };
 	virtual const FPCGStack* GetPCGStack() const override { return &PCGStack; }
+	virtual const FSlateBrush* GetIcon() const override;
 
 protected:
 	virtual FPCGStack* GetMutablePCGStack() override { return &PCGStack; }
@@ -128,6 +130,13 @@ protected:
 	bool bIsDebuggable = false;
 };
 
+enum class EPCGEditorSubgraphNodeType : uint8
+{
+	StaticSubgraph = 0,
+	DynamicSubgraph,
+	LoopSubgraph
+};
+
 class FPCGEditorGraphDebugObjectItem_PCGSubgraph : public FPCGEditorGraphDebugObjectItem
 {
 public:
@@ -137,13 +146,13 @@ public:
 		const FPCGStack& InPCGStack,
 		bool bInIsDebuggable,
 		bool bInHasInspectionData,
-		bool bInIsDynamic)
+		EPCGEditorSubgraphNodeType InSubgraphType)
 		: FPCGEditorGraphDebugObjectItem(bInHasInspectionData)
 		, PCGNode(InPCGNode)
 		, PCGGraph(InPCGGraph)
 		, PCGStack(InPCGStack)
 		, bIsDebuggable(bInIsDebuggable)
-		, bIsDynamic(bInIsDynamic)
+		, SubgraphType(InSubgraphType)
 	{}
 
 	virtual FString GetLabel(bool bForSorting = false) const override;
@@ -151,8 +160,9 @@ public:
 	virtual const UObject* GetObject() const override { return PCGNode.Get(); };
 	virtual const FPCGStack* GetPCGStack() const override { return &PCGStack; }
 	virtual bool IsDebuggable() const override { return bIsDebuggable; }
-	virtual bool IsDynamic() const override { return bIsDynamic; }
+	virtual bool IsDynamic() const override { return SubgraphType != EPCGEditorSubgraphNodeType::StaticSubgraph; }
 	virtual const UPCGGraph* GetPCGGraph() const override { return PCGGraph.Get(); }
+	virtual const FSlateBrush* GetIcon() const override;
 
 protected:
 	virtual FPCGStack* GetMutablePCGStack() override { return &PCGStack; }
@@ -162,7 +172,7 @@ protected:
 
 	FPCGStack PCGStack;
 	bool bIsDebuggable = false;
-	bool bIsDynamic = false;
+	EPCGEditorSubgraphNodeType SubgraphType = EPCGEditorSubgraphNodeType::StaticSubgraph;
 };
 
 class FPCGEditorGraphDebugObjectItem_PCGLoopIndex : public FPCGEditorGraphDebugObjectItem
@@ -205,12 +215,16 @@ protected:
 class SPCGEditorGraphDebugObjectItemRow : public SCompoundWidget
 {
 public:
-	using FDoubleClickFunc = TFunction<void(FPCGEditorGraphDebugObjectItemPtr)>;
+	DECLARE_DELEGATE_OneParam(FPCGDebugObjectItemRowAction, const FPCGEditorGraphDebugObjectItemPtr&);
+	DECLARE_DELEGATE_RetVal_OneParam(bool, FPCGDebugObjectItemRowPredicate, const FPCGEditorGraphDebugObjectItemPtr&);
 
 	SLATE_BEGIN_ARGS(SPCGEditorGraphDebugObjectItemRow)
-		: _OnDoubleClickFunc()
 	{}
-		SLATE_ARGUMENT(FDoubleClickFunc, OnDoubleClickFunc)
+		SLATE_EVENT(FPCGDebugObjectItemRowAction, OnDoubleClick)
+		SLATE_EVENT(FPCGDebugObjectItemRowAction, OnJumpTo)
+		SLATE_EVENT(FPCGDebugObjectItemRowPredicate, CanJumpTo)
+		SLATE_EVENT(FPCGDebugObjectItemRowAction, OnFocus)
+		SLATE_EVENT(FPCGDebugObjectItemRowPredicate, CanFocus)
 	SLATE_END_ARGS()
 
 	void Construct(const FArguments& InArgs, const TSharedRef<STableViewBase>& InOwnerTableView, FPCGEditorGraphDebugObjectItemPtr InItem);
@@ -219,11 +233,24 @@ public:
 	virtual FReply OnMouseButtonDoubleClick(const FGeometry& InMyGeometry, const FPointerEvent& InMouseEvent) override;
 	//~End SWidget Interface
 
+	FReply FocusClicked() const;
+	FReply JumpToClicked() const;
+	bool IsJumpToEnabled() const;
+	bool IsFocusEnabled() const;
+
 private:
 	FPCGEditorGraphDebugObjectItemPtr Item;
-	
+
 	/** Invoked when the user double clicks on the row. */
-	FDoubleClickFunc DoubleClickFunc;
+	FPCGDebugObjectItemRowAction OnDoubleClick;
+	/** Invoked when the 'Jump To' action is clicked on the row buttons. */
+	FPCGDebugObjectItemRowAction OnJumpTo;
+	/** Invoked when the 'Go to node' action is clicked on the row buttons. */
+	FPCGDebugObjectItemRowAction OnFocus;
+	/** Controls whether the jump to button will be enabled or not. */
+	FPCGDebugObjectItemRowPredicate CanJumpTo;
+	/** Controls whether the 'go to node' button will be enabled or not. */
+	FPCGDebugObjectItemRowPredicate CanFocus;
 };
 
 class SPCGEditorGraphDebugObjectTree : public SCompoundWidget
@@ -285,7 +312,7 @@ private:
 	/** If the stack matches to an item, Expands the tree view to make that item visible and select it, and return true. Otherwise, return false. */
 	bool SetDebugObjectFromStack(const FPCGStack& InStack);
 	/** Expands the tree view to make an item visible and select it. */
-	void ExpandAndSelectDebugObject(FPCGEditorGraphDebugObjectItemPtr InItem);
+	void ExpandAndSelectDebugObject(const FPCGEditorGraphDebugObjectItemPtr& InItem);
 
 	/** Expand the given row and select the deepest entry as the debug object if it is unambiguous (the only entry at its level in the tree). */
 	void ExpandAndSelectFirstLeafDebugObject(FPCGEditorGraphDebugObjectItemPtr InItem);
@@ -296,6 +323,11 @@ private:
 	void ContextMenu_JumpToGraphInTree();
 	bool ContextMenu_JumpToGraphInTree_CanExecute() const;
 
+	void JumpToGraphInTree(const FPCGEditorGraphDebugObjectItemPtr& Item);
+	bool CanJumpToGraphInTree(const FPCGEditorGraphDebugObjectItemPtr& Item) const;
+	void FocusOnItem(const FPCGEditorGraphDebugObjectItemPtr& Item);
+	bool CanFocusOnItem(const FPCGEditorGraphDebugObjectItemPtr& Item) const;
+
 	TWeakPtr<FPCGEditor> PCGEditor;
 
 	TSharedPtr<STreeView<FPCGEditorGraphDebugObjectItemPtr>> DebugObjectTreeView;
@@ -313,18 +345,11 @@ private:
 	/** Used to retain item selection state across tree refreshes. */
 	FPCGStack SelectedStack;
 
-	/** Stack before it was culled so it's in the local state. */
-	FPCGStack SelectedStackBeforeCulling;
-
 	/** Used to retain item selection state across tree refreshes if the SelectedStack is invalidated (e.g. through BP reconstruction). */
 	TWeakObjectPtr<const UPCGGraph> SelectedGraph = nullptr;
 	TWeakObjectPtr<const AActor> SelectedOwner = nullptr;
 	uint32 SelectedGridSize = PCGHiGenGrid::UnboundedGridSize();
 	FIntVector SelectedGridCoord = FIntVector::ZeroValue;
 	TWeakObjectPtr<const UPCGComponent> SelectedOriginalComponent = nullptr;
-
-	/** The previous stack that the user selected. */
-	FPCGStack PreviouslySelectedStack;
-
 	TWeakObjectPtr<const UPCGNode> PCGNodeBeingInspected = nullptr;
 };
