[PCG] Fix bug with packed instance data with the Mesh Selector by attribute (and with material override)

#rnx

diff --git a/Engine/Plugins/PCG/Source/PCG/Private/MeshSelectors/PCGMeshSelectorByAttribute.cpp b/Engine/Plugins/PCG/Source/PCG/Private/MeshSelectors/PCGMeshSelectorByAttribute.cpp
--- a/Engine/Plugins/PCG/Source/PCG/Private/MeshSelectors/PCGMeshSelectorByAttribute.cpp
+++ b/Engine/Plugins/PCG/Source/PCG/Private/MeshSelectors/PCGMeshSelectorByAttribute.cpp
@@ -294,7 +294,7 @@ bool UPCGMeshSelectorByAttribute::SelectInstances(
 						const FPCGPoint& Point = Points[PointIndices[PointIndex]];
 						FPCGMeshInstanceList& InstanceList = PCGMeshSelectorAttribute::GetInstanceList(OutMeshInstances, CurrentPartitionDescriptor, CurrentPartitionDescriptor.StaticMesh, MaterialOverrideHelper.GetMaterialOverrides(Point.MetadataEntry), bReverseTransform, InPointData, PartitionIndex);
 						InstanceList.Instances.Emplace(Point.Transform);
-						InstanceList.InstancesIndices.Emplace(PointIndex);
+						InstanceList.InstancesIndices.Emplace(PointIndices[PointIndex]);
 					}
 				}
 				else
