[PCG] Convert trivial elements to support PCGPointArrayData

#rnx

diff --git a/Engine/Plugins/PCG/Source/PCG/Private/Elements/PCGAttributeReduceElement.cpp b/Engine/Plugins/PCG/Source/PCG/Private/Elements/PCGAttributeReduceElement.cpp
--- a/Engine/Plugins/PCG/Source/PCG/Private/Elements/PCGAttributeReduceElement.cpp
+++ b/Engine/Plugins/PCG/Source/PCG/Private/Elements/PCGAttributeReduceElement.cpp
@@ -6,7 +6,6 @@
 #include "PCGCustomVersion.h"
 #include "PCGParamData.h"
 #include "PCGPin.h"
-#include "Data/PCGPointData.h"
 #include "Data/PCGSpatialData.h"
 #include "Elements/Metadata/PCGMetadataElementCommon.h"
 #include "Metadata/Accessors/PCGAttributeAccessorHelpers.h"
diff --git a/Engine/Plugins/PCG/Source/PCG/Private/Elements/PCGNumberOfElements.cpp b/Engine/Plugins/PCG/Source/PCG/Private/Elements/PCGNumberOfElements.cpp
--- a/Engine/Plugins/PCG/Source/PCG/Private/Elements/PCGNumberOfElements.cpp
+++ b/Engine/Plugins/PCG/Source/PCG/Private/Elements/PCGNumberOfElements.cpp
@@ -5,7 +5,7 @@
 #include "PCGContext.h"
 #include "PCGPin.h"
 #include "PCGParamData.h"
-#include "Data/PCGPointData.h"
+#include "Data/PCGBasePointData.h"
 #include "Metadata/PCGMetadata.h"
 
 #define LOCTEXT_NAMESPACE "PCGNumberOfElementsSettings"
@@ -118,9 +118,9 @@ bool FPCGNumberOfElementsBaseElement<DataType>::ExecuteInternal(FPCGContext* Con
 // FPCGNumberOfPointsElement
 ////////////////////////////////////
 
-int32 FPCGNumberOfPointsElement::GetNum(const UPCGPointData* InData) const
+int32 FPCGNumberOfPointsElement::GetNum(const UPCGBasePointData* InData) const
 {
-	return InData ? InData->GetPoints().Num() : 0;
+	return InData ? InData->GetNumPoints() : 0;
 }
 
 
diff --git a/Engine/Plugins/PCG/Source/PCG/Public/Elements/ControlFlow/PCGBooleanSelect.h b/Engine/Plugins/PCG/Source/PCG/Public/Elements/ControlFlow/PCGBooleanSelect.h
--- a/Engine/Plugins/PCG/Source/PCG/Public/Elements/ControlFlow/PCGBooleanSelect.h
+++ b/Engine/Plugins/PCG/Source/PCG/Public/Elements/ControlFlow/PCGBooleanSelect.h
@@ -40,4 +40,5 @@ class FPCGBooleanSelectElement : public IPCGElement
 {
 protected:
 	virtual bool ExecuteInternal(FPCGContext* Context) const override;
+	virtual bool SupportsBasePointDataInputs(FPCGContext* InContext) const { return true; }
 };
diff --git a/Engine/Plugins/PCG/Source/PCG/Public/Elements/ControlFlow/PCGBranch.h b/Engine/Plugins/PCG/Source/PCG/Public/Elements/ControlFlow/PCGBranch.h
--- a/Engine/Plugins/PCG/Source/PCG/Public/Elements/ControlFlow/PCGBranch.h
+++ b/Engine/Plugins/PCG/Source/PCG/Public/Elements/ControlFlow/PCGBranch.h
@@ -45,4 +45,5 @@ class FPCGBranchElement : public IPCGElement
 protected:
 	virtual bool ExecuteInternal(FPCGContext* Context) const override;
 	virtual bool SupportsGPUResidentData(FPCGContext* InContext) const override { return true; }
+	virtual bool SupportsBasePointDataInputs(FPCGContext* InContext) const { return true; }
 };
diff --git a/Engine/Plugins/PCG/Source/PCG/Public/Elements/ControlFlow/PCGMultiSelect.h b/Engine/Plugins/PCG/Source/PCG/Public/Elements/ControlFlow/PCGMultiSelect.h
--- a/Engine/Plugins/PCG/Source/PCG/Public/Elements/ControlFlow/PCGMultiSelect.h
+++ b/Engine/Plugins/PCG/Source/PCG/Public/Elements/ControlFlow/PCGMultiSelect.h
@@ -89,4 +89,5 @@ class FPCGMultiSelectElement : public IPCGElement
 {
 protected:
 	virtual bool ExecuteInternal(FPCGContext* Context) const override;
+	virtual bool SupportsBasePointDataInputs(FPCGContext* InContext) const override { return true; }
 };
diff --git a/Engine/Plugins/PCG/Source/PCG/Public/Elements/ControlFlow/PCGSwitch.h b/Engine/Plugins/PCG/Source/PCG/Public/Elements/ControlFlow/PCGSwitch.h
--- a/Engine/Plugins/PCG/Source/PCG/Public/Elements/ControlFlow/PCGSwitch.h
+++ b/Engine/Plugins/PCG/Source/PCG/Public/Elements/ControlFlow/PCGSwitch.h
@@ -100,4 +100,5 @@ class FPCGSwitchElement : public IPCGElement
 protected:
 	virtual bool ExecuteInternal(FPCGContext* Context) const override;
 	virtual bool SupportsGPUResidentData(FPCGContext* InContext) const override { return true; }
+	virtual bool SupportsBasePointDataInputs(FPCGContext* InContext) const { return true; }
 };
diff --git a/Engine/Plugins/PCG/Source/PCG/Public/Elements/Metadata/PCGMetadataRenameElement.h b/Engine/Plugins/PCG/Source/PCG/Public/Elements/Metadata/PCGMetadataRenameElement.h
--- a/Engine/Plugins/PCG/Source/PCG/Public/Elements/Metadata/PCGMetadataRenameElement.h
+++ b/Engine/Plugins/PCG/Source/PCG/Public/Elements/Metadata/PCGMetadataRenameElement.h
@@ -40,4 +40,5 @@ class FPCGMetadataRenameElement : public IPCGElement
 protected:
 	virtual bool ExecuteInternal(FPCGContext* Context) const override;
 	virtual EPCGElementExecutionLoopMode ExecutionLoopMode(const UPCGSettings* Settings) const override { return EPCGElementExecutionLoopMode::SinglePrimaryPin; }
+	virtual bool SupportsBasePointDataInputs(FPCGContext* InContext) const override { return true; }
 };
diff --git a/Engine/Plugins/PCG/Source/PCG/Public/Elements/PCGAddTag.h b/Engine/Plugins/PCG/Source/PCG/Public/Elements/PCGAddTag.h
--- a/Engine/Plugins/PCG/Source/PCG/Public/Elements/PCGAddTag.h
+++ b/Engine/Plugins/PCG/Source/PCG/Public/Elements/PCGAddTag.h
@@ -59,4 +59,5 @@ class FPCGAddTagElement : public IPCGElement
 protected:
 	virtual bool ExecuteInternal(FPCGContext* Context) const override;
 	virtual EPCGElementExecutionLoopMode ExecutionLoopMode(const UPCGSettings* Settings) const override { return EPCGElementExecutionLoopMode::SinglePrimaryPin; }
+	virtual bool SupportsBasePointDataInputs(FPCGContext* InContext) const override { return true; }
 };
\ No newline at end of file
diff --git a/Engine/Plugins/PCG/Source/PCG/Public/Elements/PCGAttributeReduceElement.h b/Engine/Plugins/PCG/Source/PCG/Public/Elements/PCGAttributeReduceElement.h
--- a/Engine/Plugins/PCG/Source/PCG/Public/Elements/PCGAttributeReduceElement.h
+++ b/Engine/Plugins/PCG/Source/PCG/Public/Elements/PCGAttributeReduceElement.h
@@ -89,4 +89,5 @@ class FPCGAttributeReduceElement : public IPCGElement
 protected:
 	virtual bool ExecuteInternal(FPCGContext* Context) const override;
 	virtual EPCGElementExecutionLoopMode ExecutionLoopMode(const UPCGSettings* Settings) const override;
+	virtual bool SupportsBasePointDataInputs(FPCGContext* InContext) const override { return true; }
 };
diff --git a/Engine/Plugins/PCG/Source/PCG/Public/Elements/PCGCopyAttributes.h b/Engine/Plugins/PCG/Source/PCG/Public/Elements/PCGCopyAttributes.h
--- a/Engine/Plugins/PCG/Source/PCG/Public/Elements/PCGCopyAttributes.h
+++ b/Engine/Plugins/PCG/Source/PCG/Public/Elements/PCGCopyAttributes.h
@@ -87,4 +87,5 @@ class FPCGCopyAttributesElement : public IPCGElement
 {
 protected:
 	virtual bool ExecuteInternal(FPCGContext* Context) const override;
+	virtual bool SupportsBasePointDataInputs(FPCGContext* InContext) const override { return true; }
 };
diff --git a/Engine/Plugins/PCG/Source/PCG/Public/Elements/PCGCreateAttribute.h b/Engine/Plugins/PCG/Source/PCG/Public/Elements/PCGCreateAttribute.h
--- a/Engine/Plugins/PCG/Source/PCG/Public/Elements/PCGCreateAttribute.h
+++ b/Engine/Plugins/PCG/Source/PCG/Public/Elements/PCGCreateAttribute.h
@@ -126,6 +126,7 @@ class FPCGCreateAttributeElement : public IPCGElement
 {
 protected:
 	virtual bool ExecuteInternal(FPCGContext* Context) const override;
+	virtual bool SupportsBasePointDataInputs(FPCGContext* InContext) const override { return true; }
 };
 
 class FPCGAddAttributeElement : public IPCGElement
@@ -133,4 +134,5 @@ class FPCGAddAttributeElement : public IPCGElement
 protected:
 	virtual bool ExecuteInternal(FPCGContext* Context) const override;
 	virtual EPCGElementExecutionLoopMode ExecutionLoopMode(const UPCGSettings* Settings) const override { return EPCGElementExecutionLoopMode::SinglePrimaryPin; }
+	virtual bool SupportsBasePointDataInputs(FPCGContext* InContext) const override { return true; }
 };
\ No newline at end of file
diff --git a/Engine/Plugins/PCG/Source/PCG/Public/Elements/PCGDataNum.h b/Engine/Plugins/PCG/Source/PCG/Public/Elements/PCGDataNum.h
--- a/Engine/Plugins/PCG/Source/PCG/Public/Elements/PCGDataNum.h
+++ b/Engine/Plugins/PCG/Source/PCG/Public/Elements/PCGDataNum.h
@@ -35,4 +35,5 @@ class FPCGDataNumElement : public IPCGElement
 {
 protected:
 	virtual bool ExecuteInternal(FPCGContext* Context) const override;
+	virtual bool SupportsBasePointDataInputs(FPCGContext* InContext) const override { return true; }
 };
\ No newline at end of file
diff --git a/Engine/Plugins/PCG/Source/PCG/Public/Elements/PCGDeleteAttributesElement.h b/Engine/Plugins/PCG/Source/PCG/Public/Elements/PCGDeleteAttributesElement.h
--- a/Engine/Plugins/PCG/Source/PCG/Public/Elements/PCGDeleteAttributesElement.h
+++ b/Engine/Plugins/PCG/Source/PCG/Public/Elements/PCGDeleteAttributesElement.h
@@ -81,4 +81,5 @@ class FPCGDeleteAttributesElement : public IPCGElement
 protected:
 	virtual bool ExecuteInternal(FPCGContext* Context) const override;
 	virtual EPCGElementExecutionLoopMode ExecutionLoopMode(const UPCGSettings* Settings) const override { return EPCGElementExecutionLoopMode::SinglePrimaryPin; }
+	virtual bool SupportsBasePointDataInputs(FPCGContext* InContext) const override { return true; }
 };
diff --git a/Engine/Plugins/PCG/Source/PCG/Public/Elements/PCGDeleteTags.h b/Engine/Plugins/PCG/Source/PCG/Public/Elements/PCGDeleteTags.h
--- a/Engine/Plugins/PCG/Source/PCG/Public/Elements/PCGDeleteTags.h
+++ b/Engine/Plugins/PCG/Source/PCG/Public/Elements/PCGDeleteTags.h
@@ -62,4 +62,5 @@ class FPCGDeleteTagsElement : public IPCGElement
 protected:
 	virtual bool ExecuteInternal(FPCGContext* Context) const override;
 	virtual EPCGElementExecutionLoopMode ExecutionLoopMode(const UPCGSettings* Settings) const override { return EPCGElementExecutionLoopMode::SinglePrimaryPin; }
+	virtual bool SupportsBasePointDataInputs(FPCGContext* InContext) const override { return true; }
 };
\ No newline at end of file
diff --git a/Engine/Plugins/PCG/Source/PCG/Public/Elements/PCGFilterByTag.h b/Engine/Plugins/PCG/Source/PCG/Public/Elements/PCGFilterByTag.h
--- a/Engine/Plugins/PCG/Source/PCG/Public/Elements/PCGFilterByTag.h
+++ b/Engine/Plugins/PCG/Source/PCG/Public/Elements/PCGFilterByTag.h
@@ -60,4 +60,5 @@ class FPCGFilterByTagElement : public IPCGElement
 protected:
 	virtual bool ExecuteInternal(FPCGContext* Context) const override;
 	virtual EPCGElementExecutionLoopMode ExecutionLoopMode(const UPCGSettings* Settings) const override { return EPCGElementExecutionLoopMode::SinglePrimaryPin; }
+	virtual bool SupportsBasePointDataInputs(FPCGContext* InContext) const override { return true; }
 };
diff --git a/Engine/Plugins/PCG/Source/PCG/Public/Elements/PCGFilterByType.h b/Engine/Plugins/PCG/Source/PCG/Public/Elements/PCGFilterByType.h
--- a/Engine/Plugins/PCG/Source/PCG/Public/Elements/PCGFilterByType.h
+++ b/Engine/Plugins/PCG/Source/PCG/Public/Elements/PCGFilterByType.h
@@ -43,4 +43,5 @@ class FPCGFilterByTypeElement : public IPCGElement
 protected:
 	virtual bool ExecuteInternal(FPCGContext* InContext) const override;
 	virtual EPCGElementExecutionLoopMode ExecutionLoopMode(const UPCGSettings* Settings) const override { return EPCGElementExecutionLoopMode::SinglePrimaryPin; }
+	virtual bool SupportsBasePointDataInputs(FPCGContext* InContext) const override { return true; }
 };
diff --git a/Engine/Plugins/PCG/Source/PCG/Public/Elements/PCGGather.h b/Engine/Plugins/PCG/Source/PCG/Public/Elements/PCGGather.h
--- a/Engine/Plugins/PCG/Source/PCG/Public/Elements/PCGGather.h
+++ b/Engine/Plugins/PCG/Source/PCG/Public/Elements/PCGGather.h
@@ -47,6 +47,7 @@ protected:
 	virtual bool ExecuteInternal(FPCGContext* Context) const override;
 	virtual bool ShouldVerifyIfOutputsAreUsedMultipleTimes(const UPCGSettings* InSettings) const override { return true; }
 	virtual bool SupportsGPUResidentData(FPCGContext* InContext) const override { return true; }
+	virtual bool SupportsBasePointDataInputs(FPCGContext* InContext) const { return true; }
 };
 
 namespace PCGGather
diff --git a/Engine/Plugins/PCG/Source/PCG/Public/Elements/PCGGetDataInfo.h b/Engine/Plugins/PCG/Source/PCG/Public/Elements/PCGGetDataInfo.h
--- a/Engine/Plugins/PCG/Source/PCG/Public/Elements/PCGGetDataInfo.h
+++ b/Engine/Plugins/PCG/Source/PCG/Public/Elements/PCGGetDataInfo.h
@@ -37,6 +37,7 @@ class FPCGGetTagsElement : public IPCGElement
 protected:
 	virtual bool ExecuteInternal(FPCGContext* Context) const override;
 	virtual EPCGElementExecutionLoopMode ExecutionLoopMode(const UPCGSettings* Settings) const override { return EPCGElementExecutionLoopMode::SinglePrimaryPin; }
+	virtual bool SupportsBasePointDataInputs(FPCGContext* InContext) const override { return true; }
 };
 
 /** Extracts the list of attribute on data to an attribute set with one entry per attribute. */
@@ -74,4 +75,5 @@ class FPCGGetAttributesElement : public IPCGElement
 protected:
 	virtual bool ExecuteInternal(FPCGContext* Context) const override;
 	virtual EPCGElementExecutionLoopMode ExecutionLoopMode(const UPCGSettings* Settings) const override { return EPCGElementExecutionLoopMode::SinglePrimaryPin; }
+	virtual bool SupportsBasePointDataInputs(FPCGContext* InContext) const override { return true; }
 };
\ No newline at end of file
diff --git a/Engine/Plugins/PCG/Source/PCG/Public/Elements/PCGNumberOfElements.h b/Engine/Plugins/PCG/Source/PCG/Public/Elements/PCGNumberOfElements.h
--- a/Engine/Plugins/PCG/Source/PCG/Public/Elements/PCGNumberOfElements.h
+++ b/Engine/Plugins/PCG/Source/PCG/Public/Elements/PCGNumberOfElements.h
@@ -7,7 +7,7 @@
 #include "PCGNumberOfElements.generated.h"
 
 class UPCGParamData;
-class UPCGPointData;
+class UPCGBasePointData;
 
 /**
 * Elements for getting the number of elements in a point data or a param data. Since the whole logic is identical
@@ -86,14 +86,15 @@ class FPCGNumberOfElementsBaseElement : public IPCGElement
 {
 protected:
 	virtual bool ExecuteInternal(FPCGContext* Context) const override;
+	virtual bool SupportsBasePointDataInputs(FPCGContext* InContext) const override { return true; }
 
 	virtual int32 GetNum(const DataType* InData) const = 0;
 };
 
-class FPCGNumberOfPointsElement : public FPCGNumberOfElementsBaseElement<UPCGPointData>
+class FPCGNumberOfPointsElement : public FPCGNumberOfElementsBaseElement<UPCGBasePointData>
 {
 protected:
-	virtual int32 GetNum(const UPCGPointData* InData) const override;
+	virtual int32 GetNum(const UPCGBasePointData* InData) const override;
 };
 
 class FPCGNumberOfEntriesElement : public FPCGNumberOfElementsBaseElement<UPCGParamData>
diff --git a/Engine/Plugins/PCG/Source/PCG/Public/Elements/PCGPrintElement.h b/Engine/Plugins/PCG/Source/PCG/Public/Elements/PCGPrintElement.h
--- a/Engine/Plugins/PCG/Source/PCG/Public/Elements/PCGPrintElement.h
+++ b/Engine/Plugins/PCG/Source/PCG/Public/Elements/PCGPrintElement.h
@@ -130,4 +130,5 @@ public:
 
 protected:
 	virtual bool ExecuteInternal(FPCGContext* Context) const override;
+	virtual bool SupportsBasePointDataInputs(FPCGContext* InContext) const override { return true; }
 };
diff --git a/Engine/Plugins/PCG/Source/PCG/Public/Elements/PCGReplaceTags.h b/Engine/Plugins/PCG/Source/PCG/Public/Elements/PCGReplaceTags.h
--- a/Engine/Plugins/PCG/Source/PCG/Public/Elements/PCGReplaceTags.h
+++ b/Engine/Plugins/PCG/Source/PCG/Public/Elements/PCGReplaceTags.h
@@ -55,4 +55,5 @@ class FPCGReplaceTagsElement : public IPCGElement
 protected:
 	virtual bool ExecuteInternal(FPCGContext* Context) const override;
 	virtual EPCGElementExecutionLoopMode ExecutionLoopMode(const UPCGSettings* Settings) const override { return EPCGElementExecutionLoopMode::SinglePrimaryPin; }
+	virtual bool SupportsBasePointDataInputs(FPCGContext* InContext) const override { return true; }
 };
\ No newline at end of file
diff --git a/Engine/Plugins/PCG/Source/PCG/Public/Elements/PCGSortTags.h b/Engine/Plugins/PCG/Source/PCG/Public/Elements/PCGSortTags.h
--- a/Engine/Plugins/PCG/Source/PCG/Public/Elements/PCGSortTags.h
+++ b/Engine/Plugins/PCG/Source/PCG/Public/Elements/PCGSortTags.h
@@ -45,4 +45,5 @@ class FPCGSortTagsElement : public IPCGElement
 {
 protected:
 	virtual bool ExecuteInternal(FPCGContext* Context) const override;
+	virtual bool SupportsBasePointDataInputs(FPCGContext* InContext) const override { return true; }
 };
diff --git a/Engine/Plugins/PCG/Source/PCG/Public/PCGSubgraph.h b/Engine/Plugins/PCG/Source/PCG/Public/PCGSubgraph.h
--- a/Engine/Plugins/PCG/Source/PCG/Public/PCGSubgraph.h
+++ b/Engine/Plugins/PCG/Source/PCG/Public/PCGSubgraph.h
@@ -198,6 +198,7 @@ protected:
 	void PrepareSubgraphData(const UPCGSubgraphSettings* Settings, FPCGSubgraphContext* Context, const FPCGDataCollection& InputData, FPCGDataCollection& OutputData) const;
 	void PrepareSubgraphUserParameters(const UPCGSubgraphSettings* Settings, FPCGSubgraphContext* Context, FPCGDataCollection& OutputData) const;
 	virtual bool SupportsGPUResidentData(FPCGContext* InContext) const override { return true; }
+	virtual bool SupportsBasePointDataInputs(FPCGContext* InContext) const { return true; }
 };
 
 // Implementation note: this node forwards data, but does not keep that data alive. This is the responsibility of the corresponding FPCGSubgraphContext
@@ -216,6 +217,7 @@ protected:
 	virtual bool ExecuteInternal(FPCGContext* Context) const override;
 	virtual bool IsPassthrough(const UPCGSettings* InSettings) const override { return true; }
 	virtual bool SupportsGPUResidentData(FPCGContext* InContext) const override { return true; }
+	virtual bool SupportsBasePointDataInputs(FPCGContext* InContext) const { return true; }
 
 	FPCGDataCollection Input;
 };
