Fixed double-click on debug object tree view

#rnx

diff --git a/Engine/Plugins/PCG/Source/PCGEditor/Private/SPCGEditorGraphDebugObjectTree.cpp b/Engine/Plugins/PCG/Source/PCGEditor/Private/SPCGEditorGraphDebugObjectTree.cpp
--- a/Engine/Plugins/PCG/Source/PCGEditor/Private/SPCGEditorGraphDebugObjectTree.cpp
+++ b/Engine/Plugins/PCG/Source/PCGEditor/Private/SPCGEditorGraphDebugObjectTree.cpp
@@ -1256,7 +1256,7 @@ TSharedRef<ITableRow> SPCGEditorGraphDebugObjectTree::MakeTreeRowWidget(FPCGEdit
 	return SNew( STableRow<TSharedPtr<SPCGEditorGraphDebugObjectItemRow>>, InOwnerTable)
 	[
 		SNew(SPCGEditorGraphDebugObjectItemRow, InOwnerTable, InItem)
-		.OnDoubleClick(this, &SPCGEditorGraphDebugObjectTree::ExpandAndSelectDebugObject)
+		.OnDoubleClick(this, &SPCGEditorGraphDebugObjectTree::ExpandAndSelectFirstLeafDebugObject)
 		.OnJumpTo(this, &SPCGEditorGraphDebugObjectTree::JumpToGraphInTree)
 		.CanJumpTo(this, &SPCGEditorGraphDebugObjectTree::CanJumpToGraphInTree)
 		.OnFocus(this, &SPCGEditorGraphDebugObjectTree::FocusOnItem)
@@ -1394,7 +1394,7 @@ void SPCGEditorGraphDebugObjectTree::OnSetExpansionRecursive(FPCGEditorGraphDebu
 	}
 }
 
-void SPCGEditorGraphDebugObjectTree::ExpandAndSelectFirstLeafDebugObject(FPCGEditorGraphDebugObjectItemPtr InItem)
+void SPCGEditorGraphDebugObjectTree::ExpandAndSelectFirstLeafDebugObject(const FPCGEditorGraphDebugObjectItemPtr& InItem)
 {
 	if (!InItem.IsValid())
 	{
diff --git a/Engine/Plugins/PCG/Source/PCGEditor/Private/SPCGEditorGraphDebugObjectTree.h b/Engine/Plugins/PCG/Source/PCGEditor/Private/SPCGEditorGraphDebugObjectTree.h
--- a/Engine/Plugins/PCG/Source/PCGEditor/Private/SPCGEditorGraphDebugObjectTree.h
+++ b/Engine/Plugins/PCG/Source/PCGEditor/Private/SPCGEditorGraphDebugObjectTree.h
@@ -314,8 +314,8 @@ private:
 	/** Expands the tree view to make an item visible and select it. */
 	void ExpandAndSelectDebugObject(const FPCGEditorGraphDebugObjectItemPtr& InItem);
 
-	/** Expand the given row and select the deepest entry as the debug object if it is unambiguous (the only entry at its level in the tree). */
-	void ExpandAndSelectFirstLeafDebugObject(FPCGEditorGraphDebugObjectItemPtr InItem);
+	/** Expand the given row and select the next occurrence of the current graph. */
+	void ExpandAndSelectFirstLeafDebugObject(const FPCGEditorGraphDebugObjectItemPtr& InItem);
 
 	TSharedPtr<SWidget> OpenContextMenu();
 
