[PCG] PointArrayData changes
- SetNumPoints now calls DirtyCache if num points actually changes
- New IPCGElement::ConvertInputsIfNeeded method that will be used by proxy nodes
- Create a FPCGNativePointPropertyConstAccessor instead of using the same class for const/non-const as it causes issues with flattening (with data inheritance)

#rnx

diff --git a/Engine/Plugins/PCG/Source/PCG/Private/Data/PCGBasePointData.cpp b/Engine/Plugins/PCG/Source/PCG/Private/Data/PCGBasePointData.cpp
--- a/Engine/Plugins/PCG/Source/PCG/Private/Data/PCGBasePointData.cpp
+++ b/Engine/Plugins/PCG/Source/PCG/Private/Data/PCGBasePointData.cpp
@@ -220,17 +220,78 @@ namespace PCGBasePointData
 		UPCGBasePointData* PointData = nullptr;
 	};
 
+	template<typename T> 
+	bool GetRange(TArrayView<T> OutValues, int32 StartIndex, const IPCGAttributeAccessorKeys& Keys, const auto& Range)
+	{
+		TArray<int32> KeyIndices;
+		bool bContiguous = false;
+		if (!Keys.GetKeyIndices(StartIndex, OutValues.Num(), KeyIndices, bContiguous))
+		{
+			return false;
+		}
+
+		if (bContiguous)
+		{
+			for (int32 OutIndex = 0; OutIndex < OutValues.Num(); ++OutIndex)
+			{
+				OutValues[OutIndex] = static_cast<T>(Range[StartIndex + OutIndex]);
+			}
+		}
+		else
+		{
+			for (int32 OutIndex = 0; OutIndex < OutValues.Num(); ++OutIndex)
+			{
+				OutValues[OutIndex] = static_cast<T>(Range[KeyIndices[OutIndex]]);
+			}
+		}
+
+		return true;
+	}
+
 	template <typename T, typename R = T>
-	class FPCGNativePointPropertyAccessor : public IPCGAttributeAccessorT<FPCGNativePointPropertyAccessor<T,R>>
+	class FPCGNativePointPropertyConstAccessor : public IPCGAttributeAccessorT<FPCGNativePointPropertyConstAccessor<T,R>>
 	{
 	public:
 		using Type = T;
-		using Super = IPCGAttributeAccessorT<FPCGNativePointPropertyAccessor<T,R>>;
+		using Super = IPCGAttributeAccessorT<FPCGNativePointPropertyConstAccessor<T,R>>;
 
-		FPCGNativePointPropertyAccessor(const UPCGBasePointData* InPointData, EPCGPointNativeProperties InNativeProperty)
+		FPCGNativePointPropertyConstAccessor(const UPCGBasePointData* InPointData, EPCGPointNativeProperties InNativeProperty)
+			: Super(/*bReadOnly=*/true)
+		{
+			ValueRange = InPointData->GetConstValueRange<R>(InNativeProperty);
+		}
+
+		FPCGNativePointPropertyConstAccessor(UPCGBasePointData* InPointData, EPCGPointNativeProperties InNativeProperty)
 			: Super(/*bReadOnly=*/true)
 		{
-			ValueRange = const_cast<UPCGBasePointData*>(InPointData)->GetValueRange<R>(InNativeProperty);
+			ValueRange = InPointData->GetConstValueRange<R>(InNativeProperty);
+		}
+				
+		bool GetRangeImpl(TArrayView<T> OutValues, int32 StartIndex, const IPCGAttributeAccessorKeys& Keys) const
+		{
+			return PCGBasePointData::GetRange(OutValues, StartIndex, Keys, ValueRange);
+		}
+
+		bool SetRangeImpl(TArrayView<const T> InValues, int32 StartIndex, IPCGAttributeAccessorKeys& Keys, EPCGAttributeAccessorFlags Flags)
+		{
+			return false;
+		}
+
+	private:
+		TConstPCGValueRange<R> ValueRange;
+	};
+
+	template <typename T, typename R = T>
+	class FPCGNativePointPropertyAccessor : public IPCGAttributeAccessorT<FPCGNativePointPropertyAccessor<T, R>>
+	{
+	public:
+		using Type = T;
+		using Super = IPCGAttributeAccessorT<FPCGNativePointPropertyAccessor<T, R>>;
+
+		FPCGNativePointPropertyAccessor(const UPCGBasePointData* InPointData, EPCGPointNativeProperties InNativeProperty)
+			: Super(/*bReadOnly=*/false)
+		{
+			ValueRange = const_cast<UPCGBasePointData*>(InPointData)->GetValueRange<R>(InNativeProperty, /*bAllocate=*/ true);
 		}
 
 		FPCGNativePointPropertyAccessor(UPCGBasePointData* InPointData, EPCGPointNativeProperties InNativeProperty)
@@ -241,29 +302,7 @@ namespace PCGBasePointData
 
 		bool GetRangeImpl(TArrayView<T> OutValues, int32 StartIndex, const IPCGAttributeAccessorKeys& Keys) const
 		{
-			TArray<int32> KeyIndices;
-			bool bContiguous = false;
-			if (!Keys.GetKeyIndices(StartIndex, OutValues.Num(), KeyIndices, bContiguous))
-			{
-				return false;
-			}
-						
-			if (bContiguous)
-			{
-				for (int32 OutIndex = 0; OutIndex < OutValues.Num(); ++OutIndex)
-				{
-					OutValues[OutIndex] = (T)ValueRange[StartIndex + OutIndex];
-				}
-			}
-			else
-			{
-				for (int32 OutIndex = 0; OutIndex < OutValues.Num(); ++OutIndex)
-				{
-					OutValues[OutIndex] = (T)ValueRange[KeyIndices[OutIndex]];
-				}
-			}
-
-			return true;
+			return PCGBasePointData::GetRange<T>(OutValues, StartIndex, Keys, ValueRange);
 		}
 
 		bool SetRangeImpl(TArrayView<const T> InValues, int32 StartIndex, IPCGAttributeAccessorKeys& Keys, EPCGAttributeAccessorFlags Flags)
@@ -425,18 +464,18 @@ namespace PCGBasePointData
 		switch (NativeProperty)
 		{
 		case EPCGPointNativeProperties::Transform:
-			return MakeUnique<FPCGNativePointPropertyAccessor<FTransform>>(PointData, NativeProperty);
+			return MakeUnique<FPCGNativePointPropertyConstAccessor<FTransform>>(PointData, NativeProperty);
 		case EPCGPointNativeProperties::Density:
 		case EPCGPointNativeProperties::Steepness:
 			// The values are floats but we want to remain compatible with existing UPCGPointData Steepnees/Density property accessors which are of type double 
-			return MakeUnique<FPCGNativePointPropertyAccessor<double, float>>(PointData, NativeProperty);
+			return MakeUnique<FPCGNativePointPropertyConstAccessor<double, float>>(PointData, NativeProperty);
 		case EPCGPointNativeProperties::BoundsMin:
 		case EPCGPointNativeProperties::BoundsMax:
-			return MakeUnique<FPCGNativePointPropertyAccessor<FVector>>(PointData, NativeProperty);
+			return MakeUnique<FPCGNativePointPropertyConstAccessor<FVector>>(PointData, NativeProperty);
 		case EPCGPointNativeProperties::Color:
-			return MakeUnique<FPCGNativePointPropertyAccessor<FVector4>>(PointData, NativeProperty);
+			return MakeUnique<FPCGNativePointPropertyConstAccessor<FVector4>>(PointData, NativeProperty);
 		case EPCGPointNativeProperties::Seed:
-			return MakeUnique<FPCGNativePointPropertyAccessor<int32>>(PointData, NativeProperty);
+			return MakeUnique<FPCGNativePointPropertyConstAccessor<int32>>(PointData, NativeProperty);
 		default:
 			UE_LOG(LogPCG, Error, TEXT("EPCGPointNativeProperty value '%d' does not exist."), NativeProperty);
 			return TUniquePtr<IPCGAttributeAccessor>();
@@ -894,7 +933,7 @@ void UPCGBasePointData::RebuildOctree() const
 		return;
 	}
 
-	TRACE_CPUPROFILER_EVENT_SCOPE(UPCGPointData::RebuildOctree)
+	TRACE_CPUPROFILER_EVENT_SCOPE(UPCGBasePointData::RebuildOctree)
 	
 	const FBox PointBounds = GetBounds();
 	PCGPointOctree::FPointOctree NewOctree(PointBounds.GetCenter(), PointBounds.GetExtent().Length());
@@ -1255,7 +1294,7 @@ bool UPCGBasePointData::ProjectPoint(const FTransform& InTransform, const FBox&
 
 bool UPCGBasePointData::ProjectPoint(const FTransform& InTransform, const FBox& InBounds, const FPCGProjectionParams& InParams, FPCGPoint& OutPoint, UPCGMetadata* OutMetadata, bool bUseBounds) const
 {
-	//TRACE_CPUPROFILER_EVENT_SCOPE(UPCGPointData::ProjectPoint);
+	//TRACE_CPUPROFILER_EVENT_SCOPE(UPCGBasePointData::ProjectPoint);
 	RebuildOctreeIfNeeded();
 
 	TArray<TPair<int32, FVector::FReal>, TInlineAllocator<4>> Contributions;
@@ -1407,7 +1446,7 @@ bool UPCGBasePointData::ProjectPoint(const FTransform& InTransform, const FBox&
 
 	if (OutMetadata)
 	{
-		//TRACE_CPUPROFILER_EVENT_SCOPE(UPCGPointData::SamplePoint::SetupMetadata);
+		//TRACE_CPUPROFILER_EVENT_SCOPE(UPCGBasePointData::SamplePoint::SetupMetadata);
 		// Initialise metadata entry for this temporary point
 		OutPoint.MetadataEntry = OutMetadata ? (OutMetadata->HasParent(Metadata) ? OutMetadata->AddEntry(MetadataEntryRange[MaxContributor]) : OutMetadata->AddEntry()) : PCGInvalidEntryKey;
 		
diff --git a/Engine/Plugins/PCG/Source/PCG/Private/Data/PCGPointArrayData.cpp b/Engine/Plugins/PCG/Source/PCG/Private/Data/PCGPointArrayData.cpp
--- a/Engine/Plugins/PCG/Source/PCG/Private/Data/PCGPointArrayData.cpp
+++ b/Engine/Plugins/PCG/Source/PCG/Private/Data/PCGPointArrayData.cpp
@@ -193,7 +193,12 @@ void UPCGPointArrayData::SetNumPoints(int32 InNumPoints)
 	{
 		FlattenPropertiesIfNeeded();
 	}
-	PointArray.SetNumPoints(InNumPoints);
+
+	if (InNumPoints != PointArray.GetNumPoints())
+	{
+		PointArray.SetNumPoints(InNumPoints);
+		DirtyCache();
+	}
 }
 
 void UPCGPointArrayData::AllocateProperties(EPCGPointNativeProperties Properties)
diff --git a/Engine/Plugins/PCG/Source/PCG/Private/Data/PCGPointData.cpp b/Engine/Plugins/PCG/Source/PCG/Private/Data/PCGPointData.cpp
--- a/Engine/Plugins/PCG/Source/PCG/Private/Data/PCGPointData.cpp
+++ b/Engine/Plugins/PCG/Source/PCG/Private/Data/PCGPointData.cpp
@@ -314,6 +314,15 @@ void UPCGPointData::SetPoints(const TArray<FPCGPoint>& InPoints)
 	GetMutablePoints() = InPoints;
 }
 
+void UPCGPointData::SetNumPoints(int32 InNumPoints)
+{ 
+	if (Points.Num() != InNumPoints)
+	{
+		Points.SetNum(InNumPoints);
+		DirtyCache();
+	}
+}
+
 FPCGPoint UPCGPointData::GetPoint(int32 Index) const
 {
 	// This check has an overhead but this method is blueprint callable and we want to avoid crashing
diff --git a/Engine/Plugins/PCG/Source/PCG/Private/PCGData.cpp b/Engine/Plugins/PCG/Source/PCG/Private/PCGData.cpp
--- a/Engine/Plugins/PCG/Source/PCG/Private/PCGData.cpp
+++ b/Engine/Plugins/PCG/Source/PCG/Private/PCGData.cpp
@@ -7,7 +7,7 @@
 #include "PCGParamData.h"
 #include "PCGPin.h"
 #include "PCGSettings.h"
-#include "Data/PCGPointData.h"
+#include "Data/PCGBasePointData.h"
 #include "Data/PCGSpatialData.h"
 #include "Data/PCGUnionData.h"
 
@@ -466,8 +466,8 @@ int32 FPCGDataCollection::StripEmptyPointData()
 {
 	return TaggedData.RemoveAll([](const FPCGTaggedData& Data)
 	{
-		const UPCGPointData* PointData = Cast<UPCGPointData>(Data.Data);
-		return PointData && PointData->GetPoints().IsEmpty();
+		const UPCGBasePointData* PointData = Cast<UPCGBasePointData>(Data.Data);
+		return PointData && PointData->GetNumPoints() == 0;
 	});
 }
 
diff --git a/Engine/Plugins/PCG/Source/PCG/Private/PCGElement.cpp b/Engine/Plugins/PCG/Source/PCG/Private/PCGElement.cpp
--- a/Engine/Plugins/PCG/Source/PCG/Private/PCGElement.cpp
+++ b/Engine/Plugins/PCG/Source/PCG/Private/PCGElement.cpp
@@ -7,6 +7,8 @@
 #include "PCGGraph.h"
 #include "PCGSubsystem.h"
 #include "Compute/Data/PCGProxyForGPUData.h"
+#include "Data/PCGBasePointData.h"
+#include "Data/PCGPointArrayData.h"
 #include "Data/PCGPointData.h"
 #include "Graph/PCGGraphCache.h"
 #include "Helpers/PCGHelpers.h"
@@ -35,6 +37,11 @@ static TAutoConsoleVariable<bool> CVarPCGShouldVerifyIfOutputsAreUsedMultipleTim
 	true,
 	TEXT("Add small computation at the end of each node to detect if the data is used multiple times. Necessary for data stealing."));
 
+static TAutoConsoleVariable<bool> CVarPCGEnablePointArrayDataToPointDataConversionWarnings(
+	TEXT("pcg.EnablePointArrayDataToPointDataConversionWarnings"),
+	false,
+	TEXT("Warn about input conversions from PointArrayData to PointData so that code that needs to be updated is identified."));
+
 #if WITH_EDITOR
 #define PCG_ELEMENT_EXECUTION_BREAKPOINT() \
 	if (Context && Context->GetInputSettingsInterface() && Context->GetInputSettingsInterface()->bBreakDebugger) \
@@ -363,6 +370,16 @@ void IPCGElement::PreExecutePrimaryLoopElement(FPCGContext* Context, const UPCGS
 }
 
 bool IPCGElement::PrepareData(FPCGContext* Context) const
+{
+	if (!ConvertInputsIfNeeded(Context))
+	{
+		return false;
+	}
+	
+	return PrepareDataInternal(Context);
+}
+
+bool IPCGElement::ConvertInputsIfNeeded(FPCGContext* Context) const
 {
 	if (!SupportsGPUResidentData(Context))
 	{
@@ -414,6 +431,8 @@ bool IPCGElement::PrepareData(FPCGContext* Context) const
 
 	if (!SupportsBasePointDataInputs(Context))
 	{
+		bool bDataWasConverted = false;
+
 		for (FPCGTaggedData& Data : Context->InputData.TaggedData)
 		{
 			if (const UPCGBasePointData* BasePointData = Cast<UPCGBasePointData>(Data.Data))
@@ -422,12 +441,18 @@ bool IPCGElement::PrepareData(FPCGContext* Context) const
 				if (!BasePointData->IsA<UPCGPointData>())
 				{
 					Data.Data = BasePointData->ToPointData(Context);
+					bDataWasConverted = true;
 				}
 			}
 		}
+
+		if (bDataWasConverted && CVarPCGEnablePointArrayDataToPointDataConversionWarnings.GetValueOnAnyThread())
+		{
+			PCGE_LOG(Warning, GraphAndLog, FText::Format(LOCTEXT("UnsupportedPointArrayData", "ToPointData was called on inputs of node '{0}'. Consider implementing support for UPCGBasePointData (IPCGElement::SupportsBasePointDataInputs)."), FText::FromString(Context->GetOriginalSettings<UPCGSettings>()->GetName())));
+		}
 	}
 
-	return PrepareDataInternal(Context);
+	return true;
 }
 
 bool IPCGElement::PrepareDataInternal(FPCGContext* Context) const
@@ -927,16 +952,17 @@ void IPCGElement::CleanupAndValidateOutput(FPCGContext* Context) const
 
 				if (CVarPCGValidatePointMetadata.GetValueOnAnyThread())
 				{
-					if (const UPCGPointData* PointData = Cast<UPCGPointData>(TaggedData.Data))
+					if (const UPCGBasePointData* PointData = Cast<UPCGBasePointData>(TaggedData.Data))
 					{
-						const TArray<FPCGPoint>& NewPoints = PointData->GetPoints();
 						const int32 MaxMetadataEntry = PointData->Metadata ? PointData->Metadata->GetItemCountForChild() : 0;
 
 						bool bHasError = false;
 
-						for(int32 PointIndex = 0; PointIndex < NewPoints.Num() && !bHasError; ++PointIndex)
+						const TConstPCGValueRange<int64> MetadataEntryRange = PointData->GetConstMetadataEntryValueRange();
+
+						for(int32 PointIndex = 0; PointIndex < MetadataEntryRange.Num() && !bHasError; ++PointIndex)
 						{
-							bHasError |= (NewPoints[PointIndex].MetadataEntry >= MaxMetadataEntry);
+							bHasError |= (MetadataEntryRange[PointIndex] >= MaxMetadataEntry);
 						}
 
 						if (bHasError)
diff --git a/Engine/Plugins/PCG/Source/PCG/Public/Data/PCGPointData.h b/Engine/Plugins/PCG/Source/PCG/Public/Data/PCGPointData.h
--- a/Engine/Plugins/PCG/Source/PCG/Public/Data/PCGPointData.h
+++ b/Engine/Plugins/PCG/Source/PCG/Public/Data/PCGPointData.h
@@ -130,7 +130,7 @@ public:
 	UFUNCTION(BlueprintCallable, Category = SpatialData)
 	virtual int32 GetNumPoints() const override { return Points.Num(); }
 
-	virtual void SetNumPoints(int32 InNumPoints) override { Points.SetNum(InNumPoints); }
+	virtual void SetNumPoints(int32 InNumPoints) override;
 
 	UFUNCTION(BlueprintCallable, Category = SpatialData)
 	bool IsEmpty() const { return Points.IsEmpty(); }
diff --git a/Engine/Plugins/PCG/Source/PCG/Public/PCGElement.h b/Engine/Plugins/PCG/Source/PCG/Public/PCGElement.h
--- a/Engine/Plugins/PCG/Source/PCG/Public/PCGElement.h
+++ b/Engine/Plugins/PCG/Source/PCG/Public/PCGElement.h
@@ -145,6 +145,9 @@ public:
 	bool DebugDisplay(FPCGContext* Context) const;
 #endif
 
+	/** Performs some input conversions if element doesn't support it. Called in PrepareData but available for proxy nodes */
+	bool ConvertInputsIfNeeded(FPCGContext* Context) const;
+
 protected:
 	/** This function will be called once and once only, at the beginning of an execution */
 	void PreExecute(FPCGContext* Context) const;
