[PCG] Fix a concurrency issue with FPCGContext::InitializeSettings
- Add a static thread_local that UPCGSettings subclass can check against to skip some init code
- Make sure we don't call SetupCallbacks/TeardownCallbacks when this is the case

diff --git a/Engine/Plugins/PCG/Source/PCG/Private/PCGContext.cpp b/Engine/Plugins/PCG/Source/PCG/Private/PCGContext.cpp
--- a/Engine/Plugins/PCG/Source/PCG/Private/PCGContext.cpp
+++ b/Engine/Plugins/PCG/Source/PCG/Private/PCGContext.cpp
@@ -33,6 +33,8 @@ TAutoConsoleVariable<bool> CVarPCGEnablePointArrayData(
 
 namespace PCGContextHelpers
 {
+	thread_local bool bIsInitializingSettings = false;
+
 	template <typename T>
 	bool GetOverrideParamValue(const IPCGAttributeAccessor& InAccessor, T& OutValue)
 	{
@@ -132,6 +134,11 @@ const UPCGSettingsInterface* FPCGContext::GetInputSettingsInterface() const
 	}
 }
 
+bool FPCGContext::IsInitializingSettings()
+{
+	return PCGContextHelpers::bIsInitializingSettings;
+}
+
 void FPCGContext::InitializeSettings(bool bSkipPostLoad)
 {
 	TRACE_CPUPROFILER_EVENT_SCOPE(FPCGContext::InitializeSettings);
@@ -141,6 +148,13 @@ void FPCGContext::InitializeSettings(bool bSkipPostLoad)
 		return;
 	}
 
+	// Thread local
+	PCGContextHelpers::bIsInitializingSettings = true;
+	ON_SCOPE_EXIT
+	{
+		PCGContextHelpers::bIsInitializingSettings = false;
+	};
+
 	if (const UPCGSettings* NodeSettings = GetOriginalSettings<UPCGSettings>())
 	{
 		// Don't apply overrides if the original settings are disabled
@@ -226,13 +240,6 @@ void FPCGContext::InitializeSettings(bool bSkipPostLoad)
 							// We are not calling PostLoad so remove the NeedPostLoad flags here
 							KeyValuePair.Value->ClearFlags(EObjectFlags::RF_NeedPostLoad | EObjectFlags::RF_NeedPostLoadSubobjects);
 						}
-#if WITH_EDITOR
-						// @todo_pcg: find a way to avoid the call to SetupCallbacks() all together but for now unregister the callbacks after duplication
-						if (UPCGGraphInstance* GraphInstance = Cast<UPCGGraphInstance>(KeyValuePair.Value))
-						{
-							GraphInstance->TeardownCallbacks();
-						}
-#endif
 					}
 				}
 				
diff --git a/Engine/Plugins/PCG/Source/PCG/Private/PCGGraph.cpp b/Engine/Plugins/PCG/Source/PCG/Private/PCGGraph.cpp
--- a/Engine/Plugins/PCG/Source/PCG/Private/PCGGraph.cpp
+++ b/Engine/Plugins/PCG/Source/PCG/Private/PCGGraph.cpp
@@ -2132,6 +2132,11 @@ bool UPCGGraphInstance::CanEditChange(const FProperty* InProperty) const
 
 void UPCGGraphInstance::TeardownCallbacks()
 {
+	if (FPCGContext::IsInitializingSettings())
+	{
+		return;
+	}
+
 	if (Graph)
 	{
 		Graph->OnGraphChangedDelegate.RemoveAll(this);
@@ -2141,6 +2146,11 @@ void UPCGGraphInstance::TeardownCallbacks()
 
 void UPCGGraphInstance::SetupCallbacks()
 {
+	if (FPCGContext::IsInitializingSettings())
+	{
+		return;
+	}
+
 	if (Graph && !Graph->OnGraphChangedDelegate.IsBoundToObject(this))
 	{
 		Graph->OnGraphChangedDelegate.AddUObject(this, &UPCGGraphInstance::OnGraphChanged);
diff --git a/Engine/Plugins/PCG/Source/PCG/Private/PCGSubgraph.cpp b/Engine/Plugins/PCG/Source/PCG/Private/PCGSubgraph.cpp
--- a/Engine/Plugins/PCG/Source/PCG/Private/PCGSubgraph.cpp
+++ b/Engine/Plugins/PCG/Source/PCG/Private/PCGSubgraph.cpp
@@ -44,6 +44,11 @@ namespace PCGSubgraphSettings
 #if WITH_EDITOR
 void UPCGBaseSubgraphSettings::SetupCallbacks()
 {
+	if (FPCGContext::IsInitializingSettings())
+	{
+		return;
+	}
+
 	UPCGGraphInterface* Subgraph = GetSubgraphInterface();
 
 	if (Subgraph && !Subgraph->OnGraphChangedDelegate.IsBoundToObject(this))
@@ -54,6 +59,11 @@ void UPCGBaseSubgraphSettings::SetupCallbacks()
 
 void UPCGBaseSubgraphSettings::TeardownCallbacks()
 {
+	if (FPCGContext::IsInitializingSettings())
+	{
+		return;
+	}
+
 	if (UPCGGraphInterface* Subgraph = GetSubgraphInterface())
 	{
 		Subgraph->OnGraphChangedDelegate.RemoveAll(this);
diff --git a/Engine/Plugins/PCG/Source/PCG/Public/PCGContext.h b/Engine/Plugins/PCG/Source/PCG/Public/PCGContext.h
--- a/Engine/Plugins/PCG/Source/PCG/Public/PCGContext.h
+++ b/Engine/Plugins/PCG/Source/PCG/Public/PCGContext.h
@@ -104,6 +104,8 @@ struct PCG_API FPCGContext
 	// It will create a copy of the original settings if there is indeed a possible override.
 	void InitializeSettings(bool bSkipPostLoad = false);
 
+	static bool IsInitializingSettings();
+
 	// If we any any parameter override, it will read from the params and override matching values
 	// in the settings copy.
 	void OverrideSettings();
