[PCG] Add support for render targets in PCGTextureDataInterface

diff --git a/Engine/Plugins/PCG/Source/PCG/Private/Compute/DataInterfaces/PCGTextureDataInterface.cpp b/Engine/Plugins/PCG/Source/PCG/Private/Compute/DataInterfaces/PCGTextureDataInterface.cpp
--- a/Engine/Plugins/PCG/Source/PCG/Private/Compute/DataInterfaces/PCGTextureDataInterface.cpp
+++ b/Engine/Plugins/PCG/Source/PCG/Private/Compute/DataInterfaces/PCGTextureDataInterface.cpp
@@ -118,7 +118,7 @@ UComputeDataProvider* UPCGTextureDataInterface::CreateDataProvider(TObjectPtr<UO
 
 	for (int InputTextureIndex = 0; InputTextureIndex < NumInputTextures; ++InputTextureIndex)
 	{
-		const UPCGTextureData* TextureData = Cast<UPCGTextureData>(InputTaggedData[InputTextureIndex].Data);
+		const UPCGBaseTextureData* TextureData = Cast<UPCGBaseTextureData>(InputTaggedData[InputTextureIndex].Data);
 		if (!ensure(TextureData))
 		{
 			UE_LOG(LogPCG, Error, TEXT("Unsupported data type encountered by texture data interface: '%s'"),
diff --git a/Engine/Plugins/PCG/Source/PCG/Public/Data/PCGRenderTargetData.h b/Engine/Plugins/PCG/Source/PCG/Public/Data/PCGRenderTargetData.h
--- a/Engine/Plugins/PCG/Source/PCG/Public/Data/PCGRenderTargetData.h
+++ b/Engine/Plugins/PCG/Source/PCG/Public/Data/PCGRenderTargetData.h
@@ -25,6 +25,7 @@ public:
 
 	//~Begin UPCGBaseTextureData interface
 	virtual FTextureRHIRef GetTextureRHI() const override;
+	virtual EPCGTextureResourceType GetTextureResourceType() const override { return EPCGTextureResourceType::TextureObject; }
 	//~End UPCGBaseTextureData interface
 
 	//~Begin UPCGSpatialData interface
diff --git a/Engine/Plugins/PCG/Source/PCG/Public/Data/PCGTextureData.h b/Engine/Plugins/PCG/Source/PCG/Public/Data/PCGTextureData.h
--- a/Engine/Plugins/PCG/Source/PCG/Public/Data/PCGTextureData.h
+++ b/Engine/Plugins/PCG/Source/PCG/Public/Data/PCGTextureData.h
@@ -47,8 +47,9 @@ enum class EPCGTextureAddressMode : uint8
 UENUM()
 enum class EPCGTextureResourceType : uint8
 {
-	TextureObject UMETA(ToolTip = "UObject texture us as UTexture2D."),
-	ExportedTexture UMETA(ToolTip = "Texture handle exported from a texture on the GPU.")
+	TextureObject UMETA(ToolTip = "UObject texture such as UTexture2D or UTextureRenderTarget2D."),
+	ExportedTexture UMETA(ToolTip = "Texture handle exported from a texture on the GPU."),
+	Invalid UMETA(Hidden)
 };
 
 namespace PCGTextureSamplingHelpers
@@ -90,6 +91,7 @@ public:
 	virtual bool IsValid() const;
 
 	virtual FTextureRHIRef GetTextureRHI() const PURE_VIRTUAL(UPCGBaseTextureData::GetTextureResource, return nullptr;);
+	virtual EPCGTextureResourceType GetTextureResourceType() const PURE_VIRTUAL(UPCGBaseTextureData::GetTextureResourceType, return EPCGTextureResourceType::Invalid;);
 	virtual TRefCountPtr<IPooledRenderTarget> GetRefCountedTexture() const { return nullptr; }
 	virtual int GetTextureSlice() const { return 0; }
 
@@ -177,7 +179,7 @@ public:
 	/** Data is successfully initialized and is ready to use. */
 	PCG_API bool IsSuccessfullyInitialized() const { return bSuccessfullyInitialized; }
 
-	EPCGTextureResourceType GetTextureResourceType() const { return ResourceType; }
+	EPCGTextureResourceType GetTextureResourceType() const override { return ResourceType; }
 
 	//~Begin UPCGData interface
 	virtual EPCGDataType GetDataType() const override { return EPCGDataType::Texture; }
