Implemented a fast octree-based ordered culling path in the cull points outside actor bounds node.
Normalized behavior to use point center instead of point position (in all cases).
Normalized behavior for the unordered path where it will accept only points with centers that lie in the bounds.
Ported the node to BasePointData.

diff --git a/Engine/Plugins/PCG/Config/BasePCG.ini b/Engine/Plugins/PCG/Config/BasePCG.ini
--- a/Engine/Plugins/PCG/Config/BasePCG.ini
+++ b/Engine/Plugins/PCG/Config/BasePCG.ini
@@ -30,6 +30,7 @@
 +FunctionRedirects=(OldName="PCGMetadata.CopyExistingAttribute",NewName="PCGMetadata.BP_CopyExistingAttribute")
 +FunctionRedirects=(OldName="PCGMetadata.RenameAttribute",NewName="PCGMetadata.BP_RenameAttribute")
 +FunctionRedirects=(OldName="PCGMetadata.ClearAttribute",NewName="PCGMetadata.BP_ClearAttribute")
++FunctionRedirects=(OldName="PCGBasePointData.CopyPointsFrom",NewName="PCGBasePointData.BP_CopyPointsFrom")
 
 +PropertyRedirects=(OldName="PCGSplineSamplerSettings.Params",NewName="PCGSplineSamplerSettings.SamplerParams")
 +PropertyRedirects=(OldName="PCGProjectionSettings.Params",NewName="PCGProjectionSettings.ProjectionParams")
diff --git a/Engine/Plugins/PCG/Source/PCG/Private/Data/PCGBasePointData.cpp b/Engine/Plugins/PCG/Source/PCG/Private/Data/PCGBasePointData.cpp
--- a/Engine/Plugins/PCG/Source/PCG/Private/Data/PCGBasePointData.cpp
+++ b/Engine/Plugins/PCG/Source/PCG/Private/Data/PCGBasePointData.cpp
@@ -845,12 +845,17 @@ void UPCGBasePointData::CopyPointsTo(UPCGBasePointData* To, int32 ReadStartIndex
 	CopyPropertiesTo(To, ReadStartIndex, WriteStartIndex, Count, EPCGPointNativeProperties::All);
 }
 
-void UPCGBasePointData::CopyPointsFrom(const UPCGBasePointData* InData, const TArray<int32>& InDataIndices)
+void UPCGBasePointData::BP_CopyPointsFrom(const UPCGBasePointData* InData, const TArray<int32>& InDataIndices)
 {
 	return CopyPoints(InData, this, InDataIndices, false);
 }
 
-void UPCGBasePointData::CopyPoints(const UPCGBasePointData* From, UPCGBasePointData* To, const TArray<int32>& InDataIndices, bool bCopyAll)
+void UPCGBasePointData::CopyPointsFrom(const UPCGBasePointData* InData, const TArrayView<const int32>& InDataIndices)
+{
+	return CopyPoints(InData, this, InDataIndices, false);
+}
+
+void UPCGBasePointData::CopyPoints(const UPCGBasePointData* From, UPCGBasePointData* To, const TArrayView<const int32>& InDataIndices, bool bCopyAll)
 {
 	TRACE_CPUPROFILER_EVENT_SCOPE(UPCGBasePointData::CopyPoints);
 	check(From);
diff --git a/Engine/Plugins/PCG/Source/PCG/Private/Elements/PCGCullPointsOutsideActorBounds.cpp b/Engine/Plugins/PCG/Source/PCG/Private/Elements/PCGCullPointsOutsideActorBounds.cpp
--- a/Engine/Plugins/PCG/Source/PCG/Private/Elements/PCGCullPointsOutsideActorBounds.cpp
+++ b/Engine/Plugins/PCG/Source/PCG/Private/Elements/PCGCullPointsOutsideActorBounds.cpp
@@ -5,9 +5,9 @@
 #include "PCGComponent.h"
 #include "PCGContext.h"
 #include "PCGPin.h"
-#include "Data/PCGIntersectionData.h"
-#include "Data/PCGPointData.h"
-#include "Data/PCGVolumeData.h"
+#include "Data/PCGBasePointData.h"
+#include "Helpers/PCGAsync.h"
+#include "Helpers/PCGHelpers.h"
 
 #define LOCTEXT_NAMESPACE "PCGCullPointsOutsideActorBoundsElement"
 
@@ -53,70 +53,117 @@ bool FPCGCullPointsOutsideActorBoundsElement::ExecuteInternal(FPCGContext* Conte
 	// Initialize directly. Could also have gone through PCGComponent::CreateActorPCGDataCollection.
 	const FBox BoundsBox = Context->SourceComponent->GetGridBounds().ExpandBy(Settings->BoundsExpansion);
 
-	UPCGVolumeData* VolumeData = nullptr;
-	if (Settings->Mode == EPCGCullPointsMode::Ordered)
-	{
-		VolumeData = FPCGContext::NewObject_AnyThread<UPCGVolumeData>(Context);
-		check(VolumeData);
-		VolumeData->Initialize(BoundsBox);
-	}
-	
 	const TArray<FPCGTaggedData>& Inputs = Context->InputData.TaggedData;
 	for (const FPCGTaggedData& Input : Inputs)
 	{
-		const UPCGPointData* InputPointData = Cast<UPCGPointData>(Input.Data);
+		const UPCGBasePointData* InputPointData = Cast<UPCGBasePointData>(Input.Data);
 
 		// Skip non-points or empty point data
-		if (!InputPointData || InputPointData->GetPoints().IsEmpty())
+		if(!InputPointData || InputPointData->GetNumPoints() == 0)
 		{
 			continue;
 		}
 
-		const UPCGPointData* OutputData = nullptr;
+		const UPCGBasePointData* OutputData = nullptr;
 
-		if (Settings->Mode == EPCGCullPointsMode::Ordered)
-		{
-			check(VolumeData);
-			UPCGIntersectionData* Intersection = InputPointData->IntersectWith(Context, VolumeData);
-			check(Intersection);
+		TArray<int32, TInlineAllocator<4096>> KeptPointIndices;
+		bool bNeedsSort = false;
 
-			const UPCGPointData* IntersectedPoints = Intersection->ToPointData(Context);
-			check(IntersectedPoints);
+		// If an octree is available, we'll perform a query on it instead of going through all points, as in most cases it will be significantly less work.
+		// Otherwise, fall-back to testing every point (on multiple threads), which will always yield sorted results.
+		if (InputPointData->IsPointOctreeDirty())
+		{
+			// Implementation note: while this is a slower code path in general, building the octree normally dwarfs filtering (on MT, especially)
+			// @todo_pcg: review this if we improve on the octree creation
+			auto InitializeFunc = [InputPointData, &KeptPointIndices]()
+			{
+				KeptPointIndices.SetNumUninitialized(InputPointData->GetNumPoints());
+			};
 
-			// Skip empty result
-			if (IntersectedPoints->GetPoints().IsEmpty())
+			auto AsyncProcessRangeFunc = [InputPointData, &BoundsBox, &KeptPointIndices](int32 StartReadIndex, int32 StartWriteIndex, int32 Count)
 			{
-				continue;
-			}
+				const TConstPCGValueRange<FTransform> ReadTransformRange = InputPointData->GetConstTransformValueRange();
+				const TConstPCGValueRange<FVector> ReadBoundsMinRange = InputPointData->GetConstBoundsMinValueRange();
+				const TConstPCGValueRange<FVector> ReadBoundsMaxRange = InputPointData->GetConstBoundsMaxValueRange();
+
+				int32 NumWritten = 0;
+
+				for (int32 ReadIndex = StartReadIndex; ReadIndex < StartReadIndex + Count; ++ReadIndex)
+				{
+					const FVector PointCenter = ReadTransformRange[ReadIndex].TransformPosition(PCGPointHelpers::GetLocalCenter(ReadBoundsMinRange[ReadIndex], ReadBoundsMaxRange[ReadIndex]));
+					if (PCGHelpers::IsInsideBounds(BoundsBox, PointCenter))
+					{
+						const int32 WriteIndex = StartWriteIndex + NumWritten;
+						KeptPointIndices[WriteIndex] = ReadIndex;
+						++NumWritten;
+					}
+				}
+
+				return NumWritten;
+			};
+
+			auto MoveDataRangeFunc = [&KeptPointIndices](int32 RangeStartIndex, int32 MoveToIndex, int32 NumElements)
+			{
+				std::memmove(&KeptPointIndices[MoveToIndex], &KeptPointIndices[RangeStartIndex], sizeof(int32) * NumElements);
+			};
 
-			OutputData = IntersectedPoints;
+			auto FinishedFunc = [&KeptPointIndices](int32 NumWritten)
+			{
+				KeptPointIndices.SetNum(NumWritten, EAllowShrinking::No);
+			};
+
+			FPCGAsync::AsyncProcessingRangeEx(
+				Context ? &Context->AsyncState : nullptr,
+				InputPointData->GetNumPoints(),
+				InitializeFunc,
+				AsyncProcessRangeFunc,
+				MoveDataRangeFunc,
+				FinishedFunc,
+				/*bEnableTimeSlicing=*/false
+			);
 		}
-		else if (Settings->Mode == EPCGCullPointsMode::Unordered)
+		else
 		{
-			UPCGPointData* CulledPointsData = FPCGContext::NewObject_AnyThread<UPCGPointData>(Context);
-			CulledPointsData->InitializeFromData(InputPointData);
-			TArray<FPCGPoint>& CulledPoints = CulledPointsData->GetMutablePoints();
-
-			// TODO: split up this query across mutiple threads, but will make gathering the point a bit harder
+			// Octree-based fast path; this assumes that not all points are in the bounds, otherwise it won't be faster than the other solution.
+			// @todo_pcg: Split up this query by subdividing into smaller boxes and perform the query on multiple threads
 			const PCGPointOctree::FPointOctree& InputOctree = InputPointData->GetPointOctree();
-			const TArray<FPCGPoint>& InputPoints = InputPointData->GetPoints();
 			FBoxCenterAndExtent BoxCenterAndExtents(BoundsBox);
-
-			InputOctree.FindElementsWithBoundsTest(BoxCenterAndExtents, [&InputPoints , &CulledPoints](const PCGPointOctree::FPointRef& PointRef)
+		
+			InputOctree.FindElementsWithBoundsTest(BoxCenterAndExtents, [&KeptPointIndices, &BoundsBox](const PCGPointOctree::FPointRef& PointRef)
 			{
-				CulledPoints.Add(InputPoints[PointRef.Index]);
+				// Do a final validation that the center is inside the box
+				if(PCGHelpers::IsInsideBounds(BoundsBox, PointRef.Bounds.Origin))
+				{
+					KeptPointIndices.Add(PointRef.Index);
+				}
 			});
 
-			if (!CulledPoints.IsEmpty())
-			{
-				OutputData = CulledPointsData;
-			}
+			bNeedsSort = Settings->Mode == EPCGCullPointsMode::Ordered;
 		}
-		else
+
+		if (KeptPointIndices.Num() == InputPointData->GetNumPoints())
 		{
-			PCGLog::LogErrorOnGraph(LOCTEXT("InvalidMode", "Invalid culling mode."), Context);
+			// Trivial operation, we've kept all the points
+			OutputData = InputPointData;
 		}
+		else if (!KeptPointIndices.IsEmpty())
+		{
+			UPCGBasePointData* CulledPointsData = FPCGContext::NewPointData_AnyThread(Context);
 
+			// Since we'll be selecting only a subset of points, we can't inherit the data directly
+			FPCGInitializeFromDataParams InitializeParams(InputPointData);
+			InitializeParams.bInheritSpatialData = false;
+			CulledPointsData->InitializeFromDataWithParams(InitializeParams);
+
+			if (bNeedsSort)
+			{
+				KeptPointIndices.Sort();
+			}
+
+			CulledPointsData->CopyPointsFrom(InputPointData, KeptPointIndices);
+			OutputData = CulledPointsData;
+		}
+		
 		if (OutputData)
 		{
 			FPCGTaggedData& Output = Context->OutputData.TaggedData.Emplace_GetRef(Input);
diff --git a/Engine/Plugins/PCG/Source/PCG/Public/Data/PCGBasePointData.h b/Engine/Plugins/PCG/Source/PCG/Public/Data/PCGBasePointData.h
--- a/Engine/Plugins/PCG/Source/PCG/Public/Data/PCGBasePointData.h
+++ b/Engine/Plugins/PCG/Source/PCG/Public/Data/PCGBasePointData.h
@@ -133,8 +133,10 @@ public:
 	virtual const PCGPointOctree::FPointOctree& GetPointOctree() const;
 	virtual FBox GetBounds() const override;
 		
-	UFUNCTION(BlueprintCallable, Category = SpatialData)
-	void CopyPointsFrom(const UPCGBasePointData* InData, const TArray<int32>& InDataIndices);
+	UFUNCTION(BlueprintCallable, Category = SpatialData, meta = (DisplayName="Copy Points From"))
+	void BP_CopyPointsFrom(const UPCGBasePointData* InData, const TArray<int32>& InDataIndices);
+	
+	void CopyPointsFrom(const UPCGBasePointData* InData, const TArrayView<const int32>& InDataIndices);
 
 	virtual void CopyPropertiesTo(UPCGBasePointData* To, int32 ReadStartIndex, int32 WriteStartIndex, int32 Count, EPCGPointNativeProperties Properties) const;
 	virtual void CopyPointsTo(UPCGBasePointData* To, int32 ReadStartIndex, int32 WriteStartIndex, int32 Count) const;
@@ -144,7 +146,7 @@ public:
 	virtual FPCGMetadataDomainID GetDefaultMetadataDomainID() const override { return PCGMetadataDomainID::Elements; }
 	virtual TArray<FPCGMetadataDomainID> GetAllSupportedMetadataDomainIDs() const override { return {PCGMetadataDomainID::Data, PCGMetadataDomainID::Elements}; }
 
-	static void CopyPoints(const UPCGBasePointData* From, UPCGBasePointData* To, const TArray<int32>& InDataIndices, bool bCopyAll);
+	static void CopyPoints(const UPCGBasePointData* From, UPCGBasePointData* To, const TArrayView<const int32>& InDataIndices, bool bCopyAll);
 
 protected:
 
diff --git a/Engine/Plugins/PCG/Source/PCG/Public/Elements/PCGCullPointsOutsideActorBounds.h b/Engine/Plugins/PCG/Source/PCG/Public/Elements/PCGCullPointsOutsideActorBounds.h
--- a/Engine/Plugins/PCG/Source/PCG/Public/Elements/PCGCullPointsOutsideActorBounds.h
+++ b/Engine/Plugins/PCG/Source/PCG/Public/Elements/PCGCullPointsOutsideActorBounds.h
@@ -53,4 +53,5 @@ public:
 protected:
 	virtual bool ExecuteInternal(FPCGContext* Context) const override;
 	virtual EPCGElementExecutionLoopMode ExecutionLoopMode(const UPCGSettings* Settings) const override { return EPCGElementExecutionLoopMode::SinglePrimaryPin; }
+	virtual bool SupportsBasePointDataInputs(FPCGContext* InContext) const { return true; }
 };
