Optimizing debug object tree view by implementing search from common parent.

#rnx

diff --git a/Engine/Plugins/PCG/Source/PCGEditor/Private/SPCGEditorGraphDebugObjectTree.cpp b/Engine/Plugins/PCG/Source/PCGEditor/Private/SPCGEditorGraphDebugObjectTree.cpp
--- a/Engine/Plugins/PCG/Source/PCGEditor/Private/SPCGEditorGraphDebugObjectTree.cpp
+++ b/Engine/Plugins/PCG/Source/PCGEditor/Private/SPCGEditorGraphDebugObjectTree.cpp
@@ -834,6 +834,7 @@ void SPCGEditorGraphDebugObjectTree::AddStacksToTree(const TArray<FPCGStack>& St
 	TMap<AActor*, TSharedPtr<FPCGEditorGraphDebugObjectItem_Actor>>& InOutActorItems,
 	TMap<const FPCGStack, FPCGEditorGraphDebugObjectItemPtr>& InOutStackToItem)
 {
+	//TRACE_CPUPROFILER_EVENT_SCOPE(SPCGEditorGraphDebugObjectTree::AddStacksToTree);
 	const UPCGGraph* GraphBeingEdited = GetPCGGraph();
 	UPCGSubsystem* Subsystem = PCGEditor.Pin()->GetSubsystem();
 	if (!GraphBeingEdited || !Subsystem)
@@ -843,9 +844,9 @@ void SPCGEditorGraphDebugObjectTree::AddStacksToTree(const TArray<FPCGStack>& St
 
 	for (const FPCGStack& Stack : Stacks)
 	{
+		//TRACE_CPUPROFILER_EVENT_SCOPE(SPCGEditorGraphDebugObjectTree::AddStacksToTree::PerStackLoop);
 		// If the current graph is not in the stack at all then skip.
-		bool bRemainingStackContainsEditedGraph = Stack.HasObject(GraphBeingEdited);
-		if (!bRemainingStackContainsEditedGraph)
+		if (!Stack.HasObject(GraphBeingEdited))
 		{
 			continue;
 		}
@@ -901,6 +902,7 @@ void SPCGEditorGraphDebugObjectTree::AddStacksToTree(const TArray<FPCGStack>& St
 		// Generic function to attack an item to the parent at the correct place
 		auto AddToGraphItemsAndAttachItemToParent = [&InOutStackToItem, this](FPCGStack& GraphStack, FPCGEditorGraphDebugObjectItemPtr GraphItem)
 		{
+			//TRACE_CPUPROFILER_EVENT_SCOPE(SPCGEditorGraphDebugObjectTree::AddStacksToTree::PerStackLoop::PerFrame::AddToGraphItemsAndAttachItemToParent);
 			AllGraphItems.Add(GraphItem);
 
 			TArray<FPCGStackFrame>& LoopGraphStackFrames = GraphStack.GetStackFramesMutable();
@@ -919,8 +921,9 @@ void SPCGEditorGraphDebugObjectTree::AddStacksToTree(const TArray<FPCGStack>& St
 		// Generic function to create the subgraph/loop item and hook it up to the parent properly
 		auto AddSubgraphOrLoopItemToStack = [&Stack, &InOutStackToItem, &AddToGraphItemsAndAttachItemToParent](const UPCGGraph* InGraph, const UPCGNode* InSubgraphNode, int FrameCutoff, bool bInIsDebuggable, bool bInDisplayGrayedOut, EPCGEditorSubgraphNodeType InSubgraphType)
 		{
+			//TRACE_CPUPROFILER_EVENT_SCOPE(SPCGEditorGraphDebugObjectTree::AddStacksToTree::PerStackLoop::PerFrame::AddSubgraphOrLoopItemToStack);
 			FPCGStack GraphStack = Stack;
-			GraphStack.GetStackFramesMutable().SetNum(FrameCutoff);
+			GraphStack.GetStackFramesMutable().SetNum(FrameCutoff, EAllowShrinking::No);
 
 			if (!InOutStackToItem.Contains(GraphStack))
 			{
@@ -939,8 +942,31 @@ void SPCGEditorGraphDebugObjectTree::AddStacksToTree(const TArray<FPCGStack>& St
 		//                                       ^ dynamic subgraph
 		// The loop below adds tree items for component & top graph, and then whenever a graph is encountered
 		// we look a previous frames to determine whether to add a subgraph item or loop subgraph item.
-		for (int FrameIndex = 1; FrameIndex < StackFrames.Num() && bRemainingStackContainsEditedGraph; FrameIndex++)
+		int FrameIndex = StackFrames.Num() - 1;
+
+		// Rollback from the last frames to check if there are matching stacks that would correspond to the current stack, so we don't do duplicate work.
+		{
+			//TRACE_CPUPROFILER_EVENT_SCOPE(SPCGEditorGraphDebugObjectTree::AddStacksToTree::PerStackLoop::RollbackLoop);
+			FPCGStack UnrollStack = Stack;
+			TArray<FPCGStackFrame>& UnrollFrames = UnrollStack.GetStackFramesMutable();
+
+			while (FrameIndex > 1)
+			{
+				if (InOutStackToItem.Contains(UnrollStack))
+				{
+					break;
+				}
+				else
+				{
+					--FrameIndex;
+					UnrollFrames.Pop(EAllowShrinking::No);
+				}
+			}
+		}
+
+		for(; FrameIndex < StackFrames.Num(); ++FrameIndex)
 		{
+			//TRACE_CPUPROFILER_EVENT_SCOPE(SPCGEditorGraphDebugObjectTree::AddStacksToTree::PerStackLoop::PerFrame);
 			const FPCGStackFrame& StackFrame = StackFrames[FrameIndex];
 			const FPCGStackFrame& PreviousStackFrame = StackFrames[FrameIndex - 1];
 			FPCGEditorGraphDebugObjectItemPtr CurrentItem;
@@ -1011,6 +1037,7 @@ void SPCGEditorGraphDebugObjectTree::AddStacksToTree(const TArray<FPCGStack>& St
 
 void SPCGEditorGraphDebugObjectTree::RefreshTree()
 {
+	//TRACE_CPUPROFILER_EVENT_SCOPE(SPCGEditorGraphDebugObjectTree::RefreshTree);
 	RootItems.Empty();
 	AllGraphItems.Empty();
 	DebugObjectTreeView->RequestTreeRefresh();
