[PCG] Fix resource types around data collector adaptor.

Fixes the spooky "poison attribute" which was taking down execution, due to the adaptor failing to execute properly.

#rnx

diff --git a/Engine/Plugins/PCG/Shaders/Private/PCGDataCollectionAdaptor.usf b/Engine/Plugins/PCG/Shaders/Private/PCGDataCollectionAdaptor.usf
--- a/Engine/Plugins/PCG/Shaders/Private/PCGDataCollectionAdaptor.usf
+++ b/Engine/Plugins/PCG/Shaders/Private/PCGDataCollectionAdaptor.usf
@@ -1,5 +1,9 @@
 // Copyright Epic Games, Inc. All Rights Reserved.
 
+// Updates element counts and attribute IDs in a data collection. Used to adapt a data collection
+// to a downstream compute graph which has different attribute IDs, and to apply an updated data
+// description.
+
 #include "/Engine/Private/Common.ush"
 #include "/Engine/Public/Platform.ush"
 
@@ -78,6 +82,7 @@ void MainCS(uint InThreadIndex : SV_DispatchThreadID)
 		const uint TargetAddress = GetAttributeHeaderAddress(RemapDataIndex, SourceIdToTargetId[1]);
 
 		// Write ID, stride, address to target to apply remap.
-		OutDataCollection.Store2(TargetAddress, InDataCollection.Load2(SourceAddress));
+		const uint2 AttrIDStrideAddress = InDataCollection.Load2(SourceAddress);
+		OutDataCollection.Store2(TargetAddress, AttrIDStrideAddress);
 	}
 }
diff --git a/Engine/Plugins/PCG/Source/PCG/Private/Compute/PCGDataBinding.cpp b/Engine/Plugins/PCG/Source/PCG/Private/Compute/PCGDataBinding.cpp
--- a/Engine/Plugins/PCG/Source/PCG/Private/Compute/PCGDataBinding.cpp
+++ b/Engine/Plugins/PCG/Source/PCG/Private/Compute/PCGDataBinding.cpp
@@ -170,22 +170,20 @@ void UPCGDataBinding::AddInputDataAttributesToTable(const FPCGDataCollection& In
 				{
 					// Convert from absolute ID to index into metadata attributes.
 					const int32 SourceAttributeTableIndex = PCGComputeHelpers::GetMetadataAttributeIndexFromAttributeId(SourceAttrDesc.Index);
-
-					// Negative index means the attribute is a property rather than metadata.
-					if (SourceAttributeTableIndex < 0)
+					if (SourceAttributeTableIndex == INDEX_NONE)
 					{
+						// Invalid index means the attribute is a property rather than metadata.
 						continue;
 					}
 
-					const int32 ExistingAttributeId = AttributeTable.AddAttribute(SourceAttrDesc.Name, SourceAttrDesc.Type);
+					const int32 ExistingAttributeIndex = AttributeTable.AddAttribute(SourceAttrDesc.Name, SourceAttrDesc.Type);
 
-					if (ExistingAttributeId != SourceAttributeTableIndex)
+					if (ExistingAttributeIndex != SourceAttributeTableIndex)
 					{
-						AttributeIdRemap.Add(FSourceBufferAndAttributeIndex(Proxy->GetGPUInfo(), SourceAttributeTableIndex), ExistingAttributeId);
+						AttributeIdRemap.Add(FSourceBufferAndAttributeIndex(Proxy->GetGPUInfo(), SourceAttributeTableIndex), ExistingAttributeIndex);
 					}
 				}
 			}
-
 		}
 		else if (const UPCGMetadata* Metadata = Data.Data ? Data.Data->ConstMetadata() : nullptr)
 		{
diff --git a/Engine/Plugins/PCG/Source/PCG/Private/Compute/PCGDataUploadAdaptor.cpp b/Engine/Plugins/PCG/Source/PCG/Private/Compute/PCGDataUploadAdaptor.cpp
--- a/Engine/Plugins/PCG/Source/PCG/Private/Compute/PCGDataUploadAdaptor.cpp
+++ b/Engine/Plugins/PCG/Source/PCG/Private/Compute/PCGDataUploadAdaptor.cpp
@@ -151,8 +151,7 @@ bool FPCGDataUploadAdaptor::PrepareData_GameThread(const UPCGDataBinding* InBind
 
 FRDGBufferSRVRef FPCGDataUploadAdaptor::GetAttributeRemapBufferSRV(FRDGBuilder& InGraphBuilder)
 {
-	const uint32 SizeBytes = FMath::Max(SourceToTargetAttributeId.Num(), 1) * SourceToTargetAttributeId.GetTypeSize();
-	const FRDGBufferDesc Desc = FRDGBufferDesc::CreateByteAddressDesc(SizeBytes);
+	const FRDGBufferDesc Desc = FRDGBufferDesc::CreateStructuredDesc(sizeof(FIntVector2), FMath::Max(SourceToTargetAttributeId.Num(), 1));
 
 	if (SourceToTargetAttributeId.IsEmpty())
 	{
@@ -165,7 +164,7 @@ FRDGBufferSRVRef FPCGDataUploadAdaptor::GetAttributeRemapBufferSRV(FRDGBuilder&
 	{
 		FRDGBufferRef Buffer = InGraphBuilder.CreateBuffer(Desc, TEXT("PCGDataCollectionUpload_AttributeRemap"));
 
-		InGraphBuilder.QueueBufferUpload(Buffer, SourceToTargetAttributeId.GetData(), SizeBytes);
+		InGraphBuilder.QueueBufferUpload(Buffer, MakeArrayView(SourceToTargetAttributeId));
 
 		return InGraphBuilder.CreateSRV(Buffer);
 	}
@@ -184,7 +183,7 @@ FRDGBufferSRVRef FPCGDataUploadAdaptor::GetDataElementCountsBufferSRV(FRDGBuilde
 
 	if (DataElementCounts.IsEmpty())
 	{
-		const FRDGBufferDesc Desc = FRDGBufferDesc::CreateByteAddressDesc(DataElementCounts.GetTypeSize());
+		const FRDGBufferDesc Desc = FRDGBufferDesc::CreateStructuredDesc(sizeof(uint32), DataElementCounts.Num());
 		FRDGBufferRef NullBuffer = InGraphBuilder.CreateBuffer(Desc, TEXT("EmptyDataElementCounts"));
 
 		uint32 Zero = 0;
@@ -194,12 +193,11 @@ FRDGBufferSRVRef FPCGDataUploadAdaptor::GetDataElementCountsBufferSRV(FRDGBuilde
 	}
 	else
 	{
-		const uint32 SizeBytes = DataElementCounts.Num() * DataElementCounts.GetTypeSize();
-		const FRDGBufferDesc Desc = FRDGBufferDesc::CreateByteAddressDesc(SizeBytes);
+		const FRDGBufferDesc Desc = FRDGBufferDesc::CreateStructuredDesc(sizeof(uint32), DataElementCounts.Num());
 
 		FRDGBufferRef Buffer = InGraphBuilder.CreateBuffer(Desc, TEXT("PCGDataCollectionUpload_DataElementCounts"));
 
-		InGraphBuilder.QueueBufferUpload(Buffer, DataElementCounts.GetData(), SizeBytes);
+		InGraphBuilder.QueueBufferUpload(Buffer, MakeArrayView(DataElementCounts));
 
 		return InGraphBuilder.CreateSRV(Buffer);
 	}
@@ -248,7 +246,7 @@ FRDGBufferRef FPCGDataUploadAdaptor::GetBuffer_RenderThread(FRDGBuilder& InGraph
 		OutputDesc.Usage |= BUF_SourceCopy;
 	}
 
-	FRDGBufferRef OutputBuffer = InGraphBuilder.CreateBuffer(OutputDesc, InBufferDebugName);
+	FRDGBufferRef OutputBuffer = InGraphBuilder.CreateBuffer(OutputDesc, TEXT("PCG_UDC_AttributesRemapped"));
 
 	// Populate using copy. We can't do everything in one pass trivially due to race conditions between groups.
 	AddCopyBufferPass(InGraphBuilder, OutputBuffer, InputBuffer);
