[PCG] Fixes for inspection/debugging of GPU nodes.

* Fix massive GPU->CPU readbacks when a component is inspected, by adding inspection state to UPCGSettings and using per-node state to trigger storing inspection data for GPU nodes.
* Since we don't always store GPU node inspection data, trigger generate to produce inspection data when inspection flag moved.
* Disable graying out targets in DOT when inspecting GPU node, as we don't have a flag for whether node produced data or not for GPU nodes.

#rnx

diff --git a/Engine/Plugins/PCG/Source/PCG/Private/Compute/DataInterfaces/PCGDataCollectionDataInterface.cpp b/Engine/Plugins/PCG/Source/PCG/Private/Compute/DataInterfaces/PCGDataCollectionDataInterface.cpp
--- a/Engine/Plugins/PCG/Source/PCG/Private/Compute/DataInterfaces/PCGDataCollectionDataInterface.cpp
+++ b/Engine/Plugins/PCG/Source/PCG/Private/Compute/DataInterfaces/PCGDataCollectionDataInterface.cpp
@@ -468,7 +468,7 @@ UComputeDataProvider* UPCGDataCollectionDataInterface::CreateDataProvider(TObjec
 	Provider->bExposeElementCounters = bExposeElementCounters;
 
 #if WITH_EDITOR
-	if (Binding->SourceComponent->IsInspecting())
+	if (ProducerSettings->bIsInspecting && Binding->SourceComponent->IsInspecting())
 	{
 		Provider->ExportMode |= EPCGExportMode::Inspection;
 	}
diff --git a/Engine/Plugins/PCG/Source/PCG/Public/PCGSettings.h b/Engine/Plugins/PCG/Source/PCG/Public/PCGSettings.h
--- a/Engine/Plugins/PCG/Source/PCG/Public/PCGSettings.h
+++ b/Engine/Plugins/PCG/Source/PCG/Public/PCGSettings.h
@@ -145,6 +145,9 @@ public:
 
 #if WITH_EDITOR
 	FOnPCGSettingsChanged OnSettingsChangedDelegate;
+
+	/** Whether the Inspect feature is active on the node corresponding to this settings. */
+	bool bIsInspecting = false;
 #endif
 
 	UPROPERTY(EditAnywhere, BlueprintReadWrite, Category = Debug, meta = (EditCondition = bDisplayDebuggingProperties, EditConditionHides, HideEditConditionToggle))
@@ -463,6 +466,7 @@ public:
 
 	/** Check if we need to hook the output of the pre-task to this. One use is to compute overrides in the subgraph element and pass the overrides as data, to all nodes that needs it. */
 	virtual bool RequiresDataFromPreTask() const { return false; }
+
 protected:
 	/** Iterate over OverridableParams to automatically add param pins to the list. */
 	void FillOverridableParamsPins(TArray<FPCGPinProperties>& OutPins) const;
diff --git a/Engine/Plugins/PCG/Source/PCGEditor/Private/Nodes/PCGEditorGraphNodeBase.cpp b/Engine/Plugins/PCG/Source/PCGEditor/Private/Nodes/PCGEditorGraphNodeBase.cpp
--- a/Engine/Plugins/PCG/Source/PCGEditor/Private/Nodes/PCGEditorGraphNodeBase.cpp
+++ b/Engine/Plugins/PCG/Source/PCGEditor/Private/Nodes/PCGEditorGraphNodeBase.cpp
@@ -418,6 +418,45 @@ void UPCGEditorGraphNodeBase::PostPaste()
 	bDisableReconstructFromNode = false;
 }
 
+void UPCGEditorGraphNodeBase::SetInspected(bool bInIsInspecting)
+{
+	if (UPCGSettings* Settings = PCGNode ? PCGNode->GetSettings() : nullptr; ensure(Settings))
+	{
+		const bool bWasInspectingBefore = Settings->bIsInspecting;
+
+		Settings->bIsInspecting = bInIsInspecting;
+
+		// If we start inspecting a GPU node, we need to generate to populate inspection data. The
+		// normal workflow optimization that avoids re-executing when moving inspection flag around
+		// graph relies on all nodes storing inspection data which is efficient for CPU but not for GPU,
+		// where we only do expensive GPU->CPU readbacks for the currently inspected GPU node.
+		if (!bWasInspectingBefore && bInIsInspecting && Settings->bEnabled && Settings->ShouldExecuteOnGPU())
+		{
+			const UPCGEditorGraph* EditorGraph = CastChecked<UPCGEditorGraph>(GetGraph());
+			const FPCGEditor* Editor = EditorGraph ? EditorGraph->GetEditor().Pin().Get() : nullptr;
+
+			// Force refresh/regenerate.
+			if (UPCGComponent* Component = Editor ? Editor->GetPCGComponentBeingInspected() : nullptr)
+			{
+				if (Component->IsManagedByRuntimeGenSystem())
+				{
+					Component->Refresh(EPCGChangeType::Node);
+				}
+				else
+				{
+					Component->GenerateLocal(/*bForce=*/true);
+				}
+			}
+		}
+	}
+}
+
+bool UPCGEditorGraphNodeBase::GetInspected() const
+{
+	const UPCGSettings* Settings = PCGNode ? PCGNode->GetSettings() : nullptr;
+	return ensure(Settings) && Settings->bIsInspecting;
+}
+
 void UPCGEditorGraphNodeBase::EnableDeferredReconstruct()
 {
 	ensure(DeferredReconstructCounter >= 0);
diff --git a/Engine/Plugins/PCG/Source/PCGEditor/Private/Nodes/PCGEditorGraphNodeBase.h b/Engine/Plugins/PCG/Source/PCGEditor/Private/Nodes/PCGEditorGraphNodeBase.h
--- a/Engine/Plugins/PCG/Source/PCGEditor/Private/Nodes/PCGEditorGraphNodeBase.h
+++ b/Engine/Plugins/PCG/Source/PCGEditor/Private/Nodes/PCGEditorGraphNodeBase.h
@@ -62,8 +62,8 @@ public:
 	void RebuildAfterPaste();
 	virtual void PostPaste();
 
-	void SetInspected(bool InIsInspecting) { bIsInspected = InIsInspecting; }
-	bool GetInspected() const { return bIsInspected; }
+	void SetInspected(bool bInIsInspecting);
+	bool GetInspected() const;
 
 	/** Whether node was culled either during compilation or at execution time. */
 	void SetIsCulledFromExecution(bool bInIsCulledFromExecution) { bIsCulledFromExecution = bInIsCulledFromExecution; }
@@ -156,7 +156,6 @@ protected:
 	int32 DeferredReconstructCounter = 0;
 	bool bDeferredReconstruct = false;
 	bool bDisableReconstructFromNode = false;
-	bool bIsInspected = false;
 
 	/** Whether this node was culled in the last execution. */
 	bool bIsCulledFromExecution = false;
diff --git a/Engine/Plugins/PCG/Source/PCGEditor/Private/SPCGEditorGraphDebugObjectTree.cpp b/Engine/Plugins/PCG/Source/PCGEditor/Private/SPCGEditorGraphDebugObjectTree.cpp
--- a/Engine/Plugins/PCG/Source/PCGEditor/Private/SPCGEditorGraphDebugObjectTree.cpp
+++ b/Engine/Plugins/PCG/Source/PCGEditor/Private/SPCGEditorGraphDebugObjectTree.cpp
@@ -877,7 +877,16 @@ void SPCGEditorGraphDebugObjectTree::AddStacksToTree(const TArray<FPCGStack>& St
 		}
 
 		// If we're inspecting a node which has not logged inspection data in a previous execution, display grayed out.
-		const bool bDisplayGrayedOut = PCGNodeBeingInspected.IsValid() && !PCGComponent->HasNodeProducedData(PCGNodeBeingInspected.Get(), Stack);
+		bool bDisplayGrayedOut = false;
+		if (PCGNodeBeingInspected.IsValid())
+		{
+			const UPCGSettings* Settings = PCGNodeBeingInspected->GetSettings();
+			const bool bGPUNode = Settings && Settings->bEnabled && Settings->ShouldExecuteOnGPU();
+
+			// Display grayed out if no inspection data has been stored for this node, and node does not run on GPU,
+			// because we don't opportunistically store inspection data for GPU nodes.
+			bDisplayGrayedOut = !bGPUNode && !PCGComponent->HasNodeProducedData(PCGNodeBeingInspected.Get(), Stack);
+		}
 
 		AActor* Actor = PCGComponent->GetOwner();
 		if (!Actor)
