[PCG] Change some lock scopes to avoid to many locks/unlocks inside loops

BiomeSample 4 runs (heavily multithreaded calls so this is not an equivalent wall time save :) but still reduces lock overhead a lot
- StoreResults: 5.4s -> 0.6s
- QueueNextTask: 4.7s -> 2.4s
- CullInactiveDownstreamNodes: 1.6s -> 0.4s
- ExecuteScheduling: 0.7s -> 0.4s
- ProcessCachedResults: 1.5s -> 0.4s

#rnx

diff --git a/Engine/Plugins/PCG/Source/PCG/Private/Graph/PCGGraphExecutor.cpp b/Engine/Plugins/PCG/Source/PCG/Private/Graph/PCGGraphExecutor.cpp
--- a/Engine/Plugins/PCG/Source/PCG/Private/Graph/PCGGraphExecutor.cpp
+++ b/Engine/Plugins/PCG/Source/PCG/Private/Graph/PCGGraphExecutor.cpp
@@ -778,58 +778,67 @@ bool FPCGGraphExecutor::ProcessScheduledTasks()
 	bool bProcessedNewTasks = false;
 	const bool bIsInGameThread = IsInGameThread();
 	TArray<FCachedResult*> CachedResults;
+	TArray<FPCGGraphTask, TInlineAllocator<16>> NewReadyTasks;
 
 	{
 		// Process any newly scheduled graphs to execute
 		PCGGraphExecutor::TScopeLock ScopeLock(ScheduleLock, !bIsInGameThread);
 		if (ScopeLock.IsLocked())
 		{
-			for (int32 ScheduledTaskIndex = ScheduledTasks.Num() - 1; ScheduledTaskIndex >= 0; --ScheduledTaskIndex)
+			if (ScheduledTasks.Num() > 0)
 			{
-				FPCGGraphScheduleTask& ScheduledTask = ScheduledTasks[ScheduledTaskIndex];
+				PCGGraphExecutor::TScopeLock ChildScopeLock(TasksLock);
+				PCGGraphExecutor::TScopeLock OutputLock(TaskOutputsLock);
 
-				check(ScheduledTask.Tasks.Num() > 0);
-				// Push tasks to the primary task list & build successors map
-				for (FPCGGraphTask& Task : ScheduledTask.Tasks)
+				for (int32 ScheduledTaskIndex = ScheduledTasks.Num() - 1; ScheduledTaskIndex >= 0; --ScheduledTaskIndex)
 				{
-					const FPCGTaskId TaskId = Task.NodeId;
+					FPCGGraphScheduleTask& ScheduledTask = ScheduledTasks[ScheduledTaskIndex];
 
-					// TODO: review if it's actually possible for a task with inputs to be ready at this point - it seems very unlikely
-					bool bPushToReady = true;
+					check(ScheduledTask.Tasks.Num() > 0);
+					// Push tasks to the primary task list & build successors map
+					for (FPCGGraphTask& Task : ScheduledTask.Tasks)
 					{
-						PCGGraphExecutor::TScopeLock ChildScopeLock(TasksLock);
-						PCGGraphExecutor::TScopeLock OutputLock(TaskOutputsLock);
-						for (const FPCGGraphTaskInput& Input : Task.Inputs)
+						const FPCGTaskId TaskId = Task.NodeId;
+
+						// TODO: review if it's actually possible for a task with inputs to be ready at this point - it seems very unlikely
+						bool bPushToReady = true;
 						{
-							if (!TaskOutputs.Contains(Input.TaskId))
+							for (const FPCGGraphTaskInput& Input : Task.Inputs)
 							{
-								TaskSuccessors.FindOrAdd(Input.TaskId).Add(TaskId);
-								bPushToReady = false;
+								if (!TaskOutputs.Contains(Input.TaskId))
+								{
+									TaskSuccessors.FindOrAdd(Input.TaskId).Add(TaskId);
+									bPushToReady = false;
+								}
 							}
 						}
-					}
 
-					// Automatically push inputless/already satisfied tasks to the ready queue
-					if (bPushToReady)
-					{
-						OnTaskInputsReady(Task, CachedResults, bIsInGameThread);
-					}
-					else
-					{
-						PCGGraphExecutor::TScopeLock ChildScopeLock(TasksLock);
-						Tasks.Add(TaskId, MoveTemp(Task));
+						// Automatically push inputless/already satisfied tasks to the ready queue
+						if (bPushToReady)
+						{
+							NewReadyTasks.Add(MoveTemp(Task));
+						}
+						else
+						{
+							Tasks.Add(TaskId, MoveTemp(Task));
+						}
+						bProcessedNewTasks = true;
 					}
-					bProcessedNewTasks = true;
 				}
+
+				if (PCGGraphExecutionLogging::CullingLogEnabled())
+				{
+					PCGGraphExecutionLogging::LogGraphPostSchedule(Tasks, TaskSuccessors);
+				}
+				
+				ScheduledTasks.Reset();
 			}
 
-			if (!ScheduledTasks.IsEmpty())
+			// Process ready tasks
+			for (FPCGGraphTask& ReadyTask : NewReadyTasks)
 			{
-				PCGGraphExecutor::TScopeLock ChildScopeLock(TasksLock);
-				PCGGraphExecutionLogging::LogGraphPostSchedule(Tasks, TaskSuccessors);
+				OnTaskInputsReady(ReadyTask, CachedResults, bIsInGameThread);
 			}
-
-			ScheduledTasks.Reset();
 		}
 	}
 
@@ -837,14 +846,21 @@ bool FPCGGraphExecutor::ProcessScheduledTasks()
 		
 	if (bIsInGameThread)
 	{
-		PCGGraphExecutor::TScopeLock ScopeLock(TasksLock);
-		PCGGraphExecutor::TScopeLock ChildScopeLock(LiveTasksLock);
-		PCGGraphExecutor::TScopeLock CachingResultsScopeLock(CachingResultsLock);
-		// This is a safeguard to check if we're in a stuck state
-		if (CollectGCCachingResults.Num() == 0 && ReadyTasks.Num() == 0 && ActiveTasks.Num() == 0 && ActiveTasksGameThreadOnly.Num() == 0 && SleepingTasks.Num() == 0 && Tasks.Num() > 0)
+		// Check every second to avoid locking too often for an error that shouldn't happen
+		double TimeInSeconds = FPlatformTime::Seconds();
+		if (TimeInSeconds - LastSchedulingErrorCheck > 1.0)
 		{
-			UE_LOG(LogPCG, Error, TEXT("PCG Graph executor error: tasks are in a deadlocked state. Will drop all tasks."));
-			ClearAllTasks();
+			LastSchedulingErrorCheck = TimeInSeconds;
+
+			PCGGraphExecutor::TScopeLock ScopeLock(TasksLock);
+			PCGGraphExecutor::TScopeLock ChildScopeLock(LiveTasksLock);
+			PCGGraphExecutor::TScopeLock CachingResultsScopeLock(CachingResultsLock);
+			// This is a safeguard to check if we're in a stuck state
+			if (CollectGCCachingResults.Num() == 0 && ReadyTasks.Num() == 0 && ActiveTasks.Num() == 0 && ActiveTasksGameThreadOnly.Num() == 0 && SleepingTasks.Num() == 0 && Tasks.Num() > 0)
+			{
+				UE_LOG(LogPCG, Error, TEXT("PCG Graph executor error: tasks are in a deadlocked state. Will drop all tasks."));
+				ClearAllTasks();
+			}
 		}
 	}
 
@@ -860,7 +876,6 @@ void FPCGGraphExecutor::ExecuteTasksEnded()
 	}
 	bNeedToExecuteTasksEnded = false;
 
-	check(IsInGameThread());
 	// Nothing left to do; we'll release everything here.
 	// TODO: this is fine and will make sure any intermediate data is properly
 	// garbage collected, however, this goes a bit against our goals if we want to
@@ -1621,42 +1636,47 @@ TArray<FPCGGraphExecutor::FCachedResult*> FPCGGraphExecutor::QueueNextTasksInter
 	TArray<FCachedResult*> CachedResults;
 	const bool bIsInGameThread = IsInGameThread();
 
+	
+	TSet<FPCGTaskId> Successors;
+	TArray<FPCGGraphTask, TInlineAllocator<16>> NewReadyTasks;
 	{
-		TSet<FPCGTaskId> Successors;
+		PCGGraphExecutor::TScopeLock ScopeLock(TasksLock);
+		PCGGraphExecutor::TScopeLock OutputLock(TaskOutputsLock);
+
+		if (TaskSuccessors.RemoveAndCopyValue(FinishedTask, Successors))
 		{
-			PCGGraphExecutor::TScopeLock ScopeLock(TasksLock);
-		
-			if (TaskSuccessors.RemoveAndCopyValue(FinishedTask, Successors))
+			for (FPCGTaskId Successor : Successors)
 			{
-				for (FPCGTaskId Successor : Successors)
+				bool bAllPrerequisitesMet = true;
+
+				// When running multithreaded it is possible for the Successor to already have been moved to ReadyTask after the current task stored its results
+				if (FPCGGraphTask* SuccessorTaskPtr = Tasks.Find(Successor))
 				{
-					bool bAllPrerequisitesMet = true;
+					FPCGGraphTask& SuccessorTask = *SuccessorTaskPtr;
 
-					// When running multithreaded it is possible for the Successor to already have been moved to ReadyTask after the current task stored its results
-					if (FPCGGraphTask* SuccessorTaskPtr = Tasks.Find(Successor))
+					if (SuccessorTask.Inputs.Num() > 0)
 					{
-						FPCGGraphTask& SuccessorTask = *SuccessorTaskPtr;
-
-						if (SuccessorTask.Inputs.Num() > 0)
+						for (const FPCGGraphTaskInput& Input : SuccessorTask.Inputs)
 						{
-							PCGGraphExecutor::TScopeLock OutputLock(TaskOutputsLock);
-							for (const FPCGGraphTaskInput& Input : SuccessorTask.Inputs)
-							{
-								bAllPrerequisitesMet &= TaskOutputs.Contains(Input.TaskId);
-							}
+							bAllPrerequisitesMet &= TaskOutputs.Contains(Input.TaskId);
 						}
+					}
 
-						if (bAllPrerequisitesMet)
-						{
-							OnTaskInputsReady(SuccessorTask, CachedResults, bIsInGameThread);
-							Tasks.Remove(Successor);
+					if (bAllPrerequisitesMet)
+					{
+						NewReadyTasks.Add(MoveTemp(SuccessorTask));
+						Tasks.Remove(Successor);
 					}
 				}
-				}
 			}
 		}
 	}
 
+	for (FPCGGraphTask& NewReadyTask : NewReadyTasks)
+	{
+		OnTaskInputsReady(NewReadyTask, CachedResults, bIsInGameThread);
+	}
+	
 	return CachedResults;
 }
 
@@ -1836,91 +1856,93 @@ void FPCGGraphExecutor::BuildTaskInput(FPCGGraphTask& Task, TArray<FPCGTaskId>&
 	TArray<FPCGTaggedData, TInlineAllocator<16>> InputDataOnPin;
 	TArray<FPCGCrc, TInlineAllocator<16>> InputDataCrcsOnPin;
 	
-	for (const FPCGGraphTaskInput& Input : Task.Inputs)
+	if(Task.Inputs.Num() > 0)
 	{
 		PCGGraphExecutor::TScopeLock OutputLock(TaskOutputsLock);
-		
-		const FOutputDataInfo& TaskOutput = TaskOutputs.FindChecked(Input.TaskId);
-		
-		ResultsToMarkAsRead.AddUnique(Input.TaskId);
-
-		// Input was culled before this task managed to get scheduled (moved from ScheduleTasks to Tasks)
-		if (TaskOutput.bCulled)
+		for (const FPCGGraphTaskInput& Input : Task.Inputs)
 		{
-			continue;
-		}
+			const FOutputDataInfo& TaskOutput = TaskOutputs.FindChecked(Input.TaskId);
 
-		// If the input does not provide any data, don't add it to the task input.
-		if (!Input.bProvideData)
-		{
-			continue;
-		}
+			ResultsToMarkAsRead.AddUnique(Input.TaskId);
 
-		const bool bAllowMultipleData = Input.DownstreamPin.IsSet() ? Input.DownstreamPin.GetValue().bAllowMultipleData : true;
-		const uint32 InputPinLabelCrc = Input.DownstreamPin.IsSet() ? GetTypeHash(Input.DownstreamPin.GetValue().Label) : DefaultHashForNoOutputPin;
+			// Input was culled before this task managed to get scheduled (moved from ScheduleTasks to Tasks)
+			if (TaskOutput.bCulled)
+			{
+				continue;
+			}
 
-		// Enforce single data - if already have input for this pin, don't add more. Early check before other side effects below.
-		if (Input.DownstreamPin.IsSet() && !bAllowMultipleData && Task.TaskInput.GetInputCountByPin(Input.DownstreamPin.GetValue().Label) > 0)
-		{
-			LogDiscardedData(Input.DownstreamPin);
-			continue;
-		}
+			// If the input does not provide any data, don't add it to the task input.
+			if (!Input.bProvideData)
+			{
+				continue;
+			}
 
-		const FPCGDataCollection& InputCollection = TaskOutput.DataCollection;
-		
-		Task.TaskInput.bCancelExecution |= InputCollection.bCancelExecution;
+			const bool bAllowMultipleData = Input.DownstreamPin.IsSet() ? Input.DownstreamPin.GetValue().bAllowMultipleData : true;
+			const uint32 InputPinLabelCrc = Input.DownstreamPin.IsSet() ? GetTypeHash(Input.DownstreamPin.GetValue().Label) : DefaultHashForNoOutputPin;
 
-		const int32 TaggedDataOffset = Task.TaskInput.TaggedData.Num();
+			// Enforce single data - if already have input for this pin, don't add more. Early check before other side effects below.
+			if (Input.DownstreamPin.IsSet() && !bAllowMultipleData && Task.TaskInput.GetInputCountByPin(Input.DownstreamPin.GetValue().Label) > 0)
+			{
+				LogDiscardedData(Input.DownstreamPin);
+				continue;
+			}
 
-		// Get input data at the given pin (or everything). This will add the data and include the input pin Crc to uniquely identify
-		// inputs per-pin, or use a placeholder for symmetry.
-		// Note: The input data CRC will already contain the output pin (calculated in element post execute).
-		if (Input.UpstreamPin.IsSet())
-		{
-			InputDataOnPin.Reset();
-			InputDataCrcsOnPin.Reset();
-			InputCollection.GetInputsAndCrcsByPin(Input.UpstreamPin.GetValue().Label, InputDataOnPin, InputDataCrcsOnPin);
+			const FPCGDataCollection& InputCollection = TaskOutput.DataCollection;
+
+			Task.TaskInput.bCancelExecution |= InputCollection.bCancelExecution;
 
-			if (!InputDataOnPin.IsEmpty())
+			const int32 TaggedDataOffset = Task.TaskInput.TaggedData.Num();
+
+			// Get input data at the given pin (or everything). This will add the data and include the input pin Crc to uniquely identify
+			// inputs per-pin, or use a placeholder for symmetry.
+			// Note: The input data CRC will already contain the output pin (calculated in element post execute).
+			if (Input.UpstreamPin.IsSet())
 			{
-				// Proceed carefully when adding data items - if pin is single-data, only add first item.
-				if (!ensure(InputDataOnPin.Num() == InputDataCrcsOnPin.Num()))
+				InputDataOnPin.Reset();
+				InputDataCrcsOnPin.Reset();
+				InputCollection.GetInputsAndCrcsByPin(Input.UpstreamPin.GetValue().Label, InputDataOnPin, InputDataCrcsOnPin);
+
+				if (!InputDataOnPin.IsEmpty())
 				{
-					InputDataCrcsOnPin.SetNumZeroed(InputDataOnPin.Num());
-				}
+					// Proceed carefully when adding data items - if pin is single-data, only add first item.
+					if (!ensure(InputDataOnPin.Num() == InputDataCrcsOnPin.Num()))
+					{
+						InputDataCrcsOnPin.SetNumZeroed(InputDataOnPin.Num());
+					}
 
-				const int NumberDataItemsToTake = bAllowMultipleData ? InputDataOnPin.Num() : 1;
+					const int NumberDataItemsToTake = bAllowMultipleData ? InputDataOnPin.Num() : 1;
 
-				Task.TaskInput.AddDataForPin(
-					MakeArrayView(InputDataOnPin.GetData(), NumberDataItemsToTake),
-					MakeArrayView(InputDataCrcsOnPin.GetData(), NumberDataItemsToTake),
-					InputPinLabelCrc);
+					Task.TaskInput.AddDataForPin(
+						MakeArrayView(InputDataOnPin.GetData(), NumberDataItemsToTake),
+						MakeArrayView(InputDataCrcsOnPin.GetData(), NumberDataItemsToTake),
+						InputPinLabelCrc);
 
-				if (NumberDataItemsToTake < InputDataOnPin.Num())
-				{
-					LogDiscardedData(Input.DownstreamPin);
+					if (NumberDataItemsToTake < InputDataOnPin.Num())
+					{
+						LogDiscardedData(Input.DownstreamPin);
+					}
 				}
 			}
-		}
-		else
-		{
-			Task.TaskInput.AddData(InputCollection.TaggedData, InputCollection.DataCrcs);
-		}
+			else
+			{
+				Task.TaskInput.AddData(InputCollection.TaggedData, InputCollection.DataCrcs);
+			}
 
-		// Apply labelling on data; technically, we should ensure that we do this only for pass-through nodes,
-		// Otherwise we could also null out the label on the input...
-		// Also marking the data used multiple times if the given input is used multiple times.
-		if (Input.DownstreamPin.IsSet() || Input.bIsUsedMultipleTimes)
-		{
-			for (int32 TaggedDataIndex = TaggedDataOffset; TaggedDataIndex < Task.TaskInput.TaggedData.Num(); ++TaggedDataIndex)
+			// Apply labelling on data; technically, we should ensure that we do this only for pass-through nodes,
+			// Otherwise we could also null out the label on the input...
+			// Also marking the data used multiple times if the given input is used multiple times.
+			if (Input.DownstreamPin.IsSet() || Input.bIsUsedMultipleTimes)
 			{
-				FPCGTaggedData& TaggedData = Task.TaskInput.TaggedData[TaggedDataIndex];
-				if (Input.DownstreamPin.IsSet())
+				for (int32 TaggedDataIndex = TaggedDataOffset; TaggedDataIndex < Task.TaskInput.TaggedData.Num(); ++TaggedDataIndex)
 				{
-					TaggedData.Pin = Input.DownstreamPin.GetValue().Label;
-				}
+					FPCGTaggedData& TaggedData = Task.TaskInput.TaggedData[TaggedDataIndex];
+					if (Input.DownstreamPin.IsSet())
+					{
+						TaggedData.Pin = Input.DownstreamPin.GetValue().Label;
+					}
 
-				TaggedData.bIsUsedMultipleTimes |= Input.bIsUsedMultipleTimes;
+					TaggedData.bIsUsedMultipleTimes |= Input.bIsUsedMultipleTimes;
+				}
 			}
 		}
 	}
@@ -2066,114 +2088,122 @@ void FPCGGraphExecutor::CullInactiveDownstreamNodes(FPCGTaskId InCompletedTaskId
 	// Hoisted out of loop for performance reasons.
 	TArray<FPCGTaskId, TInlineAllocator<64>> TasksToRemove;
 
-	while (!PinIdsToDeactivate.IsEmpty())
+	if (PinIdsToDeactivate.Num() > 0)
 	{
-		const FPCGPinId PinId = PinIdsToDeactivate.Pop(EAllowShrinking::No);
-		const FPCGTaskId PinTaskId = PCGPinIdHelpers::GetNodeIdFromPinId(PinId);
-
-		PCGGraphExecutionLogging::LogTaskCullingBeginLoop(PinTaskId, PCGPinIdHelpers::GetPinIndexFromPinId(PinId), PinIdsToDeactivate);
-		LogTaskState();
-
 		PCGGraphExecutor::TScopeLock ScopeLock(TasksLock);
-		const TSet<FPCGTaskId>* Successors = TaskSuccessors.Find(PinTaskId);
-		if (!Successors)
+		while (!PinIdsToDeactivate.IsEmpty())
 		{
-			continue;
-		}
+			const FPCGPinId PinId = PinIdsToDeactivate.Pop(EAllowShrinking::No);
+			const FPCGTaskId PinTaskId = PCGPinIdHelpers::GetNodeIdFromPinId(PinId);
 
-		TasksToRemove.SetNum(0, EAllowShrinking::No);
-
-		// Build set of tasks that are candidates for culling when PinId is deactivated.
-		for (const FPCGTaskId SuccessorTaskId : *Successors)
-		{
-			// Successors are updated at the end of this function, which means it may
-			// contain task IDs that have been removed.
-			if (FPCGGraphTask* FoundTask = Tasks.Find(SuccessorTaskId))
+			PCGGraphExecutionLogging::LogTaskCullingBeginLoop(PinTaskId, PCGPinIdHelpers::GetPinIndexFromPinId(PinId), PinIdsToDeactivate);
+			LogTaskStateNoLock();
+						
+			const TSet<FPCGTaskId>* Successors = TaskSuccessors.Find(PinTaskId);
+			if (!Successors)
 			{
-				bool bDependencyExpressionBecameFalse;
-				FoundTask->PinDependency.DeactivatePin(PinId, bDependencyExpressionBecameFalse);
+				continue;
+			}
+
+			TasksToRemove.SetNum(0, EAllowShrinking::No);
 
-				if (bDependencyExpressionBecameFalse)
+			// Build set of tasks that are candidates for culling when PinId is deactivated.
+			for (const FPCGTaskId SuccessorTaskId : *Successors)
+			{
+				// Successors are updated at the end of this function, which means it may
+				// contain task IDs that have been removed.
+				if (FPCGGraphTask* FoundTask = Tasks.Find(SuccessorTaskId))
 				{
-					TasksToRemove.AddUnique(SuccessorTaskId);
-				}
+					bool bDependencyExpressionBecameFalse;
+					FoundTask->PinDependency.DeactivatePin(PinId, bDependencyExpressionBecameFalse);
 
-				PCGGraphExecutionLogging::LogTaskCullingUpdatedPinDeps(SuccessorTaskId, FoundTask->PinDependency, bDependencyExpressionBecameFalse);
+					if (bDependencyExpressionBecameFalse)
+					{
+						TasksToRemove.AddUnique(SuccessorTaskId);
+					}
+
+					PCGGraphExecutionLogging::LogTaskCullingUpdatedPinDeps(SuccessorTaskId, FoundTask->PinDependency, bDependencyExpressionBecameFalse);
+				}
 			}
-		}
 
-		// Now remove the tasks.
-		for (const FPCGTaskId RemovedTaskId : TasksToRemove)
-		{
-			// Scope in which RemovedTask reference is valid.
+			// Now remove the tasks.
+			for (const FPCGTaskId RemovedTaskId : TasksToRemove)
 			{
-				FPCGGraphTask& RemovedTask = Tasks[RemovedTaskId];
+				// Scope in which RemovedTask reference is valid.
+				{
+					FPCGGraphTask& RemovedTask = Tasks[RemovedTaskId];
 
-				const UPCGNode* Node = RemovedTask.Node;
-				const int PinCount = Node ? Node->GetOutputPins().Num() : 0;
+					const UPCGNode* Node = RemovedTask.Node;
+					const int PinCount = Node ? Node->GetOutputPins().Num() : 0;
 
-				if (PinCount > 0)
-				{
-					// Deactivate all output pins.
-					const uint64 InactiveOutputPinBitmask = (1ULL << PinCount) - 1;
+					if (PinCount > 0)
+					{
+						// Deactivate all output pins.
+						const uint64 InactiveOutputPinBitmask = (1ULL << PinCount) - 1;
 
-					// Deactivate its pins - add to set of pins to deactivate.
-					GetPinIdsToDeactivate(RemovedTaskId, InactiveOutputPinBitmask, PinIdsToDeactivate);
+						// Deactivate its pins - add to set of pins to deactivate.
+						GetPinIdsToDeactivate(RemovedTaskId, InactiveOutputPinBitmask, PinIdsToDeactivate);
 
 #if WITH_EDITOR
-					SendInactivePinNotification(RemovedTask.Node, RemovedTask.GetStack(), InactiveOutputPinBitmask);
+						SendInactivePinNotification(RemovedTask.Node, RemovedTask.GetStack(), InactiveOutputPinBitmask);
 #endif
-				}
+					}
 
-				// Also register a special pin-less pin ID for this node, for task dependencies that do not have a specific pin.
-				PinIdsToDeactivate.AddUnique(PCGPinIdHelpers::NodeIdToPinId(RemovedTaskId));
+					// Also register a special pin-less pin ID for this node, for task dependencies that do not have a specific pin.
+					PinIdsToDeactivate.AddUnique(PCGPinIdHelpers::NodeIdToPinId(RemovedTaskId));
 
-				// Remove task as successor of upstream node.
-				RemoveTaskFromInputSuccessorsNoLock(RemovedTaskId, RemovedTask.Inputs);
+					// Remove task as successor of upstream node.
+					RemoveTaskFromInputSuccessorsNoLock(RemovedTaskId, RemovedTask.Inputs);
 
-				// Remove the deleted tasks from the inputs of downstream tasks.
-				if (const TSet<FPCGTaskId>* SuccessorsOfRemovedTask = TaskSuccessors.Find(RemovedTaskId))
-				{
-					for (const FPCGTaskId SuccessorTaskId : *SuccessorsOfRemovedTask)
+					// Remove the deleted tasks from the inputs of downstream tasks.
+					if (const TSet<FPCGTaskId>* SuccessorsOfRemovedTask = TaskSuccessors.Find(RemovedTaskId))
 					{
-						if (FPCGGraphTask* SuccessorTask = Tasks.Find(SuccessorTaskId))
+						for (const FPCGTaskId SuccessorTaskId : *SuccessorsOfRemovedTask)
 						{
-							for (int InputIndex = SuccessorTask->Inputs.Num() - 1; InputIndex >= 0; --InputIndex)
+							if (FPCGGraphTask* SuccessorTask = Tasks.Find(SuccessorTaskId))
 							{
-								if (SuccessorTask->Inputs[InputIndex].TaskId == RemovedTaskId)
+								for (int InputIndex = SuccessorTask->Inputs.Num() - 1; InputIndex >= 0; --InputIndex)
 								{
-									SuccessorTask->Inputs.RemoveAtSwap(InputIndex, EAllowShrinking::No);
+									if (SuccessorTask->Inputs[InputIndex].TaskId == RemovedTaskId)
+									{
+										SuccessorTask->Inputs.RemoveAtSwap(InputIndex, EAllowShrinking::No);
+									}
 								}
 							}
 						}
 					}
 				}
+
+				// Remove from tasks. After this step all traces of RemovedTaskId should be erased from tasks, task inputs. Task successors will be
+				// updated below when queuing next tasks.
+				Tasks.Remove(RemovedTaskId);
 			}
 
-			// Remove from tasks. After this step all traces of RemovedTaskId should be erased from tasks, task inputs. Task successors will be
-			// updated below when queuing next tasks.
-			Tasks.Remove(RemovedTaskId);
+			AllRemovedTasks.Append(TasksToRemove);
 		}
-
-		AllRemovedTasks.Append(TasksToRemove);
 	}
 
-	// Ensure any downstream tasks are enqueued.
-	for (const FPCGTaskId TaskId : AllRemovedTasks)
+	if (AllRemovedTasks.Num() > 0)
 	{
 		{
 			// Store culled output
 			PCGGraphExecutor::TScopeLock OutputLock(TaskOutputsLock);
-			FOutputDataInfo CullDataInfo;
-			CullDataInfo.RemainingSuccessorCount = 0;
-			CullDataInfo.bCulled = true;
-			TaskOutputs.Add(TaskId, MoveTemp(CullDataInfo));
+			for (const FPCGTaskId TaskId : AllRemovedTasks)
+			{
+				FOutputDataInfo CullDataInfo;
+				CullDataInfo.RemainingSuccessorCount = 0;
+				CullDataInfo.bCulled = true;
+				TaskOutputs.Add(TaskId, MoveTemp(CullDataInfo));
+			}
 		}
 
 		// Queue downstream tasks in a similar manner to when a task draws from the cache and is skipped.
 		// Some downstream tasks will have been culled which we don't care about (hence the ignore flag),
 		// but some may not be queued and may be ready for queuing.
-		QueueNextTasks(TaskId);
+		for (const FPCGTaskId TaskId : AllRemovedTasks)
+		{
+			QueueNextTasks(TaskId);
+		}
 	}
 }
 
@@ -2306,13 +2336,12 @@ FPCGElementPtr FPCGGraphExecutor::GetPreGraphElement()
 	return PreGraphElementPtr;
 }
 
-void FPCGGraphExecutor::LogTaskState() const
+void FPCGGraphExecutor::LogTaskStateNoLock() const
 {
 #if WITH_EDITOR
 	if (PCGGraphExecutionLogging::CullingLogEnabled())
 	{
 		UE_LOG(LogPCG, Log, TEXT("\tDORMANT (FPCGGraphExecutor::Tasks):"));
-		PCGGraphExecutor::TScopeLock ScopeLock(TasksLock);
 		PCGGraphExecutionLogging::LogGraphTasks(Tasks, &TaskSuccessors);
 	}
 #endif
diff --git a/Engine/Plugins/PCG/Source/PCG/Private/Graph/PCGGraphExecutor.h b/Engine/Plugins/PCG/Source/PCG/Private/Graph/PCGGraphExecutor.h
--- a/Engine/Plugins/PCG/Source/PCG/Private/Graph/PCGGraphExecutor.h
+++ b/Engine/Plugins/PCG/Source/PCG/Private/Graph/PCGGraphExecutor.h
@@ -246,7 +246,7 @@ private:
 	FPCGElementPtr GetFetchInputElement();
 	FPCGElementPtr GetPreGraphElement();
 
-	void LogTaskState() const;
+	void LogTaskStateNoLock() const;
 
 	int32 GetNonScheduledRemainingTaskCount() const;
 
@@ -276,19 +276,20 @@ private:
 	 * Define a Lock level for future reference. Rule is when we have a lock, we can't lock a lower or equal level lock to prevent deadlocks.
 	 * Example: When locking LiveTasksLock we can't lock ScheduleLock or CollectGCReferenceTasksLock 
 	 */
+	using FPCGLock = UE::FSpinLock;
 
 	/** Lock level - 1 (top most lock) */
-	UE::FSpinLock ScheduleLock;
+	FPCGLock ScheduleLock;
 	TArray<FPCGGraphScheduleTask> ScheduledTasks;
 	std::atomic<bool> SchedulingDisabled = false;
 
 	/** Lock level 2 */
-	mutable UE::FSpinLock TasksLock;
+	mutable FPCGLock TasksLock;
 	TMap<FPCGTaskId, FPCGGraphTask> Tasks;
 	TMap<FPCGTaskId, TSet<FPCGTaskId>> TaskSuccessors;
 
 	/** Lock level 3 */
-	UE::FSpinLock LiveTasksLock;
+	FPCGLock LiveTasksLock;
 	TArray<FPCGGraphTask> ReadyTasks;
 	TArray<TSharedPtr<FPCGGraphActiveTask>> ActiveTasks;
 	TArray<TSharedPtr<FPCGGraphActiveTask>> ActiveTasksGameThreadOnly;
@@ -296,7 +297,7 @@ private:
 	bool bNeedToCheckSleepingTasks = false;
 
 	/** Lock level 3 */
-	UE::FSpinLock CollectGCReferenceTasksLock;
+	FPCGLock CollectGCReferenceTasksLock;
 	TSet<TSharedPtr<FPCGGraphActiveTask>> CollectGCReferenceTasks;
 			
 	/** Map of node instances to their output, could be cleared once execution is done */
@@ -313,12 +314,12 @@ private:
 	};
 
 	/** Lock level 4 */
-	UE::FSpinLock CachingResultsLock;
+	FPCGLock CachingResultsLock;
 	// Used to keep GC references to in flight caching results (not yet stored to output and might not be in cache anymore)
 	TMap<FPCGTaskId, TUniquePtr<FCachedResult>> CollectGCCachingResults;
 
 	/** Lock level 4 */
-	UE::FSpinLock TaskOutputsLock;
+	FPCGLock TaskOutputsLock;
 	TMap<FPCGTaskId, FOutputDataInfo> TaskOutputs;
 	
 	/** Monotonically increasing id. Should be reset once all tasks are executed, should be protected by the ScheduleLock */
@@ -336,7 +337,7 @@ private:
 	};
 
 	/** Per graph execution cache, gets emptied when executor has no more work to do */
-	UE::FSpinLock GraphExecutionCachesLock;
+	FPCGLock GraphExecutionCachesLock;
 	TMap<FPCGTaskId, FGraphExecutionCache> GraphExecutionCaches;
 
 	std::atomic<bool> bNeedToExecuteTasksEnded = false;
@@ -350,6 +351,8 @@ private:
 	TObjectPtr<UWorld> World = nullptr;
 
 	TSharedPtr<FGameThreadHandler> GameThreadHandler;
+
+	double LastSchedulingErrorCheck = 0.0;
 };
 
 class FPCGFetchInputElement : public IPCGElement
