[PCG] Add APIs to PCG kernel to get input/output pins.

Use this to fix data description logging.

These APIs will be useful later when we support multi-kernels per node to wire up kernels.

#rnx

diff --git a/Engine/Plugins/PCG/Source/PCG/Internal/Compute/PCGDataForGPU.h b/Engine/Plugins/PCG/Source/PCG/Internal/Compute/PCGDataForGPU.h
--- a/Engine/Plugins/PCG/Source/PCG/Internal/Compute/PCGDataForGPU.h
+++ b/Engine/Plugins/PCG/Source/PCG/Internal/Compute/PCGDataForGPU.h
@@ -85,7 +85,7 @@ struct FPCGKernelAttributeTable
 	int32 AddAttribute(FName InName, EPCGKernelAttributeType InType);
 
 #if PCG_KERNEL_LOGGING_ENABLED
-	void DebugLog();
+	void DebugLog() const;
 #endif
 
 protected:
diff --git a/Engine/Plugins/PCG/Source/PCG/Private/Compute/Elements/PCGCustomHLSL.cpp b/Engine/Plugins/PCG/Source/PCG/Private/Compute/Elements/PCGCustomHLSL.cpp
--- a/Engine/Plugins/PCG/Source/PCG/Private/Compute/Elements/PCGCustomHLSL.cpp
+++ b/Engine/Plugins/PCG/Source/PCG/Private/Compute/Elements/PCGCustomHLSL.cpp
@@ -327,7 +327,7 @@ bool UPCGCustomHLSLSettings::IsShaderTextReadOnly() const
 UPCGComputeKernel* UPCGCustomHLSLSettings::CreateKernel(FPCGGPUCompilationContext& InOutContext, UObject* InObjectOuter) const
 {
 	UPCGCustomHLSLKernel* Kernel = InOutContext.NewObject_AnyThread<UPCGCustomHLSLKernel>(InObjectOuter);
-	Kernel->Settings = this;
+	Kernel->Initialize(this, bDumpDataDescriptions);
 
 	return Kernel;
 }
diff --git a/Engine/Plugins/PCG/Source/PCG/Private/Compute/Elements/PCGCustomHLSLKernel.cpp b/Engine/Plugins/PCG/Source/PCG/Private/Compute/Elements/PCGCustomHLSLKernel.cpp
--- a/Engine/Plugins/PCG/Source/PCG/Private/Compute/Elements/PCGCustomHLSLKernel.cpp
+++ b/Engine/Plugins/PCG/Source/PCG/Private/Compute/Elements/PCGCustomHLSLKernel.cpp
@@ -44,6 +44,13 @@ namespace PCGCustomHLSLKernel
 	}
 }
 
+void UPCGCustomHLSLKernel::Initialize(const UPCGCustomHLSLSettings* InSettings, bool bLogDescriptions)
+{
+	Super::Initialize(bLogDescriptions);
+
+	Settings = InSettings;
+}
+
 bool UPCGCustomHLSLKernel::IsKernelValid(FPCGContext* InContext, bool bQuiet) const
 {
 	TRACE_CPUPROFILER_EVENT_SCOPE(UPCGCustomHLSLKernel::IsKernelValid);
@@ -826,6 +833,18 @@ uint32 UPCGCustomHLSLKernel::GetElementCountMultiplier(FName InOutputPinLabel) c
 	}
 }
 
+void UPCGCustomHLSLKernel::GetInputPins(TArray<FPCGPinProperties>& OutPins) const
+{
+	check(Settings);
+	OutPins = Settings->InputPins;
+}
+
+void UPCGCustomHLSLKernel::GetOutputPins(TArray<FPCGPinPropertiesGPU>& OutPins) const
+{
+	check(Settings);
+	OutPins = Settings->OutputPins;
+}
+
 bool UPCGCustomHLSLKernel::IsThreadCountMultiplierInUse() const
 {
 	return Settings->KernelType == EPCGKernelType::Custom && Settings->DispatchThreadCount != EPCGDispatchThreadCount::Fixed;
diff --git a/Engine/Plugins/PCG/Source/PCG/Private/Compute/Elements/PCGCustomHLSLKernel.h b/Engine/Plugins/PCG/Source/PCG/Private/Compute/Elements/PCGCustomHLSLKernel.h
--- a/Engine/Plugins/PCG/Source/PCG/Private/Compute/Elements/PCGCustomHLSLKernel.h
+++ b/Engine/Plugins/PCG/Source/PCG/Private/Compute/Elements/PCGCustomHLSLKernel.h
@@ -16,6 +16,8 @@ class UPCGCustomHLSLKernel : public UPCGComputeKernel
 	GENERATED_BODY()
 
 public:
+	void Initialize(const UPCGCustomHLSLSettings* InSettings, bool bLogDescriptions);
+
 	virtual bool IsKernelValid(FPCGContext* InContext = nullptr, bool bQuiet = true) const override;
 	virtual bool ComputeOutputBindingDataDesc(const UPCGComputeGraph* InGraph, FName InOutputPinLabel, const UPCGDataBinding* InBinding, FPCGDataCollectionDesc& OutDataDesc) const override;
 #if WITH_EDITOR
@@ -26,6 +28,8 @@ public:
 	virtual void GetKernelAttributeKeys(TArray<FPCGKernelAttributeKey>& OutKeys) const override;
 	virtual uint32 GetThreadCountMultiplier() const override;
 	virtual uint32 GetElementCountMultiplier(FName InOutputPinLabel) const override;
+	virtual void GetInputPins(TArray<FPCGPinProperties>& OutPins) const override;
+	virtual void GetOutputPins(TArray<FPCGPinPropertiesGPU>& OutPins) const override;
 
 protected:
 	/** Will the ThreadCountMultiplier value be applied when calculating the dispatch thread count. */
diff --git a/Engine/Plugins/PCG/Source/PCG/Private/Compute/PCGComputeGraph.cpp b/Engine/Plugins/PCG/Source/PCG/Private/Compute/PCGComputeGraph.cpp
--- a/Engine/Plugins/PCG/Source/PCG/Private/Compute/PCGComputeGraph.cpp
+++ b/Engine/Plugins/PCG/Source/PCG/Private/Compute/PCGComputeGraph.cpp
@@ -8,6 +8,7 @@
 #include "Compute/PCGCompilerDiagnostic.h"
 #include "Compute/PCGComputeKernel.h"
 #include "Compute/PCGDataBinding.h"
+#include "Compute/PCGPinPropertiesGPU.h"
 #include "Metadata/PCGMetadata.h"
 
 #include "ComputeFramework/ComputeKernelCompileResult.h"
@@ -188,4 +189,72 @@ bool UPCGComputeGraph::IsGraphValid(FPCGContext* InContext, bool bQuiet) const
 	return bAllKernelsValid;
 }
 
+void UPCGComputeGraph::DebugLogDataDescriptions(const UPCGDataBinding* InBinding) const
+{
+#if PCG_KERNEL_LOGGING_ENABLED
+	const UEnum* PCGDataTypeEnum = StaticEnum<EPCGDataType>();
+	const UEnum* PCGKernelAttributeTypeEnum = StaticEnum<EPCGKernelAttributeType>();
+	check(PCGDataTypeEnum && PCGKernelAttributeTypeEnum);
+
+	auto LogDataDescription = [this, PCGDataTypeEnum, PCGKernelAttributeTypeEnum](const FPCGDataCollectionDesc& InDataDesc)
+	{
+		for (int32 DataIndex = 0; DataIndex < InDataDesc.DataDescs.Num(); ++DataIndex)
+		{
+			const FPCGDataDesc& DataDesc = InDataDesc.DataDescs[DataIndex];
+
+			UE_LOG(LogPCG, Display, TEXT("\t\tData Index %d"), DataIndex);
+			UE_LOG(LogPCG, Display, TEXT("\t\t\tType: %s"), *PCGDataTypeEnum->GetDisplayNameTextByValue(static_cast<int64>(DataDesc.Type)).ToString());
+			UE_LOG(LogPCG, Display, TEXT("\t\t\tElementCount: %d"), DataDesc.ElementCount);
+			UE_LOG(LogPCG, Display, TEXT("\t\t\tAttributes (%d)"), DataDesc.AttributeDescs.Num());
+
+			for (int32 AttributeIndex = 0; AttributeIndex < DataDesc.AttributeDescs.Num(); ++AttributeIndex)
+			{
+				const FPCGKernelAttributeDesc& AttributeDesc = DataDesc.AttributeDescs[AttributeIndex];
+
+				FString UniqueStringKeys;
+				for (int32 StringKey : AttributeDesc.GetUniqueStringKeys())
+				{
+					UniqueStringKeys += TEXT(", ");
+					UniqueStringKeys += FString::Format(TEXT("{0}"), { StringKey });
+				}
+
+				UE_LOG(LogPCG, Display, TEXT("\t\t\t\tID: %d\t\tName: %s\t\tType: %s\t\tUniqueStringKeys%s"),
+					AttributeDesc.Index,
+					*AttributeDesc.Name.ToString(),
+					*PCGKernelAttributeTypeEnum->GetDisplayNameTextByValue(static_cast<int64>(AttributeDesc.Type)).ToString(),
+					UniqueStringKeys.IsEmpty() ? TEXT(": ") : *UniqueStringKeys);
+			}
+		}
+	};
+
+	for (TObjectPtr<const UComputeKernel> Kernel : KernelInvocations)
+	{
+		if (const UPCGComputeKernel* PCGKernel = Cast<UPCGComputeKernel>(Kernel); PCGKernel && PCGKernel->GetLogDataDescriptions())
+		{
+			UE_LOG(LogPCG, Display, TEXT("Kernel: %s"), *PCGKernel->GetName());
+
+			TArray<FPCGPinProperties> InputPins;
+			PCGKernel->GetInputPins(InputPins);
+			for (const FPCGPinProperties& Pin : InputPins)
+			{
+				FPCGDataCollectionDesc DataDesc;
+				ensure(InBinding->GetKernelPinDataDesc_GameThread(PCGKernel, Pin.Label, true, DataDesc));
+				UE_LOG(LogPCG, Display, TEXT("\tInput Pin: %s (%d data)"), *Pin.Label.ToString(), DataDesc.DataDescs.Num());
+				LogDataDescription(DataDesc);
+			}
+
+			TArray<FPCGPinPropertiesGPU> OutputPins;
+			PCGKernel->GetOutputPins(OutputPins);
+			for (const FPCGPinPropertiesGPU& Pin : OutputPins)
+			{
+				FPCGDataCollectionDesc DataDesc;
+				ensure(InBinding->GetKernelPinDataDesc_GameThread(PCGKernel, Pin.Label, false, DataDesc));
+				UE_LOG(LogPCG, Display, TEXT("\tOutput Pin: %s (%d data)"), *Pin.Label.ToString(), DataDesc.DataDescs.Num());
+				LogDataDescription(DataDesc);
+			}
+		}
+	}
+#endif
+}
+
 #undef PCG_DATA_DESCRIPTION_LOGGING 
diff --git a/Engine/Plugins/PCG/Source/PCG/Private/Compute/PCGComputeGraph.h b/Engine/Plugins/PCG/Source/PCG/Private/Compute/PCGComputeGraph.h
--- a/Engine/Plugins/PCG/Source/PCG/Private/Compute/PCGComputeGraph.h
+++ b/Engine/Plugins/PCG/Source/PCG/Private/Compute/PCGComputeGraph.h
@@ -99,6 +99,8 @@ public:
 
 	bool IsGraphValid(FPCGContext* InContext = nullptr, bool bQuiet = true) const;
 
+	void DebugLogDataDescriptions(const UPCGDataBinding* InBinding) const;
+
 public:
 	TMap<TObjectKey<const UPCGNode>, TArray<FComputeKernelCompileMessage>> KernelToCompileMessages;
 	
diff --git a/Engine/Plugins/PCG/Source/PCG/Private/Compute/PCGComputeKernel.cpp b/Engine/Plugins/PCG/Source/PCG/Private/Compute/PCGComputeKernel.cpp
--- a/Engine/Plugins/PCG/Source/PCG/Private/Compute/PCGComputeKernel.cpp
+++ b/Engine/Plugins/PCG/Source/PCG/Private/Compute/PCGComputeKernel.cpp
@@ -20,6 +20,11 @@ namespace PCGComputeKernel
 		TEXT("Enables warnings when parameters are overidden on GPU nodes."));
 }
 
+void UPCGComputeKernel::Initialize(bool bLogDescriptions)
+{
+	bLogDataDescriptions = bLogDescriptions;
+}
+
 void UPCGComputeKernel::ComputeDataDescFromPinProperties(const FPCGPinPropertiesGPU& OutputPinProps, const TArrayView<const FPCGPinProperties>& InInputPinProps, const UPCGDataBinding* InBinding, FPCGDataCollectionDesc& OutPinDesc) const
 {
 	check(InBinding);
diff --git a/Engine/Plugins/PCG/Source/PCG/Private/Compute/PCGComputeKernel.h b/Engine/Plugins/PCG/Source/PCG/Private/Compute/PCGComputeKernel.h
--- a/Engine/Plugins/PCG/Source/PCG/Private/Compute/PCGComputeKernel.h
+++ b/Engine/Plugins/PCG/Source/PCG/Private/Compute/PCGComputeKernel.h
@@ -25,6 +25,8 @@ class UPCGComputeKernel : public UComputeKernel
 	GENERATED_BODY()
 
 public:
+	void Initialize(bool bLogDescriptions);
+
 	/** Get/set index of kernel in kernels array on owning compute graph. */
 	int32 GetKernelIndex() const { return KernelIndex; }
 	void SetKernelIndex(int32 InKernelIndex) { KernelIndex = InKernelIndex; }
@@ -76,8 +78,16 @@ public:
 	/** Do any validation of the PCG node. Returns true if node is valid and GPU execution can proceed. */
 	bool ValidatePCGNode(FPCGContext* InContext, bool bQuiet, const UPCGSettings* InSettings) const;
 
+	virtual void GetInputPins(TArray<FPCGPinProperties>& OutPins) const {}
+	virtual void GetOutputPins(TArray<FPCGPinPropertiesGPU>& OutPins) const {}
+
+	virtual bool GetLogDataDescriptions() const { return bLogDataDescriptions; }
+
 protected:
 	/** Index into kernels array on owning compute graph. */
 	UPROPERTY()
 	int32 KernelIndex = INDEX_NONE;
+
+	UPROPERTY()
+	bool bLogDataDescriptions = false;
 };
diff --git a/Engine/Plugins/PCG/Source/PCG/Private/Compute/PCGDataBinding.cpp b/Engine/Plugins/PCG/Source/PCG/Private/Compute/PCGDataBinding.cpp
--- a/Engine/Plugins/PCG/Source/PCG/Private/Compute/PCGDataBinding.cpp
+++ b/Engine/Plugins/PCG/Source/PCG/Private/Compute/PCGDataBinding.cpp
@@ -298,13 +298,9 @@ void UPCGDataBinding::BuildStringTable()
 	}
 }
 
-void UPCGDataBinding::DebugLogDataDescriptions()
+void UPCGDataBinding::DebugLogDataDescriptions() const
 {
 #if PCG_KERNEL_LOGGING_ENABLED
-	const UEnum* PCGDataTypeEnum = StaticEnum<EPCGDataType>();
-	const UEnum* PCGKernelAttributeTypeEnum = StaticEnum<EPCGKernelAttributeType>();
-	check(PCGDataTypeEnum && PCGKernelAttributeTypeEnum);
-
 	UE_LOG(LogPCG, Display, TEXT("\n### METADATA ATTRIBUTE TABLE ###"));
 	AttributeTable.DebugLog();
 
@@ -314,67 +310,8 @@ void UPCGDataBinding::DebugLogDataDescriptions()
 		UE_LOG(LogPCG, Display, TEXT("\t%d: %s"), i, *StringTable[i]);
 	}
 
-	auto LogDataDescription = [this, PCGDataTypeEnum, PCGKernelAttributeTypeEnum](const FPCGDataCollectionDesc& InDataDesc)
-	{
-		for (int32 DataIndex = 0; DataIndex < InDataDesc.DataDescs.Num(); ++DataIndex)
-		{
-			const FPCGDataDesc& DataDesc = InDataDesc.DataDescs[DataIndex];
-
-			UE_LOG(LogPCG, Display, TEXT("\t\tData Index %d"), DataIndex);
-			UE_LOG(LogPCG, Display, TEXT("\t\t\tType: %s"), *PCGDataTypeEnum->GetDisplayNameTextByValue(static_cast<int64>(DataDesc.Type)).ToString());
-			UE_LOG(LogPCG, Display, TEXT("\t\t\tElementCount: %d"), DataDesc.ElementCount);
-			UE_LOG(LogPCG, Display, TEXT("\t\t\tAttributes (%d)"), DataDesc.AttributeDescs.Num());
-				
-			for (int32 AttributeIndex = 0; AttributeIndex < DataDesc.AttributeDescs.Num(); ++AttributeIndex)
-			{
-				const FPCGKernelAttributeDesc& AttributeDesc = DataDesc.AttributeDescs[AttributeIndex];
-
-				FString UniqueStringKeys;
-				for (int32 StringKey : AttributeDesc.GetUniqueStringKeys())
-				{
-					UniqueStringKeys += TEXT(", ");
-					UniqueStringKeys += FString::Format(TEXT("{0}"), { StringKey });
-				}
-
-				UE_LOG(LogPCG, Display, TEXT("\t\t\t\tID: %d\t\tName: %s\t\tType: %s\t\tUniqueStringKeys%s"),
-					AttributeDesc.Index,
-					*AttributeDesc.Name.ToString(),
-					*PCGKernelAttributeTypeEnum->GetDisplayNameTextByValue(static_cast<int64>(AttributeDesc.Type)).ToString(),
-					UniqueStringKeys.IsEmpty() ? TEXT(": ") : *UniqueStringKeys);
-			}
-		}
-	};
-
 	// TODO - we don't have explicit pin lists for kernels because we don't need them. Decide how to bring this back or pivot the feature.
-	//UE_LOG(LogPCG, Display, TEXT("\n### INPUT PIN DATA DESCRIPTIONS ###"));
-	//for (const TSoftObjectPtr<const UPCGNode>& Node : Graph->KernelToNode)
-	//{
-	//	const UPCGSettings* Settings = Node.IsValid() ? Node->GetSettings() : nullptr;
-	//	if (Settings && Settings->bDumpDataDescriptions)
-	//	{
-	//		UE_LOG(LogPCG, Display, TEXT("Node: %s"), *Node->GetNodeTitle(EPCGNodeTitleType::ListView).ToString());
-
-	//		for (const UPCGPin* Pin : Node->GetInputPins())
-	//		{
-	//			const FPCGDataCollectionDesc DataDesc = PCGDataForGPUHelpers::ComputeInputPinDataDesc(Pin, this);
-	//			UE_LOG(LogPCG, Display, TEXT("\tInput Pin: %s (%d data)"), *Pin->Properties.Label.ToString(), DataDesc.DataDescs.Num());
-	//			LogDataDescription(DataDesc);
-	//		}
-
-	//		for (const UPCGPin* Pin : Node->GetOutputPins())
-	//		{
-	//			FPCGDataCollectionDesc DataDesc;
-	//			if (ensure(GetOutputPinDataDesc_GameThread(Pin, Settings, DataDesc)))
-	//			{
-	//				UE_LOG(LogPCG, Display, TEXT("\tOutput Pin: %s (%d data)"), *Pin->Properties.Label.ToString(), DataDesc.DataDescs.Num());
-	//				LogDataDescription(DataDesc);
-	//			}
-	//			else
-	//			{
-	//				UE_LOG(LogPCG, Display, TEXT("\tOutput Pin: %s MISSING!"), *Pin->Properties.Label.ToString());
-	//			}
-	//		}
-	//	}
-	//}
+	UE_LOG(LogPCG, Display, TEXT("\n### INPUT PIN DATA DESCRIPTIONS ###"));
+	Graph->DebugLogDataDescriptions(this);
 #endif
 }
diff --git a/Engine/Plugins/PCG/Source/PCG/Private/Compute/PCGDataBinding.h b/Engine/Plugins/PCG/Source/PCG/Private/Compute/PCGDataBinding.h
--- a/Engine/Plugins/PCG/Source/PCG/Private/Compute/PCGDataBinding.h
+++ b/Engine/Plugins/PCG/Source/PCG/Private/Compute/PCGDataBinding.h
@@ -90,7 +90,7 @@ private:
 
 	void BuildStringTable();
 
-	void DebugLogDataDescriptions();
+	void DebugLogDataDescriptions() const;
 
 public:
 	UPROPERTY()
diff --git a/Engine/Plugins/PCG/Source/PCG/Private/Compute/PCGDataForGPU.cpp b/Engine/Plugins/PCG/Source/PCG/Private/Compute/PCGDataForGPU.cpp
--- a/Engine/Plugins/PCG/Source/PCG/Private/Compute/PCGDataForGPU.cpp
+++ b/Engine/Plugins/PCG/Source/PCG/Private/Compute/PCGDataForGPU.cpp
@@ -630,7 +630,7 @@ int32 FPCGKernelAttributeTable::AddAttribute(FName InName, EPCGKernelAttributeTy
 }
 
 #if PCG_KERNEL_LOGGING_ENABLED
-void FPCGKernelAttributeTable::DebugLog()
+void FPCGKernelAttributeTable::DebugLog() const
 {
 	const UEnum* PCGKernelAttributeTypeEnum = StaticEnum<EPCGKernelAttributeType>();
 
diff --git a/Engine/Plugins/PCG/Source/PCG/Private/Elements/Metadata/PCGMetadataPartition.cpp b/Engine/Plugins/PCG/Source/PCG/Private/Elements/Metadata/PCGMetadataPartition.cpp
--- a/Engine/Plugins/PCG/Source/PCG/Private/Elements/Metadata/PCGMetadataPartition.cpp
+++ b/Engine/Plugins/PCG/Source/PCG/Private/Elements/Metadata/PCGMetadataPartition.cpp
@@ -71,7 +71,7 @@ void UPCGMetadataPartitionSettings::ApplyDeprecation(UPCGNode* InOutNode)
 UPCGComputeKernel* UPCGMetadataPartitionSettings::CreateKernel(FPCGGPUCompilationContext& InOutContext, UObject* InObjectOuter) const
 {
 	UPCGMetadataPartitionKernel* Kernel = InOutContext.NewObject_AnyThread<UPCGMetadataPartitionKernel>(InObjectOuter);
-	Kernel->Settings = this;
+	Kernel->Initialize(this, bDumpDataDescriptions);
 
 	return Kernel;
 }
diff --git a/Engine/Plugins/PCG/Source/PCG/Private/Elements/Metadata/PCGMetadataPartitionKernel.cpp b/Engine/Plugins/PCG/Source/PCG/Private/Elements/Metadata/PCGMetadataPartitionKernel.cpp
--- a/Engine/Plugins/PCG/Source/PCG/Private/Elements/Metadata/PCGMetadataPartitionKernel.cpp
+++ b/Engine/Plugins/PCG/Source/PCG/Private/Elements/Metadata/PCGMetadataPartitionKernel.cpp
@@ -6,6 +6,7 @@
 #include "PCGModule.h"
 #include "Compute/PCGComputeGraph.h"
 #include "Compute/PCGDataBinding.h"
+#include "Compute/PCGPinPropertiesGPU.h"
 #include "Compute/DataInterfaces/Elements/PCGMetadataPartitionDataInterface.h"
 #include "Compute/Elements/PCGComputeGraphElement.h"
 #include "Elements/Metadata/PCGMetadataPartition.h"
@@ -17,6 +18,13 @@
 
 #define LOCTEXT_NAMESPACE "PCGMetadataPartitionKernel"
 
+void UPCGMetadataPartitionKernel::Initialize(const UPCGMetadataPartitionSettings* InSettings, bool bLogDescriptions)
+{
+	Super::Initialize(bLogDescriptions);
+
+	Settings = InSettings;
+}
+
 bool UPCGMetadataPartitionKernel::IsKernelValid(FPCGContext* InContext, bool bQuiet) const
 {
 	TRACE_CPUPROFILER_EVENT_SCOPE(UPCGMetadataPartitionKernel::IsKernelValid);
@@ -181,4 +189,14 @@ void UPCGMetadataPartitionKernel::GetKernelAttributeKeys(TArray<FPCGKernelAttrib
 	}
 }
 
+void UPCGMetadataPartitionKernel::GetInputPins(TArray<FPCGPinProperties>& OutPins) const
+{
+	OutPins.Emplace(PCGPinConstants::DefaultInputLabel, EPCGDataType::Point);
+}
+
+void UPCGMetadataPartitionKernel::GetOutputPins(TArray<FPCGPinPropertiesGPU>& OutPins) const
+{
+	OutPins.Emplace(PCGPinConstants::DefaultOutputLabel, EPCGDataType::Point);
+}
+
 #undef LOCTEXT_NAMESPACE
diff --git a/Engine/Plugins/PCG/Source/PCG/Private/Elements/PCGCopyPoints.cpp b/Engine/Plugins/PCG/Source/PCG/Private/Elements/PCGCopyPoints.cpp
--- a/Engine/Plugins/PCG/Source/PCG/Private/Elements/PCGCopyPoints.cpp
+++ b/Engine/Plugins/PCG/Source/PCG/Private/Elements/PCGCopyPoints.cpp
@@ -25,7 +25,7 @@ FText UPCGCopyPointsSettings::GetNodeTooltipText() const
 UPCGComputeKernel* UPCGCopyPointsSettings::CreateKernel(FPCGGPUCompilationContext& InOutContext, UObject* InObjectOuter) const
 {
 	UPCGCopyPointsKernel* Kernel = InOutContext.NewObject_AnyThread<UPCGCopyPointsKernel>(InObjectOuter);
-	Kernel->Settings = this;
+	Kernel->Initialize(this, bDumpDataDescriptions);
 
 	return Kernel;
 }
diff --git a/Engine/Plugins/PCG/Source/PCG/Private/Elements/PCGCopyPointsKernel.cpp b/Engine/Plugins/PCG/Source/PCG/Private/Elements/PCGCopyPointsKernel.cpp
--- a/Engine/Plugins/PCG/Source/PCG/Private/Elements/PCGCopyPointsKernel.cpp
+++ b/Engine/Plugins/PCG/Source/PCG/Private/Elements/PCGCopyPointsKernel.cpp
@@ -18,6 +18,13 @@
 
 #define LOCTEXT_NAMESPACE "PCGCopyPointsKernel"
 
+void UPCGCopyPointsKernel::Initialize(const UPCGCopyPointsSettings* InSettings, bool bLogDescriptions)
+{
+	Super::Initialize(bLogDescriptions);
+
+	Settings = InSettings;
+}
+
 bool UPCGCopyPointsKernel::IsKernelValid(FPCGContext* InContext, bool bQuiet) const
 {
 	TRACE_CPUPROFILER_EVENT_SCOPE(UPCGCopyPointsKernel::IsKernelValid);
@@ -132,4 +139,15 @@ void UPCGCopyPointsKernel::CreateAdditionalInputDataInterfaces(FPCGGPUCompilatio
 }
 #endif
 
+void UPCGCopyPointsKernel::GetInputPins(TArray<FPCGPinProperties>& OutPins) const
+{
+	OutPins.Emplace(PCGCopyPointsConstants::SourcePointsLabel, EPCGDataType::Point);
+	OutPins.Emplace(PCGCopyPointsConstants::TargetPointsLabel, EPCGDataType::Point);
+}
+
+void UPCGCopyPointsKernel::GetOutputPins(TArray<FPCGPinPropertiesGPU>& OutPins) const
+{
+	OutPins.Emplace(PCGPinConstants::DefaultOutputLabel, EPCGDataType::Point);
+}
+
 #undef LOCTEXT_NAMESPACE
diff --git a/Engine/Plugins/PCG/Source/PCG/Private/Elements/PCGStaticMeshSpawner.cpp b/Engine/Plugins/PCG/Source/PCG/Private/Elements/PCGStaticMeshSpawner.cpp
--- a/Engine/Plugins/PCG/Source/PCG/Private/Elements/PCGStaticMeshSpawner.cpp
+++ b/Engine/Plugins/PCG/Source/PCG/Private/Elements/PCGStaticMeshSpawner.cpp
@@ -48,7 +48,7 @@ UPCGStaticMeshSpawnerSettings::UPCGStaticMeshSpawnerSettings(const FObjectInitia
 UPCGComputeKernel* UPCGStaticMeshSpawnerSettings::CreateKernel(FPCGGPUCompilationContext& InOutContext, UObject* InObjectOuter) const
 {
 	UPCGStaticMeshSpawnerKernel* Kernel = InOutContext.NewObject_AnyThread<UPCGStaticMeshSpawnerKernel>(InObjectOuter);
-	Kernel->Settings = this;
+	Kernel->Initialize(this, bDumpDataDescriptions);
 
 	return Kernel;
 }
diff --git a/Engine/Plugins/PCG/Source/PCG/Private/Elements/PCGStaticMeshSpawnerKernel.cpp b/Engine/Plugins/PCG/Source/PCG/Private/Elements/PCGStaticMeshSpawnerKernel.cpp
--- a/Engine/Plugins/PCG/Source/PCG/Private/Elements/PCGStaticMeshSpawnerKernel.cpp
+++ b/Engine/Plugins/PCG/Source/PCG/Private/Elements/PCGStaticMeshSpawnerKernel.cpp
@@ -8,6 +8,7 @@
 #include "Components/PCGProceduralISMComponent.h"
 #include "Compute/PCGComputeGraph.h"
 #include "Compute/PCGDataBinding.h"
+#include "Compute/PCGPinPropertiesGPU.h"
 #include "Compute/DataInterfaces/PCGInstanceDataInterface.h"
 #include "Compute/DataInterfaces/Elements/PCGStaticMeshSpawnerDataInterface.h"
 #include "Elements/PCGStaticMeshSpawner.h"
@@ -30,6 +31,13 @@ namespace PCGStaticMeshSpawnerKernel
 	static const FText NoMeshEntriesFormat = LOCTEXT("NoMeshEntries", "No mesh entries provided in weighted mesh selector.");
 }
 
+void UPCGStaticMeshSpawnerKernel::Initialize(const UPCGStaticMeshSpawnerSettings* InSettings, bool bLogDescriptions)
+{
+	Super::Initialize(bLogDescriptions);
+
+	Settings = InSettings;
+}
+
 bool UPCGStaticMeshSpawnerKernel::IsKernelValid(FPCGContext* InContext, bool bQuiet) const
 {
 	TRACE_CPUPROFILER_EVENT_SCOPE(UPCGStaticMeshSpawnerKernel::IsKernelValid);
@@ -198,6 +206,16 @@ void UPCGStaticMeshSpawnerKernel::GetKernelAttributeKeys(TArray<FPCGKernelAttrib
 	OutKeys.AddUnique(FPCGKernelAttributeKey(Settings->OutAttributeName, EPCGKernelAttributeType::StringKey));
 }
 
+void UPCGStaticMeshSpawnerKernel::GetInputPins(TArray<FPCGPinProperties>& OutPins) const
+{
+	OutPins.Emplace(PCGPinConstants::DefaultInputLabel, EPCGDataType::Point);
+}
+
+void UPCGStaticMeshSpawnerKernel::GetOutputPins(TArray<FPCGPinPropertiesGPU>& OutPins) const
+{
+	OutPins.Emplace(PCGPinConstants::DefaultOutputLabel, EPCGDataType::Point);
+}
+
 bool UPCGStaticMeshSpawnerKernel::SetupProceduralISMComponents(FPCGContext* InContext, UPCGDataBinding* InBinding) const
 {
 	AActor* TargetActor = Settings->TargetActor.Get() ? Settings->TargetActor.Get() : InContext->GetTargetActor(nullptr);
diff --git a/Engine/Plugins/PCG/Source/PCG/Public/Elements/Metadata/PCGMetadataPartitionKernel.h b/Engine/Plugins/PCG/Source/PCG/Public/Elements/Metadata/PCGMetadataPartitionKernel.h
--- a/Engine/Plugins/PCG/Source/PCG/Public/Elements/Metadata/PCGMetadataPartitionKernel.h
+++ b/Engine/Plugins/PCG/Source/PCG/Public/Elements/Metadata/PCGMetadataPartitionKernel.h
@@ -14,6 +14,8 @@ class UPCGMetadataPartitionKernel : public UPCGComputeKernel
 	GENERATED_BODY()
 
 public:
+	void Initialize(const UPCGMetadataPartitionSettings* InSettings, bool bLogDescriptions);
+
 	virtual bool IsKernelValid(FPCGContext* InContext = nullptr, bool bQuiet = true) const override;
 	virtual bool ComputeOutputBindingDataDesc(const UPCGComputeGraph* InGraph, FName InOutputPinLabel, const UPCGDataBinding* InBinding, FPCGDataCollectionDesc& OutDataDesc) const override;
 	virtual int ComputeThreadCount(const UPCGDataBinding* Binding) const override;
@@ -23,6 +25,8 @@ public:
 #endif
 	virtual void GetKernelAttributeKeys(TArray<FPCGKernelAttributeKey>& OutKeys) const override;
 	virtual bool DoesOutputPinRequireElementCounters(FName InOutputPinLabel) const override { return true; }
+	virtual void GetInputPins(TArray<FPCGPinProperties>& OutPins) const override;
+	virtual void GetOutputPins(TArray<FPCGPinPropertiesGPU>& OutPins) const override;
 
 public:
 	UPROPERTY()
diff --git a/Engine/Plugins/PCG/Source/PCG/Public/Elements/PCGCopyPointsKernel.h b/Engine/Plugins/PCG/Source/PCG/Public/Elements/PCGCopyPointsKernel.h
--- a/Engine/Plugins/PCG/Source/PCG/Public/Elements/PCGCopyPointsKernel.h
+++ b/Engine/Plugins/PCG/Source/PCG/Public/Elements/PCGCopyPointsKernel.h
@@ -16,6 +16,8 @@ class UPCGCopyPointsKernel : public UPCGComputeKernel
 	GENERATED_BODY()
 
 public:
+	void Initialize(const UPCGCopyPointsSettings* InSettings, bool bLogDescriptions);
+
 	virtual bool IsKernelValid(FPCGContext* InContext = nullptr, bool bQuiet = true) const override;
 	virtual bool ComputeOutputBindingDataDesc(const UPCGComputeGraph* InGraph, FName InOutputPinLabel, const UPCGDataBinding* InBinding, FPCGDataCollectionDesc& OutDataDesc) const override;
 	virtual int ComputeThreadCount(const UPCGDataBinding* Binding) const override;
@@ -23,6 +25,8 @@ public:
 	virtual FString GetCookedSource(const UPCGComputeGraph* InComputeGraph) const override;
 	void CreateAdditionalInputDataInterfaces(FPCGGPUCompilationContext& InOutContext, UObject* InObjectOuter, TArray<TObjectPtr<UComputeDataInterface>>& OutDataInterfaces) const override;
 #endif
+	virtual void GetInputPins(TArray<FPCGPinProperties>& OutPins) const override;
+	virtual void GetOutputPins(TArray<FPCGPinPropertiesGPU>& OutPins) const override;
 
 public:
 	UPROPERTY()
diff --git a/Engine/Plugins/PCG/Source/PCG/Public/Elements/PCGStaticMeshSpawnerKernel.h b/Engine/Plugins/PCG/Source/PCG/Public/Elements/PCGStaticMeshSpawnerKernel.h
--- a/Engine/Plugins/PCG/Source/PCG/Public/Elements/PCGStaticMeshSpawnerKernel.h
+++ b/Engine/Plugins/PCG/Source/PCG/Public/Elements/PCGStaticMeshSpawnerKernel.h
@@ -14,6 +14,8 @@ class UPCGStaticMeshSpawnerKernel : public UPCGComputeKernel
 	GENERATED_BODY()
 
 public:
+	void Initialize(const UPCGStaticMeshSpawnerSettings* InSettings, bool bLogDescriptions);
+
 	virtual bool IsKernelValid(FPCGContext* InContext = nullptr, bool bQuiet = true) const override;
 	virtual bool ComputeOutputBindingDataDesc(const UPCGComputeGraph* InGraph, FName InOutputPinLabel, const UPCGDataBinding* InBinding, FPCGDataCollectionDesc& OutDataDesc) const override;
 	virtual int ComputeThreadCount(const UPCGDataBinding* Binding) const override;
@@ -25,6 +27,9 @@ public:
 	virtual void AddStaticCreatedStrings(TArray<FString>& InOutStringTable) const override;
 	virtual void GetKernelAttributeKeys(TArray<FPCGKernelAttributeKey>& OutKeys) const override;
 
+	virtual void GetInputPins(TArray<FPCGPinProperties>& OutPins) const override;
+	virtual void GetOutputPins(TArray<FPCGPinPropertiesGPU>& OutPins) const override;
+
 	/** Setup Procedural ISMCs. Returns true if at least one component was set up. */
 	virtual bool SetupProceduralISMComponents(FPCGContext* InContext, UPCGDataBinding* InBinding) const;
 
