Made the set debug object from selection expand to first instance instead of just working for the top level graph.

#rnx

diff --git a/Engine/Plugins/PCG/Source/PCGEditor/Private/SPCGEditorGraphDebugObjectTree.cpp b/Engine/Plugins/PCG/Source/PCGEditor/Private/SPCGEditorGraphDebugObjectTree.cpp
--- a/Engine/Plugins/PCG/Source/PCGEditor/Private/SPCGEditorGraphDebugObjectTree.cpp
+++ b/Engine/Plugins/PCG/Source/PCGEditor/Private/SPCGEditorGraphDebugObjectTree.cpp
@@ -625,13 +625,22 @@ void SPCGEditorGraphDebugObjectTree::SetDebugObjectFromSelection_OnClicked()
 				continue;
 			}
 
+			bool bSelectedDebugObject = false;
+
 			for (const FPCGStack& Stack : StackContext.GetStacks())
 			{
-				if (SetDebugObjectFromStack(Stack))
+				if (FPCGEditorGraphDebugObjectItemPtr Item = GetItemFromStack(Stack))
 				{
+					ExpandAndSelectFirstLeafDebugObject(Item);
+					bSelectedDebugObject = true;
 					break;
 				}
 			}
+
+			if (bSelectedDebugObject)
+			{
+				break;
+			}
 		}
 	}
 }
@@ -641,38 +650,29 @@ bool SPCGEditorGraphDebugObjectTree::SetDebugObjectFromStackFromAnotherEditor(co
 	RefreshTree();
 
 	const UPCGGraph* CurrentGraph = PCGEditor.Pin()->GetPCGGraph();
+	// If the given stack is already valid for this editor, try to select it as-is.
 	if (InStack.GetGraphForCurrentFrame() == CurrentGraph)
 	{
-		return SetDebugObjectFromStack(InStack);
-	}
-	else
-	{
-		// Find first instance that starts with the given stack but has the current graph at its end.
-		for (const FPCGEditorGraphDebugObjectItemPtr& Item : AllGraphItems)
+		if (FPCGEditorGraphDebugObjectItemPtr Item = GetItemFromStack(InStack))
 		{
-			const FPCGStack* ItemStack = Item->GetPCGStack();
-
-			if (ItemStack && ItemStack->BeginsWith(InStack) && ItemStack->GetGraphForCurrentFrame() == CurrentGraph)
-			{
-				return SetDebugObjectFromStack(*ItemStack);
-			}
+			ExpandAndSelectDebugObject(Item);
+			return true;
 		}
-
-		return false;
 	}
-}
 
-bool SPCGEditorGraphDebugObjectTree::SetDebugObjectFromStack(const FPCGStack& InStack)
-{
-	if (FPCGEditorGraphDebugObjectItemPtr Item = GetItemFromStack(InStack))
-	{
-		ExpandAndSelectDebugObject(Item);
-		return true;
-	}
-	else
+	// Find first instance that starts with the given stack but has the current graph at its end.
+	for (const FPCGEditorGraphDebugObjectItemPtr& Item : AllGraphItems)
 	{
-		return false;
+		const FPCGStack* ItemStack = Item->GetPCGStack();
+
+		if (ItemStack && ItemStack->BeginsWith(InStack) && ItemStack->GetGraphForCurrentFrame() == CurrentGraph)
+		{
+			ExpandAndSelectDebugObject(Item);
+			return true;
+		}
 	}
+
+	return false;
 }
 
 void SPCGEditorGraphDebugObjectTree::ExpandAndSelectDebugObject(const FPCGEditorGraphDebugObjectItemPtr& InItem)
@@ -1466,15 +1466,8 @@ void SPCGEditorGraphDebugObjectTree::ExpandAndSelectFirstLeafDebugObject(const F
 		}
 	}
 
-	if (ItemToSelect.IsValid())
-	{
-		ExpandAndSelectDebugObject(ItemToSelect);
-	}
-	else
-	{
-		// Nothing found, just reselect the original item
-		DebugObjectTreeView->SetSelection(InItem);
-	}
+	// If we've found nothing, we'll select the original item, otherwise we'll pick the first occurrence we found.
+	ExpandAndSelectDebugObject(ItemToSelect.IsValid() ? ItemToSelect : InItem);
 }
 
 TSharedPtr<SWidget> SPCGEditorGraphDebugObjectTree::OpenContextMenu()
diff --git a/Engine/Plugins/PCG/Source/PCGEditor/Private/SPCGEditorGraphDebugObjectTree.h b/Engine/Plugins/PCG/Source/PCGEditor/Private/SPCGEditorGraphDebugObjectTree.h
--- a/Engine/Plugins/PCG/Source/PCGEditor/Private/SPCGEditorGraphDebugObjectTree.h
+++ b/Engine/Plugins/PCG/Source/PCGEditor/Private/SPCGEditorGraphDebugObjectTree.h
@@ -309,8 +309,6 @@ private:
 
 	/** Returns the matching item from the AllGraphItems array if any. */
 	FPCGEditorGraphDebugObjectItemPtr GetItemFromStack(const FPCGStack& InStack) const;
-	/** If the stack matches to an item, Expands the tree view to make that item visible and select it, and return true. Otherwise, return false. */
-	bool SetDebugObjectFromStack(const FPCGStack& InStack);
 	/** Expands the tree view to make an item visible and select it. */
 	void ExpandAndSelectDebugObject(const FPCGEditorGraphDebugObjectItemPtr& InItem);
 
