[PCG] Add support for registering selectors in the accessor factory

* Read by the Attribute Property Selector details view
* Implemented for Points and Spline Control Points

#rnx

diff --git a/Engine/Plugins/PCG/Source/PCG/Private/Metadata/Accessors/PCGAttributeAccessorFactory.cpp b/Engine/Plugins/PCG/Source/PCG/Private/Metadata/Accessors/PCGAttributeAccessorFactory.cpp
--- a/Engine/Plugins/PCG/Source/PCG/Private/Metadata/Accessors/PCGAttributeAccessorFactory.cpp
+++ b/Engine/Plugins/PCG/Source/PCG/Private/Metadata/Accessors/PCGAttributeAccessorFactory.cpp
@@ -132,6 +132,22 @@ void FPCGAttributeAccessorMethods::FillSelectorMenuEntryFromEnum(const UEnum* En
 		return;
 	}
 
+	FPCGAttributeSelectorMenu* CurrentMenu = &AttributeSelectorMenu;
+	for (const FText& MenuLabel : InMenuHierarchy)
+	{
+		FPCGAttributeSelectorMenu* It = CurrentMenu->SubMenus.FindByPredicate([&MenuLabel](const FPCGAttributeSelectorMenu& Menu) { return Menu.Label.EqualTo(MenuLabel); });
+		if (!It)
+		{
+			It = &CurrentMenu->SubMenus.Emplace_GetRef();
+			It->Label = MenuLabel;
+		}
+
+		check(It);
+		CurrentMenu = It;
+	}
+
+	check(CurrentMenu);
+
 	const int32 NumEnums = EnumType->ContainsExistingMax() ? EnumType->NumEnums() - 1 : EnumType->NumEnums();
 	for (int32 i = 0; i < NumEnums; ++i)
 	{
@@ -143,10 +159,9 @@ void FPCGAttributeAccessorMethods::FillSelectorMenuEntryFromEnum(const UEnum* En
 		// Not const to be able to move them.
 		FText EnumName = EnumType->GetDisplayNameTextByIndex(i);
 		FText Tooltip = EnumType->GetToolTipTextByIndex(i);
-		const FName InvariantName = EnumType->GetNameByIndex(i);
+		const FString InvariantName = EnumType->GetNameStringByIndex(i); // Use string version as it strips out the namespace.
 
-		// Can't move the menu hierarchy as it will be copied for each entry.
-		AttributeSelectorOptions.Emplace(std::move(EnumName), std::move(Tooltip), FPCGAttributePropertySelector::CreatePropertySelector(InvariantName), InMenuHierarchy);
+		CurrentMenu->Entries.Emplace(std::move(EnumName), std::move(Tooltip), FPCGAttributePropertySelector::CreatePropertySelector(*InvariantName));
 	}
 }
 #endif // WITH_EDITOR
@@ -218,14 +233,11 @@ TUniquePtr<const IPCGAttributeAccessorKeys> FPCGAttributeAccessorFactory::Create
 }
 
 #if WITH_EDITOR
-void FPCGAttributeAccessorFactory::ForEachSelectorMenuEntry(TFunctionRef<void(const FPCGAttributeSelectorMenuEntry&)> Callback) const
+void FPCGAttributeAccessorFactory::ForEachSelectorMenu(TFunctionRef<void(const FPCGAttributeSelectorMenu&)> Callback) const
 {
 	for (const auto& [Class, Methods] : AccessorMethods)
 	{
-		for (const FPCGAttributeSelectorMenuEntry& SelectorMenuEntry : Methods.AttributeSelectorOptions)
-		{
-			Callback(SelectorMenuEntry);
-		}
+		Callback(Methods.AttributeSelectorMenu);
 	}
 }
 #endif // WITH_EDITOR
diff --git a/Engine/Plugins/PCG/Source/PCG/Public/Metadata/Accessors/PCGAttributeAccessorFactory.h b/Engine/Plugins/PCG/Source/PCG/Public/Metadata/Accessors/PCGAttributeAccessorFactory.h
--- a/Engine/Plugins/PCG/Source/PCG/Public/Metadata/Accessors/PCGAttributeAccessorFactory.h
+++ b/Engine/Plugins/PCG/Source/PCG/Public/Metadata/Accessors/PCGAttributeAccessorFactory.h
@@ -17,26 +17,33 @@ class UPCGData;
 class UPCGMetadata;
 
 #if WITH_EDITOR
+struct FPCGAttributeSelectorMenuEntry;
+
+struct FPCGAttributeSelectorMenu
+{
+	FText Label;
+	FText Tooltip;
+	
+	TArray<FPCGAttributeSelectorMenu> SubMenus;
+	TArray<FPCGAttributeSelectorMenuEntry> Entries;
+};
+
 struct FPCGAttributeSelectorMenuEntry
 {
-	FPCGAttributeSelectorMenuEntry(FText InDisplayName, FText InTooltip, FPCGAttributePropertySelector InSelector, TArray<FText> InMenuHierarchy = {})
-		: DisplayName(std::move(InDisplayName))
+	FPCGAttributeSelectorMenuEntry(FText InLabel, FText InTooltip, FPCGAttributePropertySelector InSelector)
+		: Label(std::move(InLabel))
 		, Tooltip(std::move(InTooltip))
 		, Selector(std::move(InSelector))
-		, MenuHierarchy(std::move(InMenuHierarchy))
 	{}
 
-	/** Display name for the entry. */
-	FText DisplayName;
+	/** Label for the entry. */
+	FText Label;
 
 	/** Tooltip to display for the entry. */
 	FText Tooltip;
 
 	/** Selector to copy if this entry is selected. */
 	FPCGAttributePropertySelector Selector;
-	
-	/** Hierarchy of the menus, for all the sub menus to reach this entry. */
-	TArray<FText> MenuHierarchy;
 };
 #endif // WITH_EDITOR
 
@@ -54,7 +61,7 @@ struct FPCGAttributeAccessorMethods
 	PCG_API void FillSelectorMenuEntryFromEnum(const UEnum* EnumType, const TArray<FText>& InMenuHierarchy = {});
 	
 	// Special option to gather all the possible options for a given data type in the Selector context menu.
-	TArray<FPCGAttributeSelectorMenuEntry> AttributeSelectorOptions;
+	FPCGAttributeSelectorMenu AttributeSelectorMenu;
 #endif // WITH_EDITOR
 };
 
@@ -79,8 +86,8 @@ public:
 	TUniquePtr<const IPCGAttributeAccessorKeys> CreateSimpleConstKeys(const UPCGData* InData, const FPCGAttributePropertySelector& InSelector, const bool bQuiet) const;
 
 #if WITH_EDITOR
-	/** Call a Callback on all the FPCGAttributeSelectorMenuEntry currently registered. */
-	PCG_API void ForEachSelectorMenuEntry(TFunctionRef<void(const FPCGAttributeSelectorMenuEntry&)> Callback) const;
+	/** Call a Callback on all the FPCGAttributeSelectorMenu currently registered. */
+	PCG_API void ForEachSelectorMenu(TFunctionRef<void(const FPCGAttributeSelectorMenu&)> Callback) const;
 #endif // WITH_EDITOR
 	
 	template <typename T, typename std::enable_if_t<std::is_base_of_v<UPCGData, T>, bool> = false>
diff --git a/Engine/Plugins/PCG/Source/PCGEditor/Private/Details/PCGAttributePropertySelectorDetails.cpp b/Engine/Plugins/PCG/Source/PCGEditor/Private/Details/PCGAttributePropertySelectorDetails.cpp
--- a/Engine/Plugins/PCG/Source/PCGEditor/Private/Details/PCGAttributePropertySelectorDetails.cpp
+++ b/Engine/Plugins/PCG/Source/PCGEditor/Private/Details/PCGAttributePropertySelectorDetails.cpp
@@ -3,6 +3,7 @@
 #include "Details/PCGAttributePropertySelectorDetails.h"
 
 #include "PCGCommon.h"
+#include "PCGModule.h"
 #include "Metadata/Accessors/PCGAttributeExtractor.h"
 #include "Metadata/PCGAttributePropertySelector.h"
 #include "Metadata/PCGMetadataAttribute.h"
@@ -128,14 +129,14 @@ namespace PCGAttributePropertySelectorDetails
 
 	// Func signature: bool(FPCGAttributePropertySelector*, T)
 	template <typename T, typename Func>
-	void SetValue(FPCGAttributePropertySelector* InSelector, TSharedPtr<IPropertyHandle>& InPropertyHandle, T InValue, Func InCallback)
+	void SetValue(FPCGAttributePropertySelector* InSelector, TSharedPtr<IPropertyHandle>& InPropertyHandle, T&& InValue, Func InCallback)
 	{
 		if (InSelector)
 		{
 			FScopedTransaction Transaction(SetAttributePropertyNameTransaction);
 
 			InPropertyHandle->NotifyPreChange();
-			if (!InCallback(InSelector, std::move(InValue)))
+			if (!InCallback(InSelector, std::forward<T>(InValue)))
 			{
 				Transaction.Cancel();
 				return;
@@ -284,16 +285,54 @@ TSharedRef<SWidget> FPCGAttributePropertySelectorDetails::GenerateExtraMenu()
 
 	if (!StructProperty->HasMetaData(PCGObjectMetadata::DiscardPropertySelection))
 	{
-		MenuBuilder.BeginSection("PointProperties", LOCTEXT("PointPropertiesHeader", "Point Properties"));
+		MenuBuilder.BeginSection(TEXT("Properties"), LOCTEXT("PropertiesHeader", "Properties"));
+
+		TWeakPtr<FPCGAttributePropertySelectorDetails> ThisWeak = StaticCastWeakPtr<FPCGAttributePropertySelectorDetails>(AsWeak());
+
+		auto BuildMenu = [ThisWeak](const FPCGAttributeSelectorMenu& InSelectorMenu, FMenuBuilder& MenuBuilder, auto RecurseCallback) -> void
 		{
-			PCGAttributePropertySelectorDetails::BuildMenuSectionFromEnum<EPCGPointProperties>(MenuBuilder, this, &FPCGAttributePropertySelectorDetails::SetPointProperty);
-		}
+			for (const FPCGAttributeSelectorMenu& SubMenu : InSelectorMenu.SubMenus)
+			{
+				MenuBuilder.AddSubMenu(SubMenu.Label,
+					SubMenu.Tooltip,
+					FNewMenuDelegate::CreateLambda([SubMenu, RecurseCallback](FMenuBuilder& InSubMenuBuilder) -> void
+					{
+						RecurseCallback(SubMenu, InSubMenuBuilder, RecurseCallback);
+					}));
+			}
+
+			for (const FPCGAttributeSelectorMenuEntry& MenuEntry : InSelectorMenu.Entries)
+			{
+				auto SetSelector = [ThisWeak, Selector = MenuEntry.Selector]()
+				{
+					if (FPCGAttributePropertySelectorDetails* This = ThisWeak.Pin().Get())
+					{
+						This->SetSelector(Selector);
+					}
+				};
+				
+				MenuBuilder.AddMenuEntry(
+					MenuEntry.Label,
+					MenuEntry.Tooltip,
+					FSlateIcon(),
+					FExecuteAction::CreateLambda(std::move(SetSelector)),
+					NAME_None,
+					EUserInterfaceActionType::Button
+				);
+			}
+		};
+
+		FPCGModule::GetConstAttributeAccessorFactory().ForEachSelectorMenu([&BuildMenu, &MenuBuilder](const FPCGAttributeSelectorMenu& InSelectorMenu) -> void
+		{
+			BuildMenu(InSelectorMenu, MenuBuilder, BuildMenu);
+		});
+
 		MenuBuilder.EndSection();
 	}
 
 	if (!StructProperty->HasMetaData(PCGObjectMetadata::DiscardExtraSelection))
 	{
-		MenuBuilder.BeginSection("OtherProperties", LOCTEXT("OtherPropertiesHeader", "Other Properties"));
+		MenuBuilder.BeginSection(TEXT("OtherProperties"), LOCTEXT("OtherPropertiesHeader", "Other Properties"));
 		{
 			PCGAttributePropertySelectorDetails::BuildMenuSectionFromEnum<EPCGExtraProperties>(MenuBuilder, this, &FPCGAttributePropertySelectorDetails::SetExtraProperty);
 		}
@@ -351,6 +390,25 @@ void FPCGAttributePropertySelectorDetails::SetExtraProperty(EPCGExtraProperties
 	PCGAttributePropertySelectorDetails::SetValue(GetStruct(), PropertyHandle, EnumValue, [](FPCGAttributePropertySelector* InStruct, EPCGExtraProperties InValue) -> bool { return InStruct && InStruct->SetExtraProperty(InValue); });
 }
 
+void FPCGAttributePropertySelectorDetails::SetSelector(FPCGAttributePropertySelector InSelector)
+{
+	PCGAttributePropertySelectorDetails::SetValue(GetStruct(), PropertyHandle, InSelector, [](FPCGAttributePropertySelector* InStruct, FPCGAttributePropertySelector InValue) -> bool
+	{
+		if (!InStruct)
+		{
+			return false;
+		}
+
+		const bool bIsDifferent = *InStruct != InValue;
+		if (bIsDifferent)
+		{
+			*InStruct = InValue;
+		}
+		
+		return bIsDifferent;
+	});
+}
+
 void FPCGAttributePropertySelectorDetails::AddExtractor(FName InExtractor)
 {
 	PCGAttributePropertySelectorDetails::SetValue(GetStruct(), PropertyHandle, InExtractor, [](FPCGAttributePropertySelector* InStruct, FName InValue) -> bool 
diff --git a/Engine/Plugins/PCG/Source/PCGEditor/Public/Details/PCGAttributePropertySelectorDetails.h b/Engine/Plugins/PCG/Source/PCGEditor/Public/Details/PCGAttributePropertySelectorDetails.h
--- a/Engine/Plugins/PCG/Source/PCGEditor/Public/Details/PCGAttributePropertySelectorDetails.h
+++ b/Engine/Plugins/PCG/Source/PCGEditor/Public/Details/PCGAttributePropertySelectorDetails.h
@@ -4,6 +4,7 @@
 
 #include "IPropertyTypeCustomization.h"
 #include "Layout/Visibility.h"
+#include "Metadata/PCGAttributePropertySelector.h"
 
 enum class EPCGPointProperties : uint8;
 enum class EPCGExtraProperties : uint8;
@@ -40,6 +41,7 @@ protected:
 	void SetPointProperty(EPCGPointProperties EnumValue);
 	void SetAttributeName(FName NewName);
 	void SetExtraProperty(EPCGExtraProperties EnumValue);
+	void SetSelector(FPCGAttributePropertySelector InSelector);
 
 	TSharedPtr<IPropertyHandle> PropertyHandle;
 };
