[PCG] GPU Compiler: remove output pin map that turned out to be unused after latest refactor.

#rnx

diff --git a/Engine/Plugins/PCG/Source/PCG/Private/Graph/PCGGraphCompilerGPU.cpp b/Engine/Plugins/PCG/Source/PCG/Private/Graph/PCGGraphCompilerGPU.cpp
--- a/Engine/Plugins/PCG/Source/PCG/Private/Graph/PCGGraphCompilerGPU.cpp
+++ b/Engine/Plugins/PCG/Source/PCG/Private/Graph/PCGGraphCompilerGPU.cpp
@@ -20,7 +20,6 @@
 #include "Elements/PCGStaticMeshSpawnerKernel.h"
 #include "Graph/PCGGPUGraphCompilationContext.h"
 #include "Graph/PCGGraphCompiler.h"
-#include "Graph/PCGGraphExecutor.h"
 
 #include "ComputeFramework/ComputeKernel.h"
 #include "HAL/IConsoleManager.h"
@@ -337,8 +336,7 @@ void FPCGGraphCompilerGPU::SetupVirtualPins(
 	const TSet<FPCGTaskId>& InCollapsedTasks,
 	const TArray<FPCGGraphTask>& InCompiledTasks,
 	const FTaskToSuccessors& InTaskSuccessors,
-	FOriginalToVirtualPin& OutOriginalToVirtualPin,
-	TMap<TSoftObjectPtr<const UPCGPin>, FName>& OutOutputCPUPinToVirtualPin)
+	FOriginalToVirtualPin& OutOriginalToVirtualPin)
 {
 	TRACE_CPUPROFILER_EVENT_SCOPE(FPCGGraphCompilerGPU::SetupVirtualPins);
 
@@ -357,43 +355,12 @@ void FPCGGraphCompilerGPU::SetupVirtualPins(
 				continue;
 			}
 
-			// TODO is pinless fine with skipping?
+			// Only graph edges associated with pins are considered.
 			if (Input.DownstreamPin.IsSet())
 			{
-				const FName VirtualLabel = *FString::Format(TEXT("{0}-VirtualIn{1}"), { Input.DownstreamPin->Label.ToString(), InputCount });
+				const FName VirtualLabel = *FString::Format(TEXT("{0}-VirtualIn{1}"), { Input.DownstreamPin->Label.ToString(), InputCount++ });
 				const bool bIsInputPin = true;
 				OutOriginalToVirtualPin.Add({ GPUTaskId, Input.DownstreamPin->Label, bIsInputPin }, VirtualLabel);
-
-				++InputCount;
-
-				const FPCGGraphTask& UpstreamTask = InCompiledTasks[Input.TaskId];
-				const UPCGPin* OutputPin = UpstreamTask.Node ? UpstreamTask.Node->GetOutputPin(Input.UpstreamPin->Label) : nullptr;
-
-				// Grid Linkages don't have a node associated, so ask the GridLinkageElement for the output pin instead.
-				if (!OutputPin)
-				{
-					// TODO: In the future FPCGGridLinkageElement shoud be stateless (like all other elements), and we should always only consult the
-					// linkage settings for the upstream pin.
-					if (UpstreamTask.Element && UpstreamTask.Element->IsGridLinkage())
-					{
-						const PCGGraphExecutor::FPCGGridLinkageElement* LinkageElement = static_cast<const PCGGraphExecutor::FPCGGridLinkageElement*>(UpstreamTask.Element.Get());
-						OutputPin = LinkageElement->GetUpstreamPin();
-					}
-					else if (UpstreamTask.CookedSettings && UpstreamTask.CookedSettings->IsA<UPCGGridLinkageSettings>())
-					{
-						const UPCGGridLinkageSettings* LinkageSettings = CastChecked<UPCGGridLinkageSettings>(UpstreamTask.CookedSettings);
-						OutputPin = LinkageSettings->UpstreamPin.Get();
-					}
-				}
-
-				if (OutputPin)
-				{
-					OutOutputCPUPinToVirtualPin.Add(OutputPin, VirtualLabel);
-				}
-				else
-				{
-					UE_LOG(LogPCG, Warning, TEXT("Edge dropped during compilation due to missing upstream output pin."))
-				}
 			}
 		}
 
@@ -407,10 +374,8 @@ void FPCGGraphCompilerGPU::SetupVirtualPins(
 				const FNodePin PinKey = { GPUTaskId, PinProps.Label, /*Pin is input*/false };
 				if (!OutOriginalToVirtualPin.Contains(PinKey))
 				{
-					const FName VirtualLabel = *FString::Format(TEXT("{0}-VirtualOut{1}"), { PinProps.Label.ToString(), OutputCount });
+					const FName VirtualLabel = *FString::Format(TEXT("{0}-VirtualOut{1}"), { PinProps.Label.ToString(), OutputCount++ });
 					OutOriginalToVirtualPin.Add(PinKey, VirtualLabel);
-
-					++OutputCount;
 				}
 			}
 		}
@@ -515,8 +480,7 @@ void FPCGGraphCompilerGPU::BuildComputeGraphTask(
 	const TArray<FPCGTaskId>& InComputeElementTasks,
 	const FTaskToSuccessors& InTaskSuccessors,
 	TArray<FPCGGraphTask>& InOutCompiledTasks,
-	const FOriginalToVirtualPin& InOriginalToVirtualPin,
-	const TMap<TSoftObjectPtr<const UPCGPin>, FName>& InOutputCPUPinToVirtualPin)
+	const FOriginalToVirtualPin& InOriginalToVirtualPin)
 {
 	TRACE_CPUPROFILER_EVENT_SCOPE(FPCGGraphCompilerGPU::BuildComputeGraphTask);
 
@@ -1385,21 +1349,12 @@ void FPCGGraphCompilerGPU::CreateGPUNodes(FPCGGraphCompiler& InOutCompiler, UPCG
 	// and the input and output pins of this element must have unique virtual pin labels so that we can parse the data that
 	// PCG provides through the input data collection correctly, and route the output data to the downstream pins correctly.
 	TArray<FOriginalToVirtualPin> OriginalToVirtualPin;
-
-	TArray<TMap<TSoftObjectPtr<const UPCGPin>, FName>> OutputCPUPinToVirtualPin;
-
 	OriginalToVirtualPin.SetNum(NumComputeGraphs);
-	OutputCPUPinToVirtualPin.SetNum(NumComputeGraphs);
 
 	// Setup mappings from existing pins to compute graph element virtual pins as a prestep before wiring in the compute graph tasks.
 	for (uint32 ComputeGraphIndex = 0; ComputeGraphIndex < NumComputeGraphs; ++ComputeGraphIndex)
 	{
-		SetupVirtualPins(
-			NodeSubsetsToConvertToCFGraph[ComputeGraphIndex],
-			InOutCompiledTasks,
-			TaskSuccessors,
-			OriginalToVirtualPin[ComputeGraphIndex],
-			OutputCPUPinToVirtualPin[ComputeGraphIndex]);
+		SetupVirtualPins(NodeSubsetsToConvertToCFGraph[ComputeGraphIndex], InOutCompiledTasks, TaskSuccessors, OriginalToVirtualPin[ComputeGraphIndex]);
 	}
 
 	TArray<FPCGTaskId> ComputeElementTasks;
@@ -1448,8 +1403,7 @@ void FPCGGraphCompilerGPU::CreateGPUNodes(FPCGGraphCompiler& InOutCompiler, UPCG
 			ComputeElementTasks,
 			TaskSuccessors,
 			InOutCompiledTasks,
-			OriginalToVirtualPin[ComputeGraphIndex],
-			OutputCPUPinToVirtualPin[ComputeGraphIndex]);
+			OriginalToVirtualPin[ComputeGraphIndex]);
 	}
 
 	{
diff --git a/Engine/Plugins/PCG/Source/PCG/Private/Graph/PCGGraphCompilerGPU.h b/Engine/Plugins/PCG/Source/PCG/Private/Graph/PCGGraphCompilerGPU.h
--- a/Engine/Plugins/PCG/Source/PCG/Private/Graph/PCGGraphCompilerGPU.h
+++ b/Engine/Plugins/PCG/Source/PCG/Private/Graph/PCGGraphCompilerGPU.h
@@ -8,7 +8,6 @@ class FPCGGraphCompiler;
 class UPCGComputeDataInterface;
 class UPCGComputeGraph;
 class UPCGGraph;
-class UPCGPin;
 struct FPCGGraphTask;
 struct FPCGGPUCompilationContext;
 
@@ -51,8 +50,7 @@ private:
 		const TSet<FPCGTaskId>& InCollapsedTasks,
 		const TArray<FPCGGraphTask>& InCompiledTasks,
 		const TMap<FPCGTaskId, TArray<FPCGTaskId>>& InTaskSuccessors,
-		FOriginalToVirtualPin& OutOriginalToVirtualPin,
-		TMap<TSoftObjectPtr<const UPCGPin>, FName>& OutOutputCPUPinToVirtualPin);
+		FOriginalToVirtualPin& OutOriginalToVirtualPin);
 
 	/** Wires in a compute graph task alongside each set of GPU compatible nodes. The tasks for each node will be culled later. */
 	static void WireComputeGraphTask(
@@ -73,8 +71,7 @@ private:
 		const TArray<FPCGTaskId>& InComputeElementTasks,
 		const FTaskToSuccessors& InTaskSuccessors,
 		TArray<FPCGGraphTask>& InOutCompiledTasks,
-		const FOriginalToVirtualPin& InOriginalToVirtualPin,
-		const TMap<TSoftObjectPtr<const UPCGPin>, FName>& InOutputCPUPinToVirtualPin);
+		const FOriginalToVirtualPin& InOriginalToVirtualPin);
 
 	/** Creates and configures all data interfaces for the graph. */
 	static void CreateDataInterfaces(
