PCG: Expose a temporary way to run a local PCG executor in a blocking way, for use in prototype code that currently needs to run PCG in a non-main thread while the game thread is paused.

diff --git a/Engine/Plugins/PCG/Source/PCG/Private/Graph/PCGGraphExecutor.cpp b/Engine/Plugins/PCG/Source/PCG/Private/Graph/PCGGraphExecutor.cpp
--- a/Engine/Plugins/PCG/Source/PCG/Private/Graph/PCGGraphExecutor.cpp
+++ b/Engine/Plugins/PCG/Source/PCG/Private/Graph/PCGGraphExecutor.cpp
@@ -696,6 +696,12 @@ FPCGTaskId FPCGGraphExecutor::ScheduleGenericWithContext(TFunction<bool(FPCGCont
 }
 
 FPCGTaskId FPCGGraphExecutor::ScheduleGenericWithContext(TFunction<bool(FPCGContext*)> InOperation, TFunction<void(FPCGContext*)> InAbortOperation, UPCGComponent* InSourceComponent, const TArray<FPCGTaskId>& TaskExecutionDependencies, const TArray<FPCGTaskId>& TaskDataDependencies, bool bSupportBasePointDataInput)
+{
+	return ScheduleSingleElement(MakeShared<FPCGGenericElement>(InOperation, InAbortOperation, bSupportBasePointDataInput), 
+		!!InAbortOperation, InSourceComponent, TaskExecutionDependencies, TaskDataDependencies, bSupportBasePointDataInput);
+}
+
+FPCGTaskId FPCGGraphExecutor::ScheduleSingleElement(TSharedRef<IPCGElement> Element, bool bHasAbortCallbacks, UPCGComponent* InSourceComponent, const TArray<FPCGTaskId>& TaskExecutionDependencies, const TArray<FPCGTaskId>& TaskDataDependencies, bool bSupportBasePointDataInput)
 {
 	// Build task & element to hold the operation to perform
 	FPCGGraphTask Task;
@@ -713,7 +719,7 @@ FPCGTaskId FPCGGraphExecutor::ScheduleGenericWithContext(TFunction<bool(FPCGCont
 	}
 
 	Task.SourceComponent = InSourceComponent;
-	Task.Element = MakeShared<FPCGGenericElement>(InOperation, InAbortOperation, bSupportBasePointDataInput);
+	Task.Element = Element;
 
 	{
 		PCGGraphExecutor::TScopeLock ScopeLock(ScheduleLock);
@@ -724,11 +730,21 @@ FPCGTaskId FPCGGraphExecutor::ScheduleGenericWithContext(TFunction<bool(FPCGCont
 		FPCGGraphScheduleTask& ScheduledTask = ScheduledTasks.Emplace_GetRef();
 		ScheduledTask.Tasks.Add(Task);
 		ScheduledTask.SourceComponent = InSourceComponent;
-		ScheduledTask.bHasAbortCallbacks = !!InAbortOperation;
+		ScheduledTask.bHasAbortCallbacks = bHasAbortCallbacks;
 	}
 	return Task.NodeId;
 }
 
+// Used to expose ScheduleSingleElement to internal helpers without exposing everything
+FPCGTaskId UE::FPCGGraphExecutorInternalAccess::ScheduleSingleElement(FPCGGraphExecutor& Executor, 
+	TSharedRef<IPCGElement> Element, bool bHasAbortCallbacks, UPCGComponent* InSourceComponent, 
+	const TArray<FPCGTaskId>& TaskExecutionDependencies, const TArray<FPCGTaskId>& TaskDataDependencies, 
+	bool bSupportBasePointDataInput)
+{
+	return Executor.ScheduleSingleElement(Element, bHasAbortCallbacks, InSourceComponent, 
+		TaskExecutionDependencies, TaskDataDependencies, bSupportBasePointDataInput);
+}
+
 void FPCGGraphExecutor::MarkInputResults(TArrayView<const FPCGTaskId> InInputResults)
 {
 	TRACE_CPUPROFILER_EVENT_SCOPE(FPCGGraphExecutor::MarkInputResults);
@@ -869,7 +885,6 @@ bool FPCGGraphExecutor::ProcessScheduledTasks()
 
 void FPCGGraphExecutor::ExecuteTasksEnded()
 {
-	check(IsInGameThread());
 	if (!bNeedToExecuteTasksEnded)
 	{
 		return;
@@ -1613,8 +1628,6 @@ void FPCGGraphExecutor::Execute()
 
 void FPCGGraphExecutor::ClearAllTasks()
 {
-	check(IsInGameThread());
-
 	Tasks.Reset();
 	TaskSuccessors.Reset();
 	
diff --git a/Engine/Plugins/PCG/Source/PCG/Private/Graph/PCGGraphExecutor.h b/Engine/Plugins/PCG/Source/PCG/Private/Graph/PCGGraphExecutor.h
--- a/Engine/Plugins/PCG/Source/PCG/Private/Graph/PCGGraphExecutor.h
+++ b/Engine/Plugins/PCG/Source/PCG/Private/Graph/PCGGraphExecutor.h
@@ -29,6 +29,10 @@ class UPCGGraph;
 class UPCGNode;
 struct FPCGStack;
 class FTextFormat;
+namespace UE
+{
+	class FPCGGraphExecutorInternalAccess;
+}
 
 namespace PCGGraphExecutor
 {
@@ -143,6 +147,7 @@ public:
 
 private:
 	friend UPCGSubsystem;
+	friend UE::FPCGGraphExecutorInternalAccess;
 
 	enum class EExecutionCacheDataType : int8
 	{
@@ -200,6 +205,10 @@ private:
 		FPCGGraphExecutor* Executor = nullptr;
 	};
 
+	// Used by ScheduleGenericWithContext and exposed to an internal helper to schedule non-game-thread generics via FPCGGraphExecutorInternalAccess
+	FPCGTaskId ScheduleSingleElement(TSharedRef<IPCGElement> Element, bool bHasAbortCallbacks, UPCGComponent* InSourceComponent,
+		const TArray<FPCGTaskId>& TaskExecutionDependencies, const TArray<FPCGTaskId>& TaskDataDependencies, bool bSupportBasePointDataInput);
+
 	void PostTaskExecute(TSharedPtr<FPCGGraphActiveTask> ActiveTask, bool bIsDone);
 	bool ProcessScheduledTasks();
 	void ExecuteTasksEnded();
@@ -521,3 +530,17 @@ public:
 	TSoftObjectPtr<const UPCGPin> UpstreamPin = nullptr;
 };
 
+namespace UE
+{
+	// Helper to give FPCGInternalsHelper limited access to FPCGGraphExecutor private functions
+	//  rather than friending it completely.
+	// Currently used to allow for scheduling a custom generic element that does not require
+	//  running on game thread.
+	class FPCGGraphExecutorInternalAccess
+	{
+		friend class FPCGInternalHelpers;
+		static FPCGTaskId ScheduleSingleElement(FPCGGraphExecutor& Executor, TSharedRef<IPCGElement> Element, 
+			bool bHasAbortCallbacks, UPCGComponent* InSourceComponent, const TArray<FPCGTaskId>& TaskExecutionDependencies, 
+			const TArray<FPCGTaskId>& TaskDataDependencies, bool bSupportBasePointDataInput = false);
+	};
+}
\ No newline at end of file
diff --git a/Engine/Plugins/PCG/Source/PCG/Private/Helpers/PCGInternalHelpers.cpp b/Engine/Plugins/PCG/Source/PCG/Private/Helpers/PCGInternalHelpers.cpp
new file mode 100644
--- /dev/null
+++ b/Engine/Plugins/PCG/Source/PCG/Private/Helpers/PCGInternalHelpers.cpp
@@ -0,0 +1,72 @@
+// Copyright Epic Games, Inc. All Rights Reserved.
+
+#include "Helpers/PCGInternalHelpers.h"
+
+#include "Graph/PCGGraphExecutor.h"
+#include "PCGComponent.h"
+#include "PCGSubgraph.h" // FPCGInputForwardingElement
+
+namespace PCGInternalHelpersLocals
+{
+	class FPCGGenericElementAnyThread : public FPCGGenericElement
+	{
+	public:
+		FPCGGenericElementAnyThread(
+			TFunction<bool(FPCGContext*)> InOperation)
+			: FPCGGenericElement(InOperation)
+		{}
+
+		virtual bool CanExecuteOnlyOnMainThread(FPCGContext* Context) const override { return false; }
+	};
+}
+
+bool UE::FPCGInternalHelpers::RunGraphLocallyBlocking(UPCGComponent* PCGComponent, const FPCGDataCollection& InputData, 
+	FPCGDataCollection& OutputData, int32 CyclesToWait, double SleepDurationPerCycle)
+{
+	using namespace PCGInternalHelpersLocals;
+
+	if (!PCGComponent)
+	{
+		return false;
+	}
+
+	UPCGGraph* Subgraph = PCGComponent->GetGraph();
+	if (!Subgraph)
+	{
+		return false;
+	}
+
+	FPCGGraphExecutor LocalGraphExecutor;
+	const FPCGTaskId TaskId = LocalGraphExecutor.Schedule(Subgraph, PCGComponent, nullptr,
+		MakeShared<FPCGInputForwardingElement>(InputData),
+		/*Dependencies*/{},
+		/*InFromStack*/ nullptr,
+		/*bAllowHierarchicalGeneration*/ false);
+	PCGComponent->CurrentGenerationTask = TaskId;
+
+	bool bDone = false;
+	const FPCGTaskId FinalTaskId = FPCGGraphExecutorInternalAccess::ScheduleSingleElement(LocalGraphExecutor,
+		MakeShared<FPCGGenericElementAnyThread>([&LocalGraphExecutor, &OutputData, TaskId, &bDone](FPCGContext*) -> bool
+		{
+			LocalGraphExecutor.GetOutputData(TaskId, OutputData);
+			bDone = true;
+			return true;
+		}), 
+		false, PCGComponent, { TaskId }, {});
+
+	int32 CycleCount = 0;
+	CyclesToWait = FMath::Max(0, CyclesToWait);
+	SleepDurationPerCycle = FMath::Max(0, SleepDurationPerCycle);
+	LocalGraphExecutor.Execute();
+
+	while (!bDone && CycleCount < CyclesToWait)
+	{
+		if (SleepDurationPerCycle > 0)
+		{
+			FPlatformProcess::Sleep(SleepDurationPerCycle);
+		}
+		LocalGraphExecutor.Execute();
+		++CycleCount;
+	}
+	return bDone;
+}
\ No newline at end of file
diff --git a/Engine/Plugins/PCG/Source/PCG/Public/Helpers/PCGInternalHelpers.h b/Engine/Plugins/PCG/Source/PCG/Public/Helpers/PCGInternalHelpers.h
new file mode 100644
--- /dev/null
+++ b/Engine/Plugins/PCG/Source/PCG/Public/Helpers/PCGInternalHelpers.h
@@ -0,0 +1,27 @@
+// Copyright Epic Games, Inc. All Rights Reserved.
+
+#pragma once
+
+#include "Misc/CoreMiscDefines.h" // UE_INTERNAL
+
+class UPCGComponent;
+struct FPCGDataCollection;
+
+namespace UE
+{
+	UE_INTERNAL class FPCGInternalHelpers
+	{
+	public:
+		/**
+		 * Temporary helper that runs the graph attached to the given component in a local executor, waiting
+		 *  in a loop for it to finish or to time out. Currently used in some prototype work that needs to
+		 *  run graphs in non-main-threads while game thread is blocked, but will either be replaced with
+		 *  a similar asynchronous call, or perhaps just removed if the prototypes using it are changed.
+		 * 
+		 * Note that pcg.GraphMultithreading must be set to true if this is run outside the game thread.
+		 */
+		UE_INTERNAL PCG_API static bool RunGraphLocallyBlocking(UPCGComponent* PCGComponent, 
+			const FPCGDataCollection& InputData, FPCGDataCollection& OutputData, 
+			int32 CyclesToWait, double SleepDurationPerCycle);
+	};
+}
\ No newline at end of file
diff --git a/Engine/Plugins/PCG/Source/PCG/Public/PCGComponent.h b/Engine/Plugins/PCG/Source/PCG/Public/PCGComponent.h
--- a/Engine/Plugins/PCG/Source/PCG/Public/PCGComponent.h
+++ b/Engine/Plugins/PCG/Source/PCG/Public/PCGComponent.h
@@ -35,6 +35,10 @@ class FLandscapeProxyComponentDataChangedParams;
 class UClass;
 struct FPCGContext;
 struct FPCGStackContext;
+namespace UE 
+{
+	class FPCGInternalHelpers;
+}
 
 #if WITH_EDITOR
 	DECLARE_MULTICAST_DELEGATE_OneParam(FOnPCGGraphStartGenerating, UPCGComponent*);
@@ -94,6 +98,7 @@ class PCG_API UPCGComponent : public UActorComponent
 	friend class UPCGSubsystem;
 	friend class FPCGGraphExecutor;
 	friend class FPCGActorAndComponentMapping;
+	friend class UE::FPCGInternalHelpers;
 	friend struct FPCGWorldPartitionBuilder;
 
 public:
