[PCG] Convert IndirectionElement to support PCGPointArrayData
- also implement SupportsGPUResidentData and properly forward to proxied element

#rnx

diff --git a/Engine/Plugins/PCG/Source/PCG/Private/Elements/PCGIndirectionElement.cpp b/Engine/Plugins/PCG/Source/PCG/Private/Elements/PCGIndirectionElement.cpp
--- a/Engine/Plugins/PCG/Source/PCG/Private/Elements/PCGIndirectionElement.cpp
+++ b/Engine/Plugins/PCG/Source/PCG/Private/Elements/PCGIndirectionElement.cpp
@@ -256,6 +256,11 @@ bool FPCGIndirectionElement::PrepareDataInternal(FPCGContext* InContext) const
 
 		Context->InnerElement = InnerSettings->GetElement();
 		check(Context->InnerElement);
+				
+		if (!Context->InnerElement->ConvertInputsIfNeeded(Context))
+		{
+			return false;
+		}
 
 		// Note: we need to pass a null node here + add the inner settings as part of the input so that they are retrieved properly
 		FPCGDataCollection InnerInput = Context->InputData;
diff --git a/Engine/Plugins/PCG/Source/PCG/Public/Elements/PCGIndirectionElement.h b/Engine/Plugins/PCG/Source/PCG/Public/Elements/PCGIndirectionElement.h
--- a/Engine/Plugins/PCG/Source/PCG/Public/Elements/PCGIndirectionElement.h
+++ b/Engine/Plugins/PCG/Source/PCG/Public/Elements/PCGIndirectionElement.h
@@ -95,4 +95,6 @@ protected:
 	virtual FPCGContext* CreateContext() override;
 	virtual bool PrepareDataInternal(FPCGContext* InContext) const override;
 	virtual bool ExecuteInternal(FPCGContext* InContext) const override;
+	virtual bool SupportsGPUResidentData(FPCGContext* InContext) const override { return true; }
+	virtual bool SupportsBasePointDataInputs(FPCGContext* InContext) const override { return true; }
 };
\ No newline at end of file
