[PCG] Add support for Spline Metadata on Control points

* Add a new layer for metadata for splines, with the control points.
* Metadata entries are stored on the PCG Spline Data and are not allocated if there is no attributes.
* Enforce encapsulation of the spline struct for futureproof-ness, and change where it was accessed directly.
* Attributes are transfered from points to splines in Create Spline.
* Sampling the spline will do an interpolation if the attrbiute is interpolable. This is still not very robust as we don't have a standardized way of marking an attrbiute as interpolable (there is a flag on the attribute but no way of setting it from the PCG Graph). It won't do interpolation in the interior case.
* Improve the Attribute noise with a custom Seed, since we will probably need it for the spline (and can also be used for param data).
* Supports only Spline data, not landscape splines.
* Attributes for water will be reserved in the future and not yet accessible.

diff --git a/Engine/Plugins/Experimental/PCGInterops/PCGWaterInterop/Source/PCGWaterInterop/Private/Data/PCGWaterSplineData.cpp b/Engine/Plugins/Experimental/PCGInterops/PCGWaterInterop/Source/PCGWaterInterop/Private/Data/PCGWaterSplineData.cpp
--- a/Engine/Plugins/Experimental/PCGInterops/PCGWaterInterop/Source/PCGWaterInterop/Private/Data/PCGWaterSplineData.cpp
+++ b/Engine/Plugins/Experimental/PCGInterops/PCGWaterInterop/Source/PCGWaterInterop/Private/Data/PCGWaterSplineData.cpp
@@ -93,6 +93,10 @@ FTransform UPCGWaterSplineData::GetTransformAtDistance(int SegmentIndex, FVector
 
 void UPCGWaterSplineData::WriteMetadataToPoint(float InputKey, FPCGPoint& OutPoint, UPCGMetadata* OutMetadata) const
 {
+	// TODO: Might be interesting for those 4 custom attributes to follow the normal flow for attribute interpolation
+	// instead of having a custom interpolation here.
+	Super::WriteMetadataToPoint(InputKey, OutPoint, OutMetadata);
+
 	if (!OutMetadata)
 	{
 		return;
diff --git a/Engine/Plugins/Experimental/PCGInterops/PCGWaterInterop/Source/PCGWaterInterop/Public/Data/PCGWaterSplineData.h b/Engine/Plugins/Experimental/PCGInterops/PCGWaterInterop/Source/PCGWaterInterop/Public/Data/PCGWaterSplineData.h
--- a/Engine/Plugins/Experimental/PCGInterops/PCGWaterInterop/Source/PCGWaterInterop/Public/Data/PCGWaterSplineData.h
+++ b/Engine/Plugins/Experimental/PCGInterops/PCGWaterInterop/Source/PCGWaterInterop/Public/Data/PCGWaterSplineData.h
@@ -48,7 +48,6 @@ protected:
 	//~Begin UPCGPolyLineData interface
 	virtual FTransform GetTransformAtDistance(int SegmentIndex, FVector::FReal Distance, bool bWorldSpace = true, FBox* OutBounds = nullptr) const override;
 	virtual void WriteMetadataToPoint(float InputKey, FPCGPoint& OutPoint, UPCGMetadata* OutMetadata) const override;
-	virtual bool HasCustomMetadata() const override { return true; }
 	//~End UPCGPolyLineData interface
 
 	UPROPERTY()
diff --git a/Engine/Plugins/PCG/Source/PCG/Private/Data/PCGSplineData.cpp b/Engine/Plugins/PCG/Source/PCG/Private/Data/PCGSplineData.cpp
--- a/Engine/Plugins/PCG/Source/PCG/Private/Data/PCGSplineData.cpp
+++ b/Engine/Plugins/PCG/Source/PCG/Private/Data/PCGSplineData.cpp
@@ -38,9 +38,9 @@ void UPCGSplineData::Initialize(const USplineComponent* InSpline)
 	CachedBounds = CachedBounds.ExpandBy(SplinePointsRadius, SplinePointsRadius);
 }
 
-void UPCGSplineData::Initialize(const TArray<FSplinePoint>& InSplinePoints, bool bIsClosedLoop, const FTransform& InTransform)
+void UPCGSplineData::Initialize(const TArray<FSplinePoint>& InSplinePoints, bool bIsClosedLoop, const FTransform& InTransform, TArray<PCGMetadataEntryKey> InOptionalEntryKeys)
 {
-	SplineStruct.Initialize(InSplinePoints, bIsClosedLoop, InTransform);
+	SplineStruct.Initialize(InSplinePoints, bIsClosedLoop, InTransform, std::move(InOptionalEntryKeys));
 
 	CachedBounds = SplineStruct.GetBounds();
 
@@ -133,13 +133,13 @@ FTransform UPCGSplineData::GetTransformAtDistance(int SegmentIndex, FVector::FRe
 FVector::FReal UPCGSplineData::GetCurvatureAtDistance(int SegmentIndex, FVector::FReal Distance) const
 {
 	const float FullDistance = SplineStruct.GetDistanceAlongSplineAtSplinePoint(SegmentIndex) + Distance;
-	const float Param = SplineStruct.SplineCurves.ReparamTable.Eval(FullDistance, 0.0f);
+	const float Param = SplineStruct.GetSplineRepramTable().Eval(FullDistance, 0.0f);
 
 	// Since we need the first derivative (e.g. very similar to direction) to have its norm, we'll get the value directly
-	const FVector FirstDerivative = SplineStruct.SplineCurves.Position.EvalDerivative(Param, FVector::ZeroVector);
+	const FVector FirstDerivative = SplineStruct.GetSplinePointsPosition().EvalDerivative(Param, FVector::ZeroVector);
 	const FVector::FReal FirstDerivativeLength = FMath::Max(FirstDerivative.Length(), UE_DOUBLE_SMALL_NUMBER);
 	const FVector ForwardVector = FirstDerivative / FirstDerivativeLength;
-	const FVector SecondDerivative = SplineStruct.SplineCurves.Position.EvalSecondDerivative(Param, FVector::ZeroVector);
+	const FVector SecondDerivative = SplineStruct.GetSplinePointsPosition().EvalSecondDerivative(Param, FVector::ZeroVector);
 	// Orthogonalize the second derivative and obtain the curvature vector
 	const FVector CurvatureVector = SecondDerivative - (SecondDerivative | ForwardVector) * ForwardVector;
 	
@@ -154,14 +154,14 @@ FVector::FReal UPCGSplineData::GetCurvatureAtDistance(int SegmentIndex, FVector:
 float UPCGSplineData::GetInputKeyAtDistance(int SegmentIndex, FVector::FReal Distance) const
 {
 	const float FullDistance = GetDistanceAtSegmentStart(SegmentIndex) + Distance;
-	return SplineStruct.SplineCurves.ReparamTable.Eval(FullDistance, 0.0f);
+	return SplineStruct.GetSplineRepramTable().Eval(FullDistance, 0.0f);
 }
 
 void UPCGSplineData::GetTangentsAtSegmentStart(int SegmentIndex, FVector& OutArriveTangent, FVector& OutLeaveTangent) const
 {
-	check(SplineStruct.SplineCurves.Position.Points.IsValidIndex(SegmentIndex));
-	OutArriveTangent = SplineStruct.SplineCurves.Position.Points[SegmentIndex].ArriveTangent;
-	OutLeaveTangent = SplineStruct.SplineCurves.Position.Points[SegmentIndex].LeaveTangent;
+	check(SplineStruct.GetSplinePointsPosition().Points.IsValidIndex(SegmentIndex));
+	OutArriveTangent = SplineStruct.GetSplinePointsPosition().Points[SegmentIndex].ArriveTangent;
+	OutLeaveTangent = SplineStruct.GetSplinePointsPosition().Points[SegmentIndex].LeaveTangent;
 }
 
 FVector::FReal UPCGSplineData::GetDistanceAtSegmentStart(int SegmentIndex) const
@@ -179,10 +179,39 @@ FTransform UPCGSplineData::GetTransformAtAlpha(float Alpha) const
 	return SplineStruct.GetTransformAtSplineInputKey(GetInputKeyAtAlpha(Alpha), ESplineCoordinateSpace::World);
 }
 
+void UPCGSplineData::WriteMetadataToPoint(float InputKey, FPCGPoint& OutPoint, UPCGMetadata* OutMetadata) const
+{
+	if (!OutMetadata)
+	{
+		return;
+	}
+	
+	// TODO: Better interpolation, for now just interpolate values that can be interpolated, otherwise treat the attribute as constant along the segment.
+	TConstArrayView<PCGMetadataEntryKey> EntryKeys = SplineStruct.GetConstControlPointsEntryKeys();
+	if (!EntryKeys.IsEmpty())
+	{
+		int32 PreviousIndex = FMath::Floor(InputKey);
+		int32 NextIndex = FMath::CeilToInt32(InputKey);
+		if (EntryKeys.IsValidIndex(PreviousIndex) && EntryKeys.IsValidIndex(NextIndex))
+		{
+			OutMetadata->InitializeOnSet(OutPoint.MetadataEntry, EntryKeys[PreviousIndex], Metadata);
+			
+			if (PreviousIndex != NextIndex)
+			{
+				const float Alpha = InputKey - PreviousIndex;
+				TStaticArray<TPair<PCGMetadataEntryKey, float>, 2> Coefficients;
+				Coefficients[0] = {EntryKeys[PreviousIndex], 1.0f - Alpha};
+				Coefficients[1] = {EntryKeys[NextIndex], Alpha};
+				OutMetadata->ComputeWeightedAttribute(OutPoint.MetadataEntry, Coefficients, Metadata);
+			}
+		}
+	}
+}
+
 TUniquePtr<IPCGAttributeAccessor> UPCGSplineData::CreateStaticAccessor(const FPCGAttributePropertySelector& InSelector, bool bQuiet)
 {
 	static const FProperty* PCGSplineStructProperty = UPCGSplineData::StaticClass()->FindPropertyByName(GET_MEMBER_NAME_CHECKED(UPCGSplineData, SplineStruct));
-	static const FProperty* SplineCurveProperty = FPCGSplineStruct::StaticStruct()->FindPropertyByName(GET_MEMBER_NAME_CHECKED(FPCGSplineStruct, SplineCurves));
+	static const FProperty* SplineCurveProperty = FPCGSplineStruct::StaticStruct()->FindPropertyByName(TEXT("SplineCurves"));
 
 	const FName DomainName = InSelector.GetDomainName();
 
@@ -248,13 +277,17 @@ FPCGAttributeAccessorMethods UPCGSplineData::GetSplineAccessorMethods()
 
 	auto CreateAccessorKeysFunc = [](UPCGData* InData, const FPCGAttributePropertySelector& InSelector, bool bQuiet) -> TUniquePtr<IPCGAttributeAccessorKeys>
 	{
-		// If the selector is not a Property or ExtraProperty, return an invalid key to create the default keys for spatial data.
+		UPCGSplineData* SplineData = CastChecked<UPCGSplineData>(InData);
 		const EPCGAttributePropertySelection Selection = InSelector.GetSelection();
+		
 		if (Selection == EPCGAttributePropertySelection::Property || Selection == EPCGAttributePropertySelection::ExtraProperty)
 		{
-			UPCGSplineData* SplineData = CastChecked<UPCGSplineData>(InData);
 			return MakeUnique<FPCGAttributeAccessorKeysSplineData>(SplineData);
 		}
+		else if (Selection == EPCGAttributePropertySelection::Attribute && InSelector.GetDomainName() == PCGSplineData::ControlPointDomainName)
+		{
+			return MakeUnique<FPCGAttributeAccessorKeysSplineDataEntries>(SplineData);
+		}
 		else
 		{
 			return {};
@@ -263,13 +296,17 @@ FPCGAttributeAccessorMethods UPCGSplineData::GetSplineAccessorMethods()
 	
 	auto CreateConstAccessorKeysFunc = [](const UPCGData* InData, const FPCGAttributePropertySelector& InSelector, bool bQuiet) -> TUniquePtr<const IPCGAttributeAccessorKeys>
 	{
-		// If the selector is not a Property or ExtraProperty, return an invalid key to create the default keys for spatial data.
+		const UPCGSplineData* SplineData = CastChecked<UPCGSplineData>(InData);
+		
 		const EPCGAttributePropertySelection Selection = InSelector.GetSelection();
 		if (Selection == EPCGAttributePropertySelection::Property || Selection == EPCGAttributePropertySelection::ExtraProperty)
 		{
-			const UPCGSplineData* SplineData = CastChecked<UPCGSplineData>(InData);
 			return MakeUnique<const FPCGAttributeAccessorKeysSplineData>(SplineData);
 		}
+		else if (Selection == EPCGAttributePropertySelection::Attribute && InSelector.GetDomainName() == PCGSplineData::ControlPointDomainName)
+		{
+			return MakeUnique<const FPCGAttributeAccessorKeysSplineDataEntries>(SplineData);
+		}
 		else
 		{
 			return {};
@@ -313,7 +350,6 @@ FBox UPCGSplineData::GetBounds() const
 
 bool UPCGSplineData::SamplePoint(const FTransform& InTransform, const FBox& InBounds, FPCGPoint& OutPoint, UPCGMetadata* OutMetadata) const
 {
-	// TODO: support metadata
 	// TODO: support proper bounds
 
 	// This is a pure SamplePoint implementation.
@@ -337,6 +373,8 @@ bool UPCGSplineData::SamplePoint(const FTransform& InTransform, const FBox& InBo
 		OutPoint.SetLocalBounds(InBounds);
 		OutPoint.Density = 1.0f - Distance;
 
+		WriteMetadataToPoint(NearestPointKey, OutPoint, OutMetadata);
+
 		return true;
 	}
 }
diff --git a/Engine/Plugins/PCG/Source/PCG/Private/Data/PCGSplineInteriorSurfaceData.cpp b/Engine/Plugins/PCG/Source/PCG/Private/Data/PCGSplineInteriorSurfaceData.cpp
--- a/Engine/Plugins/PCG/Source/PCG/Private/Data/PCGSplineInteriorSurfaceData.cpp
+++ b/Engine/Plugins/PCG/Source/PCG/Private/Data/PCGSplineInteriorSurfaceData.cpp
@@ -89,7 +89,7 @@ bool UPCGSplineInteriorSurfaceData::ProjectPoint(const FTransform& InTransform,
 	{
 		// Project rotation by sampling the rotation of the nearest point on the spline.
 		float OutDistanceSquared, OutSegment;
-		const float NearestSplineKey = SplineStruct.SplineCurves.Position.FindNearest(FVector(InLocation.X, InLocation.Y, ProjectionHeight), OutDistanceSquared, OutSegment);
+		const float NearestSplineKey = SplineStruct.GetSplinePointsPosition().FindNearest(FVector(InLocation.X, InLocation.Y, ProjectionHeight), OutDistanceSquared, OutSegment);
 		ProjectedTransform.SetRotation(SplineStruct.GetQuaternionAtSplineInputKey(NearestSplineKey, ESplineCoordinateSpace::Local));
 	}
 
diff --git a/Engine/Plugins/PCG/Source/PCG/Private/Data/PCGSplineStruct.cpp b/Engine/Plugins/PCG/Source/PCG/Private/Data/PCGSplineStruct.cpp
--- a/Engine/Plugins/PCG/Source/PCG/Private/Data/PCGSplineStruct.cpp
+++ b/Engine/Plugins/PCG/Source/PCG/Private/Data/PCGSplineStruct.cpp
@@ -2,6 +2,8 @@
 
 #include "Data/PCGSplineStruct.h"
 
+#include "Metadata/PCGMetadataCommon.h"
+
 #include UE_INLINE_GENERATED_CPP_BY_NAME(PCGSplineStruct)
 
 namespace PCGSplineStruct
@@ -87,9 +89,11 @@ void FPCGSplineStruct::Initialize(const USplineComponent* InSplineComponent)
 
 	Bounds = InSplineComponent->Bounds;
 	LocalBounds = InSplineComponent->CalcLocalBounds();
+
+	ControlPointsEntryKeys.Empty();
 }
 
-void FPCGSplineStruct::Initialize(const TArray<FSplinePoint>& InSplinePoints, bool bIsClosedLoop, const FTransform& InTransform)
+void FPCGSplineStruct::Initialize(const TArray<FSplinePoint>& InSplinePoints, bool bIsClosedLoop, const FTransform& InTransform, TArray<PCGMetadataEntryKey> InOptionalEntryKeys)
 {
 	Transform = InTransform;
 	DefaultUpVector = FVector::ZAxisVector;
@@ -100,6 +104,16 @@ void FPCGSplineStruct::Initialize(const TArray<FSplinePoint>& InSplinePoints, bo
 
 	Bounds = PCGSplineStruct::CalcBounds(SplineCurves, bClosedLoop, InTransform);
 	LocalBounds = PCGSplineStruct::CalcBounds(SplineCurves, bClosedLoop, FTransform::Identity);
+
+	if (!InOptionalEntryKeys.IsEmpty() && InSplinePoints.Num() == InOptionalEntryKeys.Num())
+	{
+		ControlPointsEntryKeys = std::move(InOptionalEntryKeys);
+	}
+	else
+	{
+		// If we have a mismatch, we can't set the entry keys, so reset it
+		ControlPointsEntryKeys.Empty();
+	}
 }
 
 void FPCGSplineStruct::ApplyTo(USplineComponent* InSplineComponent) const
@@ -150,12 +164,10 @@ void FPCGSplineStruct::AddPoint(const FSplinePoint& InSplinePoint, bool bUpdateS
 		),
 		Index);
 
-	// TODO spline metadata?
-	//USplineMetadata* Metadata = GetSplinePointsMetadata();
-	//if (Metadata)
-	//{
-	//	Metadata->AddPoint(InSplinePoint.InputKey);
-	//}
+	if (!ControlPointsEntryKeys.IsEmpty())
+	{
+		ControlPointsEntryKeys.Insert(PCGInvalidEntryKey, Index);
+	}
 
 	if (bUpdateSpline)
 	{
@@ -168,6 +180,11 @@ void FPCGSplineStruct::AddPoints(const TArray<FSplinePoint>& InSplinePoints, boo
 	SplineCurves.Position.Points.Reserve(SplineCurves.Position.Points.Num() + InSplinePoints.Num());
 	SplineCurves.Rotation.Points.Reserve(SplineCurves.Rotation.Points.Num() + InSplinePoints.Num());
 	SplineCurves.Scale.Points.Reserve(SplineCurves.Scale.Points.Num() + InSplinePoints.Num());
+	
+	if (!ControlPointsEntryKeys.IsEmpty())
+	{
+		ControlPointsEntryKeys.Reserve(ControlPointsEntryKeys.Num() + InSplinePoints.Num());
+	}
 
 	for (const auto& SplinePoint : InSplinePoints)
 	{
@@ -247,6 +264,16 @@ const FInterpCurveVector& FPCGSplineStruct::GetSplinePointsPosition() const
 	return SplineCurves.Position;
 }
 
+const FInterpCurveFloat& FPCGSplineStruct::GetSplineRepramTable() const
+{
+	return SplineCurves.ReparamTable;
+}
+
+const FInterpCurveQuat& FPCGSplineStruct::GetSplinePointsRotation() const
+{
+	return SplineCurves.Rotation;
+}
+
 FVector::FReal FPCGSplineStruct::GetDistanceAlongSplineAtSplinePoint(int32 PointIndex) const
 {
 	const int32 NumPoints = SplineCurves.Position.Points.Num();
@@ -310,6 +337,25 @@ float FPCGSplineStruct::FindInputKeyClosestToWorldLocation(const FVector& WorldL
 	return SplineCurves.Position.InaccurateFindNearest(LocalLocation, Dummy);
 }
 
+void FPCGSplineStruct::AllocateMetadataEntries()
+{
+	// Add robustness to cleanup everything if we have a mismatch between the number of points and the number of entry keys
+	const int32 NumPoints = SplineCurves.Position.Points.Num();
+	if (ControlPointsEntryKeys.Num() != NumPoints)
+	{
+		ControlPointsEntryKeys.Empty(NumPoints);
+	}
+	
+	if (ControlPointsEntryKeys.IsEmpty())
+	{
+		ControlPointsEntryKeys.Reserve(NumPoints);
+		for (int32 i = 0; i < NumPoints; ++i)
+		{
+			ControlPointsEntryKeys.Add(PCGInvalidEntryKey);
+		}
+	}
+}
+
 FVector FPCGSplineStruct::GetLocationAtSplineInputKey(float InKey, ESplineCoordinateSpace::Type CoordinateSpace) const
 {
 	FVector Location = SplineCurves.Position.Eval(InKey, FVector::ZeroVector);
@@ -337,7 +383,6 @@ FQuat FPCGSplineStruct::GetQuaternionAtSplineInputKey(float InKey, ESplineCoordi
 		Rot = Transform.GetRotation() * Rot;
 	}
 
-
 	return Rot;
 }
 
diff --git a/Engine/Plugins/PCG/Source/PCG/Private/Elements/Grammar/PCGDuplicateCrossSections.cpp b/Engine/Plugins/PCG/Source/PCG/Private/Elements/Grammar/PCGDuplicateCrossSections.cpp
--- a/Engine/Plugins/PCG/Source/PCG/Private/Elements/Grammar/PCGDuplicateCrossSections.cpp
+++ b/Engine/Plugins/PCG/Source/PCG/Private/Elements/Grammar/PCGDuplicateCrossSections.cpp
@@ -79,7 +79,7 @@ bool FPCGDuplicateCrossSectionsElement::ExecuteInternal(FPCGContext* InContext)
 			continue;
 		}
 
-		const TArray<FInterpCurvePointVector>& ControlPoints = InputSplineData->SplineStruct.SplineCurves.Position.Points;
+		const TArray<FInterpCurvePointVector>& ControlPoints = InputSplineData->SplineStruct.GetSplinePointsPosition().Points;
 		if (ControlPoints.Num() < 2)
 		{
 			continue;
diff --git a/Engine/Plugins/PCG/Source/PCG/Private/Elements/PCGAttributeNoise.cpp b/Engine/Plugins/PCG/Source/PCG/Private/Elements/PCGAttributeNoise.cpp
--- a/Engine/Plugins/PCG/Source/PCG/Private/Elements/PCGAttributeNoise.cpp
+++ b/Engine/Plugins/PCG/Source/PCG/Private/Elements/PCGAttributeNoise.cpp
@@ -264,31 +264,60 @@ bool FPCGAttributeNoiseElement::ExecuteInternal(FPCGContext* InContext) const
 
 			Context->InputSource = Settings->InputSource.CopyAndFixLast(InputData);
 
-			// Create a dummy accessor on the input before allocating the output, to avoid doing useless allocation.
-			TUniquePtr<const IPCGAttributeAccessor> TempInputAccessor = PCGAttributeAccessorHelpers::CreateConstAccessor(InputData, Context->InputSource);
-			if (!TempInputAccessor)
+			Context->InputAccessor = PCGAttributeAccessorHelpers::CreateConstAccessor(InputData, Context->InputSource);
+			if (!Context->InputAccessor)
 			{
-				PCGE_LOG(Error, GraphAndLog, FText::Format(LOCTEXT("CantCreateAccessor", "Could not find Attribute/Property '{0}'"), Context->InputSource.GetDisplayText()));
+				PCGLog::Metadata::LogFailToCreateAccessorError(Context->InputSource, Context);
 				Context->CurrentInput++;
 				continue;
 			}
 
 			// Also need to make sure the accessor is a "noisable" type
-			if (!PCG::Private::IsOfTypes<int32, int64, float, double, FVector, FVector2D, FVector4, FRotator, FQuat>(TempInputAccessor->GetUnderlyingType()))
+			if (!PCG::Private::IsOfTypes<int32, int64, float, double, FVector, FVector2D, FVector4, FRotator, FQuat>(Context->InputAccessor->GetUnderlyingType()))
 			{
 				PCGE_LOG(Error, GraphAndLog, FText::Format(LOCTEXT("AttributeIsNotANumericalType", "Attribute/Property '{0}' is not a numerical type, we can't apply noise to it."), Context->InputSource.GetDisplayText()));
 				Context->CurrentInput++;
 				continue;
 			}
 
-			TempInputAccessor.Reset();
+			// Same verification for the seed. If we have no custom seed but points, use the $Seed accessor.
+			// Otherwise, it will just use the index of the element as the seed (combined with component + settings seed of course)
+			const UPCGPointData* InputPointData = Cast<UPCGPointData>(Input.Data);
+			if (Settings->bHasCustomSeedSource || InputPointData)
+			{
+				if (Settings->bHasCustomSeedSource)
+				{
+					Context->CustomSeedSource = Settings->InputSource.CopyAndFixLast(InputData);
+				}
+				else
+				{
+					Context->CustomSeedSource.SetPointProperty(EPCGPointProperties::Seed);
+				}
+				
+				Context->CustomSeedAccessor = PCGAttributeAccessorHelpers::CreateConstAccessor(InputData, Context->CustomSeedSource);
+				if (!Context->CustomSeedAccessor)
+				{
+					PCGLog::Metadata::LogFailToCreateAccessorError(Context->CustomSeedSource, Context);
+					Context->CurrentInput++;
+					continue;
+				}
+				
+				if (!PCG::Private::IsBroadcastableOrConstructible(Context->CustomSeedAccessor->GetUnderlyingType(), PCG::Private::MetadataTypes<int32>::Id))
+				{
+					PCGLog::Metadata::LogFailToGetAttributeError<int32>(Context->CustomSeedSource, Context->CustomSeedAccessor.Get(), Context);
+					Context->CurrentInput++;
+					continue;
+				}
+
+				Context->CustomSeedKeys = PCGAttributeAccessorHelpers::CreateConstKeys(InputData, Context->CustomSeedSource);
+				check(Context->CustomSeedKeys);
+			}
 
 			FPCGTaggedData& Output = Outputs.Add_GetRef(Input);
 			UPCGData* OutputData = InputData->DuplicateData(Context);
 			Output.Data = OutputData;
 
-			// Then create the accessor/keys
-			Context->InputAccessor = PCGAttributeAccessorHelpers::CreateConstAccessor(InputData, Context->InputSource);
+			// Then create the keys
 			Context->InputKeys = PCGAttributeAccessorHelpers::CreateConstKeys(InputData, Context->InputSource);
 
 			// It won't fail because we validated on the input data and output is initialized from input.
@@ -334,14 +363,10 @@ bool FPCGAttributeNoiseElement::ExecuteInternal(FPCGContext* InContext) const
 			Context->bDataPreparedForCurrentInput = true;
 		}
 
-		// Input points are solely used to use the point seed for randomness. Otherwise, it will just use the index of the element as the seed (combined with component + settings seed of course)
-		const UPCGPointData* InputPointData = Cast<UPCGPointData>(Input.Data);
-		const TArray<FPCGPoint>* InputPoints = InputPointData ? &InputPointData->GetPoints() : nullptr;
-
 		// Force clamp on Density
 		const bool bClampResult = Settings->bClampResult || (Context->OutputTarget.GetSelection() == EPCGAttributePropertySelection::Property && Context->OutputTarget.GetPointProperty() == EPCGPointProperties::Density);
 
-		const bool bDone = PCGMetadataAttribute::CallbackWithRightType(Context->InputAccessor->GetUnderlyingType(), [Settings, Context, Seed, bClampResult, InputPoints](auto&& Dummy) -> bool
+		const bool bDone = PCGMetadataAttribute::CallbackWithRightType(Context->InputAccessor->GetUnderlyingType(), [Settings, Context, Seed, bClampResult](auto&& Dummy) -> bool
 		{
 			using AttributeType = std::decay_t<decltype(Dummy)>;
 			constexpr int32 ChunkSize = 64;
@@ -354,23 +379,35 @@ bool FPCGAttributeNoiseElement::ExecuteInternal(FPCGContext* InContext) const
 			// It's finished if we processed all elements.
 			auto Finished = [NumIterations](int32 Count) { ensure(NumIterations == Count); };
 
-			return FPCGAsync::AsyncProcessingRangeEx(&Context->AsyncState, NumIterations, Initialize, [Settings, Seed, bClampResult, Context, ChunkSize, InputPoints, NumIterations](int32 StartReadIndex, int32 StartWriteIndex, int32 Count) -> int32
+			return FPCGAsync::AsyncProcessingRangeEx(&Context->AsyncState, NumIterations, Initialize, [Settings, Seed, bClampResult, Context, ChunkSize, NumIterations](int32 StartReadIndex, int32 StartWriteIndex, int32 Count) -> int32
 			{
 				TArray<AttributeType, TInlineAllocator<ChunkSize>> Values;
 				Values.SetNumUninitialized(Count);
 
+				TArray<int32, TInlineAllocator<ChunkSize>> Seeds;
+				bool bHasCustomSeed = false;
+				if (Context->CustomSeedAccessor.IsValid())
+				{
+					bHasCustomSeed = true;
+					Seeds.SetNumUninitialized(Count);
+					if (!ensure(Context->CustomSeedAccessor->GetRange<int32>(Seeds, StartReadIndex, *Context->CustomSeedKeys, EPCGAttributeAccessorFlags::AllowBroadcastAndConstructible)))
+					{
+						return 0;
+					}
+				}
+
 				if (Context->InputAccessor->GetRange<AttributeType>(Values, StartReadIndex, *Context->InputKeys))
 				{
 					for (int32 i = 0; i < Count; ++i)
 					{
 						// Use the point seed if we have points, otherwise the index. Don't start at 0 (that's why there is a +1)
 						// Warning: It makes it order independant for points, but order dependant for the rest.
-						const int32 ElementSeed = InputPoints ? (*InputPoints)[StartReadIndex + i].Seed : StartReadIndex + i + 1;
+						const int32 ElementSeed = bHasCustomSeed ? Seeds[i] : StartReadIndex + i + 1;
 						FRandomStream RandomSource(PCGHelpers::ComputeSeed(Seed, ElementSeed));
 						PCGAttributeNoiseSettings::ProcessNoise(Values[i], RandomSource, Settings, bClampResult);
 					}
 
-					Context->OutputAccessor->SetRange<AttributeType>(Values, StartWriteIndex, *Context->OutputKeys, EPCGAttributeAccessorFlags::AllowBroadcast | EPCGAttributeAccessorFlags::AllowConstructible);
+					Context->OutputAccessor->SetRange<AttributeType>(Values, StartWriteIndex, *Context->OutputKeys, EPCGAttributeAccessorFlags::AllowBroadcastAndConstructible);
 				}
 
 				return Count;
diff --git a/Engine/Plugins/PCG/Source/PCG/Private/Elements/PCGCleanSpline.cpp b/Engine/Plugins/PCG/Source/PCG/Private/Elements/PCGCleanSpline.cpp
--- a/Engine/Plugins/PCG/Source/PCG/Private/Elements/PCGCleanSpline.cpp
+++ b/Engine/Plugins/PCG/Source/PCG/Private/Elements/PCGCleanSpline.cpp
@@ -108,15 +108,23 @@ bool FPCGCleanSplineElement::ExecuteInternal(FPCGContext* InContext) const
 		}
 
 		const FTransform& SplineTransform = InputSplineData->SplineStruct.Transform;
-		const FInterpCurveVector& ControlPointsPosition = InputSplineData->SplineStruct.SplineCurves.Position;
-		const FInterpCurveQuat& ControlPointsRotation = InputSplineData->SplineStruct.SplineCurves.Rotation;
-		const FInterpCurveVector& ControlPointsScale = InputSplineData->SplineStruct.SplineCurves.Scale;
+		const FInterpCurveVector& ControlPointsPosition = InputSplineData->SplineStruct.GetSplinePointsPosition();
+		const FInterpCurveQuat& ControlPointsRotation = InputSplineData->SplineStruct.GetSplinePointsRotation();
+		const FInterpCurveVector& ControlPointsScale = InputSplineData->SplineStruct.GetSplinePointsScale();
+		const TConstArrayView<PCGMetadataEntryKey> ControlPointsMetadataEntry = InputSplineData->SplineStruct.GetConstControlPointsEntryKeys();
 		const int32 NumControlPoints = ControlPointsPosition.Points.Num();
 		bool bControlPointWasRemoved = false;
 
 		TArray<FSplinePoint> ControlPoints;
 		ControlPoints.Reserve(NumControlPoints);
 
+		// Making sure that the number of metadata entry keys matches, if not, we are in an invalid state and will reset the metadata entry keys.
+		TArray<PCGMetadataEntryKey> ControlPointsKeys;
+		if (ControlPointsMetadataEntry.Num() == NumControlPoints)
+		{
+			ControlPointsKeys.Append(ControlPointsMetadataEntry);
+		}
+
 		// For code clarity, generate the points first and remove them as needed.
 		for (int32 i = 0; i < NumControlPoints; ++i)
 		{
@@ -132,6 +140,19 @@ bool FPCGCleanSplineElement::ExecuteInternal(FPCGContext* InContext) const
 			   ESplinePointType::CurveCustomTangent);
 		}
 
+		auto RemovePoint = [&ControlPoints, &ControlPointsKeys, &bControlPointWasRemoved](int Index)
+		{
+			if (!ControlPointsKeys.IsEmpty())
+			{
+				check(ControlPointsKeys.IsValidIndex(Index));
+				ControlPointsKeys.RemoveAt(Index, EAllowShrinking::No);
+			}
+
+			check(ControlPoints.IsValidIndex(Index));
+			ControlPoints.RemoveAt(Index, EAllowShrinking::No);
+			bControlPointWasRemoved = true;
+		};
+
 		if (Settings->bFuseColocatedControlPoints)
 		{
 			// Will evaluate by pairs. Reverse order for optimizing RemoveAt.
@@ -157,12 +178,12 @@ bool FPCGCleanSplineElement::ExecuteInternal(FPCGContext* InContext) const
 						// Keep the first control point
 						case EPCGControlPointFuseMode::KeepFirst:
 							ControlPoints[Index - 1].LeaveTangent = ControlPoints[Index].LeaveTangent;
-							ControlPoints.RemoveAt(Index, EAllowShrinking::No);
+							RemovePoint(Index);
 							break;
 						// Keep the second control point
 						case EPCGControlPointFuseMode::KeepSecond:
 							ControlPoints[Index].ArriveTangent = ControlPoints[Index - 1].ArriveTangent;
-							ControlPoints.RemoveAt(Index - 1, EAllowShrinking::No);
+							RemovePoint(Index - 1);
 							break;
 						// Average the two control points' transforms and update the leave tangent
 						case EPCGControlPointFuseMode::Merge:
@@ -170,15 +191,13 @@ bool FPCGCleanSplineElement::ExecuteInternal(FPCGContext* InContext) const
 							ControlPoints[Index - 1].Rotation = FMath::Lerp(ControlPoints[Index - 1].Rotation, ControlPoints[Index].Rotation, 0.5);
 							ControlPoints[Index - 1].Scale = FMath::Lerp(ControlPoints[Index - 1].Scale, ControlPoints[Index].Scale, 0.5);
 							ControlPoints[Index - 1].LeaveTangent = ControlPoints[Index].LeaveTangent;
-							ControlPoints.RemoveAt(Index, EAllowShrinking::No);
+							RemovePoint(Index);
 							break;
 						case EPCGControlPointFuseMode::Auto: // Should've been picked up by now. Fallthrough...
 						default:
 							checkNoEntry();
 							break;
 					}
-
-					bControlPointWasRemoved = true;
 				}
 			}
 		}
@@ -196,27 +215,28 @@ bool FPCGCleanSplineElement::ExecuteInternal(FPCGContext* InContext) const
 				{
 					ControlPoints[Index - 2].LeaveTangent = ControlPoints[Index].Position - ControlPoints[Index - 2].Position;
 					ControlPoints[Index].ArriveTangent = ControlPoints[Index - 2].LeaveTangent;
-					ControlPoints.RemoveAt(Index - 1, EAllowShrinking::No);
-					bControlPointWasRemoved = true;
+					RemovePoint(Index - 1);
 				}
 			}
 		}
 
-		// Update input keys if a point was removed, to keep them monotonically incremental
+		FPCGTaggedData& Output = InContext->OutputData.TaggedData.Emplace_GetRef(InputData);
+
+		// We only need to create a new data if a point was removed.
 		if (bControlPointWasRemoved)
 		{
+			// Update input keys if a point was removed, to keep them monotonically incremental
 			for (int i = 0; i < ControlPoints.Num(); ++i)
 			{
 				ControlPoints[i].InputKey = i;
 			}
-		}
+			
+			UPCGSplineData* NewSplineData = FPCGContext::NewObject_AnyThread<UPCGSplineData>(InContext);
+			NewSplineData->InitializeFromData(InputSplineData);
+			NewSplineData->Initialize(ControlPoints, InputSplineData->IsClosed(), InputSplineData->GetTransform(), std::move(ControlPointsKeys));
 
-		UPCGSplineData* NewSplineData = FPCGContext::NewObject_AnyThread<UPCGSplineData>(InContext);
-		NewSplineData->InitializeFromData(InputSplineData);
-		NewSplineData->Initialize(ControlPoints, InputSplineData->IsClosed(), InputSplineData->GetTransform());
-
-		FPCGTaggedData& Output = InContext->OutputData.TaggedData.Emplace_GetRef(InputData);
-		Output.Data = NewSplineData;
+			Output.Data = NewSplineData;
+		}
 	}
 
 	return true;
diff --git a/Engine/Plugins/PCG/Source/PCG/Private/Elements/PCGCreateSpline.cpp b/Engine/Plugins/PCG/Source/PCG/Private/Elements/PCGCreateSpline.cpp
--- a/Engine/Plugins/PCG/Source/PCG/Private/Elements/PCGCreateSpline.cpp
+++ b/Engine/Plugins/PCG/Source/PCG/Private/Elements/PCGCreateSpline.cpp
@@ -123,12 +123,17 @@ bool FPCGCreateSplineElement::ExecuteInternal(FPCGContext* Context) const
 		}
 
 		UPCGSplineData* SplineData = FPCGContext::NewObject_AnyThread<UPCGSplineData>(Context);
+		SplineData->InitializeFromData(PointData);
 		AActor* SplineActor = TargetActor;
 
 		const TArray<FPCGPoint>& Points = PointData->GetPoints();
 		TArray<FSplinePoint> SplinePoints;
 		SplinePoints.Reserve(Points.Num());
 
+		TArray<PCGMetadataEntryKey> SplineEntryKeys;
+		SplineEntryKeys.Reserve(Points.Num());
+		bool bHasAValidEntry = false;
+
 		const FTransform SplineActorTransform = SplineActor->GetTransform();
 		
 		ESplinePointType::Type PointType = ESplinePointType::Curve;
@@ -154,9 +159,17 @@ bool FPCGCreateSplineElement::ExecuteInternal(FPCGContext* Context) const
 				PointTransform.GetRotation().Rotator(),
 				PointTransform.GetScale3D(),
 				PointType);
+
+			SplineEntryKeys.Emplace(Point.MetadataEntry);
+			bHasAValidEntry |= (Point.MetadataEntry != PCGInvalidEntryKey);
+		}
+
+		if (!bHasAValidEntry)
+		{
+			SplineEntryKeys.Empty();
 		}
 
-		SplineData->Initialize(SplinePoints, Settings->bClosedLoop, FTransform(SplineActorTransform.GetLocation()));
+		SplineData->Initialize(SplinePoints, Settings->bClosedLoop, FTransform(SplineActorTransform.GetLocation()), std::move(SplineEntryKeys));
 
 		USplineComponent* SplineComponent = nullptr;
 
diff --git a/Engine/Plugins/PCG/Source/PCG/Private/Elements/PCGSplineDirection.cpp b/Engine/Plugins/PCG/Source/PCG/Private/Elements/PCGSplineDirection.cpp
--- a/Engine/Plugins/PCG/Source/PCG/Private/Elements/PCGSplineDirection.cpp
+++ b/Engine/Plugins/PCG/Source/PCG/Private/Elements/PCGSplineDirection.cpp
@@ -14,7 +14,7 @@ namespace PCGSplineDirection
 		check(InputSplineData);
 
 		double CumulativeAngle = 0.0;
-		int NumPoints = InputSplineData->SplineStruct.SplineCurves.Position.Points.Num();
+		int NumPoints = InputSplineData->SplineStruct.GetSplinePointsPosition().Points.Num();
 		if (!InputSplineData->IsClosed())
 		{
 			NumPoints--;
@@ -39,9 +39,10 @@ namespace PCGSplineDirection
 	{
 		check(InputSplineData);
 
-		const FInterpCurveVector& ControlPointsPosition = InputSplineData->SplineStruct.SplineCurves.Position;
-		const FInterpCurveQuat& ControlPointsRotation = InputSplineData->SplineStruct.SplineCurves.Rotation;
-		const FInterpCurveVector& ControlPointsScale = InputSplineData->SplineStruct.SplineCurves.Scale;
+		const FInterpCurveVector& ControlPointsPosition = InputSplineData->SplineStruct.GetSplinePointsPosition();
+		const FInterpCurveQuat& ControlPointsRotation = InputSplineData->SplineStruct.GetSplinePointsRotation();
+		const FInterpCurveVector& ControlPointsScale = InputSplineData->SplineStruct.GetSplinePointsScale();
+		TArray<PCGMetadataEntryKey> ControlPointKeys(InputSplineData->SplineStruct.GetConstControlPointsEntryKeys());
 
 		TArray<FSplinePoint> NewControlPoints;
 		NewControlPoints.Reserve(ControlPointsPosition.Points.Num());
@@ -66,7 +67,7 @@ namespace PCGSplineDirection
 
 		UPCGSplineData* NewSplineData = FPCGContext::NewObject_AnyThread<UPCGSplineData>(Context);
 		NewSplineData->InitializeFromData(InputSplineData);
-		NewSplineData->Initialize(NewControlPoints, InputSplineData->IsClosed(), InputSplineData->GetTransform());
+		NewSplineData->Initialize(NewControlPoints, InputSplineData->IsClosed(), InputSplineData->GetTransform(), std::move(ControlPointKeys));
 
 		return NewSplineData;
 	}
diff --git a/Engine/Plugins/PCG/Source/PCG/Private/Elements/PCGSplineSampler.cpp b/Engine/Plugins/PCG/Source/PCG/Private/Elements/PCGSplineSampler.cpp
--- a/Engine/Plugins/PCG/Source/PCG/Private/Elements/PCGSplineSampler.cpp
+++ b/Engine/Plugins/PCG/Source/PCG/Private/Elements/PCGSplineSampler.cpp
@@ -687,8 +687,7 @@ namespace PCGSplineSamplerHelpers
 			{
 				constexpr double DefaultValue = 0.0;
 
-				bHasCustomMetadata = LineData->HasCustomMetadata();
-				bSetMetadata |= bHasCustomMetadata;
+				bSetMetadata = false;
 
 				if (Params.bComputeDirectionDelta)
 				{
@@ -750,61 +749,57 @@ namespace PCGSplineSamplerHelpers
 			if (bSetMetadata)
 			{
 				OutMetadata->InitializeOnSet(OutPoint.MetadataEntry);
-			}
 
-			if (bHasCustomMetadata)
-			{
-				LineData->WriteMetadataToPoint(InResult.InputKey, OutPoint, OutMetadata);
-			}
-
-			if (NextDirectionDeltaAttribute)
-			{
-				NextDirectionDeltaAttribute->SetValue(OutPoint.MetadataEntry, InResult.NextDeltaAngle);
-			}
+				if (NextDirectionDeltaAttribute)
+				{
+					NextDirectionDeltaAttribute->SetValue(OutPoint.MetadataEntry, InResult.NextDeltaAngle);
+				}
 
-			if (CurvatureAttribute)
-			{
-				CurvatureAttribute->SetValue(OutPoint.MetadataEntry, InResult.Curvature);
-			}
+				if (CurvatureAttribute)
+				{
+					CurvatureAttribute->SetValue(OutPoint.MetadataEntry, InResult.Curvature);
+				}
 
-			if (SegmentIndexAttribute)
-			{
-				SegmentIndexAttribute->SetValue(OutPoint.MetadataEntry, InResult.SegmentIndex);
-			}
+				if (SegmentIndexAttribute)
+				{
+					SegmentIndexAttribute->SetValue(OutPoint.MetadataEntry, InResult.SegmentIndex);
+				}
 			
-			if (SubsegmentIndexAttribute)
-			{
-				SubsegmentIndexAttribute->SetValue(OutPoint.MetadataEntry, InResult.SubsegmentIndex);
-			}
+				if (SubsegmentIndexAttribute)
+				{
+					SubsegmentIndexAttribute->SetValue(OutPoint.MetadataEntry, InResult.SubsegmentIndex);
+				}
 
-			if (ArriveTangentAttribute)
-			{
-				ArriveTangentAttribute->SetValue(OutPoint.MetadataEntry, InResult.ArriveTangent);
-			}
+				if (ArriveTangentAttribute)
+				{
+					ArriveTangentAttribute->SetValue(OutPoint.MetadataEntry, InResult.ArriveTangent);
+				}
 
-			if (LeaveTangentAttribute)
-			{
-				LeaveTangentAttribute->SetValue(OutPoint.MetadataEntry, InResult.LeaveTangent);
-			}
+				if (LeaveTangentAttribute)
+				{
+					LeaveTangentAttribute->SetValue(OutPoint.MetadataEntry, InResult.LeaveTangent);
+				}
 
-			if (AlphaAttribute)
-			{
-				AlphaAttribute->SetValue(OutPoint.MetadataEntry, InResult.Alpha);
-			}
+				if (AlphaAttribute)
+				{
+					AlphaAttribute->SetValue(OutPoint.MetadataEntry, InResult.Alpha);
+				}
 
-			if (DistanceAttribute)
-			{
-				DistanceAttribute->SetValue(OutPoint.MetadataEntry, InResult.Distance);
-			}
+				if (DistanceAttribute)
+				{
+					DistanceAttribute->SetValue(OutPoint.MetadataEntry, InResult.Distance);
+				}
 
-			if (InputKeyAttribute)
-			{
-				InputKeyAttribute->SetValue(OutPoint.MetadataEntry, InResult.InputKey);
+				if (InputKeyAttribute)
+				{
+					InputKeyAttribute->SetValue(OutPoint.MetadataEntry, InResult.InputKey);
+				}
 			}
 		}
 
 		virtual void Sample(const FSamplerResult& InResult, UPCGPointData* OutPointData)
 		{
+			// Mimic PCGSplineData sampling.
 			const FTransform Transform = InResult.LocalTransform * LineData->GetTransform();
 
 			bool bValid = true;
@@ -813,11 +808,13 @@ namespace PCGSplineSamplerHelpers
 			OutPoint.Density = 1.0f;
 			OutPoint.SetLocalBounds(InResult.Box);
 			OutPoint.Steepness = Params.PointSteepness;
+			
+			LineData->WriteMetadataToPoint(InResult.InputKey, OutPoint, OutPointData->Metadata);
 
 			if (ProjectionTargetData)
 			{
 				FPCGPoint ProjPoint;
-				bValid = ProjectionTargetData->ProjectPoint(Transform, InResult.Box, ProjectionParams, ProjPoint, nullptr);
+				bValid = ProjectionTargetData->ProjectPoint(Transform, InResult.Box, ProjectionParams, ProjPoint, OutPointData->Metadata);
 				OutPoint.Transform = ProjPoint.Transform;
 			}
 			else
@@ -827,7 +824,7 @@ namespace PCGSplineSamplerHelpers
 
 			// Only add valid points within bounds
 			FPCGPoint BoundsTestPoint;
-			if (bValid && (!BoundingShapeData || BoundingShapeData->SamplePoint(Transform, InResult.Box, BoundsTestPoint, nullptr)))
+			if (bValid && (!BoundingShapeData || BoundingShapeData->SamplePoint(Transform, InResult.Box, BoundsTestPoint, OutPointData->Metadata)))
 			{
 				SetSeed(OutPoint, InResult.LocalTransform.GetLocation(), Params, Seed, InResult.SampleIndex);
 				SetMetadata(InResult, OutPoint, OutPointData->Metadata);
@@ -843,7 +840,6 @@ namespace PCGSplineSamplerHelpers
 		int Seed = 0;
 
 		bool bSetMetadata = false;
-		bool bHasCustomMetadata = false;
 		FPCGMetadataAttribute<double>* NextDirectionDeltaAttribute = nullptr;
 		FPCGMetadataAttribute<double>* CurvatureAttribute = nullptr;
 		FPCGMetadataAttribute<int>* SegmentIndexAttribute = nullptr;
@@ -919,7 +915,7 @@ namespace PCGSplineSamplerHelpers
 
 					// Sample spline to get density.
 					FPCGPoint OutPoint;
-					if (!LineData->SamplePoint(TentativeTransform, SubBox, OutPoint, nullptr))
+					if (!LineData->SamplePoint(TentativeTransform, SubBox, OutPoint, OutPointData->Metadata))
 					{
 						continue;
 					}
@@ -928,7 +924,7 @@ namespace PCGSplineSamplerHelpers
 					if (ProjectionTargetData)
 					{
 						FPCGPoint ProjPoint;
-						if (!ProjectionTargetData->ProjectPoint(OutPoint.Transform, OutPoint.GetLocalBounds(), ProjectionParams, ProjPoint, nullptr))
+						if (!ProjectionTargetData->ProjectPoint(OutPoint.Transform, OutPoint.GetLocalBounds(), ProjectionParams, ProjPoint, OutPointData->Metadata))
 						{
 							continue;
 						}
@@ -937,7 +933,7 @@ namespace PCGSplineSamplerHelpers
 
 					// Prune points outside of bounds
 					FPCGPoint BoundsTestPoint;
-					if (BoundingShapeData && !BoundingShapeData->SamplePoint(OutPoint.Transform, OutPoint.GetLocalBounds(), BoundsTestPoint, nullptr))
+					if (BoundingShapeData && !BoundingShapeData->SamplePoint(OutPoint.Transform, OutPoint.GetLocalBounds(), BoundsTestPoint, OutPointData->Metadata))
 					{
 						continue;
 					}
@@ -1286,7 +1282,7 @@ namespace PCGSplineSamplerHelpers
 
 						// if bTreatAsPolyline, then we shouldnt use this, we should use nearest point on the polygon line segments
 						float Dummy;
-						const float NearestSplineKey = bFindNearestSplineKey ? Spline->SplineCurves.Position.InaccurateFindNearest(SurfaceLocation, Dummy) : 0.f;
+						const float NearestSplineKey = bFindNearestSplineKey ? Spline->GetSplinePointsPosition().InaccurateFindNearest(SurfaceLocation, Dummy) : 0.f;
 
 						const FVector PointLocationLS = (Params.bProjectOntoSurface ? FVector(SurfaceLocation) : FVector(SampleLocation, MinPoint.Z));
 
@@ -1357,7 +1353,7 @@ namespace PCGSplineSamplerHelpers
 						if (InProjectionTarget)
 						{
 							FPCGPoint ProjectedPoint;
-							if (!InProjectionTarget->ProjectPoint(TransformWS, FBox(BoundsMin, BoundsMax), InProjectionParams, ProjectedPoint, nullptr))
+							if (!InProjectionTarget->ProjectPoint(TransformWS, FBox(BoundsMin, BoundsMax), InProjectionParams, ProjectedPoint, OutPointData->Metadata))
 							{
 								continue;
 							}
@@ -1380,7 +1376,7 @@ namespace PCGSplineSamplerHelpers
 
 						// Prune points outside of bounds
 						FPCGPoint BoundsTestPoint;
-						if (InBoundingShape && !InBoundingShape->SamplePoint(TransformWS, GeneratedPointBounds, BoundsTestPoint, nullptr))
+						if (InBoundingShape && !InBoundingShape->SamplePoint(TransformWS, GeneratedPointBounds, BoundsTestPoint, OutPointData->Metadata))
 						{
 							continue;
 						}
diff --git a/Engine/Plugins/PCG/Source/PCG/Private/Metadata/Accessors/PCGAttributeAccessorKeys.cpp b/Engine/Plugins/PCG/Source/PCG/Private/Metadata/Accessors/PCGAttributeAccessorKeys.cpp
--- a/Engine/Plugins/PCG/Source/PCG/Private/Metadata/Accessors/PCGAttributeAccessorKeys.cpp
+++ b/Engine/Plugins/PCG/Source/PCG/Private/Metadata/Accessors/PCGAttributeAccessorKeys.cpp
@@ -86,6 +86,11 @@ void FPCGAttributeAccessorKeysEntries::InitializeFromMetadata(const FPCGMetadata
 	}
 }
 
+FPCGAttributeAccessorKeysEntries::FPCGAttributeAccessorKeysEntries()
+	: IPCGAttributeAccessorKeys(/*bReadOnly=*/true)
+{
+}
+
 bool FPCGAttributeAccessorKeysEntries::GetMetadataEntryKeys(int32 InStart, TArrayView<PCGMetadataEntryKey*> OutEntryKeys)
 {
 	return PCGAttributeAccessorKeys::GetKeys(Entries, InStart, OutEntryKeys, [](PCGMetadataEntryKey& Key) -> PCGMetadataEntryKey* { return &Key; });
@@ -217,5 +222,39 @@ bool FPCGAttributeAccessorKeysPointsSubset::GetMetadataEntryKeys(int32 InStart,
 
 int32 FPCGAttributeAccessorKeysSplineData::GetNum() const
 {
-	return Ptr ? Ptr->SplineStruct.SplineCurves.Position.Points.Num() : 0;
-}
\ No newline at end of file
+	return Ptr ? Ptr->SplineStruct.GetSplinePointsPosition().Points.Num() : 0;
+}
+
+/////////////////////////////////////////////////////////
+
+FPCGAttributeAccessorKeysSplineDataEntries::FPCGAttributeAccessorKeysSplineDataEntries(const UPCGSplineData* InSplineData)
+	: Ptr(InSplineData)
+{
+	check(InSplineData);
+
+	bIsReadOnly = true;
+	// const_cast the view
+	TConstArrayView<PCGMetadataEntryKey> ConstEntries = InSplineData->SplineStruct.GetConstControlPointsEntryKeys();
+	Entries = MakeArrayView(const_cast<PCGMetadataEntryKey*>(ConstEntries.GetData()), ConstEntries.Num());
+	
+	if (Entries.IsEmpty())
+	{
+		ExtractedEntries = { PCGInvalidEntryKey };
+		Entries = MakeArrayView(ExtractedEntries);
+	}
+}
+
+FPCGAttributeAccessorKeysSplineDataEntries::FPCGAttributeAccessorKeysSplineDataEntries(UPCGSplineData* InSplineData)
+	: Ptr(InSplineData)
+{
+	check(InSplineData);
+	
+	bIsReadOnly = false;
+	InSplineData->SplineStruct.AllocateMetadataEntries();
+	Entries = InSplineData->SplineStruct.GetMutableControlPointsEntryKeys();
+}
+
+int32 FPCGAttributeAccessorKeysSplineDataEntries::GetNum() const
+{
+	return Ptr ? Ptr->SplineStruct.GetSplinePointsPosition().Points.Num() : 0;
+}
diff --git a/Engine/Plugins/PCG/Source/PCG/Public/Data/PCGPolyLineData.h b/Engine/Plugins/PCG/Source/PCG/Public/Data/PCGPolyLineData.h
--- a/Engine/Plugins/PCG/Source/PCG/Public/Data/PCGPolyLineData.h
+++ b/Engine/Plugins/PCG/Source/PCG/Public/Data/PCGPolyLineData.h
@@ -68,10 +68,9 @@ public:
 
 	/** True if the line is a closed loop. */
 	virtual bool IsClosed() const { return false; }
-
-	/** If a PolyLine subtype has custom metadata, it can use this virtual call to write that into the PCG Metadata per-point. */
+	
+	/** This function should be called in the Sample/Project point function, but can also be called if the sampling is done manually.
+	 * This is meant for each child class to write its metadata given the InputKey, so interpolation can also be done.
+	 */
 	virtual void WriteMetadataToPoint(float InputKey, FPCGPoint& OutPoint, UPCGMetadata* OutMetadata) const {}
-
-	/** True if the line holds custom metadata. */
-	virtual bool HasCustomMetadata() const { return false; }
 };
diff --git a/Engine/Plugins/PCG/Source/PCG/Public/Data/PCGSplineData.h b/Engine/Plugins/PCG/Source/PCG/Public/Data/PCGSplineData.h
--- a/Engine/Plugins/PCG/Source/PCG/Public/Data/PCGSplineData.h
+++ b/Engine/Plugins/PCG/Source/PCG/Public/Data/PCGSplineData.h
@@ -27,7 +27,7 @@ class PCG_API UPCGSplineData : public UPCGPolyLineData
 
 public:
 	void Initialize(const USplineComponent* InSpline);
-	void Initialize(const TArray<FSplinePoint>& InSplinePoints, bool bInClosedLoop, const FTransform& InTransform);
+	void Initialize(const TArray<FSplinePoint>& InSplinePoints, bool bInClosedLoop, const FTransform& InTransform, TArray<PCGMetadataEntryKey> InOptionalEntryKeys = {});
 	void Initialize(const FPCGSplineStruct& InSplineStruct);
 	void ApplyTo(USplineComponent* InSpline) const;
 
@@ -49,6 +49,7 @@ public:
 	virtual FVector GetLocationAtAlpha(float Alpha) const override;
 	virtual FTransform GetTransformAtAlpha(float Alpha) const override;
 	virtual bool IsClosed() const override { return SplineStruct.bClosedLoop; }
+	virtual void WriteMetadataToPoint(float InputKey, FPCGPoint& OutPoint, UPCGMetadata* OutMetadata) const;
 	//~End UPCGPolyLineData interface
 	
 	/** Static helper to create an accessor on a data that doesn't yet exist, as accessors for spline data don't rely on existing data. */
diff --git a/Engine/Plugins/PCG/Source/PCG/Public/Data/PCGSplineStruct.h b/Engine/Plugins/PCG/Source/PCG/Public/Data/PCGSplineStruct.h
--- a/Engine/Plugins/PCG/Source/PCG/Public/Data/PCGSplineStruct.h
+++ b/Engine/Plugins/PCG/Source/PCG/Public/Data/PCGSplineStruct.h
@@ -2,6 +2,8 @@
 
 #pragma once
 
+#include "Metadata/PCGMetadataCommon.h"
+
 #include "CoreMinimal.h"
 #include "Components/SplineComponent.h"
 
@@ -25,7 +27,7 @@ struct PCG_API FPCGSplineStruct
 	GENERATED_BODY()
 
 	void Initialize(const USplineComponent* InSplineComponent);
-	void Initialize(const TArray<FSplinePoint>& InSplinePoints, bool bIsClosedLoop, const FTransform& InTransform);
+	void Initialize(const TArray<FSplinePoint>& InSplinePoints, bool bIsClosedLoop, const FTransform& InTransform, TArray<PCGMetadataEntryKey> InOptionalEntryKeys = {});
 	void ApplyTo(USplineComponent* InSplineComponent) const;
 
 	FTransform GetTransform() const { return Transform; }
@@ -41,7 +43,9 @@ struct PCG_API FPCGSplineStruct
 	FBox GetBounds() const;
 
 	const FInterpCurveVector& GetSplinePointsScale() const;
+	const FInterpCurveQuat& GetSplinePointsRotation() const;
 	const FInterpCurveVector& GetSplinePointsPosition() const;
+	const FInterpCurveFloat& GetSplineRepramTable() const;
 
 	FVector GetRightVectorAtSplineInputKey(float InKey, ESplineCoordinateSpace::Type CoordinateSpace) const;
 	FTransform GetTransformAtSplineInputKey(float InKey, ESplineCoordinateSpace::Type CoordinateSpace, bool bUseScale = false) const;
@@ -59,12 +63,14 @@ struct PCG_API FPCGSplineStruct
 	/** Given a threshold, returns a list of vertices along the spline that, treated as a list of segments (polyline), matches the spline shape. Taken from USplineComponent. */
 	bool ConvertSplineToPolyLine(ESplineCoordinateSpace::Type CoordinateSpace, const float MaxSquareDistanceFromSpline, TArray<FVector>& OutPoints) const;
 	
-	float FindInputKeyClosestToWorldLocation(const FVector& WorldLocation) const;	
+	float FindInputKeyClosestToWorldLocation(const FVector& WorldLocation) const;
 
-	UPROPERTY()
-	FSplineCurves SplineCurves;
-
-	// Replaces the component trasnform
+	/** To be called at any moment we want to access individual metadata, like with accessors. Does nothing if it was already allocated. */
+	void AllocateMetadataEntries();
+	TConstArrayView<PCGMetadataEntryKey> GetConstControlPointsEntryKeys() const { return ControlPointsEntryKeys; }
+	TArrayView<PCGMetadataEntryKey> GetMutableControlPointsEntryKeys() { return ControlPointsEntryKeys; }
+	
+	// Replaces the component transform
 	UPROPERTY()
 	FTransform Transform = FTransform::Identity;
 
@@ -83,6 +89,13 @@ struct PCG_API FPCGSplineStruct
 	UPROPERTY()
 	FBoxSphereBounds Bounds = FBoxSphereBounds(EForceInit::ForceInit);
 
+protected:
+	UPROPERTY()
+	FSplineCurves SplineCurves;
+	
+	UPROPERTY()
+	TArray<int64> ControlPointsEntryKeys; // Needs to be int64 for UHT, but it is a PCGMetadataEntryKey
+
 private:
 	// Internal helper function called by ConvertSplineSegmentToPolyLine -- assumes the input is within a half-segment, so testing the distance to midpoint will be an accurate guide to subdivision. Taken from USplineComponent.
 	bool DivideSplineIntoPolylineRecursiveWithDistancesHelper(float StartDistanceAlongSpline, float EndDistanceAlongSpline, ESplineCoordinateSpace::Type CoordinateSpace, const float MaxSquareDistanceFromSpline, TArray<FVector>& OutPoints, TArray<double>& OutDistancesAlongSpline) const;
diff --git a/Engine/Plugins/PCG/Source/PCG/Public/Elements/PCGAttributeNoise.h b/Engine/Plugins/PCG/Source/PCG/Public/Elements/PCGAttributeNoise.h
--- a/Engine/Plugins/PCG/Source/PCG/Public/Elements/PCGAttributeNoise.h
+++ b/Engine/Plugins/PCG/Source/PCG/Public/Elements/PCGAttributeNoise.h
@@ -88,6 +88,12 @@ public:
 	UPROPERTY(BlueprintReadWrite, EditAnywhere, Category = Settings, meta = (PCG_Overridable))
 	bool bClampResult = false;
 
+	UPROPERTY(BlueprintReadWrite, EditAnywhere, Category = Settings, meta = (PCG_Overridable, InlineEditConditionToggle))
+	bool bHasCustomSeedSource = false;
+	
+	UPROPERTY(BlueprintReadWrite, EditAnywhere, Category = Settings, meta = (PCG_Overridable, EditCondition = "bHasCustomSeedSource"))
+	FPCGAttributePropertyInputSelector CustomSeedSource;
+
 	// Previous names
 #if WITH_EDITORONLY_DATA
 	UPROPERTY()
@@ -116,10 +122,13 @@ struct FPCGAttributeNoiseContext : public FPCGContext
 	int32 CurrentInput = 0;
 	bool bDataPreparedForCurrentInput = false;
 	FPCGAttributePropertyInputSelector InputSource;
+	FPCGAttributePropertyInputSelector CustomSeedSource;
 	FPCGAttributePropertyOutputSelector OutputTarget;
 	TUniquePtr<const IPCGAttributeAccessor> InputAccessor;
+	TUniquePtr<const IPCGAttributeAccessor> CustomSeedAccessor;
 	TUniquePtr<IPCGAttributeAccessor> OutputAccessor;
 	TUniquePtr<const IPCGAttributeAccessorKeys> InputKeys;
+	TUniquePtr<const IPCGAttributeAccessorKeys> CustomSeedKeys;
 	TUniquePtr<IPCGAttributeAccessorKeys> OutputKeys;
 };
 
diff --git a/Engine/Plugins/PCG/Source/PCG/Public/Metadata/Accessors/PCGAttributeAccessorKeys.h b/Engine/Plugins/PCG/Source/PCG/Public/Metadata/Accessors/PCGAttributeAccessorKeys.h
--- a/Engine/Plugins/PCG/Source/PCG/Public/Metadata/Accessors/PCGAttributeAccessorKeys.h
+++ b/Engine/Plugins/PCG/Source/PCG/Public/Metadata/Accessors/PCGAttributeAccessorKeys.h
@@ -189,6 +189,9 @@ public:
 	virtual int32 GetNum() const override { return Entries.Num(); }
 
 protected:
+	/** For subclasses that have their own initialization logic, but protected to not make it default constructible. */
+	FPCGAttributeAccessorKeysEntries();
+
 	virtual bool GetMetadataEntryKeys(int32 InStart, TArrayView<PCGMetadataEntryKey*> OutEntryKeys) override;
 	virtual bool GetMetadataEntryKeys(int32 InStart, TArrayView<const PCGMetadataEntryKey*> OutEntryKeys) const override;
 
@@ -430,6 +433,21 @@ public:
 	virtual int32 GetNum() const override;
 };
 
+/**
+* Keys for metadata on a spline data
+*/
+class FPCGAttributeAccessorKeysSplineDataEntries : public FPCGAttributeAccessorKeysEntries
+{
+public:
+	explicit FPCGAttributeAccessorKeysSplineDataEntries(const UPCGSplineData* InSplineData);
+	explicit FPCGAttributeAccessorKeysSplineDataEntries(UPCGSplineData* InSplineData);
+
+	virtual int32 GetNum() const override;
+
+private:
+	const UPCGSplineData* Ptr = nullptr;
+};
+
 template <typename ObjectType>
 inline bool IPCGAttributeAccessorKeys::GetKeys(int32 InStart, TArrayView<ObjectType*> OutKeys)
 {
diff --git a/Engine/Plugins/PCG/Source/PCGEditor/Private/DataVisualizations/PCGSplineDataVisualization.cpp b/Engine/Plugins/PCG/Source/PCGEditor/Private/DataVisualizations/PCGSplineDataVisualization.cpp
--- a/Engine/Plugins/PCG/Source/PCGEditor/Private/DataVisualizations/PCGSplineDataVisualization.cpp
+++ b/Engine/Plugins/PCG/Source/PCGEditor/Private/DataVisualizations/PCGSplineDataVisualization.cpp
@@ -27,8 +27,8 @@ FPCGTableVisualizerInfo IPCGSplineDataVisualization::GetTableVisualizerInfo(cons
 			}
 			else
 			{
-				const TArray<FInterpCurvePointVector>& Positions = SplineData->SplineStruct.SplineCurves.Position.Points;
-				const TArray<FInterpCurvePointVector>& Scales = SplineData->SplineStruct.SplineCurves.Scale.Points;
+				const TArray<FInterpCurvePointVector>& Positions = SplineData->SplineStruct.GetSplinePointsPosition().Points;
+				const TArray<FInterpCurvePointVector>& Scales = SplineData->SplineStruct.GetSplinePointsScale().Points;
 
 				for (const int& Index : Indices)
 				{
@@ -60,13 +60,15 @@ const UPCGPointData* IPCGSplineDataVisualization::CollapseToDebugPointData(FPCGC
 		PointData->InitializeFromData(SplineData);
 		TArray<FPCGPoint>& Points = PointData->GetMutablePoints();
 
-		const int32 NumControlPoints = SplineData->SplineStruct.SplineCurves.Position.Points.Num();
+		const int32 NumControlPoints = SplineData->SplineStruct.GetSplinePointsPosition().Points.Num();
+		TConstArrayView<PCGMetadataEntryKey> EntryKeys = SplineData->SplineStruct.GetConstControlPointsEntryKeys();
 
 		for (int ControlPointIndex = 0; ControlPointIndex < NumControlPoints; ++ControlPointIndex)
 		{
 			FPCGPoint& Point = Points.Emplace_GetRef();
 			Point.Transform = SplineData->SplineStruct.GetTransformAtSplineInputKey(ControlPointIndex, ESplineCoordinateSpace::World);
 			Point.SetExtents(PCGSplineDataVisualizationConstants::HalfExtents);
+			Point.MetadataEntry = EntryKeys.IsValidIndex(ControlPointIndex) ? EntryKeys[ControlPointIndex] : PCGInvalidEntryKey;
 		}
 
 		return PointData;
