[PCG] Convert more nodes to support PCGPointArrayData
- Introduce FPCGPointValueRanges / FConstPCGPointValueRanges to simplify boiler plate code when algo requires most ranges
- PCGValue range now doesn't allow writing wrap around (modulo) and instead provides a stub value to write to (which is local to range), this prevents multithreaded writes to the same ViewIndex for Single value ranges / non allocated properties

#rnx

diff --git a/Engine/Plugins/PCG/Source/PCG/Private/Elements/Metadata/PCGMetadataOpElementBase.cpp b/Engine/Plugins/PCG/Source/PCG/Private/Elements/Metadata/PCGMetadataOpElementBase.cpp
--- a/Engine/Plugins/PCG/Source/PCG/Private/Elements/Metadata/PCGMetadataOpElementBase.cpp
+++ b/Engine/Plugins/PCG/Source/PCG/Private/Elements/Metadata/PCGMetadataOpElementBase.cpp
@@ -6,7 +6,7 @@
 #include "PCGCustomVersion.h"
 #include "PCGParamData.h"
 #include "PCGPin.h"
-#include "Data/PCGPointData.h"
+#include "Data/PCGBasePointData.h"
 #include "Elements/Metadata/PCGMetadataElementCommon.h"
 #include "Metadata/Accessors/PCGAttributeAccessorHelpers.h"
 
@@ -510,9 +510,9 @@ bool FPCGMetadataElementBase::PrepareDataInternal(FPCGContext* Context) const
 			InputTaggedData[OperandPinIndex] = CurrentPinInputData.Num() == 1 ? CurrentPinInputData[0] : MoveTemp(CurrentPinInputData[IterationIndex]);
 
 			// Check if we have any points
-			if (const UPCGPointData* PointInput = Cast<const UPCGPointData>(InputTaggedData[OperandPinIndex].Data))
+			if (const UPCGBasePointData* PointInput = Cast<const UPCGBasePointData>(InputTaggedData[OperandPinIndex].Data))
 			{
-				if (PointInput->GetPoints().IsEmpty())
+				if (PointInput->GetNumPoints() == 0)
 				{
 					// If we have no points, there is no operation
 					PCGE_LOG(Verbose, LogOnly, FText::Format(LOCTEXT("NoPointsForPin", "No points in point data provided on pin {0}"), FText::FromName(CurrentPinLabel)));
diff --git a/Engine/Plugins/PCG/Source/PCG/Private/Elements/PCGAttributeFilter.cpp b/Engine/Plugins/PCG/Source/PCG/Private/Elements/PCGAttributeFilter.cpp
--- a/Engine/Plugins/PCG/Source/PCG/Private/Elements/PCGAttributeFilter.cpp
+++ b/Engine/Plugins/PCG/Source/PCG/Private/Elements/PCGAttributeFilter.cpp
@@ -6,7 +6,7 @@
 #include "PCGCustomVersion.h"
 #include "PCGParamData.h"
 #include "PCGPin.h"
-#include "Data/PCGPointData.h"
+#include "Data/PCGBasePointData.h"
 #include "Data/PCGSpatialData.h"
 #include "Helpers/PCGHelpers.h"
 #include "Metadata/Accessors/IPCGAttributeAccessorTpl.h"
@@ -100,7 +100,7 @@ namespace PCGAttributeFilterHelpers
 		bool bUseInputDataForThreshold = false;
 		bool bUseSpatialQuery = false;
 
-		UPCGPointData* ThresholdPointData = nullptr;
+		UPCGBasePointData* ThresholdPointData = nullptr;
 		const UPCGSpatialData* ThresholdSpatialData = nullptr;
 	};
 
@@ -130,7 +130,7 @@ namespace PCGAttributeFilterHelpers
 					// If the threshold is spatial or points (and spatial query is enabled), we'll use spatial query (meaning we'll have to sample points).
 					// Don't create an accessor yet (ThresholdData = nullptr), it will be created further down.
 					OutThresholdInfo.ThresholdSpatialData = ThresholdSpatialData;
-					if (!ThresholdSpatialData->IsA<UPCGPointData>() || ThresholdSettings.bUseSpatialQuery)
+					if (!ThresholdSpatialData->IsA<UPCGBasePointData>() || ThresholdSettings.bUseSpatialQuery)
 					{
 						OutThresholdInfo.bUseSpatialQuery = true;
 						ThresholdData = nullptr;
@@ -176,9 +176,14 @@ namespace PCGAttributeFilterHelpers
 				// Reset the point data and reserving some points
 				// No need to reserve the full number of points, since we'll go by chunk
 				// Only allocate the chunk size
-				InOutThresholdInfo.ThresholdPointData = FPCGContext::NewObject_AnyThread<UPCGPointData>(InContext);
-				InOutThresholdInfo.ThresholdPointData->InitializeFromData(InOutThresholdInfo.ThresholdSpatialData);
-				InOutThresholdInfo.ThresholdPointData->GetMutablePoints().SetNum(PCGAttributeFilterConstants::ChunkSize);
+				InOutThresholdInfo.ThresholdPointData = FPCGContext::NewPointData_AnyThread(InContext);
+
+				FPCGInitializeFromDataParams InitializeFromDataParams(InOutThresholdInfo.ThresholdSpatialData);
+				InitializeFromDataParams.bInheritSpatialData = false;
+
+				InOutThresholdInfo.ThresholdPointData->InitializeFromDataWithParams(InitializeFromDataParams);
+				InOutThresholdInfo.ThresholdPointData->SetNumPoints(PCGAttributeFilterConstants::ChunkSize);
+				InOutThresholdInfo.ThresholdPointData->AllocateProperties(EPCGPointNativeProperties::All);
 			}
 
 			// Accessor will be valid, but keys will point to default points. But since it is a view, it will be updated when we sample the points.
@@ -533,9 +538,9 @@ bool FPCGAttributeFilterElementBase::DoFiltering(FPCGContext* Context, EPCGAttri
 
 	struct FOperationData
 	{
-		const TArray<FPCGPoint>* OriginalPoints = nullptr;
-		TArray<FPCGPoint>* InFilterPoints = nullptr;
-		TArray<FPCGPoint>* OutFilterPoints = nullptr;
+		const UPCGBasePointData* OriginalPointData = nullptr;
+		UPCGBasePointData* InFilterPointData = nullptr;
+		UPCGBasePointData* OutFilterPointData = nullptr;
 		TBitArray<TInlineAllocator<2048>> FilterBitArray;
 
 		const UPCGMetadata* OriginalMetadata = nullptr;
@@ -588,13 +593,13 @@ bool FPCGAttributeFilterElementBase::DoFiltering(FPCGContext* Context, EPCGAttri
 
 		if (const UPCGSpatialData* SpatialInput = Cast<const UPCGSpatialData>(OriginalData))
 		{
-			if (!SpatialInput->IsA<UPCGPointData>())
+			if (!SpatialInput->IsA<UPCGBasePointData>())
 			{
-				const UPCGPointData* OriginalPointData = nullptr;
+				const UPCGBasePointData* OriginalPointData = nullptr;
 
 				if (bHasSpatialToPointDeprecation)
 				{
-					OriginalPointData = SpatialInput->ToPointData(Context);
+					OriginalPointData = SpatialInput->ToBasePointData(Context);
 				}
 
 				if (!OriginalPointData)
@@ -686,19 +691,23 @@ bool FPCGAttributeFilterElementBase::DoFiltering(FPCGContext* Context, EPCGAttri
 
 		if (OperationData.bIsInputPointData)
 		{
-			const UPCGPointData* OriginalPointData = CastChecked<UPCGPointData>(OriginalData);
-			UPCGPointData* InFilterPointData = FPCGContext::NewObject_AnyThread<UPCGPointData>(Context);
-			UPCGPointData* OutFilterPointData = FPCGContext::NewObject_AnyThread<UPCGPointData>(Context);
+			const UPCGBasePointData* OriginalPointData = CastChecked<UPCGBasePointData>(OriginalData);
+			UPCGBasePointData* InFilterPointData = FPCGContext::NewPointData_AnyThread(Context);
+			UPCGBasePointData* OutFilterPointData = FPCGContext::NewPointData_AnyThread(Context);
+
+			OperationData.OriginalPointData = OriginalPointData;
+
+			FPCGInitializeFromDataParams InitializeFromDataParams(OriginalPointData);
+			InitializeFromDataParams.bInheritSpatialData = false;
 
-			OperationData.OriginalPoints = &OriginalPointData->GetPoints();
+			InFilterPointData->InitializeFromDataWithParams(InitializeFromDataParams);
+			OutFilterPointData->InitializeFromDataWithParams(InitializeFromDataParams);
 
-			InFilterPointData->InitializeFromData(OriginalPointData);
-			OutFilterPointData->InitializeFromData(OriginalPointData);
-			OperationData.InFilterPoints = &InFilterPointData->GetMutablePoints();
-			OperationData.OutFilterPoints = &OutFilterPointData->GetMutablePoints();
+			OperationData.InFilterPointData = InFilterPointData;
+			OperationData.OutFilterPointData = OutFilterPointData;
 
 			// Will be set individually in batches
-			OperationData.FilterBitArray.SetNumUninitialized(OriginalPointData->GetPoints().Num());
+			OperationData.FilterBitArray.SetNumUninitialized(OriginalPointData->GetNumPoints());
 
 			InFilterData = InFilterPointData;
 			OutFilterData = OutFilterPointData;
@@ -767,22 +776,24 @@ bool FPCGAttributeFilterElementBase::DoFiltering(FPCGContext* Context, EPCGAttri
 				}
 
 				// Sampling the points if needed
-				auto SamplePointData = [Range, StartIndex, &OperationData, &SkipTests](UPCGPointData* InPointData, const UPCGSpatialData* InSpatialData)
+				auto SamplePointData = [Range, StartIndex, &OperationData, &SkipTests](UPCGBasePointData* InPointData, const UPCGSpatialData* InSpatialData)
 				{
 					if (InPointData != nullptr)
 					{
+						const FConstPCGPointValueRanges OriginalRanges(OperationData.OriginalPointData);
+						FPCGPointValueRanges ThresholdRanges(InPointData, /*bAllocate=*/false);
+
 						// Threshold points only have "ChunkSize" points.
-						TArray<FPCGPoint>& ThresholdPoints = InPointData->GetMutablePoints();
-						check(OperationData.OriginalPoints);
 						for (int32 j = 0; j < Range; ++j)
 						{
 							FPCGPoint ThresholdPoint;
-							const FPCGPoint& SourcePoint = (*OperationData.OriginalPoints)[StartIndex + j];
+							const FPCGPoint SourcePoint = OriginalRanges.GetPoint(StartIndex + j);
+
 							// If we already mark this point to skip test, don't even try to sample.
 							// If the sample fails, mark the point to skip test.
 							if (!SkipTests[j] && InSpatialData->SamplePoint(SourcePoint.Transform, SourcePoint.GetLocalBounds(), ThresholdPoint, InPointData->Metadata))
 							{
-								ThresholdPoints[j] = ThresholdPoint;
+								ThresholdRanges.SetFromPoint(j, ThresholdPoint);
 							}
 							else
 							{
@@ -835,23 +846,32 @@ bool FPCGAttributeFilterElementBase::DoFiltering(FPCGContext* Context, EPCGAttri
 		{
 			if (OperationData.bIsInputPointData)
 			{
-				check(OperationData.OriginalPoints);
-
 				const int32 NumInFilterPoints = OperationData.FilterBitArray.CountSetBits();
-				const int32 NumOutFilterPoints = OperationData.OriginalPoints->Num() - NumInFilterPoints;
+				const int32 NumOutFilterPoints = OperationData.OriginalPointData->GetNumPoints() - NumInFilterPoints;
+
+				int32 InFilterWriteIndex = 0;
+				OperationData.InFilterPointData->SetNumPoints(NumInFilterPoints);
+				OperationData.InFilterPointData->AllocateProperties(OperationData.OriginalPointData->GetAllocatedProperties());
 
-				OperationData.InFilterPoints->Reserve(NumInFilterPoints);
-				OperationData.OutFilterPoints->Reserve(NumOutFilterPoints);
+				int32 OutFilterWriteIndex = 0;
+				OperationData.OutFilterPointData->SetNumPoints(NumOutFilterPoints);
+				OperationData.OutFilterPointData->AllocateProperties(OperationData.OriginalPointData->GetAllocatedProperties());
 
+				const FConstPCGPointValueRanges OriginalRanges(OperationData.OriginalPointData);
+				FPCGPointValueRanges InFilterRanges(OperationData.InFilterPointData, /*bAllocate=*/false);
+				FPCGPointValueRanges OutFilterRanges(OperationData.OutFilterPointData, /*bAllocate=*/false);
+								
 				for (int32 Index = 0; Index < OperationData.FilterBitArray.Num(); ++Index)
 				{
 					if (OperationData.FilterBitArray[Index])
 					{
-						OperationData.InFilterPoints->Add((*OperationData.OriginalPoints)[Index]);
+						InFilterRanges.SetFromValueRanges(InFilterWriteIndex, OriginalRanges, Index);						
+						++InFilterWriteIndex;
 					}
 					else
 					{
-						OperationData.OutFilterPoints->Add((*OperationData.OriginalPoints)[Index]);
+						OutFilterRanges.SetFromValueRanges(OutFilterWriteIndex, OriginalRanges, Index);
+						++OutFilterWriteIndex;
 					}
 				}
 			}
diff --git a/Engine/Plugins/PCG/Source/PCG/Private/Elements/PCGAttributeGetFromPointIndexElement.cpp b/Engine/Plugins/PCG/Source/PCG/Private/Elements/PCGAttributeGetFromPointIndexElement.cpp
--- a/Engine/Plugins/PCG/Source/PCG/Private/Elements/PCGAttributeGetFromPointIndexElement.cpp
+++ b/Engine/Plugins/PCG/Source/PCG/Private/Elements/PCGAttributeGetFromPointIndexElement.cpp
@@ -104,7 +104,7 @@ bool FPCGAttributeGetFromPointIndexElement::ExecuteInternal(FPCGContext* Context
 	for (int32 InputIndex = 0; InputIndex < Inputs.Num(); ++InputIndex)
 	{
 		const FPCGTaggedData& Input = Inputs[InputIndex];
-		const UPCGPointData* PointData = Cast<UPCGPointData>(Input.Data);
+		const UPCGBasePointData* PointData = Cast<UPCGBasePointData>(Input.Data);
 
 		if (!PointData)
 		{
@@ -114,14 +114,13 @@ bool FPCGAttributeGetFromPointIndexElement::ExecuteInternal(FPCGContext* Context
 
 		const int32 Index = Settings->Index;
 
-		if (Index < 0 || Index >= PointData->GetPoints().Num())
+		if (Index < 0 || Index >= PointData->GetNumPoints())
 		{
-			PCGE_LOG(Error, GraphAndLog, FText::Format(LOCTEXT("IndexOutOfBounds", "Index for input {0} is out of bounds. Index: {1}; Number of Points: {2}"), InputIndex, Index, PointData->GetPoints().Num()));
+			PCGE_LOG(Error, GraphAndLog, FText::Format(LOCTEXT("IndexOutOfBounds", "Index for input {0} is out of bounds. Index: {1}; Number of Points: {2}"), InputIndex, Index, PointData->GetNumPoints()));
 			continue;
 		}
 
 		TArray<FPCGTaggedData>& Outputs = Context->OutputData.TaggedData;
-		const FPCGPoint& Point = PointData->GetPoints()[Index];
 
 		const FPCGAttributePropertyInputSelector InputSource = Settings->InputSource.CopyAndFixLast(PointData);
 		const FPCGAttributePropertyOutputSelector OutputTarget = Settings->OutputAttributeName.CopyAndFixSource(&InputSource);
@@ -129,7 +128,7 @@ bool FPCGAttributeGetFromPointIndexElement::ExecuteInternal(FPCGContext* Context
 		const FName OutputAttributeName = OutputTarget.GetName();
 
 		TUniquePtr<const IPCGAttributeAccessor> Accessor = PCGAttributeAccessorHelpers::CreateConstAccessor(PointData, InputSource);
-		FPCGAttributeAccessorKeysPoints PointKey(Point);
+		FPCGAttributeAccessorKeysPointsSubset PointKey(PointData, { Index });
 
 		if (Accessor.IsValid())
 		{
@@ -169,9 +168,15 @@ bool FPCGAttributeGetFromPointIndexElement::ExecuteInternal(FPCGContext* Context
 				if (Context->Node && Context->Node->IsOutputPinConnected(PCGAttributeGetFromPointIndexConstants::OutputPointLabel))
 #endif
 				{
-					UPCGPointData* OutputPointData = FPCGContext::NewObject_AnyThread<UPCGPointData>(Context);
-					OutputPointData->InitializeFromData(PointData);
-					OutputPointData->GetMutablePoints().Add(Point);
+					UPCGBasePointData* OutputPointData = FPCGContext::NewPointData_AnyThread(Context);
+
+					FPCGInitializeFromDataParams InitializeFromDataParams(PointData);
+					InitializeFromDataParams.bInheritSpatialData = false;
+					
+					OutputPointData->InitializeFromDataWithParams(InitializeFromDataParams);
+					OutputPointData->SetNumPoints(1);
+
+					PointData->CopyPointsTo(OutputPointData, Index, 0, 1);
 
 					FPCGTaggedData& OutputPoint = Outputs.Add_GetRef(Input);
 					OutputPoint.Data = OutputPointData;
diff --git a/Engine/Plugins/PCG/Source/PCG/Private/Elements/PCGAttributeNoise.cpp b/Engine/Plugins/PCG/Source/PCG/Private/Elements/PCGAttributeNoise.cpp
--- a/Engine/Plugins/PCG/Source/PCG/Private/Elements/PCGAttributeNoise.cpp
+++ b/Engine/Plugins/PCG/Source/PCG/Private/Elements/PCGAttributeNoise.cpp
@@ -5,7 +5,7 @@
 #include "PCGCommon.h"
 #include "PCGContext.h"
 #include "PCGCustomVersion.h"
-#include "Data/PCGPointData.h"
+#include "Data/PCGBasePointData.h"
 #include "Helpers/PCGAsync.h"
 #include "Helpers/PCGHelpers.h"
 #include "Metadata/PCGMetadataAttribute.h"
@@ -258,7 +258,7 @@ bool FPCGAttributeNoiseElement::ExecuteInternal(FPCGContext* InContext) const
 			{
 				if (Settings->bHasSpatialToPointDeprecation)
 				{
-					InputData = InputSpatialData->ToPointData();
+					InputData = InputSpatialData->ToBasePointData(Context);
 				}
 			}
 
@@ -282,7 +282,7 @@ bool FPCGAttributeNoiseElement::ExecuteInternal(FPCGContext* InContext) const
 
 			// Same verification for the seed. If we have no custom seed but points, use the $Seed accessor.
 			// Otherwise, it will just use the index of the element as the seed (combined with component + settings seed of course)
-			const UPCGPointData* InputPointData = Cast<UPCGPointData>(Input.Data);
+			const UPCGBasePointData* InputPointData = Cast<UPCGBasePointData>(Input.Data);
 			if (Settings->bHasCustomSeedSource || InputPointData)
 			{
 				if (Settings->bHasCustomSeedSource)
diff --git a/Engine/Plugins/PCG/Source/PCG/Private/Elements/PCGCollapseElement.cpp b/Engine/Plugins/PCG/Source/PCG/Private/Elements/PCGCollapseElement.cpp
--- a/Engine/Plugins/PCG/Source/PCG/Private/Elements/PCGCollapseElement.cpp
+++ b/Engine/Plugins/PCG/Source/PCG/Private/Elements/PCGCollapseElement.cpp
@@ -4,7 +4,7 @@
 
 #include "PCGContext.h"
 #include "PCGParamData.h"
-#include "Data/PCGPointData.h"
+#include "Data/PCGBasePointData.h"
 #include "Data/PCGSpatialData.h"
 
 #include UE_INLINE_GENERATED_CPP_BY_NAME(PCGCollapseElement)
@@ -67,7 +67,7 @@ bool FPCGCollapseElement::ExecuteInternal(FPCGContext* InContext) const
 		if (const UPCGSpatialData* SpatialData = Cast<UPCGSpatialData>(Input.Data))
 		{
 			// Currently we support collapsing to point data only, but at some point in the future that might be different
-			Output.Data = Cast<const UPCGSpatialData>(Input.Data)->ToPointData(Context);
+			Output.Data = Cast<const UPCGSpatialData>(Input.Data)->ToBasePointData(Context);
 
 			if (Output.Data != Input.Data)
 			{
@@ -84,16 +84,16 @@ bool FPCGCollapseElement::ExecuteInternal(FPCGContext* InContext) const
 				continue;
 			}
 
-			UPCGPointData* PointData = FPCGContext::NewObject_AnyThread<UPCGPointData>(Context);
+			UPCGBasePointData* PointData = FPCGContext::NewPointData_AnyThread(Context);
 			check(PointData->Metadata);
 			PointData->Metadata->Initialize(ParamMetadata);
-
-			TArray<FPCGPoint>& Points = PointData->GetMutablePoints();
-			Points.SetNum(ParamItemCount);
+			PointData->SetNumPoints(ParamItemCount);
+			PointData->AllocateProperties(EPCGPointNativeProperties::MetadataEntry);
+			TPCGValueRange<int64> MetadataEntryRange = PointData->GetMetadataEntryValueRange(/*bAllocate=*/false);
 
 			for (int PointIndex = 0; PointIndex < ParamItemCount; ++PointIndex)
 			{
-				Points[PointIndex].MetadataEntry = PointIndex;
+				MetadataEntryRange[PointIndex] = PointIndex;
 			}
 
 			Output.Data = PointData;
diff --git a/Engine/Plugins/PCG/Source/PCG/Private/Elements/PCGConvertToAttributeSet.cpp b/Engine/Plugins/PCG/Source/PCG/Private/Elements/PCGConvertToAttributeSet.cpp
--- a/Engine/Plugins/PCG/Source/PCG/Private/Elements/PCGConvertToAttributeSet.cpp
+++ b/Engine/Plugins/PCG/Source/PCG/Private/Elements/PCGConvertToAttributeSet.cpp
@@ -4,7 +4,7 @@
 
 #include "PCGContext.h"
 #include "PCGParamData.h"
-#include "Data/PCGPointData.h"
+#include "Data/PCGBasePointData.h"
 #include "Helpers/PCGTagHelpers.h"
 #include "Metadata/PCGMetadata.h"
 #include "Metadata/PCGMetadataAttributeTpl.h"
@@ -47,18 +47,17 @@ bool FPCGConvertToAttributeSetElement::ExecuteInternal(FPCGContext* Context) con
 
 	for (const FPCGTaggedData& Input : Inputs)
 	{
-		if (Cast<UPCGPointData>(Input.Data) == nullptr)
+		if (Cast<UPCGBasePointData>(Input.Data) == nullptr)
 		{
 			continue;
 		}
 
-		const UPCGPointData* PointData = Cast<UPCGPointData>(Input.Data);
-		const TArray<FPCGPoint>& Points = PointData->GetPoints();
+		const UPCGBasePointData* PointData = Cast<UPCGBasePointData>(Input.Data);
 		const UPCGMetadata* SourceMetadata = PointData->Metadata;
 		check(SourceMetadata);
 
 		// Note: this is very similar in idea to UPCGPointData::Flatten
-		if (SourceMetadata->GetAttributeCount() == 0 || Points.Num() == 0)
+		if (SourceMetadata->GetAttributeCount() == 0 || PointData->GetNumPoints() == 0)
 		{
 			continue;
 		}
@@ -69,10 +68,13 @@ bool FPCGConvertToAttributeSetElement::ExecuteInternal(FPCGContext* Context) con
 
 		// Create new entry that will map (or not) to the original parent
 		TArray<PCGMetadataEntryKey> EntryKeys;
-		EntryKeys.Reserve(Points.Num());
-		for (const FPCGPoint& Point : Points)
+
+		const TConstPCGValueRange<int64> MetadataEntryRange = PointData->GetConstMetadataEntryValueRange();
+		check(PointData->GetNumPoints() == MetadataEntryRange.Num());
+		EntryKeys.Reserve(MetadataEntryRange.Num());
+		for (int32 Index = 0; Index < MetadataEntryRange.Num(); ++Index)
 		{
-			EntryKeys.Add(ParamData->Metadata->AddEntry(Point.MetadataEntry));
+			EntryKeys.Add(ParamData->Metadata->AddEntry(MetadataEntryRange[Index]));
 		}
 
 		ParamData->Metadata->FlattenAndCompress(EntryKeys);
diff --git a/Engine/Plugins/PCG/Source/PCG/Private/Elements/PCGCopyPoints.cpp b/Engine/Plugins/PCG/Source/PCG/Private/Elements/PCGCopyPoints.cpp
--- a/Engine/Plugins/PCG/Source/PCG/Private/Elements/PCGCopyPoints.cpp
+++ b/Engine/Plugins/PCG/Source/PCG/Private/Elements/PCGCopyPoints.cpp
@@ -3,7 +3,7 @@
 #include "Elements/PCGCopyPoints.h"
 
 #include "PCGContext.h"
-#include "Data/PCGPointData.h"
+#include "Data/PCGBasePointData.h"
 #include "Data/PCGSpatialData.h"
 #include "Elements/PCGCopyPointsKernel.h"
 #include "Graph/PCGGPUGraphCompilationContext.h"
@@ -122,8 +122,8 @@ bool FPCGCopyPointsElement::ExecuteInternal(FPCGContext* Context) const
 			return true;
 		}
 
-		const UPCGPointData* SourcePointData = SourceSpatialData->ToPointData(Context);
-		const UPCGPointData* TargetPointData = TargetSpatialData->ToPointData(Context);
+		const UPCGBasePointData* SourcePointData = SourceSpatialData->ToBasePointData(Context);
+		const UPCGBasePointData* TargetPointData = TargetSpatialData->ToBasePointData(Context);
 
 		if (!SourcePointData || !TargetPointData)
 		{
@@ -134,12 +134,7 @@ bool FPCGCopyPointsElement::ExecuteInternal(FPCGContext* Context) const
 		const UPCGMetadata* SourcePointMetadata = SourcePointData->Metadata;
 		const UPCGMetadata* TargetPointMetadata = TargetPointData->Metadata;
 
-		const TArray<FPCGPoint>& SourcePoints = SourcePointData->GetPoints();
-		const TArray<FPCGPoint>& TargetPoints = TargetPointData->GetPoints();
-
-
-		UPCGPointData* OutPointData = FPCGContext::NewObject_AnyThread<UPCGPointData>(Context);
-		TArray<FPCGPoint>& OutPoints = OutPointData->GetMutablePoints();
+		UPCGBasePointData* OutPointData = FPCGContext::NewPointData_AnyThread(Context);
 		Output.Data = OutPointData;
 
 		// Make sure that output contains both collection of tags from source and target
@@ -170,7 +165,9 @@ bool FPCGCopyPointsElement::ExecuteInternal(FPCGContext* Context) const
 			RootMetadata = SourcePointMetadata;
 			NonRootMetadata = nullptr;
 
-			OutPointData->InitializeFromData(SourcePointData);
+			FPCGInitializeFromDataParams InitializeFromDataParams(SourcePointData);
+			InitializeFromDataParams.bInheritSpatialData = false;
+			OutPointData->InitializeFromDataWithParams(InitializeFromDataParams);
 		}
 		else if (AttributeInheritance == EPCGCopyPointsMetadataInheritanceMode::TargetOnly)
 		{
@@ -179,7 +176,9 @@ bool FPCGCopyPointsElement::ExecuteInternal(FPCGContext* Context) const
 			RootMetadata = TargetPointMetadata;
 			NonRootMetadata = nullptr;
 
-			OutPointData->InitializeFromData(TargetPointData);
+			FPCGInitializeFromDataParams InitializeFromDataParams(TargetPointData);
+			InitializeFromDataParams.bInheritSpatialData = false;
+			OutPointData->InitializeFromDataWithParams(InitializeFromDataParams);
 		}
 		else if (AttributeInheritance == EPCGCopyPointsMetadataInheritanceMode::SourceFirst || AttributeInheritance == EPCGCopyPointsMetadataInheritanceMode::TargetFirst)
 		{
@@ -194,11 +193,18 @@ bool FPCGCopyPointsElement::ExecuteInternal(FPCGContext* Context) const
 
 			RootMetadata = bInheritMetadataFromSource ? SourcePointMetadata : TargetPointMetadata;
 			NonRootMetadata = bInheritMetadataFromSource ? TargetPointMetadata : SourcePointMetadata;
-			OutPointData->InitializeFromData(bInheritMetadataFromSource ? SourcePointData : TargetPointData);
+
+			FPCGInitializeFromDataParams InitializeFromDataParams(bInheritMetadataFromSource ? SourcePointData : TargetPointData);
+			InitializeFromDataParams.bInheritSpatialData = false;			
+			OutPointData->InitializeFromDataWithParams(InitializeFromDataParams);
 		}
 		else // None
 		{
-			OutPointData->InitializeFromData(SourcePointData, nullptr, /*bInheritMetadata=*/false, /*bInheritAttributes=*/false);
+			FPCGInitializeFromDataParams InitializeFromDataParams(SourcePointData);
+			InitializeFromDataParams.bInheritSpatialData = false;
+			InitializeFromDataParams.bInheritMetadata = false;
+			InitializeFromDataParams.bInheritAttributes = false;
+			OutPointData->InitializeFromDataWithParams(InitializeFromDataParams);
 
 			bProcessMetadata = false;
 			RootMetadata = NonRootMetadata = nullptr;
@@ -215,6 +221,8 @@ bool FPCGCopyPointsElement::ExecuteInternal(FPCGContext* Context) const
 		TArray<TTuple<int64, int64>> AllMetadataEntries;
 		TArray<TArray<TTuple<PCGMetadataEntryKey, PCGMetadataValueKey>>> AttributeValuesToSet;
 
+		const int32 NumPoints = SourcePointData->GetNumPoints() * TargetPointData->GetNumPoints();
+
 		if (bProcessMetadata)
 		{
 			TRACE_CPUPROFILER_EVENT_SCOPE(FPCGCopyPointsElement::Execute::SetupMetadata);
@@ -244,118 +252,195 @@ bool FPCGCopyPointsElement::ExecuteInternal(FPCGContext* Context) const
 
 				for (TArray<TTuple<PCGMetadataEntryKey, PCGMetadataValueKey>>& AttributeValues : AttributeValuesToSet)
 				{
-					AttributeValues.SetNumUninitialized(SourcePoints.Num() * TargetPoints.Num());
+					AttributeValues.SetNumUninitialized(NumPoints);
 				}
 			}
 
 			// Preallocate the metadata entries array if we're going to use it
-			AllMetadataEntries.SetNumUninitialized(SourcePoints.Num() * TargetPoints.Num());
+			AllMetadataEntries.SetNumUninitialized(NumPoints);
 		}
 
-		// Use implicit capture, since we capture a lot
-		FPCGAsync::AsyncPointProcessing(Context, SourcePoints.Num() * TargetPoints.Num(), OutPoints, [&](int32 Index, FPCGPoint& OutPoint)
+		// Properties that might need to be allocated
+		EPCGPointNativeProperties OptionalProperties = EPCGPointNativeProperties::Seed | EPCGPointNativeProperties::Color;
+		if (bProcessMetadata)
 		{
-			const FPCGPoint& SourcePoint = SourcePoints[Index / TargetPoints.Num()];
-			const FPCGPoint& TargetPoint = TargetPoints[Index % TargetPoints.Num()];
-
-			OutPoint = SourcePoint;
-
-			FTransform SourceTransform(SourcePoint.Transform.GetLocation());
-			FTransform TargetTransform(TargetPoint.Transform.GetLocation());
+			OptionalProperties |= EPCGPointNativeProperties::MetadataEntry;
+		}
 
-			// Set Rotation, Scale, and Color based on inheritance mode
-			if (RotationInheritance != EPCGCopyPointsInheritanceMode::Target)
+		auto InitializeFunc = [OutPointData, SourcePointData, TargetPointData, NumPoints, bProcessMetadata, ColorInheritance, SeedInheritance, OptionalProperties]()
+		{
+			OutPointData->SetNumPoints(NumPoints);
+			EPCGPointNativeProperties SourceAllocatedProperties = SourcePointData->GetAllocatedProperties();
+			EPCGPointNativeProperties TargetAllocatedProperties = TargetPointData->GetAllocatedProperties();
+						
+			// Allocate all properties from source except properties that might not need to be
+			EPCGPointNativeProperties PropertiesToAllocate = SourcePointData->GetAllocatedProperties() & ~OptionalProperties;
+
+			// Allocate MetadataEntry
+			if (bProcessMetadata)
 			{
-				SourceTransform.SetRotation(SourcePoint.Transform.GetRotation());
+				PropertiesToAllocate |= EPCGPointNativeProperties::MetadataEntry;
 			}
-
-			if (RotationInheritance != EPCGCopyPointsInheritanceMode::Source || bApplyTargetRotationToPositions)
+									
+			auto GetPropertyToAllocate = [&PropertiesToAllocate, SourceAllocatedProperties, TargetAllocatedProperties](EPCGCopyPointsInheritanceMode Inheritance, EPCGPointNativeProperties Property)
 			{
-				TargetTransform.SetRotation(TargetPoint.Transform.GetRotation());
-			}
+				if (Inheritance == EPCGCopyPointsInheritanceMode::Relative)
+				{
+					return ((SourceAllocatedProperties | TargetAllocatedProperties) & Property);
+				}
+				else if (Inheritance == EPCGCopyPointsInheritanceMode::Source)
+				{
+					return (SourceAllocatedProperties & Property);
+				}
+				else // if (Inheritance == EPCGCopyPointsInheritanceMode::Target)
+				{
+					return (TargetAllocatedProperties & Property);
+				}
+			};
 
-			if (ScaleInheritance != EPCGCopyPointsInheritanceMode::Target)
-			{ 
-				SourceTransform.SetScale3D(SourcePoint.Transform.GetScale3D());
-			}
-			
-			if (ScaleInheritance != EPCGCopyPointsInheritanceMode::Source || bApplyTargetScaleToPositions)
-			{
-				TargetTransform.SetScale3D(TargetPoint.Transform.GetScale3D());
-			}
+			PropertiesToAllocate |= GetPropertyToAllocate(SeedInheritance, EPCGPointNativeProperties::Seed);
+			PropertiesToAllocate |= GetPropertyToAllocate(ColorInheritance, EPCGPointNativeProperties::Color);
 
-			OutPoint.Transform = SourceTransform * TargetTransform;
+			OutPointData->AllocateProperties(PropertiesToAllocate);
+		};
 
-			if (RotationInheritance == EPCGCopyPointsInheritanceMode::Source && bApplyTargetRotationToPositions)
+		const int32 TargetNumPoints = TargetPointData->GetNumPoints();
+		
+		// Use implicit capture, since we capture a lot
+		auto ProcessRangeFunc = [&](int32 StartReadIndex, int32 StartWriteIndex, int32 Count)
+		{
+			int32 NumWritten = 0;
+						
+			const FConstPCGPointValueRanges SourceRanges(SourcePointData);
+			const FConstPCGPointValueRanges TargetRanges(TargetPointData);
+			FPCGPointValueRanges OutRanges(OutPointData, /*bAllocate=*/false);
+												
+			for (int32 Index = StartReadIndex; Index < (StartReadIndex+Count); ++Index)
 			{
-				OutPoint.Transform.SetRotation(SourcePoint.Transform.GetRotation());
-			}
+				const int32 SourceIndex = Index / TargetNumPoints;
+				const int32 TargetIndex = Index % TargetNumPoints;
+								
+				const int32 WriteIndex = StartWriteIndex + NumWritten;
 
-			if (ScaleInheritance == EPCGCopyPointsInheritanceMode::Source && bApplyTargetScaleToPositions)
-			{
-				OutPoint.Transform.SetScale3D(SourcePoint.Transform.GetScale3D());
-			}
+				// Copy properties from source
+				OutRanges.DensityRange[WriteIndex] = SourceRanges.DensityRange[SourceIndex];
+				OutRanges.SteepnessRange[WriteIndex] = SourceRanges.SteepnessRange[SourceIndex];
+				OutRanges.BoundsMinRange[WriteIndex] = SourceRanges.BoundsMinRange[SourceIndex];
+				OutRanges.BoundsMaxRange[WriteIndex] = SourceRanges.BoundsMaxRange[SourceIndex];
 
-			if (ColorInheritance == EPCGCopyPointsInheritanceMode::Relative)
-			{
-				OutPoint.Color = SourcePoint.Color * TargetPoint.Color;
-			}
-			else if (ColorInheritance == EPCGCopyPointsInheritanceMode::Source)
-			{ 
-				OutPoint.Color = SourcePoint.Color;
-			}
-			else // if (ColorInheritance == EPCGCopyPointsInheritanceMode::Target)
-			{ 
-				OutPoint.Color = TargetPoint.Color;
-			}
+				const FTransform& SourcePointTransform = SourceRanges.TransformRange[SourceIndex];
+				const FTransform& TargetPointTransform = TargetRanges.TransformRange[TargetIndex];
 
-			// Set seed based on inheritance mode
-			if (SeedInheritance == EPCGCopyPointsInheritanceMode::Relative)
-			{
-				OutPoint.Seed = PCGHelpers::ComputeSeed(SourcePoint.Seed, TargetPoint.Seed);
-			}
-			else if (SeedInheritance == EPCGCopyPointsInheritanceMode::Target)
-			{
-				OutPoint.Seed = TargetPoint.Seed;
-			}
+				FTransform SourceTransform(SourcePointTransform.GetLocation());
+				FTransform TargetTransform(TargetPointTransform.GetLocation());
 
-			if (bProcessMetadata)
-			{
-				const FPCGPoint* RootPoint = (bInheritMetadataFromSource ? &SourcePoint : &TargetPoint);
+				// Set Rotation, Scale, and Color based on inheritance mode
+				if (RotationInheritance != EPCGCopyPointsInheritanceMode::Target)
+				{
+					SourceTransform.SetRotation(SourcePointTransform.GetRotation());
+				}
 
-				OutPoint.MetadataEntry = OutMetadata->AddEntryPlaceholder();
-				AllMetadataEntries[Index] = TTuple<int64, int64>(OutPoint.MetadataEntry, RootPoint->MetadataEntry);
+				if (RotationInheritance != EPCGCopyPointsInheritanceMode::Source || bApplyTargetRotationToPositions)
+				{
+					TargetTransform.SetRotation(TargetPointTransform.GetRotation());
+				}
+
+				if (ScaleInheritance != EPCGCopyPointsInheritanceMode::Target)
+				{
+					SourceTransform.SetScale3D(SourcePointTransform.GetScale3D());
+				}
+
+				if (ScaleInheritance != EPCGCopyPointsInheritanceMode::Source || bApplyTargetScaleToPositions)
+				{
+					TargetTransform.SetScale3D(TargetPointTransform.GetScale3D());
+				}
+
+				OutRanges.TransformRange[WriteIndex] = SourceTransform * TargetTransform;
+
+				if (RotationInheritance == EPCGCopyPointsInheritanceMode::Source && bApplyTargetRotationToPositions)
+				{
+					OutRanges.TransformRange[WriteIndex].SetRotation(SourcePointTransform.GetRotation());
+				}
+
+				if (ScaleInheritance == EPCGCopyPointsInheritanceMode::Source && bApplyTargetScaleToPositions)
+				{
+					OutRanges.TransformRange[WriteIndex].SetScale3D(SourcePointTransform.GetScale3D());
+				}
 
-				if (NonRootMetadata)
+				if (ColorInheritance == EPCGCopyPointsInheritanceMode::Relative)
 				{
-					const FPCGPoint* NonRootPoint = (bInheritMetadataFromSource ? &TargetPoint : &SourcePoint);
+					OutRanges.ColorRange[WriteIndex] = SourceRanges.ColorRange[SourceIndex] * TargetRanges.ColorRange[TargetIndex];
+				}
+				else if (ColorInheritance == EPCGCopyPointsInheritanceMode::Target)
+				{
+					OutRanges.ColorRange[WriteIndex] = TargetRanges.ColorRange[TargetIndex];
+				}
+				else // if (ColorInheritance == EPCGCopyPointsInheritanceMode::Source)
+				{
+					OutRanges.ColorRange[WriteIndex] = SourceRanges.ColorRange[SourceIndex];
+				}
 
-					// Copy EntryToValue key mappings from NonRootAttributes - no need to do it if the non-root uses the default values
-					if (NonRootPoint->MetadataEntry != PCGInvalidEntryKey)
+				// Set seed based on inheritance mode
+				if (SeedInheritance == EPCGCopyPointsInheritanceMode::Relative)
+				{
+					OutRanges.SeedRange[WriteIndex] = PCGHelpers::ComputeSeed(SourceRanges.SeedRange[SourceIndex], TargetRanges.SeedRange[TargetIndex]);
+				}
+				else if (SeedInheritance == EPCGCopyPointsInheritanceMode::Target)
+				{
+					OutRanges.SeedRange[WriteIndex] = TargetRanges.SeedRange[TargetIndex];
+				}
+				else // if (SeedInheritance == EPCGCopyPointsInheritanceMode::Source)
+				{
+					OutRanges.SeedRange[WriteIndex] = SourceRanges.SeedRange[SourceIndex];
+				}
+
+				if (bProcessMetadata)
+				{
+					const int64 RootMetadataEntry = bInheritMetadataFromSource ? SourceRanges.MetadataEntryRange[SourceIndex] : TargetRanges.MetadataEntryRange[TargetIndex];
+					const int64 NonRootMetadataEntry = bInheritMetadataFromSource ? TargetRanges.MetadataEntryRange[TargetIndex] : SourceRanges.MetadataEntryRange[SourceIndex];
+
+					OutRanges.MetadataEntryRange[WriteIndex] = OutMetadata->AddEntryPlaceholder();
+					AllMetadataEntries[Index] = TTuple<int64, int64>(OutRanges.MetadataEntryRange[WriteIndex], RootMetadataEntry);
+
+					if (NonRootMetadata)
 					{
-						for (int32 AttributeIndex = 0; AttributeIndex < NonRootAttributes.Num(); ++AttributeIndex)
+						// Copy EntryToValue key mappings from NonRootAttributes - no need to do it if the non-root uses the default values
+						if (NonRootMetadataEntry != PCGInvalidEntryKey)
 						{
-							AttributeValuesToSet[AttributeIndex][Index] = TTuple<PCGMetadataEntryKey, PCGMetadataValueKey>(OutPoint.MetadataEntry, NonRootAttributes[AttributeIndex]->GetValueKey(NonRootPoint->MetadataEntry));
+							for (int32 AttributeIndex = 0; AttributeIndex < NonRootAttributes.Num(); ++AttributeIndex)
+							{
+								AttributeValuesToSet[AttributeIndex][Index] = TTuple<PCGMetadataEntryKey, PCGMetadataValueKey>(OutRanges.MetadataEntryRange[WriteIndex], NonRootAttributes[AttributeIndex]->GetValueKey(NonRootMetadataEntry));
+							}
 						}
-					}
-					else
-					{
-						for (int32 AttributeIndex = 0; AttributeIndex < NonRootAttributes.Num(); ++AttributeIndex)
+						else
 						{
-							AttributeValuesToSet[AttributeIndex][Index] = TTuple<PCGMetadataEntryKey, PCGMetadataValueKey>(OutPoint.MetadataEntry, PCGDefaultValueKey);
+							for (int32 AttributeIndex = 0; AttributeIndex < NonRootAttributes.Num(); ++AttributeIndex)
+							{
+								AttributeValuesToSet[AttributeIndex][Index] = TTuple<PCGMetadataEntryKey, PCGMetadataValueKey>(OutRanges.MetadataEntryRange[WriteIndex], PCGDefaultValueKey);
+							}
 						}
 					}
 				}
-			}
-			else
-			{
-				// Reset the metadata entry if we have no metadata.
-				OutPoint.MetadataEntry = PCGInvalidEntryKey;
-			}
+				else
+				{
+					// Reset the metadata entry if we have no metadata.
+					OutRanges.MetadataEntryRange[WriteIndex] = PCGInvalidEntryKey;
+				}
 
-			return true;
-		});
+				++NumWritten;
+			}
 
+			check(NumWritten == Count);
+			return NumWritten;
+		};
+
+		FPCGAsync::AsyncProcessingOneToOneRangeEx(
+			&Context->AsyncState,
+			NumPoints,
+			InitializeFunc,
+			ProcessRangeFunc,
+			/*bTimeSliceEnabled=*/false);
+				
 		if (bProcessMetadata)
 		{
 			TRACE_CPUPROFILER_EVENT_SCOPE(FPCGCopyPointsElement::Execute::SetMetadata);
diff --git a/Engine/Plugins/PCG/Source/PCG/Private/Elements/PCGCreatePointsGrid.cpp b/Engine/Plugins/PCG/Source/PCG/Private/Elements/PCGCreatePointsGrid.cpp
--- a/Engine/Plugins/PCG/Source/PCG/Private/Elements/PCGCreatePointsGrid.cpp
+++ b/Engine/Plugins/PCG/Source/PCG/Private/Elements/PCGCreatePointsGrid.cpp
@@ -279,7 +279,7 @@ bool FPCGCreatePointsGridElement::ExecuteInternal(FPCGContext* Context) const
 				const int32 WriteIndex = StartWriteIndex + NumWritten;
 				TransformRange[WriteIndex] = PointTransform;
 				SeedRange[WriteIndex] = PCGHelpers::ComputeSeedFromPosition(PointTransform.GetLocation());
-				NumWritten++;
+				++NumWritten;
 			}
 		}
 
diff --git a/Engine/Plugins/PCG/Source/PCG/Private/Elements/PCGDensityFilter.cpp b/Engine/Plugins/PCG/Source/PCG/Private/Elements/PCGDensityFilter.cpp
--- a/Engine/Plugins/PCG/Source/PCG/Private/Elements/PCGDensityFilter.cpp
+++ b/Engine/Plugins/PCG/Source/PCG/Private/Elements/PCGDensityFilter.cpp
@@ -4,7 +4,7 @@
 
 #include "Data/PCGSpatialData.h"
 #include "PCGCustomVersion.h"
-#include "Data/PCGPointData.h"
+#include "Data/PCGBasePointData.h"
 #include "Helpers/PCGAsync.h"
 #include "PCGContext.h"
 
@@ -66,7 +66,7 @@ bool FPCGDensityFilterElement::ExecuteInternal(FPCGContext* Context) const
 			continue;
 		}
 
-		const UPCGPointData* OriginalData = Cast<UPCGSpatialData>(Input.Data)->ToPointData(Context);
+		const UPCGBasePointData* OriginalData = Cast<UPCGSpatialData>(Input.Data)->ToBasePointData(Context);
 
 		if (!OriginalData)
 		{
@@ -74,39 +74,86 @@ bool FPCGDensityFilterElement::ExecuteInternal(FPCGContext* Context) const
 			continue;
 		}
 
-		const TArray<FPCGPoint>& Points = OriginalData->GetPoints();
-		
-		UPCGPointData* FilteredData = FPCGContext::NewObject_AnyThread<UPCGPointData>(Context);
-		FilteredData->InitializeFromData(OriginalData);
-		TArray<FPCGPoint>& FilteredPoints = FilteredData->GetMutablePoints();
+		UPCGBasePointData* FilteredData = FPCGContext::NewPointData_AnyThread(Context);
+
+		FPCGInitializeFromDataParams InitializeFromDataParams(OriginalData);
 
+		// Do not inherit because we are going to filter out some points
+		InitializeFromDataParams.bInheritSpatialData = false;
+		FilteredData->InitializeFromDataWithParams(InitializeFromDataParams);
+		
 		Output.Data = FilteredData;
 
-		FPCGAsync::AsyncPointProcessing(Context, Points.Num(), FilteredPoints, [&Points, MinBound, MaxBound, bInvertFilter, bKeepZeroDensityPoints](int32 Index, FPCGPoint& OutPoint)
+		auto InitializeFunc = [FilteredData, OriginalData]()
 		{
-			const FPCGPoint& Point = Points[Index];
+			FilteredData->SetNumPoints(OriginalData->GetNumPoints());
+			FilteredData->AllocateProperties(OriginalData->GetAllocatedProperties());
+		};
 
-			bool bInRange = (Point.Density >= MinBound && Point.Density <= MaxBound);
-			if (bInRange != bInvertFilter)
+		auto AsyncProcessRangeFunc = [FilteredData, OriginalData, MinBound, MaxBound, bInvertFilter, bKeepZeroDensityPoints](int32 StartReadIndex, int32 StartWriteIndex, int32 Count)
+		{
+			const FConstPCGPointValueRanges ReadRanges(OriginalData);
+			FPCGPointValueRanges WriteRanges(FilteredData, /*bAllocate=*/false);
+						
+			int32 NumWritten = 0;
+
+			for (int32 ReadIndex = StartReadIndex; ReadIndex < StartReadIndex + Count; ++ReadIndex)
 			{
-				OutPoint = Point;
-				return true;
-			}
+				const int32 WriteIndex = StartWriteIndex + NumWritten;
+				const float ReadDensity = ReadRanges.DensityRange[ReadIndex];
+				const bool bInRange = (ReadDensity >= MinBound && ReadDensity <= MaxBound);
+				bool bCopyPoint = false;
+			
+				if (bInRange != bInvertFilter)
+				{
+					bCopyPoint = true;
+					WriteRanges.DensityRange[WriteIndex] = ReadRanges.DensityRange[ReadIndex];
+				}
 #if WITH_EDITOR
-			else if (bKeepZeroDensityPoints)
-			{
-				OutPoint = Point;
-				OutPoint.Density = 0;
-				return true;
-			}
+				else if (bKeepZeroDensityPoints)
+				{
+					bCopyPoint = true;
+					WriteRanges.DensityRange[WriteIndex] = 0.f;
+				}
 #endif
-			else
-			{
-				return false;
+
+				if (bCopyPoint)
+				{
+					WriteRanges.TransformRange[WriteIndex] = ReadRanges.TransformRange[ReadIndex];
+					WriteRanges.SteepnessRange[WriteIndex] = ReadRanges.SteepnessRange[ReadIndex];
+					WriteRanges.BoundsMinRange[WriteIndex] = ReadRanges.BoundsMinRange[ReadIndex];
+					WriteRanges.BoundsMaxRange[WriteIndex] = ReadRanges.BoundsMaxRange[ReadIndex];
+					WriteRanges.ColorRange[WriteIndex] = ReadRanges.ColorRange[ReadIndex];
+					WriteRanges.MetadataEntryRange[WriteIndex] = ReadRanges.MetadataEntryRange[ReadIndex];
+					WriteRanges.SeedRange[WriteIndex] = ReadRanges.SeedRange[ReadIndex];
+					++NumWritten;
+				}
 			}
-		});
 
-		PCGE_LOG(Verbose, LogOnly, FText::Format(LOCTEXT("GenerationInfo", "Generated {0} points out of {1} source points"), FilteredPoints.Num(), Points.Num()));
+			return NumWritten;
+		};
+
+		auto MoveDataRangeFunc = [FilteredData](int32 RangeStartIndex, int32 MoveToIndex, int32 NumElements)
+		{
+			FilteredData->MoveRange(RangeStartIndex, MoveToIndex, NumElements);
+		};
+
+		auto FinishedFunc = [FilteredData](int32 NumWritten)
+		{
+			FilteredData->SetNumPoints(NumWritten);
+		};
+
+		FPCGAsyncState* AsyncState = Context ? &Context->AsyncState : nullptr;
+		FPCGAsync::AsyncProcessingRangeEx(
+			AsyncState,
+			OriginalData->GetNumPoints(),
+			InitializeFunc,
+			AsyncProcessRangeFunc,
+			MoveDataRangeFunc,
+			FinishedFunc,
+			/*bEnableTimeSlicing=*/false);
+
+		PCGE_LOG(Verbose, LogOnly, FText::Format(LOCTEXT("GenerationInfo", "Generated {0} points out of {1} source points"), FilteredData->GetNumPoints(), OriginalData->GetNumPoints()));
 	}
 
 	return true;
diff --git a/Engine/Plugins/PCG/Source/PCG/Private/Elements/PCGDistance.cpp b/Engine/Plugins/PCG/Source/PCG/Private/Elements/PCGDistance.cpp
--- a/Engine/Plugins/PCG/Source/PCG/Private/Elements/PCGDistance.cpp
+++ b/Engine/Plugins/PCG/Source/PCG/Private/Elements/PCGDistance.cpp
@@ -5,7 +5,7 @@
 #include "PCGContext.h"
 #include "PCGPin.h"
 #include "Data/PCGSpatialData.h"
-#include "Data/PCGPointData.h"
+#include "Data/PCGBasePointData.h"
 #include "Elements/PCGGather.h"
 #include "Helpers/PCGAsync.h"
 #include "Metadata/Accessors/PCGAttributeAccessorHelpers.h"
@@ -17,27 +17,27 @@ namespace PCGDistance
 	const FName SourceLabel = TEXT("Source");
 	const FName TargetLabel = TEXT("Target");
 
-	FVector CalcPosition(PCGDistanceShape Shape, const FPCGPoint& SourcePoint, const FPCGPoint& TargetPoint, const FVector SourceCenter, const FVector TargetCenter)
+	FVector CalcPosition(PCGDistanceShape Shape, const FTransform& SourceTransform, const FVector& SourceBoundsMin, const FVector& SourceBoundsMax, const FVector SourceCenter, const FVector TargetCenter)
 	{
 		if (Shape == PCGDistanceShape::SphereBounds)
 		{
 			FVector Dir = TargetCenter - SourceCenter;
 			Dir.Normalize();
 
-			return SourceCenter + Dir * SourcePoint.GetScaledExtents().Length();
+			return SourceCenter + Dir * PCGPointHelpers::GetScaledExtents(SourceTransform, SourceBoundsMin, SourceBoundsMax).Length();
 		}
 		else if (Shape == PCGDistanceShape::BoxBounds)
 		{
-			const FVector LocalTargetCenter = SourcePoint.Transform.InverseTransformPosition(TargetCenter);
+			const FVector LocalTargetCenter = SourceTransform.InverseTransformPosition(TargetCenter);
 
-			const double DistanceSquared = ComputeSquaredDistanceFromBoxToPoint(SourcePoint.BoundsMin, SourcePoint.BoundsMax, LocalTargetCenter);
+			const double DistanceSquared = ComputeSquaredDistanceFromBoxToPoint(SourceBoundsMin, SourceBoundsMax, LocalTargetCenter);
 
 			FVector Dir = -LocalTargetCenter;
 			Dir.Normalize();
 
 			const FVector LocalClosestPoint = LocalTargetCenter + Dir * FMath::Sqrt(DistanceSquared);
 
-			return SourcePoint.Transform.TransformPosition(LocalClosestPoint);
+			return SourceTransform.TransformPosition(LocalClosestPoint);
 		}
 
 		// PCGDistanceShape::Center
@@ -135,7 +135,7 @@ bool FPCGDistanceElement::ExecuteInternal(FPCGContext* Context) const
 	TArray<FPCGTaggedData> Targets = Context->InputData.GetInputsByPin(PCGDistance::TargetLabel);
 	TArray<FPCGTaggedData>& Outputs = Context->OutputData.TaggedData;
 
-	TArray<const UPCGPointData*> TargetPointDatas;
+	TArray<const UPCGBasePointData*> TargetPointDatas;
 	TargetPointDatas.Reserve(Targets.Num());
 
 	for (const FPCGTaggedData& Target : Targets)
@@ -148,7 +148,7 @@ bool FPCGDistanceElement::ExecuteInternal(FPCGContext* Context) const
 			continue;
 		}
 
-		const UPCGPointData* TargetPointData = TargetData->ToPointData(Context);
+		const UPCGBasePointData* TargetPointData = TargetData->ToBasePointData(Context);
 		if (!TargetPointData)
 		{
 			PCGE_LOG(Error, GraphAndLog, FText::Format(LOCTEXT("CannotConvertToPoint", "Cannot convert target '{0}' into Point data"), FText::FromString(Target.Data->GetClass()->GetName())));
@@ -185,15 +185,27 @@ bool FPCGDistanceElement::ExecuteInternal(FPCGContext* Context) const
 			continue;
 		}
 
-		const UPCGPointData* SourcePointData = SourceData->ToPointData(Context);
+		const UPCGBasePointData* SourcePointData = SourceData->ToBasePointData(Context);
 		if (!SourcePointData)
 		{
 			PCGE_LOG(Error, GraphAndLog, LOCTEXT("CannotConvertToPointData", "Cannot convert input Spatial data to Point data"));
 			continue;			
 		}
 
-		UPCGPointData* OutputData = FPCGContext::NewObject_AnyThread<UPCGPointData>(Context);
+		UPCGBasePointData* OutputData = FPCGContext::NewPointData_AnyThread(Context);
 		OutputData->InitializeFromData(SourcePointData);
+		OutputData->SetNumPoints(SourcePointData->GetNumPoints());
+
+		if (!OutputData->HasSpatialDataParent())
+		{
+			OutputData->AllocateProperties(SourcePointData->GetAllocatedProperties());
+		}
+
+		if (bSetDensity)
+		{
+			OutputData->AllocateProperties(EPCGPointNativeProperties::Density);
+		}
+
 		Outputs.Add_GetRef(Source).Data = OutputData;
 
 		if (Settings->bOutputToAttribute && Settings->OutputAttribute.IsBasicAttribute())
@@ -229,10 +241,6 @@ bool FPCGDistanceElement::ExecuteInternal(FPCGContext* Context) const
 			}
 		}
 
-		TArray<FPCGPoint>& OutPoints = OutputData->GetMutablePoints();
-		const TArray<FPCGPoint>& SourcePoints = SourcePointData->GetPoints();
-		OutPoints.SetNumUninitialized(SourcePoints.Num());
-
 		if (Accessor.IsValid())
 		{
 			Keys = PCGAttributeAccessorHelpers::CreateKeys(OutputData, Settings->OutputAttribute);
@@ -255,104 +263,123 @@ bool FPCGDistanceElement::ExecuteInternal(FPCGContext* Context) const
 		{
 			if (bOutputDistanceVector)
 			{
-				ResultCache.DistanceVectors.SetNumUninitialized(SourcePoints.Num());
+				ResultCache.DistanceVectors.SetNumUninitialized(SourcePointData->GetNumPoints());
 			}
 			else
 			{
-				ResultCache.Distances.SetNumUninitialized(SourcePoints.Num());
+				ResultCache.Distances.SetNumUninitialized(SourcePointData->GetNumPoints());
 			}
 		}
-
-		auto ProcessDistance = [SourceShape, TargetShape, &TargetPointDatas, MaximumDistance, MaximumDistanceRecip, bSetDensity, bOutputDistanceVector, &OutPoints, &SourcePoints, &ResultCache, bWriteToAttribute = Accessor.IsValid(), Settings, SourceIndex](int32 ReadIndex, int32 WriteIndex)
+				
+		auto ProcessDistanceFunc = [SourceShape, TargetShape, &TargetPointDatas, MaximumDistance, MaximumDistanceRecip, bSetDensity, bOutputDistanceVector, OutputData, SourcePointData, &ResultCache, bWriteToAttribute = Accessor.IsValid(), Settings, SourceIndex](int32 StartReadIndex, int32 StartWriteIndex, int32 Count)
 		{
-			FPCGPoint& OutPoint = OutPoints[WriteIndex];
-			const FPCGPoint& SourcePoint = SourcePoints[ReadIndex];
-
-			OutPoint = SourcePoint;
+			check(StartReadIndex == StartWriteIndex);
 
-			const FBoxSphereBounds SourceQueryBounds = FBoxSphereBounds(FBox(SourcePoint.BoundsMin - FVector(MaximumDistance), SourcePoint.BoundsMax + FVector(MaximumDistance))).TransformBy(SourcePoint.Transform);
+			if (!OutputData->HasSpatialDataParent())
+			{
+				SourcePointData->CopyPointsTo(OutputData, StartReadIndex, StartWriteIndex, Count);
+			}
 
-			const FVector SourceCenter = SourcePoint.Transform.TransformPosition(SourcePoint.GetLocalCenter());
+			const TConstPCGValueRange<FTransform> SourceTransformRange = SourcePointData->GetConstTransformValueRange();
+			const TConstPCGValueRange<FVector> SourceBoundsMinRange = SourcePointData->GetConstBoundsMinValueRange();
+			const TConstPCGValueRange<FVector> SourceBoundsMaxRange = SourcePointData->GetConstBoundsMaxValueRange();
 
-			double MinDistanceSquared = MaximumDistance * MaximumDistance;
-			FVector MinDistanceVector = FVector::ZeroVector;
+			TPCGValueRange<float> OutputDensityRange = bSetDensity ? OutputData->GetDensityValueRange(/*bAllocate=*/false) : TPCGValueRange<float>();
 
-			// Signed distance field for calculating the closest point of source and target
-			auto CalculateSDF = [&MinDistanceSquared, &MinDistanceVector, &SourcePoint, SourceCenter, SourceShape, TargetShape](const FPCGPoint& TargetPoint, const FBoxSphereBounds& Bounds)
+			for (int32 Index = StartReadIndex; Index < (StartReadIndex + Count); ++Index)
 			{
-				// If the source pointer and target pointer are the same, ignore distance to the exact same point
-				if (&SourcePoint == &TargetPoint)
-				{
-					return;
-				}
-
-				const FVector& TargetCenter = Bounds.Origin;
+				const FTransform& SourceTransform = SourceTransformRange[Index];
+				const FVector& SourceBoundsMin = SourceBoundsMinRange[Index];
+				const FVector& SourceBoundsMax = SourceBoundsMaxRange[Index];
 
-				const FVector SourceShapePos = PCGDistance::CalcPosition(SourceShape, SourcePoint, TargetPoint, SourceCenter, TargetCenter);
-				const FVector TargetShapePos = PCGDistance::CalcPosition(TargetShape, TargetPoint, SourcePoint, TargetCenter, SourceCenter);
+				const FBoxSphereBounds SourceQueryBounds = FBoxSphereBounds(FBox(SourceBoundsMin - FVector(MaximumDistance), SourceBoundsMax + FVector(MaximumDistance))).TransformBy(SourceTransform);
 
-				const FVector ToTargetShapeDir = TargetShapePos - SourceShapePos;
-				const FVector ToTargetCenterDir = TargetCenter - SourceCenter;
+				const FVector SourceCenter = SourceTransform.TransformPosition(PCGPointHelpers::GetLocalCenter(SourceBoundsMin, SourceBoundsMax));
 
-				const double Sign = FMath::Sign(ToTargetShapeDir.Dot(ToTargetCenterDir));
-				const double ThisDistanceSquared = ToTargetShapeDir.SquaredLength() * Sign;
+				double MinDistanceSquared = MaximumDistance * MaximumDistance;
+				FVector MinDistanceVector = FVector::ZeroVector;
 
-				if (ThisDistanceSquared < MinDistanceSquared)
+				// Signed distance field for calculating the closest point of source and target
+				auto CalculateSDF = [Index, &MinDistanceSquared, &MinDistanceVector, SourcePointData, &SourceTransform, &SourceBoundsMin, &SourceBoundsMax, SourceCenter, SourceShape, TargetShape](const UPCGBasePointData* TargetPointData, const int32 TargetPointIndex, const FBoxSphereBounds& Bounds)
 				{
-					MinDistanceSquared = ThisDistanceSquared;
-					MinDistanceVector = ToTargetShapeDir;
-				}
-			};
+					// If the source pointer and target pointer are the same, ignore distance to the exact same point
+					if (Index == TargetPointIndex && SourcePointData == TargetPointData)
+					{
+						return;
+					}
 
-			auto CheckAgainstTargetPointData = [&SourceQueryBounds, &CalculateSDF](const UPCGPointData* TargetPointData)
-			{
-				check(TargetPointData);
-				
-				const PCGPointOctree::FPointOctree& Octree = TargetPointData->GetPointOctree();
-				const TArray<FPCGPoint>& TargetPoints = TargetPointData->GetPoints();
+					const FVector& TargetCenter = Bounds.Origin;
+
+					const FVector SourceShapePos = PCGDistance::CalcPosition(SourceShape, SourceTransform, SourceBoundsMin, SourceBoundsMax, SourceCenter, TargetCenter);
+					const FVector TargetShapePos = PCGDistance::CalcPosition(TargetShape, TargetPointData->GetTransform(TargetPointIndex), TargetPointData->GetBoundsMin(TargetPointIndex), TargetPointData->GetBoundsMax(TargetPointIndex), TargetCenter, SourceCenter);
 
-				Octree.FindElementsWithBoundsTest(FBoxCenterAndExtent(SourceQueryBounds.Origin, SourceQueryBounds.BoxExtent),
-					[&TargetPoints, &CalculateSDF](const PCGPointOctree::FPointRef& InPointRef)
+					const FVector ToTargetShapeDir = TargetShapePos - SourceShapePos;
+					const FVector ToTargetCenterDir = TargetCenter - SourceCenter;
+
+					const double Sign = FMath::Sign(ToTargetShapeDir.Dot(ToTargetCenterDir));
+					const double ThisDistanceSquared = ToTargetShapeDir.SquaredLength() * Sign;
+
+					if (ThisDistanceSquared < MinDistanceSquared)
 					{
-						CalculateSDF(TargetPoints[InPointRef.Index], InPointRef.Bounds);
-					});
-			};
+						MinDistanceSquared = ThisDistanceSquared;
+						MinDistanceVector = ToTargetShapeDir;
+					}
+				};
 
-			if (Settings->bCheckSourceAgainstRespectiveTarget)
-			{
-				CheckAgainstTargetPointData(TargetPointDatas[SourceIndex % TargetPointDatas.Num()]);
-			}
-			else
-			{
-				for (const UPCGPointData* TargetPointData : TargetPointDatas)
+				auto CheckAgainstTargetPointData = [&SourceQueryBounds, &CalculateSDF](const UPCGBasePointData* TargetPointData)
 				{
-					CheckAgainstTargetPointData(TargetPointData);
+					check(TargetPointData);
+
+					const PCGPointOctree::FPointOctree& Octree = TargetPointData->GetPointOctree();
+					
+					Octree.FindElementsWithBoundsTest(FBoxCenterAndExtent(SourceQueryBounds.Origin, SourceQueryBounds.BoxExtent),
+						[TargetPointData, &CalculateSDF](const PCGPointOctree::FPointRef& InPointRef)
+						{
+							CalculateSDF(TargetPointData, InPointRef.Index, InPointRef.Bounds);
+						});
+				};
+
+				if (Settings->bCheckSourceAgainstRespectiveTarget)
+				{
+					CheckAgainstTargetPointData(TargetPointDatas[SourceIndex % TargetPointDatas.Num()]);
+				}
+				else
+				{
+					for (const UPCGBasePointData* TargetPointData : TargetPointDatas)
+					{
+						CheckAgainstTargetPointData(TargetPointData);
+					}
 				}
-			}
 
-			const double Distance = FMath::Sign(MinDistanceSquared) * FMath::Sqrt(FMath::Abs(MinDistanceSquared));
+				const double Distance = FMath::Sign(MinDistanceSquared) * FMath::Sqrt(FMath::Abs(MinDistanceSquared));
 
-			if (bWriteToAttribute)
-			{
-				if (bOutputDistanceVector)
+				if (bWriteToAttribute)
 				{
-					ResultCache.DistanceVectors[WriteIndex] = MinDistanceVector;
+					if (bOutputDistanceVector)
+					{
+						ResultCache.DistanceVectors[Index] = MinDistanceVector;
+					}
+					else
+					{
+						ResultCache.Distances[Index] = Distance;
+					}
 				}
-				else
+
+				if (bSetDensity)
 				{
-					ResultCache.Distances[WriteIndex] = Distance;
+					OutputDensityRange[Index] = MaximumDistance > UE_DOUBLE_SMALL_NUMBER ? (FMath::Clamp(Distance, -MaximumDistance, MaximumDistance) * MaximumDistanceRecip) : 1.0f;
 				}
 			}
 
-			if (bSetDensity)
-			{
-				OutPoint.Density = MaximumDistance > UE_DOUBLE_SMALL_NUMBER ? (FMath::Clamp(Distance, -MaximumDistance, MaximumDistance) * MaximumDistanceRecip) : 1.0f;
-			}
-
-			return true;
+			return Count;
 		};
 
-		if (FPCGAsync::AsyncProcessingOneToOneEx(&Context->AsyncState, SourcePoints.Num(), /*InitializeFunc=*/[]{}, ProcessDistance, /*bEnableTimeSlicing=*/false))
+		if (FPCGAsync::AsyncProcessingOneToOneRangeEx(
+			&Context->AsyncState,
+			SourcePointData->GetNumPoints(),
+			/*InitializeFunc=*/[]{},
+			ProcessDistanceFunc,
+			/*bTimeSliceEnabled=*/false))
 		{
 			if (Accessor.IsValid())
 			{
diff --git a/Engine/Plugins/PCG/Source/PCG/Private/Elements/PCGFilterByAttribute.cpp b/Engine/Plugins/PCG/Source/PCG/Private/Elements/PCGFilterByAttribute.cpp
--- a/Engine/Plugins/PCG/Source/PCG/Private/Elements/PCGFilterByAttribute.cpp
+++ b/Engine/Plugins/PCG/Source/PCG/Private/Elements/PCGFilterByAttribute.cpp
@@ -4,7 +4,7 @@
 
 #include "PCGContext.h"
 #include "PCGData.h"
-#include "Data/PCGPointData.h"
+#include "Data/PCGBasePointData.h"
 #include "Helpers/PCGHelpers.h"
 #include "Metadata/PCGMetadata.h"
 #include "Metadata/Accessors/IPCGAttributeAccessor.h"
@@ -82,7 +82,7 @@ bool FPCGFilterByAttributeElement::ExecuteInternal(FPCGContext* Context) const
 
 			if (!Settings->bIgnoreProperties)
 			{
-				if (Cast<UPCGPointData>(Input.Data))
+				if (Cast<UPCGBasePointData>(Input.Data))
 				{
 					if (const UEnum* PointProperties = StaticEnum<EPCGPointProperties>())
 					{
diff --git a/Engine/Plugins/PCG/Source/PCG/Private/Elements/PCGFilterByIndex.cpp b/Engine/Plugins/PCG/Source/PCG/Private/Elements/PCGFilterByIndex.cpp
--- a/Engine/Plugins/PCG/Source/PCG/Private/Elements/PCGFilterByIndex.cpp
+++ b/Engine/Plugins/PCG/Source/PCG/Private/Elements/PCGFilterByIndex.cpp
@@ -3,9 +3,9 @@
 #include "Elements/PCGFilterByIndex.h"
 
 #include "PCGContext.h"
-#include "Data/PCGPointData.h"
 #include "Helpers/Parsing/PCGIndexing.h"
 #include "Helpers/Parsing/PCGParsing.h"
+#include "Utils/PCGLogErrors.h"
 
 #include UE_INLINE_GENERATED_CPP_BY_NAME(PCGFilterByIndex)
 
diff --git a/Engine/Plugins/PCG/Source/PCG/Private/Elements/PCGMatchAndSetAttributes.cpp b/Engine/Plugins/PCG/Source/PCG/Private/Elements/PCGMatchAndSetAttributes.cpp
--- a/Engine/Plugins/PCG/Source/PCG/Private/Elements/PCGMatchAndSetAttributes.cpp
+++ b/Engine/Plugins/PCG/Source/PCG/Private/Elements/PCGMatchAndSetAttributes.cpp
@@ -5,7 +5,7 @@
 #include "PCGContext.h"
 #include "PCGParamData.h"
 #include "PCGPin.h"
-#include "Data/PCGPointData.h"
+#include "Data/PCGBasePointData.h"
 #include "Elements/Metadata/PCGMetadataElementCommon.h"
 #include "Helpers/PCGBlueprintHelpers.h"
 #include "Helpers/PCGHelpers.h"
@@ -703,15 +703,15 @@ public:
 			return false;
 		}
 
-		const UPCGPointData* InPointData = Cast<const UPCGPointData>(InData);
-		SetupWeights(NumElements, InPointData ? &InPointData->GetPoints() : nullptr);
+		const UPCGBasePointData* InPointData = Cast<const UPCGBasePointData>(InData);
+		SetupWeights(NumElements, InPointData);
 
 		return Weights.Num() == NumElements;
 	}
 
-	void SetupWeights(const int32 NumElements, const TArray<FPCGPoint>* Points) 
+	void SetupWeights(const int32 NumElements, const UPCGBasePointData* PointData)
 	{
-		check(!Points || NumElements == Points->Num());
+		check(!PointData || NumElements == PointData->GetNumPoints());
 
 		Weights.Reset(NumElements);
 
@@ -720,17 +720,18 @@ public:
 			Weights.SetNumUninitialized(NumElements);
 			InputWeightAccessor->GetRange<double>(Weights, 0, *InputAttributeKeys, EPCGAttributeAccessorFlags::AllowConstructible);
 		}
-		else if (Points)
+		else if (PointData)
 		{
+			const TConstPCGValueRange<int32> SeedRange = PointData->GetConstSeedValueRange();
 			// Generate a random value from the seed
-			for (const FPCGPoint& Point : *Points)
+			for (int32 Index = 0; Index < SeedRange.Num(); ++Index)
 			{
-				Weights.Add(UPCGBlueprintHelpers::GetRandomStreamFromPoint(Point, Settings, SourceComponent).FRand());
+				Weights.Add(PCGHelpers::GetRandomStreamFromSeed(SeedRange[Index], Settings, SourceComponent).FRand());
 			}
 		}
 		else
 		{
-			FRandomStream RandomStream = UPCGBlueprintHelpers::GetRandomStreamFromPoint(FPCGPoint(), Settings, SourceComponent);
+			FRandomStream RandomStream = PCGHelpers::GetRandomStreamFromSeed(0, Settings, SourceComponent);
 			for (int32 i = 0; i < NumElements; ++i)
 			{
 				Weights.Add(RandomStream.FRand());
@@ -747,9 +748,9 @@ public:
 	{
 		check(InData);
 
-		if (const UPCGPointData* InPointData = Cast<const UPCGPointData>(InData))
+		if (const UPCGBasePointData* InPointData = Cast<const UPCGBasePointData>(InData))
 		{
-			return InPointData->GetPoints().Num();
+			return InPointData->GetNumPoints();
 		}
 		else if (const UPCGMetadata* Metadata = InData->ConstMetadata())
 		{
@@ -764,9 +765,9 @@ public:
 	void Finalize(const UPCGData* InData, UPCGData* OutData)
 	{
 		check(InData && OutData);
-		if (const UPCGPointData* InPointData = Cast<const UPCGPointData>(InData))
+		if (const UPCGBasePointData* InPointData = Cast<const UPCGBasePointData>(InData))
 		{
-			FinalizeInternal(InPointData, static_cast<UPCGPointData*>(OutData), InPointData->GetPoints().Num());
+			FinalizeInternal(InPointData, static_cast<UPCGBasePointData*>(OutData), InPointData->GetNumPoints());
 		}
 		else if (const UPCGMetadata* Metadata = InData->ConstMetadata())
 		{
@@ -785,6 +786,20 @@ public:
 		UPCGMetadata* OutMetadata = OutData->MutableMetadata();
 		check(InMetadata && OutMetadata);
 
+		int32 NumWritten = 0;
+
+		FConstPCGPointValueRanges InRanges;
+		FPCGPointValueRanges OutRanges;
+				
+		if constexpr (std::is_same_v<T, UPCGBasePointData>)
+		{
+			const UPCGBasePointData* InPointData = Cast<UPCGBasePointData>(InData);
+			InRanges = FConstPCGPointValueRanges(InPointData);
+
+			UPCGBasePointData* OutPointData = Cast<UPCGBasePointData>(OutData);
+			OutRanges = FPCGPointValueRanges(OutPointData, /*bAllocate=*/false);
+		}
+
 		for (int32 Index = 0; Index < NumElements; ++Index)
 		{
 			int32 PartitionDataIndex = PartitionDataIndices[Index];
@@ -800,17 +815,20 @@ public:
 				PCGMetadataEntryKey PreviousKey = PCGMetadataEntryKey(Index);
 				PCGMetadataEntryKey NewEntry;
 
-				if constexpr (std::is_same_v<T, UPCGPointData>)
+				if constexpr (std::is_same_v<T, UPCGBasePointData>)
 				{
-					FPCGPoint& OutPoint = static_cast<UPCGPointData*>(OutData)->GetMutablePoints().Add_GetRef(static_cast<const UPCGPointData*>(InData)->GetPoints()[Index]);
-					PreviousKey = OutPoint.MetadataEntry;
+					OutRanges.SetFromValueRanges(NumWritten, InRanges, Index);
+
+					PreviousKey = OutRanges.MetadataEntryRange[NumWritten];
 
 					if (AttributeSetKey != PCGInvalidEntryKey)
 					{
-						OutPoint.MetadataEntry = OutMetadata->AddEntry(OutPoint.MetadataEntry);
+						OutRanges.MetadataEntryRange[NumWritten] = OutMetadata->AddEntry(OutRanges.MetadataEntryRange[NumWritten]);
 					}
 
-					NewEntry = OutPoint.MetadataEntry;
+					NewEntry = OutRanges.MetadataEntryRange[NumWritten];
+
+					++NumWritten;
 				}
 				else
 				{
@@ -831,6 +849,12 @@ public:
 				}
 			}
 		}
+
+		if constexpr (std::is_same_v<T, UPCGBasePointData>)
+		{
+			UPCGBasePointData* OutPointData = Cast<UPCGBasePointData>(OutData);
+			OutPointData->SetNumPoints(NumWritten);
+		}
 	}
 
 private:
@@ -1034,11 +1058,16 @@ bool FPCGMatchAndSetAttributesElement::PrepareDataInternal(FPCGContext* InContex
 			}
 		}
 
-		if (const UPCGPointData* InPointData = Cast<const UPCGPointData>(OutState.InData))
+		if (const UPCGBasePointData* InPointData = Cast<const UPCGBasePointData>(OutState.InData))
 		{
-			UPCGPointData* OutPointData = FPCGContext::NewObject_AnyThread<UPCGPointData>(InContext);
-			OutPointData->InitializeFromData(InPointData);
-			OutPointData->GetMutablePoints().Reserve(InPointData->GetPoints().Num());
+			UPCGBasePointData* OutPointData = FPCGContext::NewPointData_AnyThread(InContext);
+
+			FPCGInitializeFromDataParams InitializeFromDataParams(InPointData);
+			InitializeFromDataParams.bInheritSpatialData = false;
+			OutPointData->InitializeFromDataWithParams(InitializeFromDataParams);
+			OutPointData->SetNumPoints(InPointData->GetNumPoints());
+			OutPointData->AllocateProperties(InPointData->GetAllocatedProperties());
+
 			OutState.OutData = OutPointData;
 		}
 		else if (const UPCGParamData* InParamData = Cast<const UPCGParamData>(OutState.InData))
diff --git a/Engine/Plugins/PCG/Source/PCG/Private/Elements/PCGMergeElement.cpp b/Engine/Plugins/PCG/Source/PCG/Private/Elements/PCGMergeElement.cpp
--- a/Engine/Plugins/PCG/Source/PCG/Private/Elements/PCGMergeElement.cpp
+++ b/Engine/Plugins/PCG/Source/PCG/Private/Elements/PCGMergeElement.cpp
@@ -2,7 +2,7 @@
 
 #include "Elements/PCGMergeElement.h"
 
-#include "Data/PCGPointData.h"
+#include "Data/PCGBasePointData.h"
 #include "PCGContext.h"
 #include "PCGPin.h"
 
@@ -70,7 +70,7 @@ bool FPCGMergeElement::ExecuteInternal(FPCGContext* Context) const
 
 	TArray<FPCGTaggedData>& Outputs = Context->OutputData.TaggedData;
 
-	UPCGPointData* TargetPointData = nullptr;
+	UPCGBasePointData* TargetPointData = nullptr;
 	FPCGTaggedData* TargetTaggedData = nullptr;
 
 	int32 TotalPointCount = 0;
@@ -79,7 +79,7 @@ bool FPCGMergeElement::ExecuteInternal(FPCGContext* Context) const
 	// Done in two passes for futureproofing - expecting changes in the metadata attribute creation vs. usage in points
 	for (const FPCGTaggedData& Source : Sources)
 	{
-		const UPCGPointData* SourcePointData = Cast<const UPCGPointData>(Source.Data);
+		const UPCGBasePointData* SourcePointData = Cast<const UPCGBasePointData>(Source.Data);
 
 		if (!SourcePointData)
 		{
@@ -87,12 +87,12 @@ bool FPCGMergeElement::ExecuteInternal(FPCGContext* Context) const
 			continue;
 		}
 
-		if (SourcePointData->GetPoints().IsEmpty())
+		if (SourcePointData->GetNumPoints() == 0)
 		{
 			continue;
 		}
 
-		TotalPointCount += SourcePointData->GetPoints().Num();
+		TotalPointCount += SourcePointData->GetNumPoints();
 
 		// First data that's valid - will be the original output so we don't do a copy unless we actually need it
 		if (!TargetTaggedData)
@@ -104,8 +104,12 @@ bool FPCGMergeElement::ExecuteInternal(FPCGContext* Context) const
 			// Second valid data - we'll create the actual merged data at this point
 			check(TargetTaggedData);
 
-			TargetPointData = FPCGContext::NewObject_AnyThread<UPCGPointData>(Context);
-			TargetPointData->InitializeFromData(CastChecked<const UPCGPointData>(TargetTaggedData->Data), nullptr, bMergeMetadata);
+			TargetPointData = FPCGContext::NewPointData_AnyThread(Context);
+			FPCGInitializeFromDataParams InitializeFromDataParams(CastChecked<const UPCGBasePointData>(TargetTaggedData->Data));
+			InitializeFromDataParams.bInheritSpatialData = false;
+			InitializeFromDataParams.bInheritMetadata = bMergeMetadata;
+			
+			TargetPointData->InitializeFromDataWithParams(InitializeFromDataParams);
 			TargetTaggedData->Data = TargetPointData;
 		}
 
@@ -127,34 +131,60 @@ bool FPCGMergeElement::ExecuteInternal(FPCGContext* Context) const
 		return true;
 	}
 
-	TArray<FPCGPoint>& TargetPoints = TargetPointData->GetMutablePoints();
-	TargetPoints.Reserve(TotalPointCount);
+	TargetPointData->SetNumPoints(TotalPointCount);
+	int32 PointOffset = 0;
 
-	for(int32 SourceIndex = 0; SourceIndex < Sources.Num(); ++SourceIndex)
+	for(int32 SourceDataIndex = 0; SourceDataIndex < Sources.Num(); ++SourceDataIndex)
 	{
-		const UPCGPointData* SourcePointData = Cast<const UPCGPointData>(Sources[SourceIndex].Data);
+		const UPCGBasePointData* SourcePointData = Cast<const UPCGBasePointData>(Sources[SourceDataIndex].Data);
+		const int32 NumSourcePoints = SourcePointData ? SourcePointData->GetNumPoints() : 0;
 
-		if (!SourcePointData)
+		if (NumSourcePoints == 0)
 		{
 			continue;
 		}
 
-		int32 PointOffset = TargetPoints.Num();
-		TargetPoints.Append(SourcePointData->GetPoints());
+		SourcePointData->CopyPointsTo(TargetPointData, 0, PointOffset, NumSourcePoints);
 
-		if ((!bMergeMetadata || SourceIndex != 0) && !SourcePointData->GetPoints().IsEmpty())
+		if ((!bMergeMetadata || SourceDataIndex != 0))
 		{
-			TArrayView<FPCGPoint> TargetPointsSubset = MakeArrayView(&TargetPoints[PointOffset], SourcePointData->GetPoints().Num());
-			for (FPCGPoint& Point : TargetPointsSubset)
+			TPCGValueRange<int64> TargetMetadataEntryRange = TargetPointData->GetMetadataEntryValueRange();
+			for (int32 TargetIndex = PointOffset; TargetIndex < PointOffset + NumSourcePoints; ++TargetIndex)
 			{
-				Point.MetadataEntry = PCGInvalidEntryKey;
+				TargetMetadataEntryRange[TargetIndex] = PCGInvalidEntryKey;
 			}
 
 			if (bMergeMetadata && TargetPointData->Metadata && SourcePointData->Metadata && SourcePointData->Metadata->GetAttributeCount() > 0)
 			{
-				TargetPointData->Metadata->SetPointAttributes(MakeArrayView(SourcePointData->GetPoints()), SourcePointData->Metadata, TargetPointsSubset, Context);
+				// Extract the metadata entry keys from the in & out points
+				TArray<PCGMetadataEntryKey, TInlineAllocator<256>> SourceKeys;
+				SourceKeys.SetNumUninitialized(NumSourcePoints);
+
+				const TConstPCGValueRange<int64> SourceMetadataEntryRange = SourcePointData->GetConstMetadataEntryValueRange();
+				for (int32 SourceIndex = 0; SourceIndex < SourceKeys.Num(); ++SourceIndex)
+				{
+					SourceKeys[SourceIndex] = SourceMetadataEntryRange[SourceIndex];
+				}
+				
+				TArray<PCGMetadataEntryKey, TInlineAllocator<256>> TargetKeys;
+				TargetKeys.SetNumUninitialized(NumSourcePoints);
+								
+				for (int32 TargetIndex = 0; TargetIndex < TargetKeys.Num(); ++TargetIndex)
+				{
+					TargetKeys[TargetIndex] = PCGInvalidEntryKey;
+				}
+
+				TargetPointData->Metadata->SetAttributes(SourceKeys, SourcePointData->Metadata, TargetKeys, Context);
+
+				// Write back
+				for (int32 TargetIndex = 0; TargetIndex < TargetKeys.Num(); ++TargetIndex)
+				{
+					TargetMetadataEntryRange[PointOffset + TargetIndex] = TargetKeys[TargetIndex];
+				}
 			}
 		}
+
+		PointOffset += NumSourcePoints;
 	}
 	
 	return true;
diff --git a/Engine/Plugins/PCG/Source/PCG/Private/Elements/PCGProjectionElement.cpp b/Engine/Plugins/PCG/Source/PCG/Private/Elements/PCGProjectionElement.cpp
--- a/Engine/Plugins/PCG/Source/PCG/Private/Elements/PCGProjectionElement.cpp
+++ b/Engine/Plugins/PCG/Source/PCG/Private/Elements/PCGProjectionElement.cpp
@@ -6,7 +6,7 @@
 #include "PCGCustomVersion.h"
 #include "PCGEdge.h"
 #include "PCGPin.h"
-#include "Data/PCGPointData.h"
+#include "Data/PCGBasePointData.h"
 #include "Data/PCGProjectionData.h"
 #include "Data/PCGSpatialData.h"
 
@@ -103,9 +103,9 @@ bool FPCGProjectionElement::ExecuteInternal(FPCGContext* Context) const
 
 		const UPCGSpatialData* OutputData = ProjectionData;
 
-		if (Settings->bForceCollapseToPoint || ProjectionSource->IsA<UPCGPointData>())
+		if (Settings->bForceCollapseToPoint || ProjectionSource->IsA<UPCGBasePointData>())
 		{
-			OutputData = OutputData->ToPointData(Context);
+			OutputData = OutputData->ToBasePointData(Context);
 		}
 		else if (OutputData->RequiresCollapseToSample())
 		{
diff --git a/Engine/Plugins/PCG/Source/PCG/Private/Elements/PCGSelectPoints.cpp b/Engine/Plugins/PCG/Source/PCG/Private/Elements/PCGSelectPoints.cpp
--- a/Engine/Plugins/PCG/Source/PCG/Private/Elements/PCGSelectPoints.cpp
+++ b/Engine/Plugins/PCG/Source/PCG/Private/Elements/PCGSelectPoints.cpp
@@ -4,7 +4,7 @@
 
 #include "PCGContext.h"
 #include "Data/PCGSpatialData.h"
-#include "Data/PCGPointData.h"
+#include "Data/PCGBasePointData.h"
 #include "Helpers/PCGAsync.h"
 #include "Helpers/PCGHelpers.h"
 
@@ -74,21 +74,23 @@ bool FPCGSelectPointsElement::ExecuteInternal(FPCGContext* Context) const
 			continue;
 		}
 
-		const UPCGPointData* OriginalData = Cast<UPCGSpatialData>(Input.Data)->ToPointData(Context);
+		const UPCGBasePointData* OriginalData = Cast<UPCGSpatialData>(Input.Data)->ToBasePointData(Context);
 
 		if (!OriginalData)
 		{
 			PCGE_LOG(Error, GraphAndLog, LOCTEXT("NoPointDataInInput", "Unable to get point data from input"));
 			continue;
 		}
-
-		const TArray<FPCGPoint>& Points = OriginalData->GetPoints();
-		const int OriginalPointCount = Points.Num();
-
-		UPCGPointData* SampledData = FPCGContext::NewObject_AnyThread<UPCGPointData>(Context);
-		SampledData->InitializeFromData(OriginalData);
-		TArray<FPCGPoint>& SampledPoints = SampledData->GetMutablePoints();
-
+				
+		const int OriginalPointCount = OriginalData->GetNumPoints();
+
+		UPCGBasePointData* SampledData = FPCGContext::NewPointData_AnyThread(Context);
+		
+		FPCGInitializeFromDataParams InitializeFromDataParams(OriginalData);
+		InitializeFromDataParams.bInheritSpatialData = false;
+		
+		SampledData->InitializeFromDataWithParams(InitializeFromDataParams);
+		
 		Output.Data = SampledData;
 
 		// TODO: randomize on the fractional number of points
@@ -108,34 +110,81 @@ bool FPCGSelectPointsElement::ExecuteInternal(FPCGContext* Context) const
 		{
 			TRACE_CPUPROFILER_EVENT_SCOPE(FPCGSelectPointsElement::Execute::SelectPoints);
 
-			FPCGAsync::AsyncPointProcessing(Context, OriginalPointCount, SampledPoints, [&Points, Seed, Ratio, bKeepZeroDensityPoints](int32 Index, FPCGPoint& OutPoint)
+			auto InitializeFunc = [SampledData, OriginalData]()
 			{
-				const FPCGPoint& Point = Points[Index];
+				SampledData->SetNumPoints(OriginalData->GetNumPoints());
+				SampledData->AllocateProperties(OriginalData->GetAllocatedProperties());
+			};
 
-				// Apply a high-pass filter based on selected ratio
-				FRandomStream RandomSource(PCGHelpers::ComputeSeed(Seed, Point.Seed));
-				float Chance = RandomSource.FRand();
+			auto AsyncProcessRangeFunc = [SampledData, OriginalData, Seed, Ratio, bKeepZeroDensityPoints](int32 StartReadIndex, int32 StartWriteIndex, int32 Count)
+			{
+				int32 NumWritten = 0;
 
-				if (Chance < Ratio)
+				// @todo_pcg: could probably build a list of indices and then use CopyPoints instead of writing this boiler plate code
+				FPCGPointValueRanges SampledRanges(SampledData, /*bAllocate=*/false);
+				const FConstPCGPointValueRanges OriginalRanges(OriginalData);
+								
+				for (int32 ReadIndex = StartReadIndex; ReadIndex < StartReadIndex + Count; ++ReadIndex)
 				{
-					OutPoint = Point;
-					return true;
-				}
+					const int32 WriteIndex = StartWriteIndex + NumWritten;
+					const int32 OriginalSeed = OriginalRanges.SeedRange[ReadIndex];
+
+					// Apply a high-pass filter based on selected ratio
+					FRandomStream RandomSource(PCGHelpers::ComputeSeed(Seed, OriginalSeed));
+					float Chance = RandomSource.FRand();
+
+					bool bCopyPoint = false;
+
+					if (Chance < Ratio)
+					{
+						bCopyPoint = true;
+						SampledRanges.DensityRange[WriteIndex] = OriginalRanges.DensityRange[ReadIndex];
+					}
 #if WITH_EDITOR
-				else if (bKeepZeroDensityPoints)
-				{
-					OutPoint = Point;
-					OutPoint.Density = 0;
-					return true;
-				}
+					else if (bKeepZeroDensityPoints)
+					{
+						bCopyPoint = true;
+						SampledRanges.DensityRange[WriteIndex] = 0.f;
+					}
 #endif
-				else
-				{
-					return false;
+					
+					if (bCopyPoint)
+					{
+						SampledRanges.TransformRange[WriteIndex] = OriginalRanges.TransformRange[ReadIndex];
+						SampledRanges.BoundsMinRange[WriteIndex] = OriginalRanges.BoundsMinRange[ReadIndex];
+						SampledRanges.BoundsMaxRange[WriteIndex] = OriginalRanges.BoundsMaxRange[ReadIndex];
+						SampledRanges.SeedRange[WriteIndex] = OriginalRanges.SeedRange[ReadIndex];
+						SampledRanges.MetadataEntryRange[WriteIndex] = OriginalRanges.MetadataEntryRange[ReadIndex];
+						SampledRanges.ColorRange[WriteIndex] = OriginalRanges.ColorRange[ReadIndex];
+						SampledRanges.SteepnessRange[WriteIndex] = OriginalRanges.SteepnessRange[ReadIndex];
+						++NumWritten;
+					}
 				}
-			});
+							
+				return NumWritten;
+			};
 
-			PCGE_LOG(Verbose, LogOnly, FText::Format(LOCTEXT("GenerationInfo", "Generated {0} points from {1} source points"), SampledPoints.Num(), OriginalPointCount));
+			auto MoveDataRangeFunc = [SampledData](int32 RangeStartIndex, int32 MoveToIndex, int32 NumElements)
+			{
+				SampledData->MoveRange(RangeStartIndex, MoveToIndex, NumElements);
+			};
+
+			auto FinishedFunc = [SampledData](int32 NumWritten)
+			{
+				SampledData->SetNumPoints(NumWritten);
+			};
+
+			FPCGAsyncState* AsyncState = Context ? &Context->AsyncState : nullptr;
+			const bool bDone = FPCGAsync::AsyncProcessingRangeEx(
+				AsyncState,
+				OriginalPointCount,
+				InitializeFunc,
+				AsyncProcessRangeFunc,
+				MoveDataRangeFunc,
+				FinishedFunc,
+				/*bEnableTimeSlicing=*/false);
+
+			PCGE_LOG(Verbose, LogOnly, FText::Format(LOCTEXT("GenerationInfo", "Generated {0} points from {1} source points"), SampledData->GetNumPoints(), OriginalPointCount));
 		}
 	}
 
diff --git a/Engine/Plugins/PCG/Source/PCG/Private/Elements/PCGSelfPruning.cpp b/Engine/Plugins/PCG/Source/PCG/Private/Elements/PCGSelfPruning.cpp
--- a/Engine/Plugins/PCG/Source/PCG/Private/Elements/PCGSelfPruning.cpp
+++ b/Engine/Plugins/PCG/Source/PCG/Private/Elements/PCGSelfPruning.cpp
@@ -120,8 +120,8 @@ namespace PCGSelfPruningElement
 	{
 		TRACE_CPUPROFILER_EVENT_SCOPE(FPCGSelfPruningElement::Execute::DensityBoundsExclusion);
 
-		check(IterationState.InputPointData);
-		const PCGPointOctree::FPointOctree& Octree = IterationState.InputPointData->GetPointOctree();
+		check(IterationState.InputData);
+		const PCGPointOctree::FPointOctree& Octree = IterationState.InputData->GetPointOctree();
 
 		int32 CheckTimeSlicingCount = 0;
 
@@ -164,9 +164,9 @@ namespace PCGSelfPruningElement
 	{
 		TRACE_CPUPROFILER_EVENT_SCOPE(FPCGSelfPruningElement::Execute::CollisionExclusion);
 
-		check(IterationState.InputPointData);
-		const PCGPointOctree::FPointOctree& Octree = IterationState.InputPointData->GetPointOctree();
-		const TConstPCGValueRange<FTransform> TransformRange = IterationState.InputPointData->GetConstTransformValueRange();
+		check(IterationState.InputData);
+		const PCGPointOctree::FPointOctree& Octree = IterationState.InputData->GetPointOctree();
+		const TConstPCGValueRange<FTransform> TransformRange = IterationState.InputData->GetConstTransformValueRange();
 
 		int32 CheckTimeSlicingCount = 0;
 
@@ -385,14 +385,14 @@ namespace PCGSelfPruningElement
 	{
 		TRACE_CPUPROFILER_EVENT_SCOPE(FPCGSelfPruningElement::Execute::DuplicatePointsExclusion);
 
-		check(IterationState.InputPointData);
-		const PCGPointOctree::FPointOctree& Octree = IterationState.InputPointData->GetPointOctree();
+		check(IterationState.InputData);
+		const PCGPointOctree::FPointOctree& Octree = IterationState.InputData->GetPointOctree();
 
 		int32 CheckTimeSlicingCount = 0;
 
-		const TConstPCGValueRange<FTransform> TransformRange = IterationState.InputPointData->GetConstTransformValueRange();
-		const TConstPCGValueRange<FVector> BoundsMinRange = IterationState.InputPointData->GetConstBoundsMinValueRange();
-		const TConstPCGValueRange<FVector> BoundsMaxRange = IterationState.InputPointData->GetConstBoundsMaxValueRange();
+		const TConstPCGValueRange<FTransform> TransformRange = IterationState.InputData->GetConstTransformValueRange();
+		const TConstPCGValueRange<FVector> BoundsMinRange = IterationState.InputData->GetConstBoundsMinValueRange();
+		const TConstPCGValueRange<FVector> BoundsMaxRange = IterationState.InputData->GetConstBoundsMaxValueRange();
 
 		while (IterationState.CurrentPointIndex < IterationState.SortedPointRefs.Num())
 		{
@@ -463,16 +463,19 @@ namespace PCGSelfPruningElement
 				continue;
 			}
 
-			const UPCGPointData* InputPointData = SpatialInput->ToPointData(Context);
+			const UPCGBasePointData* InputData = SpatialInput->ToBasePointData(Context);
 			FIterationState IterationState;
-			IterationState.InputPointData = InputPointData;
+			IterationState.InputData = InputData;
+PRAGMA_DISABLE_DEPRECATION_WARNINGS
+			IterationState.InputPointData = Cast<UPCGPointData>(InputData);
+PRAGMA_ENABLE_DEPRECATION_WARNINGS
 
 			while (!ExecuteSlice(IterationState, InParameters, Context)) {}
 
 			FPCGTaggedData& Output = Context->OutputData.TaggedData.Add_GetRef(Input);
-			if (IterationState.OutputPointData != nullptr)
+			if (IterationState.OutputData != nullptr)
 			{
-				Output.Data = IterationState.OutputPointData;
+				Output.Data = IterationState.OutputData;
 			}
 		}
 	}
@@ -491,11 +494,11 @@ namespace PCGSelfPruningElement
 
 		const FVector::FReal RadiusEquality = 1.0f + InParameters.RadiusSimilarityFactor;
 		const FVector::FReal SquaredRadiusEquality = FMath::Square(RadiusEquality);
-		const TArray<FPCGPoint>& Points = InState.InputPointData->GetPoints();
+		const int32 InputNumPoints = InState.InputData->GetNumPoints();
 
 		// Force octree computation, and check if we need to stop after that if it was dirty in the first place.
-		const bool bOctreeWasDirty = InState.InputPointData->IsPointOctreeDirty();
-		const PCGPointOctree::FPointOctree& Octree = InState.InputPointData->GetPointOctree();
+		const bool bOctreeWasDirty = InState.InputData->IsPointOctreeDirty();
+		const PCGPointOctree::FPointOctree& Octree = InState.InputData->GetPointOctree();
 		if (bOctreeWasDirty && ShouldStop(InOptionalContext))
 		{
 			return false;
@@ -510,13 +513,20 @@ namespace PCGSelfPruningElement
 		{
 			// In the case of the collision-driven self-pruning, we have to populate the sorted points array earlier since we're playing with the bounds.
 			InState.SortedPointRefs.Empty();
-			InState.SortedPointRefs.Reserve(Points.Num());
-			for (int32 PointIndex = 0; PointIndex < Points.Num(); ++PointIndex)
+			InState.SortedPointRefs.Reserve(InputNumPoints);
+
+			const TConstPCGValueRange<FTransform> TransformRange = InState.InputData->GetConstTransformValueRange();
+			const TConstPCGValueRange<FVector> BoundsMinRange = InState.InputData->GetConstBoundsMinValueRange();
+			const TConstPCGValueRange<FVector> BoundsMaxRange = InState.InputData->GetConstBoundsMaxValueRange();
+			const TConstPCGValueRange<float> SteepnessRange = InState.InputData->GetConstSteepnessValueRange();
+			const TConstPCGValueRange<int32> SeedRange = InState.InputData->GetConstSeedValueRange();
+
+			for (int32 PointIndex = 0; PointIndex < InputNumPoints; ++PointIndex)
 			{
-				InState.SortedPointRefs.Emplace(PCGPointOctree::FPointRef(PointIndex, Points[PointIndex].GetDensityBounds()));
+				InState.SortedPointRefs.Emplace(PCGPointOctree::FPointRef(PointIndex, PCGPointHelpers::GetDensityBounds(TransformRange[PointIndex], SteepnessRange[PointIndex], BoundsMinRange[PointIndex], BoundsMaxRange[PointIndex])));
 			}
 
-			FPCGAttributePropertySelector ComparisonSource = InParameters.ComparisonSource.CopyAndFixLast(InState.InputPointData);
+			FPCGAttributePropertySelector ComparisonSource = InParameters.ComparisonSource.CopyAndFixLast(InState.InputData);
 
 			// First sort the points if we need sorting. 
 			// Randomize for other cases than LargeToSmall and SmallToLarge are just a seed sort, so special case for that.
@@ -528,12 +538,12 @@ namespace PCGSelfPruningElement
 			if (InParameters.bRandomizedPruning && !bSortAttribute)
 			{
 				TRACE_CPUPROFILER_EVENT_SCOPE(FPCGSelfPruningElement::Execute::SortingSeeds);
-				Algo::Sort(InState.SortedPointRefs, [&InState](const PCGPointOctree::FPointRef& A, const PCGPointOctree::FPointRef& B) { return PCGSelfPruningAlgorithms::RandomSort(InState.InputPointData, A, B); });
+				Algo::Sort(InState.SortedPointRefs, [&InState](const PCGPointOctree::FPointRef& A, const PCGPointOctree::FPointRef& B) { return PCGSelfPruningAlgorithms::RandomSort(InState.InputData, A, B); });
 			}
 			else if (ComparisonSource.GetSelection() == EPCGAttributePropertySelection::Property && ComparisonSource.GetPointProperty() == EPCGPointProperties::Extents && ComparisonSource.GetExtraNames().IsEmpty())
 			{
 				TRACE_CPUPROFILER_EVENT_SCOPE(FPCGSelfPruningElement::Execute::SortingExtents);
-				PCGSelfPruningAlgorithms::SortExtents(InState.InputPointData, InState.SortedPointRefs, InParameters, SquaredRadiusEquality);
+				PCGSelfPruningAlgorithms::SortExtents(InState.InputData, InState.SortedPointRefs, InParameters, SquaredRadiusEquality);
 			}
 			else
 			{
@@ -541,7 +551,8 @@ namespace PCGSelfPruningElement
 
 				const bool bAscending = InParameters.PruningType != EPCGSelfPruningType::LargeToSmall;
 
-				TUniquePtr<const IPCGAttributeAccessor> InputAccessor = PCGAttributeAccessorHelpers::CreateConstAccessor(InState.InputPointData, ComparisonSource);
+				TUniquePtr<const IPCGAttributeAccessor> InputAccessor = PCGAttributeAccessorHelpers::CreateConstAccessor(InState.InputData, ComparisonSource);
+				TUniquePtr<const IPCGAttributeAccessorKeys> InputKeys = PCGAttributeAccessorHelpers::CreateConstKeys(InState.InputData, FPCGAttributePropertyInputSelector());
 
 				if (!InputAccessor.IsValid())
 				{
@@ -554,7 +565,7 @@ namespace PCGSelfPruningElement
 				{
 					// Force squared length extraction on vectors
 					ComparisonSource.GetExtraNamesMutable().Add(PCGAttributeExtractorConstants::VectorSquaredLength.ToString());
-					InputAccessor = PCGAttributeAccessorHelpers::CreateConstAccessor(InState.InputPointData, ComparisonSource);
+					InputAccessor = PCGAttributeAccessorHelpers::CreateConstAccessor(InState.InputData, ComparisonSource);
 				}
 
 				if (!PCG::Private::IsOfTypes<int32, int64, float, double, uint8>(InputAccessor->GetUnderlyingType()))
@@ -565,7 +576,7 @@ namespace PCGSelfPruningElement
 
 				if (InParameters.bRandomizedPruning)
 				{
-					auto CompareLessWithRandom = [SquaredRadiusEquality, &Points](const auto& A, const auto& B, int32 IndexA, int32 IndexB, bool)
+					auto CompareLessWithRandom = [SquaredRadiusEquality, SeedRange](const auto& A, const auto& B, int32 IndexA, int32 IndexB, bool)
 					{
 						static_assert(std::is_same_v<decltype(A), decltype(B)>);
 
@@ -597,10 +608,10 @@ namespace PCGSelfPruningElement
 							}
 						}
 
-						return Points[IndexA].Seed < Points[IndexB].Seed;
+						return SeedRange[IndexA] < SeedRange[IndexB];
 					};
 
-					PCGAttributeAccessorHelpers::SortByAttribute(*InputAccessor, FPCGAttributeAccessorKeysPoints(Points), InState.SortedPointRefs, bAscending, PCGAttributeAccessorHelpers::Private::DefaultIndexGetter, std::move(CompareLessWithRandom));
+					PCGAttributeAccessorHelpers::SortByAttribute(*InputAccessor, *InputKeys, InState.SortedPointRefs, bAscending, PCGAttributeAccessorHelpers::Private::DefaultIndexGetter, std::move(CompareLessWithRandom));
 				}
 				else
 				{
@@ -621,12 +632,12 @@ namespace PCGSelfPruningElement
 						return PCG::Private::MetadataTraits<ValueType>::Less(ACopy, B);
 					};
 
-					PCGAttributeAccessorHelpers::SortByAttribute(*InputAccessor, FPCGAttributeAccessorKeysPoints(Points), InState.SortedPointRefs, bAscending, PCGAttributeAccessorHelpers::Private::DefaultIndexGetter, std::move(CompareLess));
+					PCGAttributeAccessorHelpers::SortByAttribute(*InputAccessor, *InputKeys, InState.SortedPointRefs, bAscending, PCGAttributeAccessorHelpers::Private::DefaultIndexGetter, std::move(CompareLess));
 				}
 			}
 
-			InState.ExclusionPoints.Initialize(Points.Num());
-			InState.ExcludedPoints.Initialize(Points.Num());
+			InState.ExclusionPoints.Initialize(InputNumPoints);
+			InState.ExcludedPoints.Initialize(InputNumPoints);
 
 			InState.bSortDone = true;
 
@@ -656,29 +667,42 @@ namespace PCGSelfPruningElement
 		// Finally, output all points that are present in the ExclusionPoints. This part is not time sliced, should it be too?
 		if (bIsDone)
 		{
-			UPCGPointData* PrunedData = FPCGContext::NewObject_AnyThread<UPCGPointData>(InOptionalContext);
-			PrunedData->InitializeFromData(InState.InputPointData);
-			InState.OutputPointData = PrunedData;
+			UPCGBasePointData* PrunedData = FPCGContext::NewPointData_AnyThread(InOptionalContext);
 
-			TArray<FPCGPoint>& OutputPoints = PrunedData->GetMutablePoints();
+			FPCGInitializeFromDataParams InitializeFromDataParams(InState.InputData);
+			InitializeFromDataParams.bInheritSpatialData = false;
 
-			for (int32 PointIndex = 0; PointIndex < Points.Num(); ++PointIndex)
+			PrunedData->InitializeFromDataWithParams(InitializeFromDataParams);
+
+			InState.OutputData = PrunedData;
+PRAGMA_DISABLE_DEPRECATION_WARNINGS
+			InState.OutputPointData = Cast<UPCGPointData>(PrunedData);
+PRAGMA_ENABLE_DEPRECATION_WARNINGS
+			
+			TArray<int32> Indices;
+			Indices.Reserve(InputNumPoints);
+			for (int32 PointIndex = 0; PointIndex < InputNumPoints; ++PointIndex)
 			{
 				if (InState.ExclusionPoints.Contains(PointIndex))
 				{
-					OutputPoints.Add(Points[PointIndex]);
+					Indices.Add(PointIndex);
 				}
 			}
 
+			if (Indices.Num() > 0)
+			{
+				UPCGBasePointData::CopyPoints(InState.InputData, PrunedData, Indices, /*bCopyAll=*/false);
+			}
+
 			if (InOptionalContext)
 			{
 				if (bIsDuplicateTest)
 				{
-					PCGE_LOG_C(Verbose, LogOnly, InOptionalContext, FText::Format(LOCTEXT("GenerationInfoDuplicate", "Removed {0} duplicate points from {1} source points"), Points.Num() - OutputPoints.Num(), Points.Num()));
+					PCGE_LOG_C(Verbose, LogOnly, InOptionalContext, FText::Format(LOCTEXT("GenerationInfoDuplicate", "Removed {0} duplicate points from {1} source points"), InputNumPoints - PrunedData->GetNumPoints(), InputNumPoints));
 				}
 				else
 				{
-					PCGE_LOG_C(Verbose, LogOnly, InOptionalContext, FText::Format(LOCTEXT("GenerationInfo", "Generated {0} points from {1} source points"), OutputPoints.Num(), Points.Num()));
+					PCGE_LOG_C(Verbose, LogOnly, InOptionalContext, FText::Format(LOCTEXT("GenerationInfo", "Generated {0} points from {1} source points"), PrunedData->GetNumPoints(), InputNumPoints));
 				}
 			}
 		}
@@ -763,18 +787,22 @@ bool FPCGSelfPruningElement::PrepareDataInternal(FPCGContext* InContext) const
 		}
 
 		// Retro-compatibility
-		OutState.InputPointData = SpatialData->ToPointData();
-		if (!OutState.InputPointData)
+		OutState.InputData = SpatialData->ToBasePointData(InContext);
+	PRAGMA_DISABLE_DEPRECATION_WARNINGS
+		OutState.InputPointData = Cast<UPCGPointData>(OutState.InputData);
+	PRAGMA_ENABLE_DEPRECATION_WARNINGS
+
+		if (!OutState.InputData)
 		{
 			return EPCGTimeSliceInitResult::NoOperation;
 		}
 
 		if (Settings->Parameters.PruningType != EPCGSelfPruningType::RemoveDuplicates && Settings->Parameters.bUseCollisionAttribute)
 		{
-			FPCGAttributePropertyInputSelector InputSelector = Settings->Parameters.CollisionAttribute.CopyAndFixLast(OutState.InputPointData);
+			FPCGAttributePropertyInputSelector InputSelector = Settings->Parameters.CollisionAttribute.CopyAndFixLast(OutState.InputData);
 
-			TUniquePtr<const IPCGAttributeAccessor> InputAccessor = PCGAttributeAccessorHelpers::CreateConstAccessor(OutState.InputPointData, InputSelector);
-			TUniquePtr<const IPCGAttributeAccessorKeys> InputKeys = PCGAttributeAccessorHelpers::CreateConstKeys(OutState.InputPointData, InputSelector);
+			TUniquePtr<const IPCGAttributeAccessor> InputAccessor = PCGAttributeAccessorHelpers::CreateConstAccessor(OutState.InputData, InputSelector);
+			TUniquePtr<const IPCGAttributeAccessorKeys> InputKeys = PCGAttributeAccessorHelpers::CreateConstKeys(OutState.InputData, InputSelector);
 
 			TArray<FSoftObjectPath> Meshes;
 			if (OutState.CollisionWrapper.Prepare(InputAccessor.Get(), InputKeys.Get(), Meshes))
@@ -815,9 +843,9 @@ bool FPCGSelfPruningElement::ExecuteInternal(FPCGContext* Context) const
 		if (bIsDone)
 		{
 			FPCGTaggedData& Output = Context->OutputData.TaggedData.Add_GetRef(Inputs[IterationIndex]);
-			if (IterState.OutputPointData != nullptr)
+			if (IterState.OutputData != nullptr)
 			{
-				Output.Data = IterState.OutputPointData;
+				Output.Data = IterState.OutputData;
 			}
 		}
 
diff --git a/Engine/Plugins/PCG/Source/PCG/Private/Elements/PCGSortAttributes.cpp b/Engine/Plugins/PCG/Source/PCG/Private/Elements/PCGSortAttributes.cpp
--- a/Engine/Plugins/PCG/Source/PCG/Private/Elements/PCGSortAttributes.cpp
+++ b/Engine/Plugins/PCG/Source/PCG/Private/Elements/PCGSortAttributes.cpp
@@ -3,7 +3,7 @@
 #include "Elements/PCGSortAttributes.h"
 
 #include "PCGContext.h"
-#include "Data/PCGPointData.h"
+#include "Data/PCGBasePointData.h"
 #include "Metadata/Accessors/IPCGAttributeAccessor.h"
 #include "Metadata/Accessors/PCGAttributeAccessorHelpers.h"
 #include "Metadata/Accessors/PCGAttributeAccessorKeys.h"
@@ -83,10 +83,18 @@ bool FPCGSortAttributesElement::ExecuteInternal(FPCGContext* Context) const
 			}
 
 			UPCGData* OutputData = nullptr;
-			if (const UPCGPointData* InputPointData = Cast<const UPCGPointData>(InputData))
+			if (const UPCGBasePointData* InputPointData = Cast<const UPCGBasePointData>(InputData))
 			{
-				UPCGPointData* OutputPointData = static_cast<UPCGPointData*>(InputData->DuplicateData(Context));
-				PCGAttributeAccessorHelpers::SortByAttribute(*Accessor, *Keys, OutputPointData->GetMutablePoints(), Settings->SortMethod == EPCGSortMethod::Ascending);
+				TArray<int32> SortedKeyIndices = PCGAttributeAccessorHelpers::SortKeyIndicesByAttribute(*Accessor, *Keys, Keys->GetNum(), Settings->SortMethod == EPCGSortMethod::Ascending);
+				
+				UPCGBasePointData* OutputPointData = FPCGContext::NewPointData_AnyThread(Context);
+
+				FPCGInitializeFromDataParams InitalizeFromDataParams(InputPointData);
+				InitalizeFromDataParams.bInheritSpatialData = false;
+				OutputPointData->InitializeFromDataWithParams(InitalizeFromDataParams);
+
+				UPCGBasePointData::CopyPoints(InputPointData, OutputPointData, SortedKeyIndices, /*bCopyAll=*/false);
+
 				OutputData = OutputPointData;
 			}
 			else if (const UPCGMetadata* InMetadata = InputData->ConstMetadata())
diff --git a/Engine/Plugins/PCG/Source/PCG/Private/Elements/PCGSpatialNoise.cpp b/Engine/Plugins/PCG/Source/PCG/Private/Elements/PCGSpatialNoise.cpp
--- a/Engine/Plugins/PCG/Source/PCG/Private/Elements/PCGSpatialNoise.cpp
+++ b/Engine/Plugins/PCG/Source/PCG/Private/Elements/PCGSpatialNoise.cpp
@@ -10,7 +10,7 @@
 #include "PCGCustomVersion.h"
 #include "PCGPin.h"
 
-#include "Data/PCGPointData.h"
+#include "Data/PCGBasePointData.h"
 #include "Data/PCGSpatialData.h"
 #include "Helpers/PCGAsync.h"
 #include "Helpers/PCGHelpers.h"
@@ -361,8 +361,8 @@ namespace PCGSpatialNoise
 
 	struct FBufferParams
 	{
-		const UPCGPointData* InputPointData = nullptr;
-		UPCGPointData* OutputPointData = nullptr;
+		const UPCGBasePointData* InputPointData = nullptr;
+		UPCGBasePointData* OutputPointData = nullptr;
 	};
 
 	template<typename T>
@@ -398,74 +398,74 @@ namespace PCGSpatialNoise
 	template<typename FractalNoiseFunc>
 	void DoFractal2D(const FSharedParams& SharedParams, const FBufferParams& BufferParams, const UPCGSpatialNoiseSettings& Settings, FractalNoiseFunc&& FractalNoise)
 	{
-		const TArray<FPCGPoint>& SrcPoints = BufferParams.InputPointData->GetPoints();
+		const int32 NumPoints = BufferParams.InputPointData->GetNumPoints();
 
 		struct FProcessResults
 		{
-			TArray<FPCGPoint> Points;
 			TArray<double> Values;
 		};
 
 		FProcessResults Results;
 
-		FPCGAsync::AsyncProcessingOneToOneEx(
+		FPCGAsync::AsyncProcessingOneToOneRangeEx(
 			SharedParams.Context ? &SharedParams.Context->AsyncState : nullptr,
-			SrcPoints.Num(),
-			[&Results, Count = SrcPoints.Num()]()
+			NumPoints,
+			[&Results, NumPoints]()
 			{
 				// initialize
-				Results.Points.SetNumUninitialized(Count);
-				Results.Values.SetNumUninitialized(Count);
+				Results.Values.SetNumUninitialized(NumPoints);
 			},
-			[
-				&SharedParams,
-				&BufferParams,
-				&Results,
-				&SrcPoints,
-				&FractalNoise
-			](const int32 ReadIndex, const int32 WriteIndex)
+			[&SharedParams, &BufferParams, &Results, &FractalNoise](int32 StartReadIndex, int32 StartWriteIndex, int32 Count)
 			{
-				const FPCGPoint& InPoint = SrcPoints[ReadIndex];
-				FPCGPoint& OutPoint = Results.Points[WriteIndex];
+				int32 NumWritten = 0;
 
-				OutPoint = InPoint;
+				if (!BufferParams.OutputPointData->HasSpatialDataParent())
+				{
+					BufferParams.InputPointData->CopyPointsTo(BufferParams.OutputPointData, StartReadIndex, StartWriteIndex, Count);
+				}
 
-				const FVector PointPos = InPoint.Transform.GetTranslation();
+				const TConstPCGValueRange<FTransform> InTransformRange = BufferParams.InputPointData->GetConstTransformValueRange();
 
-				double Value = 0.0;
+				for (int32 ReadIndex = StartReadIndex; ReadIndex < StartReadIndex + Count; ++ReadIndex)
+				{	
+					const FVector PointPos = InTransformRange[ReadIndex].GetTranslation();
 
-				if (SharedParams.bTiling)
-				{
-					const FLocalCoordinates2D LocalCoords = CalcLocalCoordinates2D(
-						SharedParams.ActorLocalBox,
-						SharedParams.ActorTransformInverse,
-						FVector2D(SharedParams.Transform.GetScale3D()),
-						PointPos
-					);
+					double Value = 0.0;
 
-					Value = FMath::BiLerp(
-						FractalNoise(FVector2D(LocalCoords.X0, LocalCoords.Y0), SharedParams.Iterations),
-						FractalNoise(FVector2D(LocalCoords.X1, LocalCoords.Y0), SharedParams.Iterations),
-						FractalNoise(FVector2D(LocalCoords.X0, LocalCoords.Y1), SharedParams.Iterations),
-						FractalNoise(FVector2D(LocalCoords.X1, LocalCoords.Y1), SharedParams.Iterations),
-						LocalCoords.FracX,
-						LocalCoords.FracY
-					);
-				}
-				else
-				{
-					Value = FractalNoise(FVector2D(SharedParams.Transform.TransformPosition(PointPos)), SharedParams.Iterations);
-				}
+					if (SharedParams.bTiling)
+					{
+						const FLocalCoordinates2D LocalCoords = CalcLocalCoordinates2D(
+							SharedParams.ActorLocalBox,
+							SharedParams.ActorTransformInverse,
+							FVector2D(SharedParams.Transform.GetScale3D()),
+							PointPos
+						);
+
+						Value = FMath::BiLerp(
+							FractalNoise(FVector2D(LocalCoords.X0, LocalCoords.Y0), SharedParams.Iterations),
+							FractalNoise(FVector2D(LocalCoords.X1, LocalCoords.Y0), SharedParams.Iterations),
+							FractalNoise(FVector2D(LocalCoords.X0, LocalCoords.Y1), SharedParams.Iterations),
+							FractalNoise(FVector2D(LocalCoords.X1, LocalCoords.Y1), SharedParams.Iterations),
+							LocalCoords.FracX,
+							LocalCoords.FracY
+						);
+					}
+					else
+					{
+						Value = FractalNoise(FVector2D(SharedParams.Transform.TransformPosition(PointPos)), SharedParams.Iterations);
+					}
 
-				Value = ApplyContrast(SharedParams.Brightness + Value, SharedParams.Contrast);
+					Value = ApplyContrast(SharedParams.Brightness + Value, SharedParams.Contrast);
 
-				Results.Values[WriteIndex] = Value;
-			},
-			/* bEnableTimeSlicing */ false
-		);
+					const int32 WriteIndex = StartWriteIndex + NumWritten;
+					Results.Values[WriteIndex] = Value;
+					++NumWritten;
+				}
 
-		// now apply these results
-		BufferParams.OutputPointData->GetMutablePoints() = MoveTemp(Results.Points);
+				check(NumWritten == Count);
+				return NumWritten;
+			},
+			/*bTimeSliceEnabled=*/false);
 
 		SetAttributeHelper<double>(BufferParams.OutputPointData, Settings.ValueTarget, Results.Values);
 	}
@@ -473,76 +473,73 @@ namespace PCGSpatialNoise
 	template<typename FractalNoiseFunc>
 	void DoEdgeMask2D(const FSharedParams& SharedParams, const FBufferParams& BufferParams, const UPCGSpatialNoiseSettings& Settings, FractalNoiseFunc&& FractalNoise)
 	{
-		const TArray<FPCGPoint>& SrcPoints = BufferParams.InputPointData->GetPoints();
+		const int32 NumPoints = BufferParams.InputPointData->GetNumPoints();
 
 		struct FProcessResults
 		{
-			TArray<FPCGPoint> Points;
 			TArray<double> Values;
 		};
 
 		FProcessResults Results;
 
-		FPCGAsync::AsyncProcessingOneToOneEx(
+		FPCGAsync::AsyncProcessingOneToOneRangeEx(
 			SharedParams.Context ? &SharedParams.Context->AsyncState : nullptr,
-			SrcPoints.Num(),
-			[&Results, Count = SrcPoints.Num()]()
+			NumPoints,
+			[&Results, NumPoints]()
 			{
 				// initialize
-				Results.Points.SetNumUninitialized(Count);
-				Results.Values.SetNumUninitialized(Count);
+				Results.Values.SetNumUninitialized(NumPoints);
 			},
-			[
-				&SharedParams,
-				&BufferParams,
-				&SrcPoints,
-				&Results,
-				FractalNoise,
-				EdgeBlendDistance = Settings.EdgeBlendDistance,
-				EdgeBlendCurveIntensity = Settings.EdgeBlendCurveIntensity,
-				EdgeBlendCurveOffset = Settings.EdgeBlendCurveOffset
-			] (const int32 ReadIndex, const int32 WriteIndex) {
+			[&SharedParams, &BufferParams, &Results, &FractalNoise, &Settings](int32 StartReadIndex, int32 StartWriteIndex, int32 Count)
+			{
+				int32 NumWritten = 0;
 
-				const FPCGPoint& InPoint = SrcPoints[ReadIndex];
-				FPCGPoint& OutPoint = Results.Points[WriteIndex];
+				if (!BufferParams.OutputPointData->HasSpatialDataParent())
+				{
+					BufferParams.InputPointData->CopyPointsTo(BufferParams.OutputPointData, StartReadIndex, StartWriteIndex, Count);
+				}
 
-				OutPoint = InPoint;
+				const TConstPCGValueRange<FTransform> InTransformRange = BufferParams.InputPointData->GetConstTransformValueRange();
 
-				const FVector InPosition = InPoint.Transform.GetTranslation();
+				for (int32 ReadIndex = StartReadIndex; ReadIndex < StartReadIndex + Count; ++ReadIndex)
+				{
+					const FVector InPosition = InTransformRange[ReadIndex].GetTranslation();
 
-				const FLocalCoordinates2D LocalCoords = CalcLocalCoordinates2D(
-					SharedParams.ActorLocalBox,
-					SharedParams.ActorTransformInverse,
-					FVector2D(SharedParams.Transform.GetScale3D()),
-					InPosition
-				);
+					const FLocalCoordinates2D LocalCoords = CalcLocalCoordinates2D(
+						SharedParams.ActorLocalBox,
+						SharedParams.ActorTransformInverse,
+						FVector2D(SharedParams.Transform.GetScale3D()),
+						InPosition
+					);
 
-				const double EdgeBlendAmount = CalcEdgeBlendAmount2D(LocalCoords, EdgeBlendDistance);
+					const double EdgeBlendAmount = CalcEdgeBlendAmount2D(LocalCoords, Settings.EdgeBlendDistance);
 
-				double Value = 1.0;
+					double Value = 1.0;
 
-				if (EdgeBlendAmount > 0.0001)
-				{
-					const double NoiseValue = FractalNoise(FVector2D(SharedParams.Transform.TransformPosition(InPosition)), SharedParams.Iterations);
-					const double RemappedNoiseValue = ApplyContrast(NoiseValue, EdgeBlendCurveIntensity);
-					const double OffsetAmount = FMath::Pow(EdgeBlendAmount, EdgeBlendCurveOffset);
-
-					const double NoisedAmount = FMath::Lerp(
-						FMath::Min(OffsetAmount, RemappedNoiseValue*OffsetAmount),
-						1.0 - FMath::Min(1.0-OffsetAmount, (1.0-RemappedNoiseValue) * (1.0-OffsetAmount)),
-						OffsetAmount
-					);
+					if (EdgeBlendAmount > 0.0001)
+					{
+						const double NoiseValue = FractalNoise(FVector2D(SharedParams.Transform.TransformPosition(InPosition)), SharedParams.Iterations);
+						const double RemappedNoiseValue = ApplyContrast(NoiseValue, Settings.EdgeBlendCurveIntensity);
+						const double OffsetAmount = FMath::Pow(EdgeBlendAmount, Settings.EdgeBlendCurveOffset);
+
+						const double NoisedAmount = FMath::Lerp(
+							FMath::Min(OffsetAmount, RemappedNoiseValue * OffsetAmount),
+							1.0 - FMath::Min(1.0 - OffsetAmount, (1.0 - RemappedNoiseValue) * (1.0 - OffsetAmount)),
+							OffsetAmount
+						);
+
+						Value = 1.0 - FMath::Clamp(ApplyContrast(NoisedAmount + SharedParams.Brightness, SharedParams.Contrast), 0.0, 1.0);
+					}
 
-					Value = 1.0 - FMath::Clamp(ApplyContrast(NoisedAmount + SharedParams.Brightness, SharedParams.Contrast), 0.0, 1.0);
+					const int32 WriteIndex = StartWriteIndex + NumWritten;
+					Results.Values[WriteIndex] = Value;
+					++NumWritten;
 				}
 
-				Results.Values[WriteIndex] = Value;
+				check(NumWritten == Count);
+				return NumWritten;
 			},
-			/* bEnableTimslicing */ false			
-		);
-
-		// now apply these results
-		BufferParams.OutputPointData->GetMutablePoints() = MoveTemp(Results.Points);
+			/*bTimeSliceEnabled=*/false);
 
 		SetAttributeHelper<double>(BufferParams.OutputPointData, Settings.ValueTarget, Results.Values);
 	}
@@ -550,65 +547,71 @@ namespace PCGSpatialNoise
 	template<typename VoronoiHashFunc>
 	void DoVoronoi2DInternal(const FSharedParams& SharedParams, const FBufferParams& BufferParams, const UPCGSpatialNoiseSettings& Settings, VoronoiHashFunc&& HashFunc)
 	{
-		const TArray<FPCGPoint>& SrcPoints = BufferParams.InputPointData->GetPoints();
+		const int32 NumPoints = BufferParams.InputPointData->GetNumPoints();
 
 		struct FProcessResults
 		{
-			TArray<FPCGPoint> Points;
 			TArray<double> DistanceToEdge;
 			TArray<double> CellID;
 		};
 
 		FProcessResults Results;
 
-		FPCGAsync::AsyncProcessingOneToOneEx(
+		FPCGAsync::AsyncProcessingOneToOneRangeEx(
 			SharedParams.Context ? &SharedParams.Context->AsyncState : nullptr,
-			SrcPoints.Num(),
-			// initialize
-			[&Results, Count = SrcPoints.Num()]()
+			NumPoints,
+			[&BufferParams, &Results, NumPoints, &Settings]()
 			{
 				// initialize
-				Results.Points.SetNumUninitialized(Count);
-				Results.DistanceToEdge.SetNumUninitialized(Count);
-				Results.CellID.SetNumUninitialized(Count);
+				Results.DistanceToEdge.SetNumUninitialized(NumPoints);
+				Results.CellID.SetNumUninitialized(NumPoints);
+
+				if (Settings.bVoronoiOrientSamplesToCellEdge)
+				{
+					BufferParams.OutputPointData->AllocateProperties(EPCGPointNativeProperties::Transform);
+				}
 			},
-			// process
-			[
-				&SharedParams,
-				&BufferParams,
-				&HashFunc,
-				&SrcPoints,
-				&Results,
-				bVoronoiOrientSamplesToCellEdge = Settings.bVoronoiOrientSamplesToCellEdge
-			](const int32 ReadIndex, const int32 WriteIndex)
+			[&SharedParams, &BufferParams, &HashFunc, &Results, &Settings](int32 StartReadIndex, int32 StartWriteIndex, int32 Count)
 			{
-				const FPCGPoint& InPoint = SrcPoints[ReadIndex];
-				FPCGPoint& OutPoint = Results.Points[WriteIndex];
+				int32 NumWritten = 0;
+
+				if (!BufferParams.OutputPointData->HasSpatialDataParent())
+				{
+					BufferParams.InputPointData->CopyPointsTo(BufferParams.OutputPointData, StartReadIndex, StartWriteIndex, Count);
+				}
 
-				OutPoint = InPoint;
-				const FVector2D Position = FVector2D(SharedParams.Transform.TransformPosition(InPoint.Transform.GetTranslation()));
+				const TConstPCGValueRange<FTransform> InTransformRange = BufferParams.InputPointData->GetConstTransformValueRange();
+				TPCGValueRange<FTransform> OutTransformRange = Settings.bVoronoiOrientSamplesToCellEdge ? BufferParams.OutputPointData->GetTransformValueRange(/*bAllocate=*/false) : TPCGValueRange<FTransform>();
 
-				FVoronoi2DResult Result = CalcVoronoi2D(Position, HashFunc);
+				for (int32 ReadIndex = StartReadIndex; ReadIndex < StartReadIndex + Count; ++ReadIndex)
+				{
+					const FVector2D Position = FVector2D(SharedParams.Transform.TransformPosition(InTransformRange[ReadIndex].GetTranslation()));
 
-				const double Value = ApplyContrast(SharedParams.Brightness + Result.DistanceToEdge / MAGIC_SCALE_FACTOR, SharedParams.Contrast);
+					FVoronoi2DResult Result = CalcVoronoi2D(Position, HashFunc);
 
-				Results.DistanceToEdge[WriteIndex] = Value;
-				Results.CellID[WriteIndex] = Result.ID;
+					const double Value = ApplyContrast(SharedParams.Brightness + Result.DistanceToEdge / MAGIC_SCALE_FACTOR, SharedParams.Contrast);
 
-				if (bVoronoiOrientSamplesToCellEdge)
-				{
-					OutPoint.Transform.SetRotation(
-						FQuat(UKismetMathLibrary::FindLookAtRotation(
-							FVector::ZeroVector, 
-							FVector(Result.CellEdgeDirection.X, Result.CellEdgeDirection.Y, 0.0)
-						))
-					);
+					const int32 WriteIndex = StartWriteIndex + NumWritten;
+					Results.DistanceToEdge[WriteIndex] = Value;
+					Results.CellID[WriteIndex] = Result.ID;
+				
+					if (Settings.bVoronoiOrientSamplesToCellEdge)
+					{
+						OutTransformRange[WriteIndex].SetRotation(
+							FQuat(UKismetMathLibrary::FindLookAtRotation(
+								FVector::ZeroVector,
+								FVector(Result.CellEdgeDirection.X, Result.CellEdgeDirection.Y, 0.0)
+							))
+						);
+					}
+
+					++NumWritten;
 				}
-			},
-			/* bEnableTimslicing */ false
-		);
 
-		BufferParams.OutputPointData->GetMutablePoints() = MoveTemp(Results.Points);
+				check(NumWritten == Count);
+				return NumWritten;
+			},
+			/*bTimeSliceEnabled=*/false);
 
 		SetAttributeHelper<double>(BufferParams.OutputPointData, Settings.ValueTarget, Results.DistanceToEdge);
 		SetAttributeHelper<double>(BufferParams.OutputPointData, Settings.VoronoiCellIDTarget, Results.CellID);
@@ -719,7 +722,7 @@ bool FPCGSpatialNoise::ExecuteInternal(FPCGContext* Context) const
 	{
 		PCGSpatialNoise::FBufferParams BufferParams;
 
-		BufferParams.InputPointData = Cast<UPCGPointData>(Input.Data);
+		BufferParams.InputPointData = Cast<UPCGBasePointData>(Input.Data);
 
 		if (!BufferParams.InputPointData)
 		{
@@ -727,8 +730,15 @@ bool FPCGSpatialNoise::ExecuteInternal(FPCGContext* Context) const
 			continue;
 		}
 
-		BufferParams.OutputPointData = FPCGContext::NewObject_AnyThread<UPCGPointData>(Context);
+		BufferParams.OutputPointData = FPCGContext::NewPointData_AnyThread(Context);
 		BufferParams.OutputPointData->InitializeFromData(BufferParams.InputPointData);
+		BufferParams.OutputPointData->SetNumPoints(BufferParams.InputPointData->GetNumPoints());
+
+		if (!BufferParams.OutputPointData->HasSpatialDataParent())
+		{
+			BufferParams.OutputPointData->AllocateProperties(BufferParams.InputPointData->GetAllocatedProperties());
+		}
+
 		Context->OutputData.TaggedData.Add_GetRef(Input).Data = BufferParams.OutputPointData;
 
 		switch (Settings->Mode)
diff --git a/Engine/Plugins/PCG/Source/PCG/Private/Elements/PCGSurfaceSampler.cpp b/Engine/Plugins/PCG/Source/PCG/Private/Elements/PCGSurfaceSampler.cpp
--- a/Engine/Plugins/PCG/Source/PCG/Private/Elements/PCGSurfaceSampler.cpp
+++ b/Engine/Plugins/PCG/Source/PCG/Private/Elements/PCGSurfaceSampler.cpp
@@ -286,7 +286,7 @@ namespace PCGSurfaceSampler
 				TransformRange[WriteIndex] = OutPoint.Transform;
 				MetadataEntryRange[WriteIndex] = OutPoint.MetadataEntry;
 				ColorRange[WriteIndex] = OutPoint.Color;
-				NumWritten++;
+				++NumWritten;
 			}
 
 			return NumWritten;
diff --git a/Engine/Plugins/PCG/Source/PCG/Private/Elements/PCGTransformPoints.cpp b/Engine/Plugins/PCG/Source/PCG/Private/Elements/PCGTransformPoints.cpp
--- a/Engine/Plugins/PCG/Source/PCG/Private/Elements/PCGTransformPoints.cpp
+++ b/Engine/Plugins/PCG/Source/PCG/Private/Elements/PCGTransformPoints.cpp
@@ -5,7 +5,7 @@
 #include "PCGContext.h"
 #include "PCGPin.h"
 #include "Data/PCGSpatialData.h"
-#include "Data/PCGPointData.h"
+#include "Data/PCGBasePointData.h"
 #include "Helpers/PCGAsync.h"
 #include "Helpers/PCGHelpers.h"
 
@@ -61,7 +61,7 @@ bool FPCGTransformPointsElement::ExecuteInternal(FPCGContext* Context) const
 			continue;
 		}
 
-		const UPCGPointData* PointData = SpatialData->ToPointData(Context);
+		const UPCGBasePointData* PointData = SpatialData->ToBasePointData(Context);
 
 		if (!PointData)
 		{
@@ -93,12 +93,23 @@ bool FPCGTransformPointsElement::ExecuteInternal(FPCGContext* Context) const
 
 			SourceAttribute = static_cast<const FPCGMetadataAttribute<FTransform>*>(FoundAttribute);
 		}
+				
+		UPCGBasePointData* OutputData = FPCGContext::NewPointData_AnyThread(Context);
+		
+		OutputData->InitializeFromData(PointData);
+		OutputData->SetNumPoints(PointData->GetNumPoints());
 
-		const TArray<FPCGPoint>& Points = PointData->GetPoints();
+		const EPCGPointNativeProperties PropertiesToAllocate = bApplyToAttribute ? EPCGPointNativeProperties::MetadataEntry : (bRecomputeSeed ? (EPCGPointNativeProperties::Transform | EPCGPointNativeProperties::Seed) : EPCGPointNativeProperties::Transform);
+
+		if (OutputData->HasSpatialDataParent())
+		{
+			OutputData->AllocateProperties(PropertiesToAllocate);
+		}
+		else
+		{
+			OutputData->AllocateProperties(PointData->GetAllocatedProperties() | PropertiesToAllocate);
+		}
 
-		UPCGPointData* OutputData = FPCGContext::NewObject_AnyThread<UPCGPointData>(Context);
-		OutputData->InitializeFromData(PointData);
-		TArray<FPCGPoint>& OutputPoints = OutputData->GetMutablePoints();
 		Output.Data = OutputData;
 
 		FPCGMetadataAttribute<FTransform>* TargetAttribute = nullptr;
@@ -108,99 +119,129 @@ bool FPCGTransformPointsElement::ExecuteInternal(FPCGContext* Context) const
 		{
 			check(SourceAttribute && OutputData && OutputData->Metadata);
 			TargetAttribute = OutputData->Metadata->GetMutableTypedAttribute<FTransform>(LocalAttributeName);
-			AllMetadataEntries.SetNum(Points.Num());
+			AllMetadataEntries.SetNum(PointData->GetNumPoints());
 		}
 
-		FPCGAsync::AsyncPointProcessing(Context, Points.Num(), OutputPoints, [&](int32 Index, FPCGPoint& OutPoint)
+		auto ProcessRangeFunc = [&](int32 StartReadIndex, int32 StartWriteIndex, int32 Count)
 		{
-			const FPCGPoint& InPoint = Points[Index];
-			OutPoint = InPoint;
-
-			FRandomStream RandomSource(PCGHelpers::ComputeSeed(Seed, InPoint.Seed));
+			int32 NumWritten = 0;
+			
+			// Copy all properties except the ones we are going to modify (if we are not inheriting)
+			if (!OutputData->HasSpatialDataParent())
+			{
+				PointData->CopyPropertiesTo(OutputData, StartReadIndex, StartWriteIndex, Count, EPCGPointNativeProperties::All & ~PropertiesToAllocate);
+			}
 
-			const float OffsetX = RandomSource.FRandRange(OffsetMin.X, OffsetMax.X);
-			const float OffsetY = RandomSource.FRandRange(OffsetMin.Y, OffsetMax.Y);
-			const float OffsetZ = RandomSource.FRandRange(OffsetMin.Z, OffsetMax.Z);
-			const FVector RandomOffset(OffsetX, OffsetY, OffsetZ);
+			const TConstPCGValueRange<FTransform> ReadTransformRange = PointData->GetConstTransformValueRange();
+			const TConstPCGValueRange<int32> ReadSeedRange = PointData->GetConstSeedValueRange();
+			const TConstPCGValueRange<int64> ReadMetadataEntryRange = PointData->GetConstMetadataEntryValueRange();
 
-			const float RotationX = RandomSource.FRandRange(RotationMin.Pitch, RotationMax.Pitch);
-			const float RotationY = RandomSource.FRandRange(RotationMin.Yaw, RotationMax.Yaw);
-			const float RotationZ = RandomSource.FRandRange(RotationMin.Roll, RotationMax.Roll);
-			const FQuat RandomRotation(FRotator(RotationX, RotationY, RotationZ).Quaternion());
+			TPCGValueRange<FTransform> WriteTransformRange = bApplyToAttribute ? TPCGValueRange<FTransform>() : OutputData->GetTransformValueRange(/*bAllocate=*/false);
+			TPCGValueRange<int32> WriteSeedRange = bApplyToAttribute || !bRecomputeSeed ? TPCGValueRange<int32>() : OutputData->GetSeedValueRange(/*bAllocate=*/false);
+			TPCGValueRange<int64> WriteMetadataEntryRange = bApplyToAttribute ? OutputData->GetMetadataEntryValueRange(/*bAllocate=*/false) : TPCGValueRange<int64>();
 
-			FVector RandomScale;
-			if (bUniformScale)
+			for (int32 ReadIndex = StartReadIndex; ReadIndex < StartReadIndex + Count; ++ReadIndex)
 			{
-				RandomScale = FVector(RandomSource.FRandRange(ScaleMin.X, ScaleMax.X));
-			}
-			else
-			{
-				RandomScale.X = RandomSource.FRandRange(ScaleMin.X, ScaleMax.X);
-				RandomScale.Y = RandomSource.FRandRange(ScaleMin.Y, ScaleMax.Y);
-				RandomScale.Z = RandomSource.FRandRange(ScaleMin.Z, ScaleMax.Z);
-			}
-		
-			FTransform SourceTransform;
+				const int32 WriteIndex = StartWriteIndex + NumWritten;
+				
+				FRandomStream RandomSource(PCGHelpers::ComputeSeed(Seed, ReadSeedRange[ReadIndex]));
+
+				const float OffsetX = RandomSource.FRandRange(OffsetMin.X, OffsetMax.X);
+				const float OffsetY = RandomSource.FRandRange(OffsetMin.Y, OffsetMax.Y);
+				const float OffsetZ = RandomSource.FRandRange(OffsetMin.Z, OffsetMax.Z);
+				const FVector RandomOffset(OffsetX, OffsetY, OffsetZ);
+
+				const float RotationX = RandomSource.FRandRange(RotationMin.Pitch, RotationMax.Pitch);
+				const float RotationY = RandomSource.FRandRange(RotationMin.Yaw, RotationMax.Yaw);
+				const float RotationZ = RandomSource.FRandRange(RotationMin.Roll, RotationMax.Roll);
+				const FQuat RandomRotation(FRotator(RotationX, RotationY, RotationZ).Quaternion());
+
+				FVector RandomScale;
+				if (bUniformScale)
+				{
+					RandomScale = FVector(RandomSource.FRandRange(ScaleMin.X, ScaleMax.X));
+				}
+				else
+				{
+					RandomScale.X = RandomSource.FRandRange(ScaleMin.X, ScaleMax.X);
+					RandomScale.Y = RandomSource.FRandRange(ScaleMin.Y, ScaleMax.Y);
+					RandomScale.Z = RandomSource.FRandRange(ScaleMin.Z, ScaleMax.Z);
+				}
 
-			if (!bApplyToAttribute)
-			{
-				SourceTransform = InPoint.Transform;
-			}
-			else
-			{
-				SourceTransform = SourceAttribute->GetValueFromItemKey(InPoint.MetadataEntry);
-			}
+				FTransform SourceTransform;
 
-			FTransform FinalTransform = SourceTransform;
+				if (!bApplyToAttribute)
+				{
+					SourceTransform = ReadTransformRange[ReadIndex];
+				}
+				else
+				{
+					SourceTransform = SourceAttribute->GetValueFromItemKey(ReadMetadataEntryRange[ReadIndex]);
+				}
 
-			if (bAbsoluteOffset)
-			{
-				FinalTransform.SetLocation(SourceTransform.GetLocation() + RandomOffset); 
-			}
-			else
-			{
-				const FTransform RotatedTransform(SourceTransform.GetRotation());
-				FinalTransform.SetLocation(SourceTransform.GetLocation() + RotatedTransform.TransformPosition(RandomOffset)); 
-			}
+				FTransform FinalTransform = SourceTransform;
 
-			if (bAbsoluteRotation)
-			{
-				FinalTransform.SetRotation(RandomRotation);
-			}
-			else
-			{
-				FinalTransform.SetRotation(SourceTransform.GetRotation() * RandomRotation);
-			}
+				if (bAbsoluteOffset)
+				{
+					FinalTransform.SetLocation(SourceTransform.GetLocation() + RandomOffset);
+				}
+				else
+				{
+					const FTransform RotatedTransform(SourceTransform.GetRotation());
+					FinalTransform.SetLocation(SourceTransform.GetLocation() + RotatedTransform.TransformPosition(RandomOffset));
+				}
 
-			if (bAbsoluteScale)
-			{
-				FinalTransform.SetScale3D(RandomScale);
-			}
-			else
-			{
-				FinalTransform.SetScale3D(SourceTransform.GetScale3D() * RandomScale);
-			}
+				if (bAbsoluteRotation)
+				{
+					FinalTransform.SetRotation(RandomRotation);
+				}
+				else
+				{
+					FinalTransform.SetRotation(SourceTransform.GetRotation() * RandomRotation);
+				}
 
-			if (!bApplyToAttribute)
-			{
-				OutPoint.Transform = FinalTransform;
+				if (bAbsoluteScale)
+				{
+					FinalTransform.SetScale3D(RandomScale);
+				}
+				else
+				{
+					FinalTransform.SetScale3D(SourceTransform.GetScale3D() * RandomScale);
+				}
 
-				if (bRecomputeSeed)
+				if (!bApplyToAttribute)
 				{
-					const FVector& Position = FinalTransform.GetLocation();
-					OutPoint.Seed = PCGHelpers::ComputeSeed((int)Position.X, (int)Position.Y, (int)Position.Z);
+					WriteTransformRange[WriteIndex] = FinalTransform;
+
+					if (bRecomputeSeed)
+					{
+						const FVector& Position = FinalTransform.GetLocation();
+						WriteSeedRange[WriteIndex] = PCGHelpers::ComputeSeed((int)Position.X, (int)Position.Y, (int)Position.Z);
+					}
+				}
+				else
+				{
+					WriteMetadataEntryRange[WriteIndex] = OutputData->Metadata->AddEntryPlaceholder();
+					AllMetadataEntries[ReadIndex] = MakeTuple(WriteMetadataEntryRange[WriteIndex], ReadMetadataEntryRange[ReadIndex]);
+					TargetAttribute->SetValue(WriteMetadataEntryRange[WriteIndex], FinalTransform);
 				}
-			}
-			else
-			{
-				OutPoint.MetadataEntry = OutputData->Metadata->AddEntryPlaceholder();
-				AllMetadataEntries[Index] = MakeTuple(OutPoint.MetadataEntry, InPoint.MetadataEntry);
-				TargetAttribute->SetValue(OutPoint.MetadataEntry, FinalTransform);
-			}
 
-			return true;
-		});
+				++NumWritten;
+			}
 
+			check(NumWritten == Count);
+			return NumWritten;
+		};
+				
+		FPCGAsync::AsyncProcessingRangeEx(
+			&Context->AsyncState,
+			PointData->GetNumPoints(),
+			/*InitializeFunc=*/[](){},
+			ProcessRangeFunc,
+			/*MoveRangeFunc=*/[](int32 ReadIndex, int32 WriteIndex, int32 Count) {}, // This is a one to one element so no point get discarded, so empty lambda
+			/*FinalizeFunc=*/[](int32) {}, // Not useful here, so empty lambda
+			/*bTimeSliceEnabled=*/false);
+				
 		if (TargetAttribute)
 		{
 			OutputData->Metadata->AddDelayedEntries(AllMetadataEntries);
diff --git a/Engine/Plugins/PCG/Source/PCG/Private/Helpers/PCGBlueprintHelpers.cpp b/Engine/Plugins/PCG/Source/PCG/Private/Helpers/PCGBlueprintHelpers.cpp
--- a/Engine/Plugins/PCG/Source/PCG/Private/Helpers/PCGBlueprintHelpers.cpp
+++ b/Engine/Plugins/PCG/Source/PCG/Private/Helpers/PCGBlueprintHelpers.cpp
@@ -43,42 +43,12 @@ void UPCGBlueprintHelpers::SetSeedFromPosition(FPCGPoint& InPoint)
 
 FRandomStream UPCGBlueprintHelpers::GetRandomStreamFromPoint(const FPCGPoint& InPoint, const UPCGSettings* OptionalSettings, const UPCGComponent* OptionalComponent)
 {
-	int Seed = InPoint.Seed;
-
-	if (OptionalSettings && OptionalComponent)
-	{
-		Seed = PCGHelpers::ComputeSeed(Seed, OptionalSettings->Seed, OptionalComponent->Seed);
-	}
-	else if (OptionalSettings)
-	{
-		Seed = PCGHelpers::ComputeSeed(Seed, OptionalSettings->Seed);
-	}
-	else if (OptionalComponent)
-	{
-		Seed = PCGHelpers::ComputeSeed(Seed, OptionalComponent->Seed);
-	}
-
-	return FRandomStream(Seed);
+	return PCGHelpers::GetRandomStreamFromSeed(InPoint.Seed);
 }
 
 FRandomStream UPCGBlueprintHelpers::GetRandomStreamFromTwoPoints(const FPCGPoint& InPointA, const FPCGPoint& InPointB, const UPCGSettings* OptionalSettings, const UPCGComponent* OptionalComponent)
 {
-	int Seed = PCGHelpers::ComputeSeed(InPointA.Seed, InPointB.Seed);
-
-	if (OptionalSettings && OptionalComponent)
-	{
-		Seed = PCGHelpers::ComputeSeed(Seed, OptionalSettings->Seed, OptionalComponent->Seed);
-	}
-	else if (OptionalSettings)
-	{
-		Seed = PCGHelpers::ComputeSeed(Seed, OptionalSettings->Seed);
-	}
-	else if (OptionalComponent)
-	{
-		Seed = PCGHelpers::ComputeSeed(Seed, OptionalComponent->Seed);
-	}
-
-	return FRandomStream(Seed);
+	return PCGHelpers::GetRandomStreamFromTwoSeeds(InPointA.Seed, InPointB.Seed, OptionalSettings, OptionalComponent);
 }
 
 const UPCGSettings* UPCGBlueprintHelpers::GetSettings(FPCGContext& Context)
diff --git a/Engine/Plugins/PCG/Source/PCG/Private/Helpers/PCGHelpers.cpp b/Engine/Plugins/PCG/Source/PCG/Private/Helpers/PCGHelpers.cpp
--- a/Engine/Plugins/PCG/Source/PCG/Private/Helpers/PCGHelpers.cpp
+++ b/Engine/Plugins/PCG/Source/PCG/Private/Helpers/PCGHelpers.cpp
@@ -6,6 +6,7 @@
 #include "PCGContext.h"
 #include "PCGGraph.h"
 #include "PCGModule.h"
+#include "PCGSettings.h"
 #include "PCGSubsystem.h"
 #include "PCGWorldActor.h"
 #include "Grid/PCGPartitionActor.h"
@@ -45,6 +46,46 @@ namespace PCGHelpers
 		return ComputeSeed(static_cast<int>(InPosition.X), static_cast<int>(InPosition.Y), static_cast<int>(InPosition.Z));
 	}
 
+	FRandomStream GetRandomStreamFromSeed(int32 InSeed, const UPCGSettings* OptionalSettings, const UPCGComponent* OptionalComponent)
+	{
+		int Seed = InSeed;
+
+		if (OptionalSettings && OptionalComponent)
+		{
+			Seed = PCGHelpers::ComputeSeed(Seed, OptionalSettings->Seed, OptionalComponent->Seed);
+		}
+		else if (OptionalSettings)
+		{
+			Seed = PCGHelpers::ComputeSeed(Seed, OptionalSettings->Seed);
+		}
+		else if (OptionalComponent)
+		{
+			Seed = PCGHelpers::ComputeSeed(Seed, OptionalComponent->Seed);
+		}
+
+		return FRandomStream(Seed);
+	}
+
+	FRandomStream GetRandomStreamFromTwoSeeds(int32 SeedA, int32 SeedB, const UPCGSettings* OptionalSettings, const UPCGComponent* OptionalComponent)
+	{
+		int Seed = PCGHelpers::ComputeSeed(SeedA, SeedB);
+
+		if (OptionalSettings && OptionalComponent)
+		{
+			Seed = PCGHelpers::ComputeSeed(Seed, OptionalSettings->Seed, OptionalComponent->Seed);
+		}
+		else if (OptionalSettings)
+		{
+			Seed = PCGHelpers::ComputeSeed(Seed, OptionalSettings->Seed);
+		}
+		else if (OptionalComponent)
+		{
+			Seed = PCGHelpers::ComputeSeed(Seed, OptionalComponent->Seed);
+		}
+
+		return FRandomStream(Seed);
+	}
+
 	bool IsInsideBounds(const FBox& InBox, const FVector& InPosition)
 	{
 		return (InPosition.X >= InBox.Min.X) && (InPosition.X < InBox.Max.X) &&
diff --git a/Engine/Plugins/PCG/Source/PCG/Private/Metadata/Accessors/PCGAttributeAccessorKeys.cpp b/Engine/Plugins/PCG/Source/PCG/Private/Metadata/Accessors/PCGAttributeAccessorKeys.cpp
--- a/Engine/Plugins/PCG/Source/PCG/Private/Metadata/Accessors/PCGAttributeAccessorKeys.cpp
+++ b/Engine/Plugins/PCG/Source/PCG/Private/Metadata/Accessors/PCGAttributeAccessorKeys.cpp
@@ -3,6 +3,8 @@
 #include "Metadata/Accessors/PCGAttributeAccessorKeys.h"
 
 #include "PCGPoint.h"
+#include "Data/PCGBasePointData.h"
+#include "Data/PCGPointData.h"
 #include "Data/PCGSplineData.h"
 #include "Metadata/PCGMetadata.h"
 #include "Metadata/PCGMetadataAttribute.h"
@@ -183,39 +185,152 @@ FPCGAttributeAccessorKeysPointsSubset::FPCGAttributeAccessorKeysPointsSubset(TAr
 	, Points(std::move(*reinterpret_cast<TArray<FPCGPoint*>*>(&InPointPtrs)))
 {}
 
+FPCGAttributeAccessorKeysPointsSubset::FPCGAttributeAccessorKeysPointsSubset(const UPCGBasePointData* InPointData, const TArrayView<const int32>& InPointIndices)
+	: IPCGAttributeAccessorKeys(/*bInReadOnly=*/ true)
+	, PointData(const_cast<UPCGBasePointData*>(InPointData))
+	, PointIndices(InPointIndices)
+{}
+
+FPCGAttributeAccessorKeysPointsSubset::FPCGAttributeAccessorKeysPointsSubset(UPCGBasePointData* InPointData, const TArrayView<const int32>& InPointIndices)
+	: IPCGAttributeAccessorKeys(/*bInReadOnly=*/ false)
+	, PointData(InPointData)
+	, PointIndices(InPointIndices)
+{}
+
 bool FPCGAttributeAccessorKeysPointsSubset::IsClassSupported(const UStruct* InClass) const
 {
 	return PCGAttributeAccessorKeys::IsClassSupported<FPCGPoint>(InClass);
 }
+bool FPCGAttributeAccessorKeysPointsSubset::GetKeyIndices(int32 InStart, int32 InCount, TArray<int32>& OutKeyIndices, bool& bOutContiguous) const
+{
+	bOutContiguous = false;
+
+	if (PointIndices.Num())
+	{
+		OutKeyIndices.SetNumUninitialized(InCount);
+		// @todo_pcg: could implement a memcpy version if InStart+Count <= PointIndices.Num()
+		for (int32 Index = 0; Index < InCount; ++Index)
+		{
+			const int32 KeyIndex = (InStart + Index) % PointIndices.Num();
+			OutKeyIndices[Index] = PointIndices[KeyIndex];
+		}
+
+		return true;
+	}
+
+	return false;
+}
 
 bool FPCGAttributeAccessorKeysPointsSubset::GetPointKeys(int32 InStart, TArrayView<FPCGPoint*> OutPoints)
 {
-	return PCGAttributeAccessorKeys::GetKeys(Points, InStart, OutPoints, [](FPCGPoint* Point) -> FPCGPoint* { return Point; });
+	if (Points.Num() > 0)
+	{
+		return PCGAttributeAccessorKeys::GetKeys(Points, InStart, OutPoints, [](FPCGPoint* Point) -> FPCGPoint* { return Point; });
+	}
+	else if (UPCGPointData* LocalPointData = Cast<UPCGPointData>(PointData))
+	{
+		TArray<FPCGPoint>& LocalPoints = LocalPointData->GetMutablePoints();
+		return PCGAttributeAccessorKeys::GetKeys(PointIndices, InStart, OutPoints, [&LocalPoints](int32 Index) -> FPCGPoint* { return &LocalPoints[Index]; });
+	}
+	else
+	{
+		return false;
+	}
 }
 
 bool FPCGAttributeAccessorKeysPointsSubset::GetPointKeys(int32 InStart, TArrayView<const FPCGPoint*> OutPoints) const
 {
-	return PCGAttributeAccessorKeys::GetKeys(Points, InStart, OutPoints, [](const FPCGPoint* Point) -> const FPCGPoint* { return Point; });
+	if (Points.Num() > 0)
+	{
+		return PCGAttributeAccessorKeys::GetKeys(Points, InStart, OutPoints, [](const FPCGPoint* Point) -> const FPCGPoint* { return Point; });
+	}
+	else if (UPCGPointData* LocalPointData = Cast<UPCGPointData>(PointData))
+	{
+		const TArray<FPCGPoint>& LocalPoints = LocalPointData->GetPoints();
+		return PCGAttributeAccessorKeys::GetKeys(PointIndices, InStart, OutPoints, [&LocalPoints](int32 Index) -> const FPCGPoint* { return &LocalPoints[Index]; });
+	}
+	else
+	{
+		return false;
+	}
 }
 
 bool FPCGAttributeAccessorKeysPointsSubset::GetGenericObjectKeys(int32 InStart, TArrayView<void*> OutObjects)
 {
-	return PCGAttributeAccessorKeys::GetKeys(Points, InStart, OutObjects, [](FPCGPoint* Point) -> void* { return Point; });
+	if (Points.Num() > 0)
+	{
+		return PCGAttributeAccessorKeys::GetKeys(Points, InStart, OutObjects, [](FPCGPoint* Point) -> void* { return Point; });
+	}
+	else if (UPCGPointData* LocalPointData = Cast<UPCGPointData>(PointData))
+	{
+		TArray<FPCGPoint>& LocalPoints = LocalPointData->GetMutablePoints();
+		return PCGAttributeAccessorKeys::GetKeys(PointIndices, InStart, OutObjects, [&LocalPoints](int32 Index) -> void* { return &LocalPoints[Index]; });
+	}
+	else
+	{
+		return false;
+	}
 }
 
 bool FPCGAttributeAccessorKeysPointsSubset::GetGenericObjectKeys(int32 InStart, TArrayView<const void*> OutObjects) const
 {
-	return PCGAttributeAccessorKeys::GetKeys(Points, InStart, OutObjects, [](const FPCGPoint* Point) -> const void* { return Point; });
+	if (Points.Num() > 0)
+	{
+		return PCGAttributeAccessorKeys::GetKeys(Points, InStart, OutObjects, [](const FPCGPoint* Point) -> const void* { return Point; });
+	}
+	else if (UPCGPointData* LocalPointData = Cast<UPCGPointData>(PointData))
+	{
+		const TArray<FPCGPoint>& LocalPoints = LocalPointData->GetPoints();
+		return PCGAttributeAccessorKeys::GetKeys(PointIndices, InStart, OutObjects, [&LocalPoints](int32 Index) -> const void* { return &LocalPoints[Index]; });
+	}
+	else
+	{
+		return false;
+	}
 }
 
 bool FPCGAttributeAccessorKeysPointsSubset::GetMetadataEntryKeys(int32 InStart, TArrayView<PCGMetadataEntryKey*> OutEntryKeys)
 {
-	return PCGAttributeAccessorKeys::GetKeys(Points, InStart, OutEntryKeys, [](FPCGPoint* Point) -> PCGMetadataEntryKey* { return &(Point->MetadataEntry); });
+	if (Points.Num() > 0)
+	{
+		return PCGAttributeAccessorKeys::GetKeys(Points, InStart, OutEntryKeys, [](FPCGPoint* Point) -> PCGMetadataEntryKey* { return &(Point->MetadataEntry); });
+	}
+	else if (PointIndices.Num() && PointData)
+	{
+		TPCGValueRange<int64> MetadataEntryKeys = PointData->GetMetadataEntryValueRange();
+
+		for (int32 Index = 0; Index < OutEntryKeys.Num(); ++Index)
+		{
+			const int32 PointIndex = (InStart + Index) % PointIndices.Num();
+			OutEntryKeys[Index] = &MetadataEntryKeys[PointIndices[PointIndex]];
+		}
+
+		return true;
+	}
+
+	return false;
 }
 
 bool FPCGAttributeAccessorKeysPointsSubset::GetMetadataEntryKeys(int32 InStart, TArrayView<const PCGMetadataEntryKey*> OutEntryKeys) const
 {
-	return PCGAttributeAccessorKeys::GetKeys(Points, InStart, OutEntryKeys, [](const FPCGPoint* Point) -> const PCGMetadataEntryKey* { return &(Point->MetadataEntry); });
+	if (Points.Num() > 0)
+	{
+		return PCGAttributeAccessorKeys::GetKeys(Points, InStart, OutEntryKeys, [](const FPCGPoint* Point) -> const PCGMetadataEntryKey* { return &(Point->MetadataEntry); });
+	}
+	else if (PointIndices.Num() > 0 && PointData)
+	{
+		const TConstPCGValueRange<int64> MetadataEntryKeys = PointData->GetConstMetadataEntryValueRange();
+
+		for (int32 Index = 0; Index < OutEntryKeys.Num(); ++Index)
+		{
+			const int32 PointIndex = (InStart + Index) % PointIndices.Num();
+			OutEntryKeys[Index] = &MetadataEntryKeys[PointIndices[PointIndex]];
+		}
+
+		return true;
+	}
+
+	return false;
 }
 
 /////////////////////////////////////////////////////////
diff --git a/Engine/Plugins/PCG/Source/PCG/Private/Metadata/PCGMetadataPartitionCommon.cpp b/Engine/Plugins/PCG/Source/PCG/Private/Metadata/PCGMetadataPartitionCommon.cpp
--- a/Engine/Plugins/PCG/Source/PCG/Private/Metadata/PCGMetadataPartitionCommon.cpp
+++ b/Engine/Plugins/PCG/Source/PCG/Private/Metadata/PCGMetadataPartitionCommon.cpp
@@ -5,7 +5,7 @@
 #include "PCGContext.h"
 #include "PCGModule.h"
 #include "PCGParamData.h"
-#include "Data/PCGPointData.h"
+#include "Data/PCGBasePointData.h"
 #include "Data/PCGSpatialData.h"
 #include "Elements/Metadata/PCGMetadataElementCommon.h"
 #include "Metadata/PCGMetadata.h"
@@ -43,8 +43,6 @@ namespace PCGMetadataPartitionCommon
 			return {};
 		}
 
-		TMap<PCGMetadataValueKey, UPCGPointData*> ValueToData;
-
 		// Get all value keys (-1 + 0 - N)
 		const int64 MetadataValueKeyCount = InAttribute->GetValueKeyOffsetForChild();
 
@@ -422,7 +420,7 @@ namespace PCGMetadataPartitionCommon
 	/**
 	* Do a partition on the given point data for the selector
 	*/
-	TArray<UPCGData*> AttributePointPartition(const UPCGPointData* InData, const TArrayView<const FPCGAttributePropertySelector>& InSelectorArrayView, FPCGContext* InOptionalContext, bool bSilenceMissingAttributeErrors)
+	TArray<UPCGData*> AttributePointPartition(const UPCGBasePointData* InData, const TArrayView<const FPCGAttributePropertySelector>& InSelectorArrayView, FPCGContext* InOptionalContext, bool bSilenceMissingAttributeErrors)
 	{
 		TArray<TArray<int32>> Partition = AttributeGenericPartition(InData, InSelectorArrayView, InOptionalContext, bSilenceMissingAttributeErrors);
 		if (Partition.IsEmpty())
@@ -432,7 +430,6 @@ namespace PCGMetadataPartitionCommon
 
 		TArray<UPCGData*> PartitionedData;
 		PartitionedData.Reserve(Partition.Num());
-		const TArray<FPCGPoint>& Points = InData->GetPoints();
 
 		for (TArray<int32>& Indices : Partition)
 		{
@@ -441,22 +438,21 @@ namespace PCGMetadataPartitionCommon
 				continue;
 			}
 
-			UPCGPointData* CurrentPointData = FPCGContext::NewObject_AnyThread<UPCGPointData>(InOptionalContext);
+			UPCGBasePointData* CurrentPointData = FPCGContext::NewPointData_AnyThread(InOptionalContext);
 			PartitionedData.Add(CurrentPointData);
-			CurrentPointData->InitializeFromData(InData);
 
-			TArray<FPCGPoint>& OutPoints = CurrentPointData->GetMutablePoints();
-			OutPoints.Reserve(Indices.Num());
-			for (int32 Index : Indices)
-			{
-				OutPoints.Add(Points[Index]);
-			}
+			FPCGInitializeFromDataParams InitializeFromDataParams(InData);
+			InitializeFromDataParams.bInheritSpatialData = false;
+
+			CurrentPointData->InitializeFromDataWithParams(InitializeFromDataParams);
+
+			UPCGBasePointData::CopyPoints(InData, CurrentPointData, Indices, /*bCopyAll=*/false);
 		}
 
 		return PartitionedData;
 	}
 
-	UPCGData* RemoveDuplicatesPoint(const UPCGPointData* InData, const TArrayView<const FPCGAttributePropertySelector>& InSelectorArrayView, FPCGContext* InOptionalContext, bool bSilenceMissingAttributeErrors)
+	UPCGData* RemoveDuplicatesPoint(const UPCGBasePointData* InData, const TArrayView<const FPCGAttributePropertySelector>& InSelectorArrayView, FPCGContext* InOptionalContext, bool bSilenceMissingAttributeErrors)
 	{
 		TArray<TArray<int32>> Partition = AttributeGenericPartition(InData, InSelectorArrayView, InOptionalContext, bSilenceMissingAttributeErrors);
 		if (Partition.IsEmpty())
@@ -464,28 +460,29 @@ namespace PCGMetadataPartitionCommon
 			return nullptr;
 		}
 
-		const TArray<FPCGPoint>& Points = InData->GetPoints();
-
-		UPCGPointData* OutputPointData = nullptr;
+		UPCGBasePointData* OutputPointData = nullptr;
+		TArray<int32> IndicesToCopy;
+		IndicesToCopy.Reserve(Partition.Num());
 
 		for (TArray<int32>& Indices : Partition)
 		{
-			if (Indices.IsEmpty())
-			{
-				continue;
-			}
+			IndicesToCopy.Add(Indices[0]);
+		}
 
-			if (!OutputPointData)
-			{
-				OutputPointData = FPCGContext::NewObject_AnyThread<UPCGPointData>(InOptionalContext);
-				OutputPointData->InitializeFromData(InData);
+		if (IndicesToCopy.IsEmpty())
+		{
+			return nullptr;
+		}
 
-				OutputPointData->GetMutablePoints().Reserve(Partition.Num());
-			}
+		OutputPointData = FPCGContext::NewPointData_AnyThread(InOptionalContext);
 
-			check(OutputPointData);
-			OutputPointData->GetMutablePoints().Add(Points[Indices[0]]);
-		}
+		FPCGInitializeFromDataParams InitializeFromDataParams(InData);
+		InitializeFromDataParams.bInheritSpatialData = false;
+
+		OutputPointData->InitializeFromDataWithParams(InitializeFromDataParams);
+		OutputPointData->SetNumPoints(IndicesToCopy.Num());
+		
+		UPCGBasePointData::CopyPoints(InData, OutputPointData, IndicesToCopy, /*bCopyAll=*/false);
 
 		return OutputPointData;
 	}
@@ -646,7 +643,7 @@ namespace PCGMetadataPartitionCommon
 
 	TArray<UPCGData*> AttributePartition(const UPCGData* InData, const TArrayView<const FPCGAttributePropertySelector>& InSelectorArrayView, FPCGContext* InOptionalContext, bool bSilenceMissingAttributeErrors)
 	{
-		if (const UPCGPointData* InPointData = Cast<UPCGPointData>(InData))
+		if (const UPCGBasePointData* InPointData = Cast<UPCGBasePointData>(InData))
 		{
 			return AttributePointPartition(InPointData, InSelectorArrayView, InOptionalContext, bSilenceMissingAttributeErrors);
 		}
@@ -663,7 +660,7 @@ namespace PCGMetadataPartitionCommon
 
 	UPCGData* RemoveDuplicates(const UPCGData* InData, const TArrayView<const FPCGAttributePropertySelector>& InSelectorArrayView, FPCGContext* InOptionalContext, bool bSilenceMissingAttributeErrors)
 	{
-		if (const UPCGPointData* InPointData = Cast<UPCGPointData>(InData))
+		if (const UPCGBasePointData* InPointData = Cast<UPCGBasePointData>(InData))
 		{
 			return RemoveDuplicatesPoint(InPointData, InSelectorArrayView, InOptionalContext, bSilenceMissingAttributeErrors);
 		}
diff --git a/Engine/Plugins/PCG/Source/PCG/Private/Tests/Elements/PCGAttributeFilterTest.cpp b/Engine/Plugins/PCG/Source/PCG/Private/Tests/Elements/PCGAttributeFilterTest.cpp
--- a/Engine/Plugins/PCG/Source/PCG/Private/Tests/Elements/PCGAttributeFilterTest.cpp
+++ b/Engine/Plugins/PCG/Source/PCG/Private/Tests/Elements/PCGAttributeFilterTest.cpp
@@ -95,21 +95,25 @@ bool FPCGPointFilterDensity::RunTest(const FString& Parameters)
 	UTEST_EQUAL(TEXT("OutFilter pin has 1 output"), OutFilterOutput.Num(), 1);
 
 	// Making sure the output is a PointData
-	const UPCGPointData* InFilterPointData = Cast<UPCGPointData>(InFilterOutput[0].Data);
-	const UPCGPointData* OutFilterPointData = Cast<UPCGPointData>(OutFilterOutput[0].Data);
+	const UPCGBasePointData* InFilterPointData = Cast<UPCGBasePointData>(InFilterOutput[0].Data);
+	const UPCGBasePointData* OutFilterPointData = Cast<UPCGBasePointData>(OutFilterOutput[0].Data);
 
 	UTEST_NOT_NULL(TEXT("InFilter data is a point data"), InFilterPointData);
 	UTEST_NOT_NULL(TEXT("OutFilter data is a point data"), OutFilterPointData);
 
+	const TConstPCGValueRange<float> InFilterDensityRange = InFilterPointData->GetConstDensityValueRange();
+
 	// Verifying that all points have the right density
-	for (const FPCGPoint& Point : InFilterPointData->GetPoints())
+	for (float InFilterDensity : InFilterDensityRange)
 	{
-		UTEST_TRUE(*FString::Printf(TEXT("Point has a density (%f) lower than the threshold (%f)"), Point.Density, DensityThreshold), Point.Density < DensityThreshold);
+		UTEST_TRUE(*FString::Printf(TEXT("Point has a density (%f) lower than the threshold (%f)"), InFilterDensity, DensityThreshold), InFilterDensity < DensityThreshold);
 	}
 
-	for (const FPCGPoint& Point : OutFilterPointData->GetPoints())
+	const TConstPCGValueRange<float> OutFilterDensityRange = OutFilterPointData->GetConstDensityValueRange();
+
+	for (float OutFilterDensity : OutFilterDensityRange)
 	{
-		UTEST_TRUE(*FString::Printf(TEXT("Point has a density (%f) higher or equal than the threshold (%f)"), Point.Density, DensityThreshold), Point.Density >= DensityThreshold);
+		UTEST_TRUE(*FString::Printf(TEXT("Point has a density (%f) higher or equal than the threshold (%f)"), OutFilterDensity, DensityThreshold), OutFilterDensity >= DensityThreshold);
 	}
 
 	return true;
@@ -153,21 +157,25 @@ bool FPCGPointFilterDensityRange::RunTest(const FString& Parameters)
 	UTEST_EQUAL(TEXT("OutFilter pin has 1 output"), OutFilterOutput.Num(), 1);
 
 	// Making sure the output is a PointData
-	const UPCGPointData* InFilterPointData = Cast<UPCGPointData>(InFilterOutput[0].Data);
-	const UPCGPointData* OutFilterPointData = Cast<UPCGPointData>(OutFilterOutput[0].Data);
+	const UPCGBasePointData* InFilterPointData = Cast<UPCGBasePointData>(InFilterOutput[0].Data);
+	const UPCGBasePointData* OutFilterPointData = Cast<UPCGBasePointData>(OutFilterOutput[0].Data);
 
 	UTEST_NOT_NULL(TEXT("InFilter data is a point data"), InFilterPointData);
 	UTEST_NOT_NULL(TEXT("OutFilter data is a point data"), OutFilterPointData);
 
+	const TConstPCGValueRange<float> InFilterDensityRange = InFilterPointData->GetConstDensityValueRange();
+
 	// Verifying that all points have the right density
-	for (const FPCGPoint& Point : InFilterPointData->GetPoints())
+	for (float InFilterDensity : InFilterDensityRange)
 	{
-		UTEST_TRUE(*FString::Printf(TEXT("Point has a density (%f) within the range ([%f, %f])"), Point.Density, DensityMinThreshold, DensityMaxThreshold), (Point.Density >= DensityMinThreshold) && (Point.Density <= DensityMaxThreshold));
+		UTEST_TRUE(*FString::Printf(TEXT("Point has a density (%f) within the range ([%f, %f])"), InFilterDensity, DensityMinThreshold, DensityMaxThreshold), (InFilterDensity >= DensityMinThreshold) && (InFilterDensity <= DensityMaxThreshold));
 	}
 
-	for (const FPCGPoint& Point : OutFilterPointData->GetPoints())
+	const TConstPCGValueRange<float> OutFilterDensityRange = OutFilterPointData->GetConstDensityValueRange();
+
+	for (float OutFilterDensity : OutFilterDensityRange)
 	{
-		UTEST_TRUE(*FString::Printf(TEXT("Point has a density (%f) outside the range ([%f, %f])"), Point.Density, DensityMinThreshold, DensityMaxThreshold), (Point.Density < DensityMinThreshold) || (Point.Density > DensityMaxThreshold));
+		UTEST_TRUE(*FString::Printf(TEXT("Point has a density (%f) outside the range ([%f, %f])"), OutFilterDensity, DensityMinThreshold, DensityMaxThreshold), (OutFilterDensity < DensityMinThreshold) || (OutFilterDensity > DensityMaxThreshold));
 	}
 
 	return true;
@@ -365,15 +373,15 @@ bool FPCGAttributeFilterSkipTestBug::RunTest(const FString& Parameters)
 	UTEST_EQUAL(TEXT("OutFilter pin has 1 output"), OutFilterOutput.Num(), 1);
 
 	// Making sure the output is a PointData
-	const UPCGPointData* InFilterPointData = Cast<UPCGPointData>(InFilterOutput[0].Data);
-	const UPCGPointData* OutFilterPointData = Cast<UPCGPointData>(OutFilterOutput[0].Data);
+	const UPCGBasePointData* InFilterPointData = Cast<UPCGBasePointData>(InFilterOutput[0].Data);
+	const UPCGBasePointData* OutFilterPointData = Cast<UPCGBasePointData>(OutFilterOutput[0].Data);
 
 	UTEST_NOT_NULL(TEXT("InFilter data is a point data"), InFilterPointData);
 	UTEST_NOT_NULL(TEXT("OutFilter data is a point data"), OutFilterPointData);
 
 	// We should have the inside filter with half the points and the outside filter to the rest
-	UTEST_EQUAL(TEXT("InFilter data has the right number of points"), InFilterPointData->GetPoints().Num(), HalfNumPoints);
-	UTEST_EQUAL(TEXT("OutFilter data has the right number of points"), OutFilterPointData->GetPoints().Num(), HalfNumPoints);
+	UTEST_EQUAL(TEXT("InFilter data has the right number of points"), InFilterPointData->GetNumPoints(), HalfNumPoints);
+	UTEST_EQUAL(TEXT("OutFilter data has the right number of points"), OutFilterPointData->GetNumPoints(), HalfNumPoints);
 
 	return true;
 }
\ No newline at end of file
diff --git a/Engine/Plugins/PCG/Source/PCG/Private/Tests/Elements/PCGAttributePartitionTest.cpp b/Engine/Plugins/PCG/Source/PCG/Private/Tests/Elements/PCGAttributePartitionTest.cpp
--- a/Engine/Plugins/PCG/Source/PCG/Private/Tests/Elements/PCGAttributePartitionTest.cpp
+++ b/Engine/Plugins/PCG/Source/PCG/Private/Tests/Elements/PCGAttributePartitionTest.cpp
@@ -71,15 +71,13 @@ namespace PCGAttributePartitionTest
 
 		for (int32 i = 0; i < Context->OutputData.TaggedData.Num(); ++i)
 		{
-			const UPCGPointData* OutputPointData = Cast<const UPCGPointData>(Context->OutputData.TaggedData[i].Data);
+			const UPCGBasePointData* OutputPointData = Cast<const UPCGBasePointData>(Context->OutputData.TaggedData[i].Data);
 			if (!TestClass.TestNotNull(*FString::Printf(TEXT("Output %d is a point data"), i), OutputPointData))
 			{
 				return false;
 			}
 
-			const TArray<FPCGPoint>& OutputPoints = OutputPointData->GetPoints();
-
-			if (!TestClass.TestEqual(*FString::Printf(TEXT("Output %d has 2 points"), i), OutputPoints.Num(), 2))
+			if (!TestClass.TestEqual(*FString::Printf(TEXT("Output %d has 2 points"), i), OutputPointData->GetNumPoints(), 2))
 			{
 				return false;
 			}
@@ -87,16 +85,19 @@ namespace PCGAttributePartitionTest
 			FVector TempValue;
 			bool bFirstPoint = true;
 			bool bAllEquals = true;
-			for (const FPCGPoint& Point : OutputPoints)
+			
+			const TConstPCGValueRange<FTransform> TransformRange = OutputPointData->GetConstTransformValueRange();
+			
+			for (const FTransform& Transform : TransformRange)
 			{
 				if (bFirstPoint)
 				{
-					TempValue = Point.Transform.GetLocation();
+					TempValue = Transform.GetLocation();
 					bFirstPoint = false;
 				}
 				else
 				{
-					bAllEquals &= TempValue.Equals(Point.Transform.GetLocation());
+					bAllEquals &= TempValue.Equals(Transform.GetLocation());
 				}
 			}
 
@@ -154,15 +155,13 @@ namespace PCGAttributePartitionTest
 
 		for (int32 i = 0; i < Context->OutputData.TaggedData.Num(); ++i)
 		{
-			const UPCGPointData* OutputPointData = Cast<const UPCGPointData>(Context->OutputData.TaggedData[i].Data);
+			const UPCGBasePointData* OutputPointData = Cast<const UPCGBasePointData>(Context->OutputData.TaggedData[i].Data);
 			if (!TestClass.TestNotNull(*FString::Printf(TEXT("Output %d is a point data"), i), OutputPointData))
 			{
 				return false;
 			}
 
-			const TArray<FPCGPoint>& OutputPoints = OutputPointData->GetPoints();
-
-			if (!TestClass.TestEqual(*FString::Printf(TEXT("Output %d has the right number of points"), i), OutputPoints.Num(), ExpectedPointsPerData))
+			if (!TestClass.TestEqual(*FString::Printf(TEXT("Output %d has the right number of points"), i), OutputPointData->GetNumPoints(), ExpectedPointsPerData))
 			{
 				return false;
 			}
@@ -179,22 +178,25 @@ namespace PCGAttributePartitionTest
 
 bool FPCGAttributePartition_Points::RunTest(const FString& Parameters)
 {
-	return PCGAttributePartitionTest::PointPartitionTest(*this, /*bWithPartitionIndex=*/ false, /*bNoPartition=*/ false, [this](const UPCGPointData* OutputPointData, int32 OutputIndex) -> bool
+	return PCGAttributePartitionTest::PointPartitionTest(*this, /*bWithPartitionIndex=*/ false, /*bNoPartition=*/ false, [this](const UPCGBasePointData* OutputPointData, int32 OutputIndex) -> bool
 	{
 		check(OutputPointData);
 		float TempValue = -1.0f;
 		bool bFirstPoint = true;
 		bool bAllEquals = true;
-		for (const FPCGPoint& Point : OutputPointData->GetPoints())
+
+		const TConstPCGValueRange<float> DensityRange = OutputPointData->GetConstDensityValueRange();
+
+		for (float Density : DensityRange)
 		{
 			if (bFirstPoint)
 			{
-				TempValue = Point.Density;
+				TempValue = Density;
 				bFirstPoint = false;
 			}
 			else
 			{
-				bAllEquals &= TempValue == Point.Density;
+				bAllEquals &= TempValue == Density;
 			}
 		}
 
@@ -256,23 +258,25 @@ bool FPCGAttributePartition_Order::RunTest(const FString& Parameters)
 
 	for (int32 i = 0; i < Context->OutputData.TaggedData.Num(); ++i)
 	{
-		const UPCGPointData* OutputPointData = Cast<const UPCGPointData>(Context->OutputData.TaggedData[i].Data);
+		const UPCGBasePointData* OutputPointData = Cast<const UPCGBasePointData>(Context->OutputData.TaggedData[i].Data);
 		UTEST_NOT_NULL(*FString::Printf(TEXT("Output %d is a point data"), i), OutputPointData);
 		check(OutputPointData);
 
-		const TArray<FPCGPoint>& OutputPoints = OutputPointData->GetPoints();
 		const FPCGMetadataAttribute<int>* OutputAttribute = OutputPointData->Metadata->GetConstTypedAttribute<int>(AttributeName);
 
 		UTEST_NOT_NULL(*FString::Printf(TEXT("Output %d has the expected attribute"), i), OutputAttribute);
 		check(OutputAttribute);
 
-		UTEST_EQUAL(*FString::Printf(TEXT("Output %d has %d points"), i, NumPoints / 2), OutputPoints.Num(), NumPoints / 2);
+		UTEST_EQUAL(*FString::Printf(TEXT("Output %d has %d points"), i, NumPoints / 2), OutputPointData->GetNumPoints(), NumPoints / 2);
 
-		for (int32 j = 0; j < OutputPoints.Num() - 1; ++j)
+		const TConstPCGValueRange<float> DensityRange = OutputPointData->GetConstDensityValueRange();
+		const TConstPCGValueRange<int64> MetadataEntryRange = OutputPointData->GetConstMetadataEntryValueRange();
+
+		for (int32 j = 0; j < OutputPointData->GetNumPoints() - 1; ++j)
 		{
-			UTEST_TRUE(*FString::Printf(TEXT("Output %d: Point_%d is at the right place"), i, j), OutputPoints[j].Density > OutputPoints[j + 1].Density);
+			UTEST_TRUE(*FString::Printf(TEXT("Output %d: Point_%d is at the right place"), i, j), DensityRange[j] > DensityRange[j + 1]);
 			// First partition should have the value 1, and second partition should have the value 0
-			UTEST_EQUAL(*FString::Printf(TEXT("Output %d: Point_%d has the right attribute value"), i, j), OutputAttribute->GetValue(OutputPoints[j].MetadataEntry), (i + 1) % 2);
+			UTEST_EQUAL(*FString::Printf(TEXT("Output %d: Point_%d has the right attribute value"), i, j), OutputAttribute->GetValue(MetadataEntryRange[j]), (i + 1) % 2);
 		}
 	}
 
@@ -353,7 +357,7 @@ bool FPCGAttributePartition_MultiPartitionOverride::RunTest(const FString& Param
 
 bool FPCGAttributePartition_WithPartitionIndex::RunTest(const FString& Parameters)
 {
-	return PCGAttributePartitionTest::PointPartitionTest(*this, /*bWithPartitionIndex=*/ true, /*bNoPartition=*/ false, [this](const UPCGPointData* OutputPointData, int32 OutputIndex) -> bool
+	return PCGAttributePartitionTest::PointPartitionTest(*this, /*bWithPartitionIndex=*/ true, /*bNoPartition=*/ false, [this](const UPCGBasePointData* OutputPointData, int32 OutputIndex) -> bool
 		{
 			check(OutputPointData && OutputPointData->Metadata);
 			const FPCGMetadataAttribute<int32>* PartitionIndexAttribute = OutputPointData->Metadata->GetConstTypedAttribute<int32>(PCGAttributePartitionTest::PartitionIndexAttributeName);
@@ -363,19 +367,23 @@ bool FPCGAttributePartition_WithPartitionIndex::RunTest(const FString& Parameter
 			float TempValue = -1.0f;
 			bool bFirstPoint = true;
 			bool bAllEquals = true;
-			for (const FPCGPoint& Point : OutputPointData->GetPoints())
+
+			const TConstPCGValueRange<float> DensityRange = OutputPointData->GetConstDensityValueRange();
+			const TConstPCGValueRange<int64> MetadataEntryRange = OutputPointData->GetConstMetadataEntryValueRange();
+
+			for (int32 Index = 0; Index < OutputPointData->GetNumPoints(); ++Index)
 			{
 				if (bFirstPoint)
 				{
-					TempValue = Point.Density;
+					TempValue = DensityRange[Index];
 					bFirstPoint = false;
 				}
 				else
 				{
-					bAllEquals &= TempValue == Point.Density;
+					bAllEquals &= TempValue == DensityRange[Index];
 				}
 
-				UTEST_EQUAL("Partition index value is equal to output index", PartitionIndexAttribute->GetValueFromItemKey(Point.MetadataEntry), OutputIndex);
+				UTEST_EQUAL("Partition index value is equal to output index", PartitionIndexAttribute->GetValueFromItemKey(MetadataEntryRange[Index]), OutputIndex);
 			}
 
 			return TestTrue(*FString::Printf(TEXT("Output points for output %d have all the same density"), OutputIndex), bAllEquals);
@@ -384,17 +392,20 @@ bool FPCGAttributePartition_WithPartitionIndex::RunTest(const FString& Parameter
 
 bool FPCGAttributePartition_NoPartitionWithPartitionIndex::RunTest(const FString& Parameters)
 {
-	return PCGAttributePartitionTest::PointPartitionTest(*this, /*bWithPartitionIndex=*/ true, /*bNoPartition=*/ true, [this](const UPCGPointData* OutputPointData, int32 OutputIndex) -> bool
+	return PCGAttributePartitionTest::PointPartitionTest(*this, /*bWithPartitionIndex=*/ true, /*bNoPartition=*/ true, [this](const UPCGBasePointData* OutputPointData, int32 OutputIndex) -> bool
 		{
 			check(OutputPointData && OutputPointData->Metadata);
 			const FPCGMetadataAttribute<int32>* PartitionIndexAttribute = OutputPointData->Metadata->GetConstTypedAttribute<int32>(PCGAttributePartitionTest::PartitionIndexAttributeName);
 
 			UTEST_NOT_NULL("Partition Index attribute exists", PartitionIndexAttribute);
 
-			for (const FPCGPoint& Point : OutputPointData->GetPoints())
+			const TConstPCGValueRange<float> DensityRange = OutputPointData->GetConstDensityValueRange();
+			const TConstPCGValueRange<int64> MetadataEntryRange = OutputPointData->GetConstMetadataEntryValueRange();
+
+			for (int32 Index = 0; Index < OutputPointData->GetNumPoints(); ++Index)
 			{
-				const int32 PartitionIndex = FMath::FloorToInt32(Point.Density);
-				UTEST_EQUAL("Partition index has the right value", PartitionIndexAttribute->GetValueFromItemKey(Point.MetadataEntry), PartitionIndex);
+				const int32 PartitionIndex = FMath::FloorToInt32(DensityRange[Index]);
+				UTEST_EQUAL("Partition index has the right value", PartitionIndexAttribute->GetValueFromItemKey(MetadataEntryRange[Index]), PartitionIndex);
 			}
 
 			return true;
diff --git a/Engine/Plugins/PCG/Source/PCG/Private/Tests/Elements/PCGCopyPointsTest.cpp b/Engine/Plugins/PCG/Source/PCG/Private/Tests/Elements/PCGCopyPointsTest.cpp
--- a/Engine/Plugins/PCG/Source/PCG/Private/Tests/Elements/PCGCopyPointsTest.cpp
+++ b/Engine/Plugins/PCG/Source/PCG/Private/Tests/Elements/PCGCopyPointsTest.cpp
@@ -102,14 +102,12 @@ bool FPCGCopyPointsTestBase::RunTestInternal(FParameters& Parameters)
 	UTEST_TRUE("Valid number of inputs", Sources.Num() == 1 && Targets.Num() == 1);
 	UTEST_EQUAL("Valid number of outputs", Sources.Num(), 1);
 
-	const UPCGPointData* OutPointData = Cast<const UPCGPointData>(Outputs[0].Data);
+	const UPCGBasePointData* OutPointData = Cast<const UPCGBasePointData>(Outputs[0].Data);
 	UTEST_NOT_NULL("Valid output point data", OutPointData);
 	check(OutPointData);
 
-	const TArray<FPCGPoint>& OutPoints = OutPointData->GetPoints();
-
-	UTEST_EQUAL("Valid number of output points", SourcePoints.Num() * TargetPoints.Num(), OutPoints.Num());
-	check(Parameters.ExpectedPoints.IsEmpty() || Parameters.ExpectedPoints.Num() == OutPoints.Num());
+	UTEST_EQUAL("Valid number of output points", SourcePoints.Num() * TargetPoints.Num(), OutPointData->GetNumPoints());
+	check(Parameters.ExpectedPoints.IsEmpty() || Parameters.ExpectedPoints.Num() == OutPointData->GetNumPoints());
 
 	const UPCGMetadata* RootMetadata = nullptr;
 	const UPCGMetadata* NonRootMetadata = nullptr;
@@ -186,26 +184,27 @@ bool FPCGCopyPointsTestBase::RunTestInternal(FParameters& Parameters)
 
 	check(InheritedAttributeModes.Num() == InheritedAttributeModes.Num() && InheritedAttributeModes.Num() == OutAttributes.Num())
 
-	for (int PointIndex = 0; PointIndex < OutPoints.Num(); ++PointIndex)
+	const FConstPCGPointValueRanges OutRanges(OutPointData);
+
+	for (int PointIndex = 0; PointIndex < OutPointData->GetNumPoints(); ++PointIndex)
 	{
 		const FPCGPoint& SourcePoint = SourcePoints[PointIndex / TargetPoints.Num()];
 		const FPCGPoint& TargetPoint = TargetPoints[PointIndex % TargetPoints.Num()];
-		const FPCGPoint& OutPoint = OutPoints[PointIndex];
 
 		auto FormatMessage = [PointIndex](const FString& Message)
 		{
 			return FString::Printf(TEXT("Index %d: %s"), PointIndex, *Message);
 		};
 
-		UTEST_EQUAL(*FormatMessage(TEXT("SourcePoint and OutPoint index %d have same Density")), SourcePoint.Density, OutPoint.Density);
-		UTEST_EQUAL(*FormatMessage(TEXT("SourcePoint and OutPoint index %d have same BoundsMin")), SourcePoint.BoundsMin, OutPoint.BoundsMin);
-		UTEST_EQUAL(*FormatMessage(TEXT("SourcePoint and OutPoint index %d have same BoundsMax")), SourcePoint.BoundsMax, OutPoint.BoundsMax);
-		UTEST_EQUAL(*FormatMessage(TEXT("SourcePoint and OutPoint index %d have same Steepness")), SourcePoint.Steepness, OutPoint.Steepness);
+		UTEST_EQUAL(*FormatMessage(TEXT("SourcePoint and OutPoint index %d have same Density")), SourcePoint.Density, OutRanges.DensityRange[PointIndex]);
+		UTEST_EQUAL(*FormatMessage(TEXT("SourcePoint and OutPoint index %d have same BoundsMin")), SourcePoint.BoundsMin, OutRanges.BoundsMinRange[PointIndex]);
+		UTEST_EQUAL(*FormatMessage(TEXT("SourcePoint and OutPoint index %d have same BoundsMax")), SourcePoint.BoundsMax, OutRanges.BoundsMaxRange[PointIndex]);
+		UTEST_EQUAL(*FormatMessage(TEXT("SourcePoint and OutPoint index %d have same Steepness")), SourcePoint.Steepness, OutRanges.SteepnessRange[PointIndex]);
 
 		// Validate transform inheritance
 		const FTransform& SourceTransform = SourcePoint.Transform;
 		const FTransform& TargetTransform = TargetPoint.Transform;
-		const FTransform& OutTransform = OutPoint.Transform;
+		const FTransform& OutTransform = OutRanges.TransformRange[PointIndex];
 
 		const FVector Location = TargetPoint.Transform.TransformPosition(SourcePoint.Transform.GetLocation());
 		FQuat Rotation;
@@ -267,9 +266,9 @@ bool FPCGCopyPointsTestBase::RunTestInternal(FParameters& Parameters)
 
 		UTEST_TRUE(*FormatMessage(TEXT("Valid rotation")), Rotation.Equals(OutTransform.GetRotation()));
 		UTEST_EQUAL(*FormatMessage(TEXT("Valid scale")), Scale, OutTransform.GetScale3D());
-		UTEST_TRUE(*FormatMessage(TEXT("Valid color")), Color.Equals(OutPoint.Color));
+		UTEST_TRUE(*FormatMessage(TEXT("Valid color")), Color.Equals(OutRanges.ColorRange[PointIndex]));
 		UTEST_EQUAL(*FormatMessage(TEXT("Valid location")), Location, OutTransform.GetLocation());
-		UTEST_EQUAL(*FormatMessage(TEXT("Valid seed")), Seed, OutPoint.Seed);
+		UTEST_EQUAL(*FormatMessage(TEXT("Valid seed")), Seed, OutRanges.SeedRange[PointIndex]);
 
 		// Validate point value keys
 		for (int AttributeIndex = 0; AttributeIndex < InheritedAttributes.Num(); ++AttributeIndex)
@@ -279,7 +278,7 @@ bool FPCGCopyPointsTestBase::RunTestInternal(FParameters& Parameters)
 			check(InheritedAttribute && OutAttribute);
 
 			const PCGMetadataEntryKey EntryKey = (InheritedAttributeModes[AttributeIndex] == EPCGCopyPointsMetadataInheritanceMode::SourceFirst) ? SourcePoint.MetadataEntry : TargetPoint.MetadataEntry;
-			UTEST_EQUAL(*FormatMessage(TEXT("Valid metadata value")), InheritedAttribute->GetValueKey(EntryKey), OutAttribute->GetValueKey(OutPoint.MetadataEntry));
+			UTEST_EQUAL(*FormatMessage(TEXT("Valid metadata value")), InheritedAttribute->GetValueKey(EntryKey), OutAttribute->GetValueKey(OutRanges.MetadataEntryRange[PointIndex]));
 		}
 	}
 
diff --git a/Engine/Plugins/PCG/Source/PCG/Private/Tests/Elements/PCGDistanceTest.cpp b/Engine/Plugins/PCG/Source/PCG/Private/Tests/Elements/PCGDistanceTest.cpp
--- a/Engine/Plugins/PCG/Source/PCG/Private/Tests/Elements/PCGDistanceTest.cpp
+++ b/Engine/Plugins/PCG/Source/PCG/Private/Tests/Elements/PCGDistanceTest.cpp
@@ -48,19 +48,17 @@ bool FPCGDistanceTest_PointToPoint::RunTest(const FString& Parameters)
 
 	UTEST_EQUAL("Output count", Outputs.Num(), 1);
 
-	const UPCGPointData *OutPointData = Cast<UPCGPointData>(Outputs[0].Data);
+	const UPCGBasePointData *OutPointData = Cast<UPCGBasePointData>(Outputs[0].Data);
 
 	UTEST_NOT_NULL("Output point data", OutPointData);
 
-	const TArray<FPCGPoint>& OutPoints = OutPointData->GetPoints();
-
-	UTEST_EQUAL("Output point count", OutPoints.Num(), 2);
+	UTEST_EQUAL("Output point count", OutPointData->GetNumPoints(), 2);
 
 	const FPCGMetadataAttribute<double> *DistanceAttribute = OutPointData->Metadata->GetConstTypedAttribute<double>(Settings->OutputAttribute.GetAttributeName());
 	UTEST_NOT_NULL("Distance attribute", DistanceAttribute);
 
-	UTEST_EQUAL_TOLERANCE("Point 0 distance", DistanceAttribute->GetValue(OutPoints[0].MetadataEntry), 100.0, 0.01);
-	UTEST_EQUAL_TOLERANCE("Point 1 distance", DistanceAttribute->GetValue(OutPoints[1].MetadataEntry), 50.0, 0.01);
+	UTEST_EQUAL_TOLERANCE("Point 0 distance", DistanceAttribute->GetValue(OutPointData->GetMetadataEntry(0)), 100.0, 0.01);
+	UTEST_EQUAL_TOLERANCE("Point 1 distance", DistanceAttribute->GetValue(OutPointData->GetMetadataEntry(1)), 50.0, 0.01);
 
 	return true;
 }
@@ -108,16 +106,14 @@ bool FPCGDistanceTest_SetDensity::RunTest(const FString& Parameters)
 
 	UTEST_EQUAL("Output count", Outputs.Num(), 1);
 
-	const UPCGPointData *OutPointData = Cast<UPCGPointData>(Outputs[0].Data);
+	const UPCGBasePointData *OutPointData = Cast<UPCGBasePointData>(Outputs[0].Data);
 
 	UTEST_NOT_NULL("Output point data", OutPointData);
 
-	const TArray<FPCGPoint>& OutPoints = OutPointData->GetPoints();
-
-	UTEST_EQUAL("Output point count", OutPoints.Num(), 2);
+	UTEST_EQUAL("Output point count", OutPointData->GetNumPoints(), 2);
 
-	UTEST_EQUAL_TOLERANCE("Point 0 distance", OutPoints[0].Density, 0.5f, 0.01f);
-	UTEST_EQUAL_TOLERANCE("Point 1 distance", OutPoints[1].Density, 0.25f, 0.01f);
+	UTEST_EQUAL_TOLERANCE("Point 0 distance", OutPointData->GetDensity(0), 0.5f, 0.01f);
+	UTEST_EQUAL_TOLERANCE("Point 1 distance", OutPointData->GetDensity(1), 0.25f, 0.01f);
 
 	return true;
 }
@@ -170,19 +166,17 @@ bool FPCGDistanceTest_PointToSphere::RunTest(const FString& Parameters)
 
 	UTEST_EQUAL("Output count", Outputs.Num(), 1);
 
-	const UPCGPointData *OutPointData = Cast<UPCGPointData>(Outputs[0].Data);
+	const UPCGBasePointData *OutPointData = Cast<UPCGBasePointData>(Outputs[0].Data);
 
 	UTEST_NOT_NULL("Output point data", OutPointData);
 
-	const TArray<FPCGPoint>& OutPoints = OutPointData->GetPoints();
-
-	UTEST_EQUAL("Output point count", OutPoints.Num(), 2);
+	UTEST_EQUAL("Output point count", OutPointData->GetNumPoints(), 2);
 
 	const FPCGMetadataAttribute<double> *DistanceAttribute = OutPointData->Metadata->GetConstTypedAttribute<double>(Settings->OutputAttribute.GetAttributeName());
 	UTEST_NOT_NULL("Distance attribute", DistanceAttribute);
 
-	UTEST_EQUAL_TOLERANCE("Point 0 distance", DistanceAttribute->GetValue(OutPoints[0].MetadataEntry), 100.0 - TargetPointRadius, 0.01);
-	UTEST_EQUAL_TOLERANCE("Point 1 distance", DistanceAttribute->GetValue(OutPoints[1].MetadataEntry), 50.0 - TargetPointRadius, 0.01);
+	UTEST_EQUAL_TOLERANCE("Point 0 distance", DistanceAttribute->GetValue(OutPointData->GetMetadataEntry(0)), 100.0 - TargetPointRadius, 0.01);
+	UTEST_EQUAL_TOLERANCE("Point 1 distance", DistanceAttribute->GetValue(OutPointData->GetMetadataEntry(1)), 50.0 - TargetPointRadius, 0.01);
 
 	return true;
 }
@@ -235,19 +229,17 @@ bool FPCGDistanceTest_PointToBox::RunTest(const FString& Parameters)
 
 	UTEST_EQUAL("Output count", Outputs.Num(), 1);
 
-	const UPCGPointData *OutPointData = Cast<UPCGPointData>(Outputs[0].Data);
+	const UPCGBasePointData *OutPointData = Cast<UPCGBasePointData>(Outputs[0].Data);
 
 	UTEST_NOT_NULL("Output point data", OutPointData);
 
-	const TArray<FPCGPoint>& OutPoints = OutPointData->GetPoints();
-
-	UTEST_EQUAL("Output point count", OutPoints.Num(), 2);
+	UTEST_EQUAL("Output point count", OutPointData->GetNumPoints(), 2);
 
 	const FPCGMetadataAttribute<double> *DistanceAttribute = OutPointData->Metadata->GetConstTypedAttribute<double>(Settings->OutputAttribute.GetAttributeName());
 	UTEST_NOT_NULL("Distance attribute", DistanceAttribute);
 
-	UTEST_EQUAL_TOLERANCE("Point 0 distance", DistanceAttribute->GetValue(OutPoints[0].MetadataEntry), 10.0, 0.01);
-	UTEST_EQUAL_TOLERANCE("Point 1 distance", DistanceAttribute->GetValue(OutPoints[1].MetadataEntry), 10.0, 0.01);
+	UTEST_EQUAL_TOLERANCE("Point 0 distance", DistanceAttribute->GetValue(OutPointData->GetMetadataEntry(0)), 10.0, 0.01);
+	UTEST_EQUAL_TOLERANCE("Point 1 distance", DistanceAttribute->GetValue(OutPointData->GetMetadataEntry(1)), 10.0, 0.01);
 
 	return true;
 }
diff --git a/Engine/Plugins/PCG/Source/PCG/Private/Tests/Elements/PCGSortAttributesTest.cpp b/Engine/Plugins/PCG/Source/PCG/Private/Tests/Elements/PCGSortAttributesTest.cpp
--- a/Engine/Plugins/PCG/Source/PCG/Private/Tests/Elements/PCGSortAttributesTest.cpp
+++ b/Engine/Plugins/PCG/Source/PCG/Private/Tests/Elements/PCGSortAttributesTest.cpp
@@ -5,7 +5,7 @@
 #include "PCGComponent.h"
 #include "PCGContext.h"
 #include "PCGParamData.h"
-#include "Data/PCGPointData.h"
+#include "Data/PCGBasePointData.h"
 #include "Elements/PCGSortAttributes.h"
 #include "Metadata/PCGMetadata.h"
 #include "Metadata/PCGMetadataAttributeTpl.h"
@@ -65,18 +65,17 @@ bool FPCGSortPointsTest_Ascending::RunTest(const FString& Parameters)
 
 	// compare point data
 	const FPCGTaggedData& Output = Outputs[0];
-	const UPCGPointData* OutPointData = Cast<UPCGPointData>(Output.Data);
+	const UPCGBasePointData* OutPointData = Cast<UPCGBasePointData>(Output.Data);
 	UTEST_NOT_NULL("Output is a point data", OutPointData);
 	check(OutPointData);
+	
+	const TConstPCGValueRange<float> DensityRange = OutPointData->GetConstDensityValueRange();
 
-	const TArray<FPCGPoint>& OutPoints = OutPointData->GetPoints();
-
-	for (int i = 0; i < OutPoints.Num() - 1; ++i)
+	for (int32 i = 0; i < DensityRange.Num() - 1; ++i)
 	{
-		UTEST_TRUE(FString::Format(TEXT("{0} is less than/equal to {1}"), { OutPoints[i].Density, OutPoints[i + 1].Density }), OutPoints[i].Density <= OutPoints[i + 1].Density);
+		UTEST_TRUE(FString::Format(TEXT("{0} is less than/equal to {1}"), { DensityRange[i], DensityRange[i + 1] }), DensityRange[i] <= DensityRange[i + 1]);
 	}
 
-	
 	return true;
 }
 
@@ -125,16 +124,16 @@ bool FPCGSortPointsTest_Descending::RunTest(const FString& Parameters)
 
 	// compare point data
 	const FPCGTaggedData& Output = Outputs[0];
-	const UPCGPointData* OutPointData = Cast<UPCGPointData>(Output.Data);
+	const UPCGBasePointData* OutPointData = Cast<UPCGBasePointData>(Output.Data);
 
 	UTEST_NOT_NULL("Output is a point data", OutPointData);
 	check(OutPointData);
+	
+	const TConstPCGValueRange<float> DensityRange = OutPointData->GetConstDensityValueRange();
 
-	const TArray<FPCGPoint>& OutPoints = OutPointData->GetPoints();
-
-	for (int i = 0; i < OutPoints.Num() - 1; ++i)
+	for (int i = 0; i < DensityRange.Num() - 1; ++i)
 	{
-		UTEST_TRUE(FString::Format(TEXT("{0} is greater than/equal to {1}"), { OutPoints[i].Density, OutPoints[i + 1].Density }), OutPoints[i].Density >= OutPoints[i + 1].Density);
+		UTEST_TRUE(FString::Format(TEXT("{0} is greater than/equal to {1}"), { DensityRange[i], DensityRange[i + 1] }), DensityRange[i] >= DensityRange[i + 1]);
 	}
 	
 	return true;
@@ -188,23 +187,28 @@ bool FPCGSortPointsTest_SameValues::RunTest(const FString& Parameters)
 	const FPCGTaggedData& Input = Inputs[0];
 	const FPCGTaggedData& Output = Outputs[0];
 
-	const UPCGPointData* InPointData = Cast<UPCGPointData>(Input.Data);
-	const UPCGPointData* OutPointData = Cast<UPCGPointData>(Output.Data);
+	const UPCGBasePointData* InPointData = Cast<UPCGBasePointData>(Input.Data);
+	const UPCGBasePointData* OutPointData = Cast<UPCGBasePointData>(Output.Data);
 
 	UTEST_NOT_NULL("Output is a point data", OutPointData);
 	check(InPointData && OutPointData);
+		
+	UTEST_EQUAL("Arrays have the same number of points:", InPointData->GetNumPoints(), OutPointData->GetNumPoints());
 
-	const TArray<FPCGPoint>& InPoints = InPointData->GetPoints();
-	const TArray<FPCGPoint>& OutPoints = OutPointData->GetPoints();
+	const TConstPCGValueRange<int32> InSeedRange = InPointData->GetConstSeedValueRange();
+	const TConstPCGValueRange<float> InDensityRange = InPointData->GetConstDensityValueRange();
+	const TConstPCGValueRange<FTransform> InTransformRange = InPointData->GetConstTransformValueRange();
 
-	UTEST_EQUAL("Arrays have the same number of points:", InPoints.Num(), OutPoints.Num());
+	const TConstPCGValueRange<int32> OutSeedRange = OutPointData->GetConstSeedValueRange();
+	const TConstPCGValueRange<float> OutDensityRange = OutPointData->GetConstDensityValueRange();
+	const TConstPCGValueRange<FTransform> OutTransformRange = OutPointData->GetConstTransformValueRange();
 
-	for (int i = 0; i < InPoints.Num(); ++i)
+	for (int i = 0; i < InPointData->GetNumPoints(); ++i)
 	{
 		//if they're in the same spots after sorting, they should have exactly the same properties across the board
-		UTEST_EQUAL(FString::Format(TEXT("UnsortedArray[{0}].Seed is equal to SortedArray[{0}].Seed"), {i}), OutPoints[i].Seed, InPoints[i].Seed);
-		UTEST_EQUAL(FString::Format(TEXT("UnsortedArray[{0}].Density is equal to SortedArray[{0}].Density"), {i}), OutPoints[i].Density, InPoints[i].Density);
-		UTEST_EQUAL(FString::Format(TEXT("UnsortedArray[{0}].Transform is equal to SortedArray[{0}].Transform"), {i}), OutPoints[i].Transform, InPoints[i].Transform);
+		UTEST_EQUAL(FString::Format(TEXT("UnsortedArray[{0}].Seed is equal to SortedArray[{0}].Seed"), {i}), OutSeedRange[i], InSeedRange[i]);
+		UTEST_EQUAL(FString::Format(TEXT("UnsortedArray[{0}].Density is equal to SortedArray[{0}].Density"), {i}), OutDensityRange[i], InDensityRange[i]);
+		UTEST_EQUAL(FString::Format(TEXT("UnsortedArray[{0}].Transform is equal to SortedArray[{0}].Transform"), {i}), OutTransformRange[i], InTransformRange[i]);
 	}
 
 	return true;
diff --git a/Engine/Plugins/PCG/Source/PCG/Private/Tests/Elements/PCGSpatialNoiseTest.cpp b/Engine/Plugins/PCG/Source/PCG/Private/Tests/Elements/PCGSpatialNoiseTest.cpp
--- a/Engine/Plugins/PCG/Source/PCG/Private/Tests/Elements/PCGSpatialNoiseTest.cpp
+++ b/Engine/Plugins/PCG/Source/PCG/Private/Tests/Elements/PCGSpatialNoiseTest.cpp
@@ -66,10 +66,10 @@ bool FPCGSpatialNoise_Perlin2D::RunTest(const FString& Parameters)
 
 	UTEST_EQUAL("Output count", Outputs.Num(), 1);
 
-	const UPCGPointData* OutputData = Cast<UPCGPointData>(Outputs[0].Data);
+	const UPCGBasePointData* OutputData = Cast<UPCGBasePointData>(Outputs[0].Data);
 
 	UTEST_NOT_NULL("Output data", OutputData);
-	UTEST_EQUAL("Output point data count", OutputData->GetPoints().Num(), 100);
+	UTEST_EQUAL("Output point data count", OutputData->GetNumPoints(), 100);
 
 	const FPCGMetadataAttribute<double>* NoiseAttribute = OutputData->Metadata->GetConstTypedAttribute<double>(Settings->ValueTarget.GetAttributeName());
 	UTEST_NOT_NULL("Noise Attribute Created", NoiseAttribute);
@@ -104,10 +104,10 @@ bool FPCGSpatialNoise_Caustic2D::RunTest(const FString& Parameters)
 
 	UTEST_EQUAL("Output count", Outputs.Num(), 1);
 
-	const UPCGPointData* OutputData = Cast<UPCGPointData>(Outputs[0].Data);
+	const UPCGBasePointData* OutputData = Cast<UPCGBasePointData>(Outputs[0].Data);
 
 	UTEST_NOT_NULL("Output data", OutputData);
-	UTEST_EQUAL("Output point data count", OutputData->GetPoints().Num(), 100);
+	UTEST_EQUAL("Output point data count", OutputData->GetNumPoints(), 100);
 
 	const FPCGMetadataAttribute<double>* NoiseAttribute = OutputData->Metadata->GetConstTypedAttribute<double>(Settings->ValueTarget.GetAttributeName());
 	UTEST_NOT_NULL("Noise Attribute Created", NoiseAttribute);
@@ -143,10 +143,10 @@ bool FPCGSpatialNoise_Voronoi2D::RunTest(const FString& Parameters)
 
 	UTEST_EQUAL("Output count", Outputs.Num(), 1);
 
-	const UPCGPointData* OutputData = Cast<UPCGPointData>(Outputs[0].Data);
+	const UPCGBasePointData* OutputData = Cast<UPCGBasePointData>(Outputs[0].Data);
 
 	UTEST_NOT_NULL("Output data", OutputData);
-	UTEST_EQUAL("Output point data count", OutputData->GetPoints().Num(), 100);
+	UTEST_EQUAL("Output point data count", OutputData->GetNumPoints(), 100);
 
 	const FPCGMetadataAttribute<double>* NoiseAttribute = OutputData->Metadata->GetConstTypedAttribute<double>(Settings->ValueTarget.GetAttributeName());
 	UTEST_NOT_NULL("Distance Attribute Created", NoiseAttribute);
@@ -185,10 +185,10 @@ bool FPCGSpatialNoise_FractionalBrownian2D::RunTest(const FString& Parameters)
 
 	UTEST_EQUAL("Output count", Outputs.Num(), 1);
 
-	const UPCGPointData* OutputData = Cast<UPCGPointData>(Outputs[0].Data);
+	const UPCGBasePointData* OutputData = Cast<UPCGBasePointData>(Outputs[0].Data);
 
 	UTEST_NOT_NULL("Output data", OutputData);
-	UTEST_EQUAL("Output point data count", OutputData->GetPoints().Num(), 100);
+	UTEST_EQUAL("Output point data count", OutputData->GetNumPoints(), 100);
 
 	const FPCGMetadataAttribute<double>* NoiseAttribute = OutputData->Metadata->GetConstTypedAttribute<double>(Settings->ValueTarget.GetAttributeName());
 	UTEST_NOT_NULL("Noise Attribute Created", NoiseAttribute);
@@ -223,10 +223,10 @@ bool FPCGSpatialNoise_EdgeMask2D::RunTest(const FString& Parameters)
 
 	UTEST_EQUAL("Output count", Outputs.Num(), 1);
 
-	const UPCGPointData* OutputData = Cast<UPCGPointData>(Outputs[0].Data);
+	const UPCGBasePointData* OutputData = Cast<UPCGBasePointData>(Outputs[0].Data);
 
 	UTEST_NOT_NULL("Output data", OutputData);
-	UTEST_EQUAL("Output point data count", OutputData->GetPoints().Num(), 100);
+	UTEST_EQUAL("Output point data count", OutputData->GetNumPoints(), 100);
 
 	const FPCGMetadataAttribute<double>* NoiseAttribute = OutputData->Metadata->GetConstTypedAttribute<double>(Settings->ValueTarget.GetAttributeName());
 	UTEST_NOT_NULL("Noise Attribute Created", NoiseAttribute);
@@ -261,10 +261,10 @@ bool FPCGSpatialNoise_TilingPerlin2D::RunTest(const FString& Parameters)
 
 	UTEST_EQUAL("Output count", Outputs.Num(), 1);
 
-	const UPCGPointData* OutputData = Cast<UPCGPointData>(Outputs[0].Data);
+	const UPCGBasePointData* OutputData = Cast<UPCGBasePointData>(Outputs[0].Data);
 
 	UTEST_NOT_NULL("Output data", OutputData);
-	UTEST_EQUAL("Output point data count", OutputData->GetPoints().Num(), 100);
+	UTEST_EQUAL("Output point data count", OutputData->GetNumPoints(), 100);
 
 	const FPCGMetadataAttribute<double>* NoiseAttribute = OutputData->Metadata->GetConstTypedAttribute<double>(Settings->ValueTarget.GetAttributeName());
 	UTEST_NOT_NULL("Noise Attribute Created", NoiseAttribute);
@@ -300,10 +300,10 @@ bool FPCGSpatialNoise_TilingVoronoi2D::RunTest(const FString& Parameters)
 
 	UTEST_EQUAL("Output count", Outputs.Num(), 1);
 
-	const UPCGPointData* OutputData = Cast<UPCGPointData>(Outputs[0].Data);
+	const UPCGBasePointData* OutputData = Cast<UPCGBasePointData>(Outputs[0].Data);
 
 	UTEST_NOT_NULL("Output data", OutputData);
-	UTEST_EQUAL("Output point data count", OutputData->GetPoints().Num(), 100);
+	UTEST_EQUAL("Output point data count", OutputData->GetNumPoints(), 100);
 
 	const FPCGMetadataAttribute<double>* NoiseAttribute = OutputData->Metadata->GetConstTypedAttribute<double>(Settings->ValueTarget.GetAttributeName());
 	UTEST_NOT_NULL("Distance Attribute Created", NoiseAttribute);
diff --git a/Engine/Plugins/PCG/Source/PCG/Public/Data/PCGBasePointData.h b/Engine/Plugins/PCG/Source/PCG/Public/Data/PCGBasePointData.h
--- a/Engine/Plugins/PCG/Source/PCG/Public/Data/PCGBasePointData.h
+++ b/Engine/Plugins/PCG/Source/PCG/Public/Data/PCGBasePointData.h
@@ -296,4 +296,92 @@ inline TConstPCGValueRange<T> UPCGBasePointData::GetConstValueRange(EPCGPointNat
 	}
 	
 	return TConstPCGValueRange<T>();
-}
\ No newline at end of file
+}
+
+struct FConstPCGPointValueRanges
+{
+	FConstPCGPointValueRanges() = default;
+
+	FConstPCGPointValueRanges(const UPCGBasePointData* InBasePointData)
+	{
+		TransformRange = InBasePointData->GetConstTransformValueRange();
+		DensityRange = InBasePointData->GetConstDensityValueRange();
+		SteepnessRange = InBasePointData->GetConstSteepnessValueRange();
+		BoundsMinRange = InBasePointData->GetConstBoundsMinValueRange();
+		BoundsMaxRange = InBasePointData->GetConstBoundsMaxValueRange();
+		ColorRange = InBasePointData->GetConstColorValueRange();
+		SeedRange = InBasePointData->GetConstSeedValueRange();
+		MetadataEntryRange = InBasePointData->GetConstMetadataEntryValueRange();
+	}
+
+	FPCGPoint GetPoint(int32 Index) const
+	{
+		FPCGPoint Point(TransformRange[Index], DensityRange[Index], SeedRange[Index]);
+		Point.BoundsMin = BoundsMinRange[Index];
+		Point.BoundsMax = BoundsMaxRange[Index];
+		Point.Steepness = SteepnessRange[Index];
+		Point.Color = ColorRange[Index];
+		Point.MetadataEntry = MetadataEntryRange[Index];
+
+		return Point;
+	}
+
+	TConstPCGValueRange<FTransform> TransformRange;
+	TConstPCGValueRange<float> DensityRange;
+	TConstPCGValueRange<float> SteepnessRange;
+	TConstPCGValueRange<FVector> BoundsMinRange;
+	TConstPCGValueRange<FVector> BoundsMaxRange;
+	TConstPCGValueRange<FVector4> ColorRange;
+	TConstPCGValueRange<int32> SeedRange;
+	TConstPCGValueRange<int64> MetadataEntryRange;
+};
+
+struct FPCGPointValueRanges
+{
+	FPCGPointValueRanges() = default;
+
+	FPCGPointValueRanges(UPCGBasePointData* InBasePointData, bool bAllocate = true)
+	{
+		TransformRange = InBasePointData->GetTransformValueRange(bAllocate);
+		DensityRange = InBasePointData->GetDensityValueRange(bAllocate);
+		SteepnessRange = InBasePointData->GetSteepnessValueRange(bAllocate);
+		BoundsMinRange = InBasePointData->GetBoundsMinValueRange(bAllocate);
+		BoundsMaxRange = InBasePointData->GetBoundsMaxValueRange(bAllocate);
+		ColorRange = InBasePointData->GetColorValueRange(bAllocate);
+		SeedRange = InBasePointData->GetSeedValueRange(bAllocate);
+		MetadataEntryRange = InBasePointData->GetMetadataEntryValueRange(bAllocate);
+	}
+
+	TPCGValueRange<FTransform> TransformRange;
+	TPCGValueRange<float> DensityRange;
+	TPCGValueRange<float> SteepnessRange;
+	TPCGValueRange<FVector> BoundsMinRange;
+	TPCGValueRange<FVector> BoundsMaxRange;
+	TPCGValueRange<FVector4> ColorRange;
+	TPCGValueRange<int32> SeedRange;
+	TPCGValueRange<int64> MetadataEntryRange;
+
+	void SetFromPoint(int32 Index, const FPCGPoint& Point)
+	{
+		TransformRange[Index] = Point.Transform;
+		DensityRange[Index] = Point.Density;
+		SteepnessRange[Index] = Point.Steepness;
+		BoundsMinRange[Index] = Point.BoundsMin;
+		BoundsMaxRange[Index] = Point.BoundsMax;
+		ColorRange[Index] = Point.Color;
+		SeedRange[Index] = Point.Seed;
+		MetadataEntryRange[Index] = Point.MetadataEntry;
+	}
+		
+	void SetFromValueRanges(int32 WriteIndex, const FConstPCGPointValueRanges& ReadRange, int32 ReadIndex)
+	{
+		TransformRange[WriteIndex] = ReadRange.TransformRange[ReadIndex];
+		DensityRange[WriteIndex] = ReadRange.DensityRange[ReadIndex];
+		SteepnessRange[WriteIndex] = ReadRange.SteepnessRange[ReadIndex];
+		BoundsMinRange[WriteIndex] = ReadRange.BoundsMinRange[ReadIndex];
+		BoundsMaxRange[WriteIndex] = ReadRange.BoundsMaxRange[ReadIndex];
+		ColorRange[WriteIndex] = ReadRange.ColorRange[ReadIndex];
+		SeedRange[WriteIndex] = ReadRange.SeedRange[ReadIndex];
+		MetadataEntryRange[WriteIndex] = ReadRange.MetadataEntryRange[ReadIndex];
+	}
+};
diff --git a/Engine/Plugins/PCG/Source/PCG/Public/Elements/Metadata/PCGMetadataOpElementBase.h b/Engine/Plugins/PCG/Source/PCG/Public/Elements/Metadata/PCGMetadataOpElementBase.h
--- a/Engine/Plugins/PCG/Source/PCG/Public/Elements/Metadata/PCGMetadataOpElementBase.h
+++ b/Engine/Plugins/PCG/Source/PCG/Public/Elements/Metadata/PCGMetadataOpElementBase.h
@@ -209,6 +209,7 @@ protected:
 	virtual bool PrepareDataInternal(FPCGContext* Context) const override;
 	virtual bool ExecuteInternal(FPCGContext* Context) const override;
 	virtual EPCGElementExecutionLoopMode ExecutionLoopMode(const UPCGSettings* Settings) const override { return EPCGElementExecutionLoopMode::SinglePrimaryPin; }
+	virtual bool SupportsBasePointDataInputs(FPCGContext* InContext) const override { return true; }
 
 	virtual bool DoOperation(PCGMetadataOps::FOperationData& InOperationData) const = 0;
 
diff --git a/Engine/Plugins/PCG/Source/PCG/Public/Elements/Metadata/PCGMetadataPartition.h b/Engine/Plugins/PCG/Source/PCG/Public/Elements/Metadata/PCGMetadataPartition.h
--- a/Engine/Plugins/PCG/Source/PCG/Public/Elements/Metadata/PCGMetadataPartition.h
+++ b/Engine/Plugins/PCG/Source/PCG/Public/Elements/Metadata/PCGMetadataPartition.h
@@ -77,4 +77,5 @@ class FPCGMetadataPartitionElement : public IPCGElement
 protected:
 	virtual bool ExecuteInternal(FPCGContext* Context) const override;
 	virtual EPCGElementExecutionLoopMode ExecutionLoopMode(const UPCGSettings* Settings) const override { return EPCGElementExecutionLoopMode::SinglePrimaryPin; }
+	virtual bool SupportsBasePointDataInputs(FPCGContext* InContext) const override { return true; }
 };
diff --git a/Engine/Plugins/PCG/Source/PCG/Public/Elements/PCGAttributeFilter.h b/Engine/Plugins/PCG/Source/PCG/Public/Elements/PCGAttributeFilter.h
--- a/Engine/Plugins/PCG/Source/PCG/Public/Elements/PCGAttributeFilter.h
+++ b/Engine/Plugins/PCG/Source/PCG/Public/Elements/PCGAttributeFilter.h
@@ -200,6 +200,8 @@ class FPCGAttributeFilterElementBase : public IPCGElement
 {
 protected:
 	virtual EPCGElementExecutionLoopMode ExecutionLoopMode(const UPCGSettings* Settings) const override { return EPCGElementExecutionLoopMode::SinglePrimaryPin; }
+	virtual bool SupportsBasePointDataInputs(FPCGContext* InContext) const override { return true; }
+
 	/** Performs the filter operation. Returns true when the operation is complete. */
 	bool DoFiltering(FPCGContext* Context, EPCGAttributeFilterOperator InOperation, const FPCGAttributePropertyInputSelector& TargetAttribute, bool bHasSpatialToPointDeprecation, bool bWarnOnDataMissingAttribute, const FPCGAttributeFilterThresholdSettings& FirstThreshold, const FPCGAttributeFilterThresholdSettings* SecondThreshold = nullptr) const;
 };
diff --git a/Engine/Plugins/PCG/Source/PCG/Public/Elements/PCGAttributeGetFromPointIndexElement.h b/Engine/Plugins/PCG/Source/PCG/Public/Elements/PCGAttributeGetFromPointIndexElement.h
--- a/Engine/Plugins/PCG/Source/PCG/Public/Elements/PCGAttributeGetFromPointIndexElement.h
+++ b/Engine/Plugins/PCG/Source/PCG/Public/Elements/PCGAttributeGetFromPointIndexElement.h
@@ -70,4 +70,5 @@ class FPCGAttributeGetFromPointIndexElement : public IPCGElement
 protected:
 	virtual bool ExecuteInternal(FPCGContext* Context) const override;
 	virtual EPCGElementExecutionLoopMode ExecutionLoopMode(const UPCGSettings* Settings) const override { return EPCGElementExecutionLoopMode::SinglePrimaryPin; }
+	virtual bool SupportsBasePointDataInputs(FPCGContext* InContext) const override { return true; }
 };
diff --git a/Engine/Plugins/PCG/Source/PCG/Public/Elements/PCGAttributeNoise.h b/Engine/Plugins/PCG/Source/PCG/Public/Elements/PCGAttributeNoise.h
--- a/Engine/Plugins/PCG/Source/PCG/Public/Elements/PCGAttributeNoise.h
+++ b/Engine/Plugins/PCG/Source/PCG/Public/Elements/PCGAttributeNoise.h
@@ -137,4 +137,5 @@ class FPCGAttributeNoiseElement : public IPCGElementWithCustomContext<FPCGAttrib
 protected:
 	virtual bool ExecuteInternal(FPCGContext* Context) const override;
 	virtual EPCGElementExecutionLoopMode ExecutionLoopMode(const UPCGSettings* Settings) const override { return EPCGElementExecutionLoopMode::SinglePrimaryPin; }
+	virtual bool SupportsBasePointDataInputs(FPCGContext* InContext) const override { return true; }
 };
diff --git a/Engine/Plugins/PCG/Source/PCG/Public/Elements/PCGCollapseElement.h b/Engine/Plugins/PCG/Source/PCG/Public/Elements/PCGCollapseElement.h
--- a/Engine/Plugins/PCG/Source/PCG/Public/Elements/PCGCollapseElement.h
+++ b/Engine/Plugins/PCG/Source/PCG/Public/Elements/PCGCollapseElement.h
@@ -64,4 +64,5 @@ public:
 protected:
 	virtual bool ExecuteInternal(FPCGContext* Context) const override;
 	virtual EPCGElementExecutionLoopMode ExecutionLoopMode(const UPCGSettings* Settings) const override { return EPCGElementExecutionLoopMode::SinglePrimaryPin; }
+	virtual bool SupportsBasePointDataInputs(FPCGContext* InContext) const override { return true; }
 };
diff --git a/Engine/Plugins/PCG/Source/PCG/Public/Elements/PCGConvertToAttributeSet.h b/Engine/Plugins/PCG/Source/PCG/Public/Elements/PCGConvertToAttributeSet.h
--- a/Engine/Plugins/PCG/Source/PCG/Public/Elements/PCGConvertToAttributeSet.h
+++ b/Engine/Plugins/PCG/Source/PCG/Public/Elements/PCGConvertToAttributeSet.h
@@ -40,6 +40,7 @@ public:
 protected:
 	virtual bool ExecuteInternal(FPCGContext* Context) const override;
 	virtual EPCGElementExecutionLoopMode ExecutionLoopMode(const UPCGSettings* Settings) const override { return EPCGElementExecutionLoopMode::SinglePrimaryPin; }
+	virtual bool SupportsBasePointDataInputs(FPCGContext* InContext) const override { return true; }
 };
 
 /** Converts tags on data collections to an attribute set (per input data), akin to how it's done on Get Actor Data with the Single Point option. */
@@ -68,4 +69,5 @@ class FPCGTagsToAttributeSetElement : public IPCGElement
 protected:
 	virtual bool ExecuteInternal(FPCGContext* Context) const override;
 	virtual EPCGElementExecutionLoopMode ExecutionLoopMode(const UPCGSettings* Settings) const override { return EPCGElementExecutionLoopMode::SinglePrimaryPin; }
+	virtual bool SupportsBasePointDataInputs(FPCGContext* InContext) const override { return true; }
 };
\ No newline at end of file
diff --git a/Engine/Plugins/PCG/Source/PCG/Public/Elements/PCGCopyPoints.h b/Engine/Plugins/PCG/Source/PCG/Public/Elements/PCGCopyPoints.h
--- a/Engine/Plugins/PCG/Source/PCG/Public/Elements/PCGCopyPoints.h
+++ b/Engine/Plugins/PCG/Source/PCG/Public/Elements/PCGCopyPoints.h
@@ -102,4 +102,5 @@ class FPCGCopyPointsElement : public IPCGElement
 {
 protected:
 	virtual bool ExecuteInternal(FPCGContext* Context) const override;
+	virtual bool SupportsBasePointDataInputs(FPCGContext* InContext) const override { return true; }
 };
diff --git a/Engine/Plugins/PCG/Source/PCG/Public/Elements/PCGDensityFilter.h b/Engine/Plugins/PCG/Source/PCG/Public/Elements/PCGDensityFilter.h
--- a/Engine/Plugins/PCG/Source/PCG/Public/Elements/PCGDensityFilter.h
+++ b/Engine/Plugins/PCG/Source/PCG/Public/Elements/PCGDensityFilter.h
@@ -47,4 +47,5 @@ class FPCGDensityFilterElement : public IPCGElement
 protected:
 	virtual bool ExecuteInternal(FPCGContext* Context) const override;
 	virtual EPCGElementExecutionLoopMode ExecutionLoopMode(const UPCGSettings* Settings) const override { return EPCGElementExecutionLoopMode::SinglePrimaryPin; }
+	virtual bool SupportsBasePointDataInputs(FPCGContext* InContext) const override { return true; }
 };
diff --git a/Engine/Plugins/PCG/Source/PCG/Public/Elements/PCGDistance.h b/Engine/Plugins/PCG/Source/PCG/Public/Elements/PCGDistance.h
--- a/Engine/Plugins/PCG/Source/PCG/Public/Elements/PCGDistance.h
+++ b/Engine/Plugins/PCG/Source/PCG/Public/Elements/PCGDistance.h
@@ -95,4 +95,5 @@ class FPCGDistanceElement : public IPCGElement
 protected:
 	virtual bool ExecuteInternal(FPCGContext* Context) const override;
 	virtual EPCGElementExecutionLoopMode ExecutionLoopMode(const UPCGSettings* Settings) const override { return EPCGElementExecutionLoopMode::SinglePrimaryPin; }
+	virtual bool SupportsBasePointDataInputs(FPCGContext* InContext) const override { return true; }
 };
\ No newline at end of file
diff --git a/Engine/Plugins/PCG/Source/PCG/Public/Elements/PCGFilterByAttribute.h b/Engine/Plugins/PCG/Source/PCG/Public/Elements/PCGFilterByAttribute.h
--- a/Engine/Plugins/PCG/Source/PCG/Public/Elements/PCGFilterByAttribute.h
+++ b/Engine/Plugins/PCG/Source/PCG/Public/Elements/PCGFilterByAttribute.h
@@ -47,4 +47,5 @@ class FPCGFilterByAttributeElement : public IPCGElement
 protected:
 	virtual bool ExecuteInternal(FPCGContext* Context) const override;
 	virtual EPCGElementExecutionLoopMode ExecutionLoopMode(const UPCGSettings* Settings) const override { return EPCGElementExecutionLoopMode::SinglePrimaryPin; }
+	virtual bool SupportsBasePointDataInputs(FPCGContext* InContext) const override { return true; }
 };
diff --git a/Engine/Plugins/PCG/Source/PCG/Public/Elements/PCGFilterByIndex.h b/Engine/Plugins/PCG/Source/PCG/Public/Elements/PCGFilterByIndex.h
--- a/Engine/Plugins/PCG/Source/PCG/Public/Elements/PCGFilterByIndex.h
+++ b/Engine/Plugins/PCG/Source/PCG/Public/Elements/PCGFilterByIndex.h
@@ -46,4 +46,5 @@ class FPCGFilterByIndexElement : public IPCGElement
 protected:
 	virtual bool ExecuteInternal(FPCGContext* Context) const override;
 	virtual EPCGElementExecutionLoopMode ExecutionLoopMode(const UPCGSettings* Settings) const override { return EPCGElementExecutionLoopMode::SinglePrimaryPin; }
+	virtual bool SupportsBasePointDataInputs(FPCGContext* InContext) const override { return true; }
 };
diff --git a/Engine/Plugins/PCG/Source/PCG/Public/Elements/PCGMatchAndSetAttributes.h b/Engine/Plugins/PCG/Source/PCG/Public/Elements/PCGMatchAndSetAttributes.h
--- a/Engine/Plugins/PCG/Source/PCG/Public/Elements/PCGMatchAndSetAttributes.h
+++ b/Engine/Plugins/PCG/Source/PCG/Public/Elements/PCGMatchAndSetAttributes.h
@@ -128,4 +128,5 @@ protected:
 	virtual bool PrepareDataInternal(FPCGContext* InContext) const override;
 	virtual bool ExecuteInternal(FPCGContext* InContext) const override;
 	virtual EPCGElementExecutionLoopMode ExecutionLoopMode(const UPCGSettings* Settings) const override { return EPCGElementExecutionLoopMode::SinglePrimaryPin; }
+	virtual bool SupportsBasePointDataInputs(FPCGContext* InContext) const override { return true; }
 };
diff --git a/Engine/Plugins/PCG/Source/PCG/Public/Elements/PCGMergeElement.h b/Engine/Plugins/PCG/Source/PCG/Public/Elements/PCGMergeElement.h
--- a/Engine/Plugins/PCG/Source/PCG/Public/Elements/PCGMergeElement.h
+++ b/Engine/Plugins/PCG/Source/PCG/Public/Elements/PCGMergeElement.h
@@ -47,4 +47,5 @@ class FPCGMergeElement : public IPCGElement
 {
 protected:
 	virtual bool ExecuteInternal(FPCGContext* Context) const override;
+	virtual bool SupportsBasePointDataInputs(FPCGContext* InContext) const override { return true; }
 };
diff --git a/Engine/Plugins/PCG/Source/PCG/Public/Elements/PCGProjectionElement.h b/Engine/Plugins/PCG/Source/PCG/Public/Elements/PCGProjectionElement.h
--- a/Engine/Plugins/PCG/Source/PCG/Public/Elements/PCGProjectionElement.h
+++ b/Engine/Plugins/PCG/Source/PCG/Public/Elements/PCGProjectionElement.h
@@ -59,4 +59,5 @@ class FPCGProjectionElement : public IPCGElement
 {
 protected:
 	virtual bool ExecuteInternal(FPCGContext* Context) const override;
+	virtual bool SupportsBasePointDataInputs(FPCGContext* InContext) const override { return true; }
 };
diff --git a/Engine/Plugins/PCG/Source/PCG/Public/Elements/PCGSelectPoints.h b/Engine/Plugins/PCG/Source/PCG/Public/Elements/PCGSelectPoints.h
--- a/Engine/Plugins/PCG/Source/PCG/Public/Elements/PCGSelectPoints.h
+++ b/Engine/Plugins/PCG/Source/PCG/Public/Elements/PCGSelectPoints.h
@@ -42,4 +42,5 @@ class FPCGSelectPointsElement : public IPCGElement
 protected:
 	virtual bool ExecuteInternal(FPCGContext* Context) const override;
 	virtual EPCGElementExecutionLoopMode ExecutionLoopMode(const UPCGSettings* Settings) const override { return EPCGElementExecutionLoopMode::SinglePrimaryPin; }
+	virtual bool SupportsBasePointDataInputs(FPCGContext* InContext) const { return true; }
 };
diff --git a/Engine/Plugins/PCG/Source/PCG/Public/Elements/PCGSelfPruning.h b/Engine/Plugins/PCG/Source/PCG/Public/Elements/PCGSelfPruning.h
--- a/Engine/Plugins/PCG/Source/PCG/Public/Elements/PCGSelfPruning.h
+++ b/Engine/Plugins/PCG/Source/PCG/Public/Elements/PCGSelfPruning.h
@@ -12,6 +12,9 @@
 
 #include "PCGSelfPruning.generated.h"
 
+class UPCGBasePointData;
+class UPCGPointData;
+
 UENUM()
 enum class EPCGSelfPruningType : uint8
 {
@@ -84,8 +87,9 @@ namespace PCGSelfPruningElement
 
 	struct FIterationState
 	{
-		const UPCGPointData* InputPointData = nullptr;
-		UPCGPointData* OutputPointData = nullptr;
+		const UPCGBasePointData* InputData = nullptr;
+		UPCGBasePointData* OutputData = nullptr;
+
 		TArray<PCGPointOctree::FPointRef> SortedPointRefs;
 		FPointBitSet ExcludedPoints;
 		FPointBitSet ExclusionPoints;
@@ -100,6 +104,12 @@ namespace PCGSelfPruningElement
 PRAGMA_DISABLE_DEPRECATION_WARNINGS
 		UE_DEPRECATED(5.6, "Use SortedPointRefs instead")
 		TArray<FPCGPointRef> SortedPoints;
+
+		UE_DEPRECATED(5.6, "Use InputData instead")
+		const UPCGPointData* InputPointData = nullptr;
+
+		UE_DEPRECATED(5.6, "Use OutputData instead")
+		UPCGPointData* OutputPointData = nullptr;
 PRAGMA_ENABLE_DEPRECATION_WARNINGS
 	};
 
@@ -162,4 +172,5 @@ protected:
 	virtual bool ExecuteInternal(FPCGContext* Context) const override;
 	virtual bool CanExecuteOnlyOnMainThread(FPCGContext* Context) const override;
 	virtual EPCGElementExecutionLoopMode ExecutionLoopMode(const UPCGSettings* Settings) const override { return EPCGElementExecutionLoopMode::SinglePrimaryPin; }
+	virtual bool SupportsBasePointDataInputs(FPCGContext* InContext) const override { return true; }
 };
diff --git a/Engine/Plugins/PCG/Source/PCG/Public/Elements/PCGSortAttributes.h b/Engine/Plugins/PCG/Source/PCG/Public/Elements/PCGSortAttributes.h
--- a/Engine/Plugins/PCG/Source/PCG/Public/Elements/PCGSortAttributes.h
+++ b/Engine/Plugins/PCG/Source/PCG/Public/Elements/PCGSortAttributes.h
@@ -52,4 +52,5 @@ class FPCGSortAttributesElement : public IPCGElement
 protected:
 	virtual bool ExecuteInternal(FPCGContext* Context) const override;
 	virtual EPCGElementExecutionLoopMode ExecutionLoopMode(const UPCGSettings* Settings) const override { return EPCGElementExecutionLoopMode::SinglePrimaryPin; }
+	virtual bool SupportsBasePointDataInputs(FPCGContext* InContext) const override { return true; }
 };
\ No newline at end of file
diff --git a/Engine/Plugins/PCG/Source/PCG/Public/Elements/PCGSpatialNoise.h b/Engine/Plugins/PCG/Source/PCG/Public/Elements/PCGSpatialNoise.h
--- a/Engine/Plugins/PCG/Source/PCG/Public/Elements/PCGSpatialNoise.h
+++ b/Engine/Plugins/PCG/Source/PCG/Public/Elements/PCGSpatialNoise.h
@@ -156,4 +156,5 @@ class FPCGSpatialNoise : public IPCGElement
 protected:
 	virtual bool ExecuteInternal(FPCGContext* Context) const override;
 	virtual EPCGElementExecutionLoopMode ExecutionLoopMode(const UPCGSettings* Settings) const override { return EPCGElementExecutionLoopMode::SinglePrimaryPin; }
+	virtual bool SupportsBasePointDataInputs(FPCGContext* InContext) const override { return true; }
 };
diff --git a/Engine/Plugins/PCG/Source/PCG/Public/Elements/PCGTransformPoints.h b/Engine/Plugins/PCG/Source/PCG/Public/Elements/PCGTransformPoints.h
--- a/Engine/Plugins/PCG/Source/PCG/Public/Elements/PCGTransformPoints.h
+++ b/Engine/Plugins/PCG/Source/PCG/Public/Elements/PCGTransformPoints.h
@@ -77,4 +77,5 @@ class FPCGTransformPointsElement : public IPCGElement
 protected:
 	virtual bool ExecuteInternal(FPCGContext* Context) const override;
 	virtual EPCGElementExecutionLoopMode ExecutionLoopMode(const UPCGSettings* Settings) const override { return EPCGElementExecutionLoopMode::SinglePrimaryPin; }
+	virtual bool SupportsBasePointDataInputs(FPCGContext* InContext) const override { return true; }
 };
diff --git a/Engine/Plugins/PCG/Source/PCG/Public/Helpers/PCGHelpers.h b/Engine/Plugins/PCG/Source/PCG/Public/Helpers/PCGHelpers.h
--- a/Engine/Plugins/PCG/Source/PCG/Public/Helpers/PCGHelpers.h
+++ b/Engine/Plugins/PCG/Source/PCG/Public/Helpers/PCGHelpers.h
@@ -8,6 +8,7 @@
 #include "Containers/Ticker.h"
 #include "Kismet/BlueprintFunctionLibrary.h"
 #include "Math/Box.h"
+#include "Math/RandomStream.h"
 
 #include "PCGHelpers.generated.h"
 
@@ -18,6 +19,7 @@ class ALandscapeProxy;
 class UPCGComponent;
 class UPCGGraph;
 class UPCGMetadata;
+class UPCGSettings;
 class UWorld;
 struct FPCGContext;
 struct FPCGPoint;
@@ -34,6 +36,8 @@ namespace PCGHelpers
 	PCG_API int ComputeSeed(int A, int B);
 	PCG_API int ComputeSeed(int A, int B, int C);
 	PCG_API int ComputeSeedFromPosition(const FVector& InPosition);
+	PCG_API FRandomStream GetRandomStreamFromSeed(int32 Seed, const UPCGSettings* OptionalSettings = nullptr, const UPCGComponent* OptionalComponent = nullptr);
+	PCG_API FRandomStream GetRandomStreamFromTwoSeeds(int32 SeedA, int32 SeedB, const UPCGSettings* OptionalSettings = nullptr, const UPCGComponent* OptionalComponent = nullptr);
 
 	PCG_API bool IsInsideBounds(const FBox& InBox, const FVector& InPosition);
 	PCG_API bool IsInsideBoundsXY(const FBox& InBox, const FVector& InPosition);
diff --git a/Engine/Plugins/PCG/Source/PCG/Public/Metadata/Accessors/PCGAttributeAccessorHelpers.h b/Engine/Plugins/PCG/Source/PCG/Public/Metadata/Accessors/PCGAttributeAccessorHelpers.h
--- a/Engine/Plugins/PCG/Source/PCG/Public/Metadata/Accessors/PCGAttributeAccessorHelpers.h
+++ b/Engine/Plugins/PCG/Source/PCG/Public/Metadata/Accessors/PCGAttributeAccessorHelpers.h
@@ -215,36 +215,36 @@ namespace PCGAttributeAccessorHelpers
 	}
 
 	/**
-	* Sorts array given the accessors and keys of the array. Sort is stable by default.
+	* Returns a sorted key indices array.
 	* A custom function can be used to extract the index of the elements of the array,
 	* and sort ascending or descending by the values associated with that index.
 	* CompareLess method can also be provided, method signature needs to follow DefaultStableCompareLess.
 	* We need the decltypes on both function types to allow to provide default values for both callbacks.
 	* Check DefaultIndexGetter and DefaultStableCompareLess for their signatures.
 	*/
-	template <typename T, typename GetIndexFunc = decltype(Private::DefaultIndexGetter), typename CompareLessFunc = decltype(Private::DefaultStableCompareLessLambda)>
-	void SortByAttribute(const IPCGAttributeAccessor& InAccessor, const IPCGAttributeAccessorKeys& InKeys, TArray<T>& InArray, bool bAscending, GetIndexFunc CustomGetIndex = Private::DefaultIndexGetter, CompareLessFunc CompareLess = Private::DefaultStableCompareLessLambda)
+	template <typename GetIndexFunc = decltype(Private::DefaultIndexGetter), typename CompareLessFunc = decltype(Private::DefaultStableCompareLessLambda)>
+	TArray<int32> SortKeyIndicesByAttribute(const IPCGAttributeAccessor& InAccessor, const IPCGAttributeAccessorKeys& InKeys, int32 KeyCount, bool bAscending, GetIndexFunc CustomGetIndex = Private::DefaultIndexGetter, CompareLessFunc CompareLess = Private::DefaultStableCompareLessLambda)
 	{
-		TRACE_CPUPROFILER_EVENT_SCOPE(PCGAttributeAccessorHelpers::SortByAttribute);
-		check(InArray.Num() <= InKeys.GetNum())
+		TRACE_CPUPROFILER_EVENT_SCOPE(PCGAttributeAccessorHelpers::SortKeyIndicesByAttribute);
+		check(KeyCount <= InKeys.GetNum())
 
-		if (InArray.IsEmpty())
+		if (KeyCount == 0)
 		{
-			return;
+			return {};
 		}
 
 		// Prepare integer sequence
 		TArray<int32> ElementsIndexes;
 		{
 			TRACE_CPUPROFILER_EVENT_SCOPE(PCGAttributeAccessorHelpers::SortByAttribute::PrepareIndexes);
-			ElementsIndexes.Reserve(InArray.Num());
-			for (int i = 0; i < InArray.Num(); ++i)
+			ElementsIndexes.Reserve(KeyCount);
+			for (int i = 0; i < KeyCount; ++i)
 			{
 				ElementsIndexes.Add(i);
 			}
 		}
 
-		auto SortIndexes = [&InAccessor, &InKeys, &InArray, bAscending, &CustomGetIndex, &CompareLess, &ElementsIndexes](auto Dummy)
+		auto SortIndexes = [&InAccessor, &InKeys, bAscending, &CustomGetIndex, &CompareLess, &ElementsIndexes](auto Dummy)
 		{
 			using ValueType = decltype(Dummy);
 
@@ -297,6 +297,22 @@ namespace PCGAttributeAccessorHelpers
 		// Perform the sorting on the indexes using the attribute values provided by the accessor.
 		PCGMetadataAttribute::CallbackWithRightType(InAccessor.GetUnderlyingType(), SortIndexes);
 
+		return ElementsIndexes;
+	}
+
+	/**
+	* Sorts array given the accessors and keys of the array. Sort is stable by default.
+	* A custom function can be used to extract the index of the elements of the array,
+	* and sort ascending or descending by the values associated with that index.
+	* CompareLess method can also be provided, method signature needs to follow DefaultStableCompareLess.
+	* We need the decltypes on both function types to allow to provide default values for both callbacks.
+	* Check DefaultIndexGetter and DefaultStableCompareLess for their signatures.
+	*/
+	template <typename T, typename GetIndexFunc = decltype(Private::DefaultIndexGetter), typename CompareLessFunc = decltype(Private::DefaultStableCompareLessLambda)>
+	void SortByAttribute(const IPCGAttributeAccessor& InAccessor, const IPCGAttributeAccessorKeys& InKeys, TArray<T>& InArray, bool bAscending, GetIndexFunc CustomGetIndex = Private::DefaultIndexGetter, CompareLessFunc CompareLess = Private::DefaultStableCompareLessLambda)
+	{
+		TArray<int32> SortedKeyIndices = SortKeyIndicesByAttribute(InAccessor, InKeys, InArray.Num(), bAscending, CustomGetIndex, CompareLess);
+
 		// Write back the values according to the sorted indexes.
 		TArray<T> SortedArray;
 		{
@@ -305,7 +321,7 @@ namespace PCGAttributeAccessorHelpers
 
 			for (int i = 0; i < InArray.Num(); ++i)
 			{
-				SortedArray.Add(MoveTemp(InArray[ElementsIndexes[i]]));
+				SortedArray.Add(MoveTemp(InArray[SortedKeyIndices[i]]));
 			}
 		}
 
diff --git a/Engine/Plugins/PCG/Source/PCG/Public/Metadata/Accessors/PCGAttributeAccessorKeys.h b/Engine/Plugins/PCG/Source/PCG/Public/Metadata/Accessors/PCGAttributeAccessorKeys.h
--- a/Engine/Plugins/PCG/Source/PCG/Public/Metadata/Accessors/PCGAttributeAccessorKeys.h
+++ b/Engine/Plugins/PCG/Source/PCG/Public/Metadata/Accessors/PCGAttributeAccessorKeys.h
@@ -9,6 +9,7 @@ class FPCGMetadataDomain;
 struct FPCGPoint;
 class UPCGMetadata;
 class UPCGSplineData;
+class UPCGBasePointData;
 
 ///////////////////////////////////////////////////////////////////////
 
@@ -246,11 +247,16 @@ public:
 	FPCGAttributeAccessorKeysPointsSubset(TArray<FPCGPoint*> InPointPtrs);
 	FPCGAttributeAccessorKeysPointsSubset(TArray<const FPCGPoint*> InPointPtrs);
 
-	virtual int32 GetNum() const override { return Points.Num(); }
+	FPCGAttributeAccessorKeysPointsSubset(const UPCGBasePointData* InPointData, const TArrayView<const int32>& InPointIndices);
+	FPCGAttributeAccessorKeysPointsSubset(UPCGBasePointData* InPointData, const TArrayView<const int32>& InPointIndices);
+
+	virtual int32 GetNum() const override { return Points.Num() > 0 ? Points.Num() : PointIndices.Num(); }
 
 	virtual bool IsClassSupported(const UStruct* InClass) const override;
 
 protected:
+	virtual bool GetKeyIndices(int32 InStart, int32 InCount, TArray<int32>& OutKeyIndices, bool& bOutContiguous) const override;
+
 	virtual bool GetPointKeys(int32 InStart, TArrayView<FPCGPoint*> OutPoints) override;
 	virtual bool GetPointKeys(int32 InStart, TArrayView<const FPCGPoint*> OutPoints) const override;
 
@@ -261,6 +267,9 @@ protected:
 	virtual bool GetMetadataEntryKeys(int32 InStart, TArrayView<const PCGMetadataEntryKey*> OutEntryKeys) const override;
 
 	TArray<FPCGPoint*> Points;
+
+	UPCGBasePointData* PointData = nullptr;
+	TArray<int32> PointIndices;
 };
 
 /////////////////////////////////////////////////////////////////
diff --git a/Engine/Plugins/PCG/Source/PCG/Public/Utils/PCGValueRange.h b/Engine/Plugins/PCG/Source/PCG/Public/Utils/PCGValueRange.h
--- a/Engine/Plugins/PCG/Source/PCG/Public/Utils/PCGValueRange.h
+++ b/Engine/Plugins/PCG/Source/PCG/Public/Utils/PCGValueRange.h
@@ -33,24 +33,28 @@
 * 
 * Range with 5 Elements & StridedView of 5 Elements(s) or more
 * Range Indices [0, 1, 2, 3, 4] StridedView Indices [0, 1, 2, 3, 4]
+* 
+* Note: For multithreaded usage it is preferred to declare the ranges in each thread so that the 'StubValue' is thread local
+* 
 */
 template<typename ElementType, typename ViewType = TStridedView<ElementType>>
 class TPCGValueRange
 {
 public:
 	using SizeType = typename ViewType::SizeType;
+	using NonConstElementType = std::remove_const_t<ElementType>;
 
 	TPCGValueRange()
-		: ElementView(ViewType()), NumElements(0)
+		: TPCGValueRange(ViewType(), 0)
 	{
 	}
 
-	TPCGValueRange(ViewType InElementView)
-		: ElementView(InElementView), NumElements(InElementView.Num())
+	explicit TPCGValueRange(ViewType InElementView)
+		: TPCGValueRange(InElementView, InElementView.Num())
 	{
 	}
 
-	TPCGValueRange(ViewType InElementView, SizeType InNumElements)
+	explicit TPCGValueRange(ViewType InElementView, SizeType InNumElements)
 		: ElementView(InElementView), NumElements(InNumElements)
 	{
 	}
@@ -133,18 +137,41 @@ private:
 	
 	[[nodiscard]] FORCEINLINE ElementType& GetElementUnsafe(SizeType Index) const
 	{
-		return ElementView.GetUnsafe(Index % ElementView.Num());
+		const SizeType ViewIndex = Index % ElementView.Num();
+		ElementType& Value = ElementView.GetUnsafe(ViewIndex);
+
+		// Accessing the value for read in a multithreaded context is safe so no need for the StubValue here.
+		if constexpr (std::is_const_v<ElementType>)
+		{
+			return Value;
+		}
+		else
+		{			
+			// It is possible that our ElementView is a single value for non-allocated properties in which case 
+			// we want to prevent multithreaded writes to this single value.
+			// Only allowing to return the actual value reference if Index == ViewIndex.
+			// If not we return a StubValue reference which is local to the ValueRange object.
+			if (Index == ViewIndex)
+			{
+				return Value;
+			}
+			else
+			{
+				return StubValue;
+			}
+		}
 	}
 
 	[[nodiscard]] FORCEINLINE ElementType& GetElement(SizeType Index) const
 	{
 		RangeCheck(Index);
-		return ElementView.GetUnsafe(Index % ElementView.Num());
+		return GetElementUnsafe(Index);
 	}
 
 private:
 	ViewType ElementView;
 	SizeType NumElements;
+	mutable NonConstElementType StubValue;
 };
 
 template<typename ElementType> using TConstPCGValueRange = TPCGValueRange<const ElementType, TConstStridedView<ElementType>>;
\ No newline at end of file
