[PCG] GPU: Add context handle to data binding to enable GPU code that only has the binding to log graph errors.

Use this to emit an error when data label resolution fails.

#rnx

diff --git a/Engine/Plugins/PCG/Source/PCG/Private/Compute/DataInterfaces/PCGCustomKernelDataInterface.cpp b/Engine/Plugins/PCG/Source/PCG/Private/Compute/DataInterfaces/PCGCustomKernelDataInterface.cpp
--- a/Engine/Plugins/PCG/Source/PCG/Private/Compute/DataInterfaces/PCGCustomKernelDataInterface.cpp
+++ b/Engine/Plugins/PCG/Source/PCG/Private/Compute/DataInterfaces/PCGCustomKernelDataInterface.cpp
@@ -96,17 +96,19 @@ UComputeDataProvider* UPCGCustomKernelDataInterface::CreateDataProvider(TObjectP
 	check(Kernel);
 
 	UPCGDataBinding* Binding = CastChecked<UPCGDataBinding>(InBinding);
-	check(Binding->SourceComponent.IsValid() && Binding->SourceComponent.Get());
+
+	const UPCGComponent* SourceComponent = Binding->GetSourceComponent().Get();
+	check(SourceComponent);
 
 	UPCGCustomComputeKernelDataProvider* Provider = NewObject<UPCGCustomComputeKernelDataProvider>();
 	Provider->ThreadCount = Kernel->ComputeThreadCount(Binding);
 	Provider->ThreadCountMultiplier = Kernel->GetThreadCountMultiplier();
 
-	Provider->Seed = static_cast<uint32>(Settings->GetSeed(Binding->SourceComponent.Get()));
+	Provider->Seed = static_cast<uint32>(Settings->GetSeed(SourceComponent));
 	Provider->SeedSettings = static_cast<uint32>(Settings->Seed);
-	Provider->SeedComponent = static_cast<uint32>(Binding->SourceComponent->Seed);
+	Provider->SeedComponent = static_cast<uint32>(SourceComponent->Seed);
 
-	Provider->SourceComponentBounds = Binding->SourceComponent.Get()->GetGridBounds();
+	Provider->SourceComponentBounds = SourceComponent->GetGridBounds();
 
 	return Provider;
 }
diff --git a/Engine/Plugins/PCG/Source/PCG/Private/Compute/DataInterfaces/PCGDataCollectionDataInterface.cpp b/Engine/Plugins/PCG/Source/PCG/Private/Compute/DataInterfaces/PCGDataCollectionDataInterface.cpp
--- a/Engine/Plugins/PCG/Source/PCG/Private/Compute/DataInterfaces/PCGDataCollectionDataInterface.cpp
+++ b/Engine/Plugins/PCG/Source/PCG/Private/Compute/DataInterfaces/PCGDataCollectionDataInterface.cpp
@@ -468,7 +468,8 @@ UComputeDataProvider* UPCGDataCollectionDataInterface::CreateDataProvider(TObjec
 	Provider->bExposeElementCounters = bExposeElementCounters;
 
 #if WITH_EDITOR
-	if (ProducerSettings->bIsInspecting && Binding->SourceComponent->IsInspecting())
+	const UPCGComponent* SourceComponent = Binding->GetSourceComponent().Get();
+	if (ProducerSettings->bIsInspecting && SourceComponent && SourceComponent->IsInspecting())
 	{
 		Provider->ExportMode |= EPCGExportMode::Inspection;
 	}
diff --git a/Engine/Plugins/PCG/Source/PCG/Private/Compute/DataInterfaces/PCGDataLabelResolverDataInterface.cpp b/Engine/Plugins/PCG/Source/PCG/Private/Compute/DataInterfaces/PCGDataLabelResolverDataInterface.cpp
--- a/Engine/Plugins/PCG/Source/PCG/Private/Compute/DataInterfaces/PCGDataLabelResolverDataInterface.cpp
+++ b/Engine/Plugins/PCG/Source/PCG/Private/Compute/DataInterfaces/PCGDataLabelResolverDataInterface.cpp
@@ -2,6 +2,7 @@
 
 #include "Compute/DataInterfaces/PCGDataLabelResolverDataInterface.h"
 
+#include "PCGContext.h"
 #include "PCGModule.h"
 #include "PCGSettings.h"
 #include "Compute/PCGComputeGraph.h"
@@ -112,7 +113,12 @@ UComputeDataProvider* UPCGDataLabelResolverDataInterface::CreateDataProvider(TOb
 
 			if (DataIdToDataIndexMap[LabelIndex] == INDEX_NONE)
 			{
-				UE_LOG(LogPCG, Warning, TEXT("DataLabelResolver found data label '%s', but did not match any incoming data."), *Label);
+				if (TSharedPtr<FPCGContextHandle> Context = Binding->ContextHandle.Pin())
+				{
+					PCG_KERNEL_VALIDATION_ERR(Context->GetContext(), Kernel->GetSettings(), /*bQuiet=*/false, FText::Format(
+						NSLOCTEXT("PCGDataLabelResolver", "LabelDoesNotExist", "Data label '{0}' was not present in the tags on the incoming data."),
+						FText::FromString(Label)));
+				}
 			}
 		}
 
diff --git a/Engine/Plugins/PCG/Source/PCG/Private/Compute/DataInterfaces/PCGLandscapeDataInterface.cpp b/Engine/Plugins/PCG/Source/PCG/Private/Compute/DataInterfaces/PCGLandscapeDataInterface.cpp
--- a/Engine/Plugins/PCG/Source/PCG/Private/Compute/DataInterfaces/PCGLandscapeDataInterface.cpp
+++ b/Engine/Plugins/PCG/Source/PCG/Private/Compute/DataInterfaces/PCGLandscapeDataInterface.cpp
@@ -244,9 +244,9 @@ UComputeDataProvider* UPCGLandscapeDataInterface::CreateDataProvider(TObjectPtr<
 {
 	TRACE_CPUPROFILER_EVENT_SCOPE(UPCGLandscapeDataInterface::CreateDataProvider);
 	UPCGDataBinding* Binding = CastChecked<UPCGDataBinding>(InBinding);
-	check(Binding->SourceComponent.IsValid() && Binding->SourceComponent.Get());
+	const UPCGComponent* SourceComponent = Binding->GetSourceComponent().Get();
+	check(SourceComponent);
 
-	const UPCGComponent* SourceComponent = Binding->SourceComponent.Get();
 	const UWorld* World = SourceComponent->GetWorld();
 	const FBox ComponentBounds = SourceComponent->GetGridBounds();
 
diff --git a/Engine/Plugins/PCG/Source/PCG/Private/Compute/Elements/PCGComputeGraphElement.cpp b/Engine/Plugins/PCG/Source/PCG/Private/Compute/Elements/PCGComputeGraphElement.cpp
--- a/Engine/Plugins/PCG/Source/PCG/Private/Compute/Elements/PCGComputeGraphElement.cpp
+++ b/Engine/Plugins/PCG/Source/PCG/Private/Compute/Elements/PCGComputeGraphElement.cpp
@@ -158,7 +158,7 @@ bool FPCGComputeGraphElement::ExecuteInternal(FPCGContext* InContext) const
 		UPCGDataBinding* DataBinding = FPCGContext::NewObject_AnyThread<UPCGDataBinding>(Context);
 		Context->DataBinding.Reset(DataBinding);
 
-		Context->DataBinding->Initialize(Context->ComputeGraph.Get(), Context->SourceComponent, Context->InputData);
+		Context->DataBinding->Initialize(Context->ComputeGraph.Get(), Context);
 
 		// Perform validation after input data is initialized.
 		if (!Context->ComputeGraph->IsGraphValid(InContext, /*bQuiet=*/false))
diff --git a/Engine/Plugins/PCG/Source/PCG/Private/Compute/Elements/PCGCustomHLSLKernel.cpp b/Engine/Plugins/PCG/Source/PCG/Private/Compute/Elements/PCGCustomHLSLKernel.cpp
--- a/Engine/Plugins/PCG/Source/PCG/Private/Compute/Elements/PCGCustomHLSLKernel.cpp
+++ b/Engine/Plugins/PCG/Source/PCG/Private/Compute/Elements/PCGCustomHLSLKernel.cpp
@@ -51,6 +51,11 @@ void UPCGCustomHLSLKernel::Initialize(const UPCGCustomHLSLSettings* InSettings,
 	Settings = InSettings;
 }
 
+const UPCGSettings* UPCGCustomHLSLKernel::GetSettings() const
+{
+	return Settings;
+}
+
 bool UPCGCustomHLSLKernel::IsKernelValid(FPCGContext* InContext, bool bQuiet) const
 {
 	TRACE_CPUPROFILER_EVENT_SCOPE(UPCGCustomHLSLKernel::IsKernelValid);
diff --git a/Engine/Plugins/PCG/Source/PCG/Private/Compute/Elements/PCGCustomHLSLKernel.h b/Engine/Plugins/PCG/Source/PCG/Private/Compute/Elements/PCGCustomHLSLKernel.h
--- a/Engine/Plugins/PCG/Source/PCG/Private/Compute/Elements/PCGCustomHLSLKernel.h
+++ b/Engine/Plugins/PCG/Source/PCG/Private/Compute/Elements/PCGCustomHLSLKernel.h
@@ -18,6 +18,7 @@ class UPCGCustomHLSLKernel : public UPCGComputeKernel
 public:
 	void Initialize(const UPCGCustomHLSLSettings* InSettings, bool bLogDescriptions);
 
+	virtual const UPCGSettings* GetSettings() const override;
 	virtual bool IsKernelValid(FPCGContext* InContext = nullptr, bool bQuiet = true) const override;
 	virtual bool ComputeOutputBindingDataDesc(const UPCGComputeGraph* InGraph, FName InOutputPinLabel, const UPCGDataBinding* InBinding, FPCGDataCollectionDesc& OutDataDesc) const override;
 #if WITH_EDITOR
diff --git a/Engine/Plugins/PCG/Source/PCG/Private/Compute/PCGComputeKernel.h b/Engine/Plugins/PCG/Source/PCG/Private/Compute/PCGComputeKernel.h
--- a/Engine/Plugins/PCG/Source/PCG/Private/Compute/PCGComputeKernel.h
+++ b/Engine/Plugins/PCG/Source/PCG/Private/Compute/PCGComputeKernel.h
@@ -31,6 +31,9 @@ public:
 	int32 GetKernelIndex() const { return KernelIndex; }
 	void SetKernelIndex(int32 InKernelIndex) { KernelIndex = InKernelIndex; }
 
+	/** Gets settings for node associated with this kernel, if any. */
+	virtual const UPCGSettings* GetSettings() const { return nullptr; }
+
 	/** Performs validation and returns true if this node is suitable for deployment to the GPU. */
 	virtual bool IsKernelValid(FPCGContext* InContext = nullptr, bool bQuiet = true) const { return true; }
 
diff --git a/Engine/Plugins/PCG/Source/PCG/Private/Compute/PCGDataBinding.cpp b/Engine/Plugins/PCG/Source/PCG/Private/Compute/PCGDataBinding.cpp
--- a/Engine/Plugins/PCG/Source/PCG/Private/Compute/PCGDataBinding.cpp
+++ b/Engine/Plugins/PCG/Source/PCG/Private/Compute/PCGDataBinding.cpp
@@ -2,6 +2,7 @@
 
 #include "Compute/PCGDataBinding.h"
 
+#include "PCGContext.h"
 #include "PCGNode.h"
 #include "PCGParamData.h"
 #include "PCGSettings.h"
@@ -12,16 +13,15 @@
 #include "Data/PCGPointData.h"
 #include "Metadata/PCGMetadata.h"
 
-void UPCGDataBinding::Initialize(
-	const UPCGComputeGraph* InComputeGraph,
-	TWeakObjectPtr<UPCGComponent> InSourceComponent,
-	const FPCGDataCollection& InComputeGraphElementInputData)
+void UPCGDataBinding::Initialize(const UPCGComputeGraph* InComputeGraph, FPCGContext* InContext)
 {
 	TRACE_CPUPROFILER_EVENT_SCOPE(UPCGDataBinding::Initialize);
 	check(InComputeGraph);
+	check(InContext);
+
+	ContextHandle = InContext->GetOrCreateHandle();
 
 	Graph = InComputeGraph;
-	SourceComponent = InSourceComponent;
 
 	AttributeTable = InComputeGraph->GetStaticAttributeTable();
 
@@ -29,7 +29,7 @@ void UPCGDataBinding::Initialize(
 	StringTable = { FString() };
 	StringTable.Append(Graph->GetStaticStringTable());
 
-	InitializeInputData(InComputeGraphElementInputData);
+	InitializeInputData(InContext->InputData);
 
 #if PCG_KERNEL_LOGGING_ENABLED
 	if (Graph->bLogDataDescriptions)
@@ -136,6 +136,19 @@ void UPCGDataBinding::ApplyElementCounts_GameThread()
 	UpdateCountsOnDebugData(DataToDebug);
 }
 
+TWeakObjectPtr<UPCGComponent> UPCGDataBinding::GetSourceComponent() const
+{
+	if (TSharedPtr<FPCGContextHandle> SharedHandle = ContextHandle.Pin())
+	{
+		if (FPCGContext* ContextPtr = SharedHandle->GetContext())
+		{
+			return ContextPtr->SourceComponent;
+		}
+	}
+
+	return nullptr;
+}
+
 void UPCGDataBinding::InitializeInputData(const FPCGDataCollection& InComputeGraphElementInputData)
 {
 	TRACE_CPUPROFILER_EVENT_SCOPE(UPCGDataBinding::InitializeInputData);
diff --git a/Engine/Plugins/PCG/Source/PCG/Private/Compute/PCGDataBinding.h b/Engine/Plugins/PCG/Source/PCG/Private/Compute/PCGDataBinding.h
--- a/Engine/Plugins/PCG/Source/PCG/Private/Compute/PCGDataBinding.h
+++ b/Engine/Plugins/PCG/Source/PCG/Private/Compute/PCGDataBinding.h
@@ -12,6 +12,8 @@ class UPCGComputeGraph;
 class UPCGComputeKernel;
 class UPCGProxyForGPUData;
 class UPCGStaticMeshSpawnerKernel;
+struct FPCGComputeGraphContext;
+struct FPCGContextHandle;
 
 USTRUCT()
 struct FPCGSpawnerPrimitives
@@ -60,7 +62,7 @@ class UPCGDataBinding : public UObject
 	GENERATED_BODY()
 
 public:
-	void Initialize(const UPCGComputeGraph* InComputeGraph, TWeakObjectPtr<UPCGComponent> InSourceComponent, const FPCGDataCollection& InComputeGraphElementInputData);
+	void Initialize(const UPCGComputeGraph* InComputeGraph, FPCGContext* InContext);
 
 	int32 GetAttributeId(const FPCGKernelAttributeKey& InAttribute) const { return AttributeTable.GetAttributeId(InAttribute); }
 	int32 GetAttributeId(FName InName, EPCGKernelAttributeType InType) const { return AttributeTable.GetAttributeId(InName, InType); }
@@ -82,6 +84,8 @@ public:
 
 	const TMap<FSourceBufferAndAttributeIndex, /*Attribute Index*/int32>& GetAttributeIdRemap() const { return AttributeIdRemap; }
 
+	TWeakObjectPtr<UPCGComponent> GetSourceComponent() const;
+
 private:
 	void InitializeInputData(const FPCGDataCollection& InComputeGraphElementInputData);
 
@@ -95,9 +99,6 @@ private:
 	void DebugLogDataDescriptions() const;
 
 public:
-	UPROPERTY()
-	TWeakObjectPtr<UPCGComponent> SourceComponent;
-
 	/** Compute graph element output data. Data items are labelled with unique virtual output pin labels so that the can be routed correctly by the graph
 	 * executor to downstream nodes. */
 	UPROPERTY()
@@ -116,6 +117,8 @@ public:
 	TArray<FPCGDataToDebug> DataToDebug;
 	TArray<FPCGDataToDebug> DataToInspect;
 
+	TWeakPtr<FPCGContextHandle> ContextHandle;
+
 protected:
 	/** Final element counts read back from GPU. */
 	TMap<FName, TArray<uint32>> PinToElementCounts;
diff --git a/Engine/Plugins/PCG/Source/PCG/Private/Elements/Metadata/PCGMetadataPartitionKernel.cpp b/Engine/Plugins/PCG/Source/PCG/Private/Elements/Metadata/PCGMetadataPartitionKernel.cpp
--- a/Engine/Plugins/PCG/Source/PCG/Private/Elements/Metadata/PCGMetadataPartitionKernel.cpp
+++ b/Engine/Plugins/PCG/Source/PCG/Private/Elements/Metadata/PCGMetadataPartitionKernel.cpp
@@ -25,6 +25,11 @@ void UPCGMetadataPartitionKernel::Initialize(const UPCGMetadataPartitionSettings
 	Settings = InSettings;
 }
 
+const UPCGSettings* UPCGMetadataPartitionKernel::GetSettings() const
+{
+	return Settings;
+}
+
 bool UPCGMetadataPartitionKernel::IsKernelValid(FPCGContext* InContext, bool bQuiet) const
 {
 	TRACE_CPUPROFILER_EVENT_SCOPE(UPCGMetadataPartitionKernel::IsKernelValid);
diff --git a/Engine/Plugins/PCG/Source/PCG/Private/Elements/PCGCopyPointsKernel.cpp b/Engine/Plugins/PCG/Source/PCG/Private/Elements/PCGCopyPointsKernel.cpp
--- a/Engine/Plugins/PCG/Source/PCG/Private/Elements/PCGCopyPointsKernel.cpp
+++ b/Engine/Plugins/PCG/Source/PCG/Private/Elements/PCGCopyPointsKernel.cpp
@@ -25,6 +25,11 @@ void UPCGCopyPointsKernel::Initialize(const UPCGCopyPointsSettings* InSettings,
 	Settings = InSettings;
 }
 
+const UPCGSettings* UPCGCopyPointsKernel::GetSettings() const
+{
+	return Settings;
+}
+
 bool UPCGCopyPointsKernel::IsKernelValid(FPCGContext* InContext, bool bQuiet) const
 {
 	TRACE_CPUPROFILER_EVENT_SCOPE(UPCGCopyPointsKernel::IsKernelValid);
diff --git a/Engine/Plugins/PCG/Source/PCG/Private/Elements/PCGStaticMeshSpawnerKernel.cpp b/Engine/Plugins/PCG/Source/PCG/Private/Elements/PCGStaticMeshSpawnerKernel.cpp
--- a/Engine/Plugins/PCG/Source/PCG/Private/Elements/PCGStaticMeshSpawnerKernel.cpp
+++ b/Engine/Plugins/PCG/Source/PCG/Private/Elements/PCGStaticMeshSpawnerKernel.cpp
@@ -38,6 +38,11 @@ void UPCGStaticMeshSpawnerKernel::Initialize(const UPCGStaticMeshSpawnerSettings
 	Settings = InSettings;
 }
 
+const UPCGSettings* UPCGStaticMeshSpawnerKernel::GetSettings() const
+{
+	return Settings;
+}
+
 bool UPCGStaticMeshSpawnerKernel::IsKernelValid(FPCGContext* InContext, bool bQuiet) const
 {
 	TRACE_CPUPROFILER_EVENT_SCOPE(UPCGStaticMeshSpawnerKernel::IsKernelValid);
diff --git a/Engine/Plugins/PCG/Source/PCG/Public/Elements/Metadata/PCGMetadataPartitionKernel.h b/Engine/Plugins/PCG/Source/PCG/Public/Elements/Metadata/PCGMetadataPartitionKernel.h
--- a/Engine/Plugins/PCG/Source/PCG/Public/Elements/Metadata/PCGMetadataPartitionKernel.h
+++ b/Engine/Plugins/PCG/Source/PCG/Public/Elements/Metadata/PCGMetadataPartitionKernel.h
@@ -16,6 +16,7 @@ class UPCGMetadataPartitionKernel : public UPCGComputeKernel
 public:
 	void Initialize(const UPCGMetadataPartitionSettings* InSettings, bool bLogDescriptions);
 
+	virtual const UPCGSettings* GetSettings() const override;
 	virtual bool IsKernelValid(FPCGContext* InContext = nullptr, bool bQuiet = true) const override;
 	virtual bool ComputeOutputBindingDataDesc(const UPCGComputeGraph* InGraph, FName InOutputPinLabel, const UPCGDataBinding* InBinding, FPCGDataCollectionDesc& OutDataDesc) const override;
 	virtual int ComputeThreadCount(const UPCGDataBinding* Binding) const override;
diff --git a/Engine/Plugins/PCG/Source/PCG/Public/Elements/PCGCopyPointsKernel.h b/Engine/Plugins/PCG/Source/PCG/Public/Elements/PCGCopyPointsKernel.h
--- a/Engine/Plugins/PCG/Source/PCG/Public/Elements/PCGCopyPointsKernel.h
+++ b/Engine/Plugins/PCG/Source/PCG/Public/Elements/PCGCopyPointsKernel.h
@@ -18,6 +18,7 @@ class UPCGCopyPointsKernel : public UPCGComputeKernel
 public:
 	void Initialize(const UPCGCopyPointsSettings* InSettings, bool bLogDescriptions);
 
+	virtual const UPCGSettings* GetSettings() const override;
 	virtual bool IsKernelValid(FPCGContext* InContext = nullptr, bool bQuiet = true) const override;
 	virtual bool ComputeOutputBindingDataDesc(const UPCGComputeGraph* InGraph, FName InOutputPinLabel, const UPCGDataBinding* InBinding, FPCGDataCollectionDesc& OutDataDesc) const override;
 	virtual int ComputeThreadCount(const UPCGDataBinding* Binding) const override;
diff --git a/Engine/Plugins/PCG/Source/PCG/Public/Elements/PCGStaticMeshSpawnerKernel.h b/Engine/Plugins/PCG/Source/PCG/Public/Elements/PCGStaticMeshSpawnerKernel.h
--- a/Engine/Plugins/PCG/Source/PCG/Public/Elements/PCGStaticMeshSpawnerKernel.h
+++ b/Engine/Plugins/PCG/Source/PCG/Public/Elements/PCGStaticMeshSpawnerKernel.h
@@ -16,6 +16,7 @@ class UPCGStaticMeshSpawnerKernel : public UPCGComputeKernel
 public:
 	void Initialize(const UPCGStaticMeshSpawnerSettings* InSettings, bool bLogDescriptions);
 
+	virtual const UPCGSettings* GetSettings() const override;
 	virtual bool IsKernelValid(FPCGContext* InContext = nullptr, bool bQuiet = true) const override;
 	virtual bool ComputeOutputBindingDataDesc(const UPCGComputeGraph* InGraph, FName InOutputPinLabel, const UPCGDataBinding* InBinding, FPCGDataCollectionDesc& OutDataDesc) const override;
 	virtual int ComputeThreadCount(const UPCGDataBinding* Binding) const override;
