[PCG] IPCGElement::PostExecute add more profiling scopes

#rnx

diff --git a/Engine/Plugins/PCG/Source/PCG/Private/PCGElement.cpp b/Engine/Plugins/PCG/Source/PCG/Private/PCGElement.cpp
--- a/Engine/Plugins/PCG/Source/PCG/Private/PCGElement.cpp
+++ b/Engine/Plugins/PCG/Source/PCG/Private/PCGElement.cpp
@@ -463,16 +463,23 @@ bool IPCGElement::PrepareDataInternal(FPCGContext* Context) const
 
 void IPCGElement::PostExecute(FPCGContext* Context) const
 {
+	const UPCGSettings* Settings = Context->GetInputSettings<UPCGSettings>();
+
 	// Allow sub class to do some processing here
-	PostExecuteInternal(Context);
+	{
+		TRACE_CPUPROFILER_EVENT_SCOPE_TEXT(*FString::Printf(TEXT("IPCGElement::PostExecute::PostExecuteInternal (%s)"), Settings ? *Settings->GetName() : TEXT("")));
+		PostExecuteInternal(Context);
+	}
 
-	// Cleanup and validate output
-	CleanupAndValidateOutput(Context);
-	
-	const UPCGSettings* Settings = Context->GetInputSettings<UPCGSettings>();
+	{
+		TRACE_CPUPROFILER_EVENT_SCOPE(IPCGElement::PostExecute::CleanupAndValidateOutput);
+		// Cleanup and validate output
+		CleanupAndValidateOutput(Context);
+	}
 
 	if (CVarPCGAllowPerDataCaching.GetValueOnAnyThread())
 	{
+		TRACE_CPUPROFILER_EVENT_SCOPE(IPCGElement::PostExecute::PostExecutePrimaryLoopElement);
 		if (!Context->OutputData.bCancelExecution && ExecutionLoopMode(Settings) != EPCGElementExecutionLoopMode::NotALoop && IsCacheableInstance(Settings))
 		{
 			PostExecutePrimaryLoopElement(Context, Settings);
@@ -498,30 +505,33 @@ void IPCGElement::PostExecute(FPCGContext* Context) const
 	const bool bHasErrorsOrWarnings = false;
 #endif
 
-	// Store result in cache
-	// TODO - There is a potential mismatch here between using the Settings (incl. overrides) and the input settings interface (pre-overrides), as done in the graph executor.
-	// TODO - The dependencies CRC here should always be valid except in the indirection case, which we should normalize to allow caching (tested here, otherwise it can ensure in the graph cache)
-	if (!Context->OutputData.bCancelExecution && !bHasErrorsOrWarnings && Context->DependenciesCrc.IsValid() && IsCacheableInstance(Settings))
 	{
-		bool bCacheable = true;
-
-		// GPU proxies are never cached. Caching them would hold on to precious graphics memory.
-		for (const FPCGTaggedData& TaggedData : Context->OutputData.TaggedData)
+		TRACE_CPUPROFILER_EVENT_SCOPE(IPCGElement::PostExecute::StoreInCache);
+		// Store result in cache
+		// TODO - There is a potential mismatch here between using the Settings (incl. overrides) and the input settings interface (pre-overrides), as done in the graph executor.
+		// TODO - The dependencies CRC here should always be valid except in the indirection case, which we should normalize to allow caching (tested here, otherwise it can ensure in the graph cache)
+		if (!Context->OutputData.bCancelExecution && !bHasErrorsOrWarnings && Context->DependenciesCrc.IsValid() && IsCacheableInstance(Settings))
 		{
-			if (TaggedData.Data && !TaggedData.Data->IsCacheable())
+			bool bCacheable = true;
+
+			// GPU proxies are never cached. Caching them would hold on to precious graphics memory.
+			for (const FPCGTaggedData& TaggedData : Context->OutputData.TaggedData)
 			{
-				bCacheable = false;
-				break;
+				if (TaggedData.Data && !TaggedData.Data->IsCacheable())
+				{
+					bCacheable = false;
+					break;
+				}
 			}
-		}
 
-		if (bCacheable)
-		{
-			UPCGSubsystem* Subsystem = Context->SourceComponent.Get() ? Context->SourceComponent->GetSubsystem() : nullptr;
-			if (Subsystem)
+			if (bCacheable)
 			{
-				FPCGStoreInCacheParams Params = { .Element = this, .Crc = Context->DependenciesCrc };
-				Subsystem->GetCache()->StoreInCache(Params, Context->OutputData);
+				UPCGSubsystem* Subsystem = Context->SourceComponent.Get() ? Context->SourceComponent->GetSubsystem() : nullptr;
+				if (Subsystem)
+				{
+					FPCGStoreInCacheParams Params = { .Element = this, .Crc = Context->DependenciesCrc };
+					Subsystem->GetCache()->StoreInCache(Params, Context->OutputData);
+				}
 			}
 		}
 	}
@@ -570,6 +580,7 @@ void IPCGElement::PostExecute(FPCGContext* Context) const
 	// Register the element to the component indicating the element has run and can have dynamic tracked keys.
 	if (Settings && Settings->CanDynamicallyTrackKeys() && Context->SourceComponent.IsValid())
 	{
+		TRACE_CPUPROFILER_EVENT_SCOPE(IPCGElement::PostExecute::RegisterDynamicTracking);
 		Context->SourceComponent->RegisterDynamicTracking(Settings, {});
 	}
 #endif // WITH_EDITOR
