[PCG] Temporary fix for Dynamic tracking and Runtime HiGen

* Everytime a local component managed by runtime is done processing, it will push its tracking keys to the original which will update the tracking, only with the changing keys
* This is suboptimal as when something change, the whole original and locals will be regenerated
* It's behind a CVar if it is ever an issue with performance.

#rnx

diff --git a/Engine/Plugins/PCG/Source/PCG/Private/PCGComponent.cpp b/Engine/Plugins/PCG/Source/PCG/Private/PCGComponent.cpp
--- a/Engine/Plugins/PCG/Source/PCG/Private/PCGComponent.cpp
+++ b/Engine/Plugins/PCG/Source/PCG/Private/PCGComponent.cpp
@@ -70,6 +70,12 @@ namespace PCGComponent
 		TEXT("pcg.LoadManagedActorsForCleanup"),
 		true,
 		TEXT("When enabled unloaded actors will get loaded before clean up"));
+
+	/** Have a CVar for dynamic tracking for runtime until we have a better solution to also have "culling" for local components and dynamic tracking. */
+	static TAutoConsoleVariable<bool> CVarDisableDynamicTrackingForRuntimeGen(
+		TEXT("pcg.RuntimeGeneration.DisableDynamicTracking"),
+		false,
+		TEXT("In Editor and with runtime gen, a change with one tracked element will refresh all the local components. If it is too resource intensive, it can be disabled."));
 #endif
 
 	struct FPCGManagedActorLoadingScope
@@ -3417,31 +3423,46 @@ void UPCGComponent::UpdateDynamicTracking()
 		return;
 	}
 
-	// If the component is local, we defer the tracking to the original component.
-	// So move everything to the original (while making sure we are not duplicating keys/settings).
-	// Since it can happen in parallel, we need to lock.
-	if (IsLocalComponent())
+	UPCGComponent* OriginalComponent = GetOriginalComponent();
+	if (!OriginalComponent)
 	{
-		TRACE_CPUPROFILER_EVENT_SCOPE(UPCGComponent::UpdateDynamicTracking::LocalComponent);
+		return;
+	}
 
-		UPCGComponent* OriginalComponent = GetOriginalComponent();
-		if (ensure(OriginalComponent))
+	const bool bIsLocal = IsLocalComponent();
+	const bool bIsRuntime = IsManagedByRuntimeGenSystem();
+
+	if (bIsRuntime && PCGComponent::CVarDisableDynamicTrackingForRuntimeGen.GetValueOnAnyThread())
+	{
+		return;
+	}
+
+	auto PushToOriginal = [](FPCGSelectionKeyToSettingsMap& LocalDynamicTrackingMap, FPCGSelectionKeyToSettingsMap& OriginalDynamicTrackingMap)
+	{
+		for (auto& It : LocalDynamicTrackingMap)
 		{
-			FScopeLock Lock(&OriginalComponent->CurrentExecutionDynamicTrackingLock);
-			for (auto& It : CurrentExecutionDynamicTracking)
+			TArray<FPCGSettingsAndCulling>& OriginalSettingsAndCulling = OriginalDynamicTrackingMap.FindOrAdd(It.Key);
+			for (FPCGSettingsAndCulling& SettingsAndCulling : It.Value)
 			{
-				TArray<FPCGSettingsAndCulling>& OriginalSettingsAndCulling = OriginalComponent->CurrentExecutionDynamicTracking.FindOrAdd(It.Key);
-				for (FPCGSettingsAndCulling& SettingsAndCulling : It.Value)
-				{
-					OriginalSettingsAndCulling.AddUnique(std::move(SettingsAndCulling));
-				}
+				OriginalSettingsAndCulling.AddUnique(std::move(SettingsAndCulling));
 			}
-
-			OriginalComponent->CurrentExecutionDynamicTrackingSettings.Append(CurrentExecutionDynamicTrackingSettings);
 		}
+	};
+	
+	// If the component is local, we defer the tracking to the original component.
+	// So move everything to the original (while making sure we are not duplicating keys/settings).
+	// Since it can happen in parallel, we need to lock.
+	if (bIsLocal && !bIsRuntime)
+	{
+		TRACE_CPUPROFILER_EVENT_SCOPE(UPCGComponent::UpdateDynamicTracking::LocalComponent);
+		FScopeLock Lock(&OriginalComponent->CurrentExecutionDynamicTrackingLock);
+
+		PushToOriginal(CurrentExecutionDynamicTracking, OriginalComponent->CurrentExecutionDynamicTracking);
+		OriginalComponent->CurrentExecutionDynamicTrackingSettings.Append(CurrentExecutionDynamicTrackingSettings);
 
 		CurrentExecutionDynamicTracking.Empty();
 		CurrentExecutionDynamicTrackingSettings.Empty();
+
 		return;
 	}
 
@@ -3511,7 +3532,19 @@ void UPCGComponent::UpdateDynamicTracking()
 
 	if (!ChangedKeys.IsEmpty())
 	{
-		Subsystem->UpdateComponentTracking(this, /*bShouldDirtyActors=*/false, &ChangedKeys);
+		if (bIsLocal)
+		{
+			OriginalComponent->CurrentExecutionDynamicTrackingLock.Lock();
+
+			PushToOriginal(DynamicallyTrackedKeysToSettings, OriginalComponent->DynamicallyTrackedKeysToSettings);
+		}
+		
+		Subsystem->UpdateComponentTracking(OriginalComponent, /*bShouldDirtyActors=*/false, &ChangedKeys);
+
+		if (bIsLocal)
+		{
+			OriginalComponent->CurrentExecutionDynamicTrackingLock.Unlock();
+		}
 	}
 }
 
