[PCG] Make column sizes persistent between refreshes in the attribute list view

#rnx

diff --git a/Engine/Plugins/PCG/Source/PCGEditor/Private/SPCGEditorGraphAttributeListView.cpp b/Engine/Plugins/PCG/Source/PCGEditor/Private/SPCGEditorGraphAttributeListView.cpp
--- a/Engine/Plugins/PCG/Source/PCGEditor/Private/SPCGEditorGraphAttributeListView.cpp
+++ b/Engine/Plugins/PCG/Source/PCGEditor/Private/SPCGEditorGraphAttributeListView.cpp
@@ -683,6 +683,8 @@ const FPCGDataCollection* SPCGEditorGraphAttributeListView::GetInspectionData()
 
 void SPCGEditorGraphAttributeListView::RefreshAttributeList()
 {
+	CacheColumnWidthVisibility();
+
 	HiddenAttributes = ListViewHeader->GetHiddenColumnIds();
 
 	// Swapping to an empty item list to force a widget clear, otherwise the widgets will try to update during add column and access invalid data
@@ -1057,7 +1059,6 @@ void SPCGEditorGraphAttributeListView::OnSelectionChangedPin(TSharedPtr<FPinComb
 
 	if (InSelectInfo != ESelectInfo::Direct)
 	{
-		CacheColumnWidthVisibility();
 		RefreshAttributeList();
 	}
 }
@@ -1076,7 +1077,6 @@ void SPCGEditorGraphAttributeListView::OnSelectionChanged(TSharedPtr<FName> Item
 {
 	if (SelectInfo != ESelectInfo::Direct)
 	{
-		CacheColumnWidthVisibility();
 		RefreshAttributeList();
 	}
 }
