[PCG] Set/Change GridSize nodes presets

diff --git a/Engine/Plugins/PCG/Source/PCG/Private/Elements/PCGHiGenGridSize.cpp b/Engine/Plugins/PCG/Source/PCG/Private/Elements/PCGHiGenGridSize.cpp
--- a/Engine/Plugins/PCG/Source/PCG/Private/Elements/PCGHiGenGridSize.cpp
+++ b/Engine/Plugins/PCG/Source/PCG/Private/Elements/PCGHiGenGridSize.cpp
@@ -8,14 +8,20 @@
 #include "PCGModule.h"
 #include "PCGPin.h"
 #include "Data/PCGVolumeData.h"
+#include "Helpers/PCGConversion.h"
 
 #define LOCTEXT_NAMESPACE "PCGHiGenGridSizeElement"
 
 namespace PCGHiGenGridSizeConstants
 {
 	const FName NodeName = TEXT("HiGenGridSize");
-	const FText NodeTitle = LOCTEXT("NodeTitle", "Grid Size");
 	const FName CellVolumeOutputLabel = TEXT("Grid Cell Volume");
+
+	const FText SetSizeNodeTitle = LOCTEXT("SetSizeNodeTitle", "Set Grid Size");
+	const FText ChangeSizeNodeTitle = LOCTEXT("ChangeSizeNodeTitle", "Change Grid Size");
+
+	const int32 SetGridSizePreconfiguredIndex = 0;
+	const int32 ChangeGridSizePreconfiguredIndex = 1;
 }
 
 EPCGHiGenGrid UPCGHiGenGridSizeSettings::GetGrid() const
@@ -44,7 +50,7 @@ FName UPCGHiGenGridSizeSettings::GetDefaultNodeName() const
 
 FText UPCGHiGenGridSizeSettings::GetDefaultNodeTitle() const
 {
-	return PCGHiGenGridSizeConstants::NodeTitle;
+	return bShowInputPin ? PCGHiGenGridSizeConstants::ChangeSizeNodeTitle : PCGHiGenGridSizeConstants::SetSizeNodeTitle;
 }
 
 FText UPCGHiGenGridSizeSettings::GetNodeTooltipText() const
@@ -55,6 +61,15 @@ FText UPCGHiGenGridSizeSettings::GetNodeTooltipText() const
 		"\n\t* Hierarchical Generation is disabled in the graph settings."
 		"\n\t* Executed in a subgraph, as subgraphs are always invoked on parent grid level.");
 }
+
+TArray<FPCGPreConfiguredSettingsInfo> UPCGHiGenGridSizeSettings::GetPreconfiguredInfo() const
+{
+	TArray<FPCGPreConfiguredSettingsInfo> Infos;
+	Infos.Emplace(PCGHiGenGridSizeConstants::SetGridSizePreconfiguredIndex, PCGHiGenGridSizeConstants::SetSizeNodeTitle);
+	Infos.Emplace(PCGHiGenGridSizeConstants::ChangeGridSizePreconfiguredIndex, PCGHiGenGridSizeConstants::ChangeSizeNodeTitle);
+
+	return Infos;
+}
 #endif
 
 EPCGDataType UPCGHiGenGridSizeSettings::GetCurrentPinTypes(const UPCGPin* InPin) const
@@ -74,8 +89,12 @@ EPCGDataType UPCGHiGenGridSizeSettings::GetCurrentPinTypes(const UPCGPin* InPin)
 TArray<FPCGPinProperties> UPCGHiGenGridSizeSettings::InputPinProperties() const
 {
 	TArray<FPCGPinProperties> PinProperties;
-	FPCGPinProperties& InputPinProperty = PinProperties.Emplace_GetRef(PCGPinConstants::DefaultInputLabel, EPCGDataType::Any);
-	InputPinProperty.SetRequiredPin();
+
+	if (bShowInputPin)
+	{
+		FPCGPinProperties& InputPinProperty = PinProperties.Emplace_GetRef(PCGPinConstants::DefaultInputLabel, EPCGDataType::Any);
+		InputPinProperty.SetRequiredPin();
+	}
 
 	return PinProperties;
 }
@@ -83,7 +102,17 @@ TArray<FPCGPinProperties> UPCGHiGenGridSizeSettings::InputPinProperties() const
 TArray<FPCGPinProperties> UPCGHiGenGridSizeSettings::OutputPinProperties() const
 {
 	TArray<FPCGPinProperties> PinProperties;
-	PinProperties.Emplace(PCGPinConstants::DefaultOutputLabel, EPCGDataType::Any);
+
+	if (bShowInputPin)
+	{
+		PinProperties.Emplace(PCGPinConstants::DefaultOutputLabel, EPCGDataType::Any);
+	}
+	else
+	{
+		FPCGPinProperties& DepPin = PinProperties.Emplace_GetRef(PCGPinConstants::DefaultExecutionDependencyLabel, EPCGDataType::Any, /*bInAllowMultipleConnections=*/true, /*bAllowMultipleData=*/true);
+		DepPin.Usage = EPCGPinUsage::DependencyOnly;
+	}
+
 	PinProperties.Emplace(
 		PCGHiGenGridSizeConstants::CellVolumeOutputLabel,
 		EPCGDataType::Spatial,
@@ -117,6 +146,47 @@ FString UPCGHiGenGridSizeSettings::GetAdditionalTitleInformation() const
 	}
 }
 
+void UPCGHiGenGridSizeSettings::ApplyPreconfiguredSettings(const FPCGPreConfiguredSettingsInfo& PreconfiguredInfo)
+{
+	bShowInputPin = PreconfiguredInfo.PreconfiguredIndex == PCGHiGenGridSizeConstants::ChangeGridSizePreconfiguredIndex;
+}
+
+#if WITH_EDITOR
+TArray<FPCGPreconfiguredInfo> UPCGHiGenGridSizeSettings::GetConversionInfo() const
+{
+	TArray<FPCGPreconfiguredInfo> Infos;
+
+	if (bShowInputPin)
+	{
+		Infos.Emplace(PCGHiGenGridSizeConstants::SetGridSizePreconfiguredIndex, PCGHiGenGridSizeConstants::SetSizeNodeTitle);
+	}
+	else
+	{
+		Infos.Emplace(PCGHiGenGridSizeConstants::ChangeGridSizePreconfiguredIndex, PCGHiGenGridSizeConstants::ChangeSizeNodeTitle);
+	}
+
+	return Infos;
+}
+#endif // WITH_EDITOR
+
+bool UPCGHiGenGridSizeSettings::ConvertNode(const FPCGPreconfiguredInfo& ConversionInfo)
+{
+#if WITH_EDITOR
+	const bool bWasShowingInputPin = bShowInputPin;
+#endif
+	bShowInputPin = ConversionInfo.PreconfiguredIndex == PCGHiGenGridSizeConstants::ChangeGridSizePreconfiguredIndex;
+
+#if WITH_EDITOR
+	// Only broadcast an update if bShowInputPin actually changed.
+	if (bWasShowingInputPin != bShowInputPin)
+	{
+		OnSettingsChangedDelegate.Broadcast(this, EPCGChangeType::Settings | EPCGChangeType::Node | EPCGChangeType::GenerationGrid);
+	}
+#endif
+
+	return true;
+}
+
 #if WITH_EDITOR
 EPCGChangeType UPCGHiGenGridSizeSettings::GetChangeTypeForProperty(const FName& InPropertyName) const
 {
diff --git a/Engine/Plugins/PCG/Source/PCG/Public/Elements/PCGHiGenGridSize.h b/Engine/Plugins/PCG/Source/PCG/Public/Elements/PCGHiGenGridSize.h
--- a/Engine/Plugins/PCG/Source/PCG/Public/Elements/PCGHiGenGridSize.h
+++ b/Engine/Plugins/PCG/Source/PCG/Public/Elements/PCGHiGenGridSize.h
@@ -24,17 +24,28 @@ public:
 	virtual FText GetDefaultNodeTitle() const override;
 	virtual FText GetNodeTooltipText() const override;
 	virtual EPCGSettingsType GetType() const override { return EPCGSettingsType::HierarchicalGeneration; }
+	virtual TArray<FPCGPreConfiguredSettingsInfo> GetPreconfiguredInfo() const override;
+	virtual bool OnlyExposePreconfiguredSettings() const override { return true; }
 #endif
 	virtual FString GetAdditionalTitleInformation() const override;
+	virtual void ApplyPreconfiguredSettings(const FPCGPreConfiguredSettingsInfo& PreconfigureInfo) override;
 	virtual bool HasDynamicPins() const override { return true; }
 	virtual EPCGDataType GetCurrentPinTypes(const UPCGPin* InPin) const override;
 
 	UPROPERTY(BlueprintReadWrite, EditAnywhere, Category = Settings, meta = (DisplayName = "HiGen Grid Size"))
 	EPCGHiGenGrid HiGenGridSize = EPCGHiGenGrid::Grid256;
+
+	/** This property is hidden from the user and drives the behavior of the Grid Size node. */
+	UPROPERTY()
+	bool bShowInputPin = true;
+
 protected:
 #if WITH_EDITOR
 	virtual EPCGChangeType GetChangeTypeForProperty(const FName& InPropertyName) const override;
+	virtual TArray<FPCGPreconfiguredInfo> GetConversionInfo() const override;
 #endif
+	virtual bool HasExecutionDependencyPin() const override { return bShowInputPin; }
+	virtual bool ConvertNode(const FPCGPreconfiguredInfo& ConversionInfo) override;
 	virtual TArray<FPCGPinProperties> InputPinProperties() const override;
 	virtual TArray<FPCGPinProperties> OutputPinProperties() const override;
 	virtual FPCGElementPtr CreateElement() const override;
