[PCG] Fix grass maps failing to generate consistently over time.
Cache CRC did not include generation volume, so grass maps would generate only for the landscapes overlapping the first PCG generation volume and then incorrectly reuse the cached result for subsequent generations.
#rnx

diff --git a/Engine/Plugins/PCG/Source/PCG/Private/Elements/PCGGenerateGrassMaps.cpp b/Engine/Plugins/PCG/Source/PCG/Private/Elements/PCGGenerateGrassMaps.cpp
--- a/Engine/Plugins/PCG/Source/PCG/Private/Elements/PCGGenerateGrassMaps.cpp
+++ b/Engine/Plugins/PCG/Source/PCG/Private/Elements/PCGGenerateGrassMaps.cpp
@@ -95,6 +95,20 @@ void FPCGGenerateGrassMapsContext::AddExtraStructReferencedObjects(FReferenceCol
 	}
 }
 
+void FPCGGenerateGrassMapsElement::GetDependenciesCrc(const FPCGDataCollection& InInput, const UPCGSettings* InSettings, UPCGComponent* InComponent, FPCGCrc& OutCrc) const
+{
+	FPCGCrc Crc;
+	IPCGElement::GetDependenciesCrc(InInput, InSettings, InComponent, Crc);
+
+	// todo_pcg: Technically this could be fancier and hash the set of landscape components that overlap with the generation volume.
+	if (const UPCGData* Data = InComponent->GetActorPCGData())
+	{
+		Crc.Combine(Data->GetOrComputeCrc(/*bFullDataCrc=*/false));
+	}
+
+	OutCrc = Crc;
+}
+
 FPCGContext* FPCGGenerateGrassMapsElement::CreateContext()
 {
 	return new FPCGGenerateGrassMapsContext();
diff --git a/Engine/Plugins/PCG/Source/PCG/Public/Elements/PCGGenerateGrassMaps.h b/Engine/Plugins/PCG/Source/PCG/Public/Elements/PCGGenerateGrassMaps.h
--- a/Engine/Plugins/PCG/Source/PCG/Public/Elements/PCGGenerateGrassMaps.h
+++ b/Engine/Plugins/PCG/Source/PCG/Public/Elements/PCGGenerateGrassMaps.h
@@ -78,6 +78,7 @@ protected:
 class FPCGGenerateGrassMapsElement : public IPCGElement
 {
 public:
+	virtual void GetDependenciesCrc(const FPCGDataCollection& InInput, const UPCGSettings* InSettings, UPCGComponent* InComponent, FPCGCrc& OutCrc) const override;
 	/** FLandscapeGrassWeightExporter expects to exist in scope only on the game thread. */
 	virtual bool CanExecuteOnlyOnMainThread(FPCGContext* Context) const { return true; }
 
