[PCG] GraphExecutor: Fix gap that could cause all task lists to be empty and ClearResults could be called clearing the outputs to early

#rnx

diff --git a/Engine/Plugins/PCG/Source/PCG/Private/Graph/PCGGraphExecutor.cpp b/Engine/Plugins/PCG/Source/PCG/Private/Graph/PCGGraphExecutor.cpp
--- a/Engine/Plugins/PCG/Source/PCG/Private/Graph/PCGGraphExecutor.cpp
+++ b/Engine/Plugins/PCG/Source/PCG/Private/Graph/PCGGraphExecutor.cpp
@@ -627,7 +627,7 @@ bool FPCGGraphExecutor::IsAnyGraphCurrentlyExecuting() const
 
 int32 FPCGGraphExecutor::GetNonScheduledRemainingTaskCount() const
 {
-	return Tasks.Num() + ReadyTasks.Num() + ActiveTasks.Num() + ActiveTasksGameThreadOnly.Num() + SleepingTasks.Num();
+	return NewReadyTaskCount + Tasks.Num() + ReadyTasks.Num() + ActiveTasks.Num() + ActiveTasksGameThreadOnly.Num() + SleepingTasks.Num();
 }
 
 UPCGData* FPCGGraphExecutor::GetExecutionCacheData(FPCGTaskId InGraphExecutionTaskId, EExecutionCacheDataType InExecutionCacheDataType)
@@ -833,6 +833,7 @@ bool FPCGGraphExecutor::ProcessScheduledTasks()
 						if (bPushToReady)
 						{
 							NewReadyTasks.Add(MoveTemp(Task));
+							++NewReadyTaskCount;
 						}
 						else
 						{
@@ -846,7 +847,7 @@ bool FPCGGraphExecutor::ProcessScheduledTasks()
 				{
 					PCGGraphExecutionLogging::LogGraphPostSchedule(Tasks, TaskSuccessors);
 				}
-				
+
 				ScheduledTasks.Reset();
 			}
 
@@ -854,6 +855,7 @@ bool FPCGGraphExecutor::ProcessScheduledTasks()
 			for (FPCGGraphTask& ReadyTask : NewReadyTasks)
 			{
 				OnTaskInputsReady(ReadyTask, CachedResults, bIsInGameThread);
+				--NewReadyTaskCount;
 			}
 		}
 	}
@@ -1678,6 +1680,7 @@ TArray<FPCGGraphExecutor::FCachedResult*> FPCGGraphExecutor::QueueNextTasksInter
 					if (bAllPrerequisitesMet)
 					{
 						NewReadyTasks.Add(MoveTemp(SuccessorTask));
+						++NewReadyTaskCount;
 						Tasks.Remove(Successor);
 					}
 				}
@@ -1688,6 +1691,7 @@ TArray<FPCGGraphExecutor::FCachedResult*> FPCGGraphExecutor::QueueNextTasksInter
 	for (FPCGGraphTask& NewReadyTask : NewReadyTasks)
 	{
 		OnTaskInputsReady(NewReadyTask, CachedResults, bIsInGameThread);
+		--NewReadyTaskCount;
 	}
 	
 	return CachedResults;
diff --git a/Engine/Plugins/PCG/Source/PCG/Private/Graph/PCGGraphExecutor.h b/Engine/Plugins/PCG/Source/PCG/Private/Graph/PCGGraphExecutor.h
--- a/Engine/Plugins/PCG/Source/PCG/Private/Graph/PCGGraphExecutor.h
+++ b/Engine/Plugins/PCG/Source/PCG/Private/Graph/PCGGraphExecutor.h
@@ -350,6 +350,7 @@ private:
 	TMap<FPCGTaskId, FGraphExecutionCache> GraphExecutionCaches;
 
 	std::atomic<bool> bNeedToExecuteTasksEnded = false;
+	std::atomic<int32> NewReadyTaskCount = 0;
 
 #if WITH_EDITOR
 	TWeakPtr<IPCGEditorProgressNotification> GenerationProgressNotification;
