[PCG] Allow Preview PCG Components to refresh when tracking actors inside Editing Level Instances + Fix Level Instance cancelled edit not refreshing PCG Components properly

#rnx

diff --git a/Engine/Plugins/PCG/Source/PCG/Private/PCGActorAndComponentMapping.cpp b/Engine/Plugins/PCG/Source/PCG/Private/PCGActorAndComponentMapping.cpp
--- a/Engine/Plugins/PCG/Source/PCG/Private/PCGActorAndComponentMapping.cpp
+++ b/Engine/Plugins/PCG/Source/PCG/Private/PCGActorAndComponentMapping.cpp
@@ -1254,6 +1254,7 @@ void FPCGActorAndComponentMapping::RegisterDelegates()
 		if (ULevelInstanceSubsystem* LevelInstanceSubsystem = UWorld::GetSubsystem<ULevelInstanceSubsystem>(World))
 		{
 			LevelInstanceSubsystem->OnLevelInstancesUpdated().AddRaw(this, &FPCGActorAndComponentMapping::OnLevelInstancesUpdated);
+			LevelInstanceSubsystem->OnLevelInstanceEditCancelled().AddRaw(this, &FPCGActorAndComponentMapping::OnLevelInstanceEditCancelled);
 		}
 	}
 
@@ -1297,6 +1298,7 @@ void FPCGActorAndComponentMapping::UnregisterDelegates()
 		if (ULevelInstanceSubsystem* LevelInstanceSubsystem = UWorld::GetSubsystem<ULevelInstanceSubsystem>(World))
 		{
 			LevelInstanceSubsystem->OnLevelInstancesUpdated().RemoveAll(this);
+			LevelInstanceSubsystem->OnLevelInstanceEditCancelled().RemoveAll(this);
 		}
 	}
 }
@@ -1365,6 +1367,11 @@ void FPCGActorAndComponentMapping::OnLevelInstancesUpdated(const TArray<ILevelIn
 	}
 }
 
+void FPCGActorAndComponentMapping::OnLevelInstanceEditCancelled(ILevelInstanceInterface* InLevelInstance)
+{
+	OnActorAdded(Cast<AActor>(InLevelInstance));
+}
+
 void FPCGActorAndComponentMapping::OnActorLoaded(AActor& InActor)
 {
 	// We have to make sure to not create a infinite loop
@@ -2088,8 +2095,8 @@ void FPCGActorAndComponentMapping::OnObjectChanged(UObject* InObject, const FAct
 					{
 						continue;
 					}
-				}
-				else if (Actor->IsInEditLevelInstanceHierarchy())
+				} // Allow update of Component from a Editing Level instance contained actor if the Component being refreshed is in preview mode and not itself in a Level Instance
+				else if (Actor->IsInEditLevelInstanceHierarchy() && (Component->GetEditingMode() != EPCGEditorDirtyMode::Preview || ComponentLevelInstance != nullptr))
 				{
 					continue;
 				}
diff --git a/Engine/Plugins/PCG/Source/PCG/Public/PCGActorAndComponentMapping.h b/Engine/Plugins/PCG/Source/PCG/Public/PCGActorAndComponentMapping.h
--- a/Engine/Plugins/PCG/Source/PCG/Public/PCGActorAndComponentMapping.h
+++ b/Engine/Plugins/PCG/Source/PCG/Public/PCGActorAndComponentMapping.h
@@ -186,6 +186,7 @@ private:
 	void OnActorDeleted_Internal(AActor* InActor, bool bShouldDirty);
 	void OnLandscapeChanged(ALandscapeProxy* InLandscape, const FLandscapeProxyComponentDataChangedParams& InChangeParams);
 	void OnLevelInstancesUpdated(const TArray<ILevelInstanceInterface*>& InLevelInstances);
+	void OnLevelInstanceEditCancelled(ILevelInstanceInterface* InLevelInstance);
 	void ApplyLandscapeChanges(ALandscapeProxy* InLandscape);
 	void OnObjectModified(UObject* InObject);
 	void OnObjectPropertyChanged(UObject* InObject, FPropertyChangedEvent& InEvent);
diff --git a/Engine/Source/Runtime/Engine/Private/LevelInstance/LevelInstanceSubsystem.cpp b/Engine/Source/Runtime/Engine/Private/LevelInstance/LevelInstanceSubsystem.cpp
--- a/Engine/Source/Runtime/Engine/Private/LevelInstance/LevelInstanceSubsystem.cpp
+++ b/Engine/Source/Runtime/Engine/Private/LevelInstance/LevelInstanceSubsystem.cpp
@@ -2766,6 +2766,11 @@ bool ULevelInstanceSubsystem::CommitLevelInstanceInternal(TUniquePtr<FLevelInsta
 			}
 		}
 	}
+	else
+	{
+		LevelInstance = GetLevelInstance(LevelInstanceID);
+		LevelInstanceEditCancelled.Broadcast(LevelInstance);
+	}
 
 	GEngine->BroadcastLevelActorListChanged();
 
diff --git a/Engine/Source/Runtime/Engine/Public/LevelInstance/LevelInstanceSubsystem.h b/Engine/Source/Runtime/Engine/Public/LevelInstance/LevelInstanceSubsystem.h
--- a/Engine/Source/Runtime/Engine/Public/LevelInstance/LevelInstanceSubsystem.h
+++ b/Engine/Source/Runtime/Engine/Public/LevelInstance/LevelInstanceSubsystem.h
@@ -169,6 +169,9 @@ public:
 	DECLARE_EVENT_OneParam(ULevelInstanceSubsystem, FLevelInstancesUpdated, const TArray<ILevelInstanceInterface*>&);
 	FLevelInstancesUpdated& OnLevelInstancesUpdated() { return LevelInstancesUpdatedEvent; }
 
+	DECLARE_EVENT_OneParam(ULevelInstanceSubsystem, FLevelInstanceEditCancelled, ILevelInstanceInterface*);
+	FLevelInstanceEditCancelled& OnLevelInstanceEditCancelled() { return LevelInstanceEditCancelled; }
+
 	static ENGINE_API void ResetLoadersForWorldAsset(const FString& WorldAsset);
 
 	ENGINE_API bool PassLevelInstanceFilter(UWorld* World, const FWorldPartitionHandle& Actor) const;
@@ -324,6 +327,7 @@ private:
 
 	FLevelInstanceChanged LevelInstanceChangedEvent;
 	FLevelInstancesUpdated LevelInstancesUpdatedEvent;
+	FLevelInstanceEditCancelled LevelInstanceEditCancelled;
 
 	static bool bPrimitiveColorHandlerRegistered;
 #endif
