[PCG] Make sure to copy properties from original in Refresh Component in the runtime gen scheduler

* Otherwise, any change on the original is not propagated to existing local components

#rnx

diff --git a/Engine/Plugins/PCG/Source/PCG/Private/RuntimeGen/PCGRuntimeGenScheduler.cpp b/Engine/Plugins/PCG/Source/PCG/Private/RuntimeGen/PCGRuntimeGenScheduler.cpp
--- a/Engine/Plugins/PCG/Source/PCG/Private/RuntimeGen/PCGRuntimeGenScheduler.cpp
+++ b/Engine/Plugins/PCG/Source/PCG/Private/RuntimeGen/PCGRuntimeGenScheduler.cpp
@@ -1087,6 +1087,9 @@ void FPCGRuntimeGenScheduler::RefreshComponent(UPCGComponent* InComponent, bool
 					if (LocalComponent)
 					{
 						LocalComponent->CleanupLocalImmediate(/*bRemoveComponents=*/false);
+
+						// We need to make sure that the next time this is generated that it matches the original
+						LocalComponent->SetPropertiesFromOriginal(OriginalComponent);
 					}
 
 					GeneratedComponentsToRemove.Add(Key);
