PCG: Fix crash in slate if a non-game-thread executor triggers a node update.

diff --git a/Engine/Plugins/PCG/Source/PCGEditor/Private/Nodes/SPCGEditorGraphNode.cpp b/Engine/Plugins/PCG/Source/PCGEditor/Private/Nodes/SPCGEditorGraphNode.cpp
--- a/Engine/Plugins/PCG/Source/PCGEditor/Private/Nodes/SPCGEditorGraphNode.cpp
+++ b/Engine/Plugins/PCG/Source/PCGEditor/Private/Nodes/SPCGEditorGraphNode.cpp
@@ -780,8 +780,17 @@ void SPCGEditorGraphNode::GetOverlayBrushes(bool bSelected, const FVector2D Widg
 
 void SPCGEditorGraphNode::OnNodeChanged()
 {
-	UpdateGraphNode();
-}
+	// Avoid crashing inside slate if we got triggered from a non-game-thread via any
+	//  experimental worker-thread executor
+	// @todo_pcg: revisit
+	if (IsInGameThread()) 
+	{
+		UpdateGraphNode(); 
+	}
+	else
+	{
+		ExecuteOnGameThread(UE_SOURCE_LOCATION, [this]() { UpdateGraphNode(); });
+	}}
 
 void SPCGEditorGraphNode::OnNodeRenameInitiated()
 {
