[PCG] LevelToAsset improvements
- Fix Actor Attachement information for actors inside Level Instances
- Propagate Level Instance Actor Tags to child actors

#rnx

diff --git a/Engine/Plugins/PCG/Source/PCGEditor/Private/PCGLevelToAsset.cpp b/Engine/Plugins/PCG/Source/PCGEditor/Private/PCGLevelToAsset.cpp
--- a/Engine/Plugins/PCG/Source/PCGEditor/Private/PCGLevelToAsset.cpp
+++ b/Engine/Plugins/PCG/Source/PCGEditor/Private/PCGLevelToAsset.cpp
@@ -20,11 +20,16 @@
 #include "AssetRegistry/AssetRegistryModule.h"
 #include "Components/InstancedStaticMeshComponent.h"
 #include "Engine/LevelStreaming.h"
+#include "LevelInstance/LevelInstanceEditorInstanceActor.h"
+#include "LevelInstance/LevelInstanceInterface.h"
+#include "LevelInstance/LevelInstanceSubsystem.h"
 #include "UObject/UObjectGlobals.h"
 #include "UObject/Package.h"
 #include "WorldPartition/WorldPartition.h"
 #include "WorldPartition/WorldPartitionHelpers.h"
 
+extern PCG_API TAutoConsoleVariable<bool> CVarPCGEnablePointArrayData;
+
 void UPCGLevelToAsset::CreateOrUpdatePCGAssets(const TArray<FAssetData>& WorldAssets, const FPCGAssetExporterParameters& InParameters, TSubclassOf<UPCGLevelToAsset> ExporterSubclass)
 {
 	TArray<UPackage*> PackagesToSave;
@@ -144,8 +149,11 @@ bool UPCGLevelToAsset::BP_ExportWorld_Implementation(UWorld* World, const FStrin
 	FPCGDataCollection& DataCollection = Asset->Data;
 	DataCollection.TaggedData.Reset();
 
+	// Select proper point data class
+	UClass* PointDataClass = CVarPCGEnablePointArrayData.GetValueOnAnyThread() ? UPCGPointArrayData::StaticClass() : UPCGPointData::StaticClass();
+
 	// Create Root Data
-	UPCGPointData* RootPointData = NewObject<UPCGPointData>(Asset);
+	UPCGBasePointData* RootPointData = NewObject<UPCGBasePointData>(Asset, PointDataClass);
 	UPCGMetadata* RootMetadata = RootPointData->MutableMetadata();
 
 	RootMetadata->CreateAttribute<FString>(TEXT("Name"), World->GetName(), /*bAllowsInterpolation=*/false, /*bOverrideParent=*/true);
@@ -159,7 +167,7 @@ bool UPCGLevelToAsset::BP_ExportWorld_Implementation(UWorld* World, const FStrin
 	}
 
 	// Create PCGData
-	UPCGPointData* PointData = NewObject<UPCGPointData>(Asset);
+	UPCGBasePointData* PointData = NewObject<UPCGBasePointData>(Asset, PointDataClass);
 	UPCGMetadata* PointMetadata = PointData->MutableMetadata();
 
 	// Add to data collection
@@ -266,12 +274,65 @@ bool UPCGLevelToAsset::BP_ExportWorld_Implementation(UWorld* World, const FStrin
 		}
 	}
 	
+	struct FActorEntry
+	{
+		FActorEntry(int InIndex) : Index(InIndex) {}
+		int Index = 0;
+		AActor* AttachmentParent = nullptr;
+		TArray<FName> Tags;
+	};
+
+	ULevelInstanceSubsystem* LevelInstanceSubsystem = UWorld::GetSubsystem<ULevelInstanceSubsystem>(World);
+	check(LevelInstanceSubsystem);
+
+	TFunction<AActor*(AActor*)> GetAttachParentActor = [&GetAttachParentActor, LevelInstanceSubsystem](AActor* InActor)
+	{
+		if(InActor->IsInLevelInstance())
+		{
+			AActor* AttachParentActor = InActor->GetAttachParentActor();
+			if (AttachParentActor && AttachParentActor->IsA<ALevelInstanceEditorInstanceActor>())
+			{
+				return GetAttachParentActor(Cast<AActor>(LevelInstanceSubsystem->GetParentLevelInstance(InActor)));
+			}
+			else
+			{
+				return AttachParentActor;
+			}
+		}
+		else
+		{
+			return InActor->GetAttachParentActor();
+		}
+	};
+
+	// Parent tags are meant to be at the end of the array so they are applied last in the MakePoint lambda (we want the parent tag values to override the childrens)
+	auto GetActorTags = [LevelInstanceSubsystem](AActor* InActor, TArray<FName>& OutTags)
+	{
+		OutTags.Append(InActor->Tags);
+		
+		if (InActor->IsInLevelInstance())
+		{
+			LevelInstanceSubsystem->ForEachLevelInstanceAncestors(InActor, [&OutTags](ILevelInstanceInterface* Ancestor)
+			{
+				AActor* AncestorActor = CastChecked<AActor>(Ancestor);
+				OutTags.Append(AncestorActor->Tags);
+				return true;
+			});
+		}
+	};
+
 	// Build actor-index map
-	TMap<AActor*, int> ActorIndexMap;
+	TMap<AActor*, FActorEntry> ActorIndexMap;
 	int LastActorIndex = 1; // Since the root is the "first" point we'll have, we'll have the map start from 1.
-	UPCGActorHelpers::ForEachActorInWorld(World, AActor::StaticClass(), [&ActorIndexMap, &LastActorIndex](AActor* Actor)
+	UPCGActorHelpers::ForEachActorInWorld(World, AActor::StaticClass(), [&ActorIndexMap, &LastActorIndex, &GetAttachParentActor, &GetActorTags](AActor* Actor)
 	{
-		ActorIndexMap.Add(Actor, LastActorIndex++);
+		FActorEntry ActorEntry(LastActorIndex++);
+		
+		ActorEntry.AttachmentParent = GetAttachParentActor(Actor);
+
+		GetActorTags(Actor, ActorEntry.Tags);
+
+		ActorIndexMap.Add(Actor, MoveTemp(ActorEntry));
 		return true;
 	});
 
@@ -292,7 +353,9 @@ bool UPCGLevelToAsset::BP_ExportWorld_Implementation(UWorld* World, const FStrin
 		const FBox ActorBounds = PCGHelpers::GetActorBounds(Actor, /*bIgnorePCGCreatedComponents=*/true);
 		AllActorBounds += ActorBounds;
 
-		for (FName ActorTag : Actor->Tags)
+		const FActorEntry& ActorEntry = ActorIndexMap[Actor];
+
+		for (FName ActorTag : ActorEntry.Tags)
 		{
 			PCG::Private::FParseTagResult TagData(ActorTag);
 
@@ -328,9 +391,9 @@ bool UPCGLevelToAsset::BP_ExportWorld_Implementation(UWorld* World, const FStrin
 
 		// Prepare actor-level data that's propagated to all points
 		const FTransform& ActorTransform = Actor->GetTransform();
-		const int64 ActorIndex = ActorIndexMap[Actor];
-		AActor* ParentActor = Actor->GetAttachParentActor();
-		const int64 ParentActorIndex = ParentActor ? ActorIndexMap[ParentActor] : 0;
+		const int64 ActorIndex = ActorEntry.Index;
+		AActor* ParentActor = ActorEntry.AttachmentParent;
+		const int64 ParentActorIndex = ParentActor ? ActorIndexMap[ParentActor].Index : 0;
 		const FTransform RelativeTransform = ParentActor ? ActorTransform.GetRelativeTransform(ParentActor->GetTransform()) : ActorTransform;
 
 		// Hierarchy depth, starts at 1 if the actor doesn't have a parent
@@ -338,7 +401,7 @@ bool UPCGLevelToAsset::BP_ExportWorld_Implementation(UWorld* World, const FStrin
 		while (ParentActor)
 		{
 			++HierarchyDepth;
-			ParentActor = ParentActor->GetAttachParentActor();
+			ParentActor = GetAttachParentActor(ParentActor);
 		}
 
 		auto MakePoint = [&](int32 PointIndex, const FTransform& Transform, const FSoftObjectPath& MeshPath, const FBox& MeshBounds, const TArray<UMaterialInterface*>& MeshMaterials)
@@ -370,7 +433,7 @@ bool UPCGLevelToAsset::BP_ExportWorld_Implementation(UWorld* World, const FStrin
 
 			// For all tags, if the tag is of format 'Name:Value' then create attribute Name and assign Value, otherwise
 			// create a boolean attribute with the name given by the sanitized tag string.
-			for (FName ActorTag : Actor->Tags)
+			for (FName ActorTag : ActorEntry.Tags)
 			{
 				// Implementation note: we don't set values directly from the tag in the eventuality that there are some name collisions
 				PCG::Private::FParseTagResult TagData(ActorTag);
