[PCG] Fix UPCGBlueprintHelpers::GetRandomStreamFromPoint regression (CL 39589178)

#rnx

diff --git a/Engine/Plugins/PCG/Source/PCG/Private/Helpers/PCGBlueprintHelpers.cpp b/Engine/Plugins/PCG/Source/PCG/Private/Helpers/PCGBlueprintHelpers.cpp
--- a/Engine/Plugins/PCG/Source/PCG/Private/Helpers/PCGBlueprintHelpers.cpp
+++ b/Engine/Plugins/PCG/Source/PCG/Private/Helpers/PCGBlueprintHelpers.cpp
@@ -43,7 +43,7 @@ void UPCGBlueprintHelpers::SetSeedFromPosition(FPCGPoint& InPoint)
 
 FRandomStream UPCGBlueprintHelpers::GetRandomStreamFromPoint(const FPCGPoint& InPoint, const UPCGSettings* OptionalSettings, const UPCGComponent* OptionalComponent)
 {
-	return PCGHelpers::GetRandomStreamFromSeed(InPoint.Seed);
+	return PCGHelpers::GetRandomStreamFromSeed(InPoint.Seed, OptionalSettings, OptionalComponent);
 }
 
 FRandomStream UPCGBlueprintHelpers::GetRandomStreamFromTwoPoints(const FPCGPoint& InPointA, const FPCGPoint& InPointB, const UPCGSettings* OptionalSettings, const UPCGComponent* OptionalComponent)
