Refactoring to clarify BasePointData API (copy points to vs copy points) - copy points was renamed to SetPoints.
Added copy points to API with indices views instead.

#rnx

diff --git a/Engine/Plugins/PCG/Config/BasePCG.ini b/Engine/Plugins/PCG/Config/BasePCG.ini
--- a/Engine/Plugins/PCG/Config/BasePCG.ini
+++ b/Engine/Plugins/PCG/Config/BasePCG.ini
@@ -30,7 +30,8 @@
 +FunctionRedirects=(OldName="PCGMetadata.CopyExistingAttribute",NewName="PCGMetadata.BP_CopyExistingAttribute")
 +FunctionRedirects=(OldName="PCGMetadata.RenameAttribute",NewName="PCGMetadata.BP_RenameAttribute")
 +FunctionRedirects=(OldName="PCGMetadata.ClearAttribute",NewName="PCGMetadata.BP_ClearAttribute")
-+FunctionRedirects=(OldName="PCGBasePointData.CopyPointsFrom",NewName="PCGBasePointData.BP_CopyPointsFrom")
++FunctionRedirects=(OldName="PCGBasePointData.CopyPointsFrom",NewName="PCGBasePointData.BP_SetPointsFrom")
++FunctionRedirects=(OldName="PCGBasePointData.BP_CopyPointsFrom",NewName="PCGBasePointData.BP_SetPointsFrom")
 
 +PropertyRedirects=(OldName="PCGSplineSamplerSettings.Params",NewName="PCGSplineSamplerSettings.SamplerParams")
 +PropertyRedirects=(OldName="PCGProjectionSettings.Params",NewName="PCGProjectionSettings.ProjectionParams")
diff --git a/Engine/Plugins/PCG/Source/PCG/Private/Data/PCGBasePointData.cpp b/Engine/Plugins/PCG/Source/PCG/Private/Data/PCGBasePointData.cpp
--- a/Engine/Plugins/PCG/Source/PCG/Private/Data/PCGBasePointData.cpp
+++ b/Engine/Plugins/PCG/Source/PCG/Private/Data/PCGBasePointData.cpp
@@ -767,7 +767,7 @@ FBox UPCGBasePointData::GetBounds() const
 
 void UPCGBasePointData::CopyPropertiesTo(UPCGBasePointData* To, int32 ReadStartIndex, int32 WriteStartIndex, int32 Count, EPCGPointNativeProperties Properties) const
 {
-	if (EPCGPointNativeProperties::None == Properties)
+	if (EPCGPointNativeProperties::None == Properties || Count <= 0)
 	{
 		return;
 	}
@@ -840,30 +840,97 @@ void UPCGBasePointData::CopyPropertiesTo(UPCGBasePointData* To, int32 ReadStartI
 	To->DirtyCache();
 }
 
+void UPCGBasePointData::CopyPropertiesTo(UPCGBasePointData* To, const TArrayView<const int32>& ReadIndices, const TArrayView<const int32>& WriteIndices, EPCGPointNativeProperties Properties) const
+{
+	if (EPCGPointNativeProperties::None == Properties || ReadIndices.Num() == 0 || ReadIndices.Num() != WriteIndices.Num())
+	{
+		check(ReadIndices.Num() == WriteIndices.Num());
+		return;
+	}
+
+	const int32 Count = ReadIndices.Num();
+
+	const TConstPCGValueRange<FTransform> FromTransformRange = GetConstTransformValueRange();
+	const TConstPCGValueRange<float> FromSteepnessRange = GetConstSteepnessValueRange();
+	const TConstPCGValueRange<float> FromDensityRange = GetConstDensityValueRange();
+	const TConstPCGValueRange<FVector> FromBoundsMinRange = GetConstBoundsMinValueRange();
+	const TConstPCGValueRange<FVector> FromBoundsMaxRange = GetConstBoundsMaxValueRange();
+	const TConstPCGValueRange<FVector4> FromColorRange = GetConstColorValueRange();
+	const TConstPCGValueRange<int64> FromMetadataEntryRange = GetConstMetadataEntryValueRange();
+	const TConstPCGValueRange<int32> FromSeedRange = GetConstSeedValueRange();
+
+	To->AllocateProperties(Properties);
+
+	TPCGValueRange<FTransform> ToTransformRange = To->GetTransformValueRange(/*bAllocate=*/false);
+	TPCGValueRange<float> ToSteepnessRange = To->GetSteepnessValueRange(/*bAllocate=*/false);
+	TPCGValueRange<float> ToDensityRange = To->GetDensityValueRange(/*bAllocate=*/false);
+	TPCGValueRange<FVector> ToBoundsMinRange = To->GetBoundsMinValueRange(/*bAllocate=*/false);
+	TPCGValueRange<FVector> ToBoundsMaxRange = To->GetBoundsMaxValueRange(/*bAllocate=*/false);
+	TPCGValueRange<FVector4> ToColorRange = To->GetColorValueRange(/*bAllocate=*/false);
+	TPCGValueRange<int64> ToMetadataEntryRange = To->GetMetadataEntryValueRange(/*bAllocate=*/false);
+	TPCGValueRange<int32> ToSeedRange = To->GetSeedValueRange(/*bAllocate=*/false);
+
+	auto CopyRangeIf = [&ReadIndices, &WriteIndices, Properties, Count]<typename T>(TPCGValueRange<T>& ToRange, const TConstPCGValueRange<T>& FromRange, EPCGPointNativeProperties Property)
+	{
+		if (EnumHasAllFlags(Properties, Property))
+		{
+			for (int32 Index = 0; Index < Count; ++Index)
+			{
+				const int32& ReadIndex = ReadIndices[Index];
+				const int32& WriteIndex = WriteIndices[Index];
+
+				ToRange[WriteIndex] = FromRange[ReadIndex];
+			}
+		}
+	};
+
+	CopyRangeIf(ToTransformRange, FromTransformRange, EPCGPointNativeProperties::Transform);
+	CopyRangeIf(ToDensityRange, FromDensityRange, EPCGPointNativeProperties::Density);
+	CopyRangeIf(ToBoundsMinRange, FromBoundsMinRange, EPCGPointNativeProperties::BoundsMin);
+	CopyRangeIf(ToBoundsMaxRange, FromBoundsMaxRange, EPCGPointNativeProperties::BoundsMax);
+	CopyRangeIf(ToColorRange, FromColorRange, EPCGPointNativeProperties::Color);
+	CopyRangeIf(ToSteepnessRange, FromSteepnessRange, EPCGPointNativeProperties::Steepness);
+	CopyRangeIf(ToSeedRange, FromSeedRange, EPCGPointNativeProperties::Seed);
+	CopyRangeIf(ToMetadataEntryRange, FromMetadataEntryRange, EPCGPointNativeProperties::MetadataEntry);
+
+	To->DirtyCache();
+}
+
 void UPCGBasePointData::CopyPointsTo(UPCGBasePointData* To, int32 ReadStartIndex, int32 WriteStartIndex, int32 Count) const
 {
 	CopyPropertiesTo(To, ReadStartIndex, WriteStartIndex, Count, EPCGPointNativeProperties::All);
 }
 
-void UPCGBasePointData::BP_CopyPointsFrom(const UPCGBasePointData* InData, const TArray<int32>& InDataIndices)
+void UPCGBasePointData::CopyPointsTo(UPCGBasePointData* To, const TArrayView<const int32>& ReadIndices, const TArrayView<const int32>& WriteIndices) const
+{
+	CopyPropertiesTo(To, ReadIndices, WriteIndices, EPCGPointNativeProperties::All);
+}
+
+void UPCGBasePointData::BP_SetPointsFrom(const UPCGBasePointData* InData, const TArray<int32>& InDataIndices)
 {
-	return CopyPoints(InData, this, InDataIndices, false);
+	return SetPoints(InData, this, InDataIndices, false);
 }
 
-void UPCGBasePointData::CopyPointsFrom(const UPCGBasePointData* InData, const TArrayView<const int32>& InDataIndices)
+void UPCGBasePointData::SetPointsFrom(const UPCGBasePointData* InData, const TArrayView<const int32>& InDataIndices)
 {
-	return CopyPoints(InData, this, InDataIndices, false);
+	return SetPoints(InData, this, InDataIndices, false);
 }
 
-void UPCGBasePointData::CopyPoints(const UPCGBasePointData* From, UPCGBasePointData* To, const TArrayView<const int32>& InDataIndices, bool bCopyAll)
+void UPCGBasePointData::SetPoints(const UPCGBasePointData* From, UPCGBasePointData* To, const TArrayView<const int32>& InDataIndices, bool bCopyAll)
 {
-	TRACE_CPUPROFILER_EVENT_SCOPE(UPCGBasePointData::CopyPoints);
+	TRACE_CPUPROFILER_EVENT_SCOPE(UPCGBasePointData::SetPoints);
 	check(From);
 
 	const int32 NumPoints = bCopyAll ? From->GetNumPoints() : InDataIndices.Num();
 	To->SetNumPoints(NumPoints);
 	To->AllocateProperties(From->GetAllocatedProperties());
 
+	if (NumPoints == 0)
+	{
+		To->DirtyCache();
+		return;
+	}
+
 	const TConstPCGValueRange<FTransform> FromTransformRange = From->GetConstTransformValueRange();
 	const TConstPCGValueRange<float> FromSteepnessRange = From->GetConstSteepnessValueRange();
 	const TConstPCGValueRange<float> FromDensityRange = From->GetConstDensityValueRange();
diff --git a/Engine/Plugins/PCG/Source/PCG/Private/Data/PCGPointArrayData.cpp b/Engine/Plugins/PCG/Source/PCG/Private/Data/PCGPointArrayData.cpp
--- a/Engine/Plugins/PCG/Source/PCG/Private/Data/PCGPointArrayData.cpp
+++ b/Engine/Plugins/PCG/Source/PCG/Private/Data/PCGPointArrayData.cpp
@@ -32,7 +32,7 @@ const UPCGPointData* UPCGPointArrayData::ToPointData(FPCGContext* Context, const
 	UPCGPointData* PointData = FPCGContext::NewObject_AnyThread<UPCGPointData>(Context);
 	PointData->InitializeFromData(this);
 
-	CopyPoints(this, PointData, {}, /*bCopyAll=*/true);
+	SetPoints(this, PointData, {}, /*bCopyAll=*/true);
 
 	return PointData;
 }
@@ -52,6 +52,11 @@ UPCGSpatialData* UPCGPointArrayData::CopyInternal(FPCGContext* Context) const
 
 void UPCGPointArrayData::CopyPropertiesTo(UPCGBasePointData* To, int32 ReadStartIndex, int32 WriteStartIndex, int32 Count, EPCGPointNativeProperties Properties) const
 {
+	if (Count <= 0)
+	{
+		return;
+	}
+
 	if (UPCGPointArrayData* PointArrayData = Cast<UPCGPointArrayData>(To))
 	{
 		PointArrayData->AllocateProperties(GetAllocatedProperties());
diff --git a/Engine/Plugins/PCG/Source/PCG/Private/Data/PCGSpatialData.cpp b/Engine/Plugins/PCG/Source/PCG/Private/Data/PCGSpatialData.cpp
--- a/Engine/Plugins/PCG/Source/PCG/Private/Data/PCGSpatialData.cpp
+++ b/Engine/Plugins/PCG/Source/PCG/Private/Data/PCGSpatialData.cpp
@@ -42,7 +42,7 @@ const UPCGPointArrayData* UPCGSpatialData::ToPointArrayData(FPCGContext* Context
 		UPCGPointArrayData* PointArrayData = FPCGContext::NewObject_AnyThread<UPCGPointArrayData>(Context);
 		PointArrayData->InitializeFromData(PointData);
 
-		UPCGBasePointData::CopyPoints(PointData, PointArrayData, {}, /*bCopyAll=*/true);
+		UPCGBasePointData::SetPoints(PointData, PointArrayData, {}, /*bCopyAll=*/true);
 
 		return PointArrayData;
 	}
diff --git a/Engine/Plugins/PCG/Source/PCG/Private/Elements/PCGCullPointsOutsideActorBounds.cpp b/Engine/Plugins/PCG/Source/PCG/Private/Elements/PCGCullPointsOutsideActorBounds.cpp
--- a/Engine/Plugins/PCG/Source/PCG/Private/Elements/PCGCullPointsOutsideActorBounds.cpp
+++ b/Engine/Plugins/PCG/Source/PCG/Private/Elements/PCGCullPointsOutsideActorBounds.cpp
@@ -160,7 +160,7 @@ bool FPCGCullPointsOutsideActorBoundsElement::ExecuteInternal(FPCGContext* Conte
 				KeptPointIndices.Sort();
 			}
 
-			CulledPointsData->CopyPointsFrom(InputPointData, KeptPointIndices);
+			CulledPointsData->SetPointsFrom(InputPointData, KeptPointIndices);
 			OutputData = CulledPointsData;
 		}
 		
diff --git a/Engine/Plugins/PCG/Source/PCG/Private/Elements/PCGSelfPruning.cpp b/Engine/Plugins/PCG/Source/PCG/Private/Elements/PCGSelfPruning.cpp
--- a/Engine/Plugins/PCG/Source/PCG/Private/Elements/PCGSelfPruning.cpp
+++ b/Engine/Plugins/PCG/Source/PCG/Private/Elements/PCGSelfPruning.cpp
@@ -691,7 +691,7 @@ PRAGMA_ENABLE_DEPRECATION_WARNINGS
 
 			if (Indices.Num() > 0)
 			{
-				UPCGBasePointData::CopyPoints(InState.InputData, PrunedData, Indices, /*bCopyAll=*/false);
+				UPCGBasePointData::SetPoints(InState.InputData, PrunedData, Indices, /*bCopyAll=*/false);
 			}
 
 			if (InOptionalContext)
diff --git a/Engine/Plugins/PCG/Source/PCG/Private/Elements/PCGSortAttributes.cpp b/Engine/Plugins/PCG/Source/PCG/Private/Elements/PCGSortAttributes.cpp
--- a/Engine/Plugins/PCG/Source/PCG/Private/Elements/PCGSortAttributes.cpp
+++ b/Engine/Plugins/PCG/Source/PCG/Private/Elements/PCGSortAttributes.cpp
@@ -93,7 +93,7 @@ bool FPCGSortAttributesElement::ExecuteInternal(FPCGContext* Context) const
 				InitalizeFromDataParams.bInheritSpatialData = false;
 				OutputPointData->InitializeFromDataWithParams(InitalizeFromDataParams);
 
-				UPCGBasePointData::CopyPoints(InputPointData, OutputPointData, SortedKeyIndices, /*bCopyAll=*/false);
+				UPCGBasePointData::SetPoints(InputPointData, OutputPointData, SortedKeyIndices, /*bCopyAll=*/false);
 
 				OutputData = OutputPointData;
 			}
diff --git a/Engine/Plugins/PCG/Source/PCG/Private/Metadata/PCGMetadataPartitionCommon.cpp b/Engine/Plugins/PCG/Source/PCG/Private/Metadata/PCGMetadataPartitionCommon.cpp
--- a/Engine/Plugins/PCG/Source/PCG/Private/Metadata/PCGMetadataPartitionCommon.cpp
+++ b/Engine/Plugins/PCG/Source/PCG/Private/Metadata/PCGMetadataPartitionCommon.cpp
@@ -446,7 +446,7 @@ namespace PCGMetadataPartitionCommon
 
 			CurrentPointData->InitializeFromDataWithParams(InitializeFromDataParams);
 
-			UPCGBasePointData::CopyPoints(InData, CurrentPointData, Indices, /*bCopyAll=*/false);
+			UPCGBasePointData::SetPoints(InData, CurrentPointData, Indices, /*bCopyAll=*/false);
 		}
 
 		return PartitionedData;
@@ -482,7 +482,7 @@ namespace PCGMetadataPartitionCommon
 		OutputPointData->InitializeFromDataWithParams(InitializeFromDataParams);
 		OutputPointData->SetNumPoints(IndicesToCopy.Num());
 		
-		UPCGBasePointData::CopyPoints(InData, OutputPointData, IndicesToCopy, /*bCopyAll=*/false);
+		UPCGBasePointData::SetPoints(InData, OutputPointData, IndicesToCopy, /*bCopyAll=*/false);
 
 		return OutputPointData;
 	}
diff --git a/Engine/Plugins/PCG/Source/PCG/Public/Data/PCGBasePointData.h b/Engine/Plugins/PCG/Source/PCG/Public/Data/PCGBasePointData.h
--- a/Engine/Plugins/PCG/Source/PCG/Public/Data/PCGBasePointData.h
+++ b/Engine/Plugins/PCG/Source/PCG/Public/Data/PCGBasePointData.h
@@ -57,6 +57,7 @@ public:
 	virtual EPCGDataType GetDataType() const override { return EPCGDataType::Point; }
 
 	virtual int32 GetNumPoints() const PURE_VIRTUAL(UPCGBasePointData::GetNumPoints, return 0;);
+	bool IsEmpty() const { return GetNumPoints() == 0; }
 	virtual void SetNumPoints(int32 InNumPoints) PURE_VIRTUAL(UPCGBasePointData::SetNumPoints, return;);
 
 	virtual void AllocateProperties(EPCGPointNativeProperties Properties) {}
@@ -133,20 +134,22 @@ public:
 	virtual const PCGPointOctree::FPointOctree& GetPointOctree() const;
 	virtual FBox GetBounds() const override;
 		
-	UFUNCTION(BlueprintCallable, Category = SpatialData, meta = (DisplayName="Copy Points From"))
-	void BP_CopyPointsFrom(const UPCGBasePointData* InData, const TArray<int32>& InDataIndices);
+	UFUNCTION(BlueprintCallable, Category = SpatialData, meta = (DisplayName="Set Points From"))
+	void BP_SetPointsFrom(const UPCGBasePointData* InData, const TArray<int32>& InDataIndices);
 	
-	void CopyPointsFrom(const UPCGBasePointData* InData, const TArrayView<const int32>& InDataIndices);
+	void SetPointsFrom(const UPCGBasePointData* InData, const TArrayView<const int32>& InDataIndices);
 
 	virtual void CopyPropertiesTo(UPCGBasePointData* To, int32 ReadStartIndex, int32 WriteStartIndex, int32 Count, EPCGPointNativeProperties Properties) const;
+	virtual void CopyPropertiesTo(UPCGBasePointData* To, const TArrayView<const int32>& ReadIndices, const TArrayView<const int32>& WriteIndices, EPCGPointNativeProperties Properties) const;
 	virtual void CopyPointsTo(UPCGBasePointData* To, int32 ReadStartIndex, int32 WriteStartIndex, int32 Count) const;
+	virtual void CopyPointsTo(UPCGBasePointData* To, const TArrayView<const int32>& ReadIndices, const TArrayView<const int32>& WriteIndices) const;
 	virtual EPCGPointNativeProperties GetAllocatedProperties(bool bWithInheritance = true) const { return EPCGPointNativeProperties::All; }
 
 	/** Metadata specific */
 	virtual FPCGMetadataDomainID GetDefaultMetadataDomainID() const override { return PCGMetadataDomainID::Elements; }
 	virtual TArray<FPCGMetadataDomainID> GetAllSupportedMetadataDomainIDs() const override { return {PCGMetadataDomainID::Data, PCGMetadataDomainID::Elements}; }
 
-	static void CopyPoints(const UPCGBasePointData* From, UPCGBasePointData* To, const TArrayView<const int32>& InDataIndices, bool bCopyAll);
+	static void SetPoints(const UPCGBasePointData* From, UPCGBasePointData* To, const TArrayView<const int32>& InDataIndices, bool bCopyAll);
 
 protected:
 
