[PCG] Cleanup prep-work for Random Choice Node Update

* Moved PCGAttributeCast into the Metadata folder
* Corrected BreakTransform and Compare Insights traces to avoid future confusion
* Converted seed magic number 42 into constant value
* Converted global namespace enums to enum classes to prevent spillage
+ Added a default "[REDACTED_HANDLE]" for all Point Data

#rnx

diff --git a/Engine/Plugins/PCG/Source/PCG/Private/Data/PCGBasePointData.cpp b/Engine/Plugins/PCG/Source/PCG/Private/Data/PCGBasePointData.cpp
--- a/Engine/Plugins/PCG/Source/PCG/Private/Data/PCGBasePointData.cpp
+++ b/Engine/Plugins/PCG/Source/PCG/Private/Data/PCGBasePointData.cpp
@@ -717,6 +717,11 @@ UPCGBasePointData::UPCGBasePointData(const FObjectInitializer& ObjectInitializer
 {
 	check(Metadata);
 	Metadata->SetupDomain(PCGMetadataDomainID::Elements, /*bIsDefault=*/true);
+
+	// Default to Position as the "last attribute" on creation.
+	FPCGAttributePropertySelector DefaultSelector;
+	DefaultSelector.SetPointProperty(EPCGPointProperties::Position);
+	UPCGSpatialData::SetLastSelector(DefaultSelector);
 }
 
 FPCGAttributeAccessorMethods UPCGBasePointData::GetPointAccessorMethods()
diff --git a/Engine/Plugins/PCG/Source/PCG/Private/Elements/Metadata/PCGAttributeCast.cpp b/Engine/Plugins/PCG/Source/PCG/Private/Elements/Metadata/PCGAttributeCast.cpp
new file mode 100644
--- /dev/null
+++ b/Engine/Plugins/PCG/Source/PCG/Private/Elements/Metadata/PCGAttributeCast.cpp
@@ -0,0 +1,116 @@
+// Copyright Epic Games, Inc. All Rights Reserved.
+
+#include "Elements/Metadata/PCGAttributeCast.h"
+
+#include "PCGContext.h"
+#include "PCGParamData.h"
+#include "Elements/Metadata/PCGMetadataElementCommon.h"
+#include "Metadata/PCGMetadataAttributeTraits.h"
+#include "Metadata/Accessors/IPCGAttributeAccessor.h"
+#include "Metadata/Accessors/PCGAttributeAccessorKeys.h"
+#include "Metadata/Accessors/PCGAttributeAccessorHelpers.h"
+
+#define LOCTEXT_NAMESPACE "PCGAttributeCastElement"
+
+TArray<FPCGPinProperties> UPCGAttributeCastSettings::OutputPinProperties() const
+{
+	TArray<FPCGPinProperties> PinProperties;
+	PinProperties.Emplace(PCGPinConstants::DefaultOutputLabel, EPCGDataType::Any);
+
+	return PinProperties;
+}
+
+FPCGElementPtr UPCGAttributeCastSettings::CreateElement() const
+{
+	return MakeShared<FPCGAttributeCastElement>();
+}
+
+#if WITH_EDITOR
+TArray<FPCGPreConfiguredSettingsInfo> UPCGAttributeCastSettings::GetPreconfiguredInfo() const
+{
+	return FPCGPreConfiguredSettingsInfo::PopulateFromEnum<EPCGMetadataTypes>({ EPCGMetadataTypes::Count, EPCGMetadataTypes::Unknown }, LOCTEXT("PreconfigureInfo", "Cast to "));
+}
+#endif
+
+FString UPCGAttributeCastSettings::GetAdditionalTitleInformation() const
+{
+#if WITH_EDITOR
+	const bool bIsTypeOverridden = IsPropertyOverriddenByPin(GET_MEMBER_NAME_CHECKED(UPCGAttributeCastSettings, OutputType));
+	const bool bSelectorOverridden = IsPropertyOverriddenByPin(GET_MEMBER_NAME_CHECKED(UPCGAttributeCastSettings, InputSource)) || IsPropertyOverriddenByPin(GET_MEMBER_NAME_CHECKED(UPCGAttributeCastSettings, OutputTarget));
+
+	if (bIsTypeOverridden)
+	{
+		return FString();
+	}
+	else
+	{
+		if (bSelectorOverridden)
+		{
+			return FText::Format(LOCTEXT("CastType", "To {0}"), PCG::Private::GetTypeNameText((uint16)OutputType)).ToString();
+		}
+	}
+#endif
+
+	return FText::Format(LOCTEXT("FullInfo", "{0} -> {1} as {2}"), InputSource.GetDisplayText(), OutputTarget.GetDisplayText(), PCG::Private::GetTypeNameText((uint16)OutputType)).ToString();
+}
+
+FName UPCGAttributeCastSettings::AdditionalTaskName() const
+{
+	return FName(*FText::Format(LOCTEXT("AdditionalTaskName", "Attribute Cast: {0}"), PCG::Private::GetTypeNameText(static_cast<uint16>(OutputType))).ToString());
+}
+
+void UPCGAttributeCastSettings::ApplyPreconfiguredSettings(const FPCGPreConfiguredSettingsInfo& PreconfiguredInfo)
+{
+	if (const UEnum* EnumPtr = StaticEnum<EPCGMetadataTypes>())
+	{
+		if (EnumPtr->IsValidEnumValue(PreconfiguredInfo.PreconfiguredIndex))
+		{
+			OutputType = EPCGMetadataTypes(PreconfiguredInfo.PreconfiguredIndex);
+		}
+	}
+}
+
+bool FPCGAttributeCastElement::ExecuteInternal(FPCGContext* Context) const
+{
+	check(Context);
+
+	const UPCGAttributeCastSettings* Settings = Context->GetInputSettings<UPCGAttributeCastSettings>();
+	check(Settings);
+
+	TArray<FPCGTaggedData> Inputs = Context->InputData.GetInputsByPin(PCGPinConstants::DefaultInputLabel);
+
+	for (int32 i = 0; i < Inputs.Num(); ++i)
+	{
+		const FPCGTaggedData& Input = Inputs[i];
+
+		// Only support types that supports metadata
+		if (!ensure(Input.Data) || !Input.Data->ConstMetadata())
+		{
+			PCGE_LOG(Error, GraphAndLog, FText::Format(LOCTEXT("InvalidData", "Only support Spatial and Attribute Set, Input {0} is neither of them. Skipped"), FText::AsNumber(i)));
+			continue;
+		}
+
+		FPCGTaggedData& Output = Context->OutputData.TaggedData.Add_GetRef(Input);
+
+		UPCGData* OutputData = Input.Data->DuplicateData(Context);
+		check(OutputData);
+
+		PCGMetadataHelpers::FPCGCopyAttributeParams Params{};
+		Params.SourceData = Input.Data;
+		Params.TargetData = OutputData;
+		Params.InputSource = Settings->InputSource;
+		Params.OutputTarget = Settings->OutputTarget;
+		Params.OptionalContext = Context;
+		Params.bSameOrigin = true;
+		Params.OutputType = Settings->OutputType;
+
+		if (PCGMetadataHelpers::CopyAttribute(Params))
+		{
+			Output.Data = OutputData;
+		}
+	}
+
+	return true;
+}
+
+#undef LOCTEXT_NAMESPACE
\ No newline at end of file
diff --git a/Engine/Plugins/PCG/Source/PCG/Private/Elements/Metadata/PCGMetadataBreakTransform.cpp b/Engine/Plugins/PCG/Source/PCG/Private/Elements/Metadata/PCGMetadataBreakTransform.cpp
--- a/Engine/Plugins/PCG/Source/PCG/Private/Elements/Metadata/PCGMetadataBreakTransform.cpp
+++ b/Engine/Plugins/PCG/Source/PCG/Private/Elements/Metadata/PCGMetadataBreakTransform.cpp
@@ -109,7 +109,7 @@ FPCGElementPtr UPCGMetadataBreakTransformSettings::CreateElement() const
 
 bool FPCGMetadataBreakTransformElement::DoOperation(PCGMetadataOps::FOperationData& OperationData) const
 {
-	TRACE_CPUPROFILER_EVENT_SCOPE(FPCGMetadataBreakVectorElement::Execute);
+	TRACE_CPUPROFILER_EVENT_SCOPE(FPCGMetadataBreakTransformElement::Execute);
 
 	const UPCGMetadataBreakTransformSettings* Settings = static_cast<const UPCGMetadataBreakTransformSettings*>(OperationData.Settings);
 	check(Settings);
diff --git a/Engine/Plugins/PCG/Source/PCG/Private/Elements/Metadata/PCGMetadataCompareOpElement.cpp b/Engine/Plugins/PCG/Source/PCG/Private/Elements/Metadata/PCGMetadataCompareOpElement.cpp
--- a/Engine/Plugins/PCG/Source/PCG/Private/Elements/Metadata/PCGMetadataCompareOpElement.cpp
+++ b/Engine/Plugins/PCG/Source/PCG/Private/Elements/Metadata/PCGMetadataCompareOpElement.cpp
@@ -164,7 +164,7 @@ FPCGElementPtr UPCGMetadataCompareSettings::CreateElement() const
 
 bool FPCGMetadataCompareElement::DoOperation(PCGMetadataOps::FOperationData& OperationData) const
 {
-	TRACE_CPUPROFILER_EVENT_SCOPE(FPCGMetadataTrigElement::Execute);
+	TRACE_CPUPROFILER_EVENT_SCOPE(FPCGMetadataCompareElement::Execute);
 
 	const UPCGMetadataCompareSettings* Settings = CastChecked<UPCGMetadataCompareSettings>(OperationData.Settings);
 
diff --git a/Engine/Plugins/PCG/Source/PCG/Private/Elements/PCGAttributeCast.cpp b/Engine/Plugins/PCG/Source/PCG/Private/Elements/PCGAttributeCast.cpp
deleted file mode 100644
diff --git a/Engine/Plugins/PCG/Source/PCG/Private/Elements/PCGSplineSampler.cpp b/Engine/Plugins/PCG/Source/PCG/Private/Elements/PCGSplineSampler.cpp
--- a/Engine/Plugins/PCG/Source/PCG/Private/Elements/PCGSplineSampler.cpp
+++ b/Engine/Plugins/PCG/Source/PCG/Private/Elements/PCGSplineSampler.cpp
@@ -969,7 +969,7 @@ namespace PCGSplineSamplerHelpers
 			}
 		}
 
-		const int32 Seed = Context ? Context->GetSeed() : 42;
+		const int32 Seed = Context ? Context->GetSeed() : PCGValueConstants::DefaultSeed;
 
 		FSubdivisionStepSampler SubdivisionSampler(LineData, Params, Seed);
 		FDistanceStepSampler DistanceSampler(LineData, Params, Seed);
@@ -1399,7 +1399,7 @@ namespace PCGSplineSamplerHelpers
 			PointCount += InteriorData.Num();
 		}
 
-		const int32 Seed = Context ? Context->GetSeed() : 42;
+		const int32 Seed = Context ? Context->GetSeed() : PCGValueConstants::DefaultSeed;
 		int32 SampleIndex = 0;
 
 		// TODO: should we parallel for this too?
diff --git a/Engine/Plugins/PCG/Source/PCG/Private/Elements/PCGSurfaceSampler.cpp b/Engine/Plugins/PCG/Source/PCG/Private/Elements/PCGSurfaceSampler.cpp
--- a/Engine/Plugins/PCG/Source/PCG/Private/Elements/PCGSurfaceSampler.cpp
+++ b/Engine/Plugins/PCG/Source/PCG/Private/Elements/PCGSurfaceSampler.cpp
@@ -51,7 +51,7 @@ namespace PCGSurfaceSampler
 			return false;
 		}
 
-		Seed = Context ? Context->GetSeed() : 42;
+		Seed = Context ? Context->GetSeed() : PCGValueConstants::DefaultSeed;
 
 		// Conceptually, we will break down the surface bounds in a N x M grid
 		InterstitialDistance = Params.PointExtents * 2;
diff --git a/Engine/Plugins/PCG/Source/PCG/Private/Tests/Elements/PCGAttributeCastTest.cpp b/Engine/Plugins/PCG/Source/PCG/Private/Tests/Elements/PCGAttributeCastTest.cpp
--- a/Engine/Plugins/PCG/Source/PCG/Private/Tests/Elements/PCGAttributeCastTest.cpp
+++ b/Engine/Plugins/PCG/Source/PCG/Private/Tests/Elements/PCGAttributeCastTest.cpp
@@ -7,7 +7,7 @@
 #include "PCGParamData.h"
 #include "Data/PCGPointData.h"
 
-#include "Elements/PCGAttributeCast.h"
+#include "Elements/Metadata/PCGAttributeCast.h"
 
 class FPCGAttributeCastTestBase : public FPCGTestBaseClass
 {
diff --git a/Engine/Plugins/PCG/Source/PCG/Public/Elements/Metadata/PCGAttributeCast.h b/Engine/Plugins/PCG/Source/PCG/Public/Elements/Metadata/PCGAttributeCast.h
new file mode 100644
--- /dev/null
+++ b/Engine/Plugins/PCG/Source/PCG/Public/Elements/Metadata/PCGAttributeCast.h
@@ -0,0 +1,55 @@
+// Copyright Epic Games, Inc. All Rights Reserved.
+
+#pragma once
+
+#include "PCGElement.h"
+#include "PCGSettings.h"
+#include "Metadata/PCGAttributePropertySelector.h"
+#include "Metadata/PCGMetadataAttributeTraits.h"
+
+#include "PCGAttributeCast.generated.h"
+
+/**
+* Cast an attribute to another type. Support broadcastable cast (like double -> FVector) and constructible cast (like double -> float)
+*/
+UCLASS(MinimalAPI, BlueprintType, ClassGroup = (Procedural))
+class UPCGAttributeCastSettings : public UPCGSettings
+{
+	GENERATED_BODY()
+
+public:
+	UPROPERTY(BlueprintReadWrite, EditAnywhere, Category = Settings, meta = (PCG_Overridable))
+	FPCGAttributePropertyInputSelector InputSource;
+
+	UPROPERTY(BlueprintReadWrite, EditAnywhere, Category = Settings, meta = (PCG_Overridable))
+	EPCGMetadataTypes OutputType = EPCGMetadataTypes::Float;
+
+	UPROPERTY(BlueprintReadWrite, EditAnywhere, Category = Settings, meta = (PCG_Overridable, PCG_DiscardPropertySelection, PCG_DiscardExtraSelection))
+	FPCGAttributePropertyOutputSelector OutputTarget;
+
+	//~Begin UPCGSettings interface
+public:
+#if WITH_EDITOR
+	virtual FName GetDefaultNodeName() const override { return FName("AttributeCast"); }
+	virtual FText GetDefaultNodeTitle() const override { return NSLOCTEXT("PCGAttributeCastElement", "NodeTitle", "Attribute Cast"); }
+	virtual EPCGSettingsType GetType() const override { return EPCGSettingsType::Metadata; }
+	virtual TArray<FPCGPreConfiguredSettingsInfo> GetPreconfiguredInfo() const override;
+	virtual bool OnlyExposePreconfiguredSettings() const override { return true; }
+#endif
+	virtual bool HasDynamicPins() const override { return true; }
+	virtual FString GetAdditionalTitleInformation() const override;
+	virtual FName AdditionalTaskName() const override;
+	virtual void ApplyPreconfiguredSettings(const FPCGPreConfiguredSettingsInfo& PreconfiguredInfo) override;
+	virtual FPCGElementPtr CreateElement() const override;
+
+protected:
+	virtual TArray<FPCGPinProperties> OutputPinProperties() const;
+	//~End UPCGSettings interface
+};
+
+class FPCGAttributeCastElement : public IPCGElement
+{
+protected:
+	virtual bool ExecuteInternal(FPCGContext* Context) const override;
+	virtual EPCGElementExecutionLoopMode ExecutionLoopMode(const UPCGSettings* Settings) const override { return EPCGElementExecutionLoopMode::SinglePrimaryPin; }
+};
diff --git a/Engine/Plugins/PCG/Source/PCG/Public/Elements/PCGAttributeCast.h b/Engine/Plugins/PCG/Source/PCG/Public/Elements/PCGAttributeCast.h
--- a/Engine/Plugins/PCG/Source/PCG/Public/Elements/PCGAttributeCast.h
+++ b/Engine/Plugins/PCG/Source/PCG/Public/Elements/PCGAttributeCast.h
@@ -2,54 +2,8 @@
 
 #pragma once
 
-#include "PCGElement.h"
-#include "PCGSettings.h"
-#include "Metadata/PCGAttributePropertySelector.h"
-#include "Metadata/PCGMetadataAttributeTraits.h"
+#include "CoreTypes.h"
 
-#include "PCGAttributeCast.generated.h"
+UE_DEPRECATED_HEADER(5.6, "Use /Elements/Metadata/PCGAttributeCast.h instead of /Elements/PCGAttributeCast.h")
 
-/**
-* Cast an attribute to another type. Support broadcastable cast (like double -> FVector) and constructible cast (like double -> float)
-*/
-UCLASS(MinimalAPI, BlueprintType, ClassGroup = (Procedural))
-class UPCGAttributeCastSettings : public UPCGSettings
-{
-	GENERATED_BODY()
-
-public:
-	UPROPERTY(BlueprintReadWrite, EditAnywhere, Category = Settings, meta = (PCG_Overridable))
-	FPCGAttributePropertyInputSelector InputSource;
-
-	UPROPERTY(BlueprintReadWrite, EditAnywhere, Category = Settings, meta = (PCG_Overridable))
-	EPCGMetadataTypes OutputType = EPCGMetadataTypes::Float;
-
-	UPROPERTY(BlueprintReadWrite, EditAnywhere, Category = Settings, meta = (PCG_Overridable, PCG_DiscardPropertySelection, PCG_DiscardExtraSelection))
-	FPCGAttributePropertyOutputSelector OutputTarget;
-
-	//~Begin UPCGSettings interface
-public:
-#if WITH_EDITOR
-	virtual FName GetDefaultNodeName() const override { return FName("AttributeCast"); }
-	virtual FText GetDefaultNodeTitle() const override { return NSLOCTEXT("PCGAttributeCastElement", "NodeTitle", "Attribute Cast"); }
-	virtual EPCGSettingsType GetType() const override { return EPCGSettingsType::Metadata; }
-	virtual TArray<FPCGPreConfiguredSettingsInfo> GetPreconfiguredInfo() const override;
-	virtual bool OnlyExposePreconfiguredSettings() const override { return true; }
-#endif
-	virtual bool HasDynamicPins() const override { return true; }
-	virtual FString GetAdditionalTitleInformation() const override;
-	virtual FName AdditionalTaskName() const override;
-	virtual void ApplyPreconfiguredSettings(const FPCGPreConfiguredSettingsInfo& PreconfiguredInfo) override;
-	virtual FPCGElementPtr CreateElement() const override;
-
-protected:
-	virtual TArray<FPCGPinProperties> OutputPinProperties() const;
-	//~End UPCGSettings interface
-};
-
-class FPCGAttributeCastElement : public IPCGElement
-{
-protected:
-	virtual bool ExecuteInternal(FPCGContext* Context) const override;
-	virtual EPCGElementExecutionLoopMode ExecutionLoopMode(const UPCGSettings* Settings) const override { return EPCGElementExecutionLoopMode::SinglePrimaryPin; }
-};
+#include "Metadata/PCGAttributeCast.h"
\ No newline at end of file
diff --git a/Engine/Plugins/PCG/Source/PCG/Public/PCGCommon.h b/Engine/Plugins/PCG/Source/PCG/Public/PCGCommon.h
--- a/Engine/Plugins/PCG/Source/PCG/Public/PCGCommon.h
+++ b/Engine/Plugins/PCG/Source/PCG/Public/PCGCommon.h
@@ -167,6 +167,11 @@ enum class EPCGExclusiveDataType : uint8
 	DynamicMesh
 };
 
+namespace PCGValueConstants
+{
+	constexpr int32 DefaultSeed = 42;
+}
+
 namespace PCGPinConstants
 {
 	const FName DefaultInputLabel = TEXT("In");
diff --git a/Engine/Plugins/PCG/Source/PCG/Public/PCGElement.h b/Engine/Plugins/PCG/Source/PCG/Public/PCGElement.h
--- a/Engine/Plugins/PCG/Source/PCG/Public/PCGElement.h
+++ b/Engine/Plugins/PCG/Source/PCG/Public/PCGElement.h
@@ -62,19 +62,19 @@ namespace EPCGElementLogMode
 #define PCGE_LOG_C(Verbosity, LogMode, CustomContext, Message) PCGE_LOG_BASE(Verbosity, CustomContext, Message)
 #endif
 
-enum EPCGElementExecutionLoopMode : uint8
+enum class EPCGElementExecutionLoopMode : uint8
 {
 	/** Not a trivial input -> output mapping, with respect to caching. */
 	NotALoop,
 	/** Loops on (singular) required pin. */
-	SinglePrimaryPin, 
+	SinglePrimaryPin,
 	/** Loops on matching indices on required pin(s). */
 	MatchingPrimaryPins,
 	/** Cartesian loop on required pins. */
 	// CartesianPins // TODO
 };
 
-enum EPCGCachingStatus : uint8
+enum class EPCGCachingStatus : uint8
 {
 	NotCacheable,
 	NotInCache,
