Implement native Apply Hierarchy.

diff --git a/Engine/Plugins/PCG/Source/PCG/Private/Elements/PCGApplyHierarchy.cpp b/Engine/Plugins/PCG/Source/PCG/Private/Elements/PCGApplyHierarchy.cpp
new file mode 100644
--- /dev/null
+++ b/Engine/Plugins/PCG/Source/PCG/Private/Elements/PCGApplyHierarchy.cpp
@@ -0,0 +1,525 @@
+// Copyright Epic Games, Inc. All Rights Reserved.
+
+#include "Elements/PCGApplyHierarchy.h"
+
+#include "Data/PCGBasePointData.h"
+#include "Helpers/PCGAsync.h"
+#include "Metadata/PCGMetadataPartitionCommon.h"
+#include "Metadata/Accessors/PCGAttributeAccessorHelpers.h"
+#include "Metadata/Accessors/PCGCustomAccessor.h"
+
+#define LOCTEXT_NAMESPACE "PCGApplyHierarchyElement"
+
+FPCGElementPtr UPCGApplyHierarchySettings::CreateElement() const
+{
+	return MakeShared<FPCGApplyHierarchyElement>();
+}
+
+bool FPCGApplyHierarchyElement::PrepareDataInternal(FPCGContext* InContext) const
+{
+	TRACE_CPUPROFILER_EVENT_SCOPE(FPCGApplyHierarchyElement::PrepareData);
+
+	const UPCGApplyHierarchySettings* Settings = InContext->GetInputSettings<UPCGApplyHierarchySettings>();
+	check(Settings);
+
+	ContextType* Context = static_cast<ContextType*>(InContext);
+	check(Context);
+
+	TArray<FPCGTaggedData> Inputs = Context->InputData.GetInputsByPin(PCGPinConstants::DefaultInputLabel);
+
+	Context->InitializePerExecutionState();
+
+	Context->InitializePerIterationStates(Inputs.Num(), [&Inputs, Settings, Context](IterStateType& OutState, const ExecStateType&, const uint32 IterationIndex)
+	{
+		const FPCGTaggedData& Input = Inputs[IterationIndex];
+		const UPCGBasePointData* InputData = Cast<const UPCGBasePointData>(Input.Data);
+
+		if (!InputData || InputData->GetNumPoints() == 0)
+		{
+			return EPCGTimeSliceInitResult::NoOperation;
+		}
+
+		UPCGBasePointData* OutputData = FPCGContext::NewPointData_AnyThread(Context);
+		OutputData->InitializeFromData(InputData);
+		OutputData->SetNumPoints(InputData->GetNumPoints());
+
+		// Allocate property we are going to modify
+		EPCGPointNativeProperties PropertiesToAllocate = EPCGPointNativeProperties::Transform;
+
+		// If the data doesn't support parent, also allocate any currently used properties
+		if (!OutputData->HasSpatialDataParent())
+		{
+			PropertiesToAllocate |= InputData->GetAllocatedProperties();
+		}
+
+		OutputData->AllocateProperties(PropertiesToAllocate);
+
+		OutState.InputData = InputData;
+		OutState.OutputData = OutputData;
+		OutState.OutputDataIndex = Context->OutputData.TaggedData.Num();
+		Context->OutputData.TaggedData.Add_GetRef(Input).Data = OutputData;
+		
+		// Create point index & parent index accessors
+		// Implementation note: we could have used PCGMetadataElementCommon::ApplyOnAccessorRange here, as we're doing a transformation
+		// but it would require to do the accessor & keys here.
+		auto PrepareAccessorAndKeys = [Context, InputData] (const FPCGAttributePropertyInputSelector& InputSelector, TUniquePtr<const IPCGAttributeAccessor>& OutAccessor, TUniquePtr<const IPCGAttributeAccessorKeys>& OutAccessorKeys) -> bool
+		{
+			OutAccessor = PCGAttributeAccessorHelpers::CreateConstAccessor(InputData, InputSelector);
+			OutAccessorKeys = PCGAttributeAccessorHelpers::CreateConstKeys(InputData, InputSelector);
+
+			if (!OutAccessor || !OutAccessorKeys)
+			{
+				PCGLog::Metadata::LogFailToCreateAccessorError(InputSelector, Context);
+				return false;
+			}
+
+			if (!PCG::Private::IsConstructible(OutAccessor->GetUnderlyingType(), PCG::Private::MetadataTypes<int32>::Id))
+			{
+				PCGLog::Metadata::LogFailToGetAttributeError<int32>(InputSelector, OutAccessor.Get(), Context);
+				return false;
+			}
+
+			return true;
+		};
+
+		const FPCGAttributePropertyInputSelector PointIndexSelector = Settings->PointIndexAttribute.CopyAndFixLast(InputData);
+		if (!PrepareAccessorAndKeys(PointIndexSelector, OutState.PointIndexAccessor, OutState.PointIndexKeys))
+		{
+			return EPCGTimeSliceInitResult::NoOperation;
+		}
+
+		const FPCGAttributePropertyInputSelector ParentIndexSelector = Settings->ParentIndexAttribute.CopyAndFixLast(InputData);
+		if (!PrepareAccessorAndKeys(ParentIndexSelector, OutState.ParentIndexAccessor, OutState.ParentIndexKeys))
+		{
+			return EPCGTimeSliceInitResult::NoOperation;
+		}
+
+		if (OutState.PointIndexKeys->GetNum() != OutState.ParentIndexKeys->GetNum())
+		{
+			PCGLog::LogErrorOnGraph(LOCTEXT("CardinalityMismatch", "Point index and Parent index properties do not have the same cardinality."), Context);
+			return EPCGTimeSliceInitResult::NoOperation;
+		}
+
+		// Create hierarchy depth keys (TODO: if required)
+		// Which we'll use to sort the partition by depth (from 0 to N)
+		OutState.HierarchyDepthSelector = Settings->HierarchyDepthAttribute.CopyAndFixLast(InputData);
+		if (!PrepareAccessorAndKeys(OutState.HierarchyDepthSelector, OutState.HierarchyDepthAccessor, OutState.HierarchyDepthKeys))
+		{
+			return EPCGTimeSliceInitResult::NoOperation;
+		}
+
+		if (Settings->ApplyParentRotation == EPCGApplyHierarchyOption::Always || Settings->ApplyParentRotation == EPCGApplyHierarchyOption::Never)
+		{
+			OutState.ApplyRotationAccessor = MakeUnique<FPCGConstantValueAccessor<bool>>(true);
+			OutState.bInvertApplyRotation = Settings->ApplyParentRotation == EPCGApplyHierarchyOption::Never;
+		}
+		else
+		{
+			FPCGAttributePropertyInputSelector ApplyRotationSelector = Settings->ApplyParentRotationAttribute.CopyAndFixLast(InputData);
+			OutState.ApplyRotationAccessor = PCGAttributeAccessorHelpers::CreateConstAccessor(InputData, ApplyRotationSelector);
+			OutState.bInvertApplyRotation = Settings->ApplyParentScale == EPCGApplyHierarchyOption::OptOutByAttribute;
+
+			if (!OutState.ApplyRotationAccessor)
+			{
+				PCGLog::Metadata::LogFailToCreateAccessorError(ApplyRotationSelector, Context);
+				return EPCGTimeSliceInitResult::NoOperation;
+			}
+		}
+
+		if (Settings->ApplyParentScale == EPCGApplyHierarchyOption::Always || Settings->ApplyParentScale == EPCGApplyHierarchyOption::Never)
+		{
+			OutState.ApplyScaleAccessor = MakeUnique<FPCGConstantValueAccessor<bool>>(true);
+			OutState.bInvertApplyScale = Settings->ApplyParentScale == EPCGApplyHierarchyOption::Never;
+		}
+		else
+		{
+			FPCGAttributePropertyInputSelector ApplyScaleSelector = Settings->ApplyParentScaleAttribute.CopyAndFixLast(InputData);
+			OutState.ApplyScaleAccessor = PCGAttributeAccessorHelpers::CreateConstAccessor(InputData, ApplyScaleSelector);
+			OutState.bInvertApplyScale = Settings->ApplyParentScale == EPCGApplyHierarchyOption::OptOutByAttribute;
+
+			if (!OutState.ApplyScaleAccessor)
+			{
+				PCGLog::Metadata::LogFailToCreateAccessorError(ApplyScaleSelector, Context);
+				return EPCGTimeSliceInitResult::NoOperation;
+			}
+		}
+
+		// Finally, get the relative transforms
+		FPCGAttributePropertyInputSelector RelativeTransformSelector = Settings->RelativeTransformAttribute.CopyAndFixLast(InputData);
+		OutState.RelativeTransformAccessor = PCGAttributeAccessorHelpers::CreateConstAccessor(InputData, RelativeTransformSelector);
+
+		if (!OutState.RelativeTransformAccessor)
+		{
+			PCGLog::Metadata::LogFailToCreateAccessorError(RelativeTransformSelector, Context);
+			return EPCGTimeSliceInitResult::NoOperation;
+		}
+		
+		if (!PCG::Private::IsBroadcastableOrConstructible(OutState.RelativeTransformAccessor->GetUnderlyingType(), PCG::Private::MetadataTypes<FTransform>::Id))
+		{
+			PCGLog::Metadata::LogFailToGetAttributeError<FTransform>(RelativeTransformSelector, OutState.RelativeTransformAccessor.Get(), Context);
+			return EPCGTimeSliceInitResult::NoOperation;
+		}
+		
+		return EPCGTimeSliceInitResult::Success;
+	});
+
+	return true;
+}
+
+bool FPCGApplyHierarchyElement::ExecuteInternal(FPCGContext* InContext) const
+{
+	TRACE_CPUPROFILER_EVENT_SCOPE(FPCGApplyHierarchyElement::Execute);
+
+	ContextType* TimeSlicedContext = static_cast<ContextType*>(InContext);
+	check(TimeSlicedContext);
+
+	if (!TimeSlicedContext->DataIsPreparedForExecution())
+	{
+		return true;
+	}
+
+	if (TimeSlicedContext->GetExecutionStateResult() == EPCGTimeSliceInitResult::NoOperation)
+	{
+		TimeSlicedContext->OutputData = TimeSlicedContext->InputData;
+		return true;
+	}
+
+	const UPCGApplyHierarchySettings* Settings = TimeSlicedContext->GetInputSettings<UPCGApplyHierarchySettings>();
+	check(Settings);
+
+	return ExecuteSlice(TimeSlicedContext, [Settings](ContextType* Context, const ExecStateType& ExecState, IterStateType& IterState, const uint32 IterIndex)
+	{
+		if (Context->GetIterationStateResult(IterIndex) == EPCGTimeSliceInitResult::NoOperation)
+		{
+			return true;
+		}
+
+		constexpr int32 InvalidParentIndex = -2;
+
+		// 1. Get point indices & parent indices
+		// 2. Build index to parent index map
+		if (!IterState.bParentMappingDone)
+		{
+			TRACE_CPUPROFILER_EVENT_SCOPE(FPCGApplyHierarchyElement::Execute::ParentMapping);
+			TArray<int32> PointIndices;
+			TArray<int32> ParentIndices;
+
+			PointIndices.SetNumUninitialized(IterState.PointIndexKeys->GetNum());
+			ParentIndices.SetNumUninitialized(IterState.ParentIndexKeys->GetNum());
+
+			auto GetPointIndexAndParentIndex = [&IterState, &PointIndices, &ParentIndices](int32 StartReadIndex, int32 StartWriteIndex, int32 Count)
+			{
+				IterState.PointIndexAccessor->GetRange<int32>(MakeArrayView(&PointIndices[StartWriteIndex], Count), StartReadIndex, *IterState.PointIndexKeys, EPCGAttributeAccessorFlags::AllowConstructible);
+				IterState.ParentIndexAccessor->GetRange<int32>(MakeArrayView(&ParentIndices[StartWriteIndex], Count), StartReadIndex, *IterState.ParentIndexKeys, EPCGAttributeAccessorFlags::AllowConstructible);
+				return Count;
+			};
+
+			FPCGAsync::AsyncProcessingOneToOneRangeEx(
+				&Context->AsyncState,
+				PointIndices.Num(),
+				[]() {},
+				GetPointIndexAndParentIndex,
+				/*bEnableTimeSlicing=*/false,
+				/*ChunkSize*/256);
+
+			TMap<int32, int32> PointIndexToIndexMap;
+			PointIndexToIndexMap.Reserve(PointIndices.Num());
+			for (int Index = 0; Index < PointIndices.Num(); ++Index)
+			{
+				PointIndexToIndexMap.Add(PointIndices[Index], Index);
+			}
+
+			bool bHasPointsWithInvalidParent = false;
+
+			IterState.ParentIndices.SetNumUninitialized(ParentIndices.Num());
+			for (int Index = 0; Index < ParentIndices.Num(); ++Index)
+			{
+				const int32 OriginalParentIndex = ParentIndices[Index];
+				if (OriginalParentIndex != INDEX_NONE)
+				{
+					const int32* ParentIndex = PointIndexToIndexMap.Find(ParentIndices[Index]);
+					if (ParentIndex)
+					{
+						IterState.ParentIndices[Index] = *ParentIndex;
+					}
+					else
+					{
+						IterState.ParentIndices[Index] = InvalidParentIndex;
+						bHasPointsWithInvalidParent = true;
+					}
+				}
+				else
+				{
+					ParentIndices[Index] = INDEX_NONE;
+				}
+			}
+
+			IterState.bHasPointsWithInvalidParent = bHasPointsWithInvalidParent;
+			IterState.bParentMappingDone = true;
+
+			if (Context->ShouldStop())
+			{
+				return false;
+			}
+		}
+
+		// 3. * TODO * - could build hierarchy depth at this stage instead of relying on an attribute -
+		// 4. Partition on hierarchy depth
+		if (!IterState.bHierarchyDepthPartitionDone)
+		{
+			TRACE_CPUPROFILER_EVENT_SCOPE(FPCGApplyHierarchyElement::Execute::HierarchyDepthPartition);
+			IterState.HierarchyPartition = PCGMetadataPartitionCommon::AttributeGenericPartition(IterState.InputData, IterState.HierarchyDepthSelector, Context);
+			if (IterState.HierarchyPartition.IsEmpty())
+			{
+				// Error logged in the partitioning code
+				return true;
+			}
+
+			// Build an indirection and perform some slight validation
+			TArray<int32, TInlineAllocator<64>> FirstIndexPerPartition;
+			Algo::Transform(IterState.HierarchyPartition, FirstIndexPerPartition, [](const TArray<int32>& InPartition) { return InPartition.IsEmpty() ? INDEX_NONE : InPartition[0]; });
+
+			TUniquePtr<IPCGAttributeAccessorKeys> FirstIndexPerPartitionKeys = MakeUnique<FPCGAttributeAccessorKeysPointsSubset>(IterState.InputData, FirstIndexPerPartition);
+
+			TArray<int32> HierarchyDepthPerPartition;
+			HierarchyDepthPerPartition.SetNumUninitialized(IterState.HierarchyPartition.Num());
+			IterState.HierarchyDepthAccessor->GetRange<int32>(HierarchyDepthPerPartition, 0, *FirstIndexPerPartitionKeys, EPCGAttributeAccessorFlags::AllowConstructible);
+
+			TArray<TPair<int32, int32>> PartitionRepresentatives;
+			PartitionRepresentatives.Reserve(HierarchyDepthPerPartition.Num());
+
+			for (int PartitionIndex = 0; PartitionIndex < HierarchyDepthPerPartition.Num(); ++PartitionIndex)
+			{
+				PartitionRepresentatives.Emplace(PartitionIndex, HierarchyDepthPerPartition[PartitionIndex]);
+			}
+
+			PartitionRepresentatives.Sort([](const TPair<int32, int32>& A, const TPair<int32, int32>& B) { return (A.Value < B.Value) || (A.Value == B.Value && A.Key < B.Key); });
+
+			IterState.HierarchyPartitionOrder.Reserve(PartitionRepresentatives.Num());
+
+			for(const auto& [PartitionIndex, PartitionDepth] : PartitionRepresentatives)
+			{
+				if (IterState.HierarchyPartitionOrder.Num() == PartitionDepth)
+				{
+					IterState.HierarchyPartitionOrder.Add(PartitionIndex);
+				}
+				else
+				{
+					// At this point, anything downstream is going to be broken. We'll go ahead and mark the parents of these as invalid
+					for (const int32& InvalidPointIndex : IterState.HierarchyPartition[PartitionIndex])
+					{
+						IterState.ParentIndices[InvalidPointIndex] = InvalidParentIndex;
+					}
+
+					IterState.bHasPointsWithInvalidParent = true;
+				}
+			}
+
+			IterState.bHierarchyDepthPartitionDone = true;
+
+			if (Context->ShouldStop())
+			{
+				return false;
+			}
+		}
+
+		// 5. For all depths partitions, compute & write transform to output data.
+		while (IterState.CurrentDepth < IterState.HierarchyPartitionOrder.Num())
+		{
+			// Perform current depth iteration
+			{
+				TRACE_CPUPROFILER_EVENT_SCOPE(FPCGApplyHierarchyElement::Execute::ComputeTransforms);
+				constexpr int32 ChunkSize = 64;
+				const TArray<int32>& CurrentDepthIndices = IterState.HierarchyPartition[IterState.HierarchyPartitionOrder[IterState.CurrentDepth]];
+				const bool bIsRoot = IterState.CurrentDepth == 0;
+
+				TUniquePtr<IPCGAttributeAccessorKeys> SubsetKeys = MakeUnique<FPCGAttributeAccessorKeysPointsSubset>(IterState.InputData, CurrentDepthIndices);
+
+				if (bIsRoot)
+				{
+					// Set transform to relative transform x current point transform
+					auto AsyncProcessFunc = [&CurrentDepthIndices, &IterState, &SubsetKeys](int32 StartReadIndex, int32 StartWriteIndex, int32 Count)
+					{
+						if(!IterState.OutputData->HasSpatialDataParent())
+						{
+							const TArrayView<const int32> IndicesView = MakeArrayView(&CurrentDepthIndices[StartReadIndex], Count);
+							IterState.InputData->CopyPointsTo(IterState.OutputData, IndicesView, IndicesView);
+						}
+
+						TPCGValueRange<FTransform> ReadWriteTransformRange = IterState.OutputData->GetTransformValueRange();
+
+						TArray<FTransform, TInlineAllocator<ChunkSize>> RelativeTransforms;
+						RelativeTransforms.SetNumUninitialized(Count);
+						IterState.RelativeTransformAccessor->GetRange<FTransform>(RelativeTransforms, StartReadIndex, *SubsetKeys, EPCGAttributeAccessorFlags::AllowConstructible);
+
+						for(int32 Index = 0; Index < Count; ++Index)
+						{
+							const int32 CurrentPointIndex = CurrentDepthIndices[Index + StartReadIndex];
+							FTransform& Transform = ReadWriteTransformRange[CurrentPointIndex];
+							const FTransform& RelativeTransform = RelativeTransforms[Index];
+
+							Transform = RelativeTransform * Transform;
+						}
+
+						return Count;
+					};
+
+					FPCGAsync::AsyncProcessingOneToOneRangeEx(
+						&Context->AsyncState,
+						CurrentDepthIndices.Num(),
+						[]() {},
+						AsyncProcessFunc,
+						/*bEnableTimeSlicing=*/false);
+				}
+				else
+				{
+					bool bHasPointsWithInvalidParent = false;
+					// Set transform to relative transform x parent transform (& apply options)
+					auto AsyncProcessFunc = [&CurrentDepthIndices, &IterState, &SubsetKeys, &bHasPointsWithInvalidParent](int32 StartReadIndex, int32 StartWriteIndex, int32 Count)
+					{
+						if (!IterState.OutputData->HasSpatialDataParent())
+						{
+							const TArrayView<const int32> IndicesView = MakeArrayView(&CurrentDepthIndices[StartReadIndex], Count);
+							IterState.InputData->CopyPointsTo(IterState.OutputData, IndicesView, IndicesView);
+						}
+
+						check(StartReadIndex == StartWriteIndex);
+						TPCGValueRange<FTransform> ReadWriteTransformRange = IterState.OutputData->GetTransformValueRange();
+
+						TArray<FTransform, TInlineAllocator<ChunkSize>> RelativeTransforms;
+						RelativeTransforms.SetNumUninitialized(Count);
+						IterState.RelativeTransformAccessor->GetRange<FTransform>(RelativeTransforms, StartReadIndex, *SubsetKeys, EPCGAttributeAccessorFlags::AllowConstructible);
+
+						TArray<bool, TInlineAllocator<ChunkSize>> ApplyParentRotation;
+						ApplyParentRotation.SetNumUninitialized(Count);
+						IterState.ApplyRotationAccessor->GetRange<bool>(ApplyParentRotation, StartReadIndex, *SubsetKeys, EPCGAttributeAccessorFlags::AllowConstructible);
+
+						TArray<bool, TInlineAllocator<ChunkSize>> ApplyParentScale;
+						ApplyParentScale.SetNumUninitialized(Count);
+						IterState.ApplyScaleAccessor->GetRange<bool>(ApplyParentScale, StartReadIndex, *SubsetKeys, EPCGAttributeAccessorFlags::AllowConstructible);
+
+						for (int32 Index = 0; Index < Count; ++Index)
+						{
+							const int32 CurrentPointIndex = CurrentDepthIndices[Index + StartReadIndex];
+							const int32 CurrentParentIndex = IterState.ParentIndices[CurrentPointIndex];
+
+							// At this point, if the current parent is invalid, or it has been marked invalid, we'll mark this point as invalid too.
+							if(CurrentParentIndex < 0 || IterState.ParentIndices[CurrentParentIndex] == InvalidParentIndex)
+							{
+								if (CurrentParentIndex != InvalidParentIndex)
+								{
+									IterState.ParentIndices[CurrentPointIndex] = InvalidParentIndex;
+									bHasPointsWithInvalidParent = true;
+								}
+								
+								continue;
+							}
+
+							FTransform& Transform = ReadWriteTransformRange[CurrentPointIndex];
+							const FTransform& ParentTransform = ReadWriteTransformRange[CurrentParentIndex];
+							const FTransform& RelativeTransform = RelativeTransforms[Index];
+
+							Transform = RelativeTransform * ParentTransform;
+
+							// Finally, apply options (ignore scale/rotation from parent)
+							if (ApplyParentRotation[Index] == IterState.bInvertApplyScale)
+							{
+								Transform.SetRotation(RelativeTransform.GetRotation());
+							}
+
+							if (ApplyParentScale[Index] == IterState.bInvertApplyScale)
+							{
+								Transform.SetScale3D(RelativeTransform.GetScale3D());
+							}
+						}
+
+						return Count;
+					};
+
+					FPCGAsync::AsyncProcessingOneToOneRangeEx(
+						&Context->AsyncState,
+						CurrentDepthIndices.Num(),
+						[]() {},
+						AsyncProcessFunc,
+						/*bEnableTimeSlicing=*/false);
+
+					IterState.bHasPointsWithInvalidParent |= bHasPointsWithInvalidParent;
+				}
+
+				++IterState.CurrentDepth;
+			}
+
+			if ((IterState.CurrentDepth != IterState.HierarchyPartitionOrder.Num() || IterState.bHasPointsWithInvalidParent) && Context->ShouldStop())
+			{
+				return false;
+			}
+
+			// Finally, if we had points that were unparented, we need to cull them out.
+			if (IterState.CurrentDepth == IterState.HierarchyPartitionOrder.Num() && IterState.bHasPointsWithInvalidParent)
+			{
+				if (Settings->bWarnOnPointsWithInvalidParent)
+				{
+					PCGLog::LogWarningOnGraph(LOCTEXT("SomePointsHaveAnInvalidParent", "Some points have either an invalid parent index, or an invalid depth. They will be culled from the final results."), Context);
+				}
+
+				const UPCGBasePointData* OriginalData = IterState.OutputData;
+				UPCGBasePointData* FilteredData = FPCGContext::NewPointData_AnyThread(Context);
+				FPCGInitializeFromDataParams InitializeFromDataParams(OriginalData);
+
+				// Do not inherit because we will filter out the invalid points.
+				InitializeFromDataParams.bInheritSpatialData = false;
+				FilteredData->InitializeFromDataWithParams(InitializeFromDataParams);
+
+				Context->OutputData.TaggedData[IterState.OutputDataIndex].Data = FilteredData;
+				
+				auto InitializeFunc = [FilteredData, OriginalData]()
+				{
+					FilteredData->SetNumPoints(OriginalData->GetNumPoints());
+					FilteredData->AllocateProperties(OriginalData->GetAllocatedProperties());
+				};
+
+				auto AsyncProcessRangeFunc = [FilteredData, OriginalData, &IterState](int32 StartReadIndex, int32 StartWriteIndex, int32 Count)
+				{
+					const FConstPCGPointValueRanges ReadRanges(OriginalData);
+					FPCGPointValueRanges WriteRanges(FilteredData, /*bAllocate=*/false);
+
+					int32 NumWritten = 0;
+
+					for (int32 ReadIndex = StartReadIndex; ReadIndex < StartReadIndex + Count; ++ReadIndex)
+					{
+						if (IterState.ParentIndices[ReadIndex] != InvalidParentIndex)
+						{
+							WriteRanges.SetFromValueRanges(StartWriteIndex + NumWritten, ReadRanges, ReadIndex);
+							++NumWritten;
+						}
+					}
+
+					return NumWritten;
+				};
+
+				auto MoveDataRangeFunc = [FilteredData](int32 RangeStartIndex, int32 MoveToIndex, int32 NumElements)
+				{
+					FilteredData->MoveRange(RangeStartIndex, MoveToIndex, NumElements);
+				};
+
+				auto FinishedFunc = [FilteredData](int32 NumWritten)
+				{
+					FilteredData->SetNumPoints(NumWritten);
+				};
+
+				FPCGAsync::AsyncProcessingRangeEx(
+					&Context->AsyncState,
+					OriginalData->GetNumPoints(),
+					InitializeFunc,
+					AsyncProcessRangeFunc,
+					MoveDataRangeFunc,
+					FinishedFunc,
+					/*bEnableTimeSlicing=*/false);
+			}
+		}
+
+		return true;
+	});
+}
+
+#undef LOCTEXT_NAMESPACE
\ No newline at end of file
diff --git a/Engine/Plugins/PCG/Source/PCG/Public/Elements/PCGApplyHierarchy.h b/Engine/Plugins/PCG/Source/PCG/Public/Elements/PCGApplyHierarchy.h
new file mode 100644
--- /dev/null
+++ b/Engine/Plugins/PCG/Source/PCG/Public/Elements/PCGApplyHierarchy.h
@@ -0,0 +1,121 @@
+// Copyright Epic Games, Inc. All Rights Reserved.
+
+#pragma once
+
+#include "PCGSettings.h"
+#include "Elements/PCGTimeSlicedElementBase.h"
+#include "Metadata/PCGAttributePropertySelector.h"
+#include "Metadata/Accessors/PCGPropertyAccessor.h"
+
+#include "PCGApplyHierarchy.generated.h"
+
+UENUM(BlueprintType)
+enum class EPCGApplyHierarchyOption : uint8
+{
+	Always,
+	Never,
+	OptInByAttribute,
+	OptOutByAttribute
+};
+
+/** Applies hierarchy transformations based on a hierarchy depth, point index & parent index scheme.
+* This is used in the context of PCG Data Assets that have these fields by default.
+*/
+UCLASS(MinimalAPI, BlueprintType, ClassGroup = (Procedural))
+class UPCGApplyHierarchySettings : public UPCGSettings
+{
+	GENERATED_BODY()
+
+	//~Begin UPCGSettings interface
+public:
+#if WITH_EDITOR
+	virtual FName GetDefaultNodeName() const override { return FName(TEXT("ApplyHierarchy")); }
+	virtual FText GetDefaultNodeTitle() const override { return NSLOCTEXT("PCGApplyHierarchyElement", "NodeTitle", "Apply Hierarchy"); }
+	virtual EPCGSettingsType GetType() const override { return EPCGSettingsType::PointOps; }
+#endif
+
+protected:
+	virtual TArray<FPCGPinProperties> InputPinProperties() const override { return DefaultPointInputPinProperties(); }
+	virtual TArray<FPCGPinProperties> OutputPinProperties() const override { return DefaultPointOutputPinProperties(); }
+	virtual FPCGElementPtr CreateElement() const override;
+	//~End UPCGSettings interface
+
+public:
+	UPROPERTY(BlueprintReadWrite, EditAnywhere, Category = Settings, meta = (PCG_Overridable))
+	FPCGAttributePropertyInputSelector PointIndexAttribute;
+
+	UPROPERTY(BlueprintReadWrite, EditAnywhere, Category = Settings, meta = (PCG_Overridable))
+	FPCGAttributePropertyInputSelector ParentIndexAttribute;
+
+	UPROPERTY(BlueprintReadWrite, EditAnywhere, Category = Settings, meta = (PCG_Overridable))
+	FPCGAttributePropertyInputSelector HierarchyDepthAttribute;
+
+	UPROPERTY(BlueprintReadWrite, EditAnywhere, Category = Settings, meta = (PCG_Overridable))
+	FPCGAttributePropertyInputSelector RelativeTransformAttribute;
+
+	UPROPERTY(BlueprintReadWrite, EditAnywhere, Category = Settings, meta = (PCG_Overridable))
+	EPCGApplyHierarchyOption ApplyParentRotation = EPCGApplyHierarchyOption::Always;
+
+	UPROPERTY(BlueprintReadWrite, EditAnywhere, Category = Settings, meta = (EditCondition = "ApplyParentRotation==EPCGApplyHierarchyOption::OptInByAttribute||ApplyParentRotation==EPCGApplyHierarchyOption::OptOutByAttribute", EditConditionHides, PCG_Overridable))
+	FPCGAttributePropertyInputSelector ApplyParentRotationAttribute;
+
+	UPROPERTY(BlueprintReadWrite, EditAnywhere, Category = Settings, meta = (PCG_Overridable))
+	EPCGApplyHierarchyOption ApplyParentScale = EPCGApplyHierarchyOption::Always;
+
+	UPROPERTY(BlueprintReadWrite, EditAnywhere, Category = Settings, meta = (EditCondition = "ApplyParentScale==EPCGApplyHierarchyOption::OptInByAttribute||ApplyParentScale==EPCGApplyHierarchyOption::OptOutByAttribute", EditConditionHides, PCG_Overridable))
+	FPCGAttributePropertyInputSelector ApplyParentScaleAttribute;
+
+	UPROPERTY(BlueprintReadWrite, EditAnywhere, Category = Advanced, meta = (PCG_Overridable))
+	bool bWarnOnPointsWithInvalidParent = true;
+};
+
+namespace PCGApplyHierarchyElement
+{
+	struct IterationState
+	{
+		const UPCGBasePointData* InputData = nullptr;
+		UPCGBasePointData* OutputData = nullptr;
+		int OutputDataIndex = INDEX_NONE;
+		
+
+		TUniquePtr<const IPCGAttributeAccessor> PointIndexAccessor;
+		TUniquePtr<const IPCGAttributeAccessorKeys> PointIndexKeys;
+
+		TUniquePtr<const IPCGAttributeAccessor> ParentIndexAccessor;
+		TUniquePtr<const IPCGAttributeAccessorKeys> ParentIndexKeys;
+
+		FPCGAttributePropertyInputSelector HierarchyDepthSelector;
+		TUniquePtr<const IPCGAttributeAccessor> HierarchyDepthAccessor;
+		TUniquePtr<const IPCGAttributeAccessorKeys> HierarchyDepthKeys;
+		
+		TUniquePtr<const IPCGAttributeAccessor> RelativeTransformAccessor;
+		TUniquePtr<const IPCGAttributeAccessorKeys> RelativeTransformKeys;
+		
+		TUniquePtr<const IPCGAttributeAccessor> ApplyRotationAccessor;
+		bool bInvertApplyRotation = false;
+
+		TUniquePtr<const IPCGAttributeAccessor> ApplyScaleAccessor;
+		bool bInvertApplyScale = false;
+		
+		// Built during execution
+		TArray<int32> ParentIndices;
+		TArray<TArray<int32>> HierarchyPartition;
+		TArray<int32> HierarchyPartitionOrder;
+
+		// Execution state variables
+		bool bParentMappingDone = false;
+		bool bHierarchyDepthPartitionDone = false;
+		int CurrentDepth = 0;
+
+		bool bHasPointsWithInvalidParent = false;
+	};
+}
+
+class FPCGApplyHierarchyElement : public TPCGTimeSlicedElementBase<PCGTimeSlice::FEmptyStruct, PCGApplyHierarchyElement::IterationState>
+{
+protected:
+	virtual bool PrepareDataInternal(FPCGContext* InContext) const override;
+	virtual bool ExecuteInternal(FPCGContext* InContext) const override;
+	virtual EPCGElementExecutionLoopMode ExecutionLoopMode(const UPCGSettings* Settings) const override { return EPCGElementExecutionLoopMode::SinglePrimaryPin; }
+	virtual bool SupportsBasePointDataInputs(FPCGContext* InContext) const override { return true; }
+};
\ No newline at end of file
