Exposed the object path in the PCG data asset to BP to make the Update feature a bit more flexible

#rnx

diff --git a/Engine/Plugins/PCG/Source/PCG/Public/PCGDataAsset.h b/Engine/Plugins/PCG/Source/PCG/Public/PCGDataAsset.h
--- a/Engine/Plugins/PCG/Source/PCG/Public/PCGDataAsset.h
+++ b/Engine/Plugins/PCG/Source/PCG/Public/PCGDataAsset.h
@@ -41,7 +41,7 @@ public:
 	TSubclassOf<UPCGSettings> SettingsClass;
 
 	/** Reference to originating object (often will be a level) */
-	UPROPERTY(BlueprintReadOnly, VisibleAnywhere, Category = Data, AssetRegistrySearchable, meta = (NoResetToDefault))
+	UPROPERTY(BlueprintReadWrite, VisibleAnywhere, Category = Data, AssetRegistrySearchable, meta = (NoResetToDefault))
 	FSoftObjectPath ObjectPath;
 
 	/** Custom node color to use for this specific asset */
