[PCG] Throttle Debug Object Tree refreshes which invalidate the entire graph editor window and cause big performance issues.

#rnx

diff --git a/Engine/Plugins/PCG/Source/PCGEditor/Private/SPCGEditorGraphDebugObjectTree.cpp b/Engine/Plugins/PCG/Source/PCGEditor/Private/SPCGEditorGraphDebugObjectTree.cpp
--- a/Engine/Plugins/PCG/Source/PCGEditor/Private/SPCGEditorGraphDebugObjectTree.cpp
+++ b/Engine/Plugins/PCG/Source/PCGEditor/Private/SPCGEditorGraphDebugObjectTree.cpp
@@ -508,7 +508,7 @@ void SPCGEditorGraphDebugObjectTree::Tick(const FGeometry& AllottedGeometry, con
 {
 	SCompoundWidget::Tick(AllottedGeometry, InCurrentTime, InDeltaTime);
 
-	if (bNeedsRefresh)
+	if (bNeedsRefresh && InCurrentTime >= NextRefreshTime)
 	{
 		// Updating the tree while the inspected component is generating can be bad as the selection can
 		// be lost. Don't change the tree if we're inspecting something that is generating.
@@ -517,6 +517,9 @@ void SPCGEditorGraphDebugObjectTree::Tick(const FGeometry& AllottedGeometry, con
 		{
 			bNeedsRefresh = false;
 			RefreshTree();
+
+			constexpr double RefreshCooldownTime = 0.25f;
+			NextRefreshTime = InCurrentTime + RefreshCooldownTime;
 		}
 	}
 
@@ -526,6 +529,11 @@ void SPCGEditorGraphDebugObjectTree::Tick(const FGeometry& AllottedGeometry, con
 	}
 }
 
+void SPCGEditorGraphDebugObjectTree::RequestRefresh()
+{
+	bNeedsRefresh = true;
+}
+
 void SPCGEditorGraphDebugObjectTree::OnObjectsReplaced(const TMap<UObject*, UObject*>& ReplacementMap)
 {
 	for (FPCGEditorGraphDebugObjectItemPtr& GraphItem : AllGraphItems)
diff --git a/Engine/Plugins/PCG/Source/PCGEditor/Private/SPCGEditorGraphDebugObjectTree.h b/Engine/Plugins/PCG/Source/PCGEditor/Private/SPCGEditorGraphDebugObjectTree.h
--- a/Engine/Plugins/PCG/Source/PCGEditor/Private/SPCGEditorGraphDebugObjectTree.h
+++ b/Engine/Plugins/PCG/Source/PCGEditor/Private/SPCGEditorGraphDebugObjectTree.h
@@ -265,7 +265,7 @@ public:
 
 	virtual void Tick(const FGeometry& AllottedGeometry, const double InCurrentTime, const float InDeltaTime) override;
 
-	void RequestRefresh() { bNeedsRefresh = true; };
+	void RequestRefresh();
 
 	void SetNodeBeingInspected(const UPCGNode* InPCGNode);
 
@@ -334,6 +334,8 @@ private:
 
 	bool bNeedsRefresh = false;
 
+	double NextRefreshTime = 0.0;
+
 	/** Latest value for IsSetDebugObjectFromSelectionButtonEnabled */
 	TOptional<bool> IsSetDebugObjectFromSelectionEnabled;
 
