[PCG] Structural changes to allow nodes to emit multiple GPU kernels.

Early in the GPU compiler flow the GPU-compatible tasks are invited to provide one or more kernels and their connections. These are then expanded into one graph task per kernel, and then the GPU compiler proceeds as before.

#rnx

diff --git a/Engine/Plugins/PCG/Source/PCG/Internal/Compute/PCGComputeCommon.h b/Engine/Plugins/PCG/Source/PCG/Internal/Compute/PCGComputeCommon.h
--- a/Engine/Plugins/PCG/Source/PCG/Internal/Compute/PCGComputeCommon.h
+++ b/Engine/Plugins/PCG/Source/PCG/Internal/Compute/PCGComputeCommon.h
@@ -6,6 +6,7 @@
 
 #include "PCGComputeCommon.generated.h"
 
+class UPCGComputeKernel;
 class UPCGData;
 class UPCGDataBinding;
 class UPCGSettings;
@@ -123,3 +124,51 @@ namespace PCGComputeHelpers
 	/** Produces the data interface name of a data label resolver. */
 	FString GetDataLabelResolverName(FName InPinLabel);
 }
+
+/** A by-label reference to a pin, used for wiring kernels within a node. */
+struct FPCGPinReference
+{
+	/** Reference a pin by label only, used for referencing node pins. */
+	explicit FPCGPinReference(FName InLabel)
+		: Kernel(nullptr)
+		, Label(InLabel)
+	{
+	}
+
+	/** Reference a pin by kernel and label. */
+	explicit FPCGPinReference(const UPCGComputeKernel* InKernel, FName InLabel)
+		: Kernel(InKernel)
+		, Label(InLabel)
+	{
+	}
+
+	bool operator==(const FPCGPinReference& Other) const
+	{
+		return Label == Other.Label
+			&& Kernel == Other.Kernel;
+	}
+
+	/** Associated kernel. If null then compiler will look for pin on owning node. */
+	const UPCGComputeKernel* Kernel = nullptr;
+
+	/** Pin label. */
+	FName Label;
+};
+
+uint32 GetTypeHash(const FPCGPinReference& In);
+
+/** A connection for wiring kernels within a node. */
+struct FPCGKernelEdge
+{
+	FPCGKernelEdge(const FPCGPinReference& InUpstreamPin, const FPCGPinReference& InDownstreamPin)
+		: UpstreamPin(InUpstreamPin)
+		, DownstreamPin(InDownstreamPin)
+	{
+	}
+
+	bool IsConnectedToNodeInput() const { return UpstreamPin.Kernel == nullptr; }
+	bool IsConnectedToNodeOutput() const { return DownstreamPin.Kernel == nullptr; }
+
+	FPCGPinReference UpstreamPin;
+	FPCGPinReference DownstreamPin;
+};
diff --git a/Engine/Plugins/PCG/Source/PCG/Private/Compute/Elements/PCGCustomHLSL.cpp b/Engine/Plugins/PCG/Source/PCG/Private/Compute/Elements/PCGCustomHLSL.cpp
--- a/Engine/Plugins/PCG/Source/PCG/Private/Compute/Elements/PCGCustomHLSL.cpp
+++ b/Engine/Plugins/PCG/Source/PCG/Private/Compute/Elements/PCGCustomHLSL.cpp
@@ -324,12 +324,22 @@ bool UPCGCustomHLSLSettings::IsShaderTextReadOnly() const
 	return false;
 }
 
-UPCGComputeKernel* UPCGCustomHLSLSettings::CreateKernel(FPCGGPUCompilationContext& InOutContext, UObject* InObjectOuter) const
+void UPCGCustomHLSLSettings::CreateKernels(FPCGGPUCompilationContext& InOutContext, UObject* InObjectOuter, TArray<UPCGComputeKernel*>& OutKernels, TArray<FPCGKernelEdge>& OutEdges) const
 {
 	UPCGCustomHLSLKernel* Kernel = InOutContext.NewObject_AnyThread<UPCGCustomHLSLKernel>(InObjectOuter);
 	Kernel->Initialize(this, bDumpDataDescriptions);
+	OutKernels.Add(Kernel);
 
-	return Kernel;
+	// Connect node pins to kernel pins
+	for (const FPCGPinProperties& Pin : InputPins)
+	{
+		OutEdges.Emplace(FPCGPinReference(Pin.Label), FPCGPinReference(Kernel, Pin.Label));
+	}
+
+	for (const FPCGPinPropertiesGPU& Pin : OutputPins)
+	{
+		OutEdges.Emplace(FPCGPinReference(Kernel, Pin.Label), FPCGPinReference(Pin.Label));
+	}
 }
 #endif // WITH_EDITOR
 
diff --git a/Engine/Plugins/PCG/Source/PCG/Private/Compute/Elements/PCGCustomHLSL.h b/Engine/Plugins/PCG/Source/PCG/Private/Compute/Elements/PCGCustomHLSL.h
--- a/Engine/Plugins/PCG/Source/PCG/Private/Compute/Elements/PCGCustomHLSL.h
+++ b/Engine/Plugins/PCG/Source/PCG/Private/Compute/Elements/PCGCustomHLSL.h
@@ -57,7 +57,7 @@ public:
 	virtual FText GetDefaultNodeTitle() const override { return NSLOCTEXT("PCGCustomHLSLElement", "NodeTitle", "Custom HLSL"); }
 	virtual FText GetNodeTooltipText() const override { return NSLOCTEXT("PCGCustomHLSLElement", "NodeTooltip", "[EXPERIMENTAL] Produces a HLSL compute shader which will be executed on the GPU."); }
 	virtual EPCGSettingsType GetType() const override { return EPCGSettingsType::GPU; }
-	UPCGComputeKernel* CreateKernel(FPCGGPUCompilationContext& InOutContext, UObject* InObjectOuter) const override;
+	void CreateKernels(FPCGGPUCompilationContext& InOutContext, UObject* InObjectOuter, TArray<UPCGComputeKernel*>& OutKernels, TArray<FPCGKernelEdge>& OutEdges) const override;
 
 	virtual TArray<FPCGPreConfiguredSettingsInfo> GetPreconfiguredInfo() const override;
 	virtual bool OnlyExposePreconfiguredSettings() const override { return true; }
diff --git a/Engine/Plugins/PCG/Source/PCG/Private/Compute/PCGComputeCommon.cpp b/Engine/Plugins/PCG/Source/PCG/Private/Compute/PCGComputeCommon.cpp
--- a/Engine/Plugins/PCG/Source/PCG/Private/Compute/PCGComputeCommon.cpp
+++ b/Engine/Plugins/PCG/Source/PCG/Private/Compute/PCGComputeCommon.cpp
@@ -175,3 +175,8 @@ namespace PCGComputeHelpers
 		return FString::Format(TEXT("{0}_DataResolver"), { *InPinLabel.ToString() });
 	}
 }
+
+uint32 GetTypeHash(const FPCGPinReference& In)
+{
+	return HashCombine(/*GetTypeHash(In.TaskId),*/ PointerHash(In.Kernel), GetTypeHash(In.Label));
+}
diff --git a/Engine/Plugins/PCG/Source/PCG/Private/Elements/Metadata/PCGMetadataPartition.cpp b/Engine/Plugins/PCG/Source/PCG/Private/Elements/Metadata/PCGMetadataPartition.cpp
--- a/Engine/Plugins/PCG/Source/PCG/Private/Elements/Metadata/PCGMetadataPartition.cpp
+++ b/Engine/Plugins/PCG/Source/PCG/Private/Elements/Metadata/PCGMetadataPartition.cpp
@@ -4,6 +4,7 @@
 
 #include "PCGContext.h"
 #include "PCGCustomVersion.h"
+#include "Compute/PCGComputeCommon.h"
 #include "Elements/Metadata/PCGMetadataPartitionKernel.h"
 #include "Graph/PCGGPUGraphCompilationContext.h"
 #include "Helpers/PCGHelpers.h"
@@ -56,6 +57,17 @@ void UPCGMetadataPartitionSettings::PostLoad()
 }
 
 #if WITH_EDITOR
+void UPCGMetadataPartitionSettings::CreateKernels(FPCGGPUCompilationContext& InOutContext, UObject* InObjectOuter, TArray<UPCGComputeKernel*>& OutKernels, TArray<FPCGKernelEdge>& OutEdges) const
+{
+	UPCGMetadataPartitionKernel* Kernel = InOutContext.NewObject_AnyThread<UPCGMetadataPartitionKernel>(InObjectOuter);
+	Kernel->Initialize(this, bDumpDataDescriptions);
+	OutKernels.Add(Kernel);
+
+	// Connect node pins to kernel pins
+	OutEdges.Emplace(FPCGPinReference(PCGPinConstants::DefaultInputLabel), FPCGPinReference(Kernel, PCGPinConstants::DefaultInputLabel));
+	OutEdges.Emplace(FPCGPinReference(Kernel, PCGPinConstants::DefaultOutputLabel), FPCGPinReference(PCGPinConstants::DefaultOutputLabel));
+}
+
 void UPCGMetadataPartitionSettings::ApplyDeprecation(UPCGNode* InOutNode)
 {
 	PRAGMA_DISABLE_DEPRECATION_WARNINGS
@@ -67,14 +79,6 @@ void UPCGMetadataPartitionSettings::ApplyDeprecation(UPCGNode* InOutNode)
 
 	Super::ApplyDeprecation(InOutNode);
 }
-
-UPCGComputeKernel* UPCGMetadataPartitionSettings::CreateKernel(FPCGGPUCompilationContext& InOutContext, UObject* InObjectOuter) const
-{
-	UPCGMetadataPartitionKernel* Kernel = InOutContext.NewObject_AnyThread<UPCGMetadataPartitionKernel>(InObjectOuter);
-	Kernel->Initialize(this, bDumpDataDescriptions);
-
-	return Kernel;
-}
 #endif // WITH_EDITOR
 
 FString UPCGMetadataPartitionSettings::GetAdditionalTitleInformation() const
diff --git a/Engine/Plugins/PCG/Source/PCG/Private/Elements/PCGCopyPoints.cpp b/Engine/Plugins/PCG/Source/PCG/Private/Elements/PCGCopyPoints.cpp
--- a/Engine/Plugins/PCG/Source/PCG/Private/Elements/PCGCopyPoints.cpp
+++ b/Engine/Plugins/PCG/Source/PCG/Private/Elements/PCGCopyPoints.cpp
@@ -22,12 +22,16 @@ FText UPCGCopyPointsSettings::GetNodeTooltipText() const
 	return LOCTEXT("NodeTooltip", "For each point pair from the source and the target, create a copy, inheriting properties & attributes depending on the node settings.");
 }
 
-UPCGComputeKernel* UPCGCopyPointsSettings::CreateKernel(FPCGGPUCompilationContext& InOutContext, UObject* InObjectOuter) const
+void UPCGCopyPointsSettings::CreateKernels(FPCGGPUCompilationContext& InOutContext, UObject* InObjectOuter, TArray<UPCGComputeKernel*>& OutKernels, TArray<FPCGKernelEdge>& OutEdges) const
 {
 	UPCGCopyPointsKernel* Kernel = InOutContext.NewObject_AnyThread<UPCGCopyPointsKernel>(InObjectOuter);
 	Kernel->Initialize(this, bDumpDataDescriptions);
+	OutKernels.Add(Kernel);
 
-	return Kernel;
+	// Connect node pins to kernel pins
+	OutEdges.Emplace(FPCGPinReference(PCGCopyPointsConstants::SourcePointsLabel), FPCGPinReference(Kernel, PCGCopyPointsConstants::SourcePointsLabel));
+	OutEdges.Emplace(FPCGPinReference(PCGCopyPointsConstants::TargetPointsLabel), FPCGPinReference(Kernel, PCGCopyPointsConstants::TargetPointsLabel));
+	OutEdges.Emplace(FPCGPinReference(Kernel, PCGPinConstants::DefaultOutputLabel), FPCGPinReference(PCGPinConstants::DefaultOutputLabel));
 }
 #endif
 
diff --git a/Engine/Plugins/PCG/Source/PCG/Private/Elements/PCGStaticMeshSpawner.cpp b/Engine/Plugins/PCG/Source/PCG/Private/Elements/PCGStaticMeshSpawner.cpp
--- a/Engine/Plugins/PCG/Source/PCG/Private/Elements/PCGStaticMeshSpawner.cpp
+++ b/Engine/Plugins/PCG/Source/PCG/Private/Elements/PCGStaticMeshSpawner.cpp
@@ -5,6 +5,7 @@
 #include "PCGComponent.h"
 #include "PCGCustomVersion.h"
 #include "PCGManagedResource.h"
+#include "Compute/PCGComputeCommon.h"
 #include "Data/PCGPointData.h"
 #include "Data/PCGSpatialData.h"
 #include "Elements/PCGStaticMeshSpawnerContext.h"
@@ -45,12 +46,15 @@ UPCGStaticMeshSpawnerSettings::UPCGStaticMeshSpawnerSettings(const FObjectInitia
 }
 
 #if WITH_EDITOR
-UPCGComputeKernel* UPCGStaticMeshSpawnerSettings::CreateKernel(FPCGGPUCompilationContext& InOutContext, UObject* InObjectOuter) const
+void UPCGStaticMeshSpawnerSettings::CreateKernels(FPCGGPUCompilationContext& InOutContext, UObject* InObjectOuter, TArray<UPCGComputeKernel*>& OutKernels, TArray<FPCGKernelEdge>& OutEdges) const
 {
 	UPCGStaticMeshSpawnerKernel* Kernel = InOutContext.NewObject_AnyThread<UPCGStaticMeshSpawnerKernel>(InObjectOuter);
 	Kernel->Initialize(this, bDumpDataDescriptions);
+	OutKernels.Add(Kernel);
 
-	return Kernel;
+	// Connect node pins to kernel pins
+	OutEdges.Emplace(FPCGPinReference(PCGPinConstants::DefaultInputLabel), FPCGPinReference(Kernel, PCGPinConstants::DefaultInputLabel));
+	OutEdges.Emplace(FPCGPinReference(Kernel, PCGPinConstants::DefaultOutputLabel), FPCGPinReference(PCGPinConstants::DefaultOutputLabel));
 }
 
 FText UPCGStaticMeshSpawnerSettings::GetDefaultNodeTitle() const
diff --git a/Engine/Plugins/PCG/Source/PCG/Private/Graph/PCGGraphCompiler.cpp b/Engine/Plugins/PCG/Source/PCG/Private/Graph/PCGGraphCompiler.cpp
--- a/Engine/Plugins/PCG/Source/PCG/Private/Graph/PCGGraphCompiler.cpp
+++ b/Engine/Plugins/PCG/Source/PCG/Private/Graph/PCGGraphCompiler.cpp
@@ -846,7 +846,7 @@ void FPCGGraphCompiler::CullTasksStaticInactive(TArray<FPCGGraphTask>& InOutComp
 	CullTasks(InOutCompiledTasks, /*bAddPassthroughWires=*/false, ShouldCull);
 }
 
-void FPCGGraphCompiler::CullTasks(TArray<FPCGGraphTask>& InOutCompiledTasks, bool bAddPassthroughWires, TFunctionRef<bool(const FPCGGraphTask&)> CullTask)
+void FPCGGraphCompiler::CullTasks(TArray<FPCGGraphTask>& InOutCompiledTasks, bool bAddPassthroughWires, TFunctionRef<bool(const FPCGGraphTask&)> CullTask, TArray<FPCGTaskId>* OutOptionalOldIdToNewIdMap)
 {
 	TRACE_CPUPROFILER_EVENT_SCOPE(FPCGGraphCompiler::CullTasks);
 
@@ -856,14 +856,14 @@ void FPCGGraphCompiler::CullTasks(TArray<FPCGGraphTask>& InOutCompiledTasks, boo
 		return;
 	}
 
-	TArray<FPCGTaskId> TaskRemapping;
-	TaskRemapping.SetNumUninitialized(InOutCompiledTasks.Num());
+	TArray<FPCGTaskId> OldIdToNewId;
+	OldIdToNewId.SetNumUninitialized(InOutCompiledTasks.Num());
 
 	// Mark culled tasks by remapping to INDEX_NONE. First task is input task and is never culled.
-	TaskRemapping[0] = 0;
+	OldIdToNewId[0] = 0;
 	for (int TaskIndex = 1; TaskIndex < InOutCompiledTasks.Num(); ++TaskIndex)
 	{
-		TaskRemapping[TaskIndex] = CullTask(InOutCompiledTasks[TaskIndex]) ? InvalidPCGTaskId : 0;
+		OldIdToNewId[TaskIndex] = CullTask(InOutCompiledTasks[TaskIndex]) ? InvalidPCGTaskId : 0;
 	}
 
 	// Optionally add wires that bypass culled nodes.
@@ -883,7 +883,7 @@ void FPCGGraphCompiler::CullTasks(TArray<FPCGGraphTask>& InOutCompiledTasks, boo
 				const FPCGTaskId InputTaskId = Task.Inputs[InputIndex].TaskId;
 
 				// Is node culled?
-				if (TaskRemapping[InputTaskId] == InvalidPCGTaskId)
+				if (OldIdToNewId[InputTaskId] == InvalidPCGTaskId)
 				{
 					const FPCGGraphTask& CulledInputTask = InOutCompiledTasks[InputTaskId];
 					if (CulledInputTask.Node && CulledInputTask.Node->GetInputPins().Num() > 1)
@@ -916,7 +916,7 @@ void FPCGGraphCompiler::CullTasks(TArray<FPCGGraphTask>& InOutCompiledTasks, boo
 	while (ReadIndex < InOutCompiledTasks.Num())
 	{
 		// If not culled, then move the task to it's final remapped position in the task array.
-		if (TaskRemapping[ReadIndex] != InvalidPCGTaskId)
+		if (OldIdToNewId[ReadIndex] != InvalidPCGTaskId)
 		{
 			if (WriteIndex != ReadIndex)
 			{
@@ -924,7 +924,7 @@ void FPCGGraphCompiler::CullTasks(TArray<FPCGGraphTask>& InOutCompiledTasks, boo
 				InOutCompiledTasks[WriteIndex].NodeId = WriteIndex;
 			}
 
-			TaskRemapping[ReadIndex] = WriteIndex;
+			OldIdToNewId[ReadIndex] = WriteIndex;
 			++WriteIndex;
 		}
 
@@ -939,7 +939,7 @@ void FPCGGraphCompiler::CullTasks(TArray<FPCGGraphTask>& InOutCompiledTasks, boo
 		for (int InputIndex = Task.Inputs.Num() - 1; InputIndex >= 0; --InputIndex)
 		{
 			const FPCGTaskId InputTaskId = Task.Inputs[InputIndex].TaskId;
-			const FPCGTaskId Remap = TaskRemapping[InputTaskId];
+			const FPCGTaskId Remap = OldIdToNewId[InputTaskId];
 			if (Remap != InvalidPCGTaskId)
 			{
 				Task.Inputs[InputIndex].TaskId = Remap;
@@ -953,7 +953,7 @@ void FPCGGraphCompiler::CullTasks(TArray<FPCGGraphTask>& InOutCompiledTasks, boo
 		// Remap parent ID if there tasks is in a child scope.
 		if (Task.ParentId != InvalidPCGTaskId)
 		{
-			const FPCGTaskId RemappedParentId = TaskRemapping[Task.ParentId];
+			const FPCGTaskId RemappedParentId = OldIdToNewId[Task.ParentId];
 			// Parent task should not have been culled.
 			ensure(RemappedParentId != InvalidPCGTaskId);
 
@@ -961,6 +961,12 @@ void FPCGGraphCompiler::CullTasks(TArray<FPCGGraphTask>& InOutCompiledTasks, boo
 			Task.ParentId = RemappedParentId;
 		}
 	}
+
+	// Return the mapping if caller requested it.
+	if (OutOptionalOldIdToNewIdMap)
+	{
+		*OutOptionalOldIdToNewIdMap = MoveTemp(OldIdToNewId);
+	}
 }
 
 void FPCGGraphCompiler::AddReferencedObjects(FReferenceCollector& Collector)
diff --git a/Engine/Plugins/PCG/Source/PCG/Private/Graph/PCGGraphCompiler.h b/Engine/Plugins/PCG/Source/PCG/Private/Graph/PCGGraphCompiler.h
--- a/Engine/Plugins/PCG/Source/PCG/Private/Graph/PCGGraphCompiler.h
+++ b/Engine/Plugins/PCG/Source/PCG/Private/Graph/PCGGraphCompiler.h
@@ -94,7 +94,7 @@ public:
 	FPCGElementPtr GetSharedGatherElement();
 
 	/** Culls tasks based on a given lambda. Never culls the first (input) task in the array. */
-	static void CullTasks(TArray<FPCGGraphTask>& InOutCompiledTasks, bool bAddPassthroughWires, TFunctionRef<bool(const FPCGGraphTask&)> CullTask);
+	static void CullTasks(TArray<FPCGGraphTask>& InOutCompiledTasks, bool bAddPassthroughWires, TFunctionRef<bool(const FPCGGraphTask&)> CullTask, TArray<FPCGTaskId>* OutOptionalOldIdToNewIdMap = nullptr);
 
 	void AddReferencedObjects(FReferenceCollector& Collector);
 
diff --git a/Engine/Plugins/PCG/Source/PCG/Private/Graph/PCGGraphCompilerGPU.cpp b/Engine/Plugins/PCG/Source/PCG/Private/Graph/PCGGraphCompilerGPU.cpp
--- a/Engine/Plugins/PCG/Source/PCG/Private/Graph/PCGGraphCompilerGPU.cpp
+++ b/Engine/Plugins/PCG/Source/PCG/Private/Graph/PCGGraphCompilerGPU.cpp
@@ -8,6 +8,7 @@
 #include "Compute/PCGComputeCommon.h"
 #include "Compute/PCGComputeKernel.h"
 #include "Compute/PCGComputeKernelSource.h"
+#include "Compute/PCGPinPropertiesGPU.h"
 #include "Compute/DataInterfaces/PCGComputeDataInterface.h"
 #include "Compute/DataInterfaces/PCGCustomKernelDataInterface.h"
 #include "Compute/DataInterfaces/PCGDataCollectionDataInterface.h"
@@ -226,6 +227,244 @@ void FPCGGraphCompilerGPU::CollectGPUNodeSubsets(
 	}
 }
 
+void FPCGGraphCompilerGPU::ExpandGPUNodeKernelsToTasks(
+	FPCGGPUCompilationContext& InOutContext,
+	TSet<FPCGTaskId>& InOutGPUCompatibleTaskIds,
+	TArray<FPCGGraphTask>& InOutCompiledTasks,
+	TMap<FPCGTaskId, UPCGComputeKernel*>& OutTaskIdToKernel)
+{
+	const int NumTasksBefore = InOutCompiledTasks.Num();
+
+	// Build successors map, only for nodes that are downstream from one or more of the relevant GPU tasks.
+	FTaskToSuccessors TaskSuccessors;
+	TaskSuccessors.Reserve(InOutCompiledTasks.Num());
+	for (FPCGTaskId TaskId = 0; TaskId < InOutCompiledTasks.Num(); ++TaskId)
+	{
+		for (const FPCGGraphTaskInput& TaskInput : InOutCompiledTasks[TaskId].Inputs)
+		{
+			if (InOutGPUCompatibleTaskIds.Contains(TaskInput.TaskId))
+			{
+				TaskSuccessors.FindOrAdd(TaskInput.TaskId).AddUnique(TaskId);
+			}
+		}
+	}
+
+	// Local to loop body, but hoisted for performance.
+	TArray<UPCGComputeKernel*> NodeKernels;
+	TArray<FPCGKernelEdge> NodeKernelEdges;
+	TMap<const UPCGComputeKernel*, FPCGTaskId> KernelToTaskId;
+	TArray<FPCGPinProperties> KernelInputs;
+	TArray<FPCGPinPropertiesGPU> UpstreamKernelOutputs;
+
+	KernelToTaskId.Reserve(InOutGPUCompatibleTaskIds.Num());
+	OutTaskIdToKernel.Reserve(InOutGPUCompatibleTaskIds.Num());
+
+	// Process one GPU task at a time.
+	for (FPCGTaskId TaskId : InOutGPUCompatibleTaskIds)
+	{
+		const UPCGSettings* Settings = InOutCompiledTasks[TaskId].Node ? InOutCompiledTasks[TaskId].Node->GetSettings() : nullptr;
+		if (!ensure(Settings))
+		{
+			continue;
+		}
+
+		NodeKernels.Reset();
+		NodeKernelEdges.Reset();
+		Settings->CreateKernels(InOutContext, GetTransientPackage(), NodeKernels, NodeKernelEdges);
+
+		ensureMsgf(NodeKernels.Remove(nullptr) == 0, TEXT("Settings '%s' returned one or more null kernels."), *Settings->GetName());
+
+		if (NodeKernels.IsEmpty())
+		{
+			ensureMsgf(false, TEXT("PCG GPU compiler: Settings '%s' did not emit any kernels, check implementation of UPCGSettings::CreateKernels()."), *Settings->GetName());
+			continue;
+		}
+
+		if (NodeKernelEdges.IsEmpty())
+		{
+			continue;
+		}
+
+		KernelToTaskId.Reset();
+
+		// Create a new task for each kernel emitted by the node.
+		while (!NodeKernels.IsEmpty())
+		{
+			// Find a kernel that is ready to process (does not depend on another kernel in this node).
+			// Number of kernels is likely to be small so brute force loop over them.
+			int ReadyKernelIndex = 0;
+
+			for (int Index = 0; Index < NodeKernels.Num(); ++Index)
+			{
+				// Look for an edge connected to an upstream kernel that is yet to be processed.
+				const FPCGKernelEdge* BlockingEdge = NodeKernelEdges.FindByPredicate([Index, &NodeKernels, &KernelToTaskId](const FPCGKernelEdge& InEdge)
+				{
+					const UPCGComputeKernel* UpstreamKernel = InEdge.UpstreamPin.Kernel;
+					return InEdge.DownstreamPin.Kernel == NodeKernels[Index] && UpstreamKernel && !KernelToTaskId.Contains(UpstreamKernel);
+				});
+
+				if (!BlockingEdge)
+				{
+					// No pending kernels, ready to execute this one.
+					ReadyKernelIndex = Index;
+					break;
+				}
+			}
+
+			if (ReadyKernelIndex == INDEX_NONE)
+			{
+				ensureMsgf(false, TEXT("Compilation did not make progress, %d kernels were not processed."), NodeKernels.Num());
+				break;
+			}
+
+			UPCGComputeKernel* Kernel = NodeKernels[ReadyKernelIndex];
+
+			// Create a new task for this kernel.
+			FPCGTaskId NewKernelTaskId = InOutCompiledTasks.Num();
+
+			FPCGGraphTask TaskCopy = InOutCompiledTasks[TaskId];
+			TaskCopy.NodeId = NewKernelTaskId;
+			TaskCopy.Inputs.Reset();
+
+			FPCGGraphTask& NewKernelTask = InOutCompiledTasks.Add_GetRef(MoveTemp(TaskCopy));
+
+			// Update maps.
+			KernelToTaskId.Add(Kernel, NewKernelTaskId);
+			OutTaskIdToKernel.Add(NewKernelTaskId, Kernel);
+
+			// Get kernel inputs in preparation for wiring.
+			KernelInputs.Reset();
+			Kernel->GetInputPins(KernelInputs);
+
+			// Wire inputs
+			for (const FPCGPinProperties& KernelInputProps : KernelInputs)
+			{
+				const FPCGKernelEdge* Edge = NodeKernelEdges.FindByPredicate([Kernel, &KernelInputProps](const FPCGKernelEdge& Edge)
+				{
+					return Edge.DownstreamPin.Kernel == Kernel && Edge.DownstreamPin.Label == KernelInputProps.Label;
+				});
+
+				if (!Edge)
+				{
+					continue;
+				}
+
+				if (const UPCGComputeKernel* UpstreamKernel = Edge->UpstreamPin.Kernel)
+				{
+					// Connection from an upstream kernel from this node.
+					const FPCGTaskId UpstreamTaskId = KernelToTaskId[UpstreamKernel];
+
+					FPCGGraphTaskInput& TaskInput = NewKernelTask.Inputs.AddDefaulted_GetRef();
+					TaskInput.TaskId = UpstreamTaskId;
+					TaskInput.bProvideData = true;
+					TaskInput.DownstreamPin = KernelInputProps;
+
+					// Look for output pin of upstream kernel in order to set the input upstream pin properties.
+					UpstreamKernelOutputs.Reset();
+					Kernel->GetOutputPins(UpstreamKernelOutputs);
+
+					const FPCGPinPropertiesGPU* UpstreamOutput = UpstreamKernelOutputs.FindByPredicate([Edge](const FPCGPinPropertiesGPU& InUpstreamOutput)
+					{
+						return InUpstreamOutput.Label == Edge->UpstreamPin.Label;
+					});
+
+					if (ensure(UpstreamOutput))
+					{
+						TaskInput.UpstreamPin = *UpstreamOutput;
+					}
+				}
+				// Edge from an upstream node/task. Correlates to an input for the original node task.
+				else
+				{
+					check(Edge->IsConnectedToNodeInput());
+					const FName NodePinLabel = Edge->UpstreamPin.Label;
+
+					const FPCGGraphTaskInput* OriginalTaskInput = InOutCompiledTasks[TaskId].Inputs.FindByPredicate([NodePinLabel, &KernelInputProps](const FPCGGraphTaskInput& TaskInput)
+					{
+						return TaskInput.DownstreamPin && TaskInput.DownstreamPin->Label == NodePinLabel;
+					});
+
+					if (ensure(OriginalTaskInput))
+					{
+						// Match - add a new input wire to the new task.
+						FPCGGraphTaskInput& NewTaskInput = NewKernelTask.Inputs.Add_GetRef(*OriginalTaskInput);
+
+						// Update the pin properties with those specified by the kernel.
+						NewTaskInput.DownstreamPin = KernelInputProps;
+					}
+				}
+			}
+
+			NodeKernels.RemoveAt(ReadyKernelIndex);
+		}
+
+		// Now rewire all tasks downstream of the original task to the appropriate kernel outputs.
+		if (const TArray<FPCGTaskId>* Successors = TaskSuccessors.Find(TaskId))
+		{
+			for (FPCGTaskId SuccessorTaskId : *Successors)
+			{
+				for (FPCGGraphTaskInput& Input : InOutCompiledTasks[SuccessorTaskId].Inputs)
+				{
+					if (Input.TaskId == TaskId && Input.UpstreamPin)
+					{
+						// Outgoing edge. Rewire to appropriate kernel.
+						for (FPCGKernelEdge& Edge : NodeKernelEdges)
+						{
+							if (Edge.IsConnectedToNodeOutput() && Edge.DownstreamPin.Label == Input.UpstreamPin->Label && ensure(Edge.UpstreamPin.Kernel))
+							{
+								Input.TaskId = KernelToTaskId[Edge.UpstreamPin.Kernel];
+							}
+						}
+					}
+				}
+			}
+		}
+	}
+
+	const int NumTasksAfter = InOutCompiledTasks.Num();
+
+	TArray<FPCGTaskId> OldIdToNewId;
+
+	// All the original GPU tasks have now been replaced with a task for each kernel. Remove them without rewiring.
+	FPCGGraphCompiler::CullTasks(
+		InOutCompiledTasks,
+		/*bAddPassthroughWires=*/false,
+		[&InOutGPUCompatibleTaskIds](const FPCGGraphTask& InTask)
+		{
+			return InOutGPUCompatibleTaskIds.Contains(InTask.NodeId);
+		},
+		&OldIdToNewId);
+
+	check(OldIdToNewId.Num() == NumTasksAfter);
+
+	// Refresh GPU compatible task IDs to reflect shifted indices after culling.
+	{
+		InOutGPUCompatibleTaskIds.Reset();
+
+		for (FPCGTaskId OldGPUTaskId = NumTasksBefore; OldGPUTaskId < NumTasksAfter; ++OldGPUTaskId)
+		{
+			const FPCGTaskId RemappedId = OldIdToNewId[OldGPUTaskId];
+			check(RemappedId != INDEX_NONE);
+
+			InOutGPUCompatibleTaskIds.Add(RemappedId);
+		}
+	}
+
+	// Remap task IDs of the task->kernel map to reflect shifted indices after culling.
+	{
+		TMap<FPCGTaskId, UPCGComputeKernel*> TaskIdToKernelBeforeCull = MoveTemp(OutTaskIdToKernel);
+		OutTaskIdToKernel.Reset();
+
+		for (TPair<FPCGTaskId, UPCGComputeKernel*>& Entry : TaskIdToKernelBeforeCull)
+		{
+			const FPCGTaskId RemappedId = OldIdToNewId[Entry.Key];
+			check(RemappedId != INDEX_NONE);
+
+			OutTaskIdToKernel.Add(RemappedId, Entry.Value);
+		}
+	}
+}
+
 void FPCGGraphCompilerGPU::CreateGatherTasksAtGPUInputs(UPCGGraph* InGraph, const TSet<FPCGTaskId>& InGPUCompatibleTaskIds, TArray<FPCGGraphTask>& InOutCompiledTasks)
 {
 	TRACE_CPUPROFILER_EVENT_SCOPE(FPCGGraphCompilerGPU::CreateGatherTasksAtGPUInputs);
@@ -480,7 +719,8 @@ void FPCGGraphCompilerGPU::BuildComputeGraphTask(
 	const TArray<FPCGTaskId>& InComputeElementTasks,
 	const FTaskToSuccessors& InTaskSuccessors,
 	TArray<FPCGGraphTask>& InOutCompiledTasks,
-	const FOriginalToVirtualPin& InOriginalToVirtualPin)
+	const FOriginalToVirtualPin& InOriginalToVirtualPin,
+	TMap<FPCGTaskId, UPCGComputeKernel*>& InTaskIdToKernel)
 {
 	TRACE_CPUPROFILER_EVENT_SCOPE(FPCGGraphCompilerGPU::BuildComputeGraphTask);
 
@@ -505,25 +745,37 @@ void FPCGGraphCompilerGPU::BuildComputeGraphTask(
 		// But execution code requires it currently.
 		ComputeGraph->Bindings.Add(UPCGDataBinding::StaticClass());
 
-		for (FPCGTaskId TaskId : InCollapsedTasks)
-		{
-			if (const UPCGSettings* Settings = InOutCompiledTasks[TaskId].Node ? InOutCompiledTasks[TaskId].Node->GetSettings() : nullptr)
-			{
-				InOutCompiledTasks[TaskId].Kernel = Settings->CreateKernel(InOutContext, ComputeGraph);
-			}
-		}
+		BuildComputeGraphStaticData(InOutCompiledTasks, InCollapsedTasks, InTaskIdToKernel, ComputeGraph);
 
 		FPinToDataInterface OutputPinDataInterfaces;
 
-		BuildComputeGraphStaticData(InOutCompiledTasks, InCollapsedTasks, ComputeGraph);
-
-		CreateDataInterfaces(InOutContext, InOutCompiledTasks, InCollapsedTasks, InComputeElementTasks, InTaskSuccessors, InOriginalToVirtualPin, ComputeGraph, OutputPinDataInterfaces);
+		CreateDataInterfaces(InOutContext, InOutCompiledTasks, InCollapsedTasks, InComputeElementTasks, InTaskSuccessors, InOriginalToVirtualPin, InTaskIdToKernel, ComputeGraph, OutputPinDataInterfaces);
 
-		CompileComputeGraph(InOutContext, InOutCompiledTasks, InCollapsedTasks, InComputeElementTasks, InTaskSuccessors, InOriginalToVirtualPin, OutputPinDataInterfaces, ComputeGraph);
+		CompileComputeGraph(InOutContext, InOutCompiledTasks, InCollapsedTasks, InComputeElementTasks, InTaskSuccessors, InOriginalToVirtualPin, OutputPinDataInterfaces, InTaskIdToKernel, ComputeGraph);
 
 		// Remove empty strings (at execution time empty string is placed in table slot 0).
 		// NOTE: This can scramble order but order is not important at this stage (key values created at execution time in data binding).
 		ComputeGraph->StringTable.RemoveAllSwap([](const FString& InString) { return InString.IsEmpty(); });
+
+		for (FPCGTaskId TaskId : InCollapsedTasks)
+		{
+			// Re-outer kernels to compute graph.
+			InTaskIdToKernel[TaskId]->Rename(nullptr, ComputeGraph);
+		}
+	}
+	else
+	{
+		// Compute graph obtained from cache, it will have its own kernels, so garbage the kernels we created.
+		// todo_pcg: Pretty wasteful to create kernel objects and then throw them away, even if only happens on first graph compilation.
+		// We could probably defer creation of these objects until they're really needed by storing "createinfo" data. Or possibly the caching
+		// of compute graphs could be done per pcg-graph, rather than per-compute-graph, and we could avoid doing any GPU compilation work if already cached.
+		for (TPair<FPCGTaskId, UPCGComputeKernel*> TaskIdAndKernel : InTaskIdToKernel)
+		{
+			if (ensure(TaskIdAndKernel.Value))
+			{
+				TaskIdAndKernel.Value->MarkAsGarbage();
+			}
+		}
 	}
 
 	const int32 ComputeGraphIndex = InOutContext.AddCompiledComputeGraph(ComputeGraph);
@@ -553,17 +805,14 @@ void FPCGGraphCompilerGPU::BuildComputeGraphTask(
 void FPCGGraphCompilerGPU::BuildComputeGraphStaticData(
 	const TArray<FPCGGraphTask>& InCompiledTasks,
 	const TSet<FPCGTaskId>& InCollapsedTasks,
+	const TMap<FPCGTaskId, UPCGComputeKernel*>& InTaskIdToKernel,
 	UPCGComputeGraph* InOutComputeGraph)
 {
 	TRACE_CPUPROFILER_EVENT_SCOPE(FPCGGraphCompilerGPU::BuildComputeGraphStaticData);
 
 	for (FPCGTaskId TaskId : InCollapsedTasks)
 	{
-		// TODO: The kernel being used here is a duplicate of the actual kernel that will be created during compute graph creation.
-		// This is okay for now because the functions just pass through data from the Settings at the moment, which doesn't require a fully initialized kernel.
-		const UPCGComputeKernel* Kernel = InCompiledTasks[TaskId].Kernel;
-
-		if (ensure(Kernel))
+		if (const UPCGComputeKernel* Kernel = InTaskIdToKernel[TaskId]; ensure(Kernel))
 		{
 			// Record attributes.
 			TArray<FPCGKernelAttributeKey> AttributeKeys;
@@ -587,6 +836,7 @@ void FPCGGraphCompilerGPU::CreateDataInterfaces(
 	const TArray<FPCGTaskId>& InComputeElementTasks,
 	const FTaskToSuccessors& InTaskSuccessors,
 	const FOriginalToVirtualPin& InOriginalToVirtualPin,
+	const TMap<FPCGTaskId, UPCGComputeKernel*>& InTaskIdToKernel,
 	UPCGComputeGraph* InOutComputeGraph,
 	FPinToDataInterface& InOutPinToDataInterface)
 {
@@ -594,12 +844,16 @@ void FPCGGraphCompilerGPU::CreateDataInterfaces(
 
 	// Create data interfaces which allow kernels to read or write data. Each data interface is associated with a node output pin.
 	// For CPU->GPU edges, an upload data interface is created. For GPU->CPU edges, a readback data interface is created.
-	auto CreateOutputDataInterface = [&InOutContext, &InCollapsedTasks, &InCompiledTasks, InOutComputeGraph](FPCGTaskId InProducerTaskId, bool bRequiresExport, const FPCGPinProperties& InOutputPinProperties) -> UPCGComputeDataInterface*
+	auto CreateOutputDataInterface = [&InOutContext, &InCollapsedTasks, &InCompiledTasks, &InTaskIdToKernel, InOutComputeGraph](FPCGTaskId InProducerTaskId, bool bRequiresExport, const FPCGPinProperties& InOutputPinProperties) -> UPCGComputeDataInterface*
 	{
 		TRACE_CPUPROFILER_EVENT_SCOPE(CreateDataInterface);
 
 		// The kernel that produced this data, if any (otherwise could be CPU node or a node from a different compute graph).
-		const UPCGComputeKernel* ProducerKernel = InCompiledTasks[InProducerTaskId].Kernel;
+		const UPCGComputeKernel* ProducerKernel = nullptr;
+		if (UPCGComputeKernel* const* FoundKernel = InTaskIdToKernel.Find(InProducerTaskId))
+		{
+			ProducerKernel = *FoundKernel;
+		}
 
 		const bool bUpstreamIsGPUTask = InCollapsedTasks.Contains(InProducerTaskId);
 
@@ -797,6 +1051,7 @@ void FPCGGraphCompilerGPU::CompileComputeGraph(
 	const FTaskToSuccessors& InTaskSuccessors,
 	const FOriginalToVirtualPin& InOriginalToVirtualPin,
 	const FPinToDataInterface& InOutputPinToDataInterface,
+	const TMap<FPCGTaskId, UPCGComputeKernel*>& InTaskIdToKernel,
 	UPCGComputeGraph* InOutComputeGraph)
 {
 	TRACE_CPUPROFILER_EVENT_SCOPE(FPCGGraphCompilerGPU::CompileComputeGraph);
@@ -884,11 +1139,11 @@ void FPCGGraphCompilerGPU::CompileComputeGraph(
 		FKernelWithDataBindings KernelWithBindings;
 
 		const int32 KernelIndex = InOutComputeGraph->KernelInvocations.Num();
-		KernelWithBindings.Kernel = Settings->CreateKernel(InOutContext, /*InObjectOuter=*/InOutComputeGraph);
+		KernelWithBindings.Kernel = InTaskIdToKernel[TaskId];
 		check(KernelWithBindings.Kernel);
 		KernelWithBindings.Kernel->SetKernelIndex(KernelIndex);
 
-		InOutComputeGraph->KernelInvocations.Add(Cast<UComputeKernel>(KernelWithBindings.Kernel));
+		InOutComputeGraph->KernelInvocations.Add(KernelWithBindings.Kernel);
 
 		if (const UPCGStaticMeshSpawnerKernel* SMSpawnerKernel = Cast<UPCGStaticMeshSpawnerKernel>(KernelWithBindings.Kernel))
 		{
@@ -1352,6 +1607,11 @@ void FPCGGraphCompilerGPU::CreateGPUNodes(FPCGGraphCompiler& InOutCompiler, UPCG
 	// before passing to GPU.
 	CreateGatherTasksAtGPUInputs(InGraph, GPUCompatibleTaskIds, InOutCompiledTasks);
 
+	TMap<FPCGTaskId, UPCGComputeKernel*> TaskIdToKernel;
+
+	// Create one or more kernel objects for each GPU node. A graph task will be created for each.
+	ExpandGPUNodeKernelsToTasks(Context, GPUCompatibleTaskIds, InOutCompiledTasks, TaskIdToKernel);
+
 	FTaskToSuccessors TaskSuccessors;
 	TaskSuccessors.Reserve(InOutCompiledTasks.Num());
 	for (FPCGTaskId TaskId = 0; TaskId < InOutCompiledTasks.Num(); ++TaskId)
@@ -1427,7 +1687,8 @@ void FPCGGraphCompilerGPU::CreateGPUNodes(FPCGGraphCompiler& InOutCompiler, UPCG
 			ComputeElementTasks,
 			TaskSuccessors,
 			InOutCompiledTasks,
-			OriginalToVirtualPin[ComputeGraphIndex]);
+			OriginalToVirtualPin[ComputeGraphIndex],
+			TaskIdToKernel);
 	}
 
 	{
diff --git a/Engine/Plugins/PCG/Source/PCG/Private/Graph/PCGGraphCompilerGPU.h b/Engine/Plugins/PCG/Source/PCG/Private/Graph/PCGGraphCompilerGPU.h
--- a/Engine/Plugins/PCG/Source/PCG/Private/Graph/PCGGraphCompilerGPU.h
+++ b/Engine/Plugins/PCG/Source/PCG/Private/Graph/PCGGraphCompilerGPU.h
@@ -7,6 +7,7 @@
 class FPCGGraphCompiler;
 class UPCGComputeDataInterface;
 class UPCGComputeGraph;
+class UPCGComputeKernel;
 class UPCGGraph;
 struct FPCGGraphTask;
 struct FPCGGPUCompilationContext;
@@ -39,6 +40,13 @@ private:
 		const TSet<FPCGTaskId>& InGPUCompatibleTaskIds,
 		TArray<TSet<FPCGTaskId>>& OutNodeSubsetsToConvertToCFGraph);
 	
+	/** Create a separate task for each GPU kernel emitted by each GPU node. */
+	static void ExpandGPUNodeKernelsToTasks(
+		FPCGGPUCompilationContext& InOutContext,
+		TSet<FPCGTaskId>& InOutGPUCompatibleTaskIds,
+		TArray<FPCGGraphTask>& InOutCompiledTasks,
+		TMap<FPCGTaskId, UPCGComputeKernel*>& OutTaskIdToKernel);
+
 	/** For GPU node inputs that have multiple incident edges, bundle them into a single edge. This is to avoid an inefficient
 	* gather operation on the GPU, and allows data interfaces to pick their data from the compute graph element input data collection
 	* using unique virtual input pin labels. */
@@ -71,12 +79,14 @@ private:
 		const TArray<FPCGTaskId>& InComputeElementTasks,
 		const FTaskToSuccessors& InTaskSuccessors,
 		TArray<FPCGGraphTask>& InOutCompiledTasks,
-		const FOriginalToVirtualPin& InOriginalToVirtualPin);
+		const FOriginalToVirtualPin& InOriginalToVirtualPin,
+		TMap<FPCGTaskId, UPCGComputeKernel*>& InTaskIdToKernel);
 
 	/** Collates data for the compute graph which is statically known. */
 	static void BuildComputeGraphStaticData(
 		const TArray<FPCGGraphTask>& InCompiledTasks,
 		const TSet<FPCGTaskId>& InCollapsedTasks,
+		const TMap<FPCGTaskId, UPCGComputeKernel*>& InTaskIdToKernel,
 		UPCGComputeGraph* InOutComputeGraph);
 
 	/** Creates and configures all data interfaces for the graph. */
@@ -87,7 +97,8 @@ private:
 		const TArray<FPCGTaskId>& InComputeElementTasks,
 		const FTaskToSuccessors& InTaskSuccessors,
 		const FOriginalToVirtualPin& InOriginalToVirtualPin,
-		UPCGComputeGraph* InOutComputeGraph, // In
+		const TMap<FPCGTaskId, UPCGComputeKernel*>& InTaskIdToKernel,
+		UPCGComputeGraph* InOutComputeGraph,
 		FPinToDataInterface& InOutPinToDataInterface);
 
 	/** Builds the compute graph from the input tasks and data interfaces. */
@@ -99,6 +110,7 @@ private:
 		const FTaskToSuccessors& InTaskSuccessors,
 		const FOriginalToVirtualPin& InOriginalToVirtualPin,
 		const FPinToDataInterface& InOutputPinToDataInterface,
+		const TMap<FPCGTaskId, UPCGComputeKernel*>& InTaskIdToKernel,
 		UPCGComputeGraph* InOutComputeGraph);
 };
 #endif // WITH_EDITOR
diff --git a/Engine/Plugins/PCG/Source/PCG/Private/Graph/PCGGraphTask.h b/Engine/Plugins/PCG/Source/PCG/Private/Graph/PCGGraphTask.h
--- a/Engine/Plugins/PCG/Source/PCG/Private/Graph/PCGGraphTask.h
+++ b/Engine/Plugins/PCG/Source/PCG/Private/Graph/PCGGraphTask.h
@@ -140,9 +140,6 @@ struct FPCGGraphTask
 	FPCGContext* Context = nullptr;
 	TSharedPtr<const FPCGStackContext> StackContext;
 
-	/** Kernel for GPU execution if this task supports it. */
-	const UPCGComputeKernel* Kernel = nullptr;
-
 	// Whether SetupTask has been called on this task
 	bool bHasDoneSetup = false;
 	// BuildTaskInput will initialize this Collection which will later be used by PrepareForExecute
diff --git a/Engine/Plugins/PCG/Source/PCG/Public/Elements/Metadata/PCGMetadataPartition.h b/Engine/Plugins/PCG/Source/PCG/Public/Elements/Metadata/PCGMetadataPartition.h
--- a/Engine/Plugins/PCG/Source/PCG/Public/Elements/Metadata/PCGMetadataPartition.h
+++ b/Engine/Plugins/PCG/Source/PCG/Public/Elements/Metadata/PCGMetadataPartition.h
@@ -21,11 +21,11 @@ public:
 #if WITH_EDITOR
 	virtual bool SplitGraphAtOutput() const override { return true; }
 	virtual bool DisplayExecuteOnGPUSetting() const override { return true; }
+	virtual void CreateKernels(FPCGGPUCompilationContext& InOutContext, UObject* InObjectOuter, TArray<UPCGComputeKernel*>& OutKernels, TArray<FPCGKernelEdge>& OutEdges) const override;
 	virtual void ApplyDeprecation(UPCGNode* InOutNode) override;
 	virtual FName GetDefaultNodeName() const override { return FName(TEXT("AttributePartition")); }
 	virtual FText GetDefaultNodeTitle() const override { return NSLOCTEXT("PCGMetadataPartitionSettings", "NodeTitle", "Attribute Partition"); }
 	virtual EPCGSettingsType GetType() const override { return EPCGSettingsType::Metadata; }
-	UPCGComputeKernel* CreateKernel(FPCGGPUCompilationContext& InOutContext, UObject* InObjectOuter) const override;
 #endif
 	virtual bool HasDynamicPins() const override { return true; }
 	virtual FString GetAdditionalTitleInformation() const override;
diff --git a/Engine/Plugins/PCG/Source/PCG/Public/Elements/PCGCopyPoints.h b/Engine/Plugins/PCG/Source/PCG/Public/Elements/PCGCopyPoints.h
--- a/Engine/Plugins/PCG/Source/PCG/Public/Elements/PCGCopyPoints.h
+++ b/Engine/Plugins/PCG/Source/PCG/Public/Elements/PCGCopyPoints.h
@@ -51,7 +51,7 @@ public:
 	virtual FText GetNodeTooltipText() const override;
 	virtual EPCGSettingsType GetType() const override { return EPCGSettingsType::Sampler; }
 	virtual bool DisplayExecuteOnGPUSetting() const override { return true; }
-	UPCGComputeKernel* CreateKernel(FPCGGPUCompilationContext& InOutContext, UObject* InObjectOuter) const override;
+	virtual void CreateKernels(FPCGGPUCompilationContext& InOutContext, UObject* InObjectOuter, TArray<UPCGComputeKernel*>& OutKernels, TArray<FPCGKernelEdge>& OutEdges) const override;
 #endif
 
 protected:
diff --git a/Engine/Plugins/PCG/Source/PCG/Public/Elements/PCGStaticMeshSpawner.h b/Engine/Plugins/PCG/Source/PCG/Public/Elements/PCGStaticMeshSpawner.h
--- a/Engine/Plugins/PCG/Source/PCG/Public/Elements/PCGStaticMeshSpawner.h
+++ b/Engine/Plugins/PCG/Source/PCG/Public/Elements/PCGStaticMeshSpawner.h
@@ -30,7 +30,7 @@ public:
 	// ~Begin UPCGSettings interface
 #if WITH_EDITOR
 	virtual bool DisplayExecuteOnGPUSetting() const override { return true; }
-	UPCGComputeKernel* CreateKernel(FPCGGPUCompilationContext& InOutContext, UObject* InObjectOuter) const override;
+	virtual void CreateKernels(FPCGGPUCompilationContext& InOutContext, UObject* InObjectOuter, TArray<UPCGComputeKernel*>& OutKernels, TArray<FPCGKernelEdge>& OutEdges) const override;
 	virtual FName GetDefaultNodeName() const override { return FName(TEXT("StaticMeshSpawner")); }
 	virtual FText GetDefaultNodeTitle() const override;
 	virtual EPCGSettingsType GetType() const override { return EPCGSettingsType::Spawner; }
diff --git a/Engine/Plugins/PCG/Source/PCG/Public/PCGSettings.h b/Engine/Plugins/PCG/Source/PCG/Public/PCGSettings.h
--- a/Engine/Plugins/PCG/Source/PCG/Public/PCGSettings.h
+++ b/Engine/Plugins/PCG/Source/PCG/Public/PCGSettings.h
@@ -21,7 +21,9 @@ class UPCGGraph;
 class UPCGNode;
 class UPCGSettings;
 struct FPCGGPUCompilationContext;
+struct FPCGKernelEdge;
 struct FPCGPinProperties;
+struct FPCGPinReference;
 struct FPropertyChangedEvent;
 
 using FPCGSettingsAndCulling = TPair<TSoftObjectPtr<const UPCGSettings>, bool>;
@@ -509,7 +511,10 @@ public:
 	virtual bool ShouldExecuteOnGPU() const { return bExecuteOnGPU; }
 
 #if WITH_EDITOR
-	virtual UPCGComputeKernel* CreateKernel(FPCGGPUCompilationContext& InOutContext, UObject* InObjectOuter) const { return nullptr; }
+	/** [EXPERIMENTAL] Create the GPU kernels and kernel edges that should execute for this node.
+	* Note: This API function is subject to change without deprecation.
+	*/
+	virtual void CreateKernels(FPCGGPUCompilationContext& InOutContext, UObject* InObjectOuter, TArray<UPCGComputeKernel*>& OutKernels, TArray<FPCGKernelEdge>& OutEdges) const {};
 
 	// TODO - should also move to kernel, but needs flow to be shuffled in the GPU graph compiler.
 	/** [EXPERIMENTAL] Whether compute graphs should be split at the output of this node. Useful for reading back runtime statistics or diagnostic info before continuing.
