[PCG] Optimizations for CPU logic around GPU execution.

* GPU string table building - don't repeatedly add strings from the same table.
* Write header - properly track element address rather than computing it each iteration N^2.

#rnx

diff --git a/Engine/Plugins/PCG/Source/PCG/Private/Compute/PCGDataBinding.cpp b/Engine/Plugins/PCG/Source/PCG/Private/Compute/PCGDataBinding.cpp
--- a/Engine/Plugins/PCG/Source/PCG/Private/Compute/PCGDataBinding.cpp
+++ b/Engine/Plugins/PCG/Source/PCG/Private/Compute/PCGDataBinding.cpp
@@ -223,86 +223,112 @@ void UPCGDataBinding::AddInputDataStringsToTable()
 	TRACE_CPUPROFILER_EVENT_SCOPE(UPCGDataBinding::AddInputDataStringsToTable);
 	check(Graph);
 
+	// Often many data items come from a small number of buffers, so make sure we only grab strings from each buffer once.
+	TArray<TSharedPtr<const FPCGProxyForGPUDataCollection>> ProcessedBuffers;
+
 	for (const FPCGTaggedData& Data : InputDataCollection.TaggedData)
 	{
+		// GPU proxies hold a pointer to GPU buffer info which contains the string table.
 		if (const UPCGProxyForGPUData* Proxy = Cast<UPCGProxyForGPUData>(Data.Data))
 		{
 			if (TSharedPtr<const FPCGProxyForGPUDataCollection> GPUInfo = Proxy->GetGPUInfo())
 			{
-				for (const FString& String : GPUInfo->GetStringTable())
+				if (ProcessedBuffers.Contains(GPUInfo))
+				{
+					continue;
+				}
+
+				ProcessedBuffers.Add(GPUInfo);
+
+				if (GPUInfo->GetStringTable().IsEmpty())
+				{
+					continue;
+				}
+
+				if (StringTable.IsEmpty())
 				{
-					StringTable.AddUnique(String);
+					StringTable = GPUInfo->GetStringTable();
+				}
+				else
+				{
+					for (const FString& String : GPUInfo->GetStringTable())
+					{
+						StringTable.AddUnique(String);
+					}
 				}
 
 				continue;
 			}
 		}
-
-		const UPCGMetadata* Metadata = Data.Data ? Data.Data->ConstMetadata() : nullptr;
-		if (!Metadata)
+		// Non-GPU-proxy: collect strings from metadata.
+		else
 		{
-			continue;
-		}
-
-		TArray<FName> AttributeNames;
-		TArray<EPCGMetadataTypes> AttributeTypes;
-		Metadata->GetAttributes(AttributeNames, AttributeTypes);
+			const UPCGMetadata* Metadata = Data.Data ? Data.Data->ConstMetadata() : nullptr;
+			if (!Metadata)
+			{
+				continue;
+			}
 
-		for (int AttributeIndex = 0; AttributeIndex < AttributeNames.Num(); ++AttributeIndex)
-		{
-			const FName AttributeName = AttributeNames[AttributeIndex];
-			const EPCGKernelAttributeType AttributeType = PCGDataForGPUHelpers::GetAttributeTypeFromMetadataType(AttributeTypes[AttributeIndex]);
+			TArray<FName> AttributeNames;
+			TArray<EPCGMetadataTypes> AttributeTypes;
+			Metadata->GetAttributes(AttributeNames, AttributeTypes);
 
-			if (AttributeType == EPCGKernelAttributeType::StringKey || AttributeType == EPCGKernelAttributeType::Name)
+			for (int AttributeIndex = 0; AttributeIndex < AttributeNames.Num(); ++AttributeIndex)
 			{
-				const FPCGMetadataAttributeBase* AttributeBase = Metadata->GetConstAttribute(AttributeName);
-				if (!AttributeBase)
-				{
-					UE_LOG(LogPCG, Warning, TEXT("Attribute '%s' not found."), *AttributeName.ToString());
-					continue;
-				}
+				const FName AttributeName = AttributeNames[AttributeIndex];
+				const EPCGKernelAttributeType AttributeType = PCGDataForGPUHelpers::GetAttributeTypeFromMetadataType(AttributeTypes[AttributeIndex]);
 
-				auto AddEntryToStringTable = [AttributeBase, &StringTable=StringTable](int64 InValueKey)
+				if (AttributeType == EPCGKernelAttributeType::StringKey || AttributeType == EPCGKernelAttributeType::Name)
 				{
-					if (AttributeBase->GetTypeId() == PCG::Private::MetadataTypes<FSoftObjectPath>::Id)
-					{
-						StringTable.AddUnique(static_cast<const FPCGMetadataAttribute<FSoftObjectPath>*>(AttributeBase)->GetValue(InValueKey).ToString());
-					}
-					else if (AttributeBase->GetTypeId() == PCG::Private::MetadataTypes<FString>::Id)
-					{
-						StringTable.AddUnique(static_cast<const FPCGMetadataAttribute<FString>*>(AttributeBase)->GetValue(InValueKey));
-					}
-					else if (AttributeBase->GetTypeId() == PCG::Private::MetadataTypes<FSoftClassPath>::Id)
+					const FPCGMetadataAttributeBase* AttributeBase = Metadata->GetConstAttribute(AttributeName);
+					if (!AttributeBase)
 					{
-						StringTable.AddUnique(static_cast<const FPCGMetadataAttribute<FSoftClassPath>*>(AttributeBase)->GetValue(InValueKey).ToString());
+						UE_LOG(LogPCG, Warning, TEXT("Attribute '%s' not found."), *AttributeName.ToString());
+						continue;
 					}
-					else if (AttributeBase->GetTypeId() == PCG::Private::MetadataTypes<FName>::Id)
+
+					auto AddEntryToStringTable = [AttributeBase, &StringTable=StringTable](int64 InValueKey)
 					{
-						StringTable.AddUnique(static_cast<const FPCGMetadataAttribute<FName>*>(AttributeBase)->GetValue(InValueKey).ToString());
-					}
-					else
+						if (AttributeBase->GetTypeId() == PCG::Private::MetadataTypes<FSoftObjectPath>::Id)
+						{
+							StringTable.AddUnique(static_cast<const FPCGMetadataAttribute<FSoftObjectPath>*>(AttributeBase)->GetValue(InValueKey).ToString());
+						}
+						else if (AttributeBase->GetTypeId() == PCG::Private::MetadataTypes<FString>::Id)
+						{
+							StringTable.AddUnique(static_cast<const FPCGMetadataAttribute<FString>*>(AttributeBase)->GetValue(InValueKey));
+						}
+						else if (AttributeBase->GetTypeId() == PCG::Private::MetadataTypes<FSoftClassPath>::Id)
+						{
+							StringTable.AddUnique(static_cast<const FPCGMetadataAttribute<FSoftClassPath>*>(AttributeBase)->GetValue(InValueKey).ToString());
+						}
+						else if (AttributeBase->GetTypeId() == PCG::Private::MetadataTypes<FName>::Id)
+						{
+							StringTable.AddUnique(static_cast<const FPCGMetadataAttribute<FName>*>(AttributeBase)->GetValue(InValueKey).ToString());
+						}
+						else
+						{
+							checkNoEntry();
+						}
+					};
+
+					if (const UPCGPointData* PointData = Cast<UPCGPointData>(Data.Data))
 					{
-						checkNoEntry();
+						for (const FPCGPoint& Point : PointData->GetPoints())
+						{
+							AddEntryToStringTable(AttributeBase->GetValueKey(Point.MetadataEntry));
+						}
 					}
-				};
-
-				if (const UPCGPointData* PointData = Cast<UPCGPointData>(Data.Data))
-				{
-					for (const FPCGPoint& Point : PointData->GetPoints())
+					else if (const UPCGParamData* ParamData = Cast<UPCGParamData>(Data.Data))
 					{
-						AddEntryToStringTable(AttributeBase->GetValueKey(Point.MetadataEntry));
-					}
-				}
-				else if (const UPCGParamData* ParamData = Cast<UPCGParamData>(Data.Data))
-				{
-					const int32 NumElements = Metadata->GetItemCountForChild();
+						const int32 NumElements = Metadata->GetItemCountForChild();
 
-					for (int64 MetadataKey = 0; MetadataKey < NumElements; ++MetadataKey)
-					{
-						AddEntryToStringTable(AttributeBase->GetValueKey(MetadataKey));
+						for (int64 MetadataKey = 0; MetadataKey < NumElements; ++MetadataKey)
+						{
+							AddEntryToStringTable(AttributeBase->GetValueKey(MetadataKey));
+						}
 					}
+					else { /** TODO: More attribute types! */ }
 				}
-				else { /** TODO: More attribute types! */ }
 			}
 		}
 	}
diff --git a/Engine/Plugins/PCG/Source/PCG/Private/Compute/PCGDataForGPU.cpp b/Engine/Plugins/PCG/Source/PCG/Private/Compute/PCGDataForGPU.cpp
--- a/Engine/Plugins/PCG/Source/PCG/Private/Compute/PCGDataForGPU.cpp
+++ b/Engine/Plugins/PCG/Source/PCG/Private/Compute/PCGDataForGPU.cpp
@@ -968,6 +968,9 @@ void FPCGDataCollectionDesc::WriteHeader(TArray<uint32>& OutPackedDataCollection
 
 	uint32 WriteAddressUints = 0;
 
+	// Start at end of header
+	uint32 DataStartAddressBytes = HeaderSizeBytes;
+
 	// Num data
 	OutPackedDataCollectionHeader[WriteAddressUints++] = DataDescs.Num();
 
@@ -997,31 +1000,19 @@ void FPCGDataCollectionDesc::WriteHeader(TArray<uint32>& OutPackedDataCollection
 		for (int32 AttrIndex = 0; AttrIndex < DataDesc.AttributeDescs.Num(); ++AttrIndex)
 		{
 			const FPCGKernelAttributeDesc& AttributeDesc = DataDesc.AttributeDescs[AttrIndex];
+			const int AttributeStride = PCGDataForGPUHelpers::GetAttributeTypeStrideBytes(AttributeDesc.Type);
 
 			// Scatter from attributes that are present into header which has slots for all possible attributes.
 			WriteAddressUints = (AttributeDesc.Index * ATTRIBUTE_HEADER_SIZE_BYTES + DataAttributesHeaderStartAddressBytes) >> 2;
 
 			// Data i element j: packed ID and stride
 			const uint32 AttributeId = AttributeDesc.Index;
-			const uint32 AttributeStride = PCGDataForGPUHelpers::GetAttributeTypeStrideBytes(AttributeDesc.Type);
 			const uint32 PackedIdAndStride = (AttributeId << 8) + AttributeStride;
 			OutPackedDataCollectionHeader[WriteAddressUints++] = PackedIdAndStride;
-
-			// Data i element j: data start address bytes
-			// TODO: Accumulate rather than calculate from scratch.
-			uint32 DataStartAddressBytes = HeaderSizeBytes; // Start at end of header
-			for (int32 PreviousDataIndex = 0; PreviousDataIndex < DataIndex; ++PreviousDataIndex) // Fast forward past previous data
-			{
-				for (const FPCGKernelAttributeDesc& AttrDesc : DataDescs[PreviousDataIndex].AttributeDescs)
-				{
-					DataStartAddressBytes += DataDescs[PreviousDataIndex].ElementCount * PCGDataForGPUHelpers::GetAttributeTypeStrideBytes(AttrDesc.Type);
-				}
-			}
-			for (int32 PreviousAttrIndex = 0; PreviousAttrIndex < AttrIndex; ++PreviousAttrIndex) // Fast forward past previous attributes
-			{
-				DataStartAddressBytes += DataDesc.ElementCount * PCGDataForGPUHelpers::GetAttributeTypeStrideBytes(DataDesc.AttributeDescs[PreviousAttrIndex].Type);
-			}
 			OutPackedDataCollectionHeader[WriteAddressUints++] = DataStartAddressBytes;
+
+			// Move the DataStartAddress to the beginning of the next attribute by skipping all the bytes of the current attribute.
+			DataStartAddressBytes += DataDesc.ElementCount * AttributeStride;
 		}
 
 		// After scattering in attribute headers, fast forward to end of section.
