[PCG] Fix BP Reinstancing crash: remove raw pointer cache which isn't resistant to object reinstancing.
-This was done as an optimization a while back 24453611 but current test shows very small overhead (750k calls result in 57ms cost)

#rnx

diff --git a/Engine/Plugins/PCG/Source/PCG/Private/Components/PCGProceduralISMComponent.cpp b/Engine/Plugins/PCG/Source/PCG/Private/Components/PCGProceduralISMComponent.cpp
--- a/Engine/Plugins/PCG/Source/PCG/Private/Components/PCGProceduralISMComponent.cpp
+++ b/Engine/Plugins/PCG/Source/PCG/Private/Components/PCGProceduralISMComponent.cpp
@@ -608,12 +608,6 @@ void UPCGManagedProceduralISMComponent::PostLoad()
 	GetComponent();
 }
 
-void UPCGManagedProceduralISMComponent::ForgetComponent()
-{
-	Super::ForgetComponent();
-	CachedRawComponentPtr = nullptr;
-}
-
 void UPCGManagedProceduralISMComponent::SetDescriptor(const FPCGProceduralISMComponentDescriptor& InDescriptor)
 {
 	Descriptor = InDescriptor;
@@ -709,25 +703,11 @@ void UPCGManagedProceduralISMComponent::SetRootLocation(const FVector& InRootLoc
 
 UPCGProceduralISMComponent* UPCGManagedProceduralISMComponent::GetComponent() const
 {
-	if (!CachedRawComponentPtr)
-	{
-		UPCGProceduralISMComponent* GeneratedComponentPtr = Cast<UPCGProceduralISMComponent>(GeneratedComponent.Get());
-
-		// Implementation note:
-		// There is no surefire way to make sure that we can use the raw pointer UNLESS it is from the same owner
-		if (GeneratedComponentPtr && Cast<UPCGComponent>(GetOuter()) && GeneratedComponentPtr->GetOwner() == Cast<UPCGComponent>(GetOuter())->GetOwner())
-		{
-			CachedRawComponentPtr = GeneratedComponentPtr;
-		}
-
-		return GeneratedComponentPtr;
-	}
-
-	return CachedRawComponentPtr;
+	TRACE_CPUPROFILER_EVENT_SCOPE(UPCGManagedProceduralISMComponent::GetComponent);
+	return Cast<UPCGProceduralISMComponent>(GeneratedComponent.Get());
 }
 
 void UPCGManagedProceduralISMComponent::SetComponent(UPCGProceduralISMComponent* InComponent)
 {
 	GeneratedComponent = InComponent;
-	CachedRawComponentPtr = InComponent;
 }
diff --git a/Engine/Plugins/PCG/Source/PCG/Private/Components/PCGProceduralISMComponent.h b/Engine/Plugins/PCG/Source/PCG/Private/Components/PCGProceduralISMComponent.h
--- a/Engine/Plugins/PCG/Source/PCG/Private/Components/PCGProceduralISMComponent.h
+++ b/Engine/Plugins/PCG/Source/PCG/Private/Components/PCGProceduralISMComponent.h
@@ -310,7 +310,6 @@ public:
 	virtual bool SupportsComponentReset() const override { return true; }
 	virtual void MarkAsUsed() override;
 	virtual void MarkAsReused() override;
-	virtual void ForgetComponent() override;
 	//~End UPCGManagedComponents interface
 
 	UPCGProceduralISMComponent* GetComponent() const;
@@ -336,9 +335,6 @@ protected:
 
 	UPROPERTY(Transient)
 	uint64 SettingsUID = -1; // purposefully a value that will never happen in data
-
-	// Cached raw pointer to ISM component
-	mutable UPCGProceduralISMComponent* CachedRawComponentPtr = nullptr;
 };
 
 namespace PCGManagedProceduralISMComponent
diff --git a/Engine/Plugins/PCG/Source/PCG/Private/PCGManagedResource.cpp b/Engine/Plugins/PCG/Source/PCG/Private/PCGManagedResource.cpp
--- a/Engine/Plugins/PCG/Source/PCG/Private/PCGManagedResource.cpp
+++ b/Engine/Plugins/PCG/Source/PCG/Private/PCGManagedResource.cpp
@@ -760,12 +760,6 @@ void UPCGManagedISMComponent::PostLoad()
 	GetComponent();
 }
 
-void UPCGManagedISMComponent::ForgetComponent()
-{
-	Super::ForgetComponent();
-	CachedRawComponentPtr = nullptr;
-}
-
 void UPCGManagedISMComponent::SetDescriptor(const FISMComponentDescriptor& InDescriptor)
 {
 	bHasDescriptor = true;
@@ -871,27 +865,13 @@ void UPCGManagedISMComponent::SetRootLocation(const FVector& InRootLocation)
 
 UInstancedStaticMeshComponent* UPCGManagedISMComponent::GetComponent() const
 {
-	if (!CachedRawComponentPtr)
-	{
-		UInstancedStaticMeshComponent* GeneratedComponentPtr = Cast<UInstancedStaticMeshComponent>(GeneratedComponent.Get());
-
-		// Implementation note:
-		// There is no surefire way to make sure that we can use the raw pointer UNLESS it is from the same owner
-		if (GeneratedComponentPtr && Cast<UPCGComponent>(GetOuter()) && GeneratedComponentPtr->GetOwner() == Cast<UPCGComponent>(GetOuter())->GetOwner())
-		{
-			CachedRawComponentPtr = GeneratedComponentPtr;
-		}
-
-		return GeneratedComponentPtr;
-	}
-
-	return CachedRawComponentPtr;
+	TRACE_CPUPROFILER_EVENT_SCOPE(UPCGManagedISMComponent::GetComponent);
+	return Cast<UInstancedStaticMeshComponent>(GeneratedComponent.Get());
 }
 
 void UPCGManagedISMComponent::SetComponent(UInstancedStaticMeshComponent* InComponent)
 {
 	GeneratedComponent = InComponent;
-	CachedRawComponentPtr = InComponent;
 }
 
 void UPCGManagedISKMComponent::PostLoad()
@@ -913,12 +893,6 @@ void UPCGManagedISKMComponent::PostLoad()
 	GetComponent();
 }
 
-void UPCGManagedISKMComponent::ForgetComponent()
-{
-	Super::ForgetComponent();
-	CachedRawComponentPtr = nullptr;
-}
-
 void UPCGManagedISKMComponent::SetDescriptor(const FSkinnedMeshComponentDescriptor& InDescriptor)
 {
 	bHasDescriptor = true;
@@ -1024,56 +998,22 @@ void UPCGManagedISKMComponent::SetRootLocation(const FVector& InRootLocation)
 
 UInstancedSkinnedMeshComponent* UPCGManagedISKMComponent::GetComponent() const
 {
-	if (!CachedRawComponentPtr)
-	{
-		UInstancedSkinnedMeshComponent* GeneratedComponentPtr = Cast<UInstancedSkinnedMeshComponent>(GeneratedComponent.Get());
-
-		// Implementation note:
-		// There is no surefire way to make sure that we can use the raw pointer UNLESS it is from the same owner
-		if (GeneratedComponentPtr && Cast<UPCGComponent>(GetOuter()) && GeneratedComponentPtr->GetOwner() == Cast<UPCGComponent>(GetOuter())->GetOwner())
-		{
-			CachedRawComponentPtr = GeneratedComponentPtr;
-		}
-
-		return GeneratedComponentPtr;
-	}
-
-	return CachedRawComponentPtr;
+	TRACE_CPUPROFILER_EVENT_SCOPE(UPCGManagedISKMComponent::GetComponent);
+	return Cast<UInstancedSkinnedMeshComponent>(GeneratedComponent.Get());
 }
 
 void UPCGManagedISKMComponent::SetComponent(UInstancedSkinnedMeshComponent* InComponent)
 {
 	GeneratedComponent = InComponent;
-	CachedRawComponentPtr = InComponent;
-}
-
-void UPCGManagedSplineMeshComponent::ForgetComponent()
-{
-	Super::ForgetComponent();
-	CachedRawComponentPtr = nullptr;
 }
 
 USplineMeshComponent* UPCGManagedSplineMeshComponent::GetComponent() const
 {
-	if (!CachedRawComponentPtr)
-	{
-		USplineMeshComponent* GeneratedComponentPtr = Cast<USplineMeshComponent>(GeneratedComponent.Get());
-
-		// Implementation note:
-		// There is no surefire way to make sure that we can use the raw pointer UNLESS it is from the same owner
-		if (GeneratedComponentPtr && Cast<UPCGComponent>(GetOuter()) && GeneratedComponentPtr->GetOwner() == Cast<UPCGComponent>(GetOuter())->GetOwner())
-		{
-			CachedRawComponentPtr = GeneratedComponentPtr;
-		}
-
-		return GeneratedComponentPtr;
-	}
-
-	return CachedRawComponentPtr;
+	TRACE_CPUPROFILER_EVENT_SCOPE(UPCGManagedSplineMeshComponent::GetComponent);
+	return Cast<USplineMeshComponent>(GeneratedComponent.Get());
 }
 
 void UPCGManagedSplineMeshComponent::SetComponent(USplineMeshComponent* InComponent)
 {
 	GeneratedComponent = InComponent;
-	CachedRawComponentPtr = InComponent;
 }
diff --git a/Engine/Plugins/PCG/Source/PCG/Public/PCGManagedResource.h b/Engine/Plugins/PCG/Source/PCG/Public/PCGManagedResource.h
--- a/Engine/Plugins/PCG/Source/PCG/Public/PCGManagedResource.h
+++ b/Engine/Plugins/PCG/Source/PCG/Public/PCGManagedResource.h
@@ -273,7 +273,6 @@ public:
 	virtual bool SupportsComponentReset() const override{ return true; }
 	virtual void MarkAsUsed() override;
 	virtual void MarkAsReused() override;
-	virtual void ForgetComponent() override;
 	//~End UPCGManagedComponents interface
 
 	UInstancedStaticMeshComponent* GetComponent() const;
@@ -308,9 +307,6 @@ protected:
 
 	UPROPERTY()
 	FPCGCrc SettingsCrc;
-
-	// Cached raw pointer to ISM component
-	mutable UInstancedStaticMeshComponent* CachedRawComponentPtr = nullptr;
 };
 
 UCLASS(BlueprintType)
@@ -332,7 +328,6 @@ public:
 	virtual bool SupportsComponentReset() const override{ return true; }
 	virtual void MarkAsUsed() override;
 	virtual void MarkAsReused() override;
-	virtual void ForgetComponent() override;
 	//~End UPCGManagedComponents interface
 
 	UInstancedSkinnedMeshComponent* GetComponent() const;
@@ -361,9 +356,6 @@ protected:
 
 	UPROPERTY()
 	FPCGCrc SettingsCrc;
-
-	// Cached raw pointer to ISKM component
-	mutable UInstancedSkinnedMeshComponent* CachedRawComponentPtr = nullptr;
 };
 
 UCLASS(BlueprintType)
@@ -375,7 +367,6 @@ public:
 	//~Begin UPCGManagedComponents interface
 	virtual void ResetComponent() override { /* Does nothing, but implementation is required to support reuse. */ }
 	virtual bool SupportsComponentReset() const override { return true; }
-	virtual void ForgetComponent() override;
 	//~End UPCGManagedComponents interface
 
 	USplineMeshComponent* GetComponent() const;
@@ -405,7 +396,4 @@ protected:
 
 	UPROPERTY()
 	FPCGCrc SettingsCrc;
-
-	// Cached raw pointer to USplineMeshComponent
-	mutable USplineMeshComponent* CachedRawComponentPtr = nullptr;
 };
diff --git a/Engine/Plugins/PCGInterops/PCGGeometryScriptInterop/Source/PCGGeometryScriptInterop/Private/Resources/PCGDynamicMeshManagedComponent.cpp b/Engine/Plugins/PCGInterops/PCGGeometryScriptInterop/Source/PCGGeometryScriptInterop/Private/Resources/PCGDynamicMeshManagedComponent.cpp
--- a/Engine/Plugins/PCGInterops/PCGGeometryScriptInterop/Source/PCGGeometryScriptInterop/Private/Resources/PCGDynamicMeshManagedComponent.cpp
+++ b/Engine/Plugins/PCGInterops/PCGGeometryScriptInterop/Source/PCGGeometryScriptInterop/Private/Resources/PCGDynamicMeshManagedComponent.cpp
@@ -11,12 +11,6 @@
 
 #define LOCTEXT_NAMESPACE "PCGDynamicMeshManagedComponent"
 
-void UPCGDynamicMeshManagedComponent::ForgetComponent()
-{
-	Super::ForgetComponent();
-	CachedRawComponentPtr = nullptr;
-}
-
 void UPCGDynamicMeshManagedComponent::MarkAsReused()
 {
 	Super::MarkAsReused();
@@ -40,27 +34,13 @@ void UPCGDynamicMeshManagedComponent::MarkAsReused()
 
 UDynamicMeshComponent* UPCGDynamicMeshManagedComponent::GetComponent() const
 {
-	if (!CachedRawComponentPtr)
-	{
-		UDynamicMeshComponent* GeneratedComponentPtr = Cast<UDynamicMeshComponent>(GeneratedComponent.Get());
-
-		// Implementation note:
-		// There is no surefire way to make sure that we can use the raw pointer UNLESS it is from the same owner
-		if (GeneratedComponentPtr && Cast<UPCGComponent>(GetOuter()) && GeneratedComponentPtr->GetOwner() == Cast<UPCGComponent>(GetOuter())->GetOwner())
-		{
-			CachedRawComponentPtr = GeneratedComponentPtr;
-		}
-
-		return GeneratedComponentPtr;
-	}
-
-	return CachedRawComponentPtr;
+	TRACE_CPUPROFILER_EVENT_SCOPE(UPCGDynamicMeshManagedComponent::GetComponent);
+	return Cast<UDynamicMeshComponent>(GeneratedComponent.Get());
 }
 
 void UPCGDynamicMeshManagedComponent::SetComponent(UDynamicMeshComponent* InComponent)
 {
 	GeneratedComponent = InComponent;
-	CachedRawComponentPtr = InComponent;
 }
 
 UPCGDynamicMeshManagedComponent* PCGDynamicMeshManagedComponent::GetOrCreateDynamicMeshManagedComponent(FPCGContext* Context, const UPCGSettingsInterface* SettingsInterface, const UPCGDynamicMeshData* InMeshData, AActor* TargetActor, TOptional<EPCGEditorDirtyMode> OptionalDirtyModeOverride)
diff --git a/Engine/Plugins/PCGInterops/PCGGeometryScriptInterop/Source/PCGGeometryScriptInterop/Public/Resources/PCGDynamicMeshManagedComponent.h b/Engine/Plugins/PCGInterops/PCGGeometryScriptInterop/Source/PCGGeometryScriptInterop/Public/Resources/PCGDynamicMeshManagedComponent.h
--- a/Engine/Plugins/PCGInterops/PCGGeometryScriptInterop/Source/PCGGeometryScriptInterop/Public/Resources/PCGDynamicMeshManagedComponent.h
+++ b/Engine/Plugins/PCGInterops/PCGGeometryScriptInterop/Source/PCGGeometryScriptInterop/Public/Resources/PCGDynamicMeshManagedComponent.h
@@ -24,7 +24,6 @@ public:
 	virtual void ResetComponent() override { /* Does nothing, but implementation is required to support reuse. */ }
 	virtual bool SupportsComponentReset() const override { return true; }
 	virtual void MarkAsReused() override;
-	virtual void ForgetComponent() override;
 	//~End UPCGManagedComponents interface
 
 	PCGGEOMETRYSCRIPTINTEROP_API UDynamicMeshComponent* GetComponent() const;
@@ -36,9 +35,6 @@ public:
 protected:
 	UPROPERTY(Transient)
 	uint64 DataUID = -1;
-
-	// Cached raw pointer to UDynamicMeshComponent
-	mutable UDynamicMeshComponent* CachedRawComponentPtr = nullptr;
 };
 
 namespace PCGDynamicMeshManagedComponent
