[PCG] Convert StaticMeshSpawner node to support PCGPointArrayData

#rnx

diff --git a/Engine/Plugins/PCG/Source/PCG/Private/Elements/PCGStaticMeshSpawner.cpp b/Engine/Plugins/PCG/Source/PCG/Private/Elements/PCGStaticMeshSpawner.cpp
--- a/Engine/Plugins/PCG/Source/PCG/Private/Elements/PCGStaticMeshSpawner.cpp
+++ b/Engine/Plugins/PCG/Source/PCG/Private/Elements/PCGStaticMeshSpawner.cpp
@@ -6,7 +6,7 @@
 #include "PCGCustomVersion.h"
 #include "PCGManagedResource.h"
 #include "Compute/PCGComputeCommon.h"
-#include "Data/PCGPointData.h"
+#include "Data/PCGBasePointData.h"
 #include "Data/PCGSpatialData.h"
 #include "Elements/PCGStaticMeshSpawnerContext.h"
 #include "Elements/PCGStaticMeshSpawnerKernel.h"
@@ -194,7 +194,7 @@ bool FPCGStaticMeshSpawnerElement::PrepareDataInternal(FPCGContext* InContext) c
 				continue;
 			}
 
-			const UPCGPointData* PointData = SpatialData->ToPointData(Context);
+			const UPCGBasePointData* PointData = SpatialData->ToBasePointData(Context);
 			if (!PointData)
 			{
 				PCGE_LOG(Error, GraphAndLog, LOCTEXT("NoPointDataInInput", "Unable to get point data from input"));
@@ -214,8 +214,15 @@ bool FPCGStaticMeshSpawnerElement::PrepareDataInternal(FPCGContext* InContext) c
 			{
 				FPCGTaggedData& Output = Outputs.Add_GetRef(Input);
 
-				UPCGPointData* OutputPointData = FPCGContext::NewObject_AnyThread<UPCGPointData>(Context);
-				OutputPointData->InitializeFromData(PointData);
+				UPCGBasePointData* OutputPointData = FPCGContext::NewPointData_AnyThread(Context);
+
+				// @todo_pcg: this could probably be inherited, since not all selectors output all points this works for now.
+				FPCGInitializeFromDataParams InitializeFromDataParams(PointData);
+				InitializeFromDataParams.bInheritSpatialData = false;
+				OutputPointData->InitializeFromDataWithParams(InitializeFromDataParams);
+
+				OutputPointData->SetNumPoints(PointData->GetNumPoints());
+				OutputPointData->AllocateProperties(PointData->GetAllocatedProperties());
 
 				if (OutputPointData->Metadata->HasAttribute(Settings->OutAttributeName))
 				{
@@ -243,7 +250,7 @@ bool FPCGStaticMeshSpawnerElement::PrepareDataInternal(FPCGContext* InContext) c
 		if (!Context->bSelectionDone)
 		{
 			check(Context->CurrentPointData);
-			Context->bSelectionDone = Settings->MeshSelectorParameters->SelectInstances(*Context, Settings, Context->CurrentPointData, Context->MeshInstancesData.Last().MeshInstances, Context->CurrentOutputPointData);
+			Context->bSelectionDone = Settings->MeshSelectorParameters->SelectMeshInstances(*Context, Settings, Context->CurrentPointData, Context->MeshInstancesData.Last().MeshInstances, Context->CurrentOutputPointData);
 		}
 
 		if (!Context->bSelectionDone)
@@ -338,18 +345,20 @@ bool FPCGStaticMeshSpawnerElement::ExecuteInternal(FPCGContext* InContext) const
 				// Now that the mesh is loaded/spawned, set the bounds to out points if requested.
 				if (MeshInstance.Descriptor.StaticMesh && Settings->bApplyMeshBoundsToPoints)
 				{
-					if (TMap<UPCGPointData*, TArray<int32>>* OutPointDataToPointIndex = Context->MeshToOutPoints.Find(MeshInstance.Descriptor.StaticMesh))
+					if (TMap<UPCGBasePointData*, TArray<int32>>* OutPointDataToPointIndex = Context->MeshToOutPoints.Find(MeshInstance.Descriptor.StaticMesh))
 					{
 						const FBox Bounds = MeshInstance.Descriptor.StaticMesh->GetBoundingBox();
-						for (TPair<UPCGPointData*, TArray<int32>>& It : *OutPointDataToPointIndex)
+						for (TPair<UPCGBasePointData*, TArray<int32>>& It : *OutPointDataToPointIndex)
 						{
 							check(It.Key);
-							TArray<FPCGPoint>& OutPoints = It.Key->GetMutablePoints();
+
+							TPCGValueRange<FVector> BoundsMinRange = It.Key->GetBoundsMinValueRange();
+							TPCGValueRange<FVector> BoundsMaxRange = It.Key->GetBoundsMaxValueRange();
+
 							for (int32 Index : It.Value)
 							{
-								FPCGPoint& Point = OutPoints[Index];
-								Point.BoundsMin = Bounds.Min;
-								Point.BoundsMax = Bounds.Max;
+								BoundsMinRange[Index] = Bounds.Min;
+								BoundsMaxRange[Index] = Bounds.Max;
 							}
 						}
 					}
diff --git a/Engine/Plugins/PCG/Source/PCG/Private/InstanceDataPackers/PCGInstanceDataPackerBase.cpp b/Engine/Plugins/PCG/Source/PCG/Private/InstanceDataPackers/PCGInstanceDataPackerBase.cpp
--- a/Engine/Plugins/PCG/Source/PCG/Private/InstanceDataPackers/PCGInstanceDataPackerBase.cpp
+++ b/Engine/Plugins/PCG/Source/PCG/Private/InstanceDataPackers/PCGInstanceDataPackerBase.cpp
@@ -64,10 +64,12 @@ void UPCGInstanceDataPackerBase::PackCustomDataFromAttributes(const FPCGMeshInst
 
 void UPCGInstanceDataPackerBase::PackCustomDataFromAttributes(const FPCGMeshInstanceList& InstanceList, const TArray<const FPCGMetadataAttributeBase*>& Attributes, FPCGPackedCustomData& OutPackedCustomData) const
 {
-	const UPCGPointData* PointData = InstanceList.PointData.Get();
+	const UPCGBasePointData* PointData = InstanceList.PointData.Get();
+	const TConstPCGValueRange<int64> MetadataEntryRange = PointData ? PointData->GetConstMetadataEntryValueRange() : TConstPCGValueRange<int64>();
+
 	for (const int32 Index : InstanceList.InstancesIndices)
 	{
-		const PCGMetadataEntryKey EntryKey = PointData ? PointData->GetPoints()[Index].MetadataEntry : Index;
+		const PCGMetadataEntryKey EntryKey = PointData ? MetadataEntryRange[Index] : Index;
 		for (const FPCGMetadataAttributeBase* AttributeBase : Attributes)
 		{
 			check(AttributeBase);
diff --git a/Engine/Plugins/PCG/Source/PCG/Private/InstanceDataPackers/PCGInstanceDataPackerByAttribute.cpp b/Engine/Plugins/PCG/Source/PCG/Private/InstanceDataPackers/PCGInstanceDataPackerByAttribute.cpp
--- a/Engine/Plugins/PCG/Source/PCG/Private/InstanceDataPackers/PCGInstanceDataPackerByAttribute.cpp
+++ b/Engine/Plugins/PCG/Source/PCG/Private/InstanceDataPackers/PCGInstanceDataPackerByAttribute.cpp
@@ -49,12 +49,12 @@ void UPCGInstanceDataPackerByAttribute::PackInstances_Implementation(FPCGContext
 	{
 		TUniquePtr<const IPCGAttributeAccessor> Accessor = PCGAttributeAccessorHelpers::CreateConstAccessor(InSpatialData, Selector);
 		TUniquePtr<const IPCGAttributeAccessorKeys> Keys;
-		const UPCGPointData* PointData = InstanceList.PointData.Get();
+		const UPCGBasePointData* PointData = InstanceList.PointData.Get();
 		TArray<const PCGMetadataEntryKey> ExtractedKeys;
 		
 		if (InSpatialData == PointData)
 		{
-			Keys = MakeUnique<const FPCGAttributeAccessorKeysPointsSubset>(PointData->GetPoints(), InstanceList.InstancesIndices);
+			Keys = MakeUnique<const FPCGAttributeAccessorKeysPointsSubset>(PointData, InstanceList.InstancesIndices);
 		}
 		else
 		{
diff --git a/Engine/Plugins/PCG/Source/PCG/Private/MeshSelectors/PCGMeshSelectorBase.cpp b/Engine/Plugins/PCG/Source/PCG/Private/MeshSelectors/PCGMeshSelectorBase.cpp
--- a/Engine/Plugins/PCG/Source/PCG/Private/MeshSelectors/PCGMeshSelectorBase.cpp
+++ b/Engine/Plugins/PCG/Source/PCG/Private/MeshSelectors/PCGMeshSelectorBase.cpp
@@ -9,6 +9,33 @@
 
 #define LOCTEXT_NAMESPACE "PCGMeshSelectorBase"
 
+
+bool UPCGMeshSelectorBase::SelectMeshInstances(
+	FPCGStaticMeshSpawnerContext& Context,
+	const UPCGStaticMeshSpawnerSettings* Settings,
+	const UPCGBasePointData* InPointData,
+	TArray<FPCGMeshInstanceList>& OutMeshInstances,
+	UPCGBasePointData* OutPointData) const
+{
+	// This method should be overwritten by the subclasses but in case that subclass hasn't yet been updated try and call the old signature
+	const UPCGPointData* TypedInPointData = Cast<UPCGPointData>(InPointData);
+	UPCGPointData* TypedOutPointData = Cast<UPCGPointData>(OutPointData);
+
+	const FText Message = FText::Format(LOCTEXT("SelectMeshInstancesNotImplemented", "MeshSelector class '{0}' needs to implement SelectMeshInstances"), FText::FromString(GetClass()->GetName()));
+	if (TypedInPointData && TypedOutPointData)
+	{
+		PCGE_LOG_C(Warning, GraphAndLog, &Context, Message);
+PRAGMA_DISABLE_DEPRECATION_WARNINGS
+		return SelectInstances(Context, Settings, TypedInPointData, OutMeshInstances, TypedOutPointData);
+PRAGMA_ENABLE_DEPRECATION_WARNINGS
+	}
+	else
+	{
+		PCGE_LOG_C(Error, GraphAndLog, &Context, Message);
+		return false;
+	}
+}
+
 void FPCGMeshMaterialOverrideHelper::Initialize(
 	FPCGContext& InContext,
 	bool bInUseMaterialOverrideAttributes,
diff --git a/Engine/Plugins/PCG/Source/PCG/Private/MeshSelectors/PCGMeshSelectorByAttribute.cpp b/Engine/Plugins/PCG/Source/PCG/Private/MeshSelectors/PCGMeshSelectorByAttribute.cpp
--- a/Engine/Plugins/PCG/Source/PCG/Private/MeshSelectors/PCGMeshSelectorByAttribute.cpp
+++ b/Engine/Plugins/PCG/Source/PCG/Private/MeshSelectors/PCGMeshSelectorByAttribute.cpp
@@ -2,7 +2,7 @@
 
 #include "MeshSelectors/PCGMeshSelectorByAttribute.h"
 
-#include "Data/PCGPointData.h"
+#include "Data/PCGBasePointData.h"
 #include "Data/PCGSpatialData.h"
 #include "Elements/PCGStaticMeshSpawner.h"
 #include "Elements/PCGStaticMeshSpawnerContext.h"
@@ -28,7 +28,7 @@ namespace PCGMeshSelectorAttribute
 		TSoftObjectPtr<UStaticMesh> Mesh,
 		const TArray<TSoftObjectPtr<UMaterialInterface>>& MaterialOverrides,
 		bool bReverseCulling,
-		const UPCGPointData* InPointData,
+		const UPCGBasePointData* InPointData,
 		const int AttributePartitionIndex = INDEX_NONE)
 	{
 		for (FPCGMeshInstanceList& InstanceList : InstanceLists)
@@ -108,12 +108,12 @@ void UPCGMeshSelectorByAttribute::PostLoad()
 	TemplateDescriptor.PostLoadFixup(this);
 }
 
-bool UPCGMeshSelectorByAttribute::SelectInstances(
+bool UPCGMeshSelectorByAttribute::SelectMeshInstances(
 	FPCGStaticMeshSpawnerContext& Context,
 	const UPCGStaticMeshSpawnerSettings* Settings,
-	const UPCGPointData* InPointData,
+	const UPCGBasePointData* InPointData,
 	TArray<FPCGMeshInstanceList>& OutMeshInstances,
-	UPCGPointData* OutPointData) const
+	UPCGBasePointData* OutPointData) const
 {
 	TRACE_CPUPROFILER_EVENT_SCOPE(UPCGMeshSelectorByAttribute::SelectInstances);
 
@@ -160,7 +160,11 @@ bool UPCGMeshSelectorByAttribute::SelectInstances(
 	if (Context.CurrentPointIndex == 0 && OutPointData)
 	{
 		TRACE_CPUPROFILER_EVENT_SCOPE(UPCGMeshSelectorByAttribute::SetupOutPointData);
-		OutPointData->SetPoints(InPointData->GetPoints());
+
+		const int32 NumPoints = InPointData->GetNumPoints();
+		OutPointData->SetNumPoints(NumPoints);
+		InPointData->CopyPointsTo(OutPointData, 0, 0, InPointData->GetNumPoints());
+
 		OutPointData->Metadata->DeleteAttribute(Settings->OutAttributeName);
 		OutPointData->Metadata->CopyAttribute(InPointData->Metadata, AttributeName, Settings->OutAttributeName);
 	}
@@ -235,10 +239,12 @@ bool UPCGMeshSelectorByAttribute::SelectInstances(
 
 	// Assign points to entries
 	int32 CurrentPartitionIndex = Context.CurrentPointIndex; // misnomer but we're reusing another concept from the context
-	const TArray<FPCGPoint>& Points = InPointData->GetPoints();
 	{
 		TRACE_CPUPROFILER_EVENT_SCOPE(UPCGMeshSelectorByAttribute::SelectEntries::PushingPointsToInstanceLists);
 
+		const TConstPCGValueRange<FTransform> InTransformRange = InPointData->GetConstTransformValueRange();
+		const TConstPCGValueRange<int64> InMetadataEntryRange = InPointData->GetConstMetadataEntryValueRange();
+				
 		// TODO: Revisit this when attribute partitioning is returned in a more optimized form
 		// The partition index is used to assign the point to the correct partition's instance
 		while(CurrentPartitionIndex < Context.AttributeOverridePartition.Num())
@@ -267,8 +273,8 @@ bool UPCGMeshSelectorByAttribute::SelectInstances(
 				int32 WriteIndex = 0;
 				for (int32 ReadIndex = 0; ReadIndex < Partition.Num(); ++ReadIndex)
 				{
-					const FPCGPoint& Point = Points[Partition[ReadIndex]];
-					if (Point.Transform.GetDeterminant() < 0)
+					const FTransform& InTransform = InTransformRange[Partition[ReadIndex]];
+					if (InTransform.GetDeterminant() < 0)
 					{
 						ReverseInstances.Add(Partition[ReadIndex]);
 					}
@@ -285,15 +291,16 @@ bool UPCGMeshSelectorByAttribute::SelectInstances(
 				Partition.SetNum(WriteIndex);
 			}
 
-			auto AddPointsToInstanceList = [&OutMeshInstances, &CurrentPartitionDescriptor, &MaterialOverrideHelper, &Points, PartitionIndex, InPointData](const TArray<int32>& PointIndices, bool bReverseTransform)
+			auto AddPointsToInstanceList = [&OutMeshInstances, &CurrentPartitionDescriptor, &MaterialOverrideHelper, &InTransformRange, &InMetadataEntryRange, PartitionIndex, InPointData](const TArray<int32>& PointIndices, bool bReverseTransform)
 			{
 				if (MaterialOverrideHelper.OverridesMaterials())
 				{
 					for (int32 PointIndex = 0; PointIndex < PointIndices.Num(); ++PointIndex)
 					{
-						const FPCGPoint& Point = Points[PointIndices[PointIndex]];
-						FPCGMeshInstanceList& InstanceList = PCGMeshSelectorAttribute::GetInstanceList(OutMeshInstances, CurrentPartitionDescriptor, CurrentPartitionDescriptor.StaticMesh, MaterialOverrideHelper.GetMaterialOverrides(Point.MetadataEntry), bReverseTransform, InPointData, PartitionIndex);
-						InstanceList.Instances.Emplace(Point.Transform);
+						const FTransform& InTransform = InTransformRange[PointIndices[PointIndex]];
+						const int64& InMetadataEntry = InMetadataEntryRange[PointIndices[PointIndex]];
+						FPCGMeshInstanceList& InstanceList = PCGMeshSelectorAttribute::GetInstanceList(OutMeshInstances, CurrentPartitionDescriptor, CurrentPartitionDescriptor.StaticMesh, MaterialOverrideHelper.GetMaterialOverrides(InMetadataEntry), bReverseTransform, InPointData, PartitionIndex);
+						InstanceList.Instances.Emplace(InTransform);
 						InstanceList.InstancesIndices.Emplace(PointIndices[PointIndex]);
 					}
 				}
@@ -309,8 +316,8 @@ bool UPCGMeshSelectorByAttribute::SelectInstances(
 
 					for(int32 PointIndex = 0; PointIndex < PointIndices.Num(); ++PointIndex)
 					{
-						const FPCGPoint& Point = Points[PointIndices[PointIndex]];
-						InstanceList.Instances[InstanceOffset + PointIndex] = Point.Transform;
+						const FTransform& InTransform = InTransformRange[PointIndices[PointIndex]];
+						InstanceList.Instances[InstanceOffset + PointIndex] = InTransform;
 					}
 				}
 			};
diff --git a/Engine/Plugins/PCG/Source/PCG/Private/MeshSelectors/PCGMeshSelectorWeighted.cpp b/Engine/Plugins/PCG/Source/PCG/Private/MeshSelectors/PCGMeshSelectorWeighted.cpp
--- a/Engine/Plugins/PCG/Source/PCG/Private/MeshSelectors/PCGMeshSelectorWeighted.cpp
+++ b/Engine/Plugins/PCG/Source/PCG/Private/MeshSelectors/PCGMeshSelectorWeighted.cpp
@@ -2,11 +2,11 @@
 
 #include "MeshSelectors/PCGMeshSelectorWeighted.h"
 
-#include "Data/PCGPointData.h"
+#include "Data/PCGBasePointData.h"
 #include "Data/PCGSpatialData.h"
 #include "Elements/PCGStaticMeshSpawner.h"
 #include "Elements/PCGStaticMeshSpawnerContext.h"
-#include "Helpers/PCGBlueprintHelpers.h"
+#include "Helpers/PCGHelpers.h"
 #include "MeshSelectors/PCGMeshSelectorBase.h"
 
 #include "Engine/StaticMesh.h"
@@ -23,7 +23,7 @@ namespace PCGMeshSelectorWeighted
 		bool bUseMaterialOverrides,
 		const TArray<TSoftObjectPtr<UMaterialInterface>>& InMaterialOverrides,
 		bool bInReverseCulling,
-		const UPCGPointData* InPointData)
+		const UPCGBasePointData* InPointData)
 	{
 		check(InstanceLists.Num() > 0);
 
@@ -108,12 +108,12 @@ void FPCGMeshSelectorWeightedEntry::ApplyDeprecation()
 }
 #endif
 
-bool UPCGMeshSelectorWeighted::SelectInstances(
+bool UPCGMeshSelectorWeighted::SelectMeshInstances(
 	FPCGStaticMeshSpawnerContext& Context,
 	const UPCGStaticMeshSpawnerSettings* Settings, 
-	const UPCGPointData* InPointData,
+	const UPCGBasePointData* InPointData,
 	TArray<FPCGMeshInstanceList>& OutMeshInstances,
-	UPCGPointData* OutPointData) const
+	UPCGBasePointData* OutPointData) const
 {
 	if (!InPointData)
 	{
@@ -169,11 +169,14 @@ bool UPCGMeshSelectorWeighted::SelectInstances(
 		return true;
 	}
 
-	TArray<FPCGPoint>* OutPoints = nullptr;
 	FPCGMetadataAttribute<FString>* OutAttribute = nullptr;
 
+	FPCGPointValueRanges OutRanges;
+	
 	if (OutPointData)
 	{
+		OutRanges = FPCGPointValueRanges(OutPointData);
+
 		check(OutPointData->Metadata);
 
 		if (!OutPointData->Metadata->HasAttribute(Settings->OutAttributeName)) 
@@ -188,7 +191,6 @@ bool UPCGMeshSelectorWeighted::SelectInstances(
 			if (OutAttributeBase->GetTypeId() == PCG::Private::MetadataTypes<FString>::Id)
 			{
 				OutAttribute = static_cast<FPCGMetadataAttribute<FString>*>(OutAttributeBase);
-				OutPoints = &OutPointData->GetMutablePoints();
 			}
 			else
 			{
@@ -200,6 +202,8 @@ bool UPCGMeshSelectorWeighted::SelectInstances(
 	// Assign points to entries
 	{
 		TRACE_CPUPROFILER_EVENT_SCOPE(FPCGStaticMeshSpawnerElement::Execute::SelectEntries);
+		
+		const FConstPCGPointValueRanges InRanges(InPointData);
 
 		int32 CurrentPointIndex = Context.CurrentPointIndex;
 		int32 LastCheckpointIndex = CurrentPointIndex;
@@ -207,12 +211,15 @@ bool UPCGMeshSelectorWeighted::SelectInstances(
 		TMap<TSoftObjectPtr<UStaticMesh>, PCGMetadataValueKey>& MeshToValueKey = Context.MeshToValueKey;
 		const int32 TotalWeight = CumulativeWeights.Last();
 
-		const TArray<FPCGPoint>& Points = InPointData->GetPoints();
-		while (CurrentPointIndex < Points.Num())
+		while (CurrentPointIndex < InPointData->GetNumPoints())
 		{
-			const FPCGPoint& Point = Points[CurrentPointIndex++];
+			const int32 InPointIndex = CurrentPointIndex++;
 
-			FRandomStream RandomSource = UPCGBlueprintHelpers::GetRandomStreamFromPoint(Point, Settings, Context.SourceComponent.Get());
+			const FTransform& InTransform = InRanges.TransformRange[InPointIndex];
+			const int64& InMetadataEntry = InRanges.MetadataEntryRange[InPointIndex];
+			const int32& InSeed = InRanges.SeedRange[InPointIndex];
+			
+			FRandomStream RandomSource = PCGHelpers::GetRandomStreamFromSeed(InSeed, Settings, Context.SourceComponent.Get());
 			const int RandomWeightedPick = RandomSource.RandRange(0, TotalWeight - 1);
 
 			int RandomPick = 0;
@@ -223,16 +230,25 @@ bool UPCGMeshSelectorWeighted::SelectInstances(
 
 			if(RandomPick < MeshInstances.Num())
 			{
-				const bool bNeedsReverseCulling = (Point.Transform.GetDeterminant() < 0);
-				FPCGMeshInstanceList& InstanceList = PCGMeshSelectorWeighted::GetInstanceList(MeshInstances[RandomPick], bUseAttributeMaterialOverrides, MaterialOverrideHelper.GetMaterialOverrides(Point.MetadataEntry), bNeedsReverseCulling, InPointData);
-				InstanceList.Instances.Emplace(Point.Transform);
+				const bool bNeedsReverseCulling = (InTransform.GetDeterminant() < 0);
+				FPCGMeshInstanceList& InstanceList = PCGMeshSelectorWeighted::GetInstanceList(MeshInstances[RandomPick], bUseAttributeMaterialOverrides, MaterialOverrideHelper.GetMaterialOverrides(InMetadataEntry), bNeedsReverseCulling, InPointData);
+				InstanceList.Instances.Emplace(InTransform);
 				InstanceList.InstancesIndices.Emplace(CurrentPointIndex - 1); // - 1 because it is already incremented.
 
 				const TSoftObjectPtr<UStaticMesh>& Mesh = InstanceList.Descriptor.StaticMesh;
 
 				if (OutPointData && OutAttribute)
 				{
-					FPCGPoint& OutPoint = OutPoints->Add_GetRef(Point);
+					OutRanges.TransformRange[Context.CurrentWriteIndex] = InTransform;
+					OutRanges.MetadataEntryRange[Context.CurrentWriteIndex] = InMetadataEntry;
+					OutRanges.SeedRange[Context.CurrentWriteIndex] = InSeed;
+
+					OutRanges.SteepnessRange[Context.CurrentWriteIndex] = InRanges.SteepnessRange[InPointIndex];
+					OutRanges.DensityRange[Context.CurrentWriteIndex] = InRanges.DensityRange[InPointIndex];
+					OutRanges.BoundsMinRange[Context.CurrentWriteIndex] = InRanges.BoundsMinRange[InPointIndex];
+					OutRanges.BoundsMaxRange[Context.CurrentWriteIndex] = InRanges.BoundsMaxRange[InPointIndex];
+					OutRanges.ColorRange[Context.CurrentWriteIndex] = InRanges.ColorRange[InPointIndex];
+
 					PCGMetadataValueKey* OutValueKey = MeshToValueKey.Find(Mesh);
 					if (!OutValueKey)
 					{
@@ -242,8 +258,10 @@ bool UPCGMeshSelectorWeighted::SelectInstances(
 
 					check(OutValueKey);
 
-					OutPointData->Metadata->InitializeOnSet(OutPoint.MetadataEntry);
-					OutAttribute->SetValueFromValueKey(OutPoint.MetadataEntry, *OutValueKey);
+					OutPointData->Metadata->InitializeOnSet(OutRanges.MetadataEntryRange[Context.CurrentWriteIndex]);
+					OutAttribute->SetValueFromValueKey(OutRanges.MetadataEntryRange[Context.CurrentWriteIndex], *OutValueKey);
+
+					++Context.CurrentWriteIndex;
 
 					if (Settings->bApplyMeshBoundsToPoints)
 					{
@@ -271,8 +289,13 @@ bool UPCGMeshSelectorWeighted::SelectInstances(
 		Context.CurrentPointIndex = CurrentPointIndex;
 	}
 
-	if (Context.CurrentPointIndex == InPointData->GetPoints().Num())
+	if (Context.CurrentPointIndex == InPointData->GetNumPoints())
 	{
+		if (OutPointData)
+		{
+			OutPointData->SetNumPoints(Context.CurrentWriteIndex);
+		}
+
 		// Finally, collapse to OutMeshInstances
 		for (TArray<FPCGMeshInstanceList>& PickedMeshInstances : MeshInstances)
 		{
diff --git a/Engine/Plugins/PCG/Source/PCG/Private/MeshSelectors/PCGMeshSelectorWeightedByCategory.cpp b/Engine/Plugins/PCG/Source/PCG/Private/MeshSelectors/PCGMeshSelectorWeightedByCategory.cpp
--- a/Engine/Plugins/PCG/Source/PCG/Private/MeshSelectors/PCGMeshSelectorWeightedByCategory.cpp
+++ b/Engine/Plugins/PCG/Source/PCG/Private/MeshSelectors/PCGMeshSelectorWeightedByCategory.cpp
@@ -2,11 +2,11 @@
 
 #include "MeshSelectors/PCGMeshSelectorWeightedByCategory.h"
 
-#include "Data/PCGPointData.h"
+#include "Data/PCGBasePointData.h"
 #include "Data/PCGSpatialData.h"
 #include "Elements/PCGStaticMeshSpawner.h"
 #include "Elements/PCGStaticMeshSpawnerContext.h"
-#include "Helpers/PCGBlueprintHelpers.h"
+#include "Helpers/PCGHelpers.h"
 #include "MeshSelectors/PCGMeshSelectorWeighted.h"
 
 #include "Engine/StaticMesh.h"
@@ -26,12 +26,12 @@ void FPCGWeightedByCategoryEntryList::ApplyDeprecation()
 }
 #endif
 
-bool UPCGMeshSelectorWeightedByCategory::SelectInstances(
+bool UPCGMeshSelectorWeightedByCategory::SelectMeshInstances(
 	FPCGStaticMeshSpawnerContext& Context,
 	const UPCGStaticMeshSpawnerSettings* Settings,
-	const UPCGPointData* InPointData,
+	const UPCGBasePointData* InPointData,
 	TArray<FPCGMeshInstanceList>& OutMeshInstances,
-	UPCGPointData* OutPointData) const
+	UPCGBasePointData* OutPointData) const
 {
 	if (!InPointData)
 	{
@@ -156,11 +156,14 @@ bool UPCGMeshSelectorWeightedByCategory::SelectInstances(
 		return true;
 	}
 
-	TArray<FPCGPoint>* OutPoints = nullptr;
 	FPCGMetadataAttribute<FString>* OutAttribute = nullptr;
 
+	FPCGPointValueRanges OutRanges;
+
 	if (OutPointData)
 	{
+		OutRanges = FPCGPointValueRanges(OutPointData);
+		
 		check(OutPointData->Metadata);
 
 		if (!OutPointData->Metadata->HasAttribute(Settings->OutAttributeName)) 
@@ -175,7 +178,6 @@ bool UPCGMeshSelectorWeightedByCategory::SelectInstances(
 			if (OutAttributeBase->GetTypeId() == PCG::Private::MetadataTypes<FString>::Id)
 			{
 				OutAttribute = static_cast<FPCGMetadataAttribute<FString>*>(OutAttributeBase);
-				OutPoints = &OutPointData->GetMutablePoints();
 			}
 			else
 			{
@@ -192,12 +194,17 @@ bool UPCGMeshSelectorWeightedByCategory::SelectInstances(
 	TMap<TSoftObjectPtr<UStaticMesh>, PCGMetadataValueKey>& MeshToValueKey = Context.MeshToValueKey;
 
 	// Assign points to entries
-	const TArray<FPCGPoint>& Points = InPointData->GetPoints();
-	while (CurrentPointIndex < Points.Num())
+	const FConstPCGPointValueRanges InRanges(InPointData);
+
+	while (CurrentPointIndex < InPointData->GetNumPoints())
 	{
-		const FPCGPoint& Point = Points[CurrentPointIndex++];
+		const int32 InPointIndex = CurrentPointIndex++;
 
-		const PCGMetadataValueKey ValueKey = Attribute->GetValueKey(Point.MetadataEntry);
+		const FTransform& InTransform = InRanges.TransformRange[InPointIndex];
+		const int64& InMetadataEntry = InRanges.MetadataEntryRange[InPointIndex];
+		const int32& InSeed = InRanges.SeedRange[InPointIndex];
+
+		const PCGMetadataValueKey ValueKey = Attribute->GetValueKey(InMetadataEntry);
 
 		// if no mesh list was processed for this attribute value, fallback to the default mesh list
 		FPCGInstancesAndWeights* InstancesAndWeights = CategoryEntryToInstancesAndWeights.Find(ValueKey);
@@ -216,7 +223,7 @@ bool UPCGMeshSelectorWeightedByCategory::SelectInstances(
 
 		const int TotalWeight = InstancesAndWeights->CumulativeWeights.Last();
 
-		FRandomStream RandomSource = UPCGBlueprintHelpers::GetRandomStreamFromPoint(Point, Settings, Context.SourceComponent.Get());
+		FRandomStream RandomSource = PCGHelpers::GetRandomStreamFromSeed(InSeed, Settings, Context.SourceComponent.Get());
 		const int RandomWeightedPick = RandomSource.RandRange(0, TotalWeight - 1);
 
 		int RandomPick = 0;
@@ -227,16 +234,24 @@ bool UPCGMeshSelectorWeightedByCategory::SelectInstances(
 
 		if(RandomPick < InstancesAndWeights->MeshInstances.Num())
 		{
-			const bool bNeedsReverseCulling = (Point.Transform.GetDeterminant() < 0);
-			FPCGMeshInstanceList& InstanceList = PCGMeshSelectorWeighted::GetInstanceList(InstancesAndWeights->MeshInstances[RandomPick], bUseAttributeMaterialOverrides, MaterialOverrideHelper.GetMaterialOverrides(Point.MetadataEntry), bNeedsReverseCulling, InPointData);
-			InstanceList.Instances.Emplace(Point.Transform);
+			const bool bNeedsReverseCulling = (InTransform.GetDeterminant() < 0);
+			FPCGMeshInstanceList& InstanceList = PCGMeshSelectorWeighted::GetInstanceList(InstancesAndWeights->MeshInstances[RandomPick], bUseAttributeMaterialOverrides, MaterialOverrideHelper.GetMaterialOverrides(InMetadataEntry), bNeedsReverseCulling, InPointData);
+			InstanceList.Instances.Emplace(InTransform);
 			InstanceList.InstancesIndices.Emplace(CurrentPointIndex - 1); // - 1 because it is already incremented.
 
 			const TSoftObjectPtr<UStaticMesh>& Mesh = InstanceList.Descriptor.StaticMesh;
 
 			if (OutPointData && OutAttribute)
 			{
-				FPCGPoint& OutPoint = OutPoints->Add_GetRef(Point);
+				OutRanges.TransformRange[Context.CurrentWriteIndex] = InTransform;
+				OutRanges.MetadataEntryRange[Context.CurrentWriteIndex] = InMetadataEntry;
+				OutRanges.SeedRange[Context.CurrentWriteIndex] = InSeed;
+
+				OutRanges.SteepnessRange[Context.CurrentWriteIndex] = InRanges.SteepnessRange[InPointIndex];
+				OutRanges.DensityRange[Context.CurrentWriteIndex] = InRanges.DensityRange[InPointIndex];
+				OutRanges.BoundsMinRange[Context.CurrentWriteIndex] = InRanges.BoundsMinRange[InPointIndex];
+				OutRanges.BoundsMaxRange[Context.CurrentWriteIndex] = InRanges.BoundsMaxRange[InPointIndex];
+				OutRanges.ColorRange[Context.CurrentWriteIndex] = InRanges.ColorRange[InPointIndex];
 
 				PCGMetadataValueKey* OutValueKey = MeshToValueKey.Find(Mesh);
 				if (!OutValueKey)
@@ -247,8 +262,10 @@ bool UPCGMeshSelectorWeightedByCategory::SelectInstances(
 
 				check(OutValueKey);
 
-				OutPointData->Metadata->InitializeOnSet(OutPoint.MetadataEntry);
-				OutAttribute->SetValueFromValueKey(OutPoint.MetadataEntry, *OutValueKey);
+				OutPointData->Metadata->InitializeOnSet(OutRanges.MetadataEntryRange[Context.CurrentWriteIndex]);
+				OutAttribute->SetValueFromValueKey(OutRanges.MetadataEntryRange[Context.CurrentWriteIndex], *OutValueKey);
+
+				++Context.CurrentWriteIndex;
 
 				if (Settings->bApplyMeshBoundsToPoints)
 				{
@@ -275,8 +292,13 @@ bool UPCGMeshSelectorWeightedByCategory::SelectInstances(
 
 	Context.CurrentPointIndex = CurrentPointIndex;
 
-	if (CurrentPointIndex == Points.Num())
+	if (CurrentPointIndex == InPointData->GetNumPoints())
 	{
+		if (OutPointData)
+		{
+			OutPointData->SetNumPoints(Context.CurrentWriteIndex);
+		}
+
 		// Collapse to OutMeshInstances
 		for (TPair<PCGMetadataValueKey, FPCGInstancesAndWeights>& Entry : CategoryEntryToInstancesAndWeights)
 		{
diff --git a/Engine/Plugins/PCG/Source/PCG/Public/Elements/PCGStaticMeshSpawner.h b/Engine/Plugins/PCG/Source/PCG/Public/Elements/PCGStaticMeshSpawner.h
--- a/Engine/Plugins/PCG/Source/PCG/Public/Elements/PCGStaticMeshSpawner.h
+++ b/Engine/Plugins/PCG/Source/PCG/Public/Elements/PCGStaticMeshSpawner.h
@@ -131,4 +131,5 @@ protected:
 	virtual bool ExecuteInternal(FPCGContext* Context) const override;
 	virtual void AbortInternal(FPCGContext* Context) const override;
 	void SpawnStaticMeshInstances(FPCGStaticMeshSpawnerContext* Context, const FPCGMeshInstanceList& InstanceList, AActor* TargetActor, const FPCGPackedCustomData* InPackedCustomData) const;
+	virtual bool SupportsBasePointDataInputs(FPCGContext* InContext) const { return true; }
 };
diff --git a/Engine/Plugins/PCG/Source/PCG/Public/Elements/PCGStaticMeshSpawnerContext.h b/Engine/Plugins/PCG/Source/PCG/Public/Elements/PCGStaticMeshSpawnerContext.h
--- a/Engine/Plugins/PCG/Source/PCG/Public/Elements/PCGStaticMeshSpawnerContext.h
+++ b/Engine/Plugins/PCG/Source/PCG/Public/Elements/PCGStaticMeshSpawnerContext.h
@@ -49,15 +49,16 @@ struct PCG_API FPCGStaticMeshSpawnerContext : public FPCGContext, public IPCGAsy
 	bool bSelectionDone = false;
 	bool bPartitionDone = false;
 
-	const UPCGPointData* CurrentPointData = nullptr;
-	UPCGPointData* CurrentOutputPointData = nullptr;
+	const UPCGBasePointData* CurrentPointData = nullptr;
+	UPCGBasePointData* CurrentOutputPointData = nullptr;
 	FPCGMeshMaterialOverrideHelper MaterialOverrideHelper;
 	int32 CurrentPointIndex = 0;
+	int32 CurrentWriteIndex = 0;
 
 	// Used in all selectors if we have to change the out points bounds by the mesh bounds. Will be empty otherwise.
 	// We need to keep all point indices that will spawn this mesh, in all output point data.
 
-	TMap<TSoftObjectPtr<UStaticMesh>, TMap<UPCGPointData*, TArray<int32>>> MeshToOutPoints;
+	TMap<TSoftObjectPtr<UStaticMesh>, TMap<UPCGBasePointData*, TArray<int32>>> MeshToOutPoints;
 
 	// Used in weighted selector
 	TArray<TArray<FPCGMeshInstanceList>> WeightedMeshInstances;
diff --git a/Engine/Plugins/PCG/Source/PCG/Public/MeshSelectors/PCGMeshSelectorBase.h b/Engine/Plugins/PCG/Source/PCG/Public/MeshSelectors/PCGMeshSelectorBase.h
--- a/Engine/Plugins/PCG/Source/PCG/Public/MeshSelectors/PCGMeshSelectorBase.h
+++ b/Engine/Plugins/PCG/Source/PCG/Public/MeshSelectors/PCGMeshSelectorBase.h
@@ -49,7 +49,7 @@ struct FPCGMeshInstanceList
 	int64 AttributePartitionIndex;
 
 	UPROPERTY(BlueprintReadWrite, EditAnywhere, Category = Settings)
-	TWeakObjectPtr<const UPCGPointData> PointData;
+	TWeakObjectPtr<const UPCGBasePointData> PointData;
 	
 	UPROPERTY(BlueprintReadWrite, EditAnywhere, Category = Settings)
 	TArray<int32> InstancesIndices;
@@ -67,10 +67,18 @@ class PCG_API UPCGMeshSelectorBase : public UObject
 	GENERATED_BODY()
 
 public:
+	UE_DEPRECATED(5.6, "Call / Override SelectMeshInstances with UPCGBasePointData parameter")
 	virtual bool SelectInstances(
 		FPCGStaticMeshSpawnerContext& Context,
 		const UPCGStaticMeshSpawnerSettings* Settings,
 		const UPCGPointData* InPointData,
 		TArray<FPCGMeshInstanceList>& OutMeshInstances,
-		UPCGPointData* OutPointData) const PURE_VIRTUAL(UPCGMeshSelectorBase::SelectInstances, return true;);
+		UPCGPointData* OutPointData) const { return true; }
+
+	virtual bool SelectMeshInstances(
+		FPCGStaticMeshSpawnerContext& Context,
+		const UPCGStaticMeshSpawnerSettings* Settings,
+		const UPCGBasePointData* InPointData,
+		TArray<FPCGMeshInstanceList>& OutMeshInstances,
+		UPCGBasePointData* OutPointData) const;
 };
diff --git a/Engine/Plugins/PCG/Source/PCG/Public/MeshSelectors/PCGMeshSelectorByAttribute.h b/Engine/Plugins/PCG/Source/PCG/Public/MeshSelectors/PCGMeshSelectorByAttribute.h
--- a/Engine/Plugins/PCG/Source/PCG/Public/MeshSelectors/PCGMeshSelectorByAttribute.h
+++ b/Engine/Plugins/PCG/Source/PCG/Public/MeshSelectors/PCGMeshSelectorByAttribute.h
@@ -17,12 +17,12 @@ public:
 	virtual void PostLoad() override;
 	// ~End UObject interface
 
-	virtual bool SelectInstances(
+	virtual bool SelectMeshInstances(
 		FPCGStaticMeshSpawnerContext& Context,
 		const UPCGStaticMeshSpawnerSettings* Settings,
-		const UPCGPointData* InPointData,
+		const UPCGBasePointData* InPointData,
 		TArray<FPCGMeshInstanceList>& OutMeshInstances,
-		UPCGPointData* OutPointData) const override;
+		UPCGBasePointData* OutPointData) const override;
 
 public:
 	UPROPERTY(BlueprintReadWrite, EditAnywhere, Category = MeshSelector)
diff --git a/Engine/Plugins/PCG/Source/PCG/Public/MeshSelectors/PCGMeshSelectorWeighted.h b/Engine/Plugins/PCG/Source/PCG/Public/MeshSelectors/PCGMeshSelectorWeighted.h
--- a/Engine/Plugins/PCG/Source/PCG/Public/MeshSelectors/PCGMeshSelectorWeighted.h
+++ b/Engine/Plugins/PCG/Source/PCG/Public/MeshSelectors/PCGMeshSelectorWeighted.h
@@ -15,7 +15,7 @@ namespace PCGMeshSelectorWeighted
 		bool bUseMaterialOverrides,
 		const TArray<TSoftObjectPtr<UMaterialInterface>>& InMaterialOverrides,
 		bool bInIsLocalToWorldDeterminantNegative,
-		const UPCGPointData* InPointData);
+		const UPCGBasePointData* InPointData);
 }
 
 USTRUCT(BlueprintType)
@@ -74,12 +74,12 @@ class UPCGMeshSelectorWeighted : public UPCGMeshSelectorBase
 	GENERATED_BODY()
 
 public:
-	virtual bool SelectInstances(
+	virtual bool SelectMeshInstances(
 		FPCGStaticMeshSpawnerContext& Context,
 		const UPCGStaticMeshSpawnerSettings* Settings, 
-		const UPCGPointData* InPointData,
+		const UPCGBasePointData* InPointData,
 		TArray<FPCGMeshInstanceList>& OutMeshInstances,
-		UPCGPointData* OutPointData) const override;
+		UPCGBasePointData* OutPointData) const override;
 
 	// ~Begin UObject interface
 	virtual void PostLoad() override;
diff --git a/Engine/Plugins/PCG/Source/PCG/Public/MeshSelectors/PCGMeshSelectorWeightedByCategory.h b/Engine/Plugins/PCG/Source/PCG/Public/MeshSelectors/PCGMeshSelectorWeightedByCategory.h
--- a/Engine/Plugins/PCG/Source/PCG/Public/MeshSelectors/PCGMeshSelectorWeightedByCategory.h
+++ b/Engine/Plugins/PCG/Source/PCG/Public/MeshSelectors/PCGMeshSelectorWeightedByCategory.h
@@ -38,12 +38,12 @@ class UPCGMeshSelectorWeightedByCategory : public UPCGMeshSelectorBase
 	GENERATED_BODY()
 
 public:
-	virtual bool SelectInstances(
+	virtual bool SelectMeshInstances(
 		FPCGStaticMeshSpawnerContext& Context,
 		const UPCGStaticMeshSpawnerSettings* Settings,
-		const UPCGPointData* InPointData,
+		const UPCGBasePointData* InPointData,
 		TArray<FPCGMeshInstanceList>& OutMeshInstances,
-		UPCGPointData* OutPointData) const override;
+		UPCGBasePointData* OutPointData) const override;
 
 	// ~Begin UObject interface
 	virtual void PostLoad() override;
