[PCG] Add missing Count check in UPCGPointData::CopyPointsTo

#rnx

diff --git a/Engine/Plugins/PCG/Source/PCG/Private/Data/PCGPointData.cpp b/Engine/Plugins/PCG/Source/PCG/Private/Data/PCGPointData.cpp
--- a/Engine/Plugins/PCG/Source/PCG/Private/Data/PCGPointData.cpp
+++ b/Engine/Plugins/PCG/Source/PCG/Private/Data/PCGPointData.cpp
@@ -373,6 +373,11 @@ UPCGSpatialData* UPCGPointData::CopyInternal(FPCGContext* Context) const
 
 void UPCGPointData::CopyPointsTo(UPCGBasePointData* OutData, int32 ReadStartIndex, int32 WriteStartIndex, int32 Count) const
 {
+	if (Count <= 0)
+	{
+		return;
+	}
+
 	if (UPCGPointData* OutPointData = Cast<UPCGPointData>(OutData))
 	{
 		std::memcpy(&OutPointData->GetMutablePoints()[WriteStartIndex], &Points[ReadStartIndex], sizeof(FPCGPoint) * Count);
