[PCG] Convert PCGDataFromActor/PCGPointFromMesh/PCGCreatePoints/PCGDuplicatePoint/PCGSplineSampler/PCGSpawnActor to support UPCGPointArrayData

#rnx

diff --git a/Engine/Plugins/PCG/Source/PCG/Private/Data/PCGGetDataFunctionRegistry.cpp b/Engine/Plugins/PCG/Source/PCG/Private/Data/PCGGetDataFunctionRegistry.cpp
--- a/Engine/Plugins/PCG/Source/PCG/Private/Data/PCGGetDataFunctionRegistry.cpp
+++ b/Engine/Plugins/PCG/Source/PCG/Private/Data/PCGGetDataFunctionRegistry.cpp
@@ -7,11 +7,11 @@
 #include "PCGContext.h"
 #include "PCGData.h"
 #include "PCGGraph.h"
+#include "Data/PCGBasePointData.h"
 #include "Data/PCGCollisionShapeData.h"
 #include "Data/PCGIntersectionData.h"
 #include "Data/PCGLandscapeData.h"
 #include "Data/PCGLandscapeSplineData.h"
-#include "Data/PCGPointData.h"
 #include "Data/PCGPrimitiveData.h"
 #include "Data/PCGProjectionData.h"
 #include "Data/PCGSpatialData.h"
@@ -117,7 +117,7 @@ int FPCGGetDataFunctionRegistry::DefaultDataFromActor(FPCGContext* InContext, co
 	AVolume* VolumeActor = !!(InDataFilter & EPCGDataType::Volume) ? Cast<AVolume>(Actor) : nullptr;
 	if (!bParseActor && !!(InDataFilter & EPCGDataType::Point))
 	{
-		UPCGPointData* Data = FPCGContext::NewObject_AnyThread<UPCGPointData>(InContext);
+		UPCGBasePointData* Data = FPCGContext::NewPointData_AnyThread(InContext);
 		Data->InitializeFromActor(Actor, &Output.bSanitizedTagAttributeNames);
 
 		FPCGTaggedData& TaggedData = Output.Collection.TaggedData.Emplace_GetRef();
diff --git a/Engine/Plugins/PCG/Source/PCG/Private/Elements/PCGCreatePoints.cpp b/Engine/Plugins/PCG/Source/PCG/Private/Elements/PCGCreatePoints.cpp
--- a/Engine/Plugins/PCG/Source/PCG/Private/Elements/PCGCreatePoints.cpp
+++ b/Engine/Plugins/PCG/Source/PCG/Private/Elements/PCGCreatePoints.cpp
@@ -104,56 +104,102 @@ bool FPCGCreatePointsElement::ExecuteInternal(FPCGContext* Context) const
 	TArray<FPCGTaggedData>& Outputs = Context->OutputData.TaggedData;
 	FPCGTaggedData& Output = Outputs.Emplace_GetRef();
 
-	UPCGPointData* PointData = FPCGContext::NewObject_AnyThread<UPCGPointData>(Context);
+	UPCGBasePointData* PointData = FPCGContext::NewPointData_AnyThread(Context);
 	check(PointData); 
+	PointData->SetNumPoints(PointsToLoopOn.Num());
+
+	// @todo_pcg: Could probably be optimzed by comparing the source points to find the properties that are default but this could also be expensive
+	PointData->AllocateProperties(EPCGPointNativeProperties::All);
 
-	TArray<FPCGPoint>& OutputPoints = PointData->GetMutablePoints();
 	Output.Data = PointData;
+	int32 NumWritten = 0;
 
 	if (Settings->CoordinateSpace == EPCGCoordinateSpace::World)
 	{
+		FPCGPointValueRanges OutValueRanges(PointData, /*bAllocate=*/false);
+
 		if (CullingShape)
 		{
-			OutputPoints.Reserve(PointsToLoopOn.Num());
-
 			for (const FPCGPoint& Point : PointsToLoopOn)
 			{
 				if (CullingShape->GetDensityAtPosition(Point.Transform.GetLocation()) > 0)
 				{
-					OutputPoints.Add(Point);
+					OutValueRanges.SetFromPoint(NumWritten, Point);
+					++NumWritten;
 				}
 			}
 		}
 		else
 		{
-			OutputPoints = PointsToLoopOn;
+			for (const FPCGPoint& Point : PointsToLoopOn)
+			{
+				OutValueRanges.SetFromPoint(NumWritten, Point);
+				++NumWritten;
+			}
 		}
-		
-		for (FPCGPoint& Point : OutputPoints)
+
+		for (int32 Index = 0; Index < NumWritten; ++Index)
 		{
-			if (Point.Seed == 0)
+			if (OutValueRanges.SeedRange[Index] == 0)
 			{
 				// If the seed is the default value, generate a new seed based on the its transform
-				Point.Seed = UPCGBlueprintHelpers::ComputeSeedFromPosition(Point.Transform.GetLocation());
+				OutValueRanges.SeedRange[Index] = UPCGBlueprintHelpers::ComputeSeedFromPosition(OutValueRanges.TransformRange[Index].GetLocation());
 			}
 		}
+
+		PointData->SetNumPoints(NumWritten);
 	}
 	else
 	{
 		check(Settings->CoordinateSpace == EPCGCoordinateSpace::LocalComponent || Settings->CoordinateSpace == EPCGCoordinateSpace::OriginalComponent);
 
-		FPCGAsync::AsyncPointProcessing(Context, PointsToLoopOn.Num(), OutputPoints, [&PointsToLoopOn, &LocalTransform, CullingShape](int32 Index, FPCGPoint& OutPoint)
+		auto MoveDataRangeFunc = [PointData](int32 RangeStartIndex, int32 MoveToIndex, int32 NumElements)
+		{
+			PointData->MoveRange(RangeStartIndex, MoveToIndex, NumElements);
+		};
+
+		auto FinishedFunc = [PointData](int32 NumWritten)
 		{
-			const FPCGPoint& InPoint = PointsToLoopOn[Index];
-			OutPoint = InPoint;
-			OutPoint.Transform *= LocalTransform;
+			PointData->SetNumPoints(NumWritten);
+		};
 
-			const int SeedFromPosition = UPCGBlueprintHelpers::ComputeSeedFromPosition(OutPoint.Transform.GetLocation());
-			OutPoint.Seed = (InPoint.Seed == 0 ? SeedFromPosition : PCGHelpers::ComputeSeed(InPoint.Seed, SeedFromPosition));
+		auto ProcessRangeFunc = [PointData, &PointsToLoopOn, &LocalTransform, CullingShape](int32 StartReadIndex, int32 StartWriteIndex, int32 Count)
+		{
+			int32 NumWritten = 0;
+
+			FPCGPointValueRanges OutValueRanges(PointData, /*bAllocate=*/false);
+
+			for (int32 ReadIndex = StartReadIndex; ReadIndex < StartReadIndex + Count; ++ReadIndex)
+			{
+				const int32 WriteIndex = StartWriteIndex + NumWritten;
+
+				const FPCGPoint& InPoint = PointsToLoopOn[ReadIndex];
+				const FTransform OutTransform = InPoint.Transform * LocalTransform;
+				
+				if (!CullingShape || (CullingShape->GetDensityAtPosition(OutTransform.GetLocation()) > 0.0f))
+				{
+					OutValueRanges.SetFromPoint(WriteIndex, InPoint);
+					OutValueRanges.TransformRange[WriteIndex] = OutTransform;
+
+					const int SeedFromPosition = UPCGBlueprintHelpers::ComputeSeedFromPosition(OutTransform.GetLocation());
+					OutValueRanges.SeedRange[WriteIndex] = (InPoint.Seed == 0 ? SeedFromPosition : PCGHelpers::ComputeSeed(InPoint.Seed, SeedFromPosition));
+
+					++NumWritten;
+				}
+			}
 
-			// Discards all points that are outside the volume
-			return !CullingShape || (CullingShape->GetDensityAtPosition(OutPoint.Transform.GetLocation()) > 0.0f);
-		});
+			return NumWritten;
+		};
+
+		FPCGAsync::AsyncProcessingRangeEx(
+			&Context->AsyncState,
+			PointsToLoopOn.Num(),
+			[](){},
+			ProcessRangeFunc,
+			MoveDataRangeFunc,
+			FinishedFunc,
+			/*bEnableTimeSlicing=*/false
+		);
 	}
 
 	return true;
diff --git a/Engine/Plugins/PCG/Source/PCG/Private/Elements/PCGDataFromActor.cpp b/Engine/Plugins/PCG/Source/PCG/Private/Elements/PCGDataFromActor.cpp
--- a/Engine/Plugins/PCG/Source/PCG/Private/Elements/PCGDataFromActor.cpp
+++ b/Engine/Plugins/PCG/Source/PCG/Private/Elements/PCGDataFromActor.cpp
@@ -7,7 +7,7 @@
 #include "PCGCustomVersion.h"
 #include "PCGParamData.h"
 #include "PCGSubsystem.h"
-#include "Data/PCGPointData.h"
+#include "Data/PCGBasePointData.h"
 #include "Data/PCGSpatialData.h"
 #include "Data/PCGVolumeData.h"
 #include "Elements/PCGAddComponent.h"
@@ -768,7 +768,7 @@ void FPCGDataFromActorElement::MergeActorsIntoData(FPCGContext* Context, const U
 
 	// At this point in time, the partition actors behave slightly differently, so if we are in the case where
 	// we have one or more partition actors, we'll go through the normal process and do post-processing to merge the point data instead.
-	UPCGPointData* PointData = FPCGContext::NewObject_AnyThread<UPCGPointData>(Context);
+	UPCGBasePointData* PointData = FPCGContext::NewPointData_AnyThread(Context);
 
 	bool bHasData = false;
 	bool bAnyAttributeNameWasSanitized = false;
diff --git a/Engine/Plugins/PCG/Source/PCG/Private/Elements/PCGDuplicatePoint.cpp b/Engine/Plugins/PCG/Source/PCG/Private/Elements/PCGDuplicatePoint.cpp
--- a/Engine/Plugins/PCG/Source/PCG/Private/Elements/PCGDuplicatePoint.cpp
+++ b/Engine/Plugins/PCG/Source/PCG/Private/Elements/PCGDuplicatePoint.cpp
@@ -4,7 +4,7 @@
 
 #include "PCGContext.h"
 #include "PCGPin.h"
-#include "Data/PCGPointData.h"
+#include "Data/PCGBasePointData.h"
 #include "Helpers/PCGAsync.h"
 #include "Helpers/PCGHelpers.h"
 #include "Metadata/PCGMetadataAccessor.h"
@@ -39,92 +39,125 @@ bool FPCGDuplicatePointElement::ExecuteInternal(FPCGContext* Context) const
 
 	for (int i = 0; i < Inputs.Num(); ++i)
 	{
-		const UPCGPointData* InputPointData = Cast<UPCGPointData>(Inputs[i].Data);
+		const UPCGBasePointData* InputPointData = Cast<UPCGBasePointData>(Inputs[i].Data);
 		if (!InputPointData)
 		{
 			PCGE_LOG(Warning, GraphAndLog, LOCTEXT("InvalidInputPointData", "The input is not point data, skipped."));
 			continue;
 		}
 
-		const TArray<FPCGPoint>& InputPoints = InputPointData->GetPoints();
-
 		// Determines whether or not to include the source point in data
 		const bool bKeepSourcePoint = Settings->bOutputSourcePoint;
 		const int DuplicatesPerPoint = Iterations + (bKeepSourcePoint ? 1 : 0);
-		const int NumIterations = DuplicatesPerPoint * InputPoints.Num();
+		const int NumIterations = DuplicatesPerPoint * InputPointData->GetNumPoints();
 		const int FirstDuplicateIndex = Settings->bOutputSourcePoint ? 0 : 1;
 
 		FPCGTaggedData& Output = Outputs.Add_GetRef(Inputs[i]);
-		if (InputPoints.IsEmpty())
+		if (InputPointData->GetNumPoints() == 0)
 		{
 			Output.Data = InputPointData;
 			return true;
 		}
 
-		UPCGPointData* OutPointData = FPCGContext::NewObject_AnyThread<UPCGPointData>(Context);
-		OutPointData->InitializeFromData(InputPointData);
-		TArray<FPCGPoint>& OutputPoints = OutPointData->GetMutablePoints();
-		OutputPoints.SetNumUninitialized(NumIterations);
+		UPCGBasePointData* OutPointData = FPCGContext::NewPointData_AnyThread(Context);
+		
+		FPCGInitializeFromDataParams InitializeFromDataParams(InputPointData);
+		InitializeFromDataParams.bInheritSpatialData = false;
+
+		OutPointData->InitializeFromDataWithParams(InitializeFromDataParams);
 		Output.Data = OutPointData;
 
+		auto InitializeFunc = [NumIterations, OutPointData, InputPointData]()
+		{
+			OutPointData->SetNumPoints(NumIterations);
+			OutPointData->AllocateProperties(InputPointData->GetAllocatedProperties() | EPCGPointNativeProperties::Transform | EPCGPointNativeProperties::Seed);
+		};
+
 		const FTransform& SourceDuplicateTransform = Settings->PointTransform;
 
 		if (Settings->bDirectionAppliedInRelativeSpace)
 		{
-			auto ProcessPoint = [SourceDuplicateTransform, DuplicatesPerPoint, bKeepSourcePoint, &Direction, &InputPoints, &OutputPoints](const int32 ReadIndex, const int32 WriteIndex)
+			auto ProcessRangeFunc = [SourceDuplicateTransform, DuplicatesPerPoint, bKeepSourcePoint, &Direction, InputPointData, OutPointData](const int32 StartReadIndex, const int32 StartWriteIndex, const int32 Count)
 			{
-				const FPCGPoint& ReadPoint = InputPoints[ReadIndex];
-				TArrayView<FPCGPoint> WritePoints(OutputPoints.GetData() + WriteIndex * DuplicatesPerPoint, DuplicatesPerPoint);
-
-				const FTransform DuplicateAxisTransform = FTransform((ReadPoint.BoundsMax - ReadPoint.BoundsMin) * Direction);
-				const FTransform DuplicateTransform = DuplicateAxisTransform * SourceDuplicateTransform;
-				FTransform CurrentTransform = ReadPoint.Transform;
+				const FConstPCGPointValueRanges InRanges(InputPointData);
+				FPCGPointValueRanges OutRanges(OutPointData, /*bAllocate=*/false);
 
-				int WritePointIndex = 0;
+				int32 NumWritten = 0;
 
-				if (bKeepSourcePoint)
+				for (int32 ReadIndex = StartReadIndex; ReadIndex < StartReadIndex + Count; ++ReadIndex)
 				{
-					WritePoints[WritePointIndex++] = ReadPoint;
-				}
-
-				while (WritePointIndex < DuplicatesPerPoint)
-				{
-					FPCGPoint& WritePoint = WritePoints[WritePointIndex++];
-					WritePoint = ReadPoint;
-
-					CurrentTransform = DuplicateTransform * CurrentTransform;
-					WritePoint.Transform = CurrentTransform;
-					WritePoint.Seed = PCGHelpers::ComputeSeedFromPosition(CurrentTransform.GetLocation());
+					const FTransform DuplicateAxisTransform = FTransform((InRanges.BoundsMaxRange[ReadIndex] - InRanges.BoundsMinRange[ReadIndex]) * Direction);
+					const FTransform DuplicateTransform = DuplicateAxisTransform * SourceDuplicateTransform;
+					FTransform CurrentTransform = InRanges.TransformRange[ReadIndex];
+
+					int32 LocalWriteIndex = 0;
+					if (bKeepSourcePoint)
+					{
+						const int32 WriteIndex = StartWriteIndex + NumWritten;
+						OutRanges.SetFromValueRanges(WriteIndex, InRanges, ReadIndex);
+						++LocalWriteIndex;
+					}
+
+					while (LocalWriteIndex < DuplicatesPerPoint)
+					{
+						const int32 WriteIndex = StartWriteIndex + NumWritten + LocalWriteIndex;
+						OutRanges.SetFromValueRanges(WriteIndex, InRanges, ReadIndex);
+						
+						CurrentTransform = DuplicateTransform * CurrentTransform;
+						OutRanges.TransformRange[WriteIndex] = CurrentTransform;
+						OutRanges.SeedRange[WriteIndex] = PCGHelpers::ComputeSeedFromPosition(CurrentTransform.GetLocation());
+
+						++LocalWriteIndex;
+					}
+
+					NumWritten += LocalWriteIndex;
 				}
+				
+				check(NumWritten == Count * DuplicatesPerPoint);
+				return Count;
 			};
 
-			FPCGAsync::AsyncProcessingOneToOneEx(&Context->AsyncState, InputPoints.Num(), /*Initialize=*/[]() {}, ProcessPoint, /*bEnableTimeSlicing=*/false);
+			FPCGAsync::AsyncProcessingOneToOneRangeEx(&Context->AsyncState, InputPointData->GetNumPoints(), InitializeFunc, ProcessRangeFunc, /*bEnableTimeSlicing=*/false);
 		}
 		else
 		{
-			auto ProcessPoint = [SourceDuplicateTransform, FirstDuplicateIndex, OutPointData, InputPointData, &Direction, &InputPoints, &OutputPoints](const int32 ReadIndex, const int32 WriteIndex)
+			auto ProcessRangeFunc = [SourceDuplicateTransform, &Direction, InputPointData, OutPointData, FirstDuplicateIndex](const int32 StartReadIndex, const int32 StartWriteIndex, const int32 Count)
 			{
-				const FPCGPoint& InPoint = InputPoints[ReadIndex % InputPoints.Num()];
-				FPCGPoint& OutputPoint = OutputPoints[WriteIndex];
-				const int DuplicateIndex = FirstDuplicateIndex + ReadIndex / InputPoints.Num();
+				const FConstPCGPointValueRanges InRanges(InputPointData);
+				FPCGPointValueRanges OutRanges(OutPointData, /*bAllocate=*/false);
 
-				OutputPoint = InPoint;
+				int32 NumWritten = 0;
+				const int32 NumInputPoints = InputPointData->GetNumPoints();
 
-				if (DuplicateIndex != 0)
+				for (int32 GlobalReadIndex = StartReadIndex; GlobalReadIndex < StartReadIndex + Count; ++GlobalReadIndex)
 				{
-					const FVector DuplicateLocationOffset = ((OutputPoint.BoundsMax - OutputPoint.BoundsMin) * Direction + SourceDuplicateTransform.GetLocation()) * DuplicateIndex;
-					const FRotator DuplicateRotationOffset = SourceDuplicateTransform.Rotator() * DuplicateIndex;
-					const FVector DuplicateScaleMultiplier = FVector(
-						FMath::Pow(SourceDuplicateTransform.GetScale3D().X, DuplicateIndex), 
-						FMath::Pow(SourceDuplicateTransform.GetScale3D().Y, DuplicateIndex), 
-						FMath::Pow(SourceDuplicateTransform.GetScale3D().Z, DuplicateIndex));
-
-					OutputPoint.Transform = FTransform(DuplicateRotationOffset, DuplicateLocationOffset, DuplicateScaleMultiplier) * InPoint.Transform;
-					OutputPoint.Seed = PCGHelpers::ComputeSeedFromPosition(OutputPoint.Transform.GetLocation());
+					const int32 ReadIndex = GlobalReadIndex % NumInputPoints;
+					const int32 WriteIndex = StartWriteIndex + NumWritten;
+					const int DuplicateIndex = FirstDuplicateIndex + GlobalReadIndex / NumInputPoints;
+
+					OutRanges.SetFromValueRanges(WriteIndex, InRanges, ReadIndex);
+
+					if (DuplicateIndex != 0)
+					{
+						const FVector DuplicateLocationOffset = ((OutRanges.BoundsMaxRange[WriteIndex] - OutRanges.BoundsMinRange[WriteIndex]) * Direction + SourceDuplicateTransform.GetLocation()) * DuplicateIndex;
+						const FRotator DuplicateRotationOffset = SourceDuplicateTransform.Rotator() * DuplicateIndex;
+						const FVector DuplicateScaleMultiplier = FVector(
+							FMath::Pow(SourceDuplicateTransform.GetScale3D().X, DuplicateIndex),
+							FMath::Pow(SourceDuplicateTransform.GetScale3D().Y, DuplicateIndex),
+							FMath::Pow(SourceDuplicateTransform.GetScale3D().Z, DuplicateIndex));
+
+						OutRanges.TransformRange[WriteIndex] = FTransform(DuplicateRotationOffset, DuplicateLocationOffset, DuplicateScaleMultiplier) * InRanges.TransformRange[ReadIndex];
+						OutRanges.SeedRange[WriteIndex] = PCGHelpers::ComputeSeedFromPosition(OutRanges.TransformRange[WriteIndex].GetLocation());
+					}
+
+					++NumWritten;
 				}
+				
+				check(NumWritten == Count);
+				return Count;
 			};
 
-			FPCGAsync::AsyncProcessingOneToOneEx(&Context->AsyncState, NumIterations, /*Initialize=*/[]() {}, ProcessPoint, /*bEnableTimeSlicing=*/false);
+			FPCGAsync::AsyncProcessingOneToOneRangeEx(&Context->AsyncState, NumIterations, InitializeFunc, ProcessRangeFunc, /*bEnableTimeSlicing=*/false);
 		}
 	}
 
diff --git a/Engine/Plugins/PCG/Source/PCG/Private/Elements/PCGPointFromMeshElement.cpp b/Engine/Plugins/PCG/Source/PCG/Private/Elements/PCGPointFromMeshElement.cpp
--- a/Engine/Plugins/PCG/Source/PCG/Private/Elements/PCGPointFromMeshElement.cpp
+++ b/Engine/Plugins/PCG/Source/PCG/Private/Elements/PCGPointFromMeshElement.cpp
@@ -75,16 +75,16 @@ bool FPCGPointFromMeshElement::ExecuteInternal(FPCGContext* Context) const
 	}
 
 	TArray<FPCGTaggedData>& Outputs = Context->OutputData.TaggedData;
-	UPCGPointData* OutPointData = FPCGContext::NewObject_AnyThread<UPCGPointData>(Context);
+	UPCGBasePointData* OutPointData = FPCGContext::NewPointData_AnyThread(Context);
 	Outputs.Emplace_GetRef().Data = OutPointData;
 	
-	TArray<FPCGPoint>& Points = OutPointData->GetMutablePoints();
-	FPCGPoint& Point = Points.Emplace_GetRef();
+	OutPointData->SetNumPoints(1);
 
 	// Capture StaticMesh bounds
 	const FBox StaticMeshBounds = Settings->StaticMesh->GetBoundingBox();
-	Point.BoundsMin = StaticMeshBounds.Min;
-	Point.BoundsMax = StaticMeshBounds.Max;
+
+	OutPointData->SetBoundsMin(StaticMeshBounds.Min);
+	OutPointData->SetBoundsMax(StaticMeshBounds.Max);
 
 	// Write StaticMesh path to MeshPathAttribute
 	check(OutPointData->Metadata);
diff --git a/Engine/Plugins/PCG/Source/PCG/Private/Elements/PCGSpawnActor.cpp b/Engine/Plugins/PCG/Source/PCG/Private/Elements/PCGSpawnActor.cpp
--- a/Engine/Plugins/PCG/Source/PCG/Private/Elements/PCGSpawnActor.cpp
+++ b/Engine/Plugins/PCG/Source/PCG/Private/Elements/PCGSpawnActor.cpp
@@ -6,7 +6,7 @@
 #include "PCGGraph.h"
 #include "PCGManagedResource.h"
 #include "PCGSubsystem.h"
-#include "Data/PCGPointData.h"
+#include "Data/PCGBasePointData.h"
 #include "Data/PCGSpatialData.h"
 #include "Graph/PCGStackContext.h"
 #include "Grid/PCGPartitionActor.h"
@@ -45,7 +45,7 @@ public:
 
 	bool InitializeForData(const UPCGData* InData, UPCGData* OutData)
 	{
-		if (!InData || !InData->IsA<UPCGPointData>())
+		if (!InData || !InData->IsA<UPCGBasePointData>())
 		{
 			return false;
 		}
@@ -59,23 +59,34 @@ public:
 		return SpawnAttributeAccessor.IsValid() && SpawnAttributeKeys.IsValid();
 	}
 
-	void AddToPartitionData(FPCGDataPartitionBase::Element* SelectedElement, const UPCGData* ParentData, int32 Index)
+	void Finalize(const UPCGData* InData, UPCGData* OutData)
 	{
-		// Already checked in initialize that this is a valid cast
-		const UPCGPointData* ParentPointData = static_cast<const UPCGPointData*>(ParentData);
-
-		check(SelectedElement);
-		if (!SelectedElement->PartitionData)
+		const UPCGBasePointData* InPointData = Cast<UPCGBasePointData>(InData);
+		if (!InPointData)
 		{
-			UPCGPointData* PartitionData = NewObject<UPCGPointData>();
-			PartitionData->InitializeFromData(ParentPointData);
-			SelectedElement->PartitionData = PartitionData;
+			return;
 		}
 
-		check(SelectedElement->PartitionData);
-		static_cast<UPCGPointData*>(SelectedElement->PartitionData)->GetMutablePoints().Add(ParentPointData->GetPoints()[Index]);
-	}
+		for (auto& KeyValuePair : ElementMap)
+		{
+			FPCGDataPartitionBase::Element& Element = KeyValuePair.Value;
+			
+			if (!Element.Indices.IsEmpty())
+			{
+				check(!Element.PartitionData);
 
+				UPCGBasePointData* PartitionPointData = FPCGContext::NewPointData_AnyThread(Context);
+				Element.PartitionData = PartitionPointData;
+
+				FPCGInitializeFromDataParams InitializeFromDataParams(InPointData);
+				InitializeFromDataParams.bInheritSpatialData = false;
+				
+				PartitionPointData->InitializeFromDataWithParams(InitializeFromDataParams);
+				PartitionPointData->SetPointsFrom(InPointData, Element.Indices);
+			}
+		}
+	}
+		
 	FPCGDataPartitionBase::Element* Select(int32 Index)
 	{
 		FSoftClassPath ActorPath;
@@ -540,7 +551,7 @@ bool FPCGSpawnActorElement::SpawnAndPrepareSubgraphs(FPCGSubgraphContext* Contex
 		}
 
 		// First, create target instance transforms
-		const UPCGPointData* PointData = SpatialData->ToPointData(Context);
+		const UPCGBasePointData* PointData = SpatialData->ToBasePointData(Context);
 
 		if (!PointData)
 		{
@@ -548,15 +559,13 @@ bool FPCGSpawnActorElement::SpawnAndPrepareSubgraphs(FPCGSubgraphContext* Contex
 			continue;
 		}
 
-		const TArray<FPCGPoint>& Points = PointData->GetPoints();
-
-		if (Points.IsEmpty())
+		if (PointData->GetNumPoints() == 0)
 		{
 			PCGE_LOG(Verbose, LogOnly, LOCTEXT("SkippedNoPoints", "Skipped - no points"));
 			continue;
 		}
 
-		auto SpawnOrCollapse = [this, bHasAuthority, &Context, &TargetActor, &Settings](TSubclassOf<AActor> TemplateActorClass, AActor* TemplateActor, FPCGTaggedData& Output, const UPCGPointData* PointData, UPCGPointData* OutPointData)
+		auto SpawnOrCollapse = [this, bHasAuthority, &Context, &TargetActor, &Settings](TSubclassOf<AActor> TemplateActorClass, AActor* TemplateActor, FPCGTaggedData& Output, const UPCGBasePointData* PointData, UPCGBasePointData* OutPointData)
 		{
 			const bool bSpawnedActorsRequireAuthority = (TemplateActor ? TemplateActor->GetIsReplicated() : CastChecked<AActor>(TemplateActorClass->GetDefaultObject())->GetIsReplicated());
 
@@ -570,10 +579,10 @@ bool FPCGSpawnActorElement::SpawnAndPrepareSubgraphs(FPCGSubgraphContext* Contex
 			}
 		};
 
-		UPCGPointData* OutPointData = nullptr;
+		UPCGBasePointData* OutPointData = nullptr;
 		if (bGenerateOutputsWithActorReference)
 		{
-			OutPointData = NewObject<UPCGPointData>();
+			OutPointData = FPCGContext::NewPointData_AnyThread(Context);
 			OutPointData->InitializeFromData(PointData);
 		}
 
@@ -585,7 +594,7 @@ bool FPCGSpawnActorElement::SpawnAndPrepareSubgraphs(FPCGSubgraphContext* Contex
 			int32 CurrentPointIndex = 0;
 
 			// Selection is still needed if are fully skipped in order to write to the OutPointData.
-			Selector.SelectMultiple(*Context, PointData, CurrentPointIndex, PointData->GetPoints().Num(), OutPointData);
+			Selector.SelectMultiple(*Context, PointData, CurrentPointIndex, PointData->GetNumPoints(), OutPointData);
 
 			if (!bFullySkippedDueToReuse)
 			{
@@ -594,7 +603,7 @@ bool FPCGSpawnActorElement::SpawnAndPrepareSubgraphs(FPCGSubgraphContext* Contex
 					FPCGTaggedData PartialInput = Input;
 					PartialInput.Data = Element.Value.PartitionData;
 
-					SpawnOrCollapse(Element.Key, nullptr, PartialInput, static_cast<UPCGPointData*>(Element.Value.PartitionData), OutPointData);
+					SpawnOrCollapse(Element.Key, nullptr, PartialInput, Cast<UPCGBasePointData>(Element.Value.PartitionData), OutPointData);
 
 					// Exception case here: if we've spawned actors but are merging the PCG inputs,
 					// normally this node is taken as a subgraph node (e.g. no need to do anything more than forwarding the inputs)
@@ -685,13 +694,13 @@ bool FPCGSpawnActorElement::SpawnAndPrepareSubgraphs(FPCGSubgraphContext* Contex
 	}
 }
 
-void FPCGSpawnActorElement::CollapseIntoTargetActor(FPCGSubgraphContext* Context, AActor* TargetActor, TSubclassOf<AActor> TemplateActorClass, const UPCGPointData* PointData) const
+void FPCGSpawnActorElement::CollapseIntoTargetActor(FPCGSubgraphContext* Context, AActor* TargetActor, TSubclassOf<AActor> TemplateActorClass, const UPCGBasePointData* PointData) const
 {
 	TRACE_CPUPROFILER_EVENT_SCOPE(FPCGSpawnActorElement::ExecuteInternal::CollapseActors);
 	check(Context && TargetActor && PointData);
 
-	const TArray<FPCGPoint>& Points = PointData->GetPoints();
-	if (Points.IsEmpty())
+	const int32 NumPoints = PointData->GetNumPoints();
+	if (NumPoints == 0)
 	{
 		return;
 	}
@@ -735,6 +744,8 @@ void FPCGSpawnActorElement::CollapseIntoTargetActor(FPCGSubgraphContext* Context
 		return true;
 	});
 
+	const TConstPCGValueRange<FTransform> TransformRange = PointData->GetConstTransformValueRange();
+
 	for (const TPair<FPCGISMComponentBuilderParams, TArray<FTransform>>& ISMCBuilderTransforms : MeshDescriptorTransforms)
 	{
 		const FPCGISMComponentBuilderParams& ISMCParams = ISMCBuilderTransforms.Key;
@@ -753,14 +764,13 @@ void FPCGSpawnActorElement::CollapseIntoTargetActor(FPCGSubgraphContext* Context
 		const TArray<FTransform>& ISMCTransforms = ISMCBuilderTransforms.Value;
 
 		TArray<FTransform> Transforms;
-		Transforms.Reserve(Points.Num() * ISMCTransforms.Num());
-		for (int32 PointIndex = 0; PointIndex < Points.Num(); PointIndex++)
+		Transforms.Reserve(NumPoints * ISMCTransforms.Num());
+		for (int32 PointIndex = 0; PointIndex < NumPoints; PointIndex++)
 		{
-			const FPCGPoint& Point = Points[PointIndex];
 			for (int32 TransformIndex = 0; TransformIndex < ISMCTransforms.Num(); TransformIndex++)
 			{
 				const FTransform& Transform = ISMCTransforms[TransformIndex];
-				Transforms.Add(Transform * Point.Transform);
+				Transforms.Add(Transform * TransformRange[PointIndex]);
 			}
 		}
 
@@ -773,14 +783,13 @@ void FPCGSpawnActorElement::CollapseIntoTargetActor(FPCGSubgraphContext* Context
 	}
 }
 
-void FPCGSpawnActorElement::SpawnActors(FPCGSubgraphContext* Context, AActor* TargetActor, TSubclassOf<AActor> InTemplateActorClass, AActor* InTemplateActor, FPCGTaggedData& Output, const UPCGPointData* PointData, UPCGPointData* OutPointData) const
+void FPCGSpawnActorElement::SpawnActors(FPCGSubgraphContext* Context, AActor* TargetActor, TSubclassOf<AActor> InTemplateActorClass, AActor* InTemplateActor, FPCGTaggedData& Output, const UPCGBasePointData* PointData, UPCGBasePointData* OutPointData) const
 {
 	TRACE_CPUPROFILER_EVENT_SCOPE(FPCGSpawnActorElement::ExecuteInternal::SpawnActors);
 	check(Context && TargetActor && PointData);
 	check(IsInGameThread());
 
-	const TArray<FPCGPoint>& Points = PointData->GetPoints();
-	if (Points.IsEmpty())
+	if (PointData->GetNumPoints() == 0)
 	{
 		return;
 	}
@@ -790,8 +799,9 @@ void FPCGSpawnActorElement::SpawnActors(FPCGSubgraphContext* Context, AActor* Ta
 
 	if (OutPointData)
 	{
-		OutPointOffset = OutPointData->GetMutablePoints().Num();
-		OutPointData->GetMutablePoints().Append(Points);
+		OutPointOffset = OutPointData->GetNumPoints();
+		OutPointData->SetNumPoints(OutPointOffset + PointData->GetNumPoints());
+		PointData->CopyPointsTo(OutPointData, 0, OutPointOffset, PointData->GetNumPoints());
 		ActorReferenceAttribute = OutPointData->MutableMetadata()->FindOrCreateAttribute<FSoftObjectPath>(PCGPointDataConstants::ActorReferenceAttribute, FSoftObjectPath(), /*bAllowsInterpolation=*/false, /*bOverrideParent=*/false, /*bOverwriteIfTypeMismatch=*/false);
 	}
 
@@ -878,7 +888,7 @@ void FPCGSpawnActorElement::SpawnActors(FPCGSubgraphContext* Context, AActor* Ta
 
 		if (InputDependenciesCrc.IsValid())
 		{
-			Context->SourceComponent->ForEachManagedResource([&ReusedManagedActorsResource, &InputDependenciesCrc, &Context, NumPoints = Points.Num()](UPCGManagedResource* InResource)
+			Context->SourceComponent->ForEachManagedResource([&ReusedManagedActorsResource, &InputDependenciesCrc, &Context, NumPoints = PointData->GetNumPoints()](UPCGManagedResource* InResource)
 			{
 				if (ReusedManagedActorsResource)
 				{
@@ -909,6 +919,8 @@ void FPCGSpawnActorElement::SpawnActors(FPCGSubgraphContext* Context, AActor* Ta
 		// Also set to the output data the actor reference.
 		if (Settings->Option == EPCGSpawnActorOption::NoMerging)
 		{
+			TPCGValueRange<int64> MetadataEntryRange = OutPointData->GetMetadataEntryValueRange();
+
 			const TArray<TSoftObjectPtr<AActor>>& GeneratedActors = ReusedManagedActorsResource->GetConstGeneratedActors();
 			for (int32 i = 0; i < GeneratedActors.Num(); ++i)
 			{
@@ -917,9 +929,9 @@ void FPCGSpawnActorElement::SpawnActors(FPCGSubgraphContext* Context, AActor* Ta
 				// Write to out data the actor reference
 				if (OutPointData && ActorReferenceAttribute)
 				{
-					FPCGPoint& OutPoint = OutPointData->GetMutablePoints()[i + OutPointOffset];
-					OutPointData->Metadata->InitializeOnSet(OutPoint.MetadataEntry);
-					ActorReferenceAttribute->SetValue(OutPoint.MetadataEntry, ManagedActorPtr.ToSoftObjectPath());
+					int64& MetadataEntry = MetadataEntryRange[i + OutPointOffset];
+					OutPointData->Metadata->InitializeOnSet(MetadataEntry);
+					ActorReferenceAttribute->SetValue(MetadataEntry, ManagedActorPtr.ToSoftObjectPath());
 				}
 
 				if (bActorsHavePCGComponents)
@@ -958,13 +970,14 @@ void FPCGSpawnActorElement::SpawnActors(FPCGSubgraphContext* Context, AActor* Ta
 
 		bool bAllActorOverridesSucceeded = true;
 
-		for (int32 i = 0; i < Points.Num(); ++i)
-		{
-			const FPCGPoint& Point = Points[i];
+		const FConstPCGPointValueRanges ValueRanges(PointData);
+		TPCGValueRange<int64> OutMetadataEntryRange = OutPointData->GetMetadataEntryValueRange();
 
+		for (int32 i = 0; i < PointData->GetNumPoints(); ++i)
+		{
 			bAllActorOverridesSucceeded &= ActorOverrides.Apply(i);
 
-			SpawnDefaultActorParams.Transform = Point.Transform;
+			SpawnDefaultActorParams.Transform = ValueRanges.TransformRange[i];
 			AActor* GeneratedActor = UPCGActorHelpers::SpawnDefaultActor(SpawnDefaultActorParams);
 
 			if (!GeneratedActor)
@@ -974,7 +987,7 @@ void FPCGSpawnActorElement::SpawnActors(FPCGSubgraphContext* Context, AActor* Ta
 			}
 
 			// HACK: until UE-62747 is fixed, we have to force set the scale after spawning the actor
-			GeneratedActor->SetActorRelativeScale3D(Point.Transform.GetScale3D());
+			GeneratedActor->SetActorRelativeScale3D(ValueRanges.TransformRange[i].GetScale3D());
 			GeneratedActor->Tags.Append(NewActorTags);
 			PCGHelpers::AttachToParent(GeneratedActor, TargetActor, Settings->AttachOptions, Context, GeneratedActorsFolderPath);
 
@@ -986,7 +999,7 @@ void FPCGSpawnActorElement::SpawnActors(FPCGSubgraphContext* Context, AActor* Ta
 				}
 				else if (PostSpawnFunction->IsSignatureCompatibleWith(FunctionPrototypeWithPointAndMetadata))
 				{
-					TPair<FPCGPoint, const UPCGMetadata*> PointAndMetadata = { Point, PointData->ConstMetadata() };
+					TPair<FPCGPoint, const UPCGMetadata*> PointAndMetadata = { ValueRanges.GetPoint(i), PointData->ConstMetadata()};
 					GeneratedActor->ProcessEvent(PostSpawnFunction, &PointAndMetadata);
 				}
 			}
@@ -1001,9 +1014,9 @@ void FPCGSpawnActorElement::SpawnActors(FPCGSubgraphContext* Context, AActor* Ta
 			// Write to out data the actor reference
 			if (OutPointData && ActorReferenceAttribute)
 			{
-				FPCGPoint& OutPoint = OutPointData->GetMutablePoints()[i + OutPointOffset];
-				OutPointData->Metadata->InitializeOnSet(OutPoint.MetadataEntry);
-				ActorReferenceAttribute->SetValue(OutPoint.MetadataEntry, FSoftObjectPath(GeneratedActor));
+				int64& MetadataEntry = OutMetadataEntryRange[i + OutPointOffset];
+				OutPointData->Metadata->InitializeOnSet(MetadataEntry);
+				ActorReferenceAttribute->SetValue(MetadataEntry, FSoftObjectPath(GeneratedActor));
 			}
 		}
 
@@ -1014,7 +1027,7 @@ void FPCGSpawnActorElement::SpawnActors(FPCGSubgraphContext* Context, AActor* Ta
 
 		Context->SourceComponent->AddToManagedResources(ManagedActors);
 
-		PCGE_LOG(Verbose, LogOnly, FText::Format(LOCTEXT("GenerationInfo", "Generated {0} actors"), Points.Num()));
+		PCGE_LOG(Verbose, LogOnly, FText::Format(LOCTEXT("GenerationInfo", "Generated {0} actors"), PointData->GetNumPoints()));
 	}
 
 	// Setup & Generate on PCG components if needed
diff --git a/Engine/Plugins/PCG/Source/PCG/Private/Elements/PCGSplineSampler.cpp b/Engine/Plugins/PCG/Source/PCG/Private/Elements/PCGSplineSampler.cpp
--- a/Engine/Plugins/PCG/Source/PCG/Private/Elements/PCGSplineSampler.cpp
+++ b/Engine/Plugins/PCG/Source/PCG/Private/Elements/PCGSplineSampler.cpp
@@ -7,11 +7,10 @@
 #include "PCGCustomVersion.h"
 #include "Data/PCGIntersectionData.h"
 #include "Data/PCGLandscapeSplineData.h"
-#include "Data/PCGPointData.h"
+#include "Data/PCGBasePointData.h"
 #include "Data/PCGPolyLineData.h"
 #include "Data/PCGSplineData.h"
 #include "Data/PCGSpatialData.h"
-#include "Helpers/PCGBlueprintHelpers.h"
 #include "Helpers/PCGHelpers.h"
 #include "Helpers/PCGSettingsHelpers.h"
 
@@ -353,7 +352,7 @@ namespace PCGSplineSamplerHelpers
 		FVector::FReal Distance = 0.0f;
 	};
 
-	void SetSeed(FPCGPoint& Point, const FVector& LSPosition, const FPCGSplineSamplerParams& Params, int ExternalSeed, int SampleIndex)
+	void SetSeed(int32& OutSeed, const FTransform& InTransform, const FVector& LSPosition, const FPCGSplineSamplerParams& Params, int ExternalSeed, int SampleIndex)
 	{
 		if (Params.SeedingMode == EPCGSplineSamplingSeedingMode::SeedFromPosition)
 		{
@@ -361,29 +360,29 @@ namespace PCGSplineSamplerHelpers
 			{
 				if (!Params.bSeedFrom2DPosition)
 				{
-					Point.Seed = UPCGBlueprintHelpers::ComputeSeedFromPosition(Point.Transform.GetLocation());
+					OutSeed = PCGHelpers::ComputeSeedFromPosition(InTransform.GetLocation());
 				}
 				else
 				{
-					const FVector WSPosition = Point.Transform.GetLocation();
-					Point.Seed = PCGHelpers::ComputeSeed((int)WSPosition.X, (int)WSPosition.Y);
+					const FVector WSPosition = InTransform.GetLocation();
+					OutSeed = PCGHelpers::ComputeSeed((int)WSPosition.X, (int)WSPosition.Y);
 				}
 			}
 			else // Use provided local position
 			{
 				if (!Params.bSeedFrom2DPosition)
 				{
-					Point.Seed = UPCGBlueprintHelpers::ComputeSeedFromPosition(LSPosition);
+					OutSeed = PCGHelpers::ComputeSeedFromPosition(LSPosition);
 				}
 				else
 				{
-					Point.Seed = PCGHelpers::ComputeSeed((int)LSPosition.X, (int)LSPosition.Y);
+					OutSeed = PCGHelpers::ComputeSeed((int)LSPosition.X, (int)LSPosition.Y);
 				}
 			}
 		}
 		else
 		{
-			Point.Seed = PCGHelpers::ComputeSeed(ExternalSeed, SampleIndex);
+			OutSeed = PCGHelpers::ComputeSeed(ExternalSeed, SampleIndex);
 		}
 	}
 
@@ -670,7 +669,7 @@ namespace PCGSplineSamplerHelpers
 
 	struct FDimensionSampler
 	{
-		FDimensionSampler(const UPCGPolyLineData* InLineData, const UPCGSpatialData* InBoundingShapeData, const UPCGSpatialData* InProjectionTarget, const FPCGProjectionParams& InProjectionParams, const FPCGSplineSamplerParams& InParams, int InSeed, UPCGPointData* OutPointData)
+		FDimensionSampler(const UPCGPolyLineData* InLineData, const UPCGSpatialData* InBoundingShapeData, const UPCGSpatialData* InProjectionTarget, const FPCGProjectionParams& InProjectionParams, const FPCGSplineSamplerParams& InParams, int InSeed, UPCGBasePointData* OutPointData)
 			: Params(InParams)
 			, Seed(InSeed)
 		{
@@ -797,7 +796,7 @@ namespace PCGSplineSamplerHelpers
 			}
 		}
 
-		virtual void Sample(const FSamplerResult& InResult, UPCGPointData* OutPointData)
+		virtual void Sample(const FSamplerResult& InResult, UPCGBasePointData* OutPointData)
 		{
 			// Mimic PCGSplineData sampling.
 			const FTransform Transform = InResult.LocalTransform * LineData->GetTransform();
@@ -826,9 +825,14 @@ namespace PCGSplineSamplerHelpers
 			FPCGPoint BoundsTestPoint;
 			if (bValid && (!BoundingShapeData || BoundingShapeData->SamplePoint(Transform, InResult.Box, BoundsTestPoint, OutPointData->Metadata)))
 			{
-				SetSeed(OutPoint, InResult.LocalTransform.GetLocation(), Params, Seed, InResult.SampleIndex);
+				SetSeed(OutPoint.Seed, OutPoint.Transform, InResult.LocalTransform.GetLocation(), Params, Seed, InResult.SampleIndex);
 				SetMetadata(InResult, OutPoint, OutPointData->Metadata);
-				OutPointData->GetMutablePoints().Add(OutPoint);
+				
+				// @todo_pcg: ideally we could preallocate the number of points here if the Samplers could give a hint on number of points
+				const int32 NumPoints = OutPointData->GetNumPoints();
+				OutPointData->SetNumPoints(NumPoints+1);
+				FPCGPointValueRanges OutRanges(OutPointData, /*bAllocate=*/false);
+				OutRanges.SetFromPoint(NumPoints, OutPoint);
 			}
 		}
 
@@ -854,7 +858,7 @@ namespace PCGSplineSamplerHelpers
 	/** Samples in a volume surrounding the poly line. */
 	struct FVolumeSampler : public FDimensionSampler
 	{
-		FVolumeSampler(const UPCGPolyLineData* InLineData, const UPCGSpatialData* InBoundingShapeData, const UPCGSpatialData* InProjectionTarget, const FPCGProjectionParams& InProjectionParams, const FPCGSplineSamplerParams& Params, int InSeed, UPCGPointData* OutPointData)
+		FVolumeSampler(const UPCGPolyLineData* InLineData, const UPCGSpatialData* InBoundingShapeData, const UPCGSpatialData* InProjectionTarget, const FPCGProjectionParams& InProjectionParams, const FPCGSplineSamplerParams& Params, int InSeed, UPCGBasePointData* OutPointData)
 			: FDimensionSampler(InLineData, InBoundingShapeData, InProjectionTarget, InProjectionParams, Params, InSeed, OutPointData)
 		{
 			Fill = Params.Fill;
@@ -862,7 +866,7 @@ namespace PCGSplineSamplerHelpers
 			NumHeightSteps = 1 + ((Params.Dimension == EPCGSplineSamplingDimension::OnHorizontal) ? 0 : Params.NumHeightSubdivisions);
 		}
 
-		virtual void Sample(const FSamplerResult& InResult, UPCGPointData* OutPointData) override
+		virtual void Sample(const FSamplerResult& InResult, UPCGBasePointData* OutPointData) override
 		{
 			const FTransform& TransformLS = InResult.LocalTransform;
 			FTransform TransformWS = TransformLS * LineData->GetTransform();
@@ -938,10 +942,15 @@ namespace PCGSplineSamplerHelpers
 						continue;
 					}
 
-					SetSeed(OutPoint, TentativeLocationLS, Params, Seed, InResult.SampleIndex);
+					SetSeed(OutPoint.Seed, OutPoint.Transform, TentativeLocationLS, Params, Seed, InResult.SampleIndex);
 					OutPoint.Steepness = Params.PointSteepness;
 					SetMetadata(InResult, OutPoint, OutPointData->Metadata);
-					OutPointData->GetMutablePoints().Add(OutPoint);
+
+					// @todo_pcg: ideally we could preallocate the number of points here if the Samplers could give a hint on number of points
+					const int32 NumPoints = OutPointData->GetNumPoints();
+					OutPointData->SetNumPoints(NumPoints + 1);
+					FPCGPointValueRanges OutRanges(OutPointData, /*bAllocate=*/false);
+					OutRanges.SetFromPoint(NumPoints, OutPoint);
 				}
 
 				CurrentZ += 2.0 * ZHalfStep;
@@ -953,7 +962,7 @@ namespace PCGSplineSamplerHelpers
 		int NumHeightSteps;
 	};
 
-	void SampleLineData(FPCGContext* Context, const UPCGPolyLineData* LineData, const UPCGSpatialData* InBoundingShapeData, const UPCGSpatialData* InProjectionTarget, const FPCGProjectionParams& InProjectionParams, const FPCGSplineSamplerParams& Params, UPCGPointData* OutPointData)
+	void SampleLineData(FPCGContext* Context, const UPCGPolyLineData* LineData, const UPCGSpatialData* InBoundingShapeData, const UPCGSpatialData* InProjectionTarget, const FPCGProjectionParams& InProjectionParams, const FPCGSplineSamplerParams& Params, UPCGBasePointData* OutPointData)
 	{
 		check(LineData && OutPointData);
 
@@ -988,6 +997,9 @@ namespace PCGSplineSamplerHelpers
 		FSamplerResult* PreviousResult = &Results[0];
 		FSamplerResult* CurrentResult = &Results[1];
 
+		// @todo_pcg: be more precise about allocation (sampler could provide properties to allocate)
+		OutPointData->AllocateProperties(EPCGPointNativeProperties::All);
+
 		if (!Sampler->IsDone())
 		{
 			Sampler->Step(*PreviousResult);
@@ -1031,7 +1043,7 @@ namespace PCGSplineSamplerHelpers
 		}
 	}
  
-	void SampleInteriorData(FPCGContext* Context, const UPCGPolyLineData* LineData, const UPCGSpatialData* InBoundingShape, const UPCGSpatialData* InProjectionTarget, const FPCGProjectionParams& InProjectionParams, const FPCGSplineSamplerParams& Params, UPCGPointData* OutPointData)
+	void SampleInteriorData(FPCGContext* Context, const UPCGPolyLineData* LineData, const UPCGSpatialData* InBoundingShape, const UPCGSpatialData* InProjectionTarget, const FPCGProjectionParams& InProjectionParams, const FPCGSplineSamplerParams& Params, UPCGBasePointData* OutPointData)
 	{
 		check(LineData && OutPointData);
 
@@ -1204,8 +1216,6 @@ namespace PCGSplineSamplerHelpers
 			}
 		}
 
-		TArray<FPCGPoint>& OutPoints = OutPointData->GetMutablePoints();
-
 		// RayPadding should be large enough to avoid potential misclassifications in SegmentPolygonIntersection2D
 		const FVector::FReal RayPadding = 100.f;
 		const FVector::FReal MinY = FMath::CeilToDouble(MinPoint.Y / Params.InteriorSampleSpacing) * Params.InteriorSampleSpacing;
@@ -1403,20 +1413,29 @@ namespace PCGSplineSamplerHelpers
 		int32 SampleIndex = 0;
 
 		// TODO: should we parallel for this too?
-		OutPoints.Reserve(PointCount);
+		OutPointData->SetNumPoints(PointCount);
+		OutPointData->AllocateProperties(EPCGPointNativeProperties::Transform | EPCGPointNativeProperties::Seed | EPCGPointNativeProperties::Density);
+
+		// Set Constant values
+		OutPointData->SetBoundsMin(BoundsMin);
+		OutPointData->SetBoundsMax(BoundsMax);
+		OutPointData->SetSteepness(Params.PointSteepness);
+				
+		FPCGPointValueRanges OutRanges(OutPointData, /*bAllocate=*/false);
+
+		int32 NumWritten = 0;
+
 		for (const TArray<TTuple<FTransform, FVector, float>>& InteriorData : InteriorSplinePointData)
 		{
 			for (const TTuple<FTransform, FVector, float>& InteriorPoint : InteriorData)
 			{
-				FPCGPoint& Point = OutPoints.Emplace_GetRef();
-				Point.Transform = InteriorPoint.Get<0>();
-				SetSeed(Point, InteriorPoint.Get<1>(), Params, Seed, ++SampleIndex);
-				Point.Density = InteriorPoint.Get<2>();
-				Point.BoundsMin = BoundsMin;
-				Point.BoundsMax = BoundsMax;
-				Point.Steepness = Params.PointSteepness;
+				OutRanges.TransformRange[NumWritten] = InteriorPoint.Get<0>();
+				SetSeed(OutRanges.SeedRange[NumWritten], OutRanges.TransformRange[NumWritten], InteriorPoint.Get<1>(), Params, Seed, ++SampleIndex);
+				OutRanges.DensityRange[NumWritten] = InteriorPoint.Get<2>();
+				++NumWritten;
 			}
 		}
+		check(NumWritten == PointCount);
 	}
 
 	const UPCGPolyLineData* GetPolyLineData(const UPCGSpatialData* InSpatialData)
@@ -1568,8 +1587,12 @@ bool FPCGSplineSamplerElement::ExecuteInternal(FPCGContext* Context) const
 
 		FPCGTaggedData& Output = Outputs.Emplace_GetRef(Input);
 
-		UPCGPointData* SampledPointData = FPCGContext::NewObject_AnyThread<UPCGPointData>(Context);
-		SampledPointData->InitializeFromData(SpatialData);
+		UPCGBasePointData* SampledPointData = FPCGContext::NewPointData_AnyThread(Context);
+
+		FPCGInitializeFromDataParams InitializeFromDataParams(SpatialData);
+		InitializeFromDataParams.bInheritSpatialData = false;
+
+		SampledPointData->InitializeFromDataWithParams(InitializeFromDataParams);
 		Output.Data = SampledPointData;
 
 		if (SamplerParams.Dimension == EPCGSplineSamplingDimension::OnInterior)
diff --git a/Engine/Plugins/PCG/Source/PCG/Private/Tests/Elements/PCGCreatePointsTest.cpp b/Engine/Plugins/PCG/Source/PCG/Private/Tests/Elements/PCGCreatePointsTest.cpp
--- a/Engine/Plugins/PCG/Source/PCG/Private/Tests/Elements/PCGCreatePointsTest.cpp
+++ b/Engine/Plugins/PCG/Source/PCG/Private/Tests/Elements/PCGCreatePointsTest.cpp
@@ -37,21 +37,21 @@ bool FPCGCreatePointsTest_Basic::RunTest(const FString& Parameters)
 
 	UTEST_EQUAL("Output count", Outputs.Num(), 1);
 
-	const UPCGPointData* OutPointData = Cast<UPCGPointData>(Outputs[0].Data);
+	const UPCGBasePointData* OutPointData = Cast<UPCGBasePointData>(Outputs[0].Data);
 
 	UTEST_NOT_NULL("Output point data", OutPointData);
 
-	const TArray<FPCGPoint>& OutPoints = OutPointData->GetPoints();
+	FConstPCGPointValueRanges OutRanges(OutPointData);
 
-	UTEST_EQUAL("Output point count", OutPoints.Num(), 100);
+	UTEST_EQUAL("Output point count", OutPointData->GetNumPoints(), 100);
 
-	for (int i = 0; i < OutPoints.Num(); ++i)
+	for (int i = 0; i < OutPointData->GetNumPoints(); ++i)
 	{
 		// Seed at 0 will compute a seed depending on the point position
 		static constexpr int32 SeedAtZero = 907633527;
-		UTEST_EQUAL(FString::Format(TEXT("InArray[{0}].Seed is equal to OutArray[{0}].Seed"), { i }), OutPoints[i].Seed, i == 0 ? SeedAtZero : Settings->PointsToCreate[i].Seed);
-		UTEST_EQUAL(FString::Format(TEXT("InArray[{0}].Density is equal to OutArray[{0}].Density"), { i }), OutPoints[i].Density, Settings->PointsToCreate[i].Density);
-		UTEST_EQUAL(FString::Format(TEXT("InArray[{0}].Transform is equal to OutArray[{0}].Transform"), { i }), OutPoints[i].Transform, Settings->PointsToCreate[i].Transform);
+		UTEST_EQUAL(FString::Format(TEXT("InArray[{0}].Seed is equal to OutArray[{0}].Seed"), { i }), OutRanges.SeedRange[i], i == 0 ? SeedAtZero : Settings->PointsToCreate[i].Seed);
+		UTEST_EQUAL(FString::Format(TEXT("InArray[{0}].Density is equal to OutArray[{0}].Density"), { i }), OutRanges.DensityRange[i], Settings->PointsToCreate[i].Density);
+		UTEST_EQUAL(FString::Format(TEXT("InArray[{0}].Transform is equal to OutArray[{0}].Transform"), { i }), OutRanges.TransformRange[i], Settings->PointsToCreate[i].Transform);
 	}
 
 	return true;
diff --git a/Engine/Plugins/PCG/Source/PCG/Private/Tests/Elements/PCGDuplicatePointTest.cpp b/Engine/Plugins/PCG/Source/PCG/Private/Tests/Elements/PCGDuplicatePointTest.cpp
--- a/Engine/Plugins/PCG/Source/PCG/Private/Tests/Elements/PCGDuplicatePointTest.cpp
+++ b/Engine/Plugins/PCG/Source/PCG/Private/Tests/Elements/PCGDuplicatePointTest.cpp
@@ -50,36 +50,36 @@ protected:
 
 		UTEST_EQUAL("Output count", Outputs.Num(), 1);
 
-		const UPCGPointData* OutPointData = Cast<UPCGPointData>(Outputs[0].Data);
+		const UPCGBasePointData* OutPointData = Cast<UPCGBasePointData>(Outputs[0].Data);
 
 		UTEST_NOT_NULL("OutputA point data", OutPointData);
+				
+		UTEST_EQUAL("OutputA point count", OutPointData->GetNumPoints(), Parameters.ExpectedPointCount);
 
-		const TArray<FPCGPoint>& OutPoints = OutPointData->GetPoints();
+		const FConstPCGPointValueRanges OutRanges(OutPointData);
 
-		UTEST_EQUAL("OutputA point count", OutPoints.Num(), Parameters.ExpectedPointCount);
-
-		const FVector LocationGetter = (OutPoints[0].BoundsMax - OutPoints[0].BoundsMin) * Parameters.Direction + Parameters.PtTransform.GetLocation();
+		const FVector LocationGetter = (OutRanges.BoundsMaxRange[0] - OutRanges.BoundsMinRange[0]) * Parameters.Direction + Parameters.PtTransform.GetLocation();
 
 		if (Parameters.bCheckLocation)
 		{
-			for (int i = 0; i < OutPoints.Num(); ++i)
+			for (int i = 0; i < OutPointData->GetNumPoints(); ++i)
 			{
-				UTEST_EQUAL("Output Location", OutPoints[i].Transform.GetLocation(), LocationGetter * i);
+				UTEST_EQUAL("Output Location", OutRanges.TransformRange[i].GetLocation(), LocationGetter * i);
 			}
 		}
 		if (Parameters.bCheckRotation)
 		{
-			for (int i = 0; i < OutPoints.Num(); ++i)
+			for (int i = 0; i < OutPointData->GetNumPoints(); ++i)
 			{
-				UTEST_EQUAL("Output Rotation", OutPoints[i].Transform.GetRotation().Rotator(), Parameters.PtTransform.GetRotation().Rotator() * i);
+				UTEST_EQUAL("Output Rotation", OutRanges.TransformRange[i].GetRotation().Rotator(), Parameters.PtTransform.GetRotation().Rotator() * i);
 			}
 		}
 		if (Parameters.bCheckScale)
 		{
 			const FVector ParameterTestScale = Parameters.PtTransform.GetScale3D();
-			for (int i = 0; i < OutPoints.Num(); ++i)
+			for (int i = 0; i < OutPointData->GetNumPoints(); ++i)
 			{
-				UTEST_EQUAL("Output Scale", OutPoints[i].Transform.GetScale3D(), FVector(FMath::Pow(ParameterTestScale.X, i), FMath::Pow(ParameterTestScale.Y, i), FMath::Pow(ParameterTestScale.Z, i)));
+				UTEST_EQUAL("Output Scale", OutRanges.TransformRange[i].GetScale3D(), FVector(FMath::Pow(ParameterTestScale.X, i), FMath::Pow(ParameterTestScale.Y, i), FMath::Pow(ParameterTestScale.Z, i)));
 			}
 		}
 
diff --git a/Engine/Plugins/PCG/Source/PCG/Public/Elements/PCGDataFromActor.h b/Engine/Plugins/PCG/Source/PCG/Public/Elements/PCGDataFromActor.h
--- a/Engine/Plugins/PCG/Source/PCG/Public/Elements/PCGDataFromActor.h
+++ b/Engine/Plugins/PCG/Source/PCG/Public/Elements/PCGDataFromActor.h
@@ -158,6 +158,7 @@ public:
 protected:
 	virtual bool PrepareDataInternal(FPCGContext* Context) const override;
 	virtual bool ExecuteInternal(FPCGContext* Context) const override;
+	virtual bool SupportsBasePointDataInputs(FPCGContext* InContext) const override { return true; }
 	void GatherWaitTasks(AActor* FoundActor, FPCGContext* InContext, TArray<FPCGTaskId>& OutWaitTasks) const;
 	virtual void ProcessActors(FPCGContext* Context, const UPCGDataFromActorSettings* Settings, const TArray<AActor*>& FoundActors) const;
 	virtual void ProcessActor(FPCGContext* Context, const UPCGDataFromActorSettings* Settings, AActor* FoundActor) const;
diff --git a/Engine/Plugins/PCG/Source/PCG/Public/Elements/PCGDuplicatePoint.h b/Engine/Plugins/PCG/Source/PCG/Public/Elements/PCGDuplicatePoint.h
--- a/Engine/Plugins/PCG/Source/PCG/Public/Elements/PCGDuplicatePoint.h
+++ b/Engine/Plugins/PCG/Source/PCG/Public/Elements/PCGDuplicatePoint.h
@@ -57,4 +57,5 @@ class FPCGDuplicatePointElement : public IPCGElement
 protected:
 	virtual bool ExecuteInternal(FPCGContext* Context) const override;
 	virtual EPCGElementExecutionLoopMode ExecutionLoopMode(const UPCGSettings* Settings) const override { return EPCGElementExecutionLoopMode::SinglePrimaryPin; }
+	virtual bool SupportsBasePointDataInputs(FPCGContext* InContext) const override { return true; }
 };
\ No newline at end of file
diff --git a/Engine/Plugins/PCG/Source/PCG/Public/Elements/PCGPointFromMeshElement.h b/Engine/Plugins/PCG/Source/PCG/Public/Elements/PCGPointFromMeshElement.h
--- a/Engine/Plugins/PCG/Source/PCG/Public/Elements/PCGPointFromMeshElement.h
+++ b/Engine/Plugins/PCG/Source/PCG/Public/Elements/PCGPointFromMeshElement.h
@@ -57,4 +57,5 @@ public:
 protected:
 	virtual bool PrepareDataInternal(FPCGContext* Context) const override;
 	virtual bool ExecuteInternal(FPCGContext* Context) const override;
+	virtual bool SupportsBasePointDataInputs(FPCGContext* InContext) const override { return true; }
 };
diff --git a/Engine/Plugins/PCG/Source/PCG/Public/Elements/PCGSpawnActor.h b/Engine/Plugins/PCG/Source/PCG/Public/Elements/PCGSpawnActor.h
--- a/Engine/Plugins/PCG/Source/PCG/Public/Elements/PCGSpawnActor.h
+++ b/Engine/Plugins/PCG/Source/PCG/Public/Elements/PCGSpawnActor.h
@@ -10,7 +10,7 @@
 
 class AActor;
 class UPCGGraphInterface;
-class UPCGPointData;
+class UPCGBasePointData;
 
 UENUM()
 enum class EPCGSpawnActorOption : uint8
@@ -183,6 +183,7 @@ public:
 	virtual bool CanExecuteOnlyOnMainThread(FPCGContext* Context) const override { return true; }
 	virtual bool IsCacheable(const UPCGSettings* InSettings) const override { return false; }
 	virtual bool IsPassthrough(const UPCGSettings* InSettings) const override { return !InSettings; }
+	virtual bool SupportsBasePointDataInputs(FPCGContext* InContext) const override { return true; }
 
 protected:
 	virtual bool ExecuteInternal(FPCGContext* Context) const override;
@@ -192,6 +193,6 @@ private:
 
 	TArray<FName> GetNewActorTags(FPCGContext* Context, AActor* TargetActor, bool bInheritActorTags, const TArray<FName>& AdditionalTags) const;
 
-	void CollapseIntoTargetActor(FPCGSubgraphContext* Context, AActor* TargetActor, TSubclassOf<AActor> TemplateActorClass, const UPCGPointData* PointData) const;
-	void SpawnActors(FPCGSubgraphContext* Context, AActor* TargetActor, TSubclassOf<AActor> TemplateActorClass, AActor* TemplateActor, FPCGTaggedData& Output, const UPCGPointData* PointData, UPCGPointData* OutPointData) const;
+	void CollapseIntoTargetActor(FPCGSubgraphContext* Context, AActor* TargetActor, TSubclassOf<AActor> TemplateActorClass, const UPCGBasePointData* PointData) const;
+	void SpawnActors(FPCGSubgraphContext* Context, AActor* TargetActor, TSubclassOf<AActor> TemplateActorClass, AActor* TemplateActor, FPCGTaggedData& Output, const UPCGBasePointData* PointData, UPCGBasePointData* OutPointData) const;
 };
diff --git a/Engine/Plugins/PCG/Source/PCG/Public/Elements/PCGSplineSampler.h b/Engine/Plugins/PCG/Source/PCG/Public/Elements/PCGSplineSampler.h
--- a/Engine/Plugins/PCG/Source/PCG/Public/Elements/PCGSplineSampler.h
+++ b/Engine/Plugins/PCG/Source/PCG/Public/Elements/PCGSplineSampler.h
@@ -10,9 +10,9 @@
 
 struct FPCGContext;
 struct FPCGProjectionParams;
+class UPCGBasePointData;
 class UPCGPolyLineData;
 class UPCGSpatialData;
-class UPCGPointData;
 
 namespace PCGSplineSamplerConstants
 {
@@ -234,10 +234,10 @@ struct FPCGSplineSamplerParams
 namespace PCGSplineSamplerHelpers
 {
 	/** Samples on spline or within volume around it. */
-	void SampleLineData(FPCGContext* Context, const UPCGPolyLineData* LineData, const UPCGSpatialData* InBoundingShape, const UPCGSpatialData* InProjectionTarget, const FPCGProjectionParams& InProjectionParams, const FPCGSplineSamplerParams& Params, UPCGPointData* OutPointData);
+	void SampleLineData(FPCGContext* Context, const UPCGPolyLineData* LineData, const UPCGSpatialData* InBoundingShape, const UPCGSpatialData* InProjectionTarget, const FPCGProjectionParams& InProjectionParams, const FPCGSplineSamplerParams& Params, UPCGBasePointData* OutPointData);
 
 	/** Samples 2D region bounded by spline. */
-	void SampleInteriorData(FPCGContext* Context, const UPCGPolyLineData* LineData, const UPCGSpatialData* InBoundingShape, const UPCGSpatialData* InProjectionTarget, const FPCGProjectionParams& InProjectionParams, const FPCGSplineSamplerParams& Params, UPCGPointData* OutPointData);
+	void SampleInteriorData(FPCGContext* Context, const UPCGPolyLineData* LineData, const UPCGSpatialData* InBoundingShape, const UPCGSpatialData* InProjectionTarget, const FPCGProjectionParams& InProjectionParams, const FPCGSplineSamplerParams& Params, UPCGBasePointData* OutPointData);
 
 	const UPCGPolyLineData* GetPolyLineData(const UPCGSpatialData* InSpatialData);
 
@@ -289,4 +289,5 @@ public:
 protected:
 	virtual bool ExecuteInternal(FPCGContext* Context) const override;
 	virtual EPCGElementExecutionLoopMode ExecutionLoopMode(const UPCGSettings* Settings) const override { return EPCGElementExecutionLoopMode::SinglePrimaryPin; }
+	virtual bool SupportsBasePointDataInputs(FPCGContext* InContext) const override { return true; }
 };
