[PCG] Add missing nullptr check (static analysis)

#rnx

diff --git a/Engine/Plugins/PCG/Source/PCG/Private/Elements/PCGSpawnActor.cpp b/Engine/Plugins/PCG/Source/PCG/Private/Elements/PCGSpawnActor.cpp
--- a/Engine/Plugins/PCG/Source/PCG/Private/Elements/PCGSpawnActor.cpp
+++ b/Engine/Plugins/PCG/Source/PCG/Private/Elements/PCGSpawnActor.cpp
@@ -919,7 +919,7 @@ void FPCGSpawnActorElement::SpawnActors(FPCGSubgraphContext* Context, AActor* Ta
 		// Also set to the output data the actor reference.
 		if (Settings->Option == EPCGSpawnActorOption::NoMerging)
 		{
-			TPCGValueRange<int64> MetadataEntryRange = OutPointData->GetMetadataEntryValueRange();
+			TPCGValueRange<int64> MetadataEntryRange = OutPointData ? OutPointData->GetMetadataEntryValueRange() : TPCGValueRange<int64>();
 
 			const TArray<TSoftObjectPtr<AActor>>& GeneratedActors = ReusedManagedActorsResource->GetConstGeneratedActors();
 			for (int32 i = 0; i < GeneratedActors.Num(); ++i)
