[PCG] Add InstanceMinDrawDistance to ProceduralISM

diff --git a/Engine/Plugins/PCG/Source/PCG/Private/Components/PCGProceduralISMComponent.cpp b/Engine/Plugins/PCG/Source/PCG/Private/Components/PCGProceduralISMComponent.cpp
--- a/Engine/Plugins/PCG/Source/PCG/Private/Components/PCGProceduralISMComponent.cpp
+++ b/Engine/Plugins/PCG/Source/PCG/Private/Components/PCGProceduralISMComponent.cpp
@@ -121,6 +121,7 @@ FPCGProceduralISMComponentDescriptor::FPCGProceduralISMComponentDescriptor()
 
 FPCGProceduralISMComponentDescriptor& FPCGProceduralISMComponentDescriptor::operator=(const FPCGSoftISMComponentDescriptor& Other)
 {
+	InstanceMinDrawDistance = Other.InstanceMinDrawDistance;
 	InstanceStartCullDistance = Other.InstanceStartCullDistance;
 	InstanceEndCullDistance = Other.InstanceEndCullDistance;
 	OverlayMaterial = Other.OverlayMaterial.LoadSynchronous();
@@ -171,6 +172,7 @@ void FPCGProceduralISMComponentDescriptor::InitComponent(UPCGProceduralISMCompon
 	InProceduralISMComponent->SetNumInstances(NumInstances);
 	InProceduralISMComponent->SetNumCustomDataFloats(NumCustomFloats);
 	InProceduralISMComponent->SetLocalBounds(LocalBounds);
+	InProceduralISMComponent->SetMinDrawDistance(InstanceMinDrawDistance);
 	InProceduralISMComponent->SetCullDistances(InstanceStartCullDistance, InstanceEndCullDistance);
 	InProceduralISMComponent->ComponentTags = ComponentTags;
 	InProceduralISMComponent->Mobility = Mobility;
@@ -206,6 +208,7 @@ void FPCGProceduralISMComponentDescriptor::InitFrom(UPCGProceduralISMComponent*
 	NumInstances = InProceduralISMComponent->GetNumInstances();
 	NumCustomFloats = InProceduralISMComponent->GetNumCustomDataFloats();
 	InProceduralISMComponent->GetLocalBounds(LocalBounds);
+	InstanceMinDrawDistance = InProceduralISMComponent->MinDrawDistance;
 	InProceduralISMComponent->GetCullDistances(InstanceStartCullDistance, InstanceEndCullDistance);
 	ComponentTags = InProceduralISMComponent->ComponentTags;
 	Mobility = InProceduralISMComponent->Mobility;
@@ -257,6 +260,7 @@ bool FPCGProceduralISMComponentDescriptor::operator==(const FPCGProceduralISMCom
 		&& NumInstances == Other.NumInstances
 		&& NumCustomFloats == Other.NumCustomFloats
 		&& LocalBounds == Other.LocalBounds
+		&& InstanceMinDrawDistance == Other.InstanceMinDrawDistance
 		&& InstanceStartCullDistance == Other.InstanceStartCullDistance
 		&& InstanceEndCullDistance == Other.InstanceEndCullDistance
 		&& ComponentTags == Other.ComponentTags
@@ -375,6 +379,7 @@ FPrimitiveSceneProxy* UPCGProceduralISMComponent::CreateStaticMeshSceneProxy(Nan
 	FInstancedStaticMeshSceneProxyDesc Desc;
 	Desc.InitializeFromStaticMeshComponent(this);
 	Desc.InstanceDataSceneProxy = MakeShared<FInstanceDataSceneProxy, ESPMode::ThreadSafe>(MoveTemp(InstanceSceneDataBuffers));
+	Desc.InstanceMinDrawDistance = InstanceMinDrawDistance;
 	Desc.InstanceStartCullDistance = InstanceStartCullDistance;
 	Desc.InstanceEndCullDistance = InstanceEndCullDistance;
 	Desc.bUseGpuLodSelection = true;
@@ -472,6 +477,19 @@ void UPCGProceduralISMComponent::SetCullDistances(int32 InStartCullDistance, int
 	}
 }
 
+void UPCGProceduralISMComponent::SetMinDrawDistance(int32 InMinDrawDistance)
+{
+	if (InstanceMinDrawDistance != InMinDrawDistance)
+	{
+		InstanceMinDrawDistance = InMinDrawDistance;
+
+		if (GetScene() && SceneProxy)
+		{
+			GetScene()->UpdatePrimitiveDrawDistance(this, InMinDrawDistance, /*MaxDrawDistance=*/0, /*VirtualTextureMaxDrawDistance=*/0);
+		}
+	}
+}
+
 void UPCGProceduralISMComponent::ValidateComponentSetup()
 {
 	if (bAffectDynamicIndirectLighting)
diff --git a/Engine/Plugins/PCG/Source/PCG/Private/Components/PCGProceduralISMComponent.h b/Engine/Plugins/PCG/Source/PCG/Private/Components/PCGProceduralISMComponent.h
--- a/Engine/Plugins/PCG/Source/PCG/Private/Components/PCGProceduralISMComponent.h
+++ b/Engine/Plugins/PCG/Source/PCG/Private/Components/PCGProceduralISMComponent.h
@@ -68,6 +68,14 @@ public:
 	UFUNCTION(BlueprintCallable, Category = "Components|PCGProceduralISMC")
 	void GetCullDistances(int32& OutStartCullDistance, int32& OutEndCullDistance) const { OutStartCullDistance = InstanceStartCullDistance; OutEndCullDistance = InstanceEndCullDistance; }
 
+	/** Sets the minimum distance at which instances will render for this component. */
+	UFUNCTION(BlueprintCallable, Category = "Components|PCGProceduralISMC")
+	void SetMinDrawDistance(int32 InMinDrawDistance);
+
+	/** Gets the minimum distance at which instances will render for this component. */
+	UFUNCTION(BlueprintCallable, Category = "Components|PCGProceduralISMC")
+	void GetMinDrawDistance(int32& OutMinDrawDistance) const { OutMinDrawDistance = InstanceMinDrawDistance; }
+
 	//~ Begin UActorComponent Interface
 	bool ShouldCreatePhysicsState() const override { return false; }
 	bool IsHLODRelevant() const override { return false; }
@@ -127,6 +135,10 @@ protected:
 	UPROPERTY(EditAnywhere, Category = Instances)
 	FBox LocalBounds = FBox(-FVector::OneVector * DefaultBoundsExtent, FVector::OneVector * DefaultBoundsExtent);
 
+	/** Distance from camera at which each instance begins to draw. */
+	UPROPERTY(EditAnywhere, BlueprintReadOnly, Category=Culling)
+	int32 InstanceMinDrawDistance;
+
 	/** Distance from camera at which each instance begins to fade out. */
 	UPROPERTY(EditAnywhere, BlueprintReadOnly, Category = Culling)
 	int32 InstanceStartCullDistance;
@@ -201,6 +213,9 @@ public:
 	UPROPERTY()
 	FBox LocalBounds;
 
+	UPROPERTY()
+	int32 InstanceMinDrawDistance;
+
 	UPROPERTY()
 	int32 InstanceStartCullDistance;
 
