[PCG] Stop manually marking unused kernels as garbage
* Fix ensure in GraphCompilationContext when kernels created asynchronously have already been marked for GC

#rnx

diff --git a/Engine/Plugins/PCG/Source/PCG/Private/Graph/PCGGraphCompilerGPU.cpp b/Engine/Plugins/PCG/Source/PCG/Private/Graph/PCGGraphCompilerGPU.cpp
--- a/Engine/Plugins/PCG/Source/PCG/Private/Graph/PCGGraphCompilerGPU.cpp
+++ b/Engine/Plugins/PCG/Source/PCG/Private/Graph/PCGGraphCompilerGPU.cpp
@@ -730,6 +730,10 @@ void FPCGGraphCompilerGPU::BuildComputeGraphTask(
 	UPCGComputeGraph* ComputeGraph = GraphCompiler.GetComputeGraph(InGraph, InGridSize, InComputeGraphIndex);
 
 	// If the graph does not exist already, create it.
+	// todo_pcg: If the graph does exist already, then we throw the kernels we created away, which is pretty wasteful, even if only happens
+	// on first graph compilation. We could probably defer creation of these objects until they're really needed by storing "createinfo" data.
+	// Or possibly the caching of compute graphs could be done per pcg-graph, rather than per-compute-graph, and we could avoid doing any GPU
+	// compilation work if already cached.
 	if (!ComputeGraph)
 	{
 		// TODO: There is a race condition here where two threads can try to find this graph at the same time and both fail, so both create the graph.
@@ -763,20 +767,6 @@ void FPCGGraphCompilerGPU::BuildComputeGraphTask(
 			InTaskIdToKernel[TaskId]->Rename(nullptr, ComputeGraph);
 		}
 	}
-	else
-	{
-		// Compute graph obtained from cache, it will have its own kernels, so garbage the kernels we created.
-		// todo_pcg: Pretty wasteful to create kernel objects and then throw them away, even if only happens on first graph compilation.
-		// We could probably defer creation of these objects until they're really needed by storing "createinfo" data. Or possibly the caching
-		// of compute graphs could be done per pcg-graph, rather than per-compute-graph, and we could avoid doing any GPU compilation work if already cached.
-		for (TPair<FPCGTaskId, UPCGComputeKernel*> TaskIdAndKernel : InTaskIdToKernel)
-		{
-			if (ensure(TaskIdAndKernel.Value))
-			{
-				TaskIdAndKernel.Value->MarkAsGarbage();
-			}
-		}
-	}
 
 	const int32 ComputeGraphIndex = InOutContext.AddCompiledComputeGraph(ComputeGraph);
 	ensure(ComputeGraphIndex == InComputeGraphIndex);
