[PCG] Final GPU core changes for sparse copy points feature

* Move SplitGraphAtOutput to kernel so splits can happen between kernels within a node.
* Fix data upload adaptor failing to reuse GPU buffers (and doing an expensive readback+reupload instead), by comparing data items on the correct pin. When a readback is necessary, only readback data on relevant pin.
* Add a new stage in data provider preparation "PerformPreExecuteReadbacks". This gives data providers the opportunity to do any custom readbacks for example analysis results.
* Move any data provider preparation that requires data descriptiosn to PrepareForExecution. This is tider than doing in the create method, and also means we know any analysis results that might affect data allocations are read back and data descriptions are final.

#rnx

diff --git a/Engine/Plugins/PCG/Source/PCG/Private/Compute/DataInterfaces/Elements/PCGMetadataPartitionDataInterface.cpp b/Engine/Plugins/PCG/Source/PCG/Private/Compute/DataInterfaces/Elements/PCGMetadataPartitionDataInterface.cpp
--- a/Engine/Plugins/PCG/Source/PCG/Private/Compute/DataInterfaces/Elements/PCGMetadataPartitionDataInterface.cpp
+++ b/Engine/Plugins/PCG/Source/PCG/Private/Compute/DataInterfaces/Elements/PCGMetadataPartitionDataInterface.cpp
@@ -83,37 +83,48 @@ void UPCGMetadataPartitionDataInterface::GetHLSL(FString& OutHLSL, FString const
 UComputeDataProvider* UPCGMetadataPartitionDataInterface::CreateDataProvider(TObjectPtr<UObject> InBinding, uint64 InInputMask, uint64 InOutputMask) const
 {
 	TRACE_CPUPROFILER_EVENT_SCOPE(UPCGMetadataPartitionDataInterface::CreateDataProvider);
-	UPCGDataBinding* Binding = CastChecked<UPCGDataBinding>(InBinding);
-	check(ProducerKernel);
-	const UPCGMetadataPartitionSettings* Settings = Cast<UPCGMetadataPartitionKernel>(ProducerKernel)->Settings;
-	check(Settings);
 
-	TObjectPtr<UPCGMetaDataPartitionDataProvider> DataProvider = NewObject<UPCGMetaDataPartitionDataProvider>();
+	UPCGMetaDataPartitionDataProvider* DataProvider = NewObject<UPCGMetaDataPartitionDataProvider>();
+	DataProvider->Kernel = ProducerKernel;
+
+	return DataProvider;
+}
+
+bool UPCGMetaDataPartitionDataProvider::PrepareForExecute_GameThread(UPCGDataBinding* InBinding)
+{
+	TRACE_CPUPROFILER_EVENT_SCOPE(UPCGMetadataPartitionDataInterface::PrepareForExecute_GameThread);
+	check(InBinding);
+
+	if (!Super::PrepareForExecute_GameThread(InBinding))
+	{
+		return false;
+	}
+
+	const UPCGMetadataPartitionSettings* Settings = Cast<UPCGMetadataPartitionKernel>(Kernel)->Settings;
+	check(Settings);
 
 	if (ensure(Settings->PartitionAttributeSelectors.Num() == 1 && Settings->PartitionAttributeSelectors[0].IsBasicAttribute()))
 	{
 		FPCGDataCollectionDesc InputDataDesc;
-		ensure(Binding->GetKernelPinDataDesc_GameThread(ProducerKernel, PCGPinConstants::DefaultInputLabel, /*bIsInput*/true, InputDataDesc));
+		ensure(InBinding->GetKernelPinDataDesc_GameThread(Kernel, PCGPinConstants::DefaultInputLabel, /*bIsInput*/true, InputDataDesc));
 
-		const int32 PartitionAttributeId = Binding->GetAttributeId(Settings->PartitionAttributeSelectors[0].GetAttributeName(), EPCGKernelAttributeType::StringKey);
+		PartitionAttributeId = InBinding->GetAttributeId(Settings->PartitionAttributeSelectors[0].GetAttributeName(), EPCGKernelAttributeType::StringKey);
 		if (PartitionAttributeId != INDEX_NONE)
 		{
-			InputDataDesc.GetUniqueStringKeyValues(PartitionAttributeId, DataProvider->UniqueStringKeyValues);
+			InputDataDesc.GetUniqueStringKeyValues(PartitionAttributeId, UniqueStringKeyValues);
 		}
 
-		int32 LargestValue = INDEX_NONE;
-		for (int32 Key : DataProvider->UniqueStringKeyValues)
+		MaxAttributeValue = INDEX_NONE;
+		for (int32 Key : UniqueStringKeyValues)
 		{
-			LargestValue = FMath::Max(LargestValue, Key);
+			MaxAttributeValue = FMath::Max(MaxAttributeValue, Key);
 		}
 
-		DataProvider->NumInputData = InputDataDesc.DataDescs.Num();
-		DataProvider->PartitionAttributeId = PartitionAttributeId;
-		DataProvider->NumPartitions = DataProvider->UniqueStringKeyValues.Num();
-		DataProvider->MaxAttributeValue = LargestValue;
+		NumInputData = InputDataDesc.DataDescs.Num();
+		NumPartitions = UniqueStringKeyValues.Num();
 	}
 
-	return DataProvider;
+	return true;
 }
 
 FComputeDataProviderRenderProxy* UPCGMetaDataPartitionDataProvider::GetRenderProxy()
diff --git a/Engine/Plugins/PCG/Source/PCG/Private/Compute/DataInterfaces/Elements/PCGMetadataPartitionDataInterface.h b/Engine/Plugins/PCG/Source/PCG/Private/Compute/DataInterfaces/Elements/PCGMetadataPartitionDataInterface.h
--- a/Engine/Plugins/PCG/Source/PCG/Private/Compute/DataInterfaces/Elements/PCGMetadataPartitionDataInterface.h
+++ b/Engine/Plugins/PCG/Source/PCG/Private/Compute/DataInterfaces/Elements/PCGMetadataPartitionDataInterface.h
@@ -26,15 +26,23 @@ public:
 };
 
 UCLASS()
-class UPCGMetaDataPartitionDataProvider : public UComputeDataProvider
+class UPCGMetaDataPartitionDataProvider : public UPCGComputeDataProvider
 {
 	GENERATED_BODY()
 
 public:
+	//~ Begin UPCGComputeDataProvider Interface
+	virtual bool PrepareForExecute_GameThread(UPCGDataBinding* InBinding) override;
+	//~ End UPCGComputeDataProvider Interface
+
 	//~ Begin UComputeDataProvider Interface
 	FComputeDataProviderRenderProxy* GetRenderProxy() override;
 	//~ End UComputeDataProvider Interface
 
+public:
+	UPROPERTY()
+	TObjectPtr<const UPCGComputeKernel> Kernel;
+
 	UPROPERTY()
 	int32 PartitionAttributeId = INDEX_NONE;
 
diff --git a/Engine/Plugins/PCG/Source/PCG/Private/Compute/DataInterfaces/PCGComputeDataInterface.h b/Engine/Plugins/PCG/Source/PCG/Private/Compute/DataInterfaces/PCGComputeDataInterface.h
--- a/Engine/Plugins/PCG/Source/PCG/Private/Compute/DataInterfaces/PCGComputeDataInterface.h
+++ b/Engine/Plugins/PCG/Source/PCG/Private/Compute/DataInterfaces/PCGComputeDataInterface.h
@@ -54,5 +54,9 @@ class UPCGComputeDataProvider : public UComputeDataProvider
 	GENERATED_BODY()
 
 public:
-	virtual bool PrepareForExecute_GameThread(const UPCGDataBinding* InBinding) { return true; }
+	/** Called before PrepareForExecute to allow data providers to do any readbacks and finalize data descriptions. */
+	virtual bool PerformPreExecuteReadbacks_GameThread(UPCGDataBinding* InBinding) { return true; }
+
+	/* Called prior to GetRenderProxy and execution. Any processing of data descriptions should be done here (after PerformPreExecuteReadbacks has been called). */
+	virtual bool PrepareForExecute_GameThread(UPCGDataBinding* InBinding) { return true; }
 };
diff --git a/Engine/Plugins/PCG/Source/PCG/Private/Compute/DataInterfaces/PCGCustomKernelDataInterface.cpp b/Engine/Plugins/PCG/Source/PCG/Private/Compute/DataInterfaces/PCGCustomKernelDataInterface.cpp
--- a/Engine/Plugins/PCG/Source/PCG/Private/Compute/DataInterfaces/PCGCustomKernelDataInterface.cpp
+++ b/Engine/Plugins/PCG/Source/PCG/Private/Compute/DataInterfaces/PCGCustomKernelDataInterface.cpp
@@ -101,9 +101,8 @@ UComputeDataProvider* UPCGCustomKernelDataInterface::CreateDataProvider(TObjectP
 	check(SourceComponent);
 
 	UPCGCustomComputeKernelDataProvider* Provider = NewObject<UPCGCustomComputeKernelDataProvider>();
-	Provider->ThreadCount = Kernel->ComputeThreadCount(Binding);
+	Provider->Kernel = Kernel;
 	Provider->ThreadCountMultiplier = Kernel->GetThreadCountMultiplier();
-
 	Provider->Seed = static_cast<uint32>(Settings->GetSeed(SourceComponent));
 	Provider->SeedSettings = static_cast<uint32>(Settings->Seed);
 	Provider->SeedComponent = static_cast<uint32>(SourceComponent->Seed);
@@ -113,6 +112,12 @@ UComputeDataProvider* UPCGCustomKernelDataInterface::CreateDataProvider(TObjectP
 	return Provider;
 }
 
+bool UPCGCustomComputeKernelDataProvider::PrepareForExecute_GameThread(UPCGDataBinding* InBinding)
+{
+	ThreadCount = Kernel->ComputeThreadCount(InBinding);
+	return true;
+}
+
 FComputeDataProviderRenderProxy* UPCGCustomComputeKernelDataProvider::GetRenderProxy()
 {
 	return new FPCGCustomComputeKernelDataProviderProxy(ThreadCount, ThreadCountMultiplier, Seed, SeedSettings, SeedComponent, SourceComponentBounds);
diff --git a/Engine/Plugins/PCG/Source/PCG/Private/Compute/DataInterfaces/PCGCustomKernelDataInterface.h b/Engine/Plugins/PCG/Source/PCG/Private/Compute/DataInterfaces/PCGCustomKernelDataInterface.h
--- a/Engine/Plugins/PCG/Source/PCG/Private/Compute/DataInterfaces/PCGCustomKernelDataInterface.h
+++ b/Engine/Plugins/PCG/Source/PCG/Private/Compute/DataInterfaces/PCGCustomKernelDataInterface.h
@@ -2,10 +2,9 @@
 
 #pragma once
 
-#include "Compute/PCGDataBinding.h"
+#include "Compute/DataInterfaces/PCGComputeDataInterface.h"
 
-#include "ComputeFramework/ComputeDataInterface.h"
-#include "ComputeFramework/ComputeDataProvider.h"
+#include "Compute/PCGDataBinding.h"
 
 #include "PCGCustomKernelDataInterface.generated.h"
 
@@ -42,15 +41,22 @@ public:
 
 /** Compute Framework Data Provider for each custom compute kernel. */
 UCLASS()
-class UPCGCustomComputeKernelDataProvider : public UComputeDataProvider
+class UPCGCustomComputeKernelDataProvider : public UPCGComputeDataProvider
 {
 	GENERATED_BODY()
 
 public:
+	//~ Begin UPCGComputeDataProvider Interface
+	virtual bool PrepareForExecute_GameThread(UPCGDataBinding* InBinding) override;
+	//~ End UPCGComputeDataProvider Interface
+
 	//~ Begin UComputeDataProvider Interface
 	FComputeDataProviderRenderProxy* GetRenderProxy() override;
 	//~ End UComputeDataProvider Interface
 
+	UPROPERTY()
+	TObjectPtr<const UPCGComputeKernel> Kernel = nullptr;
+
 	int32 ThreadCount = -1;
 	uint32 ThreadCountMultiplier = 0;
 
diff --git a/Engine/Plugins/PCG/Source/PCG/Private/Compute/DataInterfaces/PCGDataCollectionDataInterface.cpp b/Engine/Plugins/PCG/Source/PCG/Private/Compute/DataInterfaces/PCGDataCollectionDataInterface.cpp
--- a/Engine/Plugins/PCG/Source/PCG/Private/Compute/DataInterfaces/PCGDataCollectionDataInterface.cpp
+++ b/Engine/Plugins/PCG/Source/PCG/Private/Compute/DataInterfaces/PCGDataCollectionDataInterface.cpp
@@ -446,20 +446,11 @@ UComputeDataProvider* UPCGDataCollectionDataInterface::CreateDataProvider(TObjec
 	check(ProducerSettings);
 
 	UPCGDataBinding* Binding = CastChecked<UPCGDataBinding>(InBinding);
-	check(Binding->Graph);
-
-	FPCGDataCollectionDesc PinDataDesc;
-	ensure(Binding->GetKernelPinDataDesc_GameThread(GraphBindingIndex, PinDataDesc));
-
-	if (PCGComputeHelpers::IsBufferSizeTooLarge(PinDataDesc.ComputePackedSizeBytes()))
-	{
-		return nullptr;
-	}
 
 	UPCGDataCollectionDataProvider* Provider = NewObject<UPCGDataCollectionDataProvider>();
 	Provider->Binding = Binding;
 	Provider->ProducerSettings = ProducerSettings;
-	Provider->PinDesc = MoveTemp(PinDataDesc);
+	Provider->GraphBindingIndex = GraphBindingIndex;
 
 	Provider->ExportMode = bRequiresExport ? EPCGExportMode::ComputeGraphOutput : EPCGExportMode::NoExport;
 
@@ -489,6 +480,26 @@ UComputeDataProvider* UPCGDataCollectionDataInterface::CreateDataProvider(TObjec
 	return Provider;
 }
 
+bool UPCGDataCollectionDataProvider::PrepareForExecute_GameThread(UPCGDataBinding* InBinding)
+{
+	TRACE_CPUPROFILER_EVENT_SCOPE(UPCGDataCollectionDataProvider::PrepareForExecute_GameThread);
+	check(InBinding);
+
+	if (!Super::PrepareForExecute_GameThread(InBinding))
+	{
+		return false;
+	}
+
+	ensure(Binding->GetKernelPinDataDesc_GameThread(GraphBindingIndex, PinDesc));
+
+	if (PCGComputeHelpers::IsBufferSizeTooLarge(PinDesc.ComputePackedSizeBytes()))
+	{
+		PinDesc = {};
+	}
+
+	return true;
+}
+
 FComputeDataProviderRenderProxy* UPCGDataCollectionDataProvider::GetRenderProxy()
 {
 	TRACE_CPUPROFILER_EVENT_SCOPE(UPCGDataCollectionDataProvider::GetRenderProxy);
@@ -658,7 +669,10 @@ void FPCGDataCollectionDataProviderProxy::AllocateResources(FRDGBuilder& GraphBu
 	// Initialize with an empty data collection. The kernel may not run, for example if indirect dispatch args end up being 0. This ensures
 	// there is something meaningful to readback.
 	TArray<uint32> PackedDataCollection;
-	PinDesc.WriteHeader(PackedDataCollection);
+	{
+		TRACE_CPUPROFILER_EVENT_SCOPE(FPCGDataCollectionDataProviderProxy::AllocateResources::WriteHeader);
+		PinDesc.WriteHeader(PackedDataCollection);
+	}
 
 	GraphBuilder.QueueBufferUpload(Buffer, PackedDataCollection.GetData(), PackedDataCollection.Num() * PackedDataCollection.GetTypeSize(), ERDGInitialDataFlags::None);
 
diff --git a/Engine/Plugins/PCG/Source/PCG/Private/Compute/DataInterfaces/PCGDataCollectionDataInterface.h b/Engine/Plugins/PCG/Source/PCG/Private/Compute/DataInterfaces/PCGDataCollectionDataInterface.h
--- a/Engine/Plugins/PCG/Source/PCG/Private/Compute/DataInterfaces/PCGDataCollectionDataInterface.h
+++ b/Engine/Plugins/PCG/Source/PCG/Private/Compute/DataInterfaces/PCGDataCollectionDataInterface.h
@@ -71,6 +71,10 @@ class UPCGDataCollectionDataProvider : public UPCGComputeDataProvider
 	GENERATED_BODY()
 
 public:
+	//~ Begin UPCGComputeDataProvider Interface
+	virtual bool PrepareForExecute_GameThread(UPCGDataBinding* InBinding) override;
+	//~ End UPCGComputeDataProvider Interface
+
 	//~ Begin UComputeDataProvider Interface
 	FComputeDataProviderRenderProxy* GetRenderProxy() override;
 	//~ End UComputeDataProvider Interface
@@ -91,7 +95,8 @@ public:
 
 	TWeakObjectPtr<UPCGDataBinding> Binding;
 
-	FPCGDataCollectionDesc PinDesc;
+	UPROPERTY()
+	int32 GraphBindingIndex = INDEX_NONE;
 
 	UPROPERTY()
 	EPCGExportMode ExportMode = EPCGExportMode::NoExport;
@@ -108,6 +113,9 @@ public:
 	UPROPERTY()
 	FName OutputPinLabelAlias;
 
+protected:
+	FPCGDataCollectionDesc PinDesc;
+
 private:
 	FOnDataExported OnDataExported;
 
diff --git a/Engine/Plugins/PCG/Source/PCG/Private/Compute/DataInterfaces/PCGDataCollectionUploadDataInterface.cpp b/Engine/Plugins/PCG/Source/PCG/Private/Compute/DataInterfaces/PCGDataCollectionUploadDataInterface.cpp
--- a/Engine/Plugins/PCG/Source/PCG/Private/Compute/DataInterfaces/PCGDataCollectionUploadDataInterface.cpp
+++ b/Engine/Plugins/PCG/Source/PCG/Private/Compute/DataInterfaces/PCGDataCollectionUploadDataInterface.cpp
@@ -26,23 +26,29 @@ UComputeDataProvider* UPCGDataCollectionUploadDataInterface::CreateDataProvider(
 
 	UPCGDataProviderDataCollectionUpload* Provider = NewObject<UPCGDataProviderDataCollectionUpload>();
 	Provider->Binding = Binding;
+	Provider->GraphBindingIndex = GraphBindingIndex;
 
 	// Pick the data items from input data collection using any of the compute graph element virtual input pin labels.
 	check(!DownstreamInputPinLabelAliases.IsEmpty());
-
-	ensure(Binding->GetKernelPinDataDesc_GameThread(GraphBindingIndex, Provider->PinDesc));
-
 	Provider->DownstreamInputPinLabels = DownstreamInputPinLabelAliases;
 
-	check(!DownstreamInputPinLabelAliases.IsEmpty());
-	Provider->DataAdaptor = MakeShared<FPCGDataUploadAdaptor>(Binding, Provider->PinDesc, DownstreamInputPinLabelAliases[0]);
-
 	return Provider;
 }
 
-bool UPCGDataProviderDataCollectionUpload::PrepareForExecute_GameThread(const UPCGDataBinding* InBinding)
+bool UPCGDataProviderDataCollectionUpload::PrepareForExecute_GameThread(UPCGDataBinding* InBinding)
 {
-	return !DataAdaptor || DataAdaptor->PrepareData_GameThread(InBinding);
+	TRACE_CPUPROFILER_EVENT_SCOPE(UPCGDataProviderDataCollectionUpload::PrepareForExecute_GameThread);
+	check(InBinding);
+
+	if (!DataAdaptor)
+	{
+		ensure(InBinding->GetKernelPinDataDesc_GameThread(GraphBindingIndex, PinDesc));
+
+		check(!DownstreamInputPinLabels.IsEmpty());
+		DataAdaptor = MakeShared<FPCGDataUploadAdaptor>(InBinding, PinDesc, DownstreamInputPinLabels[0]);
+	}
+
+	return !DataAdaptor || DataAdaptor->PrepareData_GameThread();
 }
 
 FComputeDataProviderRenderProxy* UPCGDataProviderDataCollectionUpload::GetRenderProxy()
diff --git a/Engine/Plugins/PCG/Source/PCG/Private/Compute/DataInterfaces/PCGDataCollectionUploadDataInterface.h b/Engine/Plugins/PCG/Source/PCG/Private/Compute/DataInterfaces/PCGDataCollectionUploadDataInterface.h
--- a/Engine/Plugins/PCG/Source/PCG/Private/Compute/DataInterfaces/PCGDataCollectionUploadDataInterface.h
+++ b/Engine/Plugins/PCG/Source/PCG/Private/Compute/DataInterfaces/PCGDataCollectionUploadDataInterface.h
@@ -34,7 +34,7 @@ public:
 	//~ End UComputeDataProvider Interface
 
 	//~ Begin UPCGComputeDataProvider Interface
-	virtual bool PrepareForExecute_GameThread(const UPCGDataBinding* InBinding) override;
+	virtual bool PrepareForExecute_GameThread(UPCGDataBinding* InBinding) override;
 	//~ End UPCGComputeDataProvider Interface
 
 	TArray<FName> DownstreamInputPinLabels;
diff --git a/Engine/Plugins/PCG/Source/PCG/Private/Compute/DataInterfaces/PCGDataLabelResolverDataInterface.cpp b/Engine/Plugins/PCG/Source/PCG/Private/Compute/DataInterfaces/PCGDataLabelResolverDataInterface.cpp
--- a/Engine/Plugins/PCG/Source/PCG/Private/Compute/DataInterfaces/PCGDataLabelResolverDataInterface.cpp
+++ b/Engine/Plugins/PCG/Source/PCG/Private/Compute/DataInterfaces/PCGDataLabelResolverDataInterface.cpp
@@ -72,23 +72,31 @@ UComputeDataProvider* UPCGDataLabelResolverDataInterface::CreateDataProvider(TOb
 {
 	TRACE_CPUPROFILER_EVENT_SCOPE(UPCGDataLabelResolverDataInterface::CreateDataProvider);
 
+	UPCGDataLabelResolverDataProvider* DataProvider = NewObject<UPCGDataLabelResolverDataProvider>();
+	DataProvider->Kernel = Kernel;
+	DataProvider->PinLabel = PinLabel;
+	DataProvider->bIsInput = bIsInput;
+
+	return DataProvider;
+}
+
+bool UPCGDataLabelResolverDataProvider::PrepareForExecute_GameThread(UPCGDataBinding* InBinding)
+{
+	TRACE_CPUPROFILER_EVENT_SCOPE(UPCGDataLabelResolverDataProvider::PrepareForExecute_GameThread);
+	check(InBinding);
 	check(Kernel);
 
-	UPCGDataBinding* Binding = CastChecked<UPCGDataBinding>(InBinding);
-	const UPCGComputeGraph* ComputeGraph = Binding->Graph.Get();
+	const UPCGComputeGraph* ComputeGraph = InBinding->Graph.Get();
 	check(ComputeGraph);
 
-	const TArray<FString>& StringTable = Binding->GetStringTable();
+	const TArray<FString>& StringTable = InBinding->GetStringTable();
 
 	FPCGDataCollectionDesc PinDataDesc;
-	ensure(Binding->GetKernelPinDataDesc_GameThread(Kernel, PinLabel, bIsInput, PinDataDesc));
-
-	UPCGDataLabelResolverDataProvider* Provider = NewObject<UPCGDataLabelResolverDataProvider>();
+	ensure(InBinding->GetKernelPinDataDesc_GameThread(Kernel, PinLabel, bIsInput, PinDataDesc));
 
 	const FPCGPinDataLabels* FoundPinDataLabels = ComputeGraph->GetStaticDataLabelsTable().Find(Kernel->GetKernelIndex());
 	if (const FPCGDataLabels* FoundLabels = FoundPinDataLabels ? FoundPinDataLabels->PinToDataLabels.Find(PinLabel) : nullptr)
 	{
-		TArray<int32> DataIdToDataIndexMap;
 		DataIdToDataIndexMap.SetNumUninitialized(FoundLabels->Labels.Num());
 
 		for (int32 LabelIndex = 0; LabelIndex < FoundLabels->Labels.Num(); ++LabelIndex)
@@ -113,7 +121,7 @@ UComputeDataProvider* UPCGDataLabelResolverDataInterface::CreateDataProvider(TOb
 
 			if (DataIdToDataIndexMap[LabelIndex] == INDEX_NONE)
 			{
-				if (TSharedPtr<FPCGContextHandle> Context = Binding->ContextHandle.Pin())
+				if (TSharedPtr<FPCGContextHandle> Context = InBinding->ContextHandle.Pin())
 				{
 					PCG_KERNEL_VALIDATION_ERR(Context->GetContext(), Kernel->GetSettings(), /*bQuiet=*/false, FText::Format(
 						NSLOCTEXT("PCGDataLabelResolver", "LabelDoesNotExist", "Data label '{0}' was not present in the tags on the incoming data."),
@@ -121,11 +129,9 @@ UComputeDataProvider* UPCGDataLabelResolverDataInterface::CreateDataProvider(TOb
 				}
 			}
 		}
-
-		Provider->DataIdToDataIndexMap = MoveTemp(DataIdToDataIndexMap);
 	}
 
-	return Provider;
+	return true;
 }
 
 FComputeDataProviderRenderProxy* UPCGDataLabelResolverDataProvider::GetRenderProxy()
diff --git a/Engine/Plugins/PCG/Source/PCG/Private/Compute/DataInterfaces/PCGDataLabelResolverDataInterface.h b/Engine/Plugins/PCG/Source/PCG/Private/Compute/DataInterfaces/PCGDataLabelResolverDataInterface.h
--- a/Engine/Plugins/PCG/Source/PCG/Private/Compute/DataInterfaces/PCGDataLabelResolverDataInterface.h
+++ b/Engine/Plugins/PCG/Source/PCG/Private/Compute/DataInterfaces/PCGDataLabelResolverDataInterface.h
@@ -2,10 +2,10 @@
 
 #pragma once
 
-#include "Compute/PCGDataBinding.h"
-
+#include "Compute/DataInterfaces/PCGComputeDataInterface.h"
 #include "ComputeFramework/ComputeDataInterface.h"
-#include "ComputeFramework/ComputeDataProvider.h"
+
+#include "Compute/PCGDataBinding.h"
 
 #include "PCGDataLabelResolverDataInterface.generated.h"
 
@@ -42,15 +42,31 @@ public:
 };
 
 UCLASS()
-class UPCGDataLabelResolverDataProvider : public UComputeDataProvider
+class UPCGDataLabelResolverDataProvider : public UPCGComputeDataProvider
 {
 	GENERATED_BODY()
 
 public:
+	//~ Begin UPCGComputeDataProvider Interface
+	virtual bool PrepareForExecute_GameThread(UPCGDataBinding* InBinding) override;
+	//~ End UPCGComputeDataProvider Interface
+
 	//~ Begin UComputeDataProvider Interface
 	FComputeDataProviderRenderProxy* GetRenderProxy() override;
 	//~ End UComputeDataProvider Interface
 
+public:
+	UPROPERTY()
+	TObjectPtr<const UPCGComputeKernel> Kernel;
+
+	/** Label of the pin being resolved. */
+	UPROPERTY()
+	FName PinLabel;
+
+	/** Whether the pin is an input or output pin. */
+	UPROPERTY()
+	bool bIsInput = false;
+
 	/** Map from data ID to data index for accessing the input data collection. */
 	TArray<int32> DataIdToDataIndexMap;
 };
diff --git a/Engine/Plugins/PCG/Source/PCG/Private/Compute/Elements/PCGComputeGraphElement.cpp b/Engine/Plugins/PCG/Source/PCG/Private/Compute/Elements/PCGComputeGraphElement.cpp
--- a/Engine/Plugins/PCG/Source/PCG/Private/Compute/Elements/PCGComputeGraphElement.cpp
+++ b/Engine/Plugins/PCG/Source/PCG/Private/Compute/Elements/PCGComputeGraphElement.cpp
@@ -289,6 +289,32 @@ bool FPCGComputeGraphElement::ExecuteInternal(FPCGContext* InContext) const
 		Context->bPrimitiveProxiesValidated = true;
 	}
 
+	// Allow data providers to do early readbacks, which some may use to download analysis results which feed into final data descriptions.
+	if (!Context->bDataProvidersReadbacksComplete)
+	{
+		bool bAllReady = true;
+
+		for (UComputeDataProvider* ComputeDataProvider : Context->ComputeGraphInstance.GetDataProviders())
+		{
+			if (UPCGComputeDataProvider* DataProvider = Cast<UPCGComputeDataProvider>(ComputeDataProvider))
+			{
+				if (!DataProvider->PerformPreExecuteReadbacks_GameThread(Context->DataBinding.Get()))
+				{
+					bAllReady = false;
+				}
+			}
+		}
+
+		if (!bAllReady)
+		{
+			SleepUntilNextFrame(Context);
+			return false;
+		}
+
+		Context->bDataProvidersReadbacksComplete = true;
+	}
+
+	// Prepare data providers for execution.
 	if (!Context->bDataProvidersReady)
 	{
 		bool bAllReady = true;
@@ -309,6 +335,8 @@ bool FPCGComputeGraphElement::ExecuteInternal(FPCGContext* InContext) const
 			SleepUntilNextFrame(Context);
 			return false;
 		}
+
+		Context->bDataProvidersReady = true;
 	}
 
 	// 1. Prepare render resources. In editor, this will trigger shader compile if not compiled already.
diff --git a/Engine/Plugins/PCG/Source/PCG/Private/Compute/Elements/PCGComputeGraphElement.h b/Engine/Plugins/PCG/Source/PCG/Private/Compute/Elements/PCGComputeGraphElement.h
--- a/Engine/Plugins/PCG/Source/PCG/Private/Compute/Elements/PCGComputeGraphElement.h
+++ b/Engine/Plugins/PCG/Source/PCG/Private/Compute/Elements/PCGComputeGraphElement.h
@@ -40,6 +40,8 @@ public:
 	/** Whether all required primitive proxies have been created. */
 	bool bPrimitiveProxiesValidated = false;
 
+	bool bDataProvidersReadbacksComplete = false;
+
 	bool bDataProvidersReady = false;
 
 	/** Data providers with buffers that are passed to downstream tasks. The buffer will be created on the render thread, and then passed
diff --git a/Engine/Plugins/PCG/Source/PCG/Private/Compute/PCGComputeKernel.h b/Engine/Plugins/PCG/Source/PCG/Private/Compute/PCGComputeKernel.h
--- a/Engine/Plugins/PCG/Source/PCG/Private/Compute/PCGComputeKernel.h
+++ b/Engine/Plugins/PCG/Source/PCG/Private/Compute/PCGComputeKernel.h
@@ -46,6 +46,9 @@ public:
 
 	/** Create additional output data interfaces to marshal any required output data. */
 	virtual void CreateAdditionalOutputDataInterfaces(FPCGGPUCompilationContext& InOutContext, UObject* InObjectOuter, TArray<TObjectPtr<UComputeDataInterface>>& OutDataInterfaces) const {}
+
+	/** Whether compute graphs should be split at the output of this node. Useful for reading back runtime statistics or diagnostic info before continuing. */
+	virtual bool SplitGraphAtOutput() const { return false; }
 #endif
 
 	/** Compute a description of data that will be output from pin InOutputPinLabel. */
diff --git a/Engine/Plugins/PCG/Source/PCG/Private/Compute/PCGDataUploadAdaptor.cpp b/Engine/Plugins/PCG/Source/PCG/Private/Compute/PCGDataUploadAdaptor.cpp
--- a/Engine/Plugins/PCG/Source/PCG/Private/Compute/PCGDataUploadAdaptor.cpp
+++ b/Engine/Plugins/PCG/Source/PCG/Private/Compute/PCGDataUploadAdaptor.cpp
@@ -33,11 +33,22 @@ FPCGDataUploadAdaptor::FPCGDataUploadAdaptor(UPCGDataBinding* InDataBinding, con
 	InputPinLabel = InInputPinLabel;
 	DataBinding = InDataBinding;
 
+	int NumDataForThisPin = 0;
+
 	// Are all our data items coming direct from GPU - and importantly are a complete, single data collection?
 	bool bMultipleBuffersEncountered = false;
 	TSharedPtr<const FPCGProxyForGPUDataCollection> UniqueDataCollection;
 	for (int InputDataIndex = 0; InputDataIndex < DataBinding->InputDataCollection.TaggedData.Num(); ++InputDataIndex)
 	{
+		const FName DataPin = DataBinding->InputDataCollection.TaggedData[InputDataIndex].Pin;
+
+		if (InInputPinLabel != DataPin)
+		{
+			continue;
+		}
+
+		++NumDataForThisPin;
+
 		const UPCGProxyForGPUData* DataGPU = Cast<UPCGProxyForGPUData>(DataBinding->InputDataCollection.TaggedData[InputDataIndex].Data);
 		if (!DataGPU)
 		{
@@ -69,7 +80,7 @@ FPCGDataUploadAdaptor::FPCGDataUploadAdaptor(UPCGDataBinding* InDataBinding, con
 		}
 	}
 
-	if (UniqueDataCollection && UniqueDataCollection->GetDescription().DataDescs.Num() != DataBinding->InputDataCollection.TaggedData.Num())
+	if (UniqueDataCollection && UniqueDataCollection->GetDescription().DataDescs.Num() != NumDataForThisPin)
 	{
 		// We're not using all the data items from the buffer, no reuse for now.
 		// TODO: Support indirection of data items.
@@ -86,16 +97,21 @@ FPCGDataUploadAdaptor::FPCGDataUploadAdaptor(UPCGDataBinding* InDataBinding, con
 }
 
 /** Do any preparation work such as data readbacks. Returns true when preparation is complete. */
-bool FPCGDataUploadAdaptor::PrepareData_GameThread(const UPCGDataBinding* InBinding)
+bool FPCGDataUploadAdaptor::PrepareData_GameThread()
 {
-	check(InBinding);
+	UPCGDataBinding* Binding = DataBinding.Get();
+	if (!ensure(Binding))
+	{
+		return true;
+	}
+
 	check(IsInGameThread());
 
 	if (ExternalBufferForReuse)
 	{
-		SourceToTargetAttributeId.Reserve(InBinding->GetAttributeIdRemap().Num());
+		SourceToTargetAttributeId.Reserve(Binding->GetAttributeIdRemap().Num());
 
-		Algo::Transform(InBinding->GetAttributeIdRemap(), SourceToTargetAttributeId, [](const TPair<UPCGDataBinding::FSourceBufferAndAttributeIndex, int32>& Remapping)
+		Algo::Transform(Binding->GetAttributeIdRemap(), SourceToTargetAttributeId, [](const TPair<UPCGDataBinding::FSourceBufferAndAttributeIndex, int32>& Remapping)
 		{
 			const int32 TargetAttributeId = PCGComputeHelpers::GetAttributeIdFromMetadataAttributeIndex(Remapping.Value);
 			const int32 SourceAttributeId = PCGComputeHelpers::GetAttributeIdFromMetadataAttributeIndex(Remapping.Key.Value);
@@ -108,12 +124,6 @@ bool FPCGDataUploadAdaptor::PrepareData_GameThread(const UPCGDataBinding* InBind
 	}
 	else
 	{
-		UPCGDataBinding* Binding = DataBinding.Get();
-		if (!Binding)
-		{
-			return true;
-		}
-
 		const uint64 BufferSize = TargetDataCollectionDesc.ComputePackedSizeBytes();
 		if (PCGComputeHelpers::IsBufferSizeTooLarge(BufferSize))
 		{
@@ -124,6 +134,11 @@ bool FPCGDataUploadAdaptor::PrepareData_GameThread(const UPCGDataBinding* InBind
 
 		for (FPCGTaggedData& TaggedData : Binding->InputDataCollection.TaggedData)
 		{
+			if (TaggedData.Pin != InputPinLabel)
+			{
+				continue;
+			}
+
 			if (const UPCGProxyForGPUData* Proxy = Cast<UPCGProxyForGPUData>(TaggedData.Data))
 			{
 				UPCGProxyForGPUData::FReadbackResult Result = Proxy->GetCPUData(/*InContext=*/nullptr);
diff --git a/Engine/Plugins/PCG/Source/PCG/Private/Compute/PCGDataUploadAdaptor.h b/Engine/Plugins/PCG/Source/PCG/Private/Compute/PCGDataUploadAdaptor.h
--- a/Engine/Plugins/PCG/Source/PCG/Private/Compute/PCGDataUploadAdaptor.h
+++ b/Engine/Plugins/PCG/Source/PCG/Private/Compute/PCGDataUploadAdaptor.h
@@ -18,7 +18,7 @@ public:
 	FPCGDataUploadAdaptor(UPCGDataBinding* InDataBinding, const FPCGDataCollectionDesc& InTargetDataCollectionDesc, FName InInputPinLabel);
 
 	/** Do any preparation work such as data readbacks. Returns true when preparation is complete. */
-	bool PrepareData_GameThread(const UPCGDataBinding* InBinding);
+	bool PrepareData_GameThread();
 
 	FRDGBufferSRVRef GetAttributeRemapBufferSRV(FRDGBuilder& InGraphBuilder);
 
diff --git a/Engine/Plugins/PCG/Source/PCG/Private/Elements/Metadata/PCGMetadataPartitionKernel.cpp b/Engine/Plugins/PCG/Source/PCG/Private/Elements/Metadata/PCGMetadataPartitionKernel.cpp
--- a/Engine/Plugins/PCG/Source/PCG/Private/Elements/Metadata/PCGMetadataPartitionKernel.cpp
+++ b/Engine/Plugins/PCG/Source/PCG/Private/Elements/Metadata/PCGMetadataPartitionKernel.cpp
@@ -71,7 +71,7 @@ bool UPCGMetadataPartitionKernel::IsKernelValid(FPCGContext* InContext, bool bQu
 
 		if (!bAttributePresent)
 		{
-			PCG_KERNEL_VALIDATION_ERR(InContext, Settings, bQuiet, LOCTEXT("MustProvideOnePartitionAttribute", "GPU implementation currently only supports a single partition attribute."));
+			PCG_KERNEL_VALIDATION_ERR(InContext, Settings, bQuiet, LOCTEXT("PartitionAttributeMissing", "Partition attribute was not present in GPU data."));
 			return false;
 		}
 
diff --git a/Engine/Plugins/PCG/Source/PCG/Private/Graph/PCGGraphCompilerGPU.cpp b/Engine/Plugins/PCG/Source/PCG/Private/Graph/PCGGraphCompilerGPU.cpp
--- a/Engine/Plugins/PCG/Source/PCG/Private/Graph/PCGGraphCompilerGPU.cpp
+++ b/Engine/Plugins/PCG/Source/PCG/Private/Graph/PCGGraphCompilerGPU.cpp
@@ -43,12 +43,13 @@ void FPCGGraphCompilerGPU::LabelConnectedGPUNodeIslands(
 	const TArray<FPCGGraphTask>& InCompiledTasks,
 	const TSet<FPCGTaskId>& InGPUCompatibleTaskIds,
 	const FTaskToSuccessors& InTaskSuccessors,
+	const TMap<FPCGTaskId, UPCGComputeKernel*>& InTaskIdToKernel,
 	TArray<uint32>& OutIslandIDs)
 {
 	OutIslandIDs.SetNumZeroed(InCompiledTasks.Num());
 
 	// Traverses task inputs and successors and assigns the given island ID to each one. Memoized via output OutIslandIDs.
-	auto FloodFillIslandID = [&InCompiledTasks, &InTaskSuccessors, &InGPUCompatibleTaskIds, &OutIslandIDs](FPCGTaskId InTaskId, int InIslandID, FPCGTaskId InTraversedFromTaskId, auto&& RecursiveCall) -> void
+	auto FloodFillIslandID = [&InCompiledTasks, &InTaskSuccessors, &InGPUCompatibleTaskIds, &OutIslandIDs, &InTaskIdToKernel](FPCGTaskId InTaskId, int InIslandID, FPCGTaskId InTraversedFromTaskId, auto&& RecursiveCall) -> void
 	{
 		check(InTaskId != InTraversedFromTaskId);
 
@@ -67,8 +68,8 @@ void FPCGGraphCompilerGPU::LabelConnectedGPUNodeIslands(
 			}
 		}
 
-		const UPCGSettings* Settings = InCompiledTasks[InTaskId].Node ? InCompiledTasks[InTaskId].Node->GetSettings() : nullptr;
-		if (!Settings || !Settings->SplitGraphAtOutput())
+		UPCGComputeKernel* const* FoundKernel = InTaskIdToKernel.Find(InTaskId);
+		if (!FoundKernel || !(*FoundKernel)->SplitGraphAtOutput())
 		{
 			if (const TArray<FPCGTaskId>* Successors = InTaskSuccessors.Find(InTaskId))
 			{
@@ -103,13 +104,14 @@ void FPCGGraphCompilerGPU::CollectGPUNodeSubsets(
 	const TArray<FPCGGraphTask>& InCompiledTasks,
 	const FTaskToSuccessors& InTaskSuccessors,
 	const TSet<FPCGTaskId>& InGPUCompatibleTaskIds,
+	const TMap<FPCGTaskId, UPCGComputeKernel*>& InTaskIdToKernel,
 	TArray<TSet<FPCGTaskId>>& OutNodeSubsetsToConvertToCFGraph)
 {
 	TRACE_CPUPROFILER_EVENT_SCOPE(FPCGGraphCompilerGPU::CollectGPUNodeSubsets);
 
 	// Identifies connected sets of GPU nodes, giving each a non-zero ID value.
 	TArray<uint32> ConnectedGPUNodeIslandIDs;
-	LabelConnectedGPUNodeIslands(InCompiledTasks, InGPUCompatibleTaskIds, InTaskSuccessors, ConnectedGPUNodeIslandIDs);
+	LabelConnectedGPUNodeIslands(InCompiledTasks, InGPUCompatibleTaskIds, InTaskSuccessors, InTaskIdToKernel, ConnectedGPUNodeIslandIDs);
 
 	// For every CPU node that has one or more downstream GPU node connected, check for a GPU -> CPU -> GPU pattern where data flows
 	// from a GPU node island to the CPU and then back to the same island. For such cases, we traverse the entire tree of GPU nodes that
@@ -1615,7 +1617,7 @@ void FPCGGraphCompilerGPU::CreateGPUNodes(FPCGGraphCompiler& InOutCompiler, UPCG
 	}
 
 	TArray<TSet<FPCGTaskId>> NodeSubsetsToConvertToCFGraph;
-	CollectGPUNodeSubsets(InOutCompiledTasks, TaskSuccessors, GPUCompatibleTaskIds, NodeSubsetsToConvertToCFGraph);
+	CollectGPUNodeSubsets(InOutCompiledTasks, TaskSuccessors, GPUCompatibleTaskIds, TaskIdToKernel, NodeSubsetsToConvertToCFGraph);
 
 	const uint32 NumComputeGraphs = NodeSubsetsToConvertToCFGraph.Num();
 
diff --git a/Engine/Plugins/PCG/Source/PCG/Private/Graph/PCGGraphCompilerGPU.h b/Engine/Plugins/PCG/Source/PCG/Private/Graph/PCGGraphCompilerGPU.h
--- a/Engine/Plugins/PCG/Source/PCG/Private/Graph/PCGGraphCompilerGPU.h
+++ b/Engine/Plugins/PCG/Source/PCG/Private/Graph/PCGGraphCompilerGPU.h
@@ -31,6 +31,7 @@ private:
 		const TArray<FPCGGraphTask>& InCompiledTasks,
 		const TSet<FPCGTaskId>& InGPUCompatibleTaskIds,
 		const FTaskToSuccessors& InTaskSuccessors,
+		const TMap<FPCGTaskId, UPCGComputeKernel*>& InTaskIdToKernel,
 		TArray<uint32>& OutIslandIDs);
 
 	/** Outputs sets of task IDs, where each set is GPU nodes that can be compiled into a compute graph and dispatched together. */
@@ -38,6 +39,7 @@ private:
 		const TArray<FPCGGraphTask>& InCompiledTasks,
 		const FTaskToSuccessors& InTaskSuccessors,
 		const TSet<FPCGTaskId>& InGPUCompatibleTaskIds,
+		const TMap<FPCGTaskId, UPCGComputeKernel*>& InTaskIdToKernel,
 		TArray<TSet<FPCGTaskId>>& OutNodeSubsetsToConvertToCFGraph);
 	
 	/** Create a separate task for each GPU kernel emitted by each GPU node. */
diff --git a/Engine/Plugins/PCG/Source/PCG/Public/Elements/Metadata/PCGMetadataPartition.h b/Engine/Plugins/PCG/Source/PCG/Public/Elements/Metadata/PCGMetadataPartition.h
--- a/Engine/Plugins/PCG/Source/PCG/Public/Elements/Metadata/PCGMetadataPartition.h
+++ b/Engine/Plugins/PCG/Source/PCG/Public/Elements/Metadata/PCGMetadataPartition.h
@@ -19,7 +19,6 @@ public:
 
 	//~Begin UPCGSettings interface
 #if WITH_EDITOR
-	virtual bool SplitGraphAtOutput() const override { return true; }
 	virtual bool DisplayExecuteOnGPUSetting() const override { return true; }
 	virtual void CreateKernels(FPCGGPUCompilationContext& InOutContext, UObject* InObjectOuter, TArray<UPCGComputeKernel*>& OutKernels, TArray<FPCGKernelEdge>& OutEdges) const override;
 	virtual void ApplyDeprecation(UPCGNode* InOutNode) override;
diff --git a/Engine/Plugins/PCG/Source/PCG/Public/Elements/Metadata/PCGMetadataPartitionKernel.h b/Engine/Plugins/PCG/Source/PCG/Public/Elements/Metadata/PCGMetadataPartitionKernel.h
--- a/Engine/Plugins/PCG/Source/PCG/Public/Elements/Metadata/PCGMetadataPartitionKernel.h
+++ b/Engine/Plugins/PCG/Source/PCG/Public/Elements/Metadata/PCGMetadataPartitionKernel.h
@@ -23,6 +23,7 @@ public:
 #if WITH_EDITOR
 	virtual FString GetCookedSource(const UPCGComputeGraph* InComputeGraph) const override;
 	void CreateAdditionalInputDataInterfaces(FPCGGPUCompilationContext& InOutContext, UObject* InObjectOuter, TArray<TObjectPtr<UComputeDataInterface>>& OutDataInterfaces) const override;
+	virtual bool SplitGraphAtOutput() const override { return true; }
 #endif
 	virtual void GetKernelAttributeKeys(TArray<FPCGKernelAttributeKey>& OutKeys) const override;
 	virtual bool DoesOutputPinRequireElementCounters(FName InOutputPinLabel) const override { return true; }
diff --git a/Engine/Plugins/PCG/Source/PCG/Public/PCGSettings.h b/Engine/Plugins/PCG/Source/PCG/Public/PCGSettings.h
--- a/Engine/Plugins/PCG/Source/PCG/Public/PCGSettings.h
+++ b/Engine/Plugins/PCG/Source/PCG/Public/PCGSettings.h
@@ -516,12 +516,6 @@ public:
 	*/
 	virtual void CreateKernels(FPCGGPUCompilationContext& InOutContext, UObject* InObjectOuter, TArray<UPCGComputeKernel*>& OutKernels, TArray<FPCGKernelEdge>& OutEdges) const {};
 
-	// TODO - should also move to kernel, but needs flow to be shuffled in the GPU graph compiler.
-	/** [EXPERIMENTAL] Whether compute graphs should be split at the output of this node. Useful for reading back runtime statistics or diagnostic info before continuing.
-	* Note: This API function is subject to change without deprecation.
-	*/
-	virtual bool SplitGraphAtOutput() const { return false; }
-
 	/** [EXPERIMENTAL] Whether to display GPU execution option in node settings UI. 
 	* Note: This API function is subject to change without deprecation.
 	*/
@@ -570,30 +564,30 @@ private:
 #endif
 
 public:
-	UE_DEPRECATED(5.6, "All GPU logic now resides separate kernel objects, see implementations of UPCGSettings::CreateKernel().")
+	UE_DEPRECATED(5.6, "All GPU logic now resides in separate kernel objects, see implementations of UPCGSettings::CreateKernel().")
 	virtual bool IsKernelValid(FPCGContext* InContext = nullptr, bool bQuiet = true) const { return false; }
-	UE_DEPRECATED(5.6, "All GPU logic now resides separate kernel objects, see implementations of UPCGSettings::CreateKernel().")
+	UE_DEPRECATED(5.6, "All GPU logic now resides in separate kernel objects, see implementations of UPCGSettings::CreateKernel().")
 	virtual FString GetCookedKernelSource(const class UPCGComputeGraph* InComputeGraph) const { return FString(); }
-	UE_DEPRECATED(5.6, "All GPU logic now resides separate kernel objects, see implementations of UPCGSettings::CreateKernel().")
+	UE_DEPRECATED(5.6, "All GPU logic now resides in separate kernel objects, see implementations of UPCGSettings::CreateKernel().")
 	virtual void GetKernelAttributeKeys(TArray<struct FPCGKernelAttributeKey>& OutKeys) const {}
-	UE_DEPRECATED(5.6, "All GPU logic now resides separate kernel objects, see implementations of UPCGSettings::CreateKernel().")
+	UE_DEPRECATED(5.6, "All GPU logic now resides in separate kernel objects, see implementations of UPCGSettings::CreateKernel().")
 	virtual void AddStaticCreatedStrings(TArray<FString>& InOutStringTable) const {};
-	UE_DEPRECATED(5.6, "All GPU logic now resides separate kernel objects, see implementations of UPCGSettings::CreateKernel().")
+	UE_DEPRECATED(5.6, "All GPU logic now resides in separate kernel objects, see implementations of UPCGSettings::CreateKernel().")
 	virtual int ComputeKernelThreadCount(const class UPCGDataBinding* Binding) const { return 0; };
-	UE_DEPRECATED(5.6, "All GPU logic now resides separate kernel objects, see implementations of UPCGSettings::CreateKernel().")
+	UE_DEPRECATED(5.6, "All GPU logic now resides in separate kernel objects, see implementations of UPCGSettings::CreateKernel().")
 	bool ComputeOutputPinDataDesc(const FName& OutputPinLabel, const class UPCGDataBinding* InBinding, struct FPCGDataCollectionDesc& OutDesc) const { return false; }
-	UE_DEPRECATED(5.6, "All GPU logic now resides separate kernel objects, see implementations of UPCGSettings::CreateKernel().")
+	UE_DEPRECATED(5.6, "All GPU logic now resides in separate kernel objects, see implementations of UPCGSettings::CreateKernel().")
 	virtual bool ComputeOutputPinDataDesc(const UPCGPin* OutputPin, const class UPCGDataBinding* InBinding, struct FPCGDataCollectionDesc& OutDesc) const { return false; }
-	UE_DEPRECATED(5.6, "All GPU logic now resides separate kernel objects, see implementations of UPCGSettings::CreateKernel().")
+	UE_DEPRECATED(5.6, "All GPU logic now resides in separate kernel objects, see implementations of UPCGSettings::CreateKernel().")
 	virtual bool DoesOutputPinRequireElementCounters(const UPCGPin* OutputPin) const { return false; }
 #if WITH_EDITOR
-	UE_DEPRECATED(5.6, "All GPU logic now resides separate kernel objects, see implementations of UPCGSettings::CreateKernel().")
+	UE_DEPRECATED(5.6, "All GPU logic now resides in separate kernel objects, see implementations of UPCGSettings::CreateKernel().")
 	virtual void CreateAdditionalInputDataInterfaces(FPCGGPUCompilationContext& InOutContext, UObject* InObjectOuter, TArray<TObjectPtr<class UComputeDataInterface>>& OutDataInterfaces) const {}
-	UE_DEPRECATED(5.6, "All GPU logic now resides separate kernel objects, see implementations of UPCGSettings::CreateKernel().")
+	UE_DEPRECATED(5.6, "All GPU logic now resides in separate kernel objects, see implementations of UPCGSettings::CreateKernel().")
 	virtual void CreateAdditionalOutputDataInterfaces(FPCGGPUCompilationContext& InOutContext, UObject* InObjectOuter, TArray<TObjectPtr<class UComputeDataInterface>>& OutDataInterfaces) const {}
-	UE_DEPRECATED(5.6, "All GPU logic now resides separate kernel objects, see implementations of UPCGSettings::CreateKernel().")
+	UE_DEPRECATED(5.6, "All GPU logic now resides in separate kernel objects, see implementations of UPCGSettings::CreateKernel().")
 	virtual void CreateAdditionalOutputDataInterfaces(TArray<TObjectPtr<class UComputeDataInterface>>& OutDataInterfaces) const {}
-	UE_DEPRECATED(5.6, "All GPU logic now resides separate kernel objects, see implementations of UPCGSettings::CreateKernel().")
+	UE_DEPRECATED(5.6, "All GPU logic now resides in separate kernel objects, see implementations of UPCGSettings::CreateKernel().")
 	virtual void GetDataLabels(FName InPinLabel, TArray<FString>& OutDataLabels) const {}
 #endif
 };
