[PCG] ALV: Use strings for data names to fix check() in FName code when string is too long.

#rnx

diff --git a/Engine/Plugins/PCG/Source/PCGEditor/Private/SPCGEditorGraphAttributeListView.cpp b/Engine/Plugins/PCG/Source/PCGEditor/Private/SPCGEditorGraphAttributeListView.cpp
--- a/Engine/Plugins/PCG/Source/PCGEditor/Private/SPCGEditorGraphAttributeListView.cpp
+++ b/Engine/Plugins/PCG/Source/PCGEditor/Private/SPCGEditorGraphAttributeListView.cpp
@@ -335,7 +335,7 @@ void SPCGEditorGraphAttributeListView::Construct(const FArguments& InArgs, TShar
 			.Text(this, &SPCGEditorGraphAttributeListView::OnGenerateSelectedPinText)
 		];
 
-	SAssignNew(DataComboBox, SComboBox<TSharedPtr<FName>>)
+	SAssignNew(DataComboBox, SComboBox<TSharedPtr<FString>>)
 		.OptionsSource(&DataComboBoxItems)
 		.OnGenerateWidget(this, &SPCGEditorGraphAttributeListView::OnGenerateDataWidget)
 		.OnSelectionChanged(this, &SPCGEditorGraphAttributeListView::OnSelectionChanged)
@@ -850,6 +850,7 @@ void SPCGEditorGraphAttributeListView::RefreshDataComboBox()
 	for(int32 TaggedDataIndex = 0; TaggedDataIndex < InspectionData->TaggedData.Num(); ++TaggedDataIndex)
 	{
 		const FPCGTaggedData& TaggedData = InspectionData->TaggedData[TaggedDataIndex];
+
 		FString ItemName = FString::Format(
 			TEXT("[{0}] {1}"),
 			{ FText::AsNumber(TaggedDataIndex).ToString(), (TaggedData.Data ? TaggedData.Data->GetClass()->GetDisplayNameText().ToString(): TEXT("No Data")) });
@@ -859,7 +860,7 @@ void SPCGEditorGraphAttributeListView::RefreshDataComboBox()
 			ItemName.Append(FString::Format(TEXT(": ({0})"), { FString::Join(TaggedData.Tags, TEXT(", ")) }));
 		}
 
-		DataComboBoxItems.Add(MakeShared<FName>(ItemName));
+		DataComboBoxItems.Add(MakeShared<FString>(ItemName));
 	}
 
 	if (DataComboBoxItems.Num() > 0)
@@ -1068,12 +1069,12 @@ TSharedRef<SWidget> SPCGEditorGraphAttributeListView::OnGeneratePinWidget(TShare
 	return SNew(STextBlock).Text(FText::FromName(InItem.IsValid() ? InItem->Name : NAME_None));
 }
 
-TSharedRef<SWidget> SPCGEditorGraphAttributeListView::OnGenerateDataWidget(TSharedPtr<FName> InItem) const
+TSharedRef<SWidget> SPCGEditorGraphAttributeListView::OnGenerateDataWidget(TSharedPtr<FString> InItem) const
 {
-	return SNew(STextBlock).Text(FText::FromName(InItem.IsValid() ? *InItem : NAME_None));
+	return SNew(STextBlock).Text(InItem.IsValid() ? FText::FromString(*InItem) : FText());
 }
 
-void SPCGEditorGraphAttributeListView::OnSelectionChanged(TSharedPtr<FName> Item, ESelectInfo::Type SelectInfo)
+void SPCGEditorGraphAttributeListView::OnSelectionChanged(TSharedPtr<FString> Item, ESelectInfo::Type SelectInfo)
 {
 	if (SelectInfo != ESelectInfo::Direct)
 	{
@@ -1083,9 +1084,9 @@ void SPCGEditorGraphAttributeListView::OnSelectionChanged(TSharedPtr<FName> Item
 
 FText SPCGEditorGraphAttributeListView::OnGenerateSelectedDataText() const
 {
-	if (const TSharedPtr<FName> SelectedDataName = DataComboBox->GetSelectedItem())
+	if (const TSharedPtr<FString> SelectedDataName = DataComboBox->GetSelectedItem())
 	{
-		return FText::FromName(*SelectedDataName);
+		return FText::FromString(*SelectedDataName);
 	}
 	else
 	{
@@ -1096,7 +1097,7 @@ FText SPCGEditorGraphAttributeListView::OnGenerateSelectedDataText() const
 int32 SPCGEditorGraphAttributeListView::GetSelectedDataIndex() const
 {
 	int32 Index = INDEX_NONE;
-	if (const TSharedPtr<FName> SelectedItem = DataComboBox->GetSelectedItem())
+	if (const TSharedPtr<FString> SelectedItem = DataComboBox->GetSelectedItem())
 	{
 		DataComboBoxItems.Find(SelectedItem, Index);
 	}
diff --git a/Engine/Plugins/PCG/Source/PCGEditor/Private/SPCGEditorGraphAttributeListView.h b/Engine/Plugins/PCG/Source/PCGEditor/Private/SPCGEditorGraphAttributeListView.h
--- a/Engine/Plugins/PCG/Source/PCGEditor/Private/SPCGEditorGraphAttributeListView.h
+++ b/Engine/Plugins/PCG/Source/PCGEditor/Private/SPCGEditorGraphAttributeListView.h
@@ -197,8 +197,8 @@ private:
 	const FSlateBrush* GetFilterBadgeIcon() const;
 	TSharedRef<SWidget> OnGenerateFilterMenu();
 	TSharedRef<SWidget> OnGenerateAdditionalOperationsMenu();
-	TSharedRef<SWidget> OnGenerateDataWidget(TSharedPtr<FName> InItem) const;
-	void OnSelectionChanged(TSharedPtr<FName> Item, ESelectInfo::Type SelectInfo);
+	TSharedRef<SWidget> OnGenerateDataWidget(TSharedPtr<FString> InItem) const;
+	void OnSelectionChanged(TSharedPtr<FString> Item, ESelectInfo::Type SelectInfo);
 	FText OnGenerateSelectedDataText() const;
 	int32 GetSelectedDataIndex() const;
 
@@ -265,8 +265,8 @@ private:
 	TSharedPtr<SComboBox<TSharedPtr<FPinComboBoxItem>>> PinComboBox;
 	TArray<TSharedPtr<FPinComboBoxItem>> PinComboBoxItems;
 
-	TSharedPtr<SComboBox<TSharedPtr<FName>>> DataComboBox;
-	TArray<TSharedPtr<FName>> DataComboBoxItems;
+	TSharedPtr<SComboBox<TSharedPtr<FString>>> DataComboBox;
+	TArray<TSharedPtr<FString>> DataComboBoxItems;
 
 	TSharedPtr<STextBlock> NodeNameTextBlock;
 	TSharedPtr<STextBlock> InfoTextBlock;
