[PCG] GPU: Proactively purge handles to persistent buffers as early as possible rather than relying on GCs to release them.

* Track data usage using high level flags.
* Release transient resources when graph execution completes, if the data has not been placed in the per-pin output, as this data is used for higen.

#rnx

diff --git a/Engine/Plugins/PCG/Source/PCG/Internal/Compute/Data/PCGProxyForGPUData.h b/Engine/Plugins/PCG/Source/PCG/Internal/Compute/Data/PCGProxyForGPUData.h
--- a/Engine/Plugins/PCG/Source/PCG/Internal/Compute/Data/PCGProxyForGPUData.h
+++ b/Engine/Plugins/PCG/Source/PCG/Internal/Compute/Data/PCGProxyForGPUData.h
@@ -23,7 +23,9 @@ public:
 	virtual EPCGDataType GetDataType() const override { return EPCGDataType::ProxyForGPU; }
 	virtual void AddToCrc(FArchiveCrc32& Ar, bool bFullDataCrc) const override;
 	virtual bool CanBeSerialized() const override { return false; }
+	virtual bool HoldsTransientResources() const override { return true; }
 	virtual bool IsCacheable() const override { return false; }
+	virtual void ReleaseTransientResources(const TCHAR* InReason = nullptr);
 	virtual EPCGDataType GetUnderlyingDataType() const override;
 	//~End UPCGData interface
 
diff --git a/Engine/Plugins/PCG/Source/PCG/Private/Compute/Data/PCGProxyForGPUData.cpp b/Engine/Plugins/PCG/Source/PCG/Private/Compute/Data/PCGProxyForGPUData.cpp
--- a/Engine/Plugins/PCG/Source/PCG/Private/Compute/Data/PCGProxyForGPUData.cpp
+++ b/Engine/Plugins/PCG/Source/PCG/Private/Compute/Data/PCGProxyForGPUData.cpp
@@ -18,6 +18,15 @@ void UPCGProxyForGPUData::AddToCrc(FArchiveCrc32& Ar, bool bFullDataCrc) const
 	AddUIDToCrc(Ar);
 }
 
+void UPCGProxyForGPUData::ReleaseTransientResources(const TCHAR* InReason)
+{
+#ifdef PCG_DATA_USAGE_LOGGING
+	UE_LOG(LogPCG, Warning, TEXT("Releasing GPU data for '%s' due to %s"), *GetName(), InReason ? InReason : TEXT("NOREASON"));
+#endif
+
+	DataCollectionOnGPU.Reset();
+}
+
 EPCGDataType UPCGProxyForGPUData::GetUnderlyingDataType() const
 {
 	if (TSharedPtr<const FPCGProxyForGPUDataCollection> DataCollectionGPU = GetGPUInfo())
@@ -67,7 +76,7 @@ UPCGProxyForGPUData::FReadbackResult UPCGProxyForGPUData::GetCPUData(FPCGContext
 	TSharedPtr<FPCGProxyForGPUDataCollection> DataOnGPU = GetGPUInfoMutable();
 	if (!DataOnGPU)
 	{
-		UE_LOG(LogPCG, Error, TEXT("Data collection lost!"));
+		UE_LOG(LogPCG, Error, TEXT("Data collection lost! Enabling the define PCG_DATA_USAGE_LOGGING may help to identify when resource was released."));
 
 		return FReadbackResult{ .bComplete = true };
 	}
@@ -90,7 +99,7 @@ TSharedPtr<FPCGProxyForGPUDataCollection> UPCGProxyForGPUData::GetGPUInfoMutable
 {
 	if (!DataCollectionOnGPU)
 	{
-		ensureMsgf(false, TEXT("Data %s: GPU data collection lost"), *GetName());
+		ensureMsgf(false, TEXT("Data %s: GPU data collection lost. Enabling the define 'PCG_DATA_USAGE_LOGGING' may help to identify when resource was released."), *GetName());
 		return nullptr;
 	}
 	else if (!DataCollectionOnGPU->GetDescription().DataDescs.IsValidIndex(DataIndexInCollection))
diff --git a/Engine/Plugins/PCG/Source/PCG/Private/Compute/DataInterfaces/PCGComputeDataInterface.h b/Engine/Plugins/PCG/Source/PCG/Private/Compute/DataInterfaces/PCGComputeDataInterface.h
--- a/Engine/Plugins/PCG/Source/PCG/Private/Compute/DataInterfaces/PCGComputeDataInterface.h
+++ b/Engine/Plugins/PCG/Source/PCG/Private/Compute/DataInterfaces/PCGComputeDataInterface.h
@@ -59,4 +59,7 @@ public:
 
 	/* Called prior to GetRenderProxy and execution. Any processing of data descriptions should be done here (after PerformPreExecuteReadbacks has been called). */
 	virtual bool PrepareForExecute_GameThread(UPCGDataBinding* InBinding) { return true; }
+
+	/** Release any handles to transient resources like GPU buffers. */
+	virtual void ReleaseTransientResources(const TCHAR* InReason = nullptr) {}
 };
diff --git a/Engine/Plugins/PCG/Source/PCG/Private/Compute/DataInterfaces/PCGDataCollectionUploadDataInterface.cpp b/Engine/Plugins/PCG/Source/PCG/Private/Compute/DataInterfaces/PCGDataCollectionUploadDataInterface.cpp
--- a/Engine/Plugins/PCG/Source/PCG/Private/Compute/DataInterfaces/PCGDataCollectionUploadDataInterface.cpp
+++ b/Engine/Plugins/PCG/Source/PCG/Private/Compute/DataInterfaces/PCGDataCollectionUploadDataInterface.cpp
@@ -51,6 +51,15 @@ bool UPCGDataProviderDataCollectionUpload::PrepareForExecute_GameThread(UPCGData
 	return !DataAdaptor || DataAdaptor->PrepareData_GameThread();
 }
 
+void UPCGDataProviderDataCollectionUpload::ReleaseTransientResources(const TCHAR* InReason)
+{
+#ifdef PCG_DATA_USAGE_LOGGING
+	UE_LOG(LogPCG, Warning, TEXT("%s: Releasing resources due to %s"), *GetName(), InReason ? InReason : TEXT("NOREASON"));
+#endif
+
+	DataAdaptor.Reset();
+}
+
 FComputeDataProviderRenderProxy* UPCGDataProviderDataCollectionUpload::GetRenderProxy()
 {
 	TRACE_CPUPROFILER_EVENT_SCOPE(UPCGDataProviderDataCollectionUpload::GetRenderProxy);
diff --git a/Engine/Plugins/PCG/Source/PCG/Private/Compute/DataInterfaces/PCGDataCollectionUploadDataInterface.h b/Engine/Plugins/PCG/Source/PCG/Private/Compute/DataInterfaces/PCGDataCollectionUploadDataInterface.h
--- a/Engine/Plugins/PCG/Source/PCG/Private/Compute/DataInterfaces/PCGDataCollectionUploadDataInterface.h
+++ b/Engine/Plugins/PCG/Source/PCG/Private/Compute/DataInterfaces/PCGDataCollectionUploadDataInterface.h
@@ -29,6 +29,8 @@ class UPCGDataProviderDataCollectionUpload : public UPCGDataCollectionDataProvid
 	GENERATED_BODY()
 
 public:
+	virtual void ReleaseTransientResources(const TCHAR* InReason = nullptr) override;
+
 	//~ Begin UComputeDataProvider Interface
 	FComputeDataProviderRenderProxy* GetRenderProxy() override;
 	//~ End UComputeDataProvider Interface
diff --git a/Engine/Plugins/PCG/Source/PCG/Private/Compute/Elements/PCGComputeGraphElement.cpp b/Engine/Plugins/PCG/Source/PCG/Private/Compute/Elements/PCGComputeGraphElement.cpp
--- a/Engine/Plugins/PCG/Source/PCG/Private/Compute/Elements/PCGComputeGraphElement.cpp
+++ b/Engine/Plugins/PCG/Source/PCG/Private/Compute/Elements/PCGComputeGraphElement.cpp
@@ -22,6 +22,24 @@
 
 #define LOCTEXT_NAMESPACE "PCGComputeGraphElement"
 
+FPCGComputeGraphContext::~FPCGComputeGraphContext()
+{
+	if (DataBinding)
+	{
+		DataBinding->ReleaseTransientResources();
+	}
+
+	for (TObjectPtr<UComputeDataProvider> Provider : ComputeGraphInstance.GetDataProviders())
+	{
+		if (UPCGComputeDataProvider* PCGDataProvider = Cast<UPCGComputeDataProvider>(Provider.Get()))
+		{
+			PCGDataProvider->ReleaseTransientResources(TEXT("~FPCGComputeGraphContext"));
+		}
+	}
+
+	ComputeGraphInstance.DestroyDataProviders();
+}
+
 void FPCGComputeGraphContext::AddExtraStructReferencedObjects(FReferenceCollector& Collector)
 {
 	Collector.AddPropertyReferences(FComputeGraphInstance::StaticStruct(), &ComputeGraphInstance);
diff --git a/Engine/Plugins/PCG/Source/PCG/Private/Compute/Elements/PCGComputeGraphElement.h b/Engine/Plugins/PCG/Source/PCG/Private/Compute/Elements/PCGComputeGraphElement.h
--- a/Engine/Plugins/PCG/Source/PCG/Private/Compute/Elements/PCGComputeGraphElement.h
+++ b/Engine/Plugins/PCG/Source/PCG/Private/Compute/Elements/PCGComputeGraphElement.h
@@ -17,6 +17,7 @@ struct FPCGComputeGraphContext : public FPCGContext
 {
 public:
 	virtual bool IsComputeContext() const override { return true; }
+	virtual ~FPCGComputeGraphContext();
 
 protected:
 	virtual void AddExtraStructReferencedObjects(FReferenceCollector& Collector) override;
diff --git a/Engine/Plugins/PCG/Source/PCG/Private/Compute/PCGDataBinding.cpp b/Engine/Plugins/PCG/Source/PCG/Private/Compute/PCGDataBinding.cpp
--- a/Engine/Plugins/PCG/Source/PCG/Private/Compute/PCGDataBinding.cpp
+++ b/Engine/Plugins/PCG/Source/PCG/Private/Compute/PCGDataBinding.cpp
@@ -39,6 +39,17 @@ void UPCGDataBinding::Initialize(const UPCGComputeGraph* InComputeGraph, FPCGCon
 #endif
 }
 
+void UPCGDataBinding::ReleaseTransientResources()
+{
+	AttributeIdRemap.Reset();
+	OutputDataCollection.TaggedData.Reset();
+	InputDataCollection.TaggedData.Reset();
+	MeshSpawnersToPrimitives.Reset();
+	DataToDebug.Reset();
+	DataToInspect.Reset();
+	ContextHandle = nullptr;
+}
+
 bool UPCGDataBinding::GetKernelPinDataDesc_GameThread(const UPCGComputeKernel* InKernel, FName InPinLabel, bool bIsInput, FPCGDataCollectionDesc& OutDescription) const
 {
 	return ensure(Graph && InKernel) && GetKernelPinDataDesc_GameThread(Graph->GetBindingIndex(InKernel->GetKernelIndex(), InPinLabel, bIsInput), OutDescription);
diff --git a/Engine/Plugins/PCG/Source/PCG/Private/Compute/PCGDataBinding.h b/Engine/Plugins/PCG/Source/PCG/Private/Compute/PCGDataBinding.h
--- a/Engine/Plugins/PCG/Source/PCG/Private/Compute/PCGDataBinding.h
+++ b/Engine/Plugins/PCG/Source/PCG/Private/Compute/PCGDataBinding.h
@@ -64,6 +64,9 @@ class UPCGDataBinding : public UObject
 public:
 	void Initialize(const UPCGComputeGraph* InComputeGraph, FPCGContext* InContext);
 
+	/** Clear state and release any handles to resources such as GPU buffers. */
+	void ReleaseTransientResources();
+
 	int32 GetAttributeId(const FPCGKernelAttributeKey& InAttribute) const { return AttributeTable.GetAttributeId(InAttribute); }
 	int32 GetAttributeId(FName InName, EPCGKernelAttributeType InType) const { return AttributeTable.GetAttributeId(InName, InType); }
 
diff --git a/Engine/Plugins/PCG/Source/PCG/Private/Graph/PCGGraphExecutor.cpp b/Engine/Plugins/PCG/Source/PCG/Private/Graph/PCGGraphExecutor.cpp
--- a/Engine/Plugins/PCG/Source/PCG/Private/Graph/PCGGraphExecutor.cpp
+++ b/Engine/Plugins/PCG/Source/PCG/Private/Graph/PCGGraphExecutor.cpp
@@ -756,6 +756,8 @@ void FPCGGraphExecutor::MarkInputResults(TArrayView<const FPCGTaskId> InInputRes
 			FOutputDataInfo& OutputDataInfo = TaskOutputs[TaskId];
 			if (!OutputDataInfo.bNeedsManualClear && --OutputDataInfo.RemainingSuccessorCount == 0)
 			{
+				OutputDataInfo.DataCollection.ClearUsage(EPCGDataUsage::GraphExecutorTaskOutput);
+
 				TaskOutputs.Remove(TaskId);
 			}
 		}
@@ -1045,10 +1047,23 @@ void FPCGGraphExecutor::PostTaskExecute(TSharedPtr<FPCGGraphActiveTask> ActiveTa
 #endif // WITH_EDITOR
 
 		// If the task is a post execute, then we can safely clear the data after getting it from the results.
-		const bool bTaskIsPostExecute = (ActiveTask.Element == GraphCompiler.GetSharedTrivialPostGraphElement());
+		bool bNeedsManualClear = (ActiveTask.Element == GraphCompiler.GetSharedTrivialPostGraphElement());
+
+		if (!bNeedsManualClear)
+		{
+			for (const FPCGTaggedData& TaggedData : ActiveTask.Context->OutputData.TaggedData)
+			{
+				// Don't release any data that holds onto transient resources mid-execution. Clear these at the end.
+				if (TaggedData.Data && TaggedData.Data->HoldsTransientResources())
+				{
+					bNeedsManualClear = true;
+					break;
+				}
+			}
+		}
 
 		// Store output in data map.
-		StoreResults(ActiveTask.NodeId, ActiveTask.Context->OutputData, bTaskIsPostExecute);
+		StoreResults(ActiveTask.NodeId, ActiveTask.Context->OutputData, bNeedsManualClear);
 
 		// Book-keeping
 		QueueNextTasks(ActiveTask.NodeId);
@@ -2041,6 +2056,8 @@ void FPCGGraphExecutor::StoreResults(FPCGTaskId InTaskId, const FPCGDataCollecti
 	OutputDataInfo.DataCollection = InTaskOutput;
 	OutputDataInfo.bNeedsManualClear = bNeedsManualClear;
 
+	OutputDataInfo.DataCollection.MarkUsage(EPCGDataUsage::GraphExecutorTaskOutput);
+
 	{
 		TRACE_CPUPROFILER_EVENT_SCOPE(FPCGGraphExecutor::StoreResults::Successors);
 		PCGGraphExecutor::TScopeLock ScopeLock(TasksLock);
@@ -2069,6 +2086,12 @@ void FPCGGraphExecutor::ClearResults()
 	}
 
 	PCGGraphExecutor::TScopeLock OutputLock(TaskOutputsLock);
+
+	for (const TPair<FPCGTaskId, FOutputDataInfo>& TaskOutput : TaskOutputs)
+	{
+		TaskOutput.Value.DataCollection.ClearUsage(EPCGDataUsage::GraphExecutorTaskOutput);
+	}
+
 	TaskOutputs.Reset();
 }
 
diff --git a/Engine/Plugins/PCG/Source/PCG/Private/PCGComponent.cpp b/Engine/Plugins/PCG/Source/PCG/Private/PCGComponent.cpp
--- a/Engine/Plugins/PCG/Source/PCG/Private/PCGComponent.cpp
+++ b/Engine/Plugins/PCG/Source/PCG/Private/PCGComponent.cpp
@@ -734,6 +734,8 @@ void UPCGComponent::PostProcessGraph(const FBox& InNewBounds, bool bInGenerated,
 					{
 						if (InData)
 						{
+							InData->MarkUsage(EPCGDataUsage::ComponentOutputData);
+
 							const_cast<UPCGData*>(InData)->Rename(nullptr, this, IsInPreviewMode() ? REN_DoNotDirty : REN_None);
 						}
 					});
@@ -1205,7 +1207,19 @@ EPCGHiGenGrid UPCGComponent::GetGenerationGrid() const
 void UPCGComponent::StoreOutputDataForPin(const FString& InResourceKey, const FPCGDataCollection& InData)
 {
 	FWriteScopeLock ScopedWriteLock(PerPinGeneratedOutputLock);
-	PerPinGeneratedOutput.FindOrAdd(InResourceKey) = InData;
+
+	InData.MarkUsage(EPCGDataUsage::ComponentPerPinOutputData);
+
+	if (FPCGDataCollection* FoundData = PerPinGeneratedOutput.Find(InResourceKey))
+	{
+		FoundData->ClearUsage(EPCGDataUsage::ComponentPerPinOutputData);
+
+		*FoundData = InData;
+	}
+	else
+	{
+		PerPinGeneratedOutput.Add(InResourceKey, InData);
+	}
 }
 
 const FPCGDataCollection* UPCGComponent::RetrieveOutputDataForPin(const FString& InResourceKey)
@@ -1217,6 +1231,12 @@ const FPCGDataCollection* UPCGComponent::RetrieveOutputDataForPin(const FString&
 void UPCGComponent::ClearPerPinGeneratedOutput()
 {
 	FWriteScopeLock ScopedWriteLock(PerPinGeneratedOutputLock);
+
+	for (TPair<FString, FPCGDataCollection>& Entry : PerPinGeneratedOutput)
+	{
+		Entry.Value.ClearUsage(EPCGDataUsage::ComponentPerPinOutputData);
+	}
+
 	PerPinGeneratedOutput.Reset();
 }
 
@@ -1634,6 +1654,8 @@ void UPCGComponent::CleanupUnusedManagedResources()
 
 void UPCGComponent::ClearGraphGeneratedOutput()
 {
+	GeneratedGraphOutput.ClearUsage(EPCGDataUsage::ComponentOutputData);
+
 	for (FPCGTaggedData& GeneratedData : GeneratedGraphOutput.TaggedData)
 	{
 		if (GeneratedData.Data && GeneratedData.Data->GetOuter() == this)
@@ -2568,6 +2590,9 @@ void UPCGComponent::StoreInspectionData(const FPCGStack* InStack, const UPCGNode
 		return;
 	}
 
+	InInputData.MarkUsage(EPCGDataUsage::ComponentInspectionData);
+	InOutputData.MarkUsage(EPCGDataUsage::ComponentInspectionData);
+
 	// Notify component that this task executed. Useful for editor visualization.
 	NotifyNodeExecuted(InNode, InStack, InTimer, bUsedCache);
 
@@ -2663,6 +2688,12 @@ void UPCGComponent::ClearInspectionData(bool bClearPerNodeExecutionData)
 {
 	{
 		FWriteScopeLock Lock(InspectionCacheLock);
+
+		for (TPair<FPCGStack, FPCGDataCollection>& Entry : InspectionCache)
+		{
+			Entry.Value.ClearUsage(EPCGDataUsage::ComponentInspectionData);
+		}
+
 		InspectionCache.Reset();
 	}
 
diff --git a/Engine/Plugins/PCG/Source/PCG/Private/PCGData.cpp b/Engine/Plugins/PCG/Source/PCG/Private/PCGData.cpp
--- a/Engine/Plugins/PCG/Source/PCG/Private/PCGData.cpp
+++ b/Engine/Plugins/PCG/Source/PCG/Private/PCGData.cpp
@@ -17,6 +17,8 @@
 
 #include UE_INLINE_GENERATED_CPP_BY_NAME(PCGData)
 
+#define PCG_DATA_USAGE_LOGGING
+
 UPCGData::UPCGData(const FObjectInitializer& ObjectInitializer)
 	: Super(ObjectInitializer)
 {
@@ -42,6 +44,46 @@ FPCGCrc UPCGData::GetOrComputeCrc(bool bFullDataCrc) const
 	return Crc;
 }
 
+void UPCGData::MarkUsage(EPCGDataUsage InUsage) const
+{
+#ifdef PCG_DATA_USAGE_LOGGING
+	const UEnum* PCGDataUsageEnum = StaticEnum<EPCGDataUsage>();
+	UE_LOG(LogPCG, Warning, TEXT("%s: Marked usage flag %s, flags %d -> %d"),
+		*GetName(),
+		*PCGDataUsageEnum->GetNameStringByValue(static_cast<__underlying_type(EPCGDataUsage)>(InUsage)),
+		Usage.load(),
+		Usage.load() | static_cast<__underlying_type(EPCGDataUsage)>(InUsage));
+#endif
+
+	Usage |= static_cast<__underlying_type(EPCGDataUsage)>(InUsage);
+}
+
+void UPCGData::ClearUsage(EPCGDataUsage InUsage) const
+{
+#ifdef PCG_DATA_USAGE_LOGGING
+	const UEnum* PCGDataUsageEnum = StaticEnum<EPCGDataUsage>();
+	UE_LOG(LogPCG, Warning, TEXT("%s: Umarked usage flag %s, flags %d -> %d"),
+		*GetName(),
+		*PCGDataUsageEnum->GetNameStringByValue(static_cast<__underlying_type(EPCGDataUsage)>(InUsage)),
+		Usage.load(),
+		Usage.load() & ~static_cast<__underlying_type(EPCGDataUsage)>(InUsage));
+#endif
+
+	Usage &= ~static_cast<__underlying_type(EPCGDataUsage)>(InUsage);
+
+	if (Usage == static_cast<__underlying_type(EPCGDataUsage)>(EPCGDataUsage::None))
+	{
+#ifdef PCG_DATA_USAGE_LOGGING
+		UE_LOG(LogPCG, Warning, TEXT("%s: !! ReleaseTransientResources after usage unmarked !!"), *GetName());
+#endif
+
+		VisitDataNetwork([](const UPCGData* InData)
+		{
+			const_cast<UPCGData*>(InData)->ReleaseTransientResources(TEXT("Data unused"));
+		});
+	}
+}
+
 FPCGCrc UPCGData::ComputeCrc(bool bFullDataCrc) const
 {
 	FArchiveCrc32 Ar;
diff --git a/Engine/Plugins/PCG/Source/PCG/Private/PCGElement.cpp b/Engine/Plugins/PCG/Source/PCG/Private/PCGElement.cpp
--- a/Engine/Plugins/PCG/Source/PCG/Private/PCGElement.cpp
+++ b/Engine/Plugins/PCG/Source/PCG/Private/PCGElement.cpp
@@ -1001,6 +1001,11 @@ FPCGContext* IPCGElement::Initialize(const FPCGDataCollection& InputData, TWeakO
 	return Context;
 }
 
+bool IPCGElement::IsCacheable(const UPCGSettings* Settings) const
+{
+	return !Settings || !Settings->ShouldExecuteOnGPU();
+}
+
 bool IPCGElement::IsCacheableInstance(const UPCGSettingsInterface* InSettingsInterface) const
 {
 	if (InSettingsInterface)
diff --git a/Engine/Plugins/PCG/Source/PCG/Public/Data/PCGTextureData.h b/Engine/Plugins/PCG/Source/PCG/Public/Data/PCGTextureData.h
--- a/Engine/Plugins/PCG/Source/PCG/Public/Data/PCGTextureData.h
+++ b/Engine/Plugins/PCG/Source/PCG/Public/Data/PCGTextureData.h
@@ -185,7 +185,9 @@ public:
 	virtual EPCGDataType GetDataType() const override { return EPCGDataType::Texture; }
 	virtual void AddToCrc(FArchiveCrc32& Ar, bool bFullDataCrc) const override;
 	// Don't hold onto exported buffers currently as graphics memory usage (and lifetimes) may cause issues.
+	virtual bool HoldsTransientResources() const override { return ResourceType == EPCGTextureResourceType::ExportedTexture; }
 	virtual bool IsCacheable() const { return Super::IsCacheable() && ResourceType != EPCGTextureResourceType::ExportedTexture; }
+	virtual void ReleaseTransientResources(const TCHAR* InReason = nullptr) override { TextureHandle = nullptr; }
 	//~End UPCGData interface
 
 	//~Begin UPCGBaseTextureData interface
diff --git a/Engine/Plugins/PCG/Source/PCG/Public/PCGData.h b/Engine/Plugins/PCG/Source/PCG/Public/PCGData.h
--- a/Engine/Plugins/PCG/Source/PCG/Public/PCGData.h
+++ b/Engine/Plugins/PCG/Source/PCG/Public/PCGData.h
@@ -21,6 +21,17 @@ class UPCGSettingsInterface;
 class UPCGSpatialData;
 struct FPCGContext;
 
+UENUM(meta = (Bitflags))
+enum class EPCGDataUsage : uint8
+{
+	None = 0,
+	GraphExecutorTaskOutput = 1 << 0,
+	ComponentOutputData = 1 << 1,
+	ComponentPerPinOutputData = 1 << 2,
+	ComponentInspectionData = 1 << 3,
+};
+ENUM_CLASS_FLAGS(EPCGDataUsage);
+
 /**
 * Base class for any "data" class in the PCG framework.
 * This is an intentionally vague base class so we can have the required
@@ -47,6 +58,11 @@ public:
 	/** Whether the data can be placed in the graph cache. */
 	virtual bool IsCacheable() const { return true; }
 
+	/** Whether this data is holding onto one or more transient resources. This information can be used to pro-actively release the resource when the data is no longer in use. */
+	virtual bool HoldsTransientResources() const { return false; }
+
+	virtual void ReleaseTransientResources(const TCHAR* InReason = nullptr) {}
+
 	/** If this data is a proxy, returns the underlying data type, otherwise returns this data type. */
 	virtual EPCGDataType GetUnderlyingDataType() const { return GetDataType(); }
 
@@ -87,6 +103,9 @@ public:
 	UPROPERTY(VisibleAnywhere, Category = Metadata)
 	TObjectPtr<UPCGMetadata> Metadata = nullptr;
 
+	void MarkUsage(EPCGDataUsage InUsage) const;
+	void ClearUsage(EPCGDataUsage InUsage) const;
+
 protected:
 	/** Computes Crc for this and any connected data. */
 	virtual FPCGCrc ComputeCrc(bool bFullDataCrc) const;
@@ -107,6 +126,9 @@ private:
 
 	/** Serves unique ID values to instances of this object. */
 	static inline std::atomic<uint64> UIDCounter{ 1 };
+
+	/** Usage flags to record any current usages of this data. Used to help manage lifetimes of any internal transient resources. */
+	mutable std::atomic<__underlying_type(EPCGDataUsage)> Usage{ 0 };
 };
 
 USTRUCT(BlueprintType)
@@ -246,6 +268,28 @@ struct PCG_API FPCGDataCollection
 	/** Strips all empty point data from the collection. */
 	int32 StripEmptyPointData();
 
+	void MarkUsage(EPCGDataUsage InUsage) const
+	{
+		for (const FPCGTaggedData& Data : TaggedData)
+		{
+			if (Data.Data)
+			{
+				Data.Data->MarkUsage(InUsage);
+			}
+		}
+	}
+
+	void ClearUsage(EPCGDataUsage InUsage) const
+	{
+		for (const FPCGTaggedData& Data : TaggedData)
+		{
+			if (Data.Data)
+			{
+				Data.Data->ClearUsage(InUsage);
+			}
+		}
+	}
+
 	UPROPERTY(BlueprintReadWrite, EditAnywhere, Category = Data)
 	TArray<FPCGTaggedData> TaggedData;
 
diff --git a/Engine/Plugins/PCG/Source/PCG/Public/PCGElement.h b/Engine/Plugins/PCG/Source/PCG/Public/PCGElement.h
--- a/Engine/Plugins/PCG/Source/PCG/Public/PCGElement.h
+++ b/Engine/Plugins/PCG/Source/PCG/Public/PCGElement.h
@@ -108,7 +108,7 @@ public:
 	bool IsCacheableInstance(const UPCGSettingsInterface* InSettingsInterface) const;
 
 	/** Returns true if the node can be cached (e.g. does not create artifacts & does not depend on untracked data */
-	virtual bool IsCacheable(const UPCGSettings* InSettings) const { return true; }
+	virtual bool IsCacheable(const UPCGSettings* InSettings) const;
 	
 	/** Returns true if the node outputs requires to update the output tagged data to detect if the data is used multiple times. */
 	virtual bool ShouldVerifyIfOutputsAreUsedMultipleTimes(const UPCGSettings* InSettings) const;
