[PCG] Remove logging define erroneously left enabled

#rnx

diff --git a/Engine/Plugins/PCG/Source/PCG/Private/PCGData.cpp b/Engine/Plugins/PCG/Source/PCG/Private/PCGData.cpp
--- a/Engine/Plugins/PCG/Source/PCG/Private/PCGData.cpp
+++ b/Engine/Plugins/PCG/Source/PCG/Private/PCGData.cpp
@@ -17,8 +17,6 @@
 
 #include UE_INLINE_GENERATED_CPP_BY_NAME(PCGData)
 
-#define PCG_DATA_USAGE_LOGGING
-
 UPCGData::UPCGData(const FObjectInitializer& ObjectInitializer)
 	: Super(ObjectInitializer)
 {
