Added muladd metadata operation (e.g. d = a + b * c)

diff --git a/Engine/Plugins/PCG/Source/PCG/Private/Elements/Metadata/PCGMetadataMathsOpElement.cpp b/Engine/Plugins/PCG/Source/PCG/Private/Elements/Metadata/PCGMetadataMathsOpElement.cpp
--- a/Engine/Plugins/PCG/Source/PCG/Private/Elements/Metadata/PCGMetadataMathsOpElement.cpp
+++ b/Engine/Plugins/PCG/Source/PCG/Private/Elements/Metadata/PCGMetadataMathsOpElement.cpp
@@ -37,7 +37,8 @@ namespace PCGMetadataMathsSettings
 		}
 
 		if (PCGMetadataMathsSettings::IsBinaryOp(Operation)
-			|| Operation == EPCGMetadataMathsOperation::Lerp)
+			|| Operation == EPCGMetadataMathsOperation::Lerp
+			|| Operation == EPCGMetadataMathsOperation::MulAdd)
 		{
 			return PCGMetadataSettingsBaseConstants::DoubleInputFirstLabel;
 		}
@@ -77,6 +78,11 @@ namespace PCGMetadataMathsSettings
 			return PCGMetadataSettingsBaseConstants::LerpRatioLabel;
 		}
 
+		if (PCGMetadataMathsSettings::IsTernaryOp(Operation))
+		{
+			return PCGMetadataSettingsBaseConstants::DoubleInputThirdLabel;
+		}
+
 		return NAME_None;
 	}
 
@@ -153,6 +159,8 @@ namespace PCGMetadataMathsSettings
 			return PCGMetadataMaths::Clamp(Value1, Value2, Value3);
 		case EPCGMetadataMathsOperation::Lerp:
 			return PCGMetadataMaths::Lerp(Value1, Value2, Value3);
+		case EPCGMetadataMathsOperation::MulAdd:
+			return Value1 + Value2 * Value3;
 		default:
 			return T{};
 		}
diff --git a/Engine/Plugins/PCG/Source/PCG/Public/Elements/Metadata/PCGMetadataMathsOpElement.h b/Engine/Plugins/PCG/Source/PCG/Public/Elements/Metadata/PCGMetadataMathsOpElement.h
--- a/Engine/Plugins/PCG/Source/PCG/Public/Elements/Metadata/PCGMetadataMathsOpElement.h
+++ b/Engine/Plugins/PCG/Source/PCG/Public/Elements/Metadata/PCGMetadataMathsOpElement.h
@@ -42,6 +42,7 @@ enum class EPCGMetadataMathsOperation : uint16
 	TernaryOp = 1 << 12 UMETA(Hidden),
 	Clamp,
 	Lerp,
+	MulAdd UMETA(Tooltip = "Multiply Add (A + B * C)")
 };
 ENUM_CLASS_FLAGS(EPCGMetadataMathsOperation);
 
diff --git a/Engine/Plugins/PCG/Source/PCG/Public/Elements/Metadata/PCGMetadataOpElementBase.h b/Engine/Plugins/PCG/Source/PCG/Public/Elements/Metadata/PCGMetadataOpElementBase.h
--- a/Engine/Plugins/PCG/Source/PCG/Public/Elements/Metadata/PCGMetadataOpElementBase.h
+++ b/Engine/Plugins/PCG/Source/PCG/Public/Elements/Metadata/PCGMetadataOpElementBase.h
@@ -29,6 +29,7 @@ namespace PCGMetadataSettingsBaseConstants
 {
 	const FName DoubleInputFirstLabel = TEXT("InA");
 	const FName DoubleInputSecondLabel = TEXT("InB");
+	const FName DoubleInputThirdLabel = TEXT("InC");
 
 	const FName ClampMinLabel = TEXT("Min");
 	const FName ClampMaxLabel = TEXT("Max");
