[PCG] Fix crash on multi input: properly reset CurrentWriteIndex

#rnx

diff --git a/Engine/Plugins/PCG/Source/PCG/Private/Elements/PCGStaticMeshSpawnerContext.cpp b/Engine/Plugins/PCG/Source/PCG/Private/Elements/PCGStaticMeshSpawnerContext.cpp
--- a/Engine/Plugins/PCG/Source/PCG/Private/Elements/PCGStaticMeshSpawnerContext.cpp
+++ b/Engine/Plugins/PCG/Source/PCG/Private/Elements/PCGStaticMeshSpawnerContext.cpp
@@ -11,6 +11,7 @@ void FPCGStaticMeshSpawnerContext::ResetInputIterationData()
 	CurrentOutputPointData = nullptr;
 	MaterialOverrideHelper.Reset();
 	CurrentPointIndex = 0;
+	CurrentWriteIndex = 0;
 	WeightedMeshInstances.Reset();
 	MeshToValueKey.Reset();
 	CumulativeWeights.Reset();
