[PCG] Convert PCGSpatialData subclasses to support ToPointArrayData()

#rnx

diff --git a/Engine/Plugins/PCG/Source/PCG/Private/Data/PCGCollisionShapeData.cpp b/Engine/Plugins/PCG/Source/PCG/Private/Data/PCGCollisionShapeData.cpp
--- a/Engine/Plugins/PCG/Source/PCG/Private/Data/PCGCollisionShapeData.cpp
+++ b/Engine/Plugins/PCG/Source/PCG/Private/Data/PCGCollisionShapeData.cpp
@@ -2,6 +2,7 @@
 
 #include "Data/PCGCollisionShapeData.h"
 
+#include "Data/PCGPointArrayData.h"
 #include "Data/PCGPointData.h"
 #include "Data/PCGSpatialData.h"
 #include "Elements/PCGVolumeSampler.h"
@@ -94,15 +95,27 @@ bool UPCGCollisionShapeData::SamplePoint(const FTransform& InTransform, const FB
 
 const UPCGPointData* UPCGCollisionShapeData::CreatePointData(FPCGContext* Context) const
 {
-	TRACE_CPUPROFILER_EVENT_SCOPE(UPCGShapeData::CreatePointData);
+	TRACE_CPUPROFILER_EVENT_SCOPE(UPCGCollisionShapeData::CreatePointData);
 
+	return Cast<UPCGPointData>(CreateBasePointData(Context, UPCGPointData::StaticClass()));
+}
+
+const UPCGPointArrayData* UPCGCollisionShapeData::CreatePointArrayData(FPCGContext* Context, const FBox& InBounds) const
+{
+	TRACE_CPUPROFILER_EVENT_SCOPE(UPCGCollisionShapeData::CreatePointArrayData);
+
+	return Cast<UPCGPointArrayData>(CreateBasePointData(Context, UPCGPointArrayData::StaticClass()));
+}
+
+const UPCGBasePointData* UPCGCollisionShapeData::CreateBasePointData(FPCGContext* Context, TSubclassOf<UPCGBasePointData> PointDataClass) const
+{
 	const PCGVolumeSampler::FVolumeSamplerParams SamplerParams;
 
-	const UPCGPointData* Data = PCGVolumeSampler::SampleVolume(Context, SamplerParams, this);
+	const UPCGBasePointData* Data = PCGVolumeSampler::SampleVolume(Context, PointDataClass, SamplerParams, this);
 
 	if (Data)
 	{
-		UE_LOG(LogPCG, Verbose, TEXT("Shape extracted %d points"), Data->GetPoints().Num());
+		UE_LOG(LogPCG, Verbose, TEXT("Shape extracted %d points"), Data->GetNumPoints());
 	}
 
 	return Data;
diff --git a/Engine/Plugins/PCG/Source/PCG/Private/Data/PCGCollisionWrapperData.cpp b/Engine/Plugins/PCG/Source/PCG/Private/Data/PCGCollisionWrapperData.cpp
--- a/Engine/Plugins/PCG/Source/PCG/Private/Data/PCGCollisionWrapperData.cpp
+++ b/Engine/Plugins/PCG/Source/PCG/Private/Data/PCGCollisionWrapperData.cpp
@@ -3,7 +3,7 @@
 #include "Data/PCGCollisionWrapperData.h"
 
 #include "PCGContext.h"
-#include "Data/PCGPointData.h"
+#include "Data/PCGBasePointData.h"
 #include "Elements/Metadata/PCGMetadataElementCommon.h"
 
 #include "Components/StaticMeshComponent.h"
@@ -182,7 +182,7 @@ bool FPCGCollisionWrapper::GetShapeArray(FBodyInstance* Body, EPCGCollisionQuery
 	return false;
 }
 
-bool UPCGCollisionWrapperData::Initialize(const UPCGPointData* InPointData, const FPCGAttributePropertyInputSelector& InCollisionSelector, EPCGCollisionQueryFlag InCollisionQueryFlag)
+bool UPCGCollisionWrapperData::Initialize(const UPCGBasePointData* InPointData, const FPCGAttributePropertyInputSelector& InCollisionSelector, EPCGCollisionQueryFlag InCollisionQueryFlag)
 {
 	TArray<FSoftObjectPath> MeshesToLoad;
 	if (PreInitializeAndGatherMeshesEx(InPointData, InCollisionSelector, InCollisionQueryFlag, MeshesToLoad))
@@ -196,7 +196,7 @@ bool UPCGCollisionWrapperData::Initialize(const UPCGPointData* InPointData, cons
 	}
 }
 
-bool UPCGCollisionWrapperData::PreInitializeAndGatherMeshesEx(const UPCGPointData* InPointData, const FPCGAttributePropertyInputSelector& InCollisionSelector, EPCGCollisionQueryFlag InCollisionQueryFlag, TArray<FSoftObjectPath>& OutMeshesToLoad)
+bool UPCGCollisionWrapperData::PreInitializeAndGatherMeshesEx(const UPCGBasePointData* InPointData, const FPCGAttributePropertyInputSelector& InCollisionSelector, EPCGCollisionQueryFlag InCollisionQueryFlag, TArray<FSoftObjectPath>& OutMeshesToLoad)
 {
 	CollisionWrapper.Uninitialize();
 
@@ -290,7 +290,7 @@ void UPCGCollisionWrapperData::GetResourceSizeEx(FResourceSizeEx& CumulativeReso
 {
 	if (GetPointData())
 	{
-		const_cast<UPCGPointData*>(GetPointData())->GetResourceSizeEx(CumulativeResourceSize);
+		const_cast<UPCGBasePointData*>(GetPointData())->GetResourceSizeEx(CumulativeResourceSize);
 	}
 }
 
@@ -309,6 +309,11 @@ const UPCGPointData* UPCGCollisionWrapperData::ToPointData(FPCGContext* Context,
 	return GetPointData()->ToPointData(Context, InBounds);
 }
 
+const UPCGPointArrayData* UPCGCollisionWrapperData::ToPointArrayData(FPCGContext* Context, const FBox& InBounds) const
+{
+	return GetPointData()->ToPointArrayData(Context, InBounds);
+}
+
 bool UPCGCollisionWrapperData::SamplePoint(const FTransform& InTransform, const FBox& InBounds, FPCGPoint& OutPoint, UPCGMetadata* OutMetadata) const
 {
 	if(!CollisionWrapper.bInitialized)
diff --git a/Engine/Plugins/PCG/Source/PCG/Private/Data/PCGDifferenceData.cpp b/Engine/Plugins/PCG/Source/PCG/Private/Data/PCGDifferenceData.cpp
--- a/Engine/Plugins/PCG/Source/PCG/Private/Data/PCGDifferenceData.cpp
+++ b/Engine/Plugins/PCG/Source/PCG/Private/Data/PCGDifferenceData.cpp
@@ -3,6 +3,7 @@
 #include "Data/PCGDifferenceData.h"
 
 #include "PCGContext.h"
+#include "DAta/PCGPointArrayData.h"
 #include "Data/PCGPointData.h"
 #include "Data/PCGSpatialData.h"
 #include "Data/PCGSpatialDataTpl.h"
@@ -230,9 +231,22 @@ const UPCGPointData* UPCGDifferenceData::CreatePointData(FPCGContext* Context) c
 {
 	TRACE_CPUPROFILER_EVENT_SCOPE(UPCGDifferenceData::CreatePointData);
 	
-	// This is similar to what we are doing in UPCGUnionData::CreatePointData
-	const UPCGPointData* SourcePointData = GetSource()->ToPointData(Context);
+	const UPCGBasePointData* SourcePointData = GetSource()->ToPointData(Context);
 
+	return CastChecked<UPCGPointData>(CreateBasePointData(Context, SourcePointData, UPCGPointData::StaticClass()));
+}
+
+const UPCGPointArrayData* UPCGDifferenceData::CreatePointArrayData(FPCGContext* Context, const FBox& InBounds) const
+{
+	TRACE_CPUPROFILER_EVENT_SCOPE(UPCGDifferenceData::CreatePointArrayData);
+
+	const UPCGBasePointData* SourcePointData = GetSource()->ToPointArrayData(Context);
+
+	return CastChecked<UPCGPointArrayData>(CreateBasePointData(Context, SourcePointData, UPCGPointArrayData::StaticClass()));
+}
+
+const UPCGBasePointData* UPCGDifferenceData::CreateBasePointData(FPCGContext* Context, const UPCGBasePointData* SourcePointData, TSubclassOf<UPCGBasePointData> PointDataClass) const
+{
 	if (!SourcePointData)
 	{
 		UE_LOG(LogPCG, Error, TEXT("Difference unable to get source points"));
@@ -244,16 +258,18 @@ const UPCGPointData* UPCGDifferenceData::CreatePointData(FPCGContext* Context) c
 		UE_LOG(LogPCG, Verbose, TEXT("Difference is trivial"));
 		return SourcePointData;
 	}
+		
+	UPCGBasePointData* TargetPointData = FPCGContext::NewObject_AnyThread<UPCGBasePointData>(Context, GetTransientPackage(), PointDataClass);
 
-	const UPCGMetadata* SourceMetadata = SourcePointData->Metadata;
-	const TArray<FPCGPoint>& SourcePoints = SourcePointData->GetPoints();
+	FPCGInitializeFromDataParams InitializeFromDataParams(this);
+	InitializeFromDataParams.bInheritSpatialData = false;
 
-	UPCGPointData* Data = FPCGContext::NewObject_AnyThread<UPCGPointData>(Context);
-	Data->InitializeFromData(this, SourceMetadata);
+	const UPCGMetadata* SourceMetadata = SourcePointData->Metadata;
+	InitializeFromDataParams.MetadataParentOverride = SourceMetadata;
+	TargetPointData->InitializeFromDataWithParams(InitializeFromDataParams);
+	
+	UPCGMetadata* OutMetadata = TargetPointData->Metadata;
 	
-	UPCGMetadata* OutMetadata = Data->Metadata;
-	TArray<FPCGPoint>& TargetPoints = Data->GetMutablePoints();
-
 	const UPCGMetadata* DifferenceMetadata = GetDifference()->Metadata;
 	UPCGMetadata* TempDiffMetadata = nullptr;
 	if (bDiffMetadata && OutMetadata && DifferenceMetadata)
@@ -264,9 +280,14 @@ const UPCGPointData* UPCGDifferenceData::CreatePointData(FPCGContext* Context) c
 
 	constexpr int ChunkSize = FPCGSpatialDataProcessing::DefaultSamplePointsChunkSize;
 
-	auto ChunkSamplePoints = [this, SourceMetadata, TempDiffMetadata, OutMetadata](const TArrayView<TPair<FTransform, FBox>>& Samples, const TArrayView<const FPCGPoint>& SourcePoints, TArray<FPCGPoint, TInlineAllocator<ChunkSize>>& OutPoints)
+	auto ChunkSamplePoints = [this, SourceMetadata, TempDiffMetadata, OutMetadata](const TArrayView<TPair<FTransform, FBox>>& Samples, const UPCGBasePointData* SourcePointData, int32 SourceReadIndex, UPCGBasePointData* TargetPointData, int32 TargetWriteIndex)
 	{
-		const int NumPoints = Samples.Num();
+		int32 NumWritten = 0;
+
+		FConstPCGPointValueRanges SourceRanges(SourcePointData);
+		FPCGPointValueRanges TargetRanges(TargetPointData, /*bAllocate=*/false);
+
+		const int32 NumPoints = Samples.Num();
 
 		TArray<FPCGPoint, TInlineAllocator<ChunkSize>> PointsFromDiff;
 		PointsFromDiff.SetNum(NumPoints);
@@ -274,44 +295,79 @@ const UPCGPointData* UPCGDifferenceData::CreatePointData(FPCGContext* Context) c
 		check(GetDifference());
 		GetDifference()->SamplePoints(Samples, PointsFromDiff, TempDiffMetadata);
 
-		TArray<FPCGPoint, TInlineAllocator<ChunkSize>> KeptPoints;
-		TArray<FPCGPoint, TInlineAllocator<ChunkSize>> RejectedPoints;
+		struct FKeptPoint
+		{
+			int32 Index = 0;
+			float Density = 0.f;
+		};
+
+		TArray<FKeptPoint, TInlineAllocator<ChunkSize>> KeptPoints;
+		TArray<int32, TInlineAllocator<ChunkSize>> RejectedPoints;
 
 		const bool bBinaryDensity = (DensityFunction == EPCGDifferenceDensityFunction::Binary);
 
 		for (int PointIndex = 0; PointIndex < NumPoints; ++PointIndex)
 		{
-			const FPCGPoint& Point = SourcePoints[PointIndex];
-			FPCGPoint& PointFromDiff = PointsFromDiff[PointIndex];
+			const int32 SourceIndex = SourceReadIndex + PointIndex;
+			
+			const FPCGPoint& PointFromDiff = PointsFromDiff[PointIndex];
 
-			const float Density = (bBinaryDensity && PointFromDiff.Density > 0) ? 0.0f : Point.Density - PointFromDiff.Density;
+			const float Density = (bBinaryDensity && PointFromDiff.Density > 0) ? 0.0f : SourceRanges.DensityRange[SourceIndex] - PointFromDiff.Density;
 
 			if (Density > 0)
 			{
-				FPCGPoint& OutPoint = KeptPoints.Add_GetRef(Point);
-				OutPoint.Density = Density;
+				KeptPoints.Add({ PointIndex, Density });
+			}
+			else if (bKeepZeroDensityPoints)
+			{
+				RejectedPoints.Add(PointIndex);
+			}
+		}
 
-				// TODO: create an array-based MergePointsAttributeSubset
+		if (KeptPoints.Num() > 0)
+		{
+			for (const FKeptPoint& KeptPoint : KeptPoints)
+			{
+				const int32 WriteIndex = TargetWriteIndex + NumWritten;
+				const int32 ReadIndex = SourceReadIndex + KeptPoint.Index;
+
+				const FPCGPoint& PointFromDiff = PointsFromDiff[KeptPoint.Index];
+								
+				TargetRanges.SetFromValueRanges(WriteIndex, SourceRanges, ReadIndex);
+				TargetRanges.DensityRange[WriteIndex] = KeptPoint.Density;
+				
 				if (TempDiffMetadata && PointFromDiff.MetadataEntry != PCGInvalidEntryKey)
 				{
-					OutMetadata->MergePointAttributesSubset(Point, SourceMetadata, SourceMetadata, PointFromDiff, TempDiffMetadata, TempDiffMetadata, OutPoint, EPCGMetadataOp::Sub);
+					OutMetadata->MergeAttributesSubset(SourceRanges.MetadataEntryRange[ReadIndex], SourceMetadata, SourceMetadata, PointFromDiff.MetadataEntry, TempDiffMetadata, TempDiffMetadata, TargetRanges.MetadataEntryRange[WriteIndex], EPCGMetadataOp::Sub);
 				}
+				
+				++NumWritten;
 			}
-			else if(bKeepZeroDensityPoints)
+		}
+
+		if (RejectedPoints.Num() > 0)
+		{
+			for (int32 RejectedIndex : RejectedPoints)
 			{
-				FPCGPoint& RejectedPoint = RejectedPoints.Add_GetRef(Point);
-				RejectedPoint.Density = 0;
+				const int32 WriteIndex = TargetWriteIndex + NumWritten;
+				const int32 ReadIndex = SourceReadIndex + RejectedIndex;
+
+				TargetRanges.SetFromValueRanges(WriteIndex, SourceRanges, ReadIndex);
+				TargetRanges.DensityRange[WriteIndex] = 0.f;
+				
+				++NumWritten;
 			}
 		}
 
-		OutPoints.Append(KeptPoints);
-		OutPoints.Append(RejectedPoints);
+		return NumWritten;
 	};
 
-	FPCGSpatialDataProcessing::SampleBasedRangeProcessing<ChunkSize>(Context ? &Context->AsyncState : nullptr, ChunkSamplePoints, SourcePoints, TargetPoints);
-	
-	UE_LOG(LogPCG, Verbose, TEXT("Difference generated %d points from %d source points"), TargetPoints.Num(), SourcePointData->GetPoints().Num());	
-	return Data;
+	const EPCGPointNativeProperties PropertiesToAllocate = SourcePointData->GetAllocatedProperties() | EPCGPointNativeProperties::Density;
+
+	FPCGSpatialDataProcessing::SampleBasedRangeProcessing<ChunkSize>(Context ? &Context->AsyncState : nullptr, ChunkSamplePoints, SourcePointData, TargetPointData, PropertiesToAllocate);
+
+	UE_LOG(LogPCG, Verbose, TEXT("Difference generated %d points from %d source points"), TargetPointData->GetNumPoints(), SourcePointData->GetNumPoints());
+	return TargetPointData;
 }
 
 UPCGSpatialData* UPCGDifferenceData::CopyInternal(FPCGContext* Context) const
diff --git a/Engine/Plugins/PCG/Source/PCG/Private/Data/PCGIntersectionData.cpp b/Engine/Plugins/PCG/Source/PCG/Private/Data/PCGIntersectionData.cpp
--- a/Engine/Plugins/PCG/Source/PCG/Private/Data/PCGIntersectionData.cpp
+++ b/Engine/Plugins/PCG/Source/PCG/Private/Data/PCGIntersectionData.cpp
@@ -3,6 +3,7 @@
 #include "Data/PCGIntersectionData.h"
 
 #include "PCGContext.h"
+#include "Data/PCGPointArrayData.h"
 #include "Data/PCGPointData.h"
 #include "Data/PCGSpatialData.h"
 #include "Data/PCGSpatialDataTpl.h"
@@ -184,27 +185,41 @@ const UPCGSpatialData* UPCGIntersectionData::FindFirstConcreteShapeFromNetwork()
 
 const UPCGPointData* UPCGIntersectionData::CreatePointData(FPCGContext* Context) const
 {
-	check(GetA() && GetB());
+	TRACE_CPUPROFILER_EVENT_SCOPE(UPCGIntersectionData::CreatePointData);
+	return CastChecked<UPCGPointData>(CreateBasePointData(Context, UPCGPointData::StaticClass()));
+}
+
+const UPCGPointArrayData* UPCGIntersectionData::CreatePointArrayData(FPCGContext* Context, const FBox& InBounds) const
+{
+	TRACE_CPUPROFILER_EVENT_SCOPE(UPCGIntersectionData::CreatePointArrayData);
+	return CastChecked<UPCGPointArrayData>(CreateBasePointData(Context, UPCGPointArrayData::StaticClass()));
+}
+
+UPCGBasePointData* UPCGIntersectionData::CreateBasePointData(FPCGContext* Context, TSubclassOf<UPCGBasePointData> PointDataClass) const
+{
+	const UPCGSpatialData* DataA = GetA();
+	const UPCGSpatialData* DataB = GetB();
+	check(DataA && DataB);
 	// TODO: this is a placeholder;
 	// Here we will get the point data from the lower-dimensionality data
 	// and then cull out any of the points that are outside the bounds of the other
-	if (GetA()->GetDimension() <= GetB()->GetDimension())
+	if (DataA->GetDimension() <= DataB->GetDimension())
 	{
-		return CreateAndFilterPointData(Context, GetA(), GetB());
+		return CreateAndFilterPointData(Context, DataA, DataB, PointDataClass);
 	}
 	else
 	{
-		return CreateAndFilterPointData(Context, GetB(), GetA());
+		return CreateAndFilterPointData(Context, DataB, DataA, PointDataClass);
 	}
 }
 
-UPCGPointData* UPCGIntersectionData::CreateAndFilterPointData(FPCGContext* Context, const UPCGSpatialData* X, const UPCGSpatialData* Y) const
+UPCGBasePointData* UPCGIntersectionData::CreateAndFilterPointData(FPCGContext* Context, const UPCGSpatialData* X, const UPCGSpatialData* Y, TSubclassOf<UPCGBasePointData> PointDataClass) const
 {
 	TRACE_CPUPROFILER_EVENT_SCOPE(UPCGIntersectionData::CreateAndFilterPointData);
 	check(X && Y);
 	check(X->GetDimension() <= Y->GetDimension());
 
-	const UPCGPointData* SourcePointData = X->ToPointData(Context, CachedBounds);
+	const UPCGBasePointData* SourcePointData = X->ToBasePointData(Context, CachedBounds);
 
 	if (!SourcePointData)
 	{
@@ -212,12 +227,16 @@ UPCGPointData* UPCGIntersectionData::CreateAndFilterPointData(FPCGContext* Conte
 		return nullptr;
 	}
 
-	const TArray<FPCGPoint>& SourcePoints = SourcePointData->GetPoints();
-	const UPCGMetadata* SourceMetadata = SourcePointData->Metadata;
+	UPCGBasePointData* OutputData = FPCGContext::NewObject_AnyThread<UPCGBasePointData>(Context, GetTransientPackage(), PointDataClass);
+
+	FPCGInitializeFromDataParams InitializeFromDataParams(this);
+	InitializeFromDataParams.bInheritSpatialData = false;
 
-	UPCGPointData* Data = FPCGContext::NewObject_AnyThread<UPCGPointData>(Context);
-	Data->InitializeFromData(this, SourceMetadata);
-	Data->Metadata->AddAttributes(Y->Metadata);
+	const UPCGMetadata* SourceMetadata = SourcePointData->Metadata;
+	InitializeFromDataParams.MetadataParentOverride = SourceMetadata;
+	OutputData->InitializeFromDataWithParams(InitializeFromDataParams);
+	
+	OutputData->Metadata->AddAttributes(Y->Metadata);
 
 	UPCGMetadata* TempYMetadata = nullptr;
 	if (Y->Metadata)
@@ -228,12 +247,15 @@ UPCGPointData* UPCGIntersectionData::CreateAndFilterPointData(FPCGContext* Conte
 
 	const bool bPointDataHasCommonAttributes = (SourceMetadata && Y->Metadata && SourceMetadata->HasCommonAttributes(Y->Metadata));
 
-	TArray<FPCGPoint>& TargetPoints = Data->GetMutablePoints();
-
 	constexpr int ChunkSize = FPCGSpatialDataProcessing::DefaultSamplePointsChunkSize;
 
-	auto ChunkSamplePoints = [this, Data, SourceMetadata, Y, TempYMetadata, bPointDataHasCommonAttributes](const TArrayView<TPair<FTransform, FBox>>& Samples, const TArrayView<const FPCGPoint>& SourcePoints, TArray<FPCGPoint, TInlineAllocator<ChunkSize>>& OutPoints)
+	auto ChunkSamplePoints = [this, SourceMetadata, Y, TempYMetadata, bPointDataHasCommonAttributes](const TArrayView<TPair<FTransform, FBox>>& Samples, const UPCGBasePointData* SourcePointData, int32 SourceReadIndex, UPCGBasePointData* TargetPointData, int32 TargetWriteIndex)
 	{
+		int32 NumWritten = 0;
+
+		const FConstPCGPointValueRanges SourceRanges(SourcePointData);
+		FPCGPointValueRanges TargetRanges(TargetPointData, /*bAllocate=*/false);
+				
 		const int NumPoints = Samples.Num();
 
 		TArray<FPCGPoint, TInlineAllocator<ChunkSize>> PointsFromY;
@@ -241,44 +263,76 @@ UPCGPointData* UPCGIntersectionData::CreateAndFilterPointData(FPCGContext* Conte
 
 		Y->SamplePoints(Samples, PointsFromY, TempYMetadata);
 
-		TArray<FPCGPoint, TInlineAllocator<ChunkSize>> KeptPoints;
-		TArray<FPCGPoint, TInlineAllocator<ChunkSize>> RejectedPoints;
+		TArray<int32, TInlineAllocator<ChunkSize>> KeptPoints;
+		TArray<int32, TInlineAllocator<ChunkSize>> RejectedPoints;
 
 		// Filter points based on output density
 		for (int PointIndex = 0; PointIndex < NumPoints; ++PointIndex)
 		{
-			const FPCGPoint& Point = SourcePoints[PointIndex];
+			const int32 SourceIndex = SourceReadIndex + PointIndex;
+
 			FPCGPoint& PointFromY = PointsFromY[PointIndex];
 
 			if (PointFromY.Density > 0)
 			{
-				FPCGPoint& KeptPoint = KeptPoints.Add_GetRef(Point); // note: not the sampled point
-				KeptPoint.Density = PCGIntersectionDataMaths::ComputeDensity(Point.Density, PointFromY.Density, DensityFunction);
-				KeptPoint.Color = Point.Color * PointFromY.Color;
+				KeptPoints.Add(PointIndex); // note: not the sampled point
+			}
+			else if (bKeepZeroDensityPoints)
+			{
+				RejectedPoints.Add(PointIndex); // note: not the sampled point
+			}
+		}
+
+		if (KeptPoints.Num() > 0)
+		{
+			for (int32 KeptIndex : KeptPoints)
+			{
+				const int32 WriteIndex = TargetWriteIndex + NumWritten;
+				const int32 ReadIndex = SourceReadIndex + KeptIndex;
+
+				FPCGPoint& PointFromY = PointsFromY[KeptIndex];
 
+				TargetRanges.SetFromValueRanges(WriteIndex, SourceRanges, ReadIndex);
+
+				TargetRanges.DensityRange[WriteIndex] = PCGIntersectionDataMaths::ComputeDensity(SourceRanges.DensityRange[ReadIndex], PointFromY.Density, DensityFunction);
+				TargetRanges.ColorRange[WriteIndex] = SourceRanges.ColorRange[ReadIndex] * PointFromY.Color;
+				
 				// TODO: create an array-based MergePointsAttributeSubset..
 				// If either the point from Y has metadata or the merge would be a non-trivial value, then perform the full merge
-				if (Data->Metadata && (bPointDataHasCommonAttributes || PointFromY.MetadataEntry != PCGInvalidEntryKey))
+				if (TargetPointData->Metadata && (bPointDataHasCommonAttributes || PointFromY.MetadataEntry != PCGInvalidEntryKey))
 				{
-					Data->Metadata->MergePointAttributesSubset(Point, SourceMetadata, SourceMetadata, PointFromY, TempYMetadata, TempYMetadata, KeptPoint, EPCGMetadataOp::Min);
+					TargetPointData->Metadata->MergeAttributesSubset(SourceRanges.MetadataEntryRange[ReadIndex], SourceMetadata, SourceMetadata, PointFromY.MetadataEntry, TempYMetadata, TempYMetadata, TargetRanges.MetadataEntryRange[WriteIndex], EPCGMetadataOp::Min);
 				}
+
+				++NumWritten;
 			}
-			else if (bKeepZeroDensityPoints)
+		}
+
+		if (RejectedPoints.Num() > 0)
+		{
+			for (int32 RejectedIndex : RejectedPoints)
 			{
-				FPCGPoint& RejectedPoint = RejectedPoints.Add_GetRef(Point); // note: not the sampled point
-				RejectedPoint.Density = 0;
+				const int32 WriteIndex = TargetWriteIndex + NumWritten;
+				const int32 ReadIndex = SourceReadIndex + RejectedIndex;
+
+				TargetRanges.SetFromValueRanges(WriteIndex, SourceRanges, ReadIndex);
+				
+				TargetRanges.DensityRange[WriteIndex] = 0.f;
+				
+				++NumWritten;
 			}
 		}
 
-		OutPoints.Append(KeptPoints);
-		OutPoints.Append(RejectedPoints);
+		return NumWritten;
 	};
 
-	FPCGSpatialDataProcessing::SampleBasedRangeProcessing<ChunkSize>(Context ? &Context->AsyncState : nullptr, ChunkSamplePoints, SourcePoints, TargetPoints);
+	const EPCGPointNativeProperties PropertiesToAllocate = SourcePointData->GetAllocatedProperties() | EPCGPointNativeProperties::Density;
+
+	FPCGSpatialDataProcessing::SampleBasedRangeProcessing<ChunkSize>(Context ? &Context->AsyncState : nullptr, ChunkSamplePoints, SourcePointData, OutputData, PropertiesToAllocate);
 
-	UE_LOG(LogPCG, Verbose, TEXT("Intersection generated %d points from %d source points"), TargetPoints.Num(), SourcePoints.Num());
+	UE_LOG(LogPCG, Verbose, TEXT("Intersection generated %d points from %d source points"), OutputData->GetNumPoints(), SourcePointData->GetNumPoints());
 
-	return Data;
+	return OutputData;
 }
 
 UPCGSpatialData* UPCGIntersectionData::CopyInternal(FPCGContext* Context) const
diff --git a/Engine/Plugins/PCG/Source/PCG/Private/Data/PCGLandscapeData.cpp b/Engine/Plugins/PCG/Source/PCG/Private/Data/PCGLandscapeData.cpp
--- a/Engine/Plugins/PCG/Source/PCG/Private/Data/PCGLandscapeData.cpp
+++ b/Engine/Plugins/PCG/Source/PCG/Private/Data/PCGLandscapeData.cpp
@@ -4,6 +4,7 @@
 
 #include "PCGContext.h"
 #include "PCGSubsystem.h"
+#include "Data/PCGPointArrayData.h"
 #include "Data/PCGPointData.h"
 #include "Data/PCGSpatialDataTpl.h"
 #include "Data/PCGSurfaceData.h"
@@ -471,16 +472,33 @@ const UPCGPointData* UPCGLandscapeData::CreatePointData(FPCGContext* Context, co
 {
 	TRACE_CPUPROFILER_EVENT_SCOPE(UPCGLandscapeData::CreatePointData);
 
+	return Cast<UPCGPointData>(CreateBasePointData(Context, InBounds, UPCGPointData::StaticClass()));
+}
+
+const UPCGPointArrayData* UPCGLandscapeData::CreatePointArrayData(FPCGContext* Context, const FBox& InBounds) const
+{
+	TRACE_CPUPROFILER_EVENT_SCOPE(UPCGLandscapeData::CreatePointArrayData);
+
+	return Cast<UPCGPointArrayData>(CreateBasePointData(Context, InBounds, UPCGPointArrayData::StaticClass()));
+}
+
+const UPCGBasePointData* UPCGLandscapeData::CreateBasePointData(FPCGContext* Context, const FBox& InBounds, TSubclassOf<UPCGBasePointData> PointDataClass) const
+{
 	if (LandscapeInfos.IsEmpty() || !LandscapeCache)
 	{
 		UE_LOG(LogPCG, Error, TEXT("PCG Landscape cache or Landscape info are not initialized"));
 		return nullptr;
 	}
 
-	UPCGPointData* Data = FPCGContext::NewObject_AnyThread<UPCGPointData>(Context);
-	Data->InitializeFromData(this);
-	TArray<FPCGPoint>& Points = Data->GetMutablePoints();
+	UPCGBasePointData* Data = FPCGContext::NewObject_AnyThread<UPCGBasePointData>(Context, GetTransientPackage(), PointDataClass);
 
+	FPCGInitializeFromDataParams InitializeFromDataParams(this);
+	InitializeFromDataParams.bInheritSpatialData = false;
+	Data->InitializeFromDataWithParams(InitializeFromDataParams);
+
+	// @todo_pcg: try to optimize allocation of BoundsMin/BoundsMax as it currently depends on the Landscape components and should always be the same.
+	Data->AllocateProperties(EPCGPointNativeProperties::Transform | EPCGPointNativeProperties::Seed | EPCGPointNativeProperties::BoundsMin | EPCGPointNativeProperties::BoundsMax);
+	
 	FBox EffectiveBounds = Bounds;
 	if (InBounds.IsValid)
 	{
@@ -494,6 +512,12 @@ const UPCGPointData* UPCGLandscapeData::CreatePointData(FPCGContext* Context, co
 	}
 
 	UPCGMetadata* OutMetadata = DataProps.bGetLayerWeights ? Data->Metadata : nullptr;
+	if (OutMetadata)
+	{
+		Data->AllocateProperties(EPCGPointNativeProperties::MetadataEntry);
+	}
+		
+	int32 NumPoints = 0;
 
 	// Most proxies we gathered will have the same landscape info, we shouldn't loop multiple times
 	// on them, unless we add the box filtering - but even then, depending on the transform we could have overlaps
@@ -525,12 +549,14 @@ const UPCGPointData* UPCGLandscapeData::CreatePointData(FPCGContext* Context, co
 		}
 
 		const int64 PointCountUpperBound = (1 + MaxX - MinX) * (1 + MaxY - MinY);
-		const int32 PointsBeforeNum = Points.Num();
+		const int32 PointsBeforeNum = NumPoints;
 		if (PointCountUpperBound > 0)
 		{
-			Points.Reserve(PointsBeforeNum + PointCountUpperBound);
+			Data->SetNumPoints(PointsBeforeNum + PointCountUpperBound);
 		}
 
+		FPCGPointValueRanges OutRanges(Data, /*bAllocate=*/false);
+
 		const int32 MinComponentX = MinX / ComponentSizeQuads;
 		const int32 MaxComponentX = MaxX / ComponentSizeQuads;
 		const int32 MinComponentY = MinY / ComponentSizeQuads;
@@ -569,7 +595,7 @@ const UPCGPointData* UPCGLandscapeData::CreatePointData(FPCGContext* Context, co
 					{
 						const int32 PointIndex = LocalX + LocalY * (ComponentSizeQuads + 1);
 
-						FPCGPoint& Point = Points.Emplace_GetRef();
+						FPCGPoint Point;
 						if (DataProps.bGetHeightOnly)
 						{
 							LandscapeCacheEntry->GetPointHeightOnly(PointIndex, Point);
@@ -577,14 +603,21 @@ const UPCGPointData* UPCGLandscapeData::CreatePointData(FPCGContext* Context, co
 						else
 						{
 							LandscapeCacheEntry->GetPoint(PointIndex, Point, OutMetadata);
+							OutRanges.MetadataEntryRange[NumPoints] = Point.MetadataEntry;
 						}
+
+						OutRanges.TransformRange[NumPoints] = Point.Transform;
+						OutRanges.SeedRange[NumPoints] = Point.Seed;
+						OutRanges.BoundsMinRange[NumPoints] = Point.BoundsMin;
+						OutRanges.BoundsMaxRange[NumPoints] = Point.BoundsMax;
+						++NumPoints;
 					}
 				}
 			}
 		}
 
-		check(Points.Num() - PointsBeforeNum <= PointCountUpperBound);
-		UE_LOG(LogPCG, Verbose, TEXT("Landscape %s extracted %d of %d potential points"), *LandscapeProxy->GetFName().ToString(), Points.Num() - PointsBeforeNum, PointCountUpperBound);
+		check(NumPoints - PointsBeforeNum <= PointCountUpperBound);
+		UE_LOG(LogPCG, Verbose, TEXT("Landscape %s extracted %d of %d potential points"), *LandscapeProxy->GetFName().ToString(), NumPoints - PointsBeforeNum, PointCountUpperBound);
 	}
 
 	return Data;
diff --git a/Engine/Plugins/PCG/Source/PCG/Private/Data/PCGLandscapeSplineData.cpp b/Engine/Plugins/PCG/Source/PCG/Private/Data/PCGLandscapeSplineData.cpp
--- a/Engine/Plugins/PCG/Source/PCG/Private/Data/PCGLandscapeSplineData.cpp
+++ b/Engine/Plugins/PCG/Source/PCG/Private/Data/PCGLandscapeSplineData.cpp
@@ -3,8 +3,9 @@
 #include "Data/PCGLandscapeSplineData.h"
 
 #include "PCGContext.h"
-#include "Data/PCGPolyLineData.h"
+#include "Data/PCGPointArrayData.h"
 #include "Data/PCGPointData.h"
+#include "Data/PCGPolyLineData.h"
 #include "Elements/PCGSplineSampler.h"
 #include "Helpers/PCGHelpers.h"
 
@@ -262,20 +263,35 @@ FTransform UPCGLandscapeSplineData::GetTransformAtAlpha(float Alpha) const
 
 const UPCGPointData* UPCGLandscapeSplineData::CreatePointData(FPCGContext* Context) const
 {
-	check(Spline.IsValid());
 	TRACE_CPUPROFILER_EVENT_SCOPE(UPCGLandscapeSplineData::CreatePointData);
 
-	UPCGPointData* Data = FPCGContext::NewObject_AnyThread<UPCGPointData>(Context);
-	Data->InitializeFromData(this);
-	TArray<FPCGPoint>& Points = Data->GetMutablePoints();
+	return CastChecked<UPCGPointData>(CreateBasePointData(Context, UPCGPointData::StaticClass()));
+}
+
+const UPCGPointArrayData* UPCGLandscapeSplineData::CreatePointArrayData(FPCGContext* Context, const FBox& InBounds) const
+{
+	TRACE_CPUPROFILER_EVENT_SCOPE(UPCGLandscapeSplineData::CreatePointArrayData);
+
+	return CastChecked<UPCGPointArrayData>(CreateBasePointData(Context, UPCGPointArrayData::StaticClass()));
+}
+
+const UPCGBasePointData* UPCGLandscapeSplineData::CreateBasePointData(FPCGContext* Context, TSubclassOf<UPCGBasePointData> PointDataClass) const
+{
+	check(Spline.IsValid());
+
+	UPCGBasePointData* Data = FPCGContext::NewObject_AnyThread<UPCGBasePointData>(Context, GetTransientPackage(), PointDataClass);
 
+	FPCGInitializeFromDataParams InitializeFromDataParams(this);
+	InitializeFromDataParams.bInheritSpatialData = false;
+	Data->InitializeFromDataWithParams(InitializeFromDataParams);
+	
 	FPCGSplineSamplerParams SamplerParams;
 	SamplerParams.Mode = EPCGSplineSamplingMode::Distance;
 	SamplerParams.Dimension = EPCGSplineSamplingDimension::OnHorizontal;
 
 	PCGSplineSamplerHelpers::SampleLineData(Context, /*LineData=*/this, /*InBoundingShapeData=*/nullptr, /*InProjectionTarget=*/nullptr, /*InProjectionParams=*/{}, SamplerParams, Data);
 
-	UE_LOG(LogPCG, Verbose, TEXT("Landscape spline %s generated %d points"), *Spline->GetFName().ToString(), Points.Num());
+	UE_LOG(LogPCG, Verbose, TEXT("Landscape spline %s generated %d points"), *Spline->GetFName().ToString(), Data->GetNumPoints());
 
 	return Data;
 }
diff --git a/Engine/Plugins/PCG/Source/PCG/Private/Data/PCGPrimitiveData.cpp b/Engine/Plugins/PCG/Source/PCG/Private/Data/PCGPrimitiveData.cpp
--- a/Engine/Plugins/PCG/Source/PCG/Private/Data/PCGPrimitiveData.cpp
+++ b/Engine/Plugins/PCG/Source/PCG/Private/Data/PCGPrimitiveData.cpp
@@ -3,6 +3,7 @@
 #include "Data/PCGPrimitiveData.h"
 
 #include "CollisionShape.h"
+#include "Data/PCGPointArrayData.h"
 #include "Data/PCGPointData.h"
 #include "Data/PCGSpatialData.h"
 #include "Elements/PCGVolumeSampler.h"
@@ -52,15 +53,27 @@ const UPCGPointData* UPCGPrimitiveData::CreatePointData(FPCGContext* Context) co
 {
 	TRACE_CPUPROFILER_EVENT_SCOPE(UPCGPrimitiveData::CreatePointData);
 
+	return Cast<UPCGPointData>(CreateBasePointData(Context, UPCGPointData::StaticClass()));
+}
+
+const UPCGPointArrayData* UPCGPrimitiveData::CreatePointArrayData(FPCGContext* Context, const FBox& InBounds) const
+{
+	TRACE_CPUPROFILER_EVENT_SCOPE(UPCGPrimitiveData::CreatePointArrayData);
+
+	return Cast<UPCGPointArrayData>(CreateBasePointData(Context, UPCGPointArrayData::StaticClass()));
+}
+
+const UPCGBasePointData* UPCGPrimitiveData::CreateBasePointData(FPCGContext* Context, TSubclassOf<UPCGBasePointData> PointDataClass) const
+{
 	PCGVolumeSampler::FVolumeSamplerParams SamplerParams;
 	SamplerParams.VoxelSize = VoxelSize;
 	SamplerParams.Bounds = GetBounds();
 
-	const UPCGPointData* Data = PCGVolumeSampler::SampleVolume(Context, SamplerParams, this);
+	const UPCGBasePointData* Data = PCGVolumeSampler::SampleVolume(Context, PointDataClass, SamplerParams, this);
 
 	if (Data && ensure(Primitive.Get()))
 	{
-		UE_LOG(LogPCG, Verbose, TEXT("Primitive %s extracted %d points"), *Primitive->GetFName().ToString(), Data->GetPoints().Num());
+		UE_LOG(LogPCG, Verbose, TEXT("Primitive %s extracted %d points"), *Primitive->GetFName().ToString(), Data->GetNumPoints());
 	}
 
 	return Data;
diff --git a/Engine/Plugins/PCG/Source/PCG/Private/Data/PCGProjectionData.cpp b/Engine/Plugins/PCG/Source/PCG/Private/Data/PCGProjectionData.cpp
--- a/Engine/Plugins/PCG/Source/PCG/Private/Data/PCGProjectionData.cpp
+++ b/Engine/Plugins/PCG/Source/PCG/Private/Data/PCGProjectionData.cpp
@@ -4,6 +4,7 @@
 
 #include "PCGContext.h"
 #include "Data/PCGLandscapeSplineData.h"
+#include "Data/PCGPointArrayData.h"
 #include "Data/PCGPointData.h"
 #include "Data/PCGSplineData.h"
 #include "Data/PCGSpatialData.h"
@@ -253,16 +254,31 @@ bool UPCGProjectionData::RequiresCollapseToSample() const
 const UPCGPointData* UPCGProjectionData::CreatePointData(FPCGContext* Context) const
 {
 	TRACE_CPUPROFILER_EVENT_SCOPE(UPCGProjectionData::CreatePointData);
+
+	return CastChecked<UPCGPointData>(CreateBasePointData(Context, UPCGPointData::StaticClass()));
+}
+
+const UPCGPointArrayData* UPCGProjectionData::CreatePointArrayData(FPCGContext* Context, const FBox& InBounds) const
+{
+	TRACE_CPUPROFILER_EVENT_SCOPE(UPCGProjectionData::CreatePointArrayData);
+
+	return CastChecked<UPCGPointArrayData>(CreateBasePointData(Context, UPCGPointArrayData::StaticClass()));
+}
+
+const UPCGBasePointData* UPCGProjectionData::CreateBasePointData(FPCGContext* Context, TSubclassOf<UPCGBasePointData> PointDataClass) const
+{
 	// TODO: add mechanism in the ToPointData so we can pass in a transform
 	// so we can forego creating the points twice if they're not used.
-	const UPCGPointData* SourcePointData = Source->ToPointData(Context);
+	const UPCGBasePointData* SourcePointData = Source->ToBasePointData(Context);
 	const UPCGMetadata* SourceMetadata = SourcePointData->Metadata;
-	const TArray<FPCGPoint>& SourcePoints = SourcePointData->GetPoints();
-
-	UPCGPointData* PointData = FPCGContext::NewObject_AnyThread<UPCGPointData>(Context);
 
+	UPCGBasePointData* PointData = FPCGContext::NewObject_AnyThread<UPCGBasePointData>(Context, GetTransientPackage(), PointDataClass);
+	
 	// Copy metadata attributes from source point including values
-	PointData->InitializeFromData(this, SourceMetadata);
+	FPCGInitializeFromDataParams InitializeFromDataParams(this);
+	InitializeFromDataParams.MetadataParentOverride = SourceMetadata;
+	InitializeFromDataParams.bInheritSpatialData = false;
+	PointData->InitializeFromDataWithParams(InitializeFromDataParams);
 	UPCGMetadata* OutMetadata = PointData->Metadata;
 	check(OutMetadata);
 
@@ -290,56 +306,92 @@ const UPCGPointData* UPCGProjectionData::CreatePointData(FPCGContext* Context) c
 	// Add any attributes from filtered metadata to produce final list
 	OutMetadata->AddAttributes(TempTargetMetadata);
 
-	TArray<FPCGPoint>& Points = PointData->GetMutablePoints();
+	auto InitializeFunc = [PointData, SourcePointData]()
+	{
+		PointData->SetNumPoints(SourcePointData->GetNumPoints());
+
+		// @todo_pcg: try to optimize allocations here
+		PointData->AllocateProperties(SourcePointData->GetAllocatedProperties() | EPCGPointNativeProperties::Transform | EPCGPointNativeProperties::Color | EPCGPointNativeProperties::Density);
+	};
+
+	auto MoveDataRangeFunc = [PointData](int32 RangeStartIndex, int32 MoveToIndex, int32 NumElements)
+	{
+		PointData->MoveRange(RangeStartIndex, MoveToIndex, NumElements);
+	};
 
-	FPCGAsync::AsyncPointProcessing(Context, SourcePoints.Num(), Points, [this, SourceMetadata, OutMetadata, TempTargetMetadata, &SourcePoints](int32 Index, FPCGPoint& OutPoint)
+	auto FinishedFunc = [PointData](int32 NumPoints)
 	{
-		const FPCGPoint& SourcePoint = SourcePoints[Index];
+		PointData->SetNumPoints(NumPoints);
+	};
 
-		FPCGPoint PointFromTarget;
-		bool bValidProjection = true;
-		if (!Target->ProjectPoint(SourcePoint.Transform, SourcePoint.GetLocalBounds(), ProjectionParams, PointFromTarget, TempTargetMetadata))
+	auto ProcessRangeFunc = [this, PointData, SourcePointData, SourceMetadata, OutMetadata, TempTargetMetadata](int32 StartReadIndex, int32 StartWriteIndex, int32 Count)
+	{
+		int32 NumWritten = 0;
+
+		const FConstPCGPointValueRanges InRanges(SourcePointData);
+		FPCGPointValueRanges OutRanges(PointData, /*bAllocate=*/false);
+
+		for (int32 ReadIndex = StartReadIndex; ReadIndex < StartReadIndex + Count; ++ReadIndex)
 		{
-			if (!bKeepZeroDensityPoints)
+			const int32 WriteIndex = StartWriteIndex + NumWritten;
+			FPCGPoint PointFromTarget;
+			bool bValidProjection = true;
+
+			const FBox LocalBounds = PCGPointHelpers::GetLocalBounds(InRanges.BoundsMinRange[ReadIndex], InRanges.BoundsMaxRange[ReadIndex]);
+			if (!Target->ProjectPoint(InRanges.TransformRange[ReadIndex], LocalBounds, ProjectionParams, PointFromTarget, TempTargetMetadata))
 			{
-				return false;
+				if (!bKeepZeroDensityPoints)
+				{
+					continue;
+				}
+				else
+				{
+					// Point is rejected, mark its density to zero, put it in a state where we won't affect the output point
+					bValidProjection = false;
+					PointFromTarget.Transform = InRanges.TransformRange[ReadIndex];
+					PointFromTarget.Color = InRanges.ColorRange[ReadIndex];
+					PointFromTarget.Density = 0;
+					PointFromTarget.MetadataEntry = PCGInvalidEntryKey;
+				}
 			}
-			else
+
+			// Merge points into a single point
+			OutRanges.SetFromValueRanges(WriteIndex, InRanges, ReadIndex);
+
+			// Apply projection result. Some of the params are already used inside ProjectPoint, so we are just applying the remaining bits that ProjectPoint() did not have access to.
+			// TODO this would be cleaner if there was a ProjectPoint that took an FPCGPoint
+			OutRanges.TransformRange[WriteIndex] = PointFromTarget.Transform;
+			
+			if (bValidProjection)
 			{
-				// Point is rejected, mark its density to zero, put it in a state where we won't affect the output point
-				bValidProjection = false;
-				PointFromTarget.Transform = SourcePoint.Transform;
-				PointFromTarget.Color = SourcePoint.Color;
-				PointFromTarget.Density = 0;
-				PointFromTarget.MetadataEntry = PCGInvalidEntryKey;
+				OutRanges.ColorRange[WriteIndex] = PCGProjectionPrivate::ApplyProjectionColorBlend(InRanges.ColorRange[ReadIndex], PointFromTarget.Color, ProjectionParams.ColorBlendMode);
 			}
-		}
+			
+			OutRanges.DensityRange[WriteIndex] *= PointFromTarget.Density;
 
-		// Merge points into a single point
-		OutPoint = SourcePoint;
-
-		// Apply projection result. Some of the params are already used inside ProjectPoint, so we are just applying the remaining bits that ProjectPoint() did not have access to.
-		// TODO this would be cleaner if there was a ProjectPoint that took an FPCGPoint
-		OutPoint.Transform = PointFromTarget.Transform;
+			if (OutMetadata && TempTargetMetadata && PointFromTarget.MetadataEntry != PCGInvalidEntryKey)
+			{
+				// Merge metadata to produce final attribute values
+				OutMetadata->MergeAttributesSubset(InRanges.MetadataEntryRange[ReadIndex], SourceMetadata, SourceMetadata, PointFromTarget.MetadataEntry, TempTargetMetadata, TempTargetMetadata, OutRanges.MetadataEntryRange[WriteIndex], ProjectionParams.AttributeMergeOperation);
+			}
 
-		// There is no reason to project a color blend if the projection failed
-		if (bValidProjection)
-		{
-			OutPoint.Color = PCGProjectionPrivate::ApplyProjectionColorBlend(SourcePoint.Color, PointFromTarget.Color, ProjectionParams.ColorBlendMode);
+			++NumWritten;
 		}
 
-		OutPoint.Density *= PointFromTarget.Density;
+		return NumWritten;
+	};
 
-		if (OutMetadata && TempTargetMetadata && PointFromTarget.MetadataEntry != PCGInvalidEntryKey)
-		{
-			// Merge metadata to produce final attribute values
-			OutMetadata->MergePointAttributesSubset(SourcePoint, SourceMetadata, SourceMetadata, PointFromTarget, TempTargetMetadata, TempTargetMetadata, OutPoint, ProjectionParams.AttributeMergeOperation);
-		}
+	FPCGAsync::AsyncProcessingRangeEx(
+		&Context->AsyncState,
+		SourcePointData->GetNumPoints(),
+		InitializeFunc,
+		ProcessRangeFunc,
+		MoveDataRangeFunc,
+		FinishedFunc,
+		/*bEnableTimeSlicing=*/false);
 
-		return true;
-	});
 
-	UE_LOG(LogPCG, Verbose, TEXT("Projection generated %d points from %d source points"), Points.Num(), SourcePoints.Num());
+	UE_LOG(LogPCG, Verbose, TEXT("Projection generated %d points from %d source points"), PointData->GetNumPoints(), SourcePointData->GetNumPoints());
 
 	return PointData;
 }
diff --git a/Engine/Plugins/PCG/Source/PCG/Private/Data/PCGSpatialData.cpp b/Engine/Plugins/PCG/Source/PCG/Private/Data/PCGSpatialData.cpp
--- a/Engine/Plugins/PCG/Source/PCG/Private/Data/PCGSpatialData.cpp
+++ b/Engine/Plugins/PCG/Source/PCG/Private/Data/PCGSpatialData.cpp
@@ -22,25 +22,28 @@ UPCGSpatialData::UPCGSpatialData(const FObjectInitializer& ObjectInitializer)
 	Metadata->SetupDomain(PCGMetadataDomainID::Data, /*bIsDefault=*/true);
 }
 
-const UPCGBasePointData* UPCGSpatialData::ToBasePointData(FPCGContext* Context) const
+const UPCGBasePointData* UPCGSpatialData::ToBasePointData(FPCGContext* Context, const FBox& InBounds) const
 {
 	if (CVarPCGEnablePointArrayData.GetValueOnAnyThread())
 	{
-		return ToPointArrayData(Context);
+		return ToPointArrayData(Context, InBounds);
 	}
 	else
 	{
-		return ToPointData(Context);
+		return ToPointData(Context, InBounds);
 	}
 }
 
-const UPCGPointArrayData* UPCGSpatialData::ToPointArrayData(FPCGContext* Context) const
+const UPCGPointArrayData* UPCGSpatialData::ToPointArrayData(FPCGContext* Context, const FBox& InBounds) const
 {
 	// @todo_pcg: this is sub-optimal but until we implement this in other spatial types (including the cached types) it works.
 	if (const UPCGPointData* PointData = ToPointData(Context))
 	{
 		UPCGPointArrayData* PointArrayData = FPCGContext::NewObject_AnyThread<UPCGPointArrayData>(Context);
-		PointArrayData->InitializeFromData(PointData);
+
+		FPCGInitializeFromDataParams InitializeFromDataParams(PointData);
+		InitializeFromDataParams.bInheritSpatialData = false;
+		PointArrayData->InitializeFromDataWithParams(InitializeFromDataParams);
 
 		UPCGBasePointData::SetPoints(PointData, PointArrayData, {}, /*bCopyAll=*/true);
 
@@ -58,25 +61,25 @@ void UPCGSpatialDataWithPointCache::GetResourceSizeEx(FResourceSizeEx& Cumulativ
 
 	if (CachedPointData)
 	{
-		const_cast<UPCGPointData*>(CachedPointData.Get())->GetResourceSizeEx(CumulativeResourceSize);
+		const_cast<UPCGBasePointData*>(CachedPointData.Get())->GetResourceSizeEx(CumulativeResourceSize);
 	}
 
 	CumulativeResourceSize.AddDedicatedSystemMemoryBytes(CachedBoundedPointDataBoxes.GetAllocatedSize() + CachedBoundedPointData.GetAllocatedSize());
 
-	for (const UPCGPointData* Data : CachedBoundedPointData)
+	for (const UPCGBasePointData* Data : CachedBoundedPointData)
 	{
 		if (Data)
 		{
-			const_cast<UPCGPointData*>(Data)->GetResourceSizeEx(CumulativeResourceSize);
+			const_cast<UPCGBasePointData*>(Data)->GetResourceSizeEx(CumulativeResourceSize);
 		}
 	}
 }
 
-const UPCGPointData* UPCGSpatialDataWithPointCache::ToPointData(FPCGContext* Context, const FBox& InBounds) const
+const UPCGBasePointData* UPCGSpatialDataWithPointCache::ToBasePointDataInternal(FPCGContext* Context, const FBox& InBounds, TFunctionRef<const UPCGBasePointData*(FPCGContext*, const FBox&)> CreatePointDataFunc) const
 {
 	if (InBounds.IsValid && SupportsBoundedPointData())
 	{
-		const UPCGPointData* BoundedPointData = nullptr;
+		const UPCGBasePointData* BoundedPointData = nullptr;
 		CacheLock.Lock();
 		check(CachedBoundedPointDataBoxes.Num() == CachedBoundedPointData.Num());
 		for (int CachedDataIndex = 0; CachedDataIndex < CachedBoundedPointDataBoxes.Num(); ++CachedDataIndex)
@@ -90,7 +93,7 @@ const UPCGPointData* UPCGSpatialDataWithPointCache::ToPointData(FPCGContext* Con
 
 		if (!BoundedPointData)
 		{
-			BoundedPointData = CreatePointData(Context, InBounds);
+			BoundedPointData = CreatePointDataFunc(Context, InBounds);
 			CachedBoundedPointDataBoxes.Add(InBounds);
 			CachedBoundedPointData.Add(BoundedPointData);
 		}
@@ -106,7 +109,7 @@ const UPCGPointData* UPCGSpatialDataWithPointCache::ToPointData(FPCGContext* Con
 
 			if (!CachedPointData)
 			{
-				CachedPointData = CreatePointData(Context, InBounds);
+				CachedPointData = CreatePointDataFunc(Context, InBounds);
 			}
 
 			CacheLock.Unlock();
@@ -116,6 +119,22 @@ const UPCGPointData* UPCGSpatialDataWithPointCache::ToPointData(FPCGContext* Con
 	}
 }
 
+const UPCGPointData* UPCGSpatialDataWithPointCache::ToPointData(FPCGContext* Context, const FBox& InBounds) const
+{
+	return CastChecked<UPCGPointData>(ToBasePointDataInternal(Context, InBounds, [this](FPCGContext* InContext, const FBox& InBounds) { return CreatePointData(InContext, InBounds); }));
+}
+
+const UPCGPointArrayData* UPCGSpatialDataWithPointCache::ToPointArrayData(FPCGContext* Context, const FBox& InBounds) const
+{
+	return CastChecked<UPCGPointArrayData>(ToBasePointDataInternal(Context, InBounds, [this](FPCGContext* InContext, const FBox& InBounds) { return CreatePointArrayData(InContext, InBounds); }));
+}
+
+const UPCGPointArrayData* UPCGSpatialDataWithPointCache::CreatePointArrayData(FPCGContext* Context, const FBox& InBounds) const
+{
+	// Default implementation (not optimal, should be overloaded in sub classes)
+	return CreatePointData(Context, InBounds)->ToPointArrayData(Context);
+}
+
 void UPCGSpatialData::GetResourceSizeEx(FResourceSizeEx& CumulativeResourceSize)
 {
 	Super::GetResourceSizeEx(CumulativeResourceSize);
diff --git a/Engine/Plugins/PCG/Source/PCG/Private/Data/PCGSplineData.cpp b/Engine/Plugins/PCG/Source/PCG/Private/Data/PCGSplineData.cpp
--- a/Engine/Plugins/PCG/Source/PCG/Private/Data/PCGSplineData.cpp
+++ b/Engine/Plugins/PCG/Source/PCG/Private/Data/PCGSplineData.cpp
@@ -3,6 +3,7 @@
 #include "Data/PCGSplineData.h"
 
 #include "PCGContext.h"
+#include "Data/PCGPointArrayData.h"
 #include "Data/PCGPointData.h"
 #include "Data/PCGPolyLineData.h"
 #include "Data/PCGProjectionData.h"
@@ -331,14 +332,30 @@ FPCGAttributeAccessorMethods UPCGSplineData::GetSplineAccessorMethods()
 const UPCGPointData* UPCGSplineData::CreatePointData(FPCGContext* Context) const
 {
 	TRACE_CPUPROFILER_EVENT_SCOPE(UPCGSplineData::CreatePointData);
-	UPCGPointData* Data = FPCGContext::NewObject_AnyThread<UPCGPointData>(Context);
-	Data->InitializeFromData(this);
+
+	return CastChecked<UPCGPointData>(CreateBasePointData(Context, UPCGPointData::StaticClass()));
+}
+
+const UPCGPointArrayData* UPCGSplineData::CreatePointArrayData(FPCGContext* Context, const FBox& InBounds) const
+{
+	TRACE_CPUPROFILER_EVENT_SCOPE(UPCGSplineData::CreatePointArrayData);
+
+	return CastChecked<UPCGPointArrayData>(CreateBasePointData(Context, UPCGPointArrayData::StaticClass()));
+}
+
+const UPCGBasePointData* UPCGSplineData::CreateBasePointData(FPCGContext* Context, TSubclassOf<UPCGBasePointData> PointDataClass) const
+{
+	UPCGBasePointData* Data = FPCGContext::NewObject_AnyThread<UPCGBasePointData>(Context, GetTransientPackage(), PointDataClass);
+
+	FPCGInitializeFromDataParams InitializeFromDataParams(this);
+	InitializeFromDataParams.bInheritSpatialData = false;
+	Data->InitializeFromDataWithParams(InitializeFromDataParams);
 
 	FPCGSplineSamplerParams SamplerParams;
 	SamplerParams.Mode = EPCGSplineSamplingMode::Distance;
 
 	PCGSplineSamplerHelpers::SampleLineData(Context, this, /*InBoundingShape=*/nullptr, /*InProjectionTarget=*/nullptr, /*InProjectionParams=*/{}, SamplerParams, Data);
-	UE_LOG(LogPCG, Verbose, TEXT("Spline generated %d points"), Data->GetPoints().Num());
+	UE_LOG(LogPCG, Verbose, TEXT("Spline generated %d points"), Data->GetNumPoints());
 
 	return Data;
 }
diff --git a/Engine/Plugins/PCG/Source/PCG/Private/Data/PCGUnionData.cpp b/Engine/Plugins/PCG/Source/PCG/Private/Data/PCGUnionData.cpp
--- a/Engine/Plugins/PCG/Source/PCG/Private/Data/PCGUnionData.cpp
+++ b/Engine/Plugins/PCG/Source/PCG/Private/Data/PCGUnionData.cpp
@@ -3,6 +3,7 @@
 #include "Data/PCGUnionData.h"
 
 #include "PCGContext.h"
+#include "Data/PCGPointArrayData.h"
 #include "Data/PCGPointData.h"
 #include "Data/PCGSpatialData.h"
 #include "Helpers/PCGAsync.h"
@@ -236,6 +237,18 @@ const UPCGPointData* UPCGUnionData::CreatePointData(FPCGContext* Context) const
 {
 	TRACE_CPUPROFILER_EVENT_SCOPE(UPCGUnionData::CreatePointData);
 
+	return CastChecked<UPCGPointData>(CreateBasePointData(Context, UPCGPointData::StaticClass(), [](FPCGContext* InContext, const UPCGSpatialData* InSpatialData) { return InSpatialData->ToPointData(InContext); }));
+}
+
+const UPCGPointArrayData* UPCGUnionData::CreatePointArrayData(FPCGContext* Context, const FBox& InBounds) const
+{
+	TRACE_CPUPROFILER_EVENT_SCOPE(UPCGUnionData::CreatePointArrayData);
+
+	return CastChecked<UPCGPointArrayData>(CreateBasePointData(Context, UPCGPointArrayData::StaticClass(), [](FPCGContext* InContext, const UPCGSpatialData* InSpatialData) { return InSpatialData->ToPointArrayData(InContext); }));
+}
+
+const UPCGBasePointData* UPCGUnionData::CreateBasePointData(FPCGContext* Context, TSubclassOf<UPCGBasePointData> PointDataClass, TFunctionRef<const UPCGBasePointData*(FPCGContext*, const UPCGSpatialData*)> ToPointDataFunc) const
+{
 	const bool bBinaryDensity = (DensityFunction == EPCGUnionDensityFunction::Binary);
 
 	// Trivial results
@@ -247,7 +260,7 @@ const UPCGPointData* UPCGUnionData::CreatePointData(FPCGContext* Context) const
 	else if (Data.Num() == 1 && !bBinaryDensity)
 	{
 		UE_LOG(LogPCG, Verbose, TEXT("Union is trivial"));
-		return Data[0]->ToPointData(Context);
+		return ToPointDataFunc(Context, Data[0]);
 	}
 
 	// Cache raw pointers for metadata as these are much faster to use
@@ -256,11 +269,15 @@ const UPCGPointData* UPCGUnionData::CreatePointData(FPCGContext* Context) const
 	InputMetadatas.SetNumUninitialized(Data.Num());
 	for (int32 i = 0; i < DataRawPtr.Num(); i++)
 	{
-		InputMetadatas[i] = DataRawPtr[i]->ToPointData(Context)->Metadata;
+		InputMetadatas[i] = ToPointDataFunc(Context, DataRawPtr[i])->Metadata;
 	}
 
-	UPCGPointData* PointData = FPCGContext::NewObject_AnyThread<UPCGPointData>(Context);
-	PointData->InitializeFromData(this, InputMetadatas[0]);
+	UPCGBasePointData* PointData = FPCGContext::NewObject_AnyThread<UPCGBasePointData>(Context, GetTransientPackage(), PointDataClass);
+
+	FPCGInitializeFromDataParams InitializeFromDataParams(this);
+	InitializeFromDataParams.bInheritSpatialData = false;
+	InitializeFromDataParams.MetadataParentOverride = InputMetadatas[0];
+	PointData->InitializeFromDataWithParams(InitializeFromDataParams);
 
 	UPCGMetadata* OutMetadata = PointData->Metadata;
 	check(OutMetadata);
@@ -275,37 +292,82 @@ const UPCGPointData* UPCGUnionData::CreatePointData(FPCGContext* Context) const
 	{
 	case EPCGUnionType::LeftToRightPriority:
 	default:
-		CreateSequentialPointData(Context, DataRawPtr, InputMetadatas, PointData, OutMetadata, /*bLeftToRight=*/true);
+		CreateSequentialPointData(Context, DataRawPtr, InputMetadatas, PointData, OutMetadata, /*bLeftToRight=*/true, ToPointDataFunc);
 		break;
 
 	case EPCGUnionType::RightToLeftPriority:
-		CreateSequentialPointData(Context, DataRawPtr, InputMetadatas, PointData, OutMetadata, /*bLeftToRight=*/false);
+		CreateSequentialPointData(Context, DataRawPtr, InputMetadatas, PointData, OutMetadata, /*bLeftToRight=*/false, ToPointDataFunc);
 		break;
 
 	case EPCGUnionType::KeepAll:
 		{
-			TArray<FPCGPoint>& TargetPoints = PointData->GetMutablePoints();
+			int32 PointInputCount = 0;
+			TArray<const UPCGBasePointData*> PointInputDatas;
+			PointInputDatas.SetNumUninitialized(DataRawPtr.Num());
+			EPCGPointNativeProperties PropertiesToAllocate = EPCGPointNativeProperties::None;
+			for (int32 Index = 0; Index < PointInputDatas.Num(); ++Index)
+			{
+				PointInputDatas[Index] = ToPointDataFunc(Context, DataRawPtr[Index]);
+				PointInputCount += PointInputDatas[Index]->GetNumPoints();
+				PropertiesToAllocate |= PointInputDatas[Index]->GetAllocatedProperties();
+			}
+
+			PointData->SetNumPoints(PointInputCount);
+			PointData->AllocateProperties(PropertiesToAllocate);
+
+			int32 NumWritten = 0;
+
 			for(int32 DataIndex = 0; DataIndex < DataRawPtr.Num(); ++DataIndex)
 			{
 				const UPCGSpatialData* Datum = DataRawPtr[DataIndex];
-				const UPCGPointData* DatumPointData = Datum->ToPointData(Context);
+				const UPCGBasePointData* DatumPointData = PointInputDatas[DataIndex];
 				const UPCGMetadata* DatumPointMetadata = DatumPointData->Metadata;
 
-				int32 TargetPointOffset = TargetPoints.Num();
-				TargetPoints.Append(DatumPointData->GetPoints());
-
-				if (DataIndex > 0 && DatumPointData->GetPoints().Num() > 0)
+				const int32 DatumNumPoints = DatumPointData->GetNumPoints();
+				const int32 TargetPointOffset = NumWritten;
+				
+				if (DatumNumPoints > 0)
 				{
-					// TODO: could optimize case where there is a common parent between Data 0 and current data, for points that still point to common parent metadata.
-					TArrayView<FPCGPoint> TargetPointsSubset = MakeArrayView(&TargetPoints[TargetPointOffset], DatumPointData->GetPoints().Num());
-					for (FPCGPoint& Point : TargetPointsSubset)
-					{
-						Point.MetadataEntry = PCGInvalidEntryKey;
-					}
+					DatumPointData->CopyPointsTo(PointData, 0, TargetPointOffset, DatumNumPoints);
+					NumWritten += DatumNumPoints;
 
-					if (OutMetadata && DatumPointMetadata && DatumPointMetadata->GetAttributeCount() > 0)
+					if (DataIndex > 0)
 					{
-						OutMetadata->SetPointAttributes(MakeArrayView(DatumPointData->GetPoints()), DatumPointMetadata, TargetPointsSubset, Context);
+						TPCGValueRange<int64> TargetMetadataEntryRange = PointData->GetMetadataEntryValueRange();
+						
+						// TODO: could optimize case where there is a common parent between Data 0 and current data, for points that still point to common parent metadata.
+						for (int32 TargetIndex = TargetPointOffset; TargetIndex < TargetPointOffset + DatumNumPoints; ++TargetIndex)
+						{
+							TargetMetadataEntryRange[TargetIndex] = PCGInvalidEntryKey;
+						}
+
+						if (OutMetadata && DatumPointMetadata && DatumPointMetadata->GetAttributeCount() > 0)
+						{
+							TArray<PCGMetadataEntryKey, TInlineAllocator<256>> DatumKeys;
+							DatumKeys.SetNumUninitialized(DatumNumPoints);
+
+							const TConstPCGValueRange<int64> DatumMetadataEntryRange = DatumPointData->GetConstMetadataEntryValueRange();
+							for (int32 DatumIndex = 0; DatumIndex < DatumKeys.Num(); ++DatumIndex)
+							{
+								DatumKeys[DatumIndex] = DatumMetadataEntryRange[DatumIndex];
+							}
+
+							TArray<PCGMetadataEntryKey, TInlineAllocator<256>> TargetKeys;
+							TargetKeys.SetNumUninitialized(DatumNumPoints);
+
+							for (int32 TargetIndex = 0; TargetIndex < TargetKeys.Num(); ++TargetIndex)
+							{
+								TargetKeys[TargetIndex] = PCGInvalidEntryKey;
+							}
+
+							PointData->Metadata->SetAttributes(DatumKeys, DatumPointMetadata, TargetKeys, Context);
+
+							// Write back
+							for (int32 TargetIndex = 0; TargetIndex < TargetKeys.Num(); ++TargetIndex)
+							{
+								TargetMetadataEntryRange[TargetPointOffset + TargetIndex] = TargetKeys[TargetIndex];
+							}
+						}
 					}
 				}
 			}
@@ -313,49 +375,93 @@ const UPCGPointData* UPCGUnionData::CreatePointData(FPCGContext* Context) const
 			// Correct density for binary-style union
 			if (bBinaryDensity)
 			{
-				for (FPCGPoint& TargetPoint : TargetPoints)
+				if (EnumHasAllFlags(PointData->GetAllocatedProperties(), EPCGPointNativeProperties::Density))
+				{
+					TPCGValueRange<float> DensityRange = PointData->GetDensityValueRange(/*bAllocate=*/false);
+					for (float& Density : DensityRange)
+					{
+						Density = ((Density > 0) ? 1.0f : 0);
+					}
+				}
+				else
 				{
-					TargetPoint.Density = ((TargetPoint.Density > 0) ? 1.0f : 0);
+					PointData->SetDensity(PointData->GetDensity(0) > 0 ? 1.0f : 0);
 				}
 			}
 		}
 		break;
 	}
 
-	UE_LOG(LogPCG, Verbose, TEXT("Union generated %d points out of %d data sources"), PointData->GetPoints().Num(), Data.Num());
+	UE_LOG(LogPCG, Verbose, TEXT("Union generated %d points out of %d data sources"), PointData->GetNumPoints(), Data.Num());
 
 	return PointData;
 }
 
-void UPCGUnionData::CreateSequentialPointData(FPCGContext* Context, TArray<const UPCGSpatialData*>& InputDatas, TArray<const UPCGMetadata*>& InputMetadatas, UPCGPointData* PointData, UPCGMetadata* OutMetadata, bool bLeftToRight) const
+void UPCGUnionData::CreateSequentialPointData(FPCGContext* Context, TArray<const UPCGSpatialData*>& InputDatas, TArray<const UPCGMetadata*>& InputMetadatas, UPCGBasePointData* OutPointData, UPCGMetadata* OutMetadata, bool bLeftToRight, TFunctionRef<const UPCGBasePointData* (FPCGContext*, const UPCGSpatialData*)> ToPointDataFunc) const
 {
-	check(PointData);
+	check(OutPointData);
+		
+	int32 PointInputCount = 0;
+	
+	TArray<const UPCGBasePointData*> PointInputDatas;
+	PointInputDatas.SetNumUninitialized(InputDatas.Num());
+	EPCGPointNativeProperties PropertiesToAllocate = EPCGPointNativeProperties::None;
+	for (int32 Index = 0; Index < PointInputDatas.Num(); ++Index)
+	{
+		PointInputDatas[Index] = ToPointDataFunc(Context, InputDatas[Index]);
+		PointInputCount += PointInputDatas[Index]->GetNumPoints();
+		PropertiesToAllocate |= PointInputDatas[Index]->GetAllocatedProperties();
+	}
+	
+	OutPointData->SetNumPoints(PointInputCount);
+	OutPointData->AllocateProperties(PropertiesToAllocate);
 
-	TArray<FPCGPoint>& TargetPoints = PointData->GetMutablePoints();
-	TArray<FPCGPoint> SelectedDataPoints;
+	struct FIndexParams
+	{
+		FIndexParams(const int32 InFirstDataIndex, const int32 InLastDataIndex, const int32 InDataIndexIncrement)
+			: FirstDataIndex(InFirstDataIndex), LastDataIndex(InLastDataIndex), DataIndexIncrement(InDataIndexIncrement), CurrentIndex(0) {}
 
-	int32 FirstDataIndex = (bLeftToRight ? 0 : InputDatas.Num() - 1);
-	int32 LastDataIndex = (bLeftToRight ? InputDatas.Num() : -1);
-	int32 DataIndexIncrement = (bLeftToRight ? 1 : -1);
+		const int32 FirstDataIndex;
+		const int32 LastDataIndex;
+		const int32 DataIndexIncrement;
+		int32 CurrentIndex;
+	};
 
-	// Note: this is a O(N^2) implementation. 
-	// TODO: It is easy to implement a kind of divide & conquer algorithm here, but it will require some temporary storage.
-	for (int32 DataIndex = FirstDataIndex; DataIndex != LastDataIndex; DataIndex += DataIndexIncrement)
+	FIndexParams IndexParams(
+		/*InFirstDataIndex=*/bLeftToRight ? 0 : InputDatas.Num() - 1,
+		/*InLastDataIndex=*/bLeftToRight ? InputDatas.Num() : -1,
+		/*InDataIndexIncrement=*/bLeftToRight ? 1 : -1);
+
+	int32 PointOffset = 0;
+
+	auto MoveDataRange = [OutPointData, &PointOffset](int32 ReadIndex, int32 WriteIndex, int32 Count)
+	{
+		OutPointData->MoveRange(ReadIndex + PointOffset, WriteIndex + PointOffset, Count);
+	};
+
+	auto Finished = [&PointOffset](int32 Count)
+	{
+		PointOffset += Count;
+	};
+		
+	auto ProcessRange = [this, &PointOffset, &PointInputDatas, &InputDatas, &InputMetadatas, OutPointData, OutMetadata](int32 StartReadIndex, int32 StartWriteIndex, int32 Count, FIndexParams& IndexParams)
 	{
-		// For each point, if it is not already "processed" by previous data,
-		// add it & compute its final density
-		const UPCGPointData* CurrentPointData = InputDatas[DataIndex]->ToPointData(Context);
-		const TArray<FPCGPoint>& Points = CurrentPointData->GetPoints();
+		const UPCGBasePointData* CurrentPointData = PointInputDatas[IndexParams.CurrentIndex];
+				
+		FConstPCGPointValueRanges InRanges(CurrentPointData);
+		FPCGPointValueRanges OutRanges(OutPointData, /*bAllocate=*/false);
+
+		int32 NumWritten = 0;
 
-		FPCGAsync::AsyncPointProcessing(Context, Points.Num(), SelectedDataPoints, [this, PointData, OutMetadata, &Points, &InputDatas, &InputMetadatas, DataIndex, FirstDataIndex, LastDataIndex, DataIndexIncrement](int32 Index, FPCGPoint& OutPoint)
+		for (int32 ReadIndex = StartReadIndex; ReadIndex < StartReadIndex + Count; ++ReadIndex)
 		{
-			const FPCGPoint& Point = Points[Index];
+			const int32 WriteIndex = StartWriteIndex + NumWritten;
 
 			// Discard point if it is already covered by a previous data
 			bool bPointToExclude = false;
-			for (int32 PreviousDataIndex = FirstDataIndex; PreviousDataIndex != DataIndex; PreviousDataIndex += DataIndexIncrement)
+			for (int32 PreviousDataIndex = IndexParams.FirstDataIndex; PreviousDataIndex != IndexParams.CurrentIndex; PreviousDataIndex += IndexParams.DataIndexIncrement)
 			{
-				if (InputDatas[PreviousDataIndex]->GetDensityAtPosition(Point.Transform.GetLocation()) != 0)
+				if (InputDatas[PreviousDataIndex]->GetDensityAtPosition(InRanges.TransformRange[ReadIndex].GetLocation()) != 0)
 				{
 					bPointToExclude = true;
 					break;
@@ -364,67 +470,96 @@ void UPCGUnionData::CreateSequentialPointData(FPCGContext* Context, TArray<const
 
 			if (bPointToExclude)
 			{
-				return false;
+				continue;
 			}
-
-			check(PointData);
+						
 			check(OutMetadata);
+						
+			OutRanges.SetFromValueRanges(WriteIndex, InRanges, ReadIndex);
+			++NumWritten;
 
-			OutPoint = Point;
-			if (OutMetadata->GetParent() != InputMetadatas[DataIndex])
+			if (OutMetadata->GetParent() != InputMetadatas[IndexParams.CurrentIndex])
 			{
-				UPCGMetadataAccessorHelpers::InitializeMetadata(OutPoint, OutMetadata);
+				OutRanges.MetadataEntryRange[WriteIndex] = OutMetadata->AddEntry();
+								
 				// Since we can't inherit from the parent point, we'll set the values directly here
-				OutMetadata->SetPointAttributes(Point, InputMetadatas[DataIndex], OutPoint);
+				OutMetadata->SetAttributes(InRanges.MetadataEntryRange[ReadIndex], InputMetadatas[IndexParams.CurrentIndex], OutRanges.MetadataEntryRange[WriteIndex]);
 			}
 
-			if (DensityFunction == EPCGUnionDensityFunction::Binary && OutPoint.Density > 0)
+			if (DensityFunction == EPCGUnionDensityFunction::Binary && OutRanges.DensityRange[WriteIndex] > 0)
 			{
-				OutPoint.Density = 1.0f;
+				OutRanges.DensityRange[WriteIndex] = 1.0f;
 			}
 
 			// Update density & metadata based on current & following data
-			for (int32 FollowingDataIndex = DataIndex + DataIndexIncrement; FollowingDataIndex != LastDataIndex; FollowingDataIndex += DataIndexIncrement)
+			for (int32 FollowingDataIndex = IndexParams.CurrentIndex + IndexParams.DataIndexIncrement; FollowingDataIndex != IndexParams.LastDataIndex; FollowingDataIndex += IndexParams.DataIndexIncrement)
 			{
 				const UPCGMetadata* FollowingMetadata = InputMetadatas[FollowingDataIndex];
 
 				// If density is saturated and there are no metadata attributes then we can skip this data as it will not contribute.
-				if (OutPoint.Density >= 1.0f && (!FollowingMetadata || FollowingMetadata->GetAttributeCount() == 0))
+				if (OutRanges.DensityRange[WriteIndex] >= 1.0f && (!FollowingMetadata || FollowingMetadata->GetAttributeCount() == 0))
 				{
 					continue;
 				}
 
 				FPCGPoint PointInData;
-				if (InputDatas[FollowingDataIndex]->SamplePoint(OutPoint.Transform, OutPoint.GetLocalBounds(), PointInData, OutMetadata))
+
+				const FBox LocalBounds = PCGPointHelpers::GetLocalBounds(OutRanges.BoundsMinRange[WriteIndex], OutRanges.BoundsMaxRange[WriteIndex]);
+				if (InputDatas[FollowingDataIndex]->SamplePoint(OutRanges.TransformRange[WriteIndex], LocalBounds, PointInData, OutMetadata))
 				{
 					// Update density
-					PCGUnionDataMaths::UpdateDensity(OutPoint.Density, PointInData.Density, DensityFunction);
+					PCGUnionDataMaths::UpdateDensity(OutRanges.DensityRange[WriteIndex], PointInData.Density, DensityFunction);
+
+					const FVector4& InColor = InRanges.ColorRange[ReadIndex];
+					FVector4& OutColor = OutRanges.ColorRange[WriteIndex];
 
 					// Update color
-					OutPoint.Color = FVector4(
-						FMath::Max(OutPoint.Color.X, PointInData.Color.X),
-						FMath::Max(OutPoint.Color.Y, PointInData.Color.Y),
-						FMath::Max(OutPoint.Color.Z, PointInData.Color.Z),
-						FMath::Max(OutPoint.Color.W, PointInData.Color.W));
+					OutColor = FVector4(
+						FMath::Max(OutColor.X, InColor.X),
+						FMath::Max(OutColor.Y, InColor.Y),
+						FMath::Max(OutColor.Z, InColor.Z),
+						FMath::Max(OutColor.W, InColor.W));
 
-					if (OutPoint.MetadataEntry != PCGInvalidEntryKey && PointInData.MetadataEntry != PCGInvalidEntryKey)
+					if (OutRanges.MetadataEntryRange[WriteIndex] != PCGInvalidEntryKey && PointInData.MetadataEntry != PCGInvalidEntryKey)
 					{
-						OutMetadata->MergePointAttributesSubset(OutPoint, OutMetadata, OutMetadata, PointInData, OutMetadata, FollowingMetadata, OutPoint, EPCGMetadataOp::Max);
+						OutMetadata->MergeAttributesSubset(OutRanges.MetadataEntryRange[WriteIndex], OutMetadata, OutMetadata, PointInData.MetadataEntry, OutMetadata, FollowingMetadata, OutRanges.MetadataEntryRange[WriteIndex], EPCGMetadataOp::Max);
 					}
 					else if (PointInData.MetadataEntry != PCGInvalidEntryKey)
 					{
-						OutPoint.MetadataEntry = PointInData.MetadataEntry;
+						OutRanges.MetadataEntryRange[WriteIndex] = PointInData.MetadataEntry;
 					}
 				}
 			}
+		}
 
-			return true;
-		});
+		return NumWritten;
+	};
 
-		// Append current iteration results to target points
-		TargetPoints += SelectedDataPoints;
-		SelectedDataPoints.Reset();
+	// Note: this is a O(N^2) implementation. 
+	// TODO: It is easy to implement a kind of divide & conquer algorithm here, but it will require some temporary storage.
+	for (int32 DataIndex = IndexParams.FirstDataIndex; DataIndex != IndexParams.LastDataIndex; DataIndex += IndexParams.DataIndexIncrement)
+	{
+		IndexParams.CurrentIndex = DataIndex;
+
+		auto ProcessRangeCurrentIndex = [&IndexParams, &ProcessRange, &PointOffset](int32 StartReadIndex, int32 StartWriteIndex, int32 Count) -> int32
+		{
+			return ProcessRange(StartReadIndex, StartWriteIndex + PointOffset, Count, IndexParams);
+		};
+
+		const UPCGBasePointData* CurrentPointData = PointInputDatas[DataIndex];
+
+		FPCGAsync::AsyncProcessingRangeEx(
+			&Context->AsyncState,
+			CurrentPointData->GetNumPoints(),
+			[]() {},
+			ProcessRangeCurrentIndex,
+			MoveDataRange,
+			Finished,
+			/*bEnableTimeSlicing=*/false);
 	}
+
+	// Set final point count
+	OutPointData->SetNumPoints(PointOffset);
 }
 
 UPCGSpatialData* UPCGUnionData::CopyInternal(FPCGContext* Context) const
diff --git a/Engine/Plugins/PCG/Source/PCG/Private/Data/PCGVolumeData.cpp b/Engine/Plugins/PCG/Source/PCG/Private/Data/PCGVolumeData.cpp
--- a/Engine/Plugins/PCG/Source/PCG/Private/Data/PCGVolumeData.cpp
+++ b/Engine/Plugins/PCG/Source/PCG/Private/Data/PCGVolumeData.cpp
@@ -1,6 +1,7 @@
 // Copyright Epic Games, Inc. All Rights Reserved.
 
 #include "Data/PCGVolumeData.h"
+#include "Data/PCGPointArrayData.h"
 #include "Data/PCGPointData.h"
 #include "Data/PCGSpatialData.h"
 #include "Elements/PCGVolumeSampler.h"
@@ -142,15 +143,27 @@ const UPCGPointData* UPCGVolumeData::CreatePointData(FPCGContext* Context) const
 {
 	TRACE_CPUPROFILER_EVENT_SCOPE(UPCGVolumeData::CreatePointData);
 
+	return Cast<UPCGPointData>(CreateBasePointData(Context, UPCGPointData::StaticClass()));
+}
+
+const UPCGPointArrayData* UPCGVolumeData::CreatePointArrayData(FPCGContext* Context, const FBox& InBounds) const
+{
+	TRACE_CPUPROFILER_EVENT_SCOPE(UPCGVolumeData::CreatePointArrayData);
+
+	return Cast<UPCGPointArrayData>(CreateBasePointData(Context, UPCGPointArrayData::StaticClass()));
+}
+
+const UPCGBasePointData* UPCGVolumeData::CreateBasePointData(FPCGContext* Context, TSubclassOf<UPCGBasePointData> PointDataClass) const
+{
 	PCGVolumeSampler::FVolumeSamplerParams SamplerParams;
 	SamplerParams.VoxelSize = VoxelSize;
 	SamplerParams.Bounds = GetBounds();
 
-	const UPCGPointData* Data = PCGVolumeSampler::SampleVolume(Context, SamplerParams, this);
+	const UPCGBasePointData* Data = PCGVolumeSampler::SampleVolume(Context, PointDataClass, SamplerParams, this);
 
 	if (Data)
 	{
-		UE_LOG(LogPCG, Verbose, TEXT("Volume extracted %d points"), Data->GetPoints().Num());
+		UE_LOG(LogPCG, Verbose, TEXT("Volume extracted %d points"), Data->GetNumPoints());
 	}
 
 	return Data;
diff --git a/Engine/Plugins/PCG/Source/PCG/Private/Data/PCGWorldData.cpp b/Engine/Plugins/PCG/Source/PCG/Private/Data/PCGWorldData.cpp
--- a/Engine/Plugins/PCG/Source/PCG/Private/Data/PCGWorldData.cpp
+++ b/Engine/Plugins/PCG/Source/PCG/Private/Data/PCGWorldData.cpp
@@ -4,6 +4,7 @@
 
 #include "PCGComponent.h"
 #include "PCGSubsystem.h"
+#include "Data/PCGPointArrayData.h"
 #include "Data/PCGPointData.h"
 #include "Data/PCGSurfaceData.h"
 #include "Data/PCGVolumeData.h"
@@ -207,11 +208,17 @@ bool UPCGWorldVolumetricData::SamplePoint(const FTransform& InTransform, const F
 const UPCGPointData* UPCGWorldVolumetricData::CreatePointData(FPCGContext* Context, const FBox& InBounds) const
 {
 	TRACE_CPUPROFILER_EVENT_SCOPE(UPCGWorldVolumetricData::CreatePointData);
+	return Cast<UPCGPointData>(CreateBasePointData(Context, InBounds, UPCGPointData::StaticClass()));
+}
 
-	UPCGPointData* Data = FPCGContext::NewObject_AnyThread<UPCGPointData>(Context);
-	Data->InitializeFromData(this);
-	TArray<FPCGPoint>& Points = Data->GetMutablePoints();
+const UPCGPointArrayData* UPCGWorldVolumetricData::CreatePointArrayData(FPCGContext* Context, const FBox& InBounds) const
+{
+	TRACE_CPUPROFILER_EVENT_SCOPE(UPCGWorldVolumetricData::CreatePointArrayData);
+	return Cast<UPCGPointArrayData>(CreateBasePointData(Context, InBounds, UPCGPointArrayData::StaticClass()));
+}
 
+const UPCGBasePointData* UPCGWorldVolumetricData::CreateBasePointData(FPCGContext* Context, const FBox& InBounds, TSubclassOf<UPCGBasePointData> PointDataClass) const
+{		
 	FBox EffectiveBounds = Bounds;
 	if (InBounds.IsValid)
 	{
@@ -232,6 +239,12 @@ const UPCGPointData* UPCGWorldVolumetricData::CreatePointData(FPCGContext* Conte
 		{
 			UE_LOG(LogPCG, Error, TEXT("PCG World Volumetric Data cannot generate without sampling bounds. Consider using a Volume Sampler with the Unbounded option disabled."));
 		}
+
+		UPCGBasePointData* Data = FPCGContext::NewObject_AnyThread<UPCGBasePointData>(Context, GetTransientPackage(), PointDataClass);
+		
+		FPCGInitializeFromDataParams InitializeFromDataParams(this);
+		InitializeFromDataParams.bInheritSpatialData = false;
+		Data->InitializeFromDataWithParams(InitializeFromDataParams);
 		
 		return Data;
 	}
@@ -240,8 +253,8 @@ const UPCGPointData* UPCGWorldVolumetricData::CreatePointData(FPCGContext* Conte
 	SamplerParams.VoxelSize = VoxelSize;
 	SamplerParams.Bounds = EffectiveBounds;
 
-	Data = PCGVolumeSampler::SampleVolume(Context, SamplerParams, this);
-	UE_LOG(LogPCG, Verbose, TEXT("Volumetric world extracted %d points"), Data->GetPoints().Num());
+	UPCGBasePointData* Data = PCGVolumeSampler::SampleVolume(Context, PointDataClass, SamplerParams, this);
+	UE_LOG(LogPCG, Verbose, TEXT("Volumetric world extracted %d points"), Data->GetNumPoints());
 
 	return Data;
 }
diff --git a/Engine/Plugins/PCG/Source/PCG/Private/Elements/PCGDifferenceElement.cpp b/Engine/Plugins/PCG/Source/PCG/Private/Elements/PCGDifferenceElement.cpp
--- a/Engine/Plugins/PCG/Source/PCG/Private/Elements/PCGDifferenceElement.cpp
+++ b/Engine/Plugins/PCG/Source/PCG/Private/Elements/PCGDifferenceElement.cpp
@@ -7,7 +7,7 @@
 #include "PCGEdge.h"
 #include "PCGGraph.h"
 #include "PCGPin.h"
-#include "Data/PCGPointData.h"
+#include "Data/PCGBasePointData.h"
 #include "Data/PCGSpatialData.h"
 #include "Data/PCGUnionData.h"
 #include "Elements/PCGUnionElement.h"
@@ -164,7 +164,7 @@ bool FPCGDifferenceElement::ExecuteInternal(FPCGContext* Context) const
 		if (Difference.Data && Difference.Data->IsA<UPCGSpatialData>())
 		{
 			bDifferenceDataHasValidData = true;
-			bHasPointsInDifferences |= Difference.Data->IsA<UPCGPointData>();
+			bHasPointsInDifferences |= Difference.Data->IsA<UPCGBasePointData>();
 			break;
 		}
 	}
@@ -188,7 +188,7 @@ bool FPCGDifferenceElement::ExecuteInternal(FPCGContext* Context) const
 			continue;
 		}
 
-		bHasPointsInSource |= SourceSpatialData->IsA<UPCGPointData>();
+		bHasPointsInSource |= SourceSpatialData->IsA<UPCGBasePointData>();
 
 		UPCGDifferenceData* DifferenceData = FPCGContext::NewObject_AnyThread<UPCGDifferenceData>(Context);
 		DifferenceData->Initialize(SourceSpatialData);
@@ -212,7 +212,7 @@ bool FPCGDifferenceElement::ExecuteInternal(FPCGContext* Context) const
 		if ((Settings->Mode == EPCGDifferenceMode::Discrete || 
 			(Settings->Mode == EPCGDifferenceMode::Inferred && bHasPointsInSource && bHasPointsInDifferences)))
 		{
-			Output.Data = DifferenceData->ToPointData(Context);
+			Output.Data = DifferenceData->ToBasePointData(Context);
 		}
 	}
 
diff --git a/Engine/Plugins/PCG/Source/PCG/Private/Elements/PCGVolumeSampler.cpp b/Engine/Plugins/PCG/Source/PCG/Private/Elements/PCGVolumeSampler.cpp
--- a/Engine/Plugins/PCG/Source/PCG/Private/Elements/PCGVolumeSampler.cpp
+++ b/Engine/Plugins/PCG/Source/PCG/Private/Elements/PCGVolumeSampler.cpp
@@ -23,8 +23,16 @@ namespace PCGVolumeSampler
 {
 	UPCGPointData* SampleVolume(FPCGContext* Context, const FVolumeSamplerParams& SamplerSettings, const UPCGSpatialData* Volume, const UPCGSpatialData* BoundingShape)
 	{
-		UPCGPointData* Data = FPCGContext::NewObject_AnyThread<UPCGPointData>(Context);
-		Data->InitializeFromData(Volume);
+		return Cast<UPCGPointData>(SampleVolume(Context, UPCGPointData::StaticClass(), SamplerSettings, Volume, BoundingShape));
+	}
+
+	UPCGBasePointData* SampleVolume(FPCGContext* Context, TSubclassOf<UPCGBasePointData> PointDataClass, const FVolumeSamplerParams& SamplerSettings, const UPCGSpatialData* Volume, const UPCGSpatialData* BoundingShape)
+	{
+		UPCGBasePointData* Data = FPCGContext::NewObject_AnyThread<UPCGBasePointData>(Context, GetTransientPackage(), PointDataClass);
+
+		FPCGInitializeFromDataParams InitializeFromDataParams(Volume);
+		InitializeFromDataParams.bInheritSpatialData = false;
+		Data->InitializeFromDataWithParams(InitializeFromDataParams);
 
 		const bool bTimeSliceIsEnabled = Context ? Context->TimeSliceIsEnabled() : false;
 		SampleVolume(Context, SamplerSettings, Volume, BoundingShape, Data, bTimeSliceIsEnabled);
@@ -32,7 +40,7 @@ namespace PCGVolumeSampler
 		return Data;
 	}
 
-	bool SampleVolume(FPCGContext* Context, const FVolumeSamplerParams& SamplerSettings, const UPCGSpatialData* Volume, const UPCGSpatialData* BoundingShape, UPCGPointData* OutputData, const bool bTimeSlicingIsEnabled)
+	bool SampleVolume(FPCGContext* Context, const FVolumeSamplerParams& SamplerSettings, const UPCGSpatialData* Volume, const UPCGSpatialData* BoundingShape, UPCGBasePointData* OutputData, const bool bTimeSlicingIsEnabled)
 	{
 		check(Volume && OutputData);
 
@@ -47,7 +55,6 @@ namespace PCGVolumeSampler
 			}
 		}
 
-		TArray<FPCGPoint>& Points = OutputData->GetMutablePoints();
 		const FVector& VoxelSize = SamplerSettings.VoxelSize;
 
 		const int32 MinX = FMath::CeilToInt(Bounds.Min.X / VoxelSize.X);
@@ -93,45 +100,78 @@ namespace PCGVolumeSampler
 			}
 		}
 
-		auto AsyncProcessingFunc = [Volume, BoundingShape, PointSteepness = SamplerSettings.PointSteepness, VoxelSize, MinX, MaxX, MinY, MaxY, MinZ](int32 Index, FPCGPoint& OutPoint)
+		auto InitializeFunc = [NumIterations, OutputData]()
 		{
-			const int X = MinX + (Index % (1 + MaxX - MinX));
-			const int Y = MinY + (Index / (1 + MaxX - MinX) % (1 + MaxY - MinY));
-			const int Z = MinZ + (Index / ((1 + MaxX - MinX) * (1 + MaxY - MinY)));
+			OutputData->SetNumPoints(NumIterations);
+			// Volume->SamplePoint could write to all properties
+			OutputData->AllocateProperties(EPCGPointNativeProperties::All);
+		};
 
-			const FVector SampleLocation(X * VoxelSize.X, Y * VoxelSize.Y, Z * VoxelSize.Z);
-			const FBox VoxelBox(VoxelSize * -0.5, VoxelSize * 0.5);
+		auto MoveDataRangeFunc = [OutputData](int32 RangeStartIndex, int32 MoveToIndex, int32 NumElements)
+		{
+			OutputData->MoveRange(RangeStartIndex, MoveToIndex, NumElements);
+		};
 
-			const FTransform SampleTransform = FTransform(SampleLocation);
+		auto FinishedFunc = [OutputData](int32 NumWritten)
+		{
+			OutputData->SetNumPoints(NumWritten);
+		};
+
+		auto ProcessRangeFunc = [OutputData, Volume, BoundingShape, PointSteepness = SamplerSettings.PointSteepness, VoxelSize, MinX, MaxX, MinY, MaxY, MinZ](int32 StartReadIndex, int32 StartWriteIndex, int32 Count)
+		{
+			int32 NumWritten = 0;
 
-			// The OutPoint has not been initialized, so do it now
-			OutPoint = FPCGPoint();
+			FPCGPointValueRanges OutRanges(OutputData, /*bAllocate=*/false);
 
-			if (Volume->SamplePoint(SampleTransform, VoxelBox, OutPoint, nullptr))
+			for (int32 ReadIndex = StartReadIndex; ReadIndex < StartReadIndex + Count; ++ReadIndex)
 			{
-				if (BoundingShape)
+				const int32 WriteIndex = StartWriteIndex + NumWritten;
+
+				const int X = MinX + (ReadIndex % (1 + MaxX - MinX));
+				const int Y = MinY + (ReadIndex / (1 + MaxX - MinX) % (1 + MaxY - MinY));
+				const int Z = MinZ + (ReadIndex / ((1 + MaxX - MinX) * (1 + MaxY - MinY)));
+
+				const FVector SampleLocation(X * VoxelSize.X, Y * VoxelSize.Y, Z * VoxelSize.Z);
+				const FBox VoxelBox(VoxelSize * -0.5, VoxelSize * 0.5);
+
+				const FTransform SampleTransform = FTransform(SampleLocation);
+
+				FPCGPoint OutPoint;
+
+				if (Volume->SamplePoint(SampleTransform, VoxelBox, OutPoint, nullptr))
 				{
-					FPCGPoint BoundingShapeSample;
-					if (!BoundingShape->SamplePoint(SampleTransform, VoxelBox, BoundingShapeSample, nullptr))
+					if (BoundingShape)
 					{
-						return false;
+						FPCGPoint BoundingShapeSample;
+						if (!BoundingShape->SamplePoint(SampleTransform, VoxelBox, BoundingShapeSample, nullptr))
+						{
+							continue;
+						}
 					}
-				}
 
-				OutPoint.Seed = PCGHelpers::ComputeSeed(X, Y, Z);
-				OutPoint.Steepness = PointSteepness;
-				// Reset the bounds to the expected bounds, if it was modified by the Sample Point
-				OutPoint.SetLocalBounds(VoxelBox);
-				return true;
-			}
-			else
-			{
-				return false;
+					OutPoint.Seed = PCGHelpers::ComputeSeed(X, Y, Z);
+					OutPoint.Steepness = PointSteepness;
+					// Reset the bounds to the expected bounds, if it was modified by the Sample Point
+					OutPoint.SetLocalBounds(VoxelBox);
+
+					OutRanges.SetFromPoint(WriteIndex, OutPoint);
+					++NumWritten;
+				}
 			}
+
+			return NumWritten;
 		};
 
 		FPCGAsyncState* AsyncState = Context ? &Context->AsyncState : nullptr;
-		return FPCGAsync::AsyncProcessing<FPCGPoint>(AsyncState, NumIterations, Points, AsyncProcessingFunc, /*bEnableTimeSlicing=*/Context && bTimeSlicingIsEnabled);
+		
+		return FPCGAsync::AsyncProcessingRangeEx(
+			AsyncState,
+			NumIterations,
+			InitializeFunc,
+			ProcessRangeFunc,
+			MoveDataRangeFunc,
+			FinishedFunc,
+			/*bEnableTimeSlicing=*/Context && bTimeSlicingIsEnabled);
 	}
 }
 UPCGVolumeSamplerSettings::UPCGVolumeSamplerSettings()
@@ -279,9 +319,15 @@ bool FPCGVolumeSamplerElement::PrepareDataInternal(FPCGContext* Context) const
 			check(GeneratingShape);
 
 			OutState.Volume = GeneratingShape;
-			OutState.OutputData = FPCGContext::NewObject_AnyThread<UPCGPointData>(Context);
-			OutState.OutputData->InitializeFromData(OutState.Volume);
-			Outputs[IterationIndex].Data = OutState.OutputData;
+			OutState.OutputPointData = FPCGContext::NewPointData_AnyThread(Context);
+PRAGMA_DISABLE_DEPRECATION_WARNINGS
+			OutState.OutputData = Cast<UPCGPointData>(OutState.OutputPointData);
+PRAGMA_ENABLE_DEPRECATION_WARNINGS
+
+			FPCGInitializeFromDataParams InitializeFromDataParams(OutState.Volume);
+			InitializeFromDataParams.bInheritSpatialData = false;
+			OutState.OutputPointData->InitializeFromDataWithParams(InitializeFromDataParams);
+			Outputs[IterationIndex].Data = OutState.OutputPointData;
 
 			FBox& InputBounds = OutState.Settings.Bounds;
 
@@ -351,8 +397,12 @@ bool FPCGVolumeSamplerElement::ExecuteInternal(FPCGContext* Context) const
 		{
 			// TODO: Empty point data (to preserve previous behavior). Eventually, should be replaced with no output at all
 			FPCGTaggedData& Output = TimeSlicedContext->OutputData.TaggedData.Add_GetRef(Input);
-			UPCGPointData* PointData = FPCGContext::NewObject_AnyThread<UPCGPointData>(Context);
-			PointData->InitializeFromData(Cast<UPCGSpatialData>(Input.Data));
+			UPCGBasePointData* PointData = FPCGContext::NewPointData_AnyThread(Context);
+			
+			FPCGInitializeFromDataParams InitializeFromDataParams(Cast<UPCGSpatialData>(Input.Data));
+			InitializeFromDataParams.bInheritSpatialData = false;
+			PointData->InitializeFromDataWithParams(InitializeFromDataParams);
+
 			Output.Data = PointData;
 		}
 
@@ -365,7 +415,7 @@ bool FPCGVolumeSamplerElement::ExecuteInternal(FPCGContext* Context) const
 
 		if (InitResult == EPCGTimeSliceInitResult::NoOperation)
 		{
-			Context->OutputData.TaggedData[IterationIndex].Data = FPCGContext::NewObject_AnyThread<UPCGPointData>(Context);
+			Context->OutputData.TaggedData[IterationIndex].Data = FPCGContext::NewPointData_AnyThread(Context);
 
 			return true;
 		}
@@ -377,10 +427,10 @@ bool FPCGVolumeSamplerElement::ExecuteInternal(FPCGContext* Context) const
 				IterState.Settings,
 				ExecState.GeneratingShapes[IterationIndex],
 				ExecState.BoundingShape,
-				IterState.OutputData,
+				IterState.OutputPointData,
 				Context->TimeSliceIsEnabled());
 
-		PCGE_LOG_C(Verbose, LogOnly, Context, FText::Format(LOCTEXT("GenerationInfo", "Generated {0} points in volume"), IterState.OutputData->GetPoints().Num()));
+		PCGE_LOG_C(Verbose, LogOnly, Context, FText::Format(LOCTEXT("GenerationInfo", "Generated {0} points in volume"), IterState.OutputPointData->GetNumPoints()));
 
 		return bAsyncDone;
 	});
diff --git a/Engine/Plugins/PCG/Source/PCG/Private/Tests/Determinism/PCGDeterminismTestsCommon.cpp b/Engine/Plugins/PCG/Source/PCG/Private/Tests/Determinism/PCGDeterminismTestsCommon.cpp
--- a/Engine/Plugins/PCG/Source/PCG/Private/Tests/Determinism/PCGDeterminismTestsCommon.cpp
+++ b/Engine/Plugins/PCG/Source/PCG/Private/Tests/Determinism/PCGDeterminismTestsCommon.cpp
@@ -926,7 +926,7 @@ namespace PCGDeterminismTests
 		{
 			return false;
 		}
-		else if (BothDataCastsToDataType<const UPCGPointData>(FirstData, SecondData))
+		else if (BothDataCastsToDataType<const UPCGBasePointData>(FirstData, SecondData))
 		{
 			return PointDataIsIdentical(FirstData, SecondData);
 		}
@@ -955,32 +955,69 @@ namespace PCGDeterminismTests
 		return false;
 	}
 
+	template<class T>
+	bool CompareProperty(const UPCGBasePointData* FirstPointData, const UPCGBasePointData* SecondPointData, EPCGPointNativeProperties Property)
+	{
+		const TConstPCGValueRange<T> FirstRange = FirstPointData->GetConstValueRange<T>(Property);
+		const TConstPCGValueRange<T> SecondRange = SecondPointData->GetConstValueRange<T>(Property);
+		check(FirstRange.Num() == SecondRange.Num());
+
+		for (int32 Index = 0; Index < FirstRange.Num(); ++Index)
+		{
+			if (!PCG::Private::MetadataTraits<T>::Equal(FirstRange[Index], SecondRange[Index]))
+			{
+				return false;
+			}
+		}
+		
+		return true;
+	}
+
 	bool PointDataIsIdentical(const UPCGData* FirstData, const UPCGData* SecondData)
 	{
 		TRACE_CPUPROFILER_EVENT_SCOPE(PCGDeterminismTests::PointDataIsIdentical);
 
-		const UPCGPointData* FirstPointData = CastChecked<const UPCGPointData>(FirstData);
-		const UPCGPointData* SecondPointData = CastChecked<const UPCGPointData>(SecondData);
+		const UPCGBasePointData* FirstPointData = CastChecked<const UPCGBasePointData>(FirstData);
+		const UPCGBasePointData* SecondPointData = CastChecked<const UPCGBasePointData>(SecondData);
 
 		if (!SpatialBasicsAreIdentical(FirstPointData, SecondPointData))
 		{
 			return false;
 		}
+				
+		if (FirstPointData->GetNumPoints() != SecondPointData->GetNumPoints())
+		{
+			return false;
+		}
 
-		const TArray<FPCGPoint>& FirstPoints = FirstPointData->GetPoints();
-		const TArray<FPCGPoint>& SecondPoints = SecondPointData->GetPoints();
-
-		if (FirstPoints.Num() != SecondPoints.Num())
+		if (!CompareProperty<float>(FirstPointData, SecondPointData, EPCGPointNativeProperties::Steepness))
 		{
 			return false;
 		}
 
-		for (int32 I = 0; I < FirstPoints.Num(); ++I)
+		if (!CompareProperty<float>(FirstPointData, SecondPointData, EPCGPointNativeProperties::Density))
 		{
-			if (!PCGTestsCommon::PointsAreIdentical(FirstPoints[I], SecondPoints[I]))
-			{
-				return false;
-			}
+			return false;
+		}
+		
+		if (!CompareProperty<FVector>(FirstPointData, SecondPointData, EPCGPointNativeProperties::BoundsMin))
+		{
+			return false;
+		}
+		
+		if (!CompareProperty<FVector>(FirstPointData, SecondPointData, EPCGPointNativeProperties::BoundsMax))
+		{
+			return false;
+		}
+		
+		if (!CompareProperty<FVector4>(FirstPointData, SecondPointData, EPCGPointNativeProperties::Color))
+		{
+			return false;
+		}
+		
+		if (!CompareProperty<FTransform>(FirstPointData, SecondPointData, EPCGPointNativeProperties::Transform))
+		{
+			return false;
 		}
 
 		return true;
diff --git a/Engine/Plugins/PCG/Source/PCG/Public/Data/PCGCollisionShapeData.h b/Engine/Plugins/PCG/Source/PCG/Public/Data/PCGCollisionShapeData.h
--- a/Engine/Plugins/PCG/Source/PCG/Public/Data/PCGCollisionShapeData.h
+++ b/Engine/Plugins/PCG/Source/PCG/Public/Data/PCGCollisionShapeData.h
@@ -40,8 +40,12 @@ protected:
 public:
 	// ~Begin UPCGSpatialDataWithPointCache implementation
 	virtual const UPCGPointData* CreatePointData(FPCGContext* Context) const override;
+	virtual const UPCGPointArrayData* CreatePointArrayData(FPCGContext* Context, const FBox& InBounds) const override;
 	// ~End UPCGSpatialDataWithPointCache implementation
 
+private:
+	virtual const UPCGBasePointData* CreateBasePointData(FPCGContext* Context, TSubclassOf<UPCGBasePointData> PointDataClass) const;
+
 protected:
 	UPROPERTY(BlueprintReadOnly, VisibleAnywhere, Category = SpatialData)
 	FTransform Transform;
diff --git a/Engine/Plugins/PCG/Source/PCG/Public/Data/PCGCollisionWrapperData.h b/Engine/Plugins/PCG/Source/PCG/Public/Data/PCGCollisionWrapperData.h
--- a/Engine/Plugins/PCG/Source/PCG/Public/Data/PCGCollisionWrapperData.h
+++ b/Engine/Plugins/PCG/Source/PCG/Public/Data/PCGCollisionWrapperData.h
@@ -3,7 +3,7 @@
 #pragma once
 
 #include "PCGSpatialData.h"
-#include "Data/PCGPointData.h"
+#include "Data/PCGBasePointData.h"
 #include "Metadata/PCGAttributePropertySelector.h"
 #include "Metadata/Accessors/PCGAttributeAccessorHelpers.h"
 
@@ -60,10 +60,10 @@ class UPCGCollisionWrapperData : public UPCGSpatialData
 
 public:
 	/** Inititializes the collision wrapper on a point data based on the provided attribute selector */
-	PCG_API bool Initialize(const UPCGPointData* InPointData, const FPCGAttributePropertyInputSelector& InCollisionSelector, EPCGCollisionQueryFlag InCollisionQueryFlag);
+	PCG_API bool Initialize(const UPCGBasePointData* InPointData, const FPCGAttributePropertyInputSelector& InCollisionSelector, EPCGCollisionQueryFlag InCollisionQueryFlag);
 
 	/** Advanced API for async loading */
-	bool PreInitializeAndGatherMeshesEx(const UPCGPointData* InPointData, const FPCGAttributePropertyInputSelector& InCollisionSelector, EPCGCollisionQueryFlag InCollisionQueryFlag, TArray<FSoftObjectPath>& OutMeshesToLoad);
+	bool PreInitializeAndGatherMeshesEx(const UPCGBasePointData* InPointData, const FPCGAttributePropertyInputSelector& InCollisionSelector, EPCGCollisionQueryFlag InCollisionQueryFlag, TArray<FSoftObjectPath>& OutMeshesToLoad);
 	void FinalizeInitializationEx(const TArray<FSoftObjectPath>& InMeshPaths);
 
 	// ~Begin UPCGData Interface
@@ -78,6 +78,7 @@ public:
 	virtual FBox GetBounds() const override;
 	virtual FBox GetStrictBounds() const override;
 	virtual const UPCGPointData* ToPointData(FPCGContext* Context, const FBox& InBounds = FBox(EForceInit::ForceInit)) const override;
+	virtual const UPCGPointArrayData* ToPointArrayData(FPCGContext* Context, const FBox& InBounds = FBox(EForceInit::ForceInit)) const override;
 	virtual bool SamplePoint(const FTransform& Transform, const FBox& Bounds, FPCGPoint& OutPoint, UPCGMetadata* OutMetadata) const override;
 protected:
 	virtual UPCGSpatialData* CopyInternal(FPCGContext* Context) const override;
@@ -85,16 +86,16 @@ protected:
 
 	// For performance reasons, we keep a raw pointer to the point data in editor.
 #if WITH_EDITOR
-	inline const UPCGPointData* GetPointData() const { return RawPointData; }
+	inline const UPCGBasePointData* GetPointData() const { return RawPointData; }
 #else
-	inline const UPCGPointData* GetPointData() const { return PointData.Get(); }
+	inline const UPCGBasePointData* GetPointData() const { return PointData.Get(); }
 #endif
 
 private:
 	const PhysicsInterfaceTypes::FInlineShapeArray& GetCachedShapes(int32 EntryIndex) const;
 
 	UPROPERTY()
-	TObjectPtr<const UPCGPointData> PointData;
+	TObjectPtr<const UPCGBasePointData> PointData;
 
 	// Implementation note: in order to be able to duplicate this easily, we're keeping track of the arguments we used when calling Initialize (& derived functions)
 	UPROPERTY()
@@ -108,6 +109,6 @@ private:
 	TArray<PhysicsInterfaceTypes::FInlineShapeArray> CachedShapes;
 
 #if WITH_EDITOR
-	const UPCGPointData* RawPointData = nullptr;
+	const UPCGBasePointData* RawPointData = nullptr;
 #endif
 };
\ No newline at end of file
diff --git a/Engine/Plugins/PCG/Source/PCG/Public/Data/PCGDifferenceData.h b/Engine/Plugins/PCG/Source/PCG/Public/Data/PCGDifferenceData.h
--- a/Engine/Plugins/PCG/Source/PCG/Public/Data/PCGDifferenceData.h
+++ b/Engine/Plugins/PCG/Source/PCG/Public/Data/PCGDifferenceData.h
@@ -83,6 +83,7 @@ protected:
 public:
 	//~Begin UPCGSpatialDataWithPointCache interface
 	virtual const UPCGPointData* CreatePointData(FPCGContext* Context) const override;
+	virtual const UPCGPointArrayData* CreatePointArrayData(FPCGContext* Context, const FBox& InBounds) const override;
 	//~End UPCGSpatialDataWithPointCache interface
 protected:
 	UPROPERTY(BlueprintReadOnly, VisibleAnywhere, Category = SpatialData)
@@ -97,6 +98,8 @@ protected:
 	UPROPERTY(BlueprintSetter = SetDensityFunction, EditAnywhere, Category = Settings)
 	EPCGDifferenceDensityFunction DensityFunction = EPCGDifferenceDensityFunction::Minimum;
 
+	const UPCGBasePointData* CreateBasePointData(FPCGContext* Context, const UPCGBasePointData* SourcePointData, TSubclassOf<UPCGBasePointData> PointDataClass) const;
+
 #if WITH_EDITOR
 	inline const UPCGSpatialData* GetSource() const { return RawPointerSource; }
 	inline const UPCGSpatialData* GetDifference() const { return RawPointerDifference; }
diff --git a/Engine/Plugins/PCG/Source/PCG/Public/Data/PCGIntersectionData.h b/Engine/Plugins/PCG/Source/PCG/Public/Data/PCGIntersectionData.h
--- a/Engine/Plugins/PCG/Source/PCG/Public/Data/PCGIntersectionData.h
+++ b/Engine/Plugins/PCG/Source/PCG/Public/Data/PCGIntersectionData.h
@@ -55,6 +55,7 @@ protected:
 public:
 	//~Begin UPCGSpatialDataWithPointCache interface
 	virtual const UPCGPointData* CreatePointData(FPCGContext* Context) const override;
+	virtual const UPCGPointArrayData* CreatePointArrayData(FPCGContext* Context, const FBox& InBounds) const override;
 	//~End UPCGSpatialDataWithPointCache interface
 
 public:
@@ -68,7 +69,8 @@ public:
 	TObjectPtr<const UPCGSpatialData> B = nullptr;
 
 protected:
-	UPCGPointData* CreateAndFilterPointData(FPCGContext* Context, const UPCGSpatialData* X, const UPCGSpatialData* Y) const;
+	UPCGBasePointData* CreateBasePointData(FPCGContext* Context, TSubclassOf<UPCGBasePointData> PointDataClass) const;
+	UPCGBasePointData* CreateAndFilterPointData(FPCGContext* Context, const UPCGSpatialData* X, const UPCGSpatialData* Y, TSubclassOf<UPCGBasePointData> PointDataClass) const;
 
 	UPROPERTY()
 	FBox CachedBounds = FBox(EForceInit::ForceInit);
diff --git a/Engine/Plugins/PCG/Source/PCG/Public/Data/PCGLandscapeData.h b/Engine/Plugins/PCG/Source/PCG/Public/Data/PCGLandscapeData.h
--- a/Engine/Plugins/PCG/Source/PCG/Public/Data/PCGLandscapeData.h
+++ b/Engine/Plugins/PCG/Source/PCG/Public/Data/PCGLandscapeData.h
@@ -75,6 +75,7 @@ public:
 	virtual bool SupportsBoundedPointData() const { return true; }
 	virtual const UPCGPointData* CreatePointData(FPCGContext* Context) const override { return CreatePointData(Context, FBox(EForceInit::ForceInit)); }
 	virtual const UPCGPointData* CreatePointData(FPCGContext* Context, const FBox& InBounds) const override;
+	virtual const UPCGPointArrayData* CreatePointArrayData(FPCGContext* Context, const FBox& InBounds) const override;
 	// ~End UPCGSpatialDataWithPointCache interface
 
 	// TODO: add on property changed to clear cached data. This is used to populate the LandscapeInfos array.
@@ -89,6 +90,8 @@ protected:
 	*/
 	const ULandscapeInfo* GetLandscapeInfo(const FVector& InPosition) const;
 
+	const UPCGBasePointData* CreateBasePointData(FPCGContext* Context, const FBox& InBounds, TSubclassOf<UPCGBasePointData> PointDataClass) const;
+
 	UPROPERTY()
 	FBox Bounds = FBox(EForceInit::ForceInit);
 
diff --git a/Engine/Plugins/PCG/Source/PCG/Public/Data/PCGLandscapeSplineData.h b/Engine/Plugins/PCG/Source/PCG/Public/Data/PCGLandscapeSplineData.h
--- a/Engine/Plugins/PCG/Source/PCG/Public/Data/PCGLandscapeSplineData.h
+++ b/Engine/Plugins/PCG/Source/PCG/Public/Data/PCGLandscapeSplineData.h
@@ -41,6 +41,7 @@ public:
 
 	//~Begin UPCGSpatialDataWithPointCache interface
 	virtual const UPCGPointData* CreatePointData(FPCGContext* Context) const override;
+	virtual const UPCGPointArrayData* CreatePointArrayData(FPCGContext* Context, const FBox& InBounds) const override;
 	//~End UPCGSpatialDataWithPointCache interface
 
 	//~Begin UPCGSpatialData interface
@@ -50,6 +51,8 @@ protected:
 	virtual UPCGSpatialData* CopyInternal(FPCGContext* Context) const override;
 	//~End UPCGSpatialData interface
 
+	const UPCGBasePointData* CreateBasePointData(FPCGContext* Context, TSubclassOf<UPCGBasePointData> PointDataClass) const;
+
 	/** Recompute the reparameterization of the spline by distance. */
 	void UpdateReparamTable();
 
diff --git a/Engine/Plugins/PCG/Source/PCG/Public/Data/PCGPointArrayData.h b/Engine/Plugins/PCG/Source/PCG/Public/Data/PCGPointArrayData.h
--- a/Engine/Plugins/PCG/Source/PCG/Public/Data/PCGPointArrayData.h
+++ b/Engine/Plugins/PCG/Source/PCG/Public/Data/PCGPointArrayData.h
@@ -22,7 +22,7 @@ public:
 
 	virtual void VisitDataNetwork(TFunctionRef<void(const UPCGData*)> Action) const override;
 	virtual const UPCGPointData* ToPointData(FPCGContext* Context, const FBox& InBounds = FBox(EForceInit::ForceInit)) const override;
-	virtual const UPCGPointArrayData* ToPointArrayData(FPCGContext* Context) const override { return this; }
+	virtual const UPCGPointArrayData* ToPointArrayData(FPCGContext* Context, const FBox& InBounds = FBox(EForceInit::ForceInit)) const override { return this; }
 	virtual UPCGSpatialData* CopyInternal(FPCGContext* Context) const override;
 	virtual void CopyPropertiesTo(UPCGBasePointData* To, int32 ReadStartIndex, int32 WriteStartIndex, int32 Count, EPCGPointNativeProperties Properties) const override;
 	virtual EPCGPointNativeProperties GetAllocatedProperties(bool bWithInheritance = true) const override;
diff --git a/Engine/Plugins/PCG/Source/PCG/Public/Data/PCGPrimitiveData.h b/Engine/Plugins/PCG/Source/PCG/Public/Data/PCGPrimitiveData.h
--- a/Engine/Plugins/PCG/Source/PCG/Public/Data/PCGPrimitiveData.h
+++ b/Engine/Plugins/PCG/Source/PCG/Public/Data/PCGPrimitiveData.h
@@ -34,6 +34,7 @@ protected:
 public:
 	// ~Begin UPCGSpatialDataWithPointCache implementation
 	virtual const UPCGPointData* CreatePointData(FPCGContext* Context) const override;
+	virtual const UPCGPointArrayData* CreatePointArrayData(FPCGContext* Context, const FBox& InBounds) const override;
 	// ~End UPCGSpatialDataWithPointCache implementation
 
 	TWeakObjectPtr<UPrimitiveComponent> GetComponent() const { return Primitive; }
@@ -41,6 +42,9 @@ public:
 	UPROPERTY(BlueprintReadWrite, EditAnywhere, Category = "Data")
 	FVector VoxelSize = FVector(100.0, 100.0, 100.0);
 
+private:
+	virtual const UPCGBasePointData* CreateBasePointData(FPCGContext* Context, TSubclassOf<UPCGBasePointData> PointDataClass) const;
+
 protected:
 	UPROPERTY(BlueprintReadOnly, VisibleAnywhere, Category = "Data")
 	TWeakObjectPtr<UPrimitiveComponent> Primitive = nullptr;
diff --git a/Engine/Plugins/PCG/Source/PCG/Public/Data/PCGProjectionData.h b/Engine/Plugins/PCG/Source/PCG/Public/Data/PCGProjectionData.h
--- a/Engine/Plugins/PCG/Source/PCG/Public/Data/PCGProjectionData.h
+++ b/Engine/Plugins/PCG/Source/PCG/Public/Data/PCGProjectionData.h
@@ -42,9 +42,12 @@ protected:
 public:
 	//~Begin UPCGSpatialDataWithPointCache interface
 	virtual const UPCGPointData* CreatePointData(FPCGContext* Context) const override;
+	virtual const UPCGPointArrayData* CreatePointArrayData(FPCGContext* Context, const FBox& InBounds) const override;
 	//~End UPCGSpatialDataWithPointCache interface
 
 protected:
+	const UPCGBasePointData* CreateBasePointData(FPCGContext* Context, TSubclassOf<UPCGBasePointData> PointDataClass) const;
+
 	void CopyBaseProjectionClass(UPCGProjectionData* NewProjectionData) const;
 
 	FBox ProjectBounds(const FBox& InBounds) const;
diff --git a/Engine/Plugins/PCG/Source/PCG/Public/Data/PCGSpatialData.h b/Engine/Plugins/PCG/Source/PCG/Public/Data/PCGSpatialData.h
--- a/Engine/Plugins/PCG/Source/PCG/Public/Data/PCGSpatialData.h
+++ b/Engine/Plugins/PCG/Source/PCG/Public/Data/PCGSpatialData.h
@@ -121,9 +121,9 @@ public:
 
 	virtual const UPCGPointData* ToPointData(FPCGContext* Context, const FBox& InBounds = FBox(EForceInit::ForceInit)) const PURE_VIRTUAL(UPCGSpatialData::ToPointData, return nullptr;);
 
-	virtual const UPCGPointArrayData* ToPointArrayData(FPCGContext* Context) const;
+	virtual const UPCGPointArrayData* ToPointArrayData(FPCGContext* Context, const FBox& InBounds = FBox(EForceInit::ForceInit)) const;
 
-	const UPCGBasePointData* ToBasePointData(FPCGContext* Context) const;
+	const UPCGBasePointData* ToBasePointData(FPCGContext* Context, const FBox& InBounds = FBox(EForceInit::ForceInit)) const;
 
 	/** Sample rotation, scale and other attributes from this data at the query position. Returns true if Transform location and Bounds overlaps this data. */
 	virtual bool SamplePoint(const FTransform& Transform, const FBox& Bounds, FPCGPoint& OutPoint, UPCGMetadata* OutMetadata) const PURE_VIRTUAL(UPCGSpatialData::SamplePoint, return false;);
@@ -247,22 +247,26 @@ public:
 
 	// ~UPCGSpatialData implementation
 	virtual const UPCGPointData* ToPointData(FPCGContext* Context, const FBox& InBounds = FBox(EForceInit::ForceInit)) const override;
+	virtual const UPCGPointArrayData* ToPointArrayData(FPCGContext* Context, const FBox& InBounds = FBox(EForceInit::ForceInit)) const override;
 	// ~End UPCGSpatialData implementation
 
 protected:
+	const UPCGBasePointData* ToBasePointDataInternal(FPCGContext* Context, const FBox& InBounds, TFunctionRef<const UPCGBasePointData* (FPCGContext*, const FBox&)> CreatePointDataFunc) const;
+
 	virtual bool SupportsBoundedPointData() const { return false; }
 	virtual const UPCGPointData* CreatePointData(FPCGContext* Context) const PURE_VIRTUAL(UPCGSpatialData::CreatePointData, return nullptr;);
 	virtual const UPCGPointData* CreatePointData(FPCGContext* Context, const FBox& InBounds) const { return CreatePointData(Context); }
+	virtual const UPCGPointArrayData* CreatePointArrayData(FPCGContext* Context, const FBox& InBounds) const; 
 
 private:
 	UPROPERTY(Transient)
-	mutable TObjectPtr<const UPCGPointData> CachedPointData;
+	mutable TObjectPtr<const UPCGBasePointData> CachedPointData;
 
 	UPROPERTY(Transient)
 	mutable TArray<FBox> CachedBoundedPointDataBoxes;
 
 	UPROPERTY(Transient)
-	mutable TArray<TObjectPtr<const UPCGPointData>> CachedBoundedPointData;
+	mutable TArray<TObjectPtr<const UPCGBasePointData>> CachedBoundedPointData;
 
 	mutable FCriticalSection CacheLock;
 };
diff --git a/Engine/Plugins/PCG/Source/PCG/Public/Data/PCGSpatialDataTpl.h b/Engine/Plugins/PCG/Source/PCG/Public/Data/PCGSpatialDataTpl.h
--- a/Engine/Plugins/PCG/Source/PCG/Public/Data/PCGSpatialDataTpl.h
+++ b/Engine/Plugins/PCG/Source/PCG/Public/Data/PCGSpatialDataTpl.h
@@ -3,6 +3,7 @@
 #pragma once
 
 #include "PCGPoint.h"
+#include "Data/PCGBasePointData.h"
 #include "Helpers/PCGAsync.h"
 
 #include "Algo/Transform.h"
@@ -63,4 +64,54 @@ namespace FPCGSpatialDataProcessing
 
 		ensure(FPCGAsync::AsyncProcessingRangeEx(AsyncState, NumIterations, Initialize, ProcessRange, MoveDataRange, Finished, /*bEnableTimeSlicing=*/false, ChunkSize, /*bAllowChunkSizeOverride=*/false));
 	}
+
+	template<int ChunkSize, typename ProcessRangeFunc>
+	void SampleBasedRangeProcessing(FPCGAsyncState* AsyncState, ProcessRangeFunc&& InProcessRange, const UPCGBasePointData* SourceData, UPCGBasePointData* TargetData, EPCGPointNativeProperties PropertiesToAllocate)
+	{
+		const int NumIterations = SourceData->GetNumPoints();
+
+		auto Initialize = [TargetData, SourceData, NumIterations, PropertiesToAllocate]()
+		{
+			TargetData->SetNumPoints(NumIterations);
+			TargetData->AllocateProperties(PropertiesToAllocate);
+		};
+
+		auto ProcessRange = [RangeFunc = MoveTemp(InProcessRange), SourceData, TargetData](int32 StartReadIndex, int32 StartWriteIndex, int32 Count) -> int32
+		{
+			ensure(Count >= 0 && Count <= ChunkSize);
+
+			TArray<TPair<FTransform, FBox>, TInlineAllocator<ChunkSize>> Samples;
+			
+			{
+				TRACE_CPUPROFILER_EVENT_SCOPE(FPCGSpatialDataProcessing::SamplePoints::PrepareSamples)
+				const TConstPCGValueRange<FTransform> SourceTransformRange = SourceData->GetConstTransformValueRange();
+				const TConstPCGValueRange<FVector> SourceBoundsMinRange = SourceData->GetConstBoundsMinValueRange();
+				const TConstPCGValueRange<FVector> SourceBoundsMaxRange = SourceData->GetConstBoundsMaxValueRange();
+
+				Samples.Reserve(Count);
+				for (int32 Index = 0; Index < Count; ++Index)
+				{
+					const int32 ReadIndex = StartReadIndex + Index;
+					Samples.Add(TPair<FTransform, FBox>(SourceTransformRange[ReadIndex], PCGPointHelpers::GetLocalBounds(SourceBoundsMinRange[ReadIndex], SourceBoundsMaxRange[ReadIndex])));
+				}
+			}
+
+			{
+				TRACE_CPUPROFILER_EVENT_SCOPE(FPCGSpatialDataProcessing::SamplePoints::RangeFunc)
+				return RangeFunc(Samples, SourceData, StartReadIndex, TargetData, StartWriteIndex);
+			}
+		};
+
+		auto MoveDataRange = [TargetData](int32 ReadIndex, int32 WriteIndex, int32 Count)
+		{
+			TargetData->MoveRange(ReadIndex, WriteIndex, Count);
+		};
+
+		auto Finished = [TargetData](int32 Count)
+		{
+			TargetData->SetNumPoints(Count);
+		};
+
+		ensure(FPCGAsync::AsyncProcessingRangeEx(AsyncState, NumIterations, Initialize, ProcessRange, MoveDataRange, Finished, /*bEnableTimeSlicing=*/false, ChunkSize, /*bAllowChunkSizeOverride=*/false));
+	}
 }
\ No newline at end of file
diff --git a/Engine/Plugins/PCG/Source/PCG/Public/Data/PCGSplineData.h b/Engine/Plugins/PCG/Source/PCG/Public/Data/PCGSplineData.h
--- a/Engine/Plugins/PCG/Source/PCG/Public/Data/PCGSplineData.h
+++ b/Engine/Plugins/PCG/Source/PCG/Public/Data/PCGSplineData.h
@@ -60,6 +60,7 @@ public:
 
 	//~Begin UPCGSpatialDataWithPointCache interface
 	virtual const UPCGPointData* CreatePointData(FPCGContext* Context) const override;
+	virtual const UPCGPointArrayData* CreatePointArrayData(FPCGContext* Context, const FBox& InBounds) const override;
 	//~End UPCGSpatialDataWithPointCache interface
 
 	//~Begin UPCGSpatialData interface
@@ -72,6 +73,8 @@ protected:
 
 	virtual void CopySplineData(UPCGSplineData* InCopy) const;
 
+	const UPCGBasePointData* CreateBasePointData(FPCGContext* Context, TSubclassOf<UPCGBasePointData> PointDataClass) const;
+
 public:
 	// Minimal data needed to replicate the behavior from USplineComponent
 	UPROPERTY()
diff --git a/Engine/Plugins/PCG/Source/PCG/Public/Data/PCGUnionData.h b/Engine/Plugins/PCG/Source/PCG/Public/Data/PCGUnionData.h
--- a/Engine/Plugins/PCG/Source/PCG/Public/Data/PCGUnionData.h
+++ b/Engine/Plugins/PCG/Source/PCG/Public/Data/PCGUnionData.h
@@ -61,6 +61,7 @@ protected:
 public:
 	//~Begin UPCGSpatialDataWithPointCache interface
 	virtual const UPCGPointData* CreatePointData(FPCGContext* Context) const override;
+	virtual const UPCGPointArrayData* CreatePointArrayData(FPCGContext* Context, const FBox& InBounds) const override;
 	//~End UPCGSpatialDataWithPointCache interface
 
 protected:
@@ -86,5 +87,6 @@ protected:
 	int CachedDimension = 0;
 
 private:
-	void CreateSequentialPointData(FPCGContext* Context, TArray<const UPCGSpatialData*>& DataRawPtr, TArray<const UPCGMetadata*>& InputMetadatas, UPCGPointData* PointData, UPCGMetadata* OutMetadata, bool bLeftToRight) const;
+	const UPCGBasePointData* CreateBasePointData(FPCGContext* Context, TSubclassOf<UPCGBasePointData> PointDataClass, TFunctionRef<const UPCGBasePointData* (FPCGContext*, const UPCGSpatialData*)> ToPointDataFunc) const;
+	void CreateSequentialPointData(FPCGContext* Context, TArray<const UPCGSpatialData*>& DataRawPtr, TArray<const UPCGMetadata*>& InputMetadatas, UPCGBasePointData* PointData, UPCGMetadata* OutMetadata, bool bLeftToRight, TFunctionRef<const UPCGBasePointData* (FPCGContext*, const UPCGSpatialData*)> ToPointDataFunc) const;
 };
diff --git a/Engine/Plugins/PCG/Source/PCG/Public/Data/PCGVolumeData.h b/Engine/Plugins/PCG/Source/PCG/Public/Data/PCGVolumeData.h
--- a/Engine/Plugins/PCG/Source/PCG/Public/Data/PCGVolumeData.h
+++ b/Engine/Plugins/PCG/Source/PCG/Public/Data/PCGVolumeData.h
@@ -38,11 +38,15 @@ protected:
 public:
 	// ~Begin UPCGSpatialDataWithPointCache interface
 	virtual const UPCGPointData* CreatePointData(FPCGContext* Context) const override;
+	virtual const UPCGPointArrayData* CreatePointArrayData(FPCGContext* Context, const FBox& InBounds) const override;
 	// ~End UPCGSpatialDataWithPointCache interface
 
 	UPROPERTY(BlueprintReadWrite, EditAnywhere, Category = "Data")
 	FVector VoxelSize = FVector(100.0, 100.0, 100.0);
 
+private:
+	const UPCGBasePointData* CreateBasePointData(FPCGContext* Context, TSubclassOf<UPCGBasePointData> PointDataClass) const;
+
 protected:
 	void CopyBaseVolumeData(UPCGVolumeData* NewVolumeData) const;
 	void ReleaseInternalBodyInstance();
diff --git a/Engine/Plugins/PCG/Source/PCG/Public/Data/PCGWorldData.h b/Engine/Plugins/PCG/Source/PCG/Public/Data/PCGWorldData.h
--- a/Engine/Plugins/PCG/Source/PCG/Public/Data/PCGWorldData.h
+++ b/Engine/Plugins/PCG/Source/PCG/Public/Data/PCGWorldData.h
@@ -318,8 +318,12 @@ public:
 	virtual bool SupportsBoundedPointData() const { return true; }
 	virtual const UPCGPointData* CreatePointData(FPCGContext* Context) const override { return CreatePointData(Context, FBox(EForceInit::ForceInit)); }
 	virtual const UPCGPointData* CreatePointData(FPCGContext* Context, const FBox& InBounds) const override;
+	virtual const UPCGPointArrayData* CreatePointArrayData(FPCGContext* Context, const FBox& InBounds) const override;
 	//~End UPCGSpatialDataWithPointCache
+private:
+	const UPCGBasePointData* CreateBasePointData(FPCGContext* Context, const FBox& InBounds, TSubclassOf<UPCGBasePointData> PointDataClass) const;
 
+public:
 	UPROPERTY()
 	TWeakObjectPtr<UWorld> World = nullptr;
 
diff --git a/Engine/Plugins/PCG/Source/PCG/Public/Elements/PCGDifferenceElement.h b/Engine/Plugins/PCG/Source/PCG/Public/Elements/PCGDifferenceElement.h
--- a/Engine/Plugins/PCG/Source/PCG/Public/Elements/PCGDifferenceElement.h
+++ b/Engine/Plugins/PCG/Source/PCG/Public/Elements/PCGDifferenceElement.h
@@ -61,4 +61,5 @@ class FPCGDifferenceElement : public IPCGElement
 protected:
 	virtual bool ExecuteInternal(FPCGContext* Context) const override;
 	virtual EPCGElementExecutionLoopMode ExecutionLoopMode(const UPCGSettings* Settings) const override { return EPCGElementExecutionLoopMode::SinglePrimaryPin; }
+	virtual bool SupportsBasePointDataInputs(FPCGContext* InContext) const override { return true; }
 };
diff --git a/Engine/Plugins/PCG/Source/PCG/Public/Elements/PCGInnerIntersectionElement.h b/Engine/Plugins/PCG/Source/PCG/Public/Elements/PCGInnerIntersectionElement.h
--- a/Engine/Plugins/PCG/Source/PCG/Public/Elements/PCGInnerIntersectionElement.h
+++ b/Engine/Plugins/PCG/Source/PCG/Public/Elements/PCGInnerIntersectionElement.h
@@ -45,4 +45,5 @@ class FPCGInnerIntersectionElement : public IPCGElement
 {
 protected:
 	virtual bool ExecuteInternal(FPCGContext* Context) const override;
+	virtual bool SupportsBasePointDataInputs(FPCGContext* InContext) const override { return true; }
 };
diff --git a/Engine/Plugins/PCG/Source/PCG/Public/Elements/PCGOuterIntersectionElement.h b/Engine/Plugins/PCG/Source/PCG/Public/Elements/PCGOuterIntersectionElement.h
--- a/Engine/Plugins/PCG/Source/PCG/Public/Elements/PCGOuterIntersectionElement.h
+++ b/Engine/Plugins/PCG/Source/PCG/Public/Elements/PCGOuterIntersectionElement.h
@@ -68,4 +68,5 @@ class FPCGOuterIntersectionElement : public IPCGElement
 protected:
 	virtual bool ExecuteInternal(FPCGContext* Context) const override;
 	virtual EPCGElementExecutionLoopMode ExecutionLoopMode(const UPCGSettings* Settings) const override { return EPCGElementExecutionLoopMode::SinglePrimaryPin; }
+	virtual bool SupportsBasePointDataInputs(FPCGContext* InContext) const override { return true; }
 };
\ No newline at end of file
diff --git a/Engine/Plugins/PCG/Source/PCG/Public/Elements/PCGUnionElement.h b/Engine/Plugins/PCG/Source/PCG/Public/Elements/PCGUnionElement.h
--- a/Engine/Plugins/PCG/Source/PCG/Public/Elements/PCGUnionElement.h
+++ b/Engine/Plugins/PCG/Source/PCG/Public/Elements/PCGUnionElement.h
@@ -47,4 +47,5 @@ class FPCGUnionElement : public IPCGElement
 {
 protected:
 	virtual bool ExecuteInternal(FPCGContext* Context) const override;
+	virtual bool SupportsBasePointDataInputs(FPCGContext* InContext) const override { return true; }
 };
diff --git a/Engine/Plugins/PCG/Source/PCG/Public/Elements/PCGVolumeSampler.h b/Engine/Plugins/PCG/Source/PCG/Public/Elements/PCGVolumeSampler.h
--- a/Engine/Plugins/PCG/Source/PCG/Public/Elements/PCGVolumeSampler.h
+++ b/Engine/Plugins/PCG/Source/PCG/Public/Elements/PCGVolumeSampler.h
@@ -41,14 +41,23 @@ namespace PCGVolumeSampler
 		FVolumeSamplerParams Settings;
 
 		const UPCGSpatialData* Volume = nullptr;
+
+		UPCGBasePointData* OutputPointData = nullptr;
+
+		PRAGMA_DISABLE_DEPRECATION_WARNINGS
+		UE_DEPRECATED(5.6, "Use OutputPointData instead")
 		UPCGPointData* OutputData = nullptr;
+		PRAGMA_ENABLE_DEPRECATION_WARNINGS
 	};
 
 	/** Sample a volume and returns the resulting point data. */
+	PCG_API UPCGBasePointData* SampleVolume(FPCGContext* Context, TSubclassOf<UPCGBasePointData> PointDataClass, const FVolumeSamplerParams& SamplerSettings, const UPCGSpatialData* Volume, const UPCGSpatialData* BoundingShape = nullptr);
+
+	UE_DEPRECATED(5.6, "Use SampleVolume taking a UPCGBasePoint subclass instead")
 	PCG_API UPCGPointData* SampleVolume(FPCGContext* Context, const FVolumeSamplerParams& SamplerSettings, const UPCGSpatialData* Volume, const UPCGSpatialData* BoundingShape = nullptr);
 
 	/** Sample a volume and write the results in the given point data. Can be timesliced and will return false if the processing is not done, true otherwise. */
-	PCG_API bool SampleVolume(FPCGContext* Context, const FVolumeSamplerParams& SamplerSettings, const UPCGSpatialData* Volume, const UPCGSpatialData* BoundingShape, UPCGPointData* OutputData, const bool bTimeSlicingIsEnabled = false);
+	PCG_API bool SampleVolume(FPCGContext* Context, const FVolumeSamplerParams& SamplerSettings, const UPCGSpatialData* Volume, const UPCGSpatialData* BoundingShape, UPCGBasePointData* OutputData, const bool bTimeSlicingIsEnabled = false);
 }
 
 UCLASS(MinimalAPI, BlueprintType, ClassGroup = (Procedural))
diff --git a/Engine/Plugins/PCG/Source/PCGEditor/Private/DataVisualizations/PCGVolumeDataVisualization.cpp b/Engine/Plugins/PCG/Source/PCGEditor/Private/DataVisualizations/PCGVolumeDataVisualization.cpp
--- a/Engine/Plugins/PCG/Source/PCGEditor/Private/DataVisualizations/PCGVolumeDataVisualization.cpp
+++ b/Engine/Plugins/PCG/Source/PCGEditor/Private/DataVisualizations/PCGVolumeDataVisualization.cpp
@@ -4,6 +4,7 @@
 
 #include "Data/PCGCollisionShapeData.h"
 #include "Data/PCGCollisionWrapperData.h"
+#include "Data/PCGPointData.h"
 #include "Data/PCGPrimitiveData.h"
 #include "Data/PCGSpatialData.h"
 #include "Data/PCGVolumeData.h"
@@ -35,7 +36,8 @@ const UPCGPointData* IPCGPrimitiveDataVisualization::CollapseToDebugPointData(FP
 		SamplerParams.VoxelSize = FVector(VoxelSize, VoxelSize, VoxelSize);
 		SamplerParams.Bounds = Bounds;
 
-		return PCGVolumeSampler::SampleVolume(Context, SamplerParams, SpatialData);
+		// @todo_pcg: convert IPCGPrimitiveDataVisualization to support UPCGBasePointData
+		return Cast<UPCGPointData>(PCGVolumeSampler::SampleVolume(Context, UPCGPointData::StaticClass(), SamplerParams, SpatialData));
 	}
 	else
 	{
