[PCG] Fix include typo (uppercase)

#rnx

diff --git a/Engine/Plugins/PCG/Source/PCG/Private/Data/PCGDifferenceData.cpp b/Engine/Plugins/PCG/Source/PCG/Private/Data/PCGDifferenceData.cpp
--- a/Engine/Plugins/PCG/Source/PCG/Private/Data/PCGDifferenceData.cpp
+++ b/Engine/Plugins/PCG/Source/PCG/Private/Data/PCGDifferenceData.cpp
@@ -3,7 +3,7 @@
 #include "Data/PCGDifferenceData.h"
 
 #include "PCGContext.h"
-#include "DAta/PCGPointArrayData.h"
+#include "Data/PCGPointArrayData.h"
 #include "Data/PCGPointData.h"
 #include "Data/PCGSpatialData.h"
 #include "Data/PCGSpatialDataTpl.h"
