[PCG] Allow passing GPU data through higen grid links by default.

The risk is that this data will be held by the upstream compponent until the upstream component is cleaned. We need to solve this with graph editor UI/UX. Having this cvar disabled by default is not helping in the short term.

Finally lower verbosity of a shader compilation log which can be a bit spammy.

#rnx

diff --git a/Engine/Plugins/PCG/Source/PCG/Private/Compute/Elements/PCGComputeGraphElement.cpp b/Engine/Plugins/PCG/Source/PCG/Private/Compute/Elements/PCGComputeGraphElement.cpp
--- a/Engine/Plugins/PCG/Source/PCG/Private/Compute/Elements/PCGComputeGraphElement.cpp
+++ b/Engine/Plugins/PCG/Source/PCG/Private/Compute/Elements/PCGComputeGraphElement.cpp
@@ -151,7 +151,7 @@ bool FPCGComputeGraphElement::ExecuteInternal(FPCGContext* InContext) const
 	// 3. If still compiling, try again next frame.
 	if (Context->ComputeGraph->HasKernelResourcesPendingShaderCompilation())
 	{
-		UE_LOG(LogPCG, Log, TEXT("Deferring until next frame as the kernel has pending shader compilations."));
+		UE_LOG(LogPCG, Verbose, TEXT("Deferring until next frame as the kernel has pending shader compilations."));
 		SleepUntilNextFrame(Context);
 		return false;
 	}
diff --git a/Engine/Plugins/PCG/Source/PCG/Private/Graph/PCGGraphExecutor.cpp b/Engine/Plugins/PCG/Source/PCG/Private/Graph/PCGGraphExecutor.cpp
--- a/Engine/Plugins/PCG/Source/PCG/Private/Graph/PCGGraphExecutor.cpp
+++ b/Engine/Plugins/PCG/Source/PCG/Private/Graph/PCGGraphExecutor.cpp
@@ -93,7 +93,7 @@ namespace PCGGraphExecutor
 
 	TAutoConsoleVariable<bool> CVarPassGPUDataThroughGridLinks(
 		TEXT("pcg.Graph.GPU.PassGPUDataThroughGridLinks"),
-		false,
+		true,
 		TEXT("Whether proxies for GPU data are cached in per pin output data and passed through grid links. If false data is read back to CPU."));
 
 	void ClearAsyncFlags(TSet<TObjectPtr<UObject>>& AsyncObjects)
