[PCG] Fix bug in Comparison Metadata Ops, which was ignoring tolerance

diff --git a/Engine/Plugins/PCG/Source/PCG/Private/Elements/Grammar/PCGSelectGrammar.cpp b/Engine/Plugins/PCG/Source/PCG/Private/Elements/Grammar/PCGSelectGrammar.cpp
--- a/Engine/Plugins/PCG/Source/PCG/Private/Elements/Grammar/PCGSelectGrammar.cpp
+++ b/Engine/Plugins/PCG/Source/PCG/Private/Elements/Grammar/PCGSelectGrammar.cpp
@@ -46,7 +46,7 @@ static constexpr TBinarySignature<T> BinaryCompare[] =
 	[](const T&, const T&) { return true; }, // Select
 	PCG::Private::MetadataTraits<T>::Less,
 	PCG::Private::MetadataTraits<T>::LessOrEqual,
-	PCG::Private::MetadataTraits<T>::Equal,
+	[](const T& LHS, const T& RHS) { return PCG::Private::MetadataTraits<T>::Equal(LHS, RHS); }, // Adapter for 'Equal', as it changed from binary to ternary with tolerance
 	PCG::Private::MetadataTraits<T>::GreaterOrEqual,
 	PCG::Private::MetadataTraits<T>::Greater,
 };
diff --git a/Engine/Plugins/PCG/Source/PCG/Private/Elements/Metadata/PCGMetadataCompareOpElement.cpp b/Engine/Plugins/PCG/Source/PCG/Private/Elements/Metadata/PCGMetadataCompareOpElement.cpp
--- a/Engine/Plugins/PCG/Source/PCG/Private/Elements/Metadata/PCGMetadataCompareOpElement.cpp
+++ b/Engine/Plugins/PCG/Source/PCG/Private/Elements/Metadata/PCGMetadataCompareOpElement.cpp
@@ -10,15 +10,29 @@
 namespace PCGMetadataCompareSettings
 {
 	template <typename T>
-	bool ApplyCompare(const T& Input1, const T& Input2, EPCGMetadataCompareOperation Operation, double Tolerance)
+	bool ApplyCompare(const T& Input1, const T& Input2, EPCGMetadataCompareOperation Operation, double Tolerance = UE_DOUBLE_SMALL_NUMBER)
 	{
 		if (Operation == EPCGMetadataCompareOperation::Equal)
 		{
-			return PCG::Private::MetadataTraits<T>::Equal(Input1, Input2);
+			if constexpr (PCG::Private::MetadataTraits<T>::NeedsCompareTolerance)
+			{
+				return PCG::Private::MetadataTraits<T>::Equal(Input1, Input2, Tolerance);
+			}
+			else
+			{
+				return PCG::Private::MetadataTraits<T>::Equal(Input1, Input2);
+			}
 		}
 		else if (Operation == EPCGMetadataCompareOperation::NotEqual)
 		{
-			return !PCG::Private::MetadataTraits<T>::Equal(Input1, Input2);
+			if constexpr (PCG::Private::MetadataTraits<T>::NeedsCompareTolerance)
+			{
+				return !PCG::Private::MetadataTraits<T>::Equal(Input1, Input2, Tolerance);
+			}
+			else
+			{
+				return !PCG::Private::MetadataTraits<T>::Equal(Input1, Input2);
+			}
 		}
 
 		if constexpr (PCG::Private::MetadataTraits<T>::CanCompare)
diff --git a/Engine/Plugins/PCG/Source/PCG/Public/Metadata/PCGMetadataAttributeTraits.h b/Engine/Plugins/PCG/Source/PCG/Public/Metadata/PCGMetadataAttributeTraits.h
--- a/Engine/Plugins/PCG/Source/PCG/Public/Metadata/PCGMetadataAttributeTraits.h
+++ b/Engine/Plugins/PCG/Source/PCG/Public/Metadata/PCGMetadataAttributeTraits.h
@@ -331,11 +331,14 @@ namespace PCG
 		* - CanSearchString: Can do string search operations like matching substrings.
 		* - CompressData: Uses value keys, so if two values are identical, we only store one.
 		* - NeedsConstruction: Whether type requires constructor for valid state (i.e. can't simply be zero initialized).
+		* - NeedsCompareTolerance: Whether the type itself or the container's value type is floating point and needs to account for tolerance.
 		*/
 
 		template<typename T>
 		struct DefaultStringTraits
 		{
+			enum { NeedsCompareTolerance = false };
+
 			static FString ToString(const T& A)
 			{
 				return A.ToString();
@@ -362,7 +365,8 @@ namespace PCG
 		{
 			enum { CanSubAdd = true };
 			enum { CanMulDiv = true };
-			
+			enum { NeedsCompareTolerance = false };
+
 			static bool Equal(const T& A, const T& B)
 			{
 				return A == B;
@@ -537,10 +541,11 @@ namespace PCG
 			enum { CompressData = false };
 			enum { CanSearchString = false };
 			enum { NeedsConstruction = false };
+			enum { NeedsCompareTolerance = true };
 
-			static bool Equal(const float& A, const float& B)
+			static bool Equal(const float& A, const float& B, const double Tolerance = UE_SMALL_NUMBER)
 			{
-				return FMath::IsNearlyEqual(A, B);
+				return FMath::IsNearlyEqual(A, B, Tolerance);
 			}
 
 			static float ParseFromString(const FString& InValue, bool& bOutSuccess)
@@ -557,10 +562,11 @@ namespace PCG
 			enum { CompressData = false };
 			enum { CanSearchString = false };
 			enum { NeedsConstruction = false };
+			enum { NeedsCompareTolerance = true };
 
-			static bool Equal(const double& A, const double& B)
+			static bool Equal(const double& A, const double& B, const double Tolerance = UE_DOUBLE_SMALL_NUMBER)
 			{
-				return FMath::IsNearlyEqual(A, B);
+				return FMath::IsNearlyEqual(A, B, Tolerance);
 			}
 
 			static double ParseFromString(const FString& InValue, bool& bOutSuccess)
@@ -584,6 +590,7 @@ namespace PCG
 			enum { CanComputeDistance = false };
 			enum { CanSearchString = false };
 			enum { NeedsConstruction = false };
+			enum { NeedsCompareTolerance = false };
 
 			static bool ZeroValue()
 			{
@@ -654,6 +661,7 @@ namespace PCG
 			enum { NeedsConstruction = false };
 			enum { CanFindNearest = true };
 			enum { CanComputeDistance = true };
+			enum { NeedsCompareTolerance = true };
 
 			using DistanceType = typename T::FReal;
 
@@ -687,9 +695,9 @@ namespace PCG
 		template<>
 		struct MetadataTraits<FVector2D> : VectorTraits<FVector2D>
 		{
-			static bool Equal(const FVector2D& A, const FVector2D& B)
+			static bool Equal(const FVector2D& A, const FVector2D& B, const FVector2D::FReal Tolerance = UE_DOUBLE_SMALL_NUMBER)
 			{
-				return A.Equals(B);
+				return A.Equals(B, Tolerance);
 			}
 
 			static FVector2D Min(const FVector2D& A, const FVector2D& B)
@@ -734,9 +742,9 @@ namespace PCG
 		template<>
 		struct MetadataTraits<FVector> : VectorTraits<FVector>
 		{
-			static bool Equal(const FVector& A, const FVector& B)
+			static bool Equal(const FVector& A, const FVector& B, const FVector::FReal Tolerance = UE_DOUBLE_SMALL_NUMBER)
 			{
-				return A.Equals(B);
+				return A.Equals(B, Tolerance);
 			}
 
 			static FVector Min(const FVector& A, const FVector& B)
@@ -781,9 +789,9 @@ namespace PCG
 		template<>
 		struct MetadataTraits<FVector4> : VectorTraits<FVector4>
 		{
-			static bool Equal(const FVector4& A, const FVector4& B)
+			static bool Equal(const FVector4& A, const FVector4& B, const FVector4::FReal Tolerance = UE_DOUBLE_SMALL_NUMBER)
 			{
-				return A.Equals(B);
+				return A.Equals(B, Tolerance);
 			}
 
 			static FVector4 Min(const FVector4& A, const FVector4& B)
@@ -837,12 +845,13 @@ namespace PCG
 			enum { CanCompare = false };
 			enum { CanSearchString = false };
 			enum { NeedsConstruction = false };
+			enum { NeedsCompareTolerance = true };
 
 			using DistanceType = typename FQuat::FReal;
 
-			static bool Equal(const FQuat& A, const FQuat& B)
+			static bool Equal(const FQuat& A, const FQuat& B, const DistanceType Tolerance = UE_DOUBLE_SMALL_NUMBER)
 			{
-				return A.Equals(B);
+				return A.Equals(B, Tolerance);
 			}
 
 			static FQuat Add(const FQuat& A, const FQuat& B)
@@ -912,12 +921,13 @@ namespace PCG
 			enum { CanCompare = false };
 			enum { CanSearchString = false };
 			enum { NeedsConstruction = false };
+			enum { NeedsCompareTolerance = true };
 
 			using DistanceType = typename FRotator::FReal;
 
-			static bool Equal(const FRotator& A, const FRotator& B)
+			static bool Equal(const FRotator& A, const FRotator& B, const DistanceType Tolerance = UE_DOUBLE_SMALL_NUMBER)
 			{
-				return A.Equals(B);
+				return A.Equals(B, Tolerance);
 			}
 
 			static FRotator Add(const FRotator& A, const FRotator& B)
@@ -984,10 +994,11 @@ namespace PCG
 			enum { CanComputeDistance = false };
 			enum { CanSearchString = false };
 			enum { NeedsConstruction = false };
+			enum { NeedsCompareTolerance = true };
 
-			static bool Equal(const FTransform& A, const FTransform& B)
+			static bool Equal(const FTransform& A, const FTransform& B, const FTransform::FReal Tolerance = UE_DOUBLE_SMALL_NUMBER)
 			{
-				return A.Equals(B);
+				return A.Equals(B, Tolerance);
 			}
 
 			static FTransform Mul(const FTransform& A, const FTransform& B)
