[PCG] SpawnActor remove inheritance causing duplicate points in output

[RN]

diff --git a/Engine/Plugins/PCG/Source/PCG/Private/Elements/PCGSpawnActor.cpp b/Engine/Plugins/PCG/Source/PCG/Private/Elements/PCGSpawnActor.cpp
--- a/Engine/Plugins/PCG/Source/PCG/Private/Elements/PCGSpawnActor.cpp
+++ b/Engine/Plugins/PCG/Source/PCG/Private/Elements/PCGSpawnActor.cpp
@@ -583,7 +583,11 @@ bool FPCGSpawnActorElement::SpawnAndPrepareSubgraphs(FPCGSubgraphContext* Contex
 		if (bGenerateOutputsWithActorReference)
 		{
 			OutPointData = FPCGContext::NewPointData_AnyThread(Context);
-			OutPointData->InitializeFromData(PointData);
+
+			FPCGInitializeFromDataParams InitializeFromDataParams(PointData);
+			InitializeFromDataParams.bInheritSpatialData = false;
+
+			OutPointData->InitializeFromDataWithParams(InitializeFromDataParams);
 		}
 
 		FPCGTaggedData Output = Input;
