Allow negative weights in the attract for a cheap repulsion mode.

#rnx

diff --git a/Engine/Plugins/PCG/Source/PCG/Public/Elements/PCGAttractElement.h b/Engine/Plugins/PCG/Source/PCG/Public/Elements/PCGAttractElement.h
--- a/Engine/Plugins/PCG/Source/PCG/Public/Elements/PCGAttractElement.h
+++ b/Engine/Plugins/PCG/Source/PCG/Public/Elements/PCGAttractElement.h
@@ -79,7 +79,7 @@ public:
 	UPROPERTY(EditAnywhere, BlueprintReadWrite, Category = Settings, meta = (EditCondition = "bUseTargetWeight", PCG_Overridable))
 	FPCGAttributePropertyInputSelector TargetWeightAttribute;
 
-	UPROPERTY(EditAnywhere, BlueprintReadWrite, Category = Settings, meta = (ClampMin = "0.0", ClampMax = "1.0", EditCondition = "!bUseSourceWeight && !bUseTargetWeight", PCG_Overridable))
+	UPROPERTY(EditAnywhere, BlueprintReadWrite, Category = Settings, meta = (ClampMin = "-1.0", ClampMax = "1.0", EditCondition = "!bUseSourceWeight && !bUseTargetWeight", PCG_Overridable))
 	double Weight = 1.0;
 	
 	UPROPERTY(EditAnywhere, BlueprintReadWrite, Category = Settings)
