Added descriptor tags to CRC validation so that there is no false sharing between ISM component that would have different tags.

#rnx

diff --git a/Engine/Plugins/PCG/Source/PCG/Private/Elements/PCGStaticMeshSpawner.cpp b/Engine/Plugins/PCG/Source/PCG/Private/Elements/PCGStaticMeshSpawner.cpp
--- a/Engine/Plugins/PCG/Source/PCG/Private/Elements/PCGStaticMeshSpawner.cpp
+++ b/Engine/Plugins/PCG/Source/PCG/Private/Elements/PCGStaticMeshSpawner.cpp
@@ -478,6 +478,17 @@ void FPCGStaticMeshSpawnerElement::SpawnStaticMeshInstances(FPCGStaticMeshSpawne
 	Params.SettingsCrc = Settings->GetSettingsCrc();
 	ensure(Params.SettingsCrc.IsValid());
 
+	// Implementation note: in order to prevent components from merging together when they have tags, we need to incorporate the tags hash
+	// into the settings crc. This will make sure we don't select component that don't match up with what we're expecting.
+	// However, this validation is potentially too restrictive, as the order and composition of the tags ultimately don't really matter.
+	// TODO- improve this, however since this stems from a path where we build those, the ordering of tags could be sorted in other ways.
+	if (Params.Descriptor.HasTags())
+	{
+		const uint32 TagsArrayHash = Params.Descriptor.ComponentTags.IsEmpty() ? 1 : GetTypeHash(Params.Descriptor.ComponentTags);
+		const uint32 AdditionalTagsHash = Params.Descriptor.AdditionalCommaSeparatedTags.IsEmpty() ? 1 : GetTypeHash(Params.Descriptor.AdditionalCommaSeparatedTags);
+		Params.SettingsCrc.Combine(HashCombine(TagsArrayHash, AdditionalTagsHash));
+	}
+
 	Params.bAllowDescriptorChanges = Settings->bAllowDescriptorChanges;
 
 	UPCGManagedISMComponent* MISMC = UPCGActorHelpers::GetOrCreateManagedISMC(TargetActor, Context->SourceComponent.Get(), Params, Context);
diff --git a/Engine/Plugins/PCG/Source/PCG/Public/MeshSelectors/PCGISMDescriptor.h b/Engine/Plugins/PCG/Source/PCG/Public/MeshSelectors/PCGISMDescriptor.h
--- a/Engine/Plugins/PCG/Source/PCG/Public/MeshSelectors/PCGISMDescriptor.h
+++ b/Engine/Plugins/PCG/Source/PCG/Public/MeshSelectors/PCGISMDescriptor.h
@@ -18,6 +18,7 @@ struct FPCGSoftISMComponentDescriptor : public FSoftISMComponentDescriptor
 	PCG_API explicit FPCGSoftISMComponentDescriptor(const FISMComponentDescriptor& Other);
 	PCG_API virtual void InitFrom(const UStaticMeshComponent* Component, bool bInitBodyInstance = true) override;
 	PCG_API virtual void InitComponent(UInstancedStaticMeshComponent* ISMComponent) const override;
+	PCG_API bool HasTags() const { return !ComponentTags.IsEmpty() || !AdditionalCommaSeparatedTags.IsEmpty(); }
 
 public:
 	UPROPERTY(EditAnywhere, Category="Component Settings|Tags", meta = (DisplayAfter = "ComponentClass"))
