[PCG] Fix log crash with Compute element not returning a OriginalSettings (with pcg.Experimental.EnablePointArrayData & pcg.EnablePointArrayDataToPointDataConversionWarnings)

#rnx

diff --git a/Engine/Plugins/PCG/Source/PCG/Private/PCGElement.cpp b/Engine/Plugins/PCG/Source/PCG/Private/PCGElement.cpp
--- a/Engine/Plugins/PCG/Source/PCG/Private/PCGElement.cpp
+++ b/Engine/Plugins/PCG/Source/PCG/Private/PCGElement.cpp
@@ -449,7 +449,9 @@ bool IPCGElement::ConvertInputsIfNeeded(FPCGContext* Context) const
 
 		if (bDataWasConverted && CVarPCGEnablePointArrayDataToPointDataConversionWarnings.GetValueOnAnyThread())
 		{
-			PCGE_LOG(Warning, GraphAndLog, FText::Format(LOCTEXT("UnsupportedPointArrayData", "ToPointData was called on inputs of node '{0}'. Consider implementing support for UPCGBasePointData (IPCGElement::SupportsBasePointDataInputs)."), FText::FromString(Context->GetOriginalSettings<UPCGSettings>()->GetName())));
+			const UPCGSettings* Settings = Context->GetOriginalSettings<UPCGSettings>();
+			const FText SettingsName = Settings ? FText::FromString(Settings->GetName()) : LOCTEXT("UnknownSettings", "Unknown");
+			PCGE_LOG(Warning, GraphAndLog, FText::Format(LOCTEXT("UnsupportedPointArrayData", "ToPointData was called on inputs of node '{0}'. Consider implementing support for UPCGBasePointData (IPCGElement::SupportsBasePointDataInputs)."), SettingsName));
 		}
 	}
 
