Added multiple key support for point/parent indices.
Small behavior improvement on all-invalid data in Apply Hierarchy node

#rnx

diff --git a/Engine/Plugins/PCG/Source/PCG/Private/Elements/PCGApplyHierarchy.cpp b/Engine/Plugins/PCG/Source/PCG/Private/Elements/PCGApplyHierarchy.cpp
--- a/Engine/Plugins/PCG/Source/PCG/Private/Elements/PCGApplyHierarchy.cpp
+++ b/Engine/Plugins/PCG/Source/PCG/Private/Elements/PCGApplyHierarchy.cpp
@@ -2,6 +2,7 @@
 
 #include "Elements/PCGApplyHierarchy.h"
 
+#include "PCGDataAsset.h"
 #include "Data/PCGBasePointData.h"
 #include "Helpers/PCGAsync.h"
 #include "Metadata/PCGMetadataPartitionCommon.h"
@@ -10,6 +11,18 @@
 
 #define LOCTEXT_NAMESPACE "PCGApplyHierarchyElement"
 
+UPCGApplyHierarchySettings::UPCGApplyHierarchySettings()
+{
+	PointKeyAttributes.Emplace_GetRef().SetAttributeName(PCGLevelToAssetConstants::ActorIndexAttributeName);
+	ParentKeyAttributes.Emplace_GetRef().SetAttributeName(PCGLevelToAssetConstants::ParentIndexAttributeName);
+	HierarchyDepthAttribute.SetAttributeName(PCGLevelToAssetConstants::HierarchyDepthAttributeName);
+	RelativeTransformAttribute.SetAttributeName(PCGLevelToAssetConstants::RelativeTransformAttributeName);
+	ApplyParentRotation = EPCGApplyHierarchyOption::OptOutByAttribute;
+	ApplyParentRotationAttribute.SetAttributeName(PCGLevelToAssetConstants::IgnoreParentRotationAttributeName);
+	ApplyParentScale = EPCGApplyHierarchyOption::OptOutByAttribute;
+	ApplyParentScaleAttribute.SetAttributeName(PCGLevelToAssetConstants::IgnoreParentScaleAttributeName);
+}
+
 FPCGElementPtr UPCGApplyHierarchySettings::CreateElement() const
 {
 	return MakeShared<FPCGApplyHierarchyElement>();
@@ -82,21 +95,60 @@ bool FPCGApplyHierarchyElement::PrepareDataInternal(FPCGContext* InContext) cons
 			return true;
 		};
 
-		const FPCGAttributePropertyInputSelector PointIndexSelector = Settings->PointIndexAttribute.CopyAndFixLast(InputData);
-		if (!PrepareAccessorAndKeys(PointIndexSelector, OutState.PointIndexAccessor, OutState.PointIndexKeys))
+		auto PrepareAccessorAndKeysArray = [PrepareAccessorAndKeys, InputData](const TArray<FPCGAttributePropertyInputSelector>& InputSelectors, TArray<TUniquePtr<const IPCGAttributeAccessor>>& OutAccessors, TArray<TUniquePtr<const IPCGAttributeAccessorKeys>>& OutAccessorKeys) -> bool
+		{
+			for (const FPCGAttributePropertyInputSelector& InputSelector : InputSelectors)
+			{
+				const FPCGAttributePropertyInputSelector Selector = InputSelector.CopyAndFixLast(InputData);
+				TUniquePtr<const IPCGAttributeAccessor> Accessor;
+				TUniquePtr<const IPCGAttributeAccessorKeys> Keys;
+
+				if (PrepareAccessorAndKeys(Selector, Accessor, Keys))
+				{
+					OutAccessors.Add(std::move(Accessor));
+					OutAccessorKeys.Add(std::move(Keys));
+				}
+				else
+				{
+					return false;
+				}
+			}
+
+			return true;
+		};
+
+		if (Settings->PointKeyAttributes.IsEmpty() || Settings->ParentKeyAttributes.IsEmpty())
+		{
+			PCGLog::LogErrorOnGraph(LOCTEXT("RequiresAtLeastOneAttribute", "Both the point key attribute and parent key attribute require valid entries to perform the Apply Hierarchy operation."));
+			return EPCGTimeSliceInitResult::NoOperation;
+		}
+
+		if (!PrepareAccessorAndKeysArray(Settings->PointKeyAttributes, OutState.PointIndexAccessors, OutState.PointIndexKeys))
 		{
 			return EPCGTimeSliceInitResult::NoOperation;
 		}
 
-		const FPCGAttributePropertyInputSelector ParentIndexSelector = Settings->ParentIndexAttribute.CopyAndFixLast(InputData);
-		if (!PrepareAccessorAndKeys(ParentIndexSelector, OutState.ParentIndexAccessor, OutState.ParentIndexKeys))
+		if (!PrepareAccessorAndKeysArray(Settings->ParentKeyAttributes, OutState.ParentIndexAccessors, OutState.ParentIndexKeys))
 		{
 			return EPCGTimeSliceInitResult::NoOperation;
 		}
 
-		if (OutState.PointIndexKeys->GetNum() != OutState.ParentIndexKeys->GetNum())
+		const int NumKeys = OutState.PointIndexKeys[0] ? OutState.PointIndexKeys[0]->GetNum() : 0;
+		bool bHasCardinalityError = (NumKeys == 0);
+
+		for (const TUniquePtr<const IPCGAttributeAccessorKeys>& PointIndexKey : OutState.PointIndexKeys)
+		{
+			bHasCardinalityError |= (!PointIndexKey || PointIndexKey->GetNum() != NumKeys);
+		}
+		
+		for (const TUniquePtr<const IPCGAttributeAccessorKeys>& ParentIndexKey : OutState.ParentIndexKeys)
+		{
+			bHasCardinalityError |= (!ParentIndexKey || ParentIndexKey->GetNum() != NumKeys);
+		}
+
+		if (bHasCardinalityError)
 		{
-			PCGLog::LogErrorOnGraph(LOCTEXT("CardinalityMismatch", "Point index and Parent index properties do not have the same cardinality."), Context);
+			PCGLog::LogErrorOnGraph(LOCTEXT("CardinalityMismatch", "Point Key and Parent Key properties do not have the same cardinality."), Context);
 			return EPCGTimeSliceInitResult::NoOperation;
 		}
 
@@ -201,43 +253,82 @@ bool FPCGApplyHierarchyElement::ExecuteInternal(FPCGContext* InContext) const
 		if (!IterState.bParentMappingDone)
 		{
 			TRACE_CPUPROFILER_EVENT_SCOPE(FPCGApplyHierarchyElement::Execute::ParentMapping);
-			TArray<int32> PointIndices;
-			TArray<int32> ParentIndices;
+			constexpr int32 KeySize = 4;
+			TArray<TArray<int32, TInlineAllocator<KeySize>>> PointKeys;
+			TArray<TArray<int32, TInlineAllocator<KeySize>>> ParentKeys;
 
-			PointIndices.SetNumUninitialized(IterState.PointIndexKeys->GetNum());
-			ParentIndices.SetNumUninitialized(IterState.ParentIndexKeys->GetNum());
+			constexpr int32 ChunkSize = 256;
 
-			auto GetPointIndexAndParentIndex = [&IterState, &PointIndices, &ParentIndices](int32 StartReadIndex, int32 StartWriteIndex, int32 Count)
+			auto InitializePointKeysAndParentKeys = [&IterState, &PointKeys, &ParentKeys]()
 			{
-				IterState.PointIndexAccessor->GetRange<int32>(MakeArrayView(&PointIndices[StartWriteIndex], Count), StartReadIndex, *IterState.PointIndexKeys, EPCGAttributeAccessorFlags::AllowConstructible);
-				IterState.ParentIndexAccessor->GetRange<int32>(MakeArrayView(&ParentIndices[StartWriteIndex], Count), StartReadIndex, *IterState.ParentIndexKeys, EPCGAttributeAccessorFlags::AllowConstructible);
+				const int32 NumAttributes = IterState.PointIndexKeys.Num();
+				const int32 NumKeys = IterState.PointIndexKeys[0]->GetNum();
+				// Initialize and set array sizes
+				PointKeys.SetNum(NumKeys);
+				ParentKeys.SetNum(NumKeys);
+
+				for (auto& PointKeyArray : PointKeys)
+				{
+					PointKeyArray.SetNumUninitialized(NumAttributes);
+				}
+
+				for (auto& ParentKeyArray : ParentKeys)
+				{
+					ParentKeyArray.SetNumUninitialized(NumAttributes);
+				}
+			};
+
+			auto GetPointKeysAndParentKeys = [&IterState, &PointKeys, &ParentKeys](int32 StartReadIndex, int32 StartWriteIndex, int32 Count)
+			{
+				TArray<int32, TInlineAllocator<ChunkSize>> LocalArray;
+				LocalArray.SetNumUninitialized(Count);
+
+				for(int PointKeyIndex = 0; PointKeyIndex < IterState.PointIndexAccessors.Num(); ++PointKeyIndex)
+				{
+					IterState.PointIndexAccessors[PointKeyIndex]->GetRange<int32>(LocalArray, StartReadIndex, *IterState.PointIndexKeys[PointKeyIndex], EPCGAttributeAccessorFlags::AllowConstructible);
+
+					for (int32 Index = 0; Index < Count; ++Index)
+					{
+						PointKeys[StartWriteIndex + Index][PointKeyIndex] = LocalArray[Index];
+					}
+				}
+
+				for (int ParentKeyIndex = 0; ParentKeyIndex < IterState.ParentIndexAccessors.Num(); ++ParentKeyIndex)
+				{
+					IterState.ParentIndexAccessors[ParentKeyIndex]->GetRange<int32>(LocalArray, StartReadIndex, *IterState.ParentIndexKeys[ParentKeyIndex], EPCGAttributeAccessorFlags::AllowConstructible);
+
+					for (int32 Index = 0; Index < Count; ++Index)
+					{
+						ParentKeys[StartWriteIndex + Index][ParentKeyIndex] = LocalArray[Index];
+					}
+				}
+
 				return Count;
 			};
 
 			FPCGAsync::AsyncProcessingOneToOneRangeEx(
 				&Context->AsyncState,
-				PointIndices.Num(),
-				[]() {},
-				GetPointIndexAndParentIndex,
+				IterState.PointIndexKeys[0]->GetNum(),
+				InitializePointKeysAndParentKeys,
+				GetPointKeysAndParentKeys,
 				/*bEnableTimeSlicing=*/false,
-				/*ChunkSize*/256);
+				ChunkSize);
 
-			TMap<int32, int32> PointIndexToIndexMap;
-			PointIndexToIndexMap.Reserve(PointIndices.Num());
-			for (int Index = 0; Index < PointIndices.Num(); ++Index)
+			TMap<TArray<int32, TInlineAllocator<KeySize>>, int32> PointKeysToIndexMap;
+			for (int Index = 0; Index < PointKeys.Num(); ++Index)
 			{
-				PointIndexToIndexMap.Add(PointIndices[Index], Index);
+				PointKeysToIndexMap.Add(PointKeys[Index], Index);
 			}
 
 			bool bHasPointsWithInvalidParent = false;
 
-			IterState.ParentIndices.SetNumUninitialized(ParentIndices.Num());
-			for (int Index = 0; Index < ParentIndices.Num(); ++Index)
+			IterState.ParentIndices.SetNumUninitialized(ParentKeys.Num());
+			for(int Index = 0; Index < ParentKeys.Num(); ++Index)
 			{
-				const int32 OriginalParentIndex = ParentIndices[Index];
-				if (OriginalParentIndex != INDEX_NONE)
+				const TArray<int32, TInlineAllocator<KeySize>>& OriginalParentKey = ParentKeys[Index];
+				if(!OriginalParentKey.Contains(INDEX_NONE))
 				{
-					const int32* ParentIndex = PointIndexToIndexMap.Find(ParentIndices[Index]);
+					const int32* ParentIndex = PointKeysToIndexMap.Find(ParentKeys[Index]);
 					if (ParentIndex)
 					{
 						IterState.ParentIndices[Index] = *ParentIndex;
@@ -250,7 +341,7 @@ bool FPCGApplyHierarchyElement::ExecuteInternal(FPCGContext* InContext) const
 				}
 				else
 				{
-					ParentIndices[Index] = INDEX_NONE;
+					IterState.ParentIndices[Index] = INDEX_NONE;
 				}
 			}
 
@@ -518,6 +609,12 @@ bool FPCGApplyHierarchyElement::ExecuteInternal(FPCGContext* InContext) const
 			}
 		}
 
+		// Finally, if we had nothing to do, we shouldn't return any points.
+		if (IterState.HierarchyPartitionOrder.IsEmpty())
+		{
+			IterState.OutputData->SetNumPoints(0);
+		}
+
 		return true;
 	});
 }
diff --git a/Engine/Plugins/PCG/Source/PCG/Public/Elements/PCGApplyHierarchy.h b/Engine/Plugins/PCG/Source/PCG/Public/Elements/PCGApplyHierarchy.h
--- a/Engine/Plugins/PCG/Source/PCG/Public/Elements/PCGApplyHierarchy.h
+++ b/Engine/Plugins/PCG/Source/PCG/Public/Elements/PCGApplyHierarchy.h
@@ -26,6 +26,9 @@ class UPCGApplyHierarchySettings : public UPCGSettings
 {
 	GENERATED_BODY()
 
+public:
+	UPCGApplyHierarchySettings();
+
 	//~Begin UPCGSettings interface
 public:
 #if WITH_EDITOR
@@ -41,11 +44,13 @@ protected:
 	//~End UPCGSettings interface
 
 public:
+	/** Attributes that constitute a unique key representing the point. All attributes must be int32 at this time. */
 	UPROPERTY(BlueprintReadWrite, EditAnywhere, Category = Settings, meta = (PCG_Overridable))
-	FPCGAttributePropertyInputSelector PointIndexAttribute;
+	TArray<FPCGAttributePropertyInputSelector> PointKeyAttributes;
 
+	/** Attributes that constitute a unique key representing the point's parent in the hierarchy. All attributes must be int32 at this time. */
 	UPROPERTY(BlueprintReadWrite, EditAnywhere, Category = Settings, meta = (PCG_Overridable))
-	FPCGAttributePropertyInputSelector ParentIndexAttribute;
+	TArray<FPCGAttributePropertyInputSelector> ParentKeyAttributes;
 
 	UPROPERTY(BlueprintReadWrite, EditAnywhere, Category = Settings, meta = (PCG_Overridable))
 	FPCGAttributePropertyInputSelector HierarchyDepthAttribute;
@@ -54,13 +59,13 @@ public:
 	FPCGAttributePropertyInputSelector RelativeTransformAttribute;
 
 	UPROPERTY(BlueprintReadWrite, EditAnywhere, Category = Settings, meta = (PCG_Overridable))
-	EPCGApplyHierarchyOption ApplyParentRotation = EPCGApplyHierarchyOption::Always;
+	EPCGApplyHierarchyOption ApplyParentRotation;
 
 	UPROPERTY(BlueprintReadWrite, EditAnywhere, Category = Settings, meta = (EditCondition = "ApplyParentRotation==EPCGApplyHierarchyOption::OptInByAttribute||ApplyParentRotation==EPCGApplyHierarchyOption::OptOutByAttribute", EditConditionHides, PCG_Overridable))
 	FPCGAttributePropertyInputSelector ApplyParentRotationAttribute;
 
 	UPROPERTY(BlueprintReadWrite, EditAnywhere, Category = Settings, meta = (PCG_Overridable))
-	EPCGApplyHierarchyOption ApplyParentScale = EPCGApplyHierarchyOption::Always;
+	EPCGApplyHierarchyOption ApplyParentScale;
 
 	UPROPERTY(BlueprintReadWrite, EditAnywhere, Category = Settings, meta = (EditCondition = "ApplyParentScale==EPCGApplyHierarchyOption::OptInByAttribute||ApplyParentScale==EPCGApplyHierarchyOption::OptOutByAttribute", EditConditionHides, PCG_Overridable))
 	FPCGAttributePropertyInputSelector ApplyParentScaleAttribute;
@@ -77,12 +82,11 @@ namespace PCGApplyHierarchyElement
 		UPCGBasePointData* OutputData = nullptr;
 		int OutputDataIndex = INDEX_NONE;
 		
+		TArray<TUniquePtr<const IPCGAttributeAccessor>> PointIndexAccessors;
+		TArray<TUniquePtr<const IPCGAttributeAccessorKeys>> PointIndexKeys;
 
-		TUniquePtr<const IPCGAttributeAccessor> PointIndexAccessor;
-		TUniquePtr<const IPCGAttributeAccessorKeys> PointIndexKeys;
-
-		TUniquePtr<const IPCGAttributeAccessor> ParentIndexAccessor;
-		TUniquePtr<const IPCGAttributeAccessorKeys> ParentIndexKeys;
+		TArray<TUniquePtr<const IPCGAttributeAccessor>> ParentIndexAccessors;
+		TArray<TUniquePtr<const IPCGAttributeAccessorKeys>> ParentIndexKeys;
 
 		FPCGAttributePropertyInputSelector HierarchyDepthSelector;
 		TUniquePtr<const IPCGAttributeAccessor> HierarchyDepthAccessor;
diff --git a/Engine/Plugins/PCG/Source/PCG/Public/PCGDataAsset.h b/Engine/Plugins/PCG/Source/PCG/Public/PCGDataAsset.h
--- a/Engine/Plugins/PCG/Source/PCG/Public/PCGDataAsset.h
+++ b/Engine/Plugins/PCG/Source/PCG/Public/PCGDataAsset.h
@@ -7,6 +7,18 @@
 
 #include "PCGDataAsset.generated.h"
 
+namespace PCGLevelToAssetConstants
+{
+	const FName MaterialAttributeName = TEXT("Material");
+	const FName MeshAttributeName = TEXT("Mesh");
+	const FName HierarchyDepthAttributeName = TEXT("HierarchyDepth");
+	const FName ActorIndexAttributeName = TEXT("ActorIndex");
+	const FName ParentIndexAttributeName = TEXT("ParentIndex");
+	const FName RelativeTransformAttributeName = TEXT("RelativeTransform");
+	const FName IgnoreParentRotationAttributeName = TEXT("IgnoreParentRotation");
+	const FName IgnoreParentScaleAttributeName = TEXT("IgnoreParentScale");
+}
+
 /** Container for PCG data exported as standalone objects */
 UCLASS(hidecategories=Object, BlueprintType, ClassGroup=(Procedural))
 class PCG_API UPCGDataAsset : public UObject
diff --git a/Engine/Plugins/PCG/Source/PCGEditor/Private/PCGLevelToAsset.cpp b/Engine/Plugins/PCG/Source/PCGEditor/Private/PCGLevelToAsset.cpp
--- a/Engine/Plugins/PCG/Source/PCGEditor/Private/PCGLevelToAsset.cpp
+++ b/Engine/Plugins/PCG/Source/PCGEditor/Private/PCGLevelToAsset.cpp
@@ -181,12 +181,12 @@ bool UPCGLevelToAsset::BP_ExportWorld_Implementation(UWorld* World, const FStrin
 	FBox AllActorBounds(EForceInit::ForceInit);
 
 	// Hardcoded attributes
-	const FName MaterialAttributeName = TEXT("Material");
-	const FName MeshAttributeName = TEXT("Mesh");
-	const FName HierarchyDepthAttributeName = TEXT("HierarchyDepth");
-	const FName ActorIndexAttributeName = TEXT("ActorIndex");
-	const FName ParentIndexAttributeName = TEXT("ParentIndex");
-	const FName RelativeTransformAttributeName = TEXT("RelativeTransform");
+	const FName& MaterialAttributeName = PCGLevelToAssetConstants::MaterialAttributeName;
+	const FName& MeshAttributeName = PCGLevelToAssetConstants::MeshAttributeName;
+	const FName& HierarchyDepthAttributeName = PCGLevelToAssetConstants::HierarchyDepthAttributeName;
+	const FName& ActorIndexAttributeName = PCGLevelToAssetConstants::ActorIndexAttributeName;
+	const FName& ParentIndexAttributeName = PCGLevelToAssetConstants::ParentIndexAttributeName;
+	const FName& RelativeTransformAttributeName = PCGLevelToAssetConstants::RelativeTransformAttributeName;
 
 	// Attribute setup on the points
 	FPCGMetadataAttribute<FSoftObjectPath>* MaterialAttribute = PointMetadata->CreateAttribute<FSoftObjectPath>(MaterialAttributeName, FSoftObjectPath(), /*bAllowsInterpolation=*/false, /*bOverrideParent=*/true);
