[PCG] Reduce usage of FOutputDeviceWriteScopeLock by registering one OutputDevice and passing it to each FScopedCall
- Reduces FOutputDevicesWriteScopeLock from 320k calls ~300ms to 87k calls ~22ms for 4 BiomeSample generations

#rnx

diff --git a/Engine/Plugins/PCG/Source/PCG/Private/PCGElement.cpp b/Engine/Plugins/PCG/Source/PCG/Private/PCGElement.cpp
--- a/Engine/Plugins/PCG/Source/PCG/Private/PCGElement.cpp
+++ b/Engine/Plugins/PCG/Source/PCG/Private/PCGElement.cpp
@@ -164,9 +164,11 @@ bool IPCGElement::Execute(FPCGContext* Context) const
 	check(Context && Context->AsyncState.NumAvailableTasks != 0 && Context->CurrentPhase < EPCGExecutionPhase::Done);
 	check(Context->AsyncState.bIsRunningOnMainThread || !CanExecuteOnlyOnMainThread(Context));
 
+	PCGUtils::FScopedCallOutputDevice OutputDevice;
+
 	while (Context->CurrentPhase != EPCGExecutionPhase::Done)
 	{
-		PCGUtils::FScopedCall ScopedCall(*this, Context);
+		PCGUtils::FScopedCall ScopedCall(*this, Context, OutputDevice);
 		bool bExecutionPostponed = false;
 
 		switch (Context->CurrentPhase)
diff --git a/Engine/Plugins/PCG/Source/PCG/Private/Utils/PCGExtraCapture.cpp b/Engine/Plugins/PCG/Source/PCG/Private/Utils/PCGExtraCapture.cpp
--- a/Engine/Plugins/PCG/Source/PCG/Private/Utils/PCGExtraCapture.cpp
+++ b/Engine/Plugins/PCG/Source/PCG/Private/Utils/PCGExtraCapture.cpp
@@ -261,27 +261,39 @@ PCGUtils::FCallTreeInfo PCGUtils::FExtraCapture::CalculateCallTreeInfo(const UPC
 	return RootInfo;
 }
 
-PCGUtils::FScopedCall::FScopedCall(const IPCGElement& InOwner, FPCGContext* InContext)
+PCGUtils::FScopedCall::FScopedCall(const IPCGElement& InOwner, FPCGContext* InContext, FScopedCallOutputDevice& InOutputDevice)
 	: Owner(InOwner)
 	, Context(InContext)
 	, Phase(InContext->CurrentPhase)
-	, ThreadID(FPlatformTLS::GetCurrentThreadId())
+	, OutputDevice(InOutputDevice)
 {
 	StartTime = FPlatformTime::Seconds();
 
-	GLog->AddOutputDevice(this);
+	check(!OutputDevice.ScopedCall);
+	OutputDevice.ScopedCall = this;
 }
 
 PCGUtils::FScopedCall::~FScopedCall()
 {
-	GLog->RemoveOutputDevice(this);
 	if (Context && Context->SourceComponent.IsValid())
 	{
 		Context->SourceComponent->ExtraCapture.Update(*this);
 	}
+	OutputDevice.ScopedCall = nullptr;
+}
+
+PCGUtils::FScopedCallOutputDevice::FScopedCallOutputDevice()
+	: ThreadID(FPlatformTLS::GetCurrentThreadId())
+{
+	GLog->AddOutputDevice(this);
 }
 
-void PCGUtils::FScopedCall::Serialize(const TCHAR* V, ELogVerbosity::Type Verbosity, const class FName& Category)
+PCGUtils::FScopedCallOutputDevice::~FScopedCallOutputDevice()
+{
+	GLog->RemoveOutputDevice(this);
+}
+
+void PCGUtils::FScopedCallOutputDevice::Serialize(const TCHAR* V, ELogVerbosity::Type Verbosity, const class FName& Category)
 {
 	// TODO: this thread id check will also filter out messages spawned from threads spawned inside of nodes. To improve that,
 	// perhaps set at TLS bit on things from here and inside of PCGAsync spawned jobs. If this was done, CapturedMessages below also will
@@ -295,7 +307,10 @@ void PCGUtils::FScopedCall::Serialize(const TCHAR* V, ELogVerbosity::Type Verbos
 	// this is a dumb counter just so messages can be sorted in a similar order as when they were logged
 	static volatile int32 MessageCounter = 0;
 
-	CapturedMessages.Add(PCGUtils::FCapturedMessage{ MessageCounter++, Category, V, Verbosity });
+	if (ScopedCall)
+	{
+		ScopedCall->CapturedMessages.Add(PCGUtils::FCapturedMessage{ MessageCounter++, Category, V, Verbosity });
+	}
 }
 
 #endif // WITH_EDITOR
\ No newline at end of file
diff --git a/Engine/Plugins/PCG/Source/PCG/Public/Utils/PCGExtraCapture.h b/Engine/Plugins/PCG/Source/PCG/Public/Utils/PCGExtraCapture.h
--- a/Engine/Plugins/PCG/Source/PCG/Public/Utils/PCGExtraCapture.h
+++ b/Engine/Plugins/PCG/Source/PCG/Public/Utils/PCGExtraCapture.h
@@ -56,23 +56,35 @@ namespace PCGUtils
 		TArray<FCallTreeInfo> Children;
 	};
 
-	struct FScopedCall : public FOutputDevice
+	struct FScopedCall;
+
+	struct FScopedCallOutputDevice : public FOutputDevice
+	{
+		FScopedCallOutputDevice();
+		virtual ~FScopedCallOutputDevice();
+
+		// FOutputDevice
+		virtual bool IsMemoryOnly() const override { return true; }
+		virtual bool CanBeUsedOnMultipleThreads() const override { return true; }
+		virtual bool CanBeUsedOnAnyThread() const override { return true; }
+		virtual void Serialize(const TCHAR* V, ELogVerbosity::Type Verbosity, const class FName& Category) override;
+
+		FScopedCall* ScopedCall = nullptr;
+		const uint32 ThreadID;
+	};
+
+	struct FScopedCall
 	{
-		FScopedCall(const IPCGElement& InOwner, FPCGContext* InContext);
+		FScopedCall(const IPCGElement& InOwner, FPCGContext* InContext, FScopedCallOutputDevice& InOutputDevice);
 		virtual ~FScopedCall();
 
 		const IPCGElement& Owner;
 		FPCGContext* Context;
 		double StartTime;
 		EPCGExecutionPhase Phase;
-		const uint32 ThreadID;
 		TArray<FCapturedMessage> CapturedMessages;
 
-		// FOutputDevice
-		virtual bool IsMemoryOnly() const override { return true; }
-		virtual bool CanBeUsedOnMultipleThreads() const override { return true; }
-		virtual bool CanBeUsedOnAnyThread() const override { return true; }
-		virtual void Serialize(const TCHAR* V, ELogVerbosity::Type Verbosity, const class FName& Category) override;
+		FScopedCallOutputDevice& OutputDevice;
 	};
 
 	class PCG_API FExtraCapture
@@ -93,9 +105,13 @@ namespace PCGUtils
 		mutable FCriticalSection Lock;
 	};
 #else
+	struct FScopedCallOutputDevice
+	{
+	};
+
 	struct FScopedCall
 	{
-		FScopedCall(const IPCGElement& InOwner, FPCGContext* InContext) {}
+		FScopedCall(const IPCGElement& InOwner, FPCGContext* InContext, FScopedCallOutputDevice& InOutputDevice) {}
 	};
 #endif // WITH_EDITOR
 }
