[PCG] IPCGSpatialDataVisualization support UPCGPointArrayData

#rnx

diff --git a/Engine/Plugins/PCG/Source/PCG/Private/Data/PCGBasePointData.cpp b/Engine/Plugins/PCG/Source/PCG/Private/Data/PCGBasePointData.cpp
--- a/Engine/Plugins/PCG/Source/PCG/Private/Data/PCGBasePointData.cpp
+++ b/Engine/Plugins/PCG/Source/PCG/Private/Data/PCGBasePointData.cpp
@@ -12,6 +12,7 @@
 #include "Metadata/Accessors/IPCGAttributeAccessorTpl.h"
 #include "Metadata/Accessors/PCGAttributeAccessorHelpers.h"
 #include "Metadata/Accessors/PCGAttributeAccessorKeys.h"
+#include "Metadata/Accessors/PCGCustomAccessor.h"
 
 #include "GameFramework/Actor.h"
 #include "Serialization/ArchiveCrc32.h"
@@ -138,304 +139,6 @@ namespace PCGPointHelpers
 
 namespace PCGBasePointData
 {
-	class FPCGAttributeAccessorKeysPointIndices : public IPCGAttributeAccessorKeys
-	{
-	public:
-		FPCGAttributeAccessorKeysPointIndices(UPCGBasePointData* InPointData)
-			: IPCGAttributeAccessorKeys(/*bInReadOnly=*/false)
-			, PointData(InPointData)
-		{}
-
-		FPCGAttributeAccessorKeysPointIndices(const UPCGBasePointData* InPointData)
-			: IPCGAttributeAccessorKeys(/*bInReadOnly=*/true)
-			, PointData(const_cast<UPCGBasePointData*>(InPointData))
-		{}
-
-		virtual bool GetKeyIndices(int32 InStart, int32 InCount, TArray<int32>& OutKeyIndices, bool& bOutContiguous) const override
-		{
-			if (PointData->GetNumPoints() == 0)
-			{
-				return false;
-			}
-
-			// Optimization (avoid allocating indices memory)
-			if (InStart + InCount <= PointData->GetNumPoints())
-			{
-				bOutContiguous = true;
-				return true;
-			}
-
-			OutKeyIndices.SetNumUninitialized(InCount);
-			for (int32 Index = 0; Index < InCount; ++Index)
-			{
-				OutKeyIndices[Index] = (InStart + Index) % PointData->GetNumPoints();
-			}
-
-			return true;
-		}
-
-		virtual int32 GetNum() const override
-		{
-			return PointData ? PointData->GetNumPoints() : 0;
-		}
-
-	protected:
-		virtual bool GetMetadataEntryKeys(int32 InStart, TArrayView<PCGMetadataEntryKey*> OutEntryKeys) override
-		{
-			if (!PointData || PointData->GetNumPoints() == 0)
-			{
-				return false;
-			}
-
-			TPCGValueRange<int64> MetadataEntryKeys = PointData->GetMetadataEntryValueRange();
-
-			for (int32 Index = 0; Index < OutEntryKeys.Num(); ++Index)
-			{
-				int32 EntryIndex = (InStart + Index) % MetadataEntryKeys.Num();
-				OutEntryKeys[Index] = &MetadataEntryKeys[EntryIndex];
-			}
-
-			return true;
-		}
-
-		virtual bool GetMetadataEntryKeys(int32 InStart, TArrayView<const PCGMetadataEntryKey*> OutEntryKeys) const override
-		{
-			if (!PointData || PointData->GetNumPoints() == 0)
-			{
-				return false;
-			}
-
-			const TConstPCGValueRange<int64> MetadataEntryKeys = PointData->GetConstMetadataEntryValueRange();
-
-			for (int32 Index = 0; Index < OutEntryKeys.Num(); ++Index)
-			{
-				int32 EntryIndex = (InStart + Index) % MetadataEntryKeys.Num();
-				OutEntryKeys[Index] = &MetadataEntryKeys[EntryIndex];
-			}
-
-			return true;
-		}
-
-	private:
-		UPCGBasePointData* PointData = nullptr;
-	};
-
-	template<typename T> 
-	bool GetRange(TArrayView<T> OutValues, int32 StartIndex, const IPCGAttributeAccessorKeys& Keys, const auto& Range)
-	{
-		TArray<int32> KeyIndices;
-		bool bContiguous = false;
-		if (!Keys.GetKeyIndices(StartIndex, OutValues.Num(), KeyIndices, bContiguous))
-		{
-			return false;
-		}
-
-		if (bContiguous)
-		{
-			for (int32 OutIndex = 0; OutIndex < OutValues.Num(); ++OutIndex)
-			{
-				OutValues[OutIndex] = static_cast<T>(Range[StartIndex + OutIndex]);
-			}
-		}
-		else
-		{
-			for (int32 OutIndex = 0; OutIndex < OutValues.Num(); ++OutIndex)
-			{
-				OutValues[OutIndex] = static_cast<T>(Range[KeyIndices[OutIndex]]);
-			}
-		}
-
-		return true;
-	}
-
-	template <typename T, typename R = T>
-	class FPCGNativePointPropertyConstAccessor : public IPCGAttributeAccessorT<FPCGNativePointPropertyConstAccessor<T,R>>
-	{
-	public:
-		using Type = T;
-		using Super = IPCGAttributeAccessorT<FPCGNativePointPropertyConstAccessor<T,R>>;
-
-		FPCGNativePointPropertyConstAccessor(const UPCGBasePointData* InPointData, EPCGPointNativeProperties InNativeProperty)
-			: Super(/*bReadOnly=*/true)
-		{
-			ValueRange = InPointData->GetConstValueRange<R>(InNativeProperty);
-		}
-
-		FPCGNativePointPropertyConstAccessor(UPCGBasePointData* InPointData, EPCGPointNativeProperties InNativeProperty)
-			: Super(/*bReadOnly=*/true)
-		{
-			ValueRange = InPointData->GetConstValueRange<R>(InNativeProperty);
-		}
-				
-		bool GetRangeImpl(TArrayView<T> OutValues, int32 StartIndex, const IPCGAttributeAccessorKeys& Keys) const
-		{
-			return PCGBasePointData::GetRange(OutValues, StartIndex, Keys, ValueRange);
-		}
-
-		bool SetRangeImpl(TArrayView<const T> InValues, int32 StartIndex, IPCGAttributeAccessorKeys& Keys, EPCGAttributeAccessorFlags Flags)
-		{
-			return false;
-		}
-
-	private:
-		TConstPCGValueRange<R> ValueRange;
-	};
-
-	template <typename T, typename R = T>
-	class FPCGNativePointPropertyAccessor : public IPCGAttributeAccessorT<FPCGNativePointPropertyAccessor<T, R>>
-	{
-	public:
-		using Type = T;
-		using Super = IPCGAttributeAccessorT<FPCGNativePointPropertyAccessor<T, R>>;
-
-		FPCGNativePointPropertyAccessor(const UPCGBasePointData* InPointData, EPCGPointNativeProperties InNativeProperty)
-			: Super(/*bReadOnly=*/false)
-		{
-			ValueRange = const_cast<UPCGBasePointData*>(InPointData)->GetValueRange<R>(InNativeProperty, /*bAllocate=*/ true);
-		}
-
-		FPCGNativePointPropertyAccessor(UPCGBasePointData* InPointData, EPCGPointNativeProperties InNativeProperty)
-			: Super(/*bReadOnly=*/false)
-		{
-			ValueRange = InPointData->GetValueRange<R>(InNativeProperty, /*bAllocate=*/true);
-		}
-
-		bool GetRangeImpl(TArrayView<T> OutValues, int32 StartIndex, const IPCGAttributeAccessorKeys& Keys) const
-		{
-			return PCGBasePointData::GetRange<T>(OutValues, StartIndex, Keys, ValueRange);
-		}
-
-		bool SetRangeImpl(TArrayView<const T> InValues, int32 StartIndex, IPCGAttributeAccessorKeys& Keys, EPCGAttributeAccessorFlags Flags)
-		{
-			if (Keys.IsReadOnly())
-			{
-				return false;
-			}
-
-			TArray<int32> KeyIndices;
-			bool bContiguous = false;
-			if (!Keys.GetKeyIndices(StartIndex, InValues.Num(), KeyIndices, bContiguous))
-			{
-				return false;
-			}
-
-			if (bContiguous)
-			{
-				for (int32 InIndex = 0; InIndex < InValues.Num(); ++InIndex)
-				{
-					ValueRange[StartIndex + InIndex] = (R)InValues[InIndex];
-				}
-			}
-			else
-			{
-				for (int32 InIndex = 0; InIndex < InValues.Num(); ++InIndex)
-				{
-					ValueRange[KeyIndices[InIndex]] = (R)InValues[InIndex];
-				}
-			}
-
-			return true;
-		}
-
-	private:
-		TPCGValueRange<R> ValueRange;
-	};
-
-	template <typename T, typename... Args>
-	using FPointCustomPropertyGetter = TFunction<bool(int32, T&, Args...)>;
-
-	template <typename T, typename... Args>
-	using FPointCustomPropertySetter = TFunction<bool(int32, const T&, Args...)>;
-
-	template <typename T, typename... Args>
-	class FPCGCustomPointPropertyAccessor : public IPCGAttributeAccessorT<FPCGCustomPointPropertyAccessor<T, Args...>>
-	{
-	public:
-		using Type = T;
-		using Super = IPCGAttributeAccessorT<FPCGCustomPointPropertyAccessor<T, Args...>>;
-
-		FPCGCustomPointPropertyAccessor(const UPCGBasePointData* InPointData, const FPointCustomPropertyGetter<T, Args...>& InGetter, const FPointCustomPropertySetter<T, Args...>& InSetter, Args&& ...InValueRanges)
-			: Super(/*bInReadOnly=*/ false)
-			, PointData(const_cast<UPCGBasePointData*>(InPointData))
-			, Getter(InGetter)
-			, Setter(InSetter)
-			, ValueRanges(std::forward<Args>(InValueRanges)...)
-		{
-		}
-
-		FPCGCustomPointPropertyAccessor(const UPCGBasePointData* InPointData, const FPointCustomPropertyGetter<T, Args...>& InGetter, Args&& ...InValueRanges)
-			: Super(/*bInReadOnly=*/ true)
-			, PointData(const_cast<UPCGBasePointData*>(InPointData))
-			, Getter(InGetter)
-			, ValueRanges(std::forward<Args>(InValueRanges)...)
-		{}
-
-		bool GetRangeImpl(TArrayView<T> OutValues, int32 StartIndex, const IPCGAttributeAccessorKeys& Keys) const
-		{
-			TArray<int32> KeyIndices;
-			bool bContiguous = false;
-			if (!Keys.GetKeyIndices(StartIndex, OutValues.Num(), KeyIndices, bContiguous))
-			{
-				return false;
-			}
-
-			if (bContiguous)
-			{
-				for (int32 OutIndex = 0; OutIndex < OutValues.Num(); ++OutIndex)
-				{
-					ValueRanges.ApplyAfter(Getter, StartIndex + OutIndex, OutValues[OutIndex]);
-				}
-			}
-			else
-			{
-				for (int32 OutIndex = 0; OutIndex < OutValues.Num(); ++OutIndex)
-				{
-					ValueRanges.ApplyAfter(Getter, KeyIndices[OutIndex], OutValues[OutIndex]);
-				}
-			}
-
-			return true;
-		}
-
-		bool SetRangeImpl(TArrayView<const T> InValues, int32 StartIndex, IPCGAttributeAccessorKeys& Keys, EPCGAttributeAccessorFlags Flags)
-		{
-			if (Keys.IsReadOnly())
-			{
-				return false;
-			}
-
-			TArray<int32> KeyIndices;
-			bool bContiguous = false;
-			if (!Keys.GetKeyIndices(StartIndex, InValues.Num(), KeyIndices, bContiguous))
-			{
-				return false;
-			}
-
-			if (bContiguous)
-			{
-				for (int32 InIndex = 0; InIndex < InValues.Num(); ++InIndex)
-				{
-					ValueRanges.ApplyAfter(Setter, StartIndex + InIndex, InValues[InIndex]);
-				}
-			}
-			else
-			{
-				for (int32 InIndex = 0; InIndex < InValues.Num(); ++InIndex)
-				{
-					ValueRanges.ApplyAfter(Setter, KeyIndices[InIndex], InValues[InIndex]);
-				}
-			}
-
-			return true;
-		}
-
-	private:
-		UPCGBasePointData* PointData = nullptr;
-		FPointCustomPropertyGetter<T, Args...> Getter;
-		FPointCustomPropertySetter<T, Args...> Setter;
-		TTuple<Args...> ValueRanges;
-	};
-
 	TUniquePtr<IPCGAttributeAccessor> CreatePropertyAccessor(UPCGBasePointData* PointData, EPCGPointNativeProperties NativeProperty, bool bQuiet)
 	{
 		switch (NativeProperty)
@@ -453,6 +156,8 @@ namespace PCGBasePointData
 			return MakeUnique<FPCGNativePointPropertyAccessor<FVector4>>(PointData, NativeProperty);
 		case EPCGPointNativeProperties::Seed:
 			return MakeUnique<FPCGNativePointPropertyAccessor<int32>>(PointData, NativeProperty);
+		case EPCGPointNativeProperties::MetadataEntry:
+			return MakeUnique<FPCGNativePointPropertyAccessor<int64>>(PointData, NativeProperty);
 		default:
 			UE_LOG(LogPCG, Error, TEXT("EPCGPointNativeProperty value '%d' does not exist."), NativeProperty);
 			return TUniquePtr<IPCGAttributeAccessor>();
@@ -476,6 +181,8 @@ namespace PCGBasePointData
 			return MakeUnique<FPCGNativePointPropertyConstAccessor<FVector4>>(PointData, NativeProperty);
 		case EPCGPointNativeProperties::Seed:
 			return MakeUnique<FPCGNativePointPropertyConstAccessor<int32>>(PointData, NativeProperty);
+		case EPCGPointNativeProperties::MetadataEntry:
+			return MakeUnique<FPCGNativePointPropertyConstAccessor<int64>>(PointData, NativeProperty);
 		default:
 			UE_LOG(LogPCG, Error, TEXT("EPCGPointNativeProperty value '%d' does not exist."), NativeProperty);
 			return TUniquePtr<IPCGAttributeAccessor>();
@@ -739,13 +446,13 @@ FPCGAttributeAccessorMethods UPCGBasePointData::GetPointAccessorMethods()
 	auto CreateAccessorKeysFunc = [](UPCGData* InData, const FPCGAttributePropertySelector& InSelector, bool bQuiet)
 	{
 		UPCGBasePointData* PointData = CastChecked<UPCGBasePointData>(InData);
-		return MakeUnique<PCGBasePointData::FPCGAttributeAccessorKeysPointIndices>(PointData);
+		return MakeUnique<FPCGAttributeAccessorKeysPointIndices>(PointData);
 	};
 
 	auto CreateConstAccessorKeysFunc = [](const UPCGData* InData, const FPCGAttributePropertySelector& InSelector, bool bQuiet) -> TUniquePtr<const IPCGAttributeAccessorKeys>
 	{
 		const UPCGBasePointData* PointData = CastChecked<UPCGBasePointData>(InData);
-		return MakeUnique<const PCGBasePointData::FPCGAttributeAccessorKeysPointIndices>(PointData);
+		return MakeUnique<const FPCGAttributeAccessorKeysPointIndices>(PointData);
 	};
 
 	FPCGAttributeAccessorMethods Methods
diff --git a/Engine/Plugins/PCG/Source/PCG/Private/Elements/PCGVolumeSampler.cpp b/Engine/Plugins/PCG/Source/PCG/Private/Elements/PCGVolumeSampler.cpp
--- a/Engine/Plugins/PCG/Source/PCG/Private/Elements/PCGVolumeSampler.cpp
+++ b/Engine/Plugins/PCG/Source/PCG/Private/Elements/PCGVolumeSampler.cpp
@@ -26,9 +26,14 @@ namespace PCGVolumeSampler
 		return Cast<UPCGPointData>(SampleVolume(Context, UPCGPointData::StaticClass(), SamplerSettings, Volume, BoundingShape));
 	}
 
+	UPCGBasePointData* SampleVolume(FPCGContext* Context, const UPCGSpatialData* Volume, const FVolumeSamplerParams& SamplerSettings, const UPCGSpatialData* BoundingShape)
+	{
+		return SampleVolume(Context, nullptr, SamplerSettings, Volume, BoundingShape);
+	}
+
 	UPCGBasePointData* SampleVolume(FPCGContext* Context, TSubclassOf<UPCGBasePointData> PointDataClass, const FVolumeSamplerParams& SamplerSettings, const UPCGSpatialData* Volume, const UPCGSpatialData* BoundingShape)
 	{
-		UPCGBasePointData* Data = FPCGContext::NewObject_AnyThread<UPCGBasePointData>(Context, GetTransientPackage(), PointDataClass);
+		UPCGBasePointData* Data = PointDataClass != nullptr ? FPCGContext::NewObject_AnyThread<UPCGBasePointData>(Context, GetTransientPackage(), PointDataClass) : FPCGContext::NewPointData_AnyThread(Context);
 
 		FPCGInitializeFromDataParams InitializeFromDataParams(Volume);
 		InitializeFromDataParams.bInheritSpatialData = false;
diff --git a/Engine/Plugins/PCG/Source/PCG/Private/Metadata/Accessors/PCGAttributeAccessorKeys.cpp b/Engine/Plugins/PCG/Source/PCG/Private/Metadata/Accessors/PCGAttributeAccessorKeys.cpp
--- a/Engine/Plugins/PCG/Source/PCG/Private/Metadata/Accessors/PCGAttributeAccessorKeys.cpp
+++ b/Engine/Plugins/PCG/Source/PCG/Private/Metadata/Accessors/PCGAttributeAccessorKeys.cpp
@@ -128,6 +128,29 @@ bool FPCGAttributeAccessorKeysPoints::IsClassSupported(const UStruct* InClass) c
 	return PCGAttributeAccessorKeys::IsClassSupported<FPCGPoint>(InClass);
 }
 
+bool FPCGAttributeAccessorKeysPoints::GetKeyIndices(int32 InStart, int32 InCount, TArray<int32>& OutKeyIndices, bool& bOutContiguous) const
+{
+	if (Points.IsEmpty())
+	{
+		return false;
+	}
+
+	// Optimization (avoid allocating indices memory)
+	if (InStart + InCount <= Points.Num())
+	{
+		bOutContiguous = true;
+		return true;
+	}
+
+	OutKeyIndices.SetNumUninitialized(InCount);
+	for (int32 Index = 0; Index < InCount; ++Index)
+	{
+		OutKeyIndices[Index] = (InStart + Index) % Points.Num();
+	}
+
+	return true;
+}
+
 bool FPCGAttributeAccessorKeysPoints::GetPointKeys(int32 InStart, TArrayView<FPCGPoint*> OutPoints)
 {
 	return PCGAttributeAccessorKeys::GetKeys(Points, InStart, OutPoints, [](FPCGPoint& Point) -> FPCGPoint* { return &Point; });
diff --git a/Engine/Plugins/PCG/Source/PCG/Public/Elements/PCGVolumeSampler.h b/Engine/Plugins/PCG/Source/PCG/Public/Elements/PCGVolumeSampler.h
--- a/Engine/Plugins/PCG/Source/PCG/Public/Elements/PCGVolumeSampler.h
+++ b/Engine/Plugins/PCG/Source/PCG/Public/Elements/PCGVolumeSampler.h
@@ -53,9 +53,11 @@ namespace PCGVolumeSampler
 	/** Sample a volume and returns the resulting point data. */
 	PCG_API UPCGBasePointData* SampleVolume(FPCGContext* Context, TSubclassOf<UPCGBasePointData> PointDataClass, const FVolumeSamplerParams& SamplerSettings, const UPCGSpatialData* Volume, const UPCGSpatialData* BoundingShape = nullptr);
 
-	UE_DEPRECATED(5.6, "Use SampleVolume taking a UPCGBasePoint subclass instead")
+	UE_DEPRECATED(5.6, "Use SampleVolume returning UPCGBasePointData")
 	PCG_API UPCGPointData* SampleVolume(FPCGContext* Context, const FVolumeSamplerParams& SamplerSettings, const UPCGSpatialData* Volume, const UPCGSpatialData* BoundingShape = nullptr);
 
+	PCG_API UPCGBasePointData* SampleVolume(FPCGContext* Context, const UPCGSpatialData* Volume, const FVolumeSamplerParams& SamplerSettings, const UPCGSpatialData* BoundingShape = nullptr);
+
 	/** Sample a volume and write the results in the given point data. Can be timesliced and will return false if the processing is not done, true otherwise. */
 	PCG_API bool SampleVolume(FPCGContext* Context, const FVolumeSamplerParams& SamplerSettings, const UPCGSpatialData* Volume, const UPCGSpatialData* BoundingShape, UPCGBasePointData* OutputData, const bool bTimeSlicingIsEnabled = false);
 }
diff --git a/Engine/Plugins/PCG/Source/PCG/Public/Metadata/Accessors/PCGAttributeAccessorKeys.h b/Engine/Plugins/PCG/Source/PCG/Public/Metadata/Accessors/PCGAttributeAccessorKeys.h
--- a/Engine/Plugins/PCG/Source/PCG/Public/Metadata/Accessors/PCGAttributeAccessorKeys.h
+++ b/Engine/Plugins/PCG/Source/PCG/Public/Metadata/Accessors/PCGAttributeAccessorKeys.h
@@ -221,6 +221,8 @@ public:
 	virtual bool IsClassSupported(const UStruct* InClass) const override;
 
 protected:
+	virtual bool GetKeyIndices(int32 InStart, int32 InCount, TArray<int32>& OutKeyIndices, bool& bOutContiguous) const override;
+
 	virtual bool GetPointKeys(int32 InStart, TArrayView<FPCGPoint*> OutPoints) override;
 	virtual bool GetPointKeys(int32 InStart, TArrayView<const FPCGPoint*> OutPoints) const override;
 
diff --git a/Engine/Plugins/PCG/Source/PCG/Public/Metadata/Accessors/PCGCustomAccessor.h b/Engine/Plugins/PCG/Source/PCG/Public/Metadata/Accessors/PCGCustomAccessor.h
--- a/Engine/Plugins/PCG/Source/PCG/Public/Metadata/Accessors/PCGCustomAccessor.h
+++ b/Engine/Plugins/PCG/Source/PCG/Public/Metadata/Accessors/PCGCustomAccessor.h
@@ -5,6 +5,7 @@
 #include "IPCGAttributeAccessorTpl.h"
 
 #include "PCGPoint.h"
+#include "Data/PCGBasePointData.h"
 
 /**
 * Templated accessor class for custom point properties. Need a getter and a setter, defined in the
@@ -214,3 +215,307 @@ public:
 		return false;
 	}
 };
+
+class FPCGAttributeAccessorKeysPointIndices : public IPCGAttributeAccessorKeys
+{
+public:
+	FPCGAttributeAccessorKeysPointIndices(UPCGBasePointData* InPointData)
+		: IPCGAttributeAccessorKeys(/*bInReadOnly=*/false)
+		, PointData(InPointData)
+	{
+	}
+
+	FPCGAttributeAccessorKeysPointIndices(const UPCGBasePointData* InPointData)
+		: IPCGAttributeAccessorKeys(/*bInReadOnly=*/true)
+		, PointData(const_cast<UPCGBasePointData*>(InPointData))
+	{
+	}
+
+	virtual bool GetKeyIndices(int32 InStart, int32 InCount, TArray<int32>& OutKeyIndices, bool& bOutContiguous) const override
+	{
+		if (PointData->GetNumPoints() == 0)
+		{
+			return false;
+		}
+
+		// Optimization (avoid allocating indices memory)
+		if (InStart + InCount <= PointData->GetNumPoints())
+		{
+			bOutContiguous = true;
+			return true;
+		}
+
+		OutKeyIndices.SetNumUninitialized(InCount);
+		for (int32 Index = 0; Index < InCount; ++Index)
+		{
+			OutKeyIndices[Index] = (InStart + Index) % PointData->GetNumPoints();
+		}
+
+		return true;
+	}
+
+	virtual int32 GetNum() const override
+	{
+		return PointData ? PointData->GetNumPoints() : 0;
+	}
+
+protected:
+	virtual bool GetMetadataEntryKeys(int32 InStart, TArrayView<PCGMetadataEntryKey*> OutEntryKeys) override
+	{
+		if (!PointData || PointData->GetNumPoints() == 0)
+		{
+			return false;
+		}
+
+		TPCGValueRange<int64> MetadataEntryKeys = PointData->GetMetadataEntryValueRange();
+
+		for (int32 Index = 0; Index < OutEntryKeys.Num(); ++Index)
+		{
+			int32 EntryIndex = (InStart + Index) % MetadataEntryKeys.Num();
+			OutEntryKeys[Index] = &MetadataEntryKeys[EntryIndex];
+		}
+
+		return true;
+	}
+
+	virtual bool GetMetadataEntryKeys(int32 InStart, TArrayView<const PCGMetadataEntryKey*> OutEntryKeys) const override
+	{
+		if (!PointData || PointData->GetNumPoints() == 0)
+		{
+			return false;
+		}
+
+		const TConstPCGValueRange<int64> MetadataEntryKeys = PointData->GetConstMetadataEntryValueRange();
+
+		for (int32 Index = 0; Index < OutEntryKeys.Num(); ++Index)
+		{
+			int32 EntryIndex = (InStart + Index) % MetadataEntryKeys.Num();
+			OutEntryKeys[Index] = &MetadataEntryKeys[EntryIndex];
+		}
+
+		return true;
+	}
+
+private:
+	UPCGBasePointData* PointData = nullptr;
+};
+
+namespace PCGCustomAccessor
+{
+	template<typename T>
+	bool GetRange(TArrayView<T> OutValues, int32 StartIndex, const IPCGAttributeAccessorKeys& Keys, const auto& Range)
+	{
+		TArray<int32> KeyIndices;
+		bool bContiguous = false;
+		if (!Keys.GetKeyIndices(StartIndex, OutValues.Num(), KeyIndices, bContiguous))
+		{
+			return false;
+		}
+
+		if (bContiguous)
+		{
+			for (int32 OutIndex = 0; OutIndex < OutValues.Num(); ++OutIndex)
+			{
+				OutValues[OutIndex] = static_cast<T>(Range[StartIndex + OutIndex]);
+			}
+		}
+		else
+		{
+			for (int32 OutIndex = 0; OutIndex < OutValues.Num(); ++OutIndex)
+			{
+				OutValues[OutIndex] = static_cast<T>(Range[KeyIndices[OutIndex]]);
+			}
+		}
+
+		return true;
+	}
+}
+
+template <typename T, typename R = T>
+class FPCGNativePointPropertyConstAccessor : public IPCGAttributeAccessorT<FPCGNativePointPropertyConstAccessor<T, R>>
+{
+public:
+	using Type = T;
+	using Super = IPCGAttributeAccessorT<FPCGNativePointPropertyConstAccessor<T, R>>;
+
+	FPCGNativePointPropertyConstAccessor(const UPCGBasePointData* InPointData, EPCGPointNativeProperties InNativeProperty)
+		: Super(/*bReadOnly=*/true)
+	{
+		ValueRange = InPointData->GetConstValueRange<R>(InNativeProperty);
+	}
+
+	FPCGNativePointPropertyConstAccessor(UPCGBasePointData* InPointData, EPCGPointNativeProperties InNativeProperty)
+		: Super(/*bReadOnly=*/true)
+	{
+		ValueRange = InPointData->GetConstValueRange<R>(InNativeProperty);
+	}
+
+	bool GetRangeImpl(TArrayView<T> OutValues, int32 StartIndex, const IPCGAttributeAccessorKeys& Keys) const
+	{
+		return PCGCustomAccessor::GetRange(OutValues, StartIndex, Keys, ValueRange);
+	}
+
+	bool SetRangeImpl(TArrayView<const T> InValues, int32 StartIndex, IPCGAttributeAccessorKeys& Keys, EPCGAttributeAccessorFlags Flags)
+	{
+		return false;
+	}
+
+private:
+	TConstPCGValueRange<R> ValueRange;
+};
+
+template <typename T, typename R = T>
+class FPCGNativePointPropertyAccessor : public IPCGAttributeAccessorT<FPCGNativePointPropertyAccessor<T, R>>
+{
+public:
+	using Type = T;
+	using Super = IPCGAttributeAccessorT<FPCGNativePointPropertyAccessor<T, R>>;
+
+	FPCGNativePointPropertyAccessor(const UPCGBasePointData* InPointData, EPCGPointNativeProperties InNativeProperty)
+		: Super(/*bReadOnly=*/false)
+	{
+		ValueRange = const_cast<UPCGBasePointData*>(InPointData)->GetValueRange<R>(InNativeProperty, /*bAllocate=*/ true);
+	}
+
+	FPCGNativePointPropertyAccessor(UPCGBasePointData* InPointData, EPCGPointNativeProperties InNativeProperty)
+		: Super(/*bReadOnly=*/false)
+	{
+		ValueRange = InPointData->GetValueRange<R>(InNativeProperty, /*bAllocate=*/true);
+	}
+
+	bool GetRangeImpl(TArrayView<T> OutValues, int32 StartIndex, const IPCGAttributeAccessorKeys& Keys) const
+	{
+		return PCGCustomAccessor::GetRange<T>(OutValues, StartIndex, Keys, ValueRange);
+	}
+
+	bool SetRangeImpl(TArrayView<const T> InValues, int32 StartIndex, IPCGAttributeAccessorKeys& Keys, EPCGAttributeAccessorFlags Flags)
+	{
+		if (Keys.IsReadOnly())
+		{
+			return false;
+		}
+
+		TArray<int32> KeyIndices;
+		bool bContiguous = false;
+		if (!Keys.GetKeyIndices(StartIndex, InValues.Num(), KeyIndices, bContiguous))
+		{
+			return false;
+		}
+
+		if (bContiguous)
+		{
+			for (int32 InIndex = 0; InIndex < InValues.Num(); ++InIndex)
+			{
+				ValueRange[StartIndex + InIndex] = (R)InValues[InIndex];
+			}
+		}
+		else
+		{
+			for (int32 InIndex = 0; InIndex < InValues.Num(); ++InIndex)
+			{
+				ValueRange[KeyIndices[InIndex]] = (R)InValues[InIndex];
+			}
+		}
+
+		return true;
+	}
+
+private:
+	TPCGValueRange<R> ValueRange;
+};
+
+template <typename T, typename... Args>
+using FPointCustomPropertyGetter = TFunction<bool(int32, T&, Args...)>;
+
+template <typename T, typename... Args>
+using FPointCustomPropertySetter = TFunction<bool(int32, const T&, Args...)>;
+
+template <typename T, typename... Args>
+class FPCGCustomPointPropertyAccessor : public IPCGAttributeAccessorT<FPCGCustomPointPropertyAccessor<T, Args...>>
+{
+public:
+	using Type = T;
+	using Super = IPCGAttributeAccessorT<FPCGCustomPointPropertyAccessor<T, Args...>>;
+
+	FPCGCustomPointPropertyAccessor(const UPCGBasePointData* InPointData, const FPointCustomPropertyGetter<T, Args...>& InGetter, const FPointCustomPropertySetter<T, Args...>& InSetter, Args&& ...InValueRanges)
+		: Super(/*bInReadOnly=*/ false)
+		, PointData(const_cast<UPCGBasePointData*>(InPointData))
+		, Getter(InGetter)
+		, Setter(InSetter)
+		, ValueRanges(std::forward<Args>(InValueRanges)...)
+	{
+	}
+
+	FPCGCustomPointPropertyAccessor(const UPCGBasePointData* InPointData, const FPointCustomPropertyGetter<T, Args...>& InGetter, Args&& ...InValueRanges)
+		: Super(/*bInReadOnly=*/ true)
+		, PointData(const_cast<UPCGBasePointData*>(InPointData))
+		, Getter(InGetter)
+		, ValueRanges(std::forward<Args>(InValueRanges)...)
+	{
+	}
+
+	bool GetRangeImpl(TArrayView<T> OutValues, int32 StartIndex, const IPCGAttributeAccessorKeys& Keys) const
+	{
+		TArray<int32> KeyIndices;
+		bool bContiguous = false;
+		if (!Keys.GetKeyIndices(StartIndex, OutValues.Num(), KeyIndices, bContiguous))
+		{
+			return false;
+		}
+
+		if (bContiguous)
+		{
+			for (int32 OutIndex = 0; OutIndex < OutValues.Num(); ++OutIndex)
+			{
+				ValueRanges.ApplyAfter(Getter, StartIndex + OutIndex, OutValues[OutIndex]);
+			}
+		}
+		else
+		{
+			for (int32 OutIndex = 0; OutIndex < OutValues.Num(); ++OutIndex)
+			{
+				ValueRanges.ApplyAfter(Getter, KeyIndices[OutIndex], OutValues[OutIndex]);
+			}
+		}
+
+		return true;
+	}
+
+	bool SetRangeImpl(TArrayView<const T> InValues, int32 StartIndex, IPCGAttributeAccessorKeys& Keys, EPCGAttributeAccessorFlags Flags)
+	{
+		if (Keys.IsReadOnly())
+		{
+			return false;
+		}
+
+		TArray<int32> KeyIndices;
+		bool bContiguous = false;
+		if (!Keys.GetKeyIndices(StartIndex, InValues.Num(), KeyIndices, bContiguous))
+		{
+			return false;
+		}
+
+		if (bContiguous)
+		{
+			for (int32 InIndex = 0; InIndex < InValues.Num(); ++InIndex)
+			{
+				ValueRanges.ApplyAfter(Setter, StartIndex + InIndex, InValues[InIndex]);
+			}
+		}
+		else
+		{
+			for (int32 InIndex = 0; InIndex < InValues.Num(); ++InIndex)
+			{
+				ValueRanges.ApplyAfter(Setter, KeyIndices[InIndex], InValues[InIndex]);
+			}
+		}
+
+		return true;
+	}
+
+private:
+	UPCGBasePointData* PointData = nullptr;
+	FPointCustomPropertyGetter<T, Args...> Getter;
+	FPointCustomPropertySetter<T, Args...> Setter;
+	TTuple<Args...> ValueRanges;
+};
diff --git a/Engine/Plugins/PCG/Source/PCGEditor/Private/DataVisualizations/PCGLandscapeDataVisualization.cpp b/Engine/Plugins/PCG/Source/PCGEditor/Private/DataVisualizations/PCGLandscapeDataVisualization.cpp
--- a/Engine/Plugins/PCG/Source/PCGEditor/Private/DataVisualizations/PCGLandscapeDataVisualization.cpp
+++ b/Engine/Plugins/PCG/Source/PCGEditor/Private/DataVisualizations/PCGLandscapeDataVisualization.cpp
@@ -4,10 +4,10 @@
 
 #include "PCGEditorSettings.h"
 #include "Data/PCGLandscapeData.h"
-#include "Data/PCGPointData.h"
+#include "Data/PCGBasePointData.h"
 #include "Elements/PCGSurfaceSampler.h"
 
-const UPCGPointData* IPCGLandscapeDataVisualization::CollapseToDebugPointData(FPCGContext* Context, const UPCGData* Data) const
+const UPCGBasePointData* IPCGLandscapeDataVisualization::CollapseToDebugBasePointData(FPCGContext* Context, const UPCGData* Data) const
 {
 	const UPCGLandscapeData* LandscapeData = Cast<UPCGLandscapeData>(Data);
 	
@@ -30,5 +30,5 @@ const UPCGPointData* IPCGLandscapeDataVisualization::CollapseToDebugPointData(FP
 		.bApplyDensityToPoints = false
 	};
 	
-	return PCGSurfaceSampler::SampleSurface(Context, LandscapeData, nullptr, Bounds, SamplerParams);
+	return PCGSurfaceSampler::SampleSurface(Context, SamplerParams, LandscapeData, nullptr, Bounds);
 }
diff --git a/Engine/Plugins/PCG/Source/PCGEditor/Private/DataVisualizations/PCGSpatialDataVisualization.cpp b/Engine/Plugins/PCG/Source/PCGEditor/Private/DataVisualizations/PCGSpatialDataVisualization.cpp
--- a/Engine/Plugins/PCG/Source/PCGEditor/Private/DataVisualizations/PCGSpatialDataVisualization.cpp
+++ b/Engine/Plugins/PCG/Source/PCGEditor/Private/DataVisualizations/PCGSpatialDataVisualization.cpp
@@ -8,6 +8,7 @@
 #include "PCGEditorCommon.h"
 #include "PCGPoint.h"
 #include "PCGSettings.h"
+#include "Data/PCGPointArrayData.h"
 #include "Data/PCGPointData.h"
 #include "Data/PCGSpatialData.h"
 #include "DataVisualizations/PCGDataVisualizationHelpers.h"
@@ -103,27 +104,26 @@ void IPCGSpatialDataVisualization::ExecuteDebugDisplay(FPCGContext* Context, con
 	TArray<TSoftObjectPtr<UMaterialInterface>> Materials;
 	Materials.Add(DebugSettings.GetMaterial());
 
-	const UPCGPointData* PointData = CollapseToDebugPointData(Context, Data);
+	const UPCGBasePointData* PointData = CollapseToDebugBasePointData(Context, Data);
 
 	if (!PointData)
 	{
 		return;
 	}
 
-	const TArray<FPCGPoint>& Points = PointData->GetPoints();
-
-	if (Points.Num() == 0)
+	if (PointData->IsEmpty())
 	{
 		return;
 	}
 
 	const int NumCustomData = 8;
+	const int32 NumPoints = PointData->GetNumPoints();
 
 	TArray<FTransform> ForwardInstances;
 	TArray<FTransform> ReverseInstances;
 	TArray<float> InstanceCustomData;
 
-	ForwardInstances.Reserve(Points.Num());
+	ForwardInstances.Reserve(NumPoints);
 	InstanceCustomData.Reserve(NumCustomData);
 
 	// First, create target instance transforms
@@ -137,25 +137,29 @@ void IPCGSpatialDataVisualization::ExecuteDebugDisplay(FPCGContext* Context, con
 	 * the renderer is robust to over-instancing attempts and will not crash. However, MAX_INSTANCE_ID still serves as a
 	 * good, scalable heuristic for a max limit.
 	 */
-	const int32 NumDesiredInstances = FMath::Min(Points.Num(), static_cast<int32>(MAX_INSTANCE_ID));
-	if (NumDesiredInstances != Points.Num())
+	const int32 NumDesiredInstances = FMath::Min(NumPoints, static_cast<int32>(MAX_INSTANCE_ID));
+	if (NumDesiredInstances != NumPoints)
 	{
-		PCGE_LOG_C(Error, GraphAndLog, Context, FText::Format(LOCTEXT( "DebugPointsOverLimit", "Debug point display ({0}) surpassed the max instance limit ({1}) and will be clamped."), FText::AsNumber(Points.Num()), FText::AsNumber(MAX_INSTANCE_ID)));
+		PCGE_LOG_C(Error, GraphAndLog, Context, FText::Format(LOCTEXT( "DebugPointsOverLimit", "Debug point display ({0}) surpassed the max instance limit ({1}) and will be clamped."), FText::AsNumber(NumPoints), FText::AsNumber(MAX_INSTANCE_ID)));
 	}
 
+	const FConstPCGPointValueRanges ValueRanges(PointData);
+
 	for (int i = 0; i < NumDesiredInstances; ++i)
 	{
-		const FPCGPoint& Point = Points[i];
-		TArray<FTransform>& Instances = ((bIsAbsolute || Point.Transform.GetDeterminant() >= 0) ? ForwardInstances : ReverseInstances);
-		FTransform& InstanceTransform = Instances.Add_GetRef(Point.Transform);
+		TArray<FTransform>& Instances = ((bIsAbsolute || ValueRanges.TransformRange[i].GetDeterminant() >= 0) ? ForwardInstances : ReverseInstances);
+		FTransform& InstanceTransform = Instances.Add_GetRef(ValueRanges.TransformRange[i]);
 		if (bIsRelative)
 		{
 			InstanceTransform.SetScale3D(InstanceTransform.GetScale3D() * PointScale);
 		}
 		else if (bScaleWithExtents)
 		{
-			const FVector ScaleWithExtents = Point.GetExtents() / MeshExtents;
-			const FVector TransformedBoxCenterWithOffset = InstanceTransform.TransformPosition(Point.GetLocalCenter()) - InstanceTransform.GetLocation();
+			const FVector Extents = PCGPointHelpers::GetExtents(ValueRanges.BoundsMinRange[i], ValueRanges.BoundsMaxRange[i]);
+			const FVector LocalCenter = PCGPointHelpers::GetLocalCenter(ValueRanges.BoundsMinRange[i], ValueRanges.BoundsMaxRange[i]);
+
+			const FVector ScaleWithExtents = Extents / MeshExtents;
+			const FVector TransformedBoxCenterWithOffset = InstanceTransform.TransformPosition(LocalCenter) - InstanceTransform.GetLocation();
 			InstanceTransform.SetTranslation(InstanceTransform.GetTranslation() + TransformedBoxCenterWithOffset);
 			InstanceTransform.SetScale3D(InstanceTransform.GetScale3D() * ScaleWithExtents);
 		}
@@ -233,23 +237,25 @@ void IPCGSpatialDataVisualization::ExecuteDebugDisplay(FPCGContext* Context, con
 
 		// Scan all points looking for points that match current direction and add their custom data.
 		int32 PointCounter = 0;
-		for (const FPCGPoint& Point : Points)
+		for (int32 PointIndex = 0; PointIndex < NumPoints; ++PointIndex)
 		{
-			const int32 PointDirection = ((bIsAbsolute || Point.Transform.GetDeterminant() >= 0) ? 0 : 1);
+			const int32 PointDirection = ((bIsAbsolute || ValueRanges.TransformRange[PointIndex].GetDeterminant() >= 0) ? 0 : 1);
 			if (PointDirection != Direction)
 			{
 				continue;
 			}
 
-			InstanceCustomData.Add(Point.Density);
-			const FVector Extents = Point.GetExtents();
+			const FVector4& Color = ValueRanges.ColorRange[PointIndex];
+			const FVector Extents = PCGPointHelpers::GetExtents(ValueRanges.BoundsMinRange[PointIndex], ValueRanges.BoundsMaxRange[PointIndex]);
+			InstanceCustomData.Add(ValueRanges.DensityRange[PointIndex]);
+			
 			InstanceCustomData.Add(Extents[0]);
 			InstanceCustomData.Add(Extents[1]);
 			InstanceCustomData.Add(Extents[2]);
-			InstanceCustomData.Add(Point.Color[0]);
-			InstanceCustomData.Add(Point.Color[1]);
-			InstanceCustomData.Add(Point.Color[2]);
-			InstanceCustomData.Add(Point.Color[3]);
+			InstanceCustomData.Add(Color[0]);
+			InstanceCustomData.Add(Color[1]);
+			InstanceCustomData.Add(Color[2]);
+			InstanceCustomData.Add(Color[3]);
 
 			ISMC->SetCustomData(PreExistingInstanceCount + PointCounter, InstanceCustomData);
 
@@ -268,7 +274,7 @@ FPCGTableVisualizerInfo IPCGSpatialDataVisualization::GetTableVisualizerInfo(con
 	using namespace PCGPointDataVisualizationConstants;
 
 	// Collapse to point representation for visualization.
-	const UPCGPointData* PointData = CollapseToDebugPointData(/*Context=*/nullptr, Data);
+	const UPCGBasePointData* PointData = CollapseToDebugBasePointData(/*Context=*/nullptr, Data);
 
 	FPCGTableVisualizerInfo Info;
 	Info.Data = PointData;
@@ -300,22 +306,25 @@ FPCGTableVisualizerInfo IPCGSpatialDataVisualization::GetTableVisualizerInfo(con
 
 	if (FPCGEditorCommon::CVarShowAdvancedAttributesFields.GetValueOnAnyThread())
 	{
-		Info.ColumnInfos.Add(CreateColumnInfo(PointData, NAME_MetadataEntry, TEXT_MetadataEntry, []()
+		Info.ColumnInfos.Add(CreateColumnInfo(PointData, NAME_MetadataEntry, TEXT_MetadataEntry, [PointData]()
 		{
-			return TSharedPtr<const IPCGAttributeAccessor>(PCGAttributeAccessorHelpers::CreatePropertyAccessor(GET_MEMBER_NAME_CHECKED(FPCGPoint, MetadataEntry), FPCGPoint::StaticStruct()).Release());
+			FPCGAttributePropertySelector MetadataEntrySelector;
+			MetadataEntrySelector.SetPropertyName(*StaticEnum<EPCGPointNativeProperties>()->GetNameStringByValue((int64)EPCGPointNativeProperties::MetadataEntry));
+			return TSharedPtr<const IPCGAttributeAccessor>(PCGAttributeAccessorHelpers::CreateConstAccessor(PointData, MetadataEntrySelector).Release());
 		}));
 
 		Info.ColumnInfos.Add(CreateColumnInfo(PointData, NAME_MetadataEntryParent, TEXT_MetadataEntryParent, [PointData]()
 		{
-			return MakeShared<FPCGCustomPointAccessor<int64>>([PointData](const FPCGPoint& Point, void* OutValue)
+			return MakeShared<FPCGCustomPointPropertyAccessor<int64, TConstPCGValueRange<int64>>>(PointData, [Metadata = PointData->ConstMetadata()](int32 Index, int64& OutValue, const TConstPCGValueRange<int64>& MetaDataEntryRange)
 			{
-				if (const UPCGMetadata* Metadata = PointData->ConstMetadata())
+				if (Metadata)
 				{
-					*reinterpret_cast<int64*>(OutValue) = Metadata->GetParentKey(Point.MetadataEntry);
+					OutValue = Metadata->GetParentKey(MetaDataEntryRange[Index]);
 					return true;
 				}
 				return false;
-			}, nullptr);
+			},
+			PointData->GetConstMetadataEntryValueRange());
 		}));
 	}
 
@@ -330,15 +339,13 @@ FPCGTableVisualizerInfo IPCGSpatialDataVisualization::GetTableVisualizerInfo(con
 	{
 		if (const UPCGSpatialData* SpatialData = Cast<UPCGSpatialData>(Data))
 		{
-			const UPCGPointData* PointData = SpatialData->ToPointData(nullptr);
+			const UPCGBasePointData* PointData = SpatialData->ToBasePointData(nullptr);
 
 			if (!PointData)
 			{
 				return;
 			}
 
-			const TArray<FPCGPoint>& Points = PointData->GetPoints();
-
 			FBox BoundingBox(EForceInit::ForceInit);
 			if (Indices.IsEmpty())
 			{
@@ -346,12 +353,11 @@ FPCGTableVisualizerInfo IPCGSpatialDataVisualization::GetTableVisualizerInfo(con
 			}
 			else
 			{
+				const FConstPCGPointValueRanges ValueRanges(PointData);
 				for (const int& Index : Indices)
 				{
-					check(Points.IsValidIndex(Index));
-
-					const FPCGPoint& Point = Points[Index];
-					const FBox PointBoundingBox = Point.GetLocalBounds().TransformBy(Point.Transform.ToMatrixWithScale());
+					const FBox LocalBounds = PCGPointHelpers::GetLocalBounds(ValueRanges.BoundsMinRange[Index], ValueRanges.BoundsMaxRange[Index]);
+					const FBox PointBoundingBox = LocalBounds.TransformBy(ValueRanges.TransformRange[Index].ToMatrixWithScale());
 
 					BoundingBox += PointBoundingBox;
 				}
@@ -380,4 +386,23 @@ const UPCGPointData* IPCGSpatialDataVisualization::CollapseToDebugPointData(FPCG
 	return nullptr;
 }
 
+const UPCGBasePointData* IPCGSpatialDataVisualization::CollapseToDebugBasePointData(FPCGContext* Context, const UPCGData* Data) const
+{	
+	if (const UPCGSpatialData* SpatialData = Cast<UPCGSpatialData>(Data))
+	{
+		if (CVarPCGEnablePointArrayData.GetValueOnAnyThread())
+		{
+			return SpatialData->ToPointArrayData(Context);
+		}
+		else
+		{
+PRAGMA_DISABLE_DEPRECATION_WARNINGS
+			return CollapseToDebugPointData(Context, Data);
+PRAGMA_ENABLE_DEPRECATION_WARNINGS
+		}
+	}
+
+	return nullptr;
+}
+
 #undef LOCTEXT_NAMESPACE
diff --git a/Engine/Plugins/PCG/Source/PCGEditor/Private/DataVisualizations/PCGSplineDataVisualization.cpp b/Engine/Plugins/PCG/Source/PCGEditor/Private/DataVisualizations/PCGSplineDataVisualization.cpp
--- a/Engine/Plugins/PCG/Source/PCGEditor/Private/DataVisualizations/PCGSplineDataVisualization.cpp
+++ b/Engine/Plugins/PCG/Source/PCGEditor/Private/DataVisualizations/PCGSplineDataVisualization.cpp
@@ -2,6 +2,7 @@
 
 #include "DataVisualizations/PCGSplineDataVisualization.h"
 
+#include "PCGContext.h"
 #include "Data/PCGPointData.h"
 #include "Data/PCGSplineData.h"
 #include "Data/PCGSplineStruct.h"
@@ -52,23 +53,27 @@ FPCGTableVisualizerInfo IPCGSplineDataVisualization::GetTableVisualizerInfo(cons
 	return Info;
 }
 
-const UPCGPointData* IPCGSplineDataVisualization::CollapseToDebugPointData(FPCGContext* Context, const UPCGData* Data) const
+const UPCGBasePointData* IPCGSplineDataVisualization::CollapseToDebugBasePointData(FPCGContext* Context, const UPCGData* Data) const
 {
 	if (const UPCGSplineData* SplineData = Cast<UPCGSplineData>(Data))
 	{
-		UPCGPointData* PointData = NewObject<UPCGPointData>();
+		UPCGBasePointData* PointData = FPCGContext::NewPointData_AnyThread(Context);
 		PointData->InitializeFromData(SplineData);
-		TArray<FPCGPoint>& Points = PointData->GetMutablePoints();
 
 		const int32 NumControlPoints = SplineData->SplineStruct.GetSplinePointsPosition().Points.Num();
 		TConstArrayView<PCGMetadataEntryKey> EntryKeys = SplineData->SplineStruct.GetConstControlPointsEntryKeys();
 
+		PointData->SetNumPoints(NumControlPoints);
+		PointData->SetExtents(PCGSplineDataVisualizationConstants::HalfExtents);
+		PointData->AllocateProperties(EPCGPointNativeProperties::Transform | EPCGPointNativeProperties::MetadataEntry);
+
+		TPCGValueRange<FTransform> TransformRange = PointData->GetTransformValueRange(/*bAllocate=*/false);
+		TPCGValueRange<int64> MetadataEntryRange = PointData->GetMetadataEntryValueRange(/*bAllocate=*/false);
+
 		for (int ControlPointIndex = 0; ControlPointIndex < NumControlPoints; ++ControlPointIndex)
 		{
-			FPCGPoint& Point = Points.Emplace_GetRef();
-			Point.Transform = SplineData->SplineStruct.GetTransformAtSplineInputKey(ControlPointIndex, ESplineCoordinateSpace::World);
-			Point.SetExtents(PCGSplineDataVisualizationConstants::HalfExtents);
-			Point.MetadataEntry = EntryKeys.IsValidIndex(ControlPointIndex) ? EntryKeys[ControlPointIndex] : PCGInvalidEntryKey;
+			TransformRange[ControlPointIndex] = SplineData->SplineStruct.GetTransformAtSplineInputKey(ControlPointIndex, ESplineCoordinateSpace::World);
+			MetadataEntryRange[ControlPointIndex] = EntryKeys.IsValidIndex(ControlPointIndex) ? EntryKeys[ControlPointIndex] : PCGInvalidEntryKey;
 		}
 
 		return PointData;
diff --git a/Engine/Plugins/PCG/Source/PCGEditor/Private/DataVisualizations/PCGVolumeDataVisualization.cpp b/Engine/Plugins/PCG/Source/PCGEditor/Private/DataVisualizations/PCGVolumeDataVisualization.cpp
--- a/Engine/Plugins/PCG/Source/PCGEditor/Private/DataVisualizations/PCGVolumeDataVisualization.cpp
+++ b/Engine/Plugins/PCG/Source/PCGEditor/Private/DataVisualizations/PCGVolumeDataVisualization.cpp
@@ -4,7 +4,7 @@
 
 #include "Data/PCGCollisionShapeData.h"
 #include "Data/PCGCollisionWrapperData.h"
-#include "Data/PCGPointData.h"
+#include "Data/PCGBasePointData.h"
 #include "Data/PCGPrimitiveData.h"
 #include "Data/PCGSpatialData.h"
 #include "Data/PCGVolumeData.h"
@@ -12,7 +12,7 @@
 
 #include "PCGEditorSettings.h"
 
-const UPCGPointData* IPCGPrimitiveDataVisualization::CollapseToDebugPointData(FPCGContext* Context, const UPCGData* Data) const
+const UPCGBasePointData* IPCGPrimitiveDataVisualization::CollapseToDebugBasePointData(FPCGContext* Context, const UPCGData* Data) const
 {
 	const UPCGSpatialData* SpatialData = Cast<UPCGSpatialData>(Data);
 	
@@ -36,8 +36,7 @@ const UPCGPointData* IPCGPrimitiveDataVisualization::CollapseToDebugPointData(FP
 		SamplerParams.VoxelSize = FVector(VoxelSize, VoxelSize, VoxelSize);
 		SamplerParams.Bounds = Bounds;
 
-		// @todo_pcg: convert IPCGPrimitiveDataVisualization to support UPCGBasePointData
-		return Cast<UPCGPointData>(PCGVolumeSampler::SampleVolume(Context, UPCGPointData::StaticClass(), SamplerParams, SpatialData));
+		return PCGVolumeSampler::SampleVolume(Context, SpatialData, SamplerParams);
 	}
 	else
 	{
diff --git a/Engine/Plugins/PCG/Source/PCGEditor/Public/DataVisualizations/PCGLandscapeDataVisualization.h b/Engine/Plugins/PCG/Source/PCGEditor/Public/DataVisualizations/PCGLandscapeDataVisualization.h
--- a/Engine/Plugins/PCG/Source/PCGEditor/Public/DataVisualizations/PCGLandscapeDataVisualization.h
+++ b/Engine/Plugins/PCG/Source/PCGEditor/Public/DataVisualizations/PCGLandscapeDataVisualization.h
@@ -13,6 +13,6 @@ class PCGEDITOR_API IPCGLandscapeDataVisualization : public IPCGSpatialDataVisua
 {
 public:
 	// ~Begin IPCGSpatialDataVisualization interface
-	virtual const UPCGPointData* CollapseToDebugPointData(FPCGContext* Context, const UPCGData* Data) const override;
+	virtual const UPCGBasePointData* CollapseToDebugBasePointData(FPCGContext* Context, const UPCGData* Data) const override;
 	// ~End IPCGSpatialDataVisualization interface
 };
diff --git a/Engine/Plugins/PCG/Source/PCGEditor/Public/DataVisualizations/PCGSpatialDataVisualization.h b/Engine/Plugins/PCG/Source/PCGEditor/Public/DataVisualizations/PCGSpatialDataVisualization.h
--- a/Engine/Plugins/PCG/Source/PCGEditor/Public/DataVisualizations/PCGSpatialDataVisualization.h
+++ b/Engine/Plugins/PCG/Source/PCGEditor/Public/DataVisualizations/PCGSpatialDataVisualization.h
@@ -7,6 +7,7 @@
 
 class AActor;
 class UPCGData;
+class UPCGBasePointData;
 class UPCGPointData;
 struct FPCGContext;
 
@@ -19,5 +20,8 @@ public:
 	virtual FPCGTableVisualizerInfo GetTableVisualizerInfo(const UPCGData* Data) const override;
 	// ~End IPCGDataVisualization interface
 
+	UE_DEPRECATED(5.6, "Implement CollapseToDebugBasePointData instead")
 	virtual const UPCGPointData* CollapseToDebugPointData(FPCGContext* Context, const UPCGData* Data) const;
+
+	virtual const UPCGBasePointData* CollapseToDebugBasePointData(FPCGContext* Context, const UPCGData* Data) const;
 };
diff --git a/Engine/Plugins/PCG/Source/PCGEditor/Public/DataVisualizations/PCGSplineDataVisualization.h b/Engine/Plugins/PCG/Source/PCGEditor/Public/DataVisualizations/PCGSplineDataVisualization.h
--- a/Engine/Plugins/PCG/Source/PCGEditor/Public/DataVisualizations/PCGSplineDataVisualization.h
+++ b/Engine/Plugins/PCG/Source/PCGEditor/Public/DataVisualizations/PCGSplineDataVisualization.h
@@ -6,7 +6,6 @@
 #include "PCGEditorModule.h"
 
 class UPCGData;
-class UPCGPointData;
 struct FPCGContext;
 
 class PCGEDITOR_API IPCGSplineDataVisualization : public IPCGSpatialDataVisualization
@@ -18,6 +17,6 @@ public:
 
 	// ~Begin IPCGSpatialDataVisualization interface
 	/** Overrides collapse behavior to show the spline control points. */
-	virtual const UPCGPointData* CollapseToDebugPointData(FPCGContext* Context, const UPCGData* Data) const override;
+	virtual const UPCGBasePointData* CollapseToDebugBasePointData(FPCGContext* Context, const UPCGData* Data) const override;
 	// ~End IPCGSpatialDataVisualization interface
 };
diff --git a/Engine/Plugins/PCG/Source/PCGEditor/Public/DataVisualizations/PCGVolumeDataVisualization.h b/Engine/Plugins/PCG/Source/PCGEditor/Public/DataVisualizations/PCGVolumeDataVisualization.h
--- a/Engine/Plugins/PCG/Source/PCGEditor/Public/DataVisualizations/PCGVolumeDataVisualization.h
+++ b/Engine/Plugins/PCG/Source/PCGEditor/Public/DataVisualizations/PCGVolumeDataVisualization.h
@@ -5,7 +5,6 @@
 #include "PCGSpatialDataVisualization.h"
 
 class UPCGData;
-class UPCGPointData;
 struct FPCGContext;
 
 /** Overrides collapse behavior to target a given number of points generated, to avoid too coarse debug (points way bigger than data) or enormous amount of points (points way smaller than data). */
@@ -13,6 +12,6 @@ class PCGEDITOR_API IPCGPrimitiveDataVisualization : public IPCGSpatialDataVisua
 {
 public:
 	// ~Begin IPCGSpatialDataVisualization interface
-	virtual const UPCGPointData* CollapseToDebugPointData(FPCGContext* Context, const UPCGData* Data) const override;
+	virtual const UPCGBasePointData* CollapseToDebugBasePointData(FPCGContext* Context, const UPCGData* Data) const override;
 	// ~End IPCGSpatialDataVisualization interface
 };
