[PCG] Cache PCGModule ptr to avoid locking/slow GetModuleChecked
- In BiomeSample with 4 generations this was costing about 7 seconds of out of GT time (175k calls)

#rnx

diff --git a/Engine/Plugins/PCG/Source/PCG/Private/Metadata/PCGMetadata.cpp b/Engine/Plugins/PCG/Source/PCG/Private/Metadata/PCGMetadata.cpp
--- a/Engine/Plugins/PCG/Source/PCG/Private/Metadata/PCGMetadata.cpp
+++ b/Engine/Plugins/PCG/Source/PCG/Private/Metadata/PCGMetadata.cpp
@@ -303,6 +303,7 @@ void UPCGMetadata::Serialize(FArchive& InArchive)
 
 void UPCGMetadata::AddToCrc(FArchiveCrc32& Ar, bool bFullDataCrc) const
 {
+	TRACE_CPUPROFILER_EVENT_SCOPE(UPCGMetadata::AddToCrc);
 	const UPCGData* Data = Cast<UPCGData>(GetOuter());
 	check(Data);
 	
diff --git a/Engine/Plugins/PCG/Source/PCG/Private/PCGModule.cpp b/Engine/Plugins/PCG/Source/PCG/Private/PCGModule.cpp
--- a/Engine/Plugins/PCG/Source/PCG/Private/PCGModule.cpp
+++ b/Engine/Plugins/PCG/Source/PCG/Private/PCGModule.cpp
@@ -25,9 +25,14 @@
 
 #define LOCTEXT_NAMESPACE "FPCGModule"
 
+namespace PCGModule
+{
+	FPCGModule* PCGModulePtr = nullptr;
+}
+
 FPCGModule& FPCGModule::GetPCGModuleChecked()
 {
-	return FModuleManager::GetModuleChecked<FPCGModule>(TEXT("PCG"));
+	return PCGModule::PCGModulePtr ? *PCGModule::PCGModulePtr : FModuleManager::GetModuleChecked<FPCGModule>(TEXT("PCG"));
 }
 
 bool FPCGModule::IsPCGModuleLoaded()
@@ -48,6 +53,10 @@ void FPCGModule::StartupModule()
 	AddShaderSourceDirectoryMapping(TEXT("/Plugin/PCG"), PluginShaderDir);
 #endif // WITH_EDITOR
 
+	// Cache for fast access
+	check(!PCGModule::PCGModulePtr);
+	PCGModule::PCGModulePtr = this;
+
 	// Registering accessor methods
 	AttributeAccessorFactory.RegisterDefaultMethods();
 	AttributeAccessorFactory.RegisterMethods<UPCGBasePointData>(UPCGBasePointData::GetPointAccessorMethods());
@@ -62,6 +71,8 @@ void FPCGModule::StartupModule()
 void FPCGModule::ShutdownModule()
 {
 	FCoreDelegates::OnPreExit.RemoveAll(this);
+
+	PCGModule::PCGModulePtr = nullptr;
 }
 
 void FPCGModule::PreExit()
