[PCG] HLODLayer Targeting
- Partition Actors will inherit HLOD Layer of their Original Component
- SpawnActor/CreateTargetActor have 3 modes (Self, Template, Reference)

diff --git a/Engine/Plugins/PCG/Source/PCG/Private/Elements/PCGCreateTargetActor.cpp b/Engine/Plugins/PCG/Source/PCG/Private/Elements/PCGCreateTargetActor.cpp
--- a/Engine/Plugins/PCG/Source/PCG/Private/Elements/PCGCreateTargetActor.cpp
+++ b/Engine/Plugins/PCG/Source/PCG/Private/Elements/PCGCreateTargetActor.cpp
@@ -24,6 +24,7 @@
 
 #if WITH_EDITOR
 #include "WorldPartition/DataLayer/DataLayerInstance.h"
+#include "WorldPartition/HLOD/HLODLayer.h"
 #endif
 
 #include UE_INLINE_GENERATED_CPP_BY_NAME(PCGCreateTargetActor)
@@ -92,6 +93,12 @@ namespace PCGCreateTargetActor
 			UDataLayerInstance* NonConstDataLayerInstance = const_cast<UDataLayerInstance*>(DataLayerInstance);
 			Ar << NonConstDataLayerInstance;
 		}
+
+		if (InParams.HLODLayer)
+		{
+			UHLODLayer* HLODLayer = InParams.HLODLayer;
+			Ar << HLODLayer;
+		}
 #endif
 		
 		return FPCGCrc(Ar.GetCrc());
@@ -386,6 +393,9 @@ bool FPCGCreateTargetActorElement::ExecuteInternal(FPCGContext* Context) const
 #if WITH_EDITOR
 	int32 DataLayerCrc;
 	SpawnDefaultActorParams.DataLayerInstances = PCGDataLayerHelpers::GetDataLayerInstancesAndCrc(Context, Settings->DataLayerSettings, TargetActor, DataLayerCrc);
+
+	int32 HLODLayerCrc;
+	SpawnDefaultActorParams.HLODLayer = PCGHLODHelpers::GetHLODLayerAndCrc(Context, Settings->HLODSettings, TargetActor, TemplateActor, HLODLayerCrc);
 #endif
 
 	UPCGManagedActors* ReusableResource = nullptr;
@@ -406,6 +416,11 @@ bool FPCGCreateTargetActorElement::ExecuteInternal(FPCGContext* Context) const
 			{
 				Context->DependenciesCrc.Combine(DataLayerCrc);
 			}
+
+			if (HLODLayerCrc != 0)
+			{
+				Context->DependenciesCrc.Combine(HLODLayerCrc);
+			}
 #endif
 		}
 
diff --git a/Engine/Plugins/PCG/Source/PCG/Private/Elements/PCGSpawnActor.cpp b/Engine/Plugins/PCG/Source/PCG/Private/Elements/PCGSpawnActor.cpp
--- a/Engine/Plugins/PCG/Source/PCG/Private/Elements/PCGSpawnActor.cpp
+++ b/Engine/Plugins/PCG/Source/PCG/Private/Elements/PCGSpawnActor.cpp
@@ -24,6 +24,11 @@
 #include "ISMPartition/ISMComponentDescriptor.h"
 #include "UObject/Package.h"
 
+#if WITH_EDITOR
+#include "Misc/Crc.h"
+#include "WorldPartition/HLOD/HLODLayer.h"
+#endif
+
 #include UE_INLINE_GENERATED_CPP_BY_NAME(PCGSpawnActor)
 
 #define LOCTEXT_NAMESPACE "PCGSpawnActorElement"
@@ -859,6 +864,9 @@ void FPCGSpawnActorElement::SpawnActors(FPCGSubgraphContext* Context, AActor* Ta
 #if WITH_EDITOR
 	int32 DataLayerCrc;
 	SpawnDefaultActorParams.DataLayerInstances = PCGDataLayerHelpers::GetDataLayerInstancesAndCrc(Context, Settings->DataLayerSettings, TargetActor, DataLayerCrc);
+
+	int32 HLODLayerCrc;
+	SpawnDefaultActorParams.HLODLayer = PCGHLODHelpers::GetHLODLayerAndCrc(Context, Settings->HLODSettings, TargetActor, TemplateActor, HLODLayerCrc);
 #endif
 
 	const bool bForceCallGenerate = (Settings->GenerationTrigger == EPCGSpawnActorGenerationTrigger::ForceGenerate);
@@ -888,6 +896,11 @@ void FPCGSpawnActorElement::SpawnActors(FPCGSubgraphContext* Context, AActor* Ta
 		{
 			InputDependenciesCrc.Combine(DataLayerCrc);
 		}
+
+		if (HLODLayerCrc != 0)
+		{
+			InputDependenciesCrc.Combine(HLODLayerCrc);
+		}
 #endif
 
 		if (InputDependenciesCrc.IsValid())
diff --git a/Engine/Plugins/PCG/Source/PCG/Private/Grid/PCGGridDescriptor.cpp b/Engine/Plugins/PCG/Source/PCG/Private/Grid/PCGGridDescriptor.cpp
--- a/Engine/Plugins/PCG/Source/PCG/Private/Grid/PCGGridDescriptor.cpp
+++ b/Engine/Plugins/PCG/Source/PCG/Private/Grid/PCGGridDescriptor.cpp
@@ -7,16 +7,23 @@
 #include "Serialization/ArchiveCrc32.h"
 #include "WorldPartition/WorldPartition.h"
 #include "WorldPartition/DataLayer/DataLayerManager.h"
+#include "WorldPartition/HLOD/HLODLayer.h"
 
 namespace PCGGridDescriptor
 {
-	void Serialize(FArchive& Ar, const TArray<TSoftObjectPtr<UDataLayerAsset>>& InDataLayerAssets)
+	void Serialize(FArchive& Ar, const TArray<TSoftObjectPtr<UDataLayerAsset>>& InDataLayerAssets, TSoftObjectPtr<UHLODLayer> InHLODLayer)
 	{
 		for (const TSoftObjectPtr<UDataLayerAsset>& DataLayerAsset : InDataLayerAssets)
 		{
 			FString DataLayerAssetString = DataLayerAsset.ToString();
 			Ar << DataLayerAssetString;
 		}
+
+		if (!InHLODLayer.IsNull())
+		{
+			FString HLODAssetString = InHLODLayer.ToString();
+			Ar << HLODAssetString;
+		}
 	}
 }
 
@@ -26,7 +33,7 @@ bool FPCGGridDescriptor::operator==(const FPCGGridDescriptor& Other) const
 	if (GridSize == Other.GridSize && bIs2DGrid == Other.bIs2DGrid && bIsRuntime == Other.bIsRuntime && RuntimeHash == Other.RuntimeHash)
 	{
 #if WITH_EDITORONLY_DATA
-		return DataLayerAssets == Other.DataLayerAssets;
+		return DataLayerAssets == Other.DataLayerAssets && HLODLayer == Other.HLODLayer;
 #else
 		return true;
 #endif
@@ -45,7 +52,7 @@ uint32 FPCGGridDescriptor::ComputeHash() const
 
 #if WITH_EDITORONLY_DATA	
 	// Hashing of Editor only properties
-	PCGGridDescriptor::Serialize(Ar, DataLayerAssets);
+	PCGGridDescriptor::Serialize(Ar, DataLayerAssets, HLODLayer);
 #endif
 
 	return Ar.GetCrc();
@@ -60,7 +67,7 @@ uint32 FPCGGridDescriptor::GetRuntimeHash() const
 	}
 
 	FArchiveCrc32 Ar;
-	PCGGridDescriptor::Serialize(Ar, DataLayerAssets);
+	PCGGridDescriptor::Serialize(Ar, DataLayerAssets, HLODLayer);
 	return Ar.GetCrc();
 }
 
@@ -130,11 +137,8 @@ FString FPCGGridDescriptor::GetPartitionActorName(const FIntVector& GridCoords)
 	// For Data Layer Instances, build a hash to add as a suffix
 	{
 		FArchiveCrc32 Ar;
-		if (DataLayerAssets.Num() > 0)
-		{
-			PCGGridDescriptor::Serialize(Ar, DataLayerAssets);
-		}
-
+		PCGGridDescriptor::Serialize(Ar, DataLayerAssets, HLODLayer);
+		
 		uint32 NameHash = Ar.GetCrc();
 		if (NameHash != 0)
 		{
diff --git a/Engine/Plugins/PCG/Source/PCG/Private/Grid/PCGPartitionActor.cpp b/Engine/Plugins/PCG/Source/PCG/Private/Grid/PCGPartitionActor.cpp
--- a/Engine/Plugins/PCG/Source/PCG/Private/Grid/PCGPartitionActor.cpp
+++ b/Engine/Plugins/PCG/Source/PCG/Private/Grid/PCGPartitionActor.cpp
@@ -341,6 +341,7 @@ FPCGGridDescriptor APCGPartitionActor::GetGridDescriptor() const
 	else
 	{
 		GridDescriptor.SetDataLayerAssets(GetDataLayerAssets());
+		GridDescriptor.SetHLODLayer(GetHLODLayer());
 	}
 #else
 	GridDescriptor.SetRuntimeHash(RuntimeGridDescriptorHash);
@@ -641,6 +642,8 @@ void APCGPartitionActor::PostCreation(const FPCGGridDescriptor& GridDescriptor)
 		ensure(ExternalDataLayerAsset == DescriptorExternalDataLayerAsset);
 
 		DataLayerAssets = DescriptorDataLayerAssets;
+
+		SetHLODLayer(GridDescriptor.GetHLODLayer().LoadSynchronous());
 	}
 
 	// Set only once upon creation, can't change for a Partition Actor
diff --git a/Engine/Plugins/PCG/Source/PCG/Private/Grid/PCGPartitionActorDesc.cpp b/Engine/Plugins/PCG/Source/PCG/Private/Grid/PCGPartitionActorDesc.cpp
--- a/Engine/Plugins/PCG/Source/PCG/Private/Grid/PCGPartitionActorDesc.cpp
+++ b/Engine/Plugins/PCG/Source/PCG/Private/Grid/PCGPartitionActorDesc.cpp
@@ -24,7 +24,8 @@ FPCGGridDescriptor FPCGPartitionActorDesc::GetGridDescriptor(const FWorldPartiti
 	FPCGGridDescriptor GridDescriptor = FPCGGridDescriptor()
 		.SetGridSize(GridSize)
 		.SetIs2DGrid(bUse2DGrid)
-		.SetIsRuntime(false);
+		.SetIsRuntime(false)
+		.SetHLODLayer(TSoftObjectPtr<UHLODLayer>(GetHLODLayer()));
 
 	if (IsUsingDataLayerAsset())
 	{
diff --git a/Engine/Plugins/PCG/Source/PCG/Private/Helpers/PCGActorHelpers.cpp b/Engine/Plugins/PCG/Source/PCG/Private/Helpers/PCGActorHelpers.cpp
--- a/Engine/Plugins/PCG/Source/PCG/Private/Helpers/PCGActorHelpers.cpp
+++ b/Engine/Plugins/PCG/Source/PCG/Private/Helpers/PCGActorHelpers.cpp
@@ -741,6 +741,8 @@ AActor* UPCGActorHelpers::SpawnDefaultActor(const FSpawnDefaultActorParams& Para
 #if WITH_EDITOR
 	if (!Params.World->IsGameWorld())
 	{
+		NewActor->SetHLODLayer(Params.HLODLayer);
+
 		// Add remaining DataLayers (except External which was done on spawn)
 		for (const UDataLayerInstance* DataLayerInstance : Params.DataLayerInstances)
 		{
diff --git a/Engine/Plugins/PCG/Source/PCG/Private/Helpers/PCGHLODHelpers.cpp b/Engine/Plugins/PCG/Source/PCG/Private/Helpers/PCGHLODHelpers.cpp
new file mode 100644
--- /dev/null
+++ b/Engine/Plugins/PCG/Source/PCG/Private/Helpers/PCGHLODHelpers.cpp
@@ -0,0 +1,61 @@
+// Copyright Epic Games, Inc. All Rights Reserved.
+
+#include "Helpers/PCGHLODHelpers.h"
+
+#define LOCTEXT_NAMESPACE "PCGHLODHelpers"
+
+#if WITH_EDITOR
+#include "PCGContext.h"
+
+#include "GameFramework/Actor.h"
+#include "Serialization/ArchiveCrc32.h"
+#include "WorldPartition/HLOD/HLODLayer.h"
+#endif
+
+namespace PCGHLODHelpers
+{
+#if WITH_EDITOR
+	UHLODLayer* GetHLODLayerAndCrc(FPCGContext* Context, const FPCGHLODSettings& HLODSettings, AActor* DefaultHLODLayerSource, AActor* TemplateActor, int32& OutCrc)
+	{
+		UHLODLayer* HLODLayer = nullptr;
+		
+		switch (HLODSettings.HLODSourceType)
+		{
+		case EPCGHLODSource::Self:
+		{
+			if(ensure(DefaultHLODLayerSource))
+			{
+				HLODLayer = DefaultHLODLayerSource->GetHLODLayer();
+			}
+			break;
+		}
+		case EPCGHLODSource::Template:
+		{
+			if (ensure(TemplateActor))
+			{
+				HLODLayer = TemplateActor->GetHLODLayer();
+			}
+			break;
+		}
+		case EPCGHLODSource::Reference:
+		{
+			HLODLayer = HLODSettings.HLODLayer;
+			break;
+		}
+		}
+
+		OutCrc = 0;
+		if (HLODLayer)
+		{
+			FArchiveCrc32 Ar;
+			FString HLODLayerPath = HLODLayer->GetPathName();
+			Ar << HLODLayerPath;
+			OutCrc = Ar.GetCrc();
+		}
+
+		return HLODLayer;
+	}
+#endif // WITH_EDITOR
+}
+
+#undef LOCTEXT_NAMESPACE
\ No newline at end of file
diff --git a/Engine/Plugins/PCG/Source/PCG/Private/PCGComponent.cpp b/Engine/Plugins/PCG/Source/PCG/Private/PCGComponent.cpp
--- a/Engine/Plugins/PCG/Source/PCG/Private/PCGComponent.cpp
+++ b/Engine/Plugins/PCG/Source/PCG/Private/PCGComponent.cpp
@@ -218,6 +218,7 @@ FPCGGridDescriptor UPCGComponent::GetGridDescriptorInternal(uint32 GridSize, boo
 	else
 	{
 		PCGGridDescriptor.SetDataLayerAssets(GetOwner()->GetDataLayerAssets());
+		PCGGridDescriptor.SetHLODLayer(GetOwner()->GetHLODLayer());
 	}
 #else
 	PCGGridDescriptor.SetRuntimeHash(RuntimeGridDescriptorHash);
@@ -1123,6 +1124,7 @@ AActor* UPCGComponent::ClearPCGLink(UClass* TemplateActorClass)
 
 #if WITH_EDITOR
 	SpawnDefaultActorParams.DataLayerInstances = GetOwner()->GetDataLayerInstances();
+	SpawnDefaultActorParams.HLODLayer = GetOwner()->GetHLODLayer();
 #endif
 
 	// First create a new actor that will be the new owner of all the resources
@@ -1319,6 +1321,7 @@ bool UPCGComponent::MoveResourcesToNewActor(AActor* InNewActor, bool bCreateChil
 
 #if WITH_EDITOR
 		SpawnDefaultActorParams.DataLayerInstances = Owner->GetDataLayerInstances();
+		SpawnDefaultActorParams.HLODLayer = Owner->GetHLODLayer();
 #endif
 
 		NewActor = UPCGActorHelpers::SpawnDefaultActor(SpawnDefaultActorParams);
diff --git a/Engine/Plugins/PCG/Source/PCG/Public/Elements/PCGCreateTargetActor.h b/Engine/Plugins/PCG/Source/PCG/Public/Elements/PCGCreateTargetActor.h
--- a/Engine/Plugins/PCG/Source/PCG/Public/Elements/PCGCreateTargetActor.h
+++ b/Engine/Plugins/PCG/Source/PCG/Public/Elements/PCGCreateTargetActor.h
@@ -4,6 +4,7 @@
 
 #include "PCGSettings.h"
 #include "Helpers/PCGDataLayerHelpers.h"
+#include "Helpers/PCGHLODHelpers.h"
 #include "Metadata/PCGObjectPropertyOverride.h"
 
 #include "WorldPartition/DataLayer/DataLayerAsset.h"
@@ -82,9 +83,12 @@ public:
 	UPROPERTY(BlueprintReadWrite, EditAnywhere, Category = Settings)
 	TArray<FName> PostProcessFunctionNames;
 
-	UPROPERTY(BlueprintReadWrite, EditAnywhere, Category = Settings, meta = (PCG_Overridable))
+	UPROPERTY(BlueprintReadWrite, EditAnywhere, Category = Streaming, meta = (PCG_Overridable))
 	FPCGDataLayerSettings DataLayerSettings;
 
+	UPROPERTY(BlueprintReadWrite, EditAnywhere, Category = Streaming, DisplayName = "HLOD Settings", meta = (PCG_Overridable))
+	FPCGHLODSettings HLODSettings;
+
 	UPROPERTY(BlueprintReadWrite, EditAnywhere, Category = "Settings|Advanced")
 	bool bDeleteActorsBeforeGeneration = false;
 
diff --git a/Engine/Plugins/PCG/Source/PCG/Public/Elements/PCGSpawnActor.h b/Engine/Plugins/PCG/Source/PCG/Public/Elements/PCGSpawnActor.h
--- a/Engine/Plugins/PCG/Source/PCG/Public/Elements/PCGSpawnActor.h
+++ b/Engine/Plugins/PCG/Source/PCG/Public/Elements/PCGSpawnActor.h
@@ -4,6 +4,7 @@
 
 #include "PCGSubgraph.h"
 #include "Helpers/PCGDataLayerHelpers.h"
+#include "Helpers/PCGHLODHelpers.h"
 #include "Metadata/PCGObjectPropertyOverride.h"
 
 #include "PCGSpawnActor.generated.h"
@@ -92,9 +93,12 @@ public:
 	UPROPERTY(BlueprintReadWrite, EditAnywhere, Category = "Settings|Advanced")
 	bool bDeleteActorsBeforeGeneration = false;
 
-	UPROPERTY(BlueprintReadWrite, EditAnywhere, Category = Settings, meta = (PCG_Overridable))
+	UPROPERTY(BlueprintReadWrite, EditAnywhere, Category = Streaming, meta = (PCG_Overridable))
 	FPCGDataLayerSettings DataLayerSettings;
 
+	UPROPERTY(BlueprintReadWrite, EditAnywhere, Category = Streaming, DisplayName = "HLOD Settings", meta = (PCG_Overridable))
+	FPCGHLODSettings HLODSettings;
+
 protected:
 	UPROPERTY(BlueprintReadOnly, EditAnywhere, Category = Settings, meta = (OnlyPlaceable, DisallowCreateNew))
 	TSubclassOf<AActor> TemplateActorClass = nullptr;
diff --git a/Engine/Plugins/PCG/Source/PCG/Public/Grid/PCGGridDescriptor.h b/Engine/Plugins/PCG/Source/PCG/Public/Grid/PCGGridDescriptor.h
--- a/Engine/Plugins/PCG/Source/PCG/Public/Grid/PCGGridDescriptor.h
+++ b/Engine/Plugins/PCG/Source/PCG/Public/Grid/PCGGridDescriptor.h
@@ -8,6 +8,7 @@
 #include "Misc/Optional.h"
 #include "UObject/SoftObjectPtr.h"
 #include "WorldPartition/DataLayer/DataLayerAsset.h"
+#include "WorldPartition/HLOD/HLODLayer.h"
 
 #include "PCGGridDescriptor.generated.h"
 
@@ -35,6 +36,9 @@ private:
 #if WITH_EDITORONLY_DATA
 	UPROPERTY()
 	TArray<TSoftObjectPtr<UDataLayerAsset>> DataLayerAssets;
+
+	UPROPERTY()
+	TSoftObjectPtr<UHLODLayer> HLODLayer;
 #endif // WITH_EDITORONLY_DATA
 		
 	// Some descriptor properties do not exist in non editor builds / game worlds. This is the hash that represents those editor only properties.
@@ -62,6 +66,9 @@ public:
 	const TArray<TSoftObjectPtr<UDataLayerAsset>>& GetDataLayerAssets() const { return DataLayerAssets; }
 
 	void GetDataLayerAssets(TArray<TSoftObjectPtr<UDataLayerAsset>>& OutDataLayerAssets, const UExternalDataLayerAsset*& OutExternalDataLayerAsset) const;
+
+	FPCGGridDescriptor& SetHLODLayer(const TSoftObjectPtr<UHLODLayer>& InHLODLayer) { check(!Hash.IsSet()); HLODLayer = InHLODLayer; return *this; }
+	const TSoftObjectPtr<UHLODLayer>& GetHLODLayer() const { return HLODLayer; }
 #endif // WITH_EDITOR
 
 	bool operator==(const FPCGGridDescriptor& Other) const;
diff --git a/Engine/Plugins/PCG/Source/PCG/Public/Helpers/PCGActorHelpers.h b/Engine/Plugins/PCG/Source/PCG/Public/Helpers/PCGActorHelpers.h
--- a/Engine/Plugins/PCG/Source/PCG/Public/Helpers/PCGActorHelpers.h
+++ b/Engine/Plugins/PCG/Source/PCG/Public/Helpers/PCGActorHelpers.h
@@ -20,6 +20,7 @@ class AActor;
 struct FPCGContext;
 class UActorComponent;
 class UDataLayerInstance;
+class UHLODLayer;
 class UInstancedStaticMeshComponent;
 class UInstancedSkinnedMeshComponent;
 class ULevel;
@@ -212,6 +213,7 @@ PRAGMA_ENABLE_DEPRECATION_WARNINGS
 		AActor* Parent = nullptr;
 		bool bForceStaticMobility = true;
 #if WITH_EDITOR
+		UHLODLayer* HLODLayer = nullptr;
 		TArray<const UDataLayerInstance*> DataLayerInstances;
 #endif
 	};
diff --git a/Engine/Plugins/PCG/Source/PCG/Public/Helpers/PCGHLODHelpers.h b/Engine/Plugins/PCG/Source/PCG/Public/Helpers/PCGHLODHelpers.h
new file mode 100644
--- /dev/null
+++ b/Engine/Plugins/PCG/Source/PCG/Public/Helpers/PCGHLODHelpers.h
@@ -0,0 +1,47 @@
+// Copyright Epic Games, Inc. All Rights Reserved.
+
+#pragma once
+
+#include "Engine/EngineTypes.h"
+
+class UHLODLayer;
+
+class AActor;
+
+struct FPCGContext;
+struct FPCGHLODSettings;
+
+#include "PCGHLODHelpers.generated.h"
+
+namespace PCGHLODHelpers
+{
+	namespace Constants
+	{
+		const FName HLODLayerAttribute = TEXT("HLODLayer");
+	}
+
+#if WITH_EDITOR
+	UHLODLayer* GetHLODLayerAndCrc(FPCGContext* Context, const FPCGHLODSettings& HLODSettings, AActor* DefaultHLODLayerSource, AActor* TemplateActor, int32& OutCrc);
+#endif
+};
+
+UENUM()
+enum class EPCGHLODSource : uint8
+{
+	Self,
+	Reference,
+	Template,
+};
+
+USTRUCT(BlueprintType)
+struct FPCGHLODSettings
+{
+	GENERATED_BODY()
+
+	/** What source should be used to assign HLOD Layer to the spawned actor */
+	UPROPERTY(BlueprintReadWrite, EditAnywhere, Category = Default)
+	EPCGHLODSource HLODSourceType = EPCGHLODSource::Self;
+
+	UPROPERTY(BlueprintReadWrite, EditAnywhere, Category = Default, meta = (EditCondition = "HLODSourceType==EPCGHLODSource::Reference", EditConditionHides))
+	TObjectPtr<UHLODLayer> HLODLayer;
+};
\ No newline at end of file
