[PCG] Fix SpatialData registering accessor method

#rnx

diff --git a/Engine/Plugins/PCG/Source/PCG/Private/Metadata/Accessors/PCGAttributeAccessorFactory.cpp b/Engine/Plugins/PCG/Source/PCG/Private/Metadata/Accessors/PCGAttributeAccessorFactory.cpp
--- a/Engine/Plugins/PCG/Source/PCG/Private/Metadata/Accessors/PCGAttributeAccessorFactory.cpp
+++ b/Engine/Plugins/PCG/Source/PCG/Private/Metadata/Accessors/PCGAttributeAccessorFactory.cpp
@@ -276,7 +276,7 @@ void FPCGAttributeAccessorFactory::RegisterDefaultMethods()
 	};
 
 	RegisterMethods<UPCGData>(std::move(DefaultMethods));
-	RegisterMethods<UPCGSpatialData>(std::move(DefaultMethods));
+	RegisterMethods<UPCGSpatialData>(std::move(SpatialMethods));
 }
 
 void FPCGAttributeAccessorFactory::UnregisterDefaultMethods()
