Temporarily expose PCG functions needed for some prototype work where a separate PCG executor needs to run on the game thread.

diff --git a/Engine/Plugins/PCG/Source/PCG/Private/Helpers/PCGInternalHelpers.cpp b/Engine/Plugins/PCG/Source/PCG/Private/Helpers/PCGInternalHelpers.cpp
--- a/Engine/Plugins/PCG/Source/PCG/Private/Helpers/PCGInternalHelpers.cpp
+++ b/Engine/Plugins/PCG/Source/PCG/Private/Helpers/PCGInternalHelpers.cpp
@@ -20,7 +20,25 @@ namespace PCGInternalHelpersLocals
 	};
 }
 
-bool UE::FPCGInternalHelpers::RunGraphLocallyBlocking(UPCGComponent* PCGComponent, const FPCGDataCollection& InputData, 
+TSharedPtr<FPCGGraphExecutor> UE::FPCGInternalHelpers::CreateExecutor()
+{
+	return MakeShared<FPCGGraphExecutor>();
+}
+
+void UE::FPCGInternalHelpers::ExecuteExecutor(FPCGGraphExecutor& Executor)
+{
+	Executor.Execute();
+}
+
+#if WITH_EDITOR
+void UE::FPCGInternalHelpers::NotifyGraphChanged(FPCGGraphExecutor& Executor, UPCGGraph* Graph, EPCGChangeType ChangeType)
+{
+	Executor.NotifyGraphChanged(Graph, ChangeType);
+}
+#endif // WITH_EDITOR
+
+bool UE::FPCGInternalHelpers::ScheduleOnExecutorAndWait(FPCGGraphExecutor& Executor,
+	UPCGComponent* PCGComponent, const FPCGDataCollection& InputData,
 	FPCGDataCollection& OutputData, int32 CyclesToWait, double SleepDurationPerCycle)
 {
 	using namespace PCGInternalHelpersLocals;
@@ -36,8 +54,7 @@ bool UE::FPCGInternalHelpers::RunGraphLocallyBlocking(UPCGComponent* PCGComponen
 		return false;
 	}
 
-	FPCGGraphExecutor LocalGraphExecutor;
-	const FPCGTaskId TaskId = LocalGraphExecutor.Schedule(Subgraph, PCGComponent, nullptr,
+	const FPCGTaskId TaskId = Executor.Schedule(Subgraph, PCGComponent, nullptr,
 		MakeShared<FPCGInputForwardingElement>(InputData),
 		/*Dependencies*/{},
 		/*InFromStack*/ nullptr,
@@ -45,10 +62,10 @@ bool UE::FPCGInternalHelpers::RunGraphLocallyBlocking(UPCGComponent* PCGComponen
 	PCGComponent->CurrentGenerationTask = TaskId;
 
 	bool bDone = false;
-	const FPCGTaskId FinalTaskId = FPCGGraphExecutorInternalAccess::ScheduleSingleElement(LocalGraphExecutor,
-		MakeShared<FPCGGenericElementAnyThread>([&LocalGraphExecutor, &OutputData, TaskId, &bDone](FPCGContext*) -> bool
+	const FPCGTaskId FinalTaskId = FPCGGraphExecutorInternalAccess::ScheduleSingleElement(Executor,
+		MakeShared<FPCGGenericElementAnyThread>([&Executor, &OutputData, TaskId, &bDone](FPCGContext*) -> bool
 		{
-			LocalGraphExecutor.GetOutputData(TaskId, OutputData);
+			Executor.GetOutputData(TaskId, OutputData);
 			bDone = true;
 			return true;
 		}), 
@@ -57,15 +74,13 @@ bool UE::FPCGInternalHelpers::RunGraphLocallyBlocking(UPCGComponent* PCGComponen
 	int32 CycleCount = 0;
 	CyclesToWait = FMath::Max(0, CyclesToWait);
 	SleepDurationPerCycle = FMath::Max(0, SleepDurationPerCycle);
-	LocalGraphExecutor.Execute();
-
+	
 	while (!bDone && CycleCount < CyclesToWait)
 	{
 		if (SleepDurationPerCycle > 0)
 		{
 			FPlatformProcess::Sleep(SleepDurationPerCycle);
 		}
-		LocalGraphExecutor.Execute();
 		++CycleCount;
 	}
 	return bDone;
diff --git a/Engine/Plugins/PCG/Source/PCG/Public/Helpers/PCGInternalHelpers.h b/Engine/Plugins/PCG/Source/PCG/Public/Helpers/PCGInternalHelpers.h
--- a/Engine/Plugins/PCG/Source/PCG/Public/Helpers/PCGInternalHelpers.h
+++ b/Engine/Plugins/PCG/Source/PCG/Public/Helpers/PCGInternalHelpers.h
@@ -3,25 +3,33 @@
 #pragma once
 
 #include "Misc/CoreMiscDefines.h" // UE_INTERNAL
+#include "Templates/SharedPointer.h"
+#include "PCGCommon.h" // EPCGChangeType
 
 class UPCGComponent;
 struct FPCGDataCollection;
-
+class FPCGGraphExecutor;
+class UPCGGraph;
 namespace UE
 {
 	UE_INTERNAL class FPCGInternalHelpers
 	{
 	public:
+		// Expose methods (temporarily) to be able to manage a custom executor in some prototype work, where the
+		//  PCG subsystem executor (currently) cannot be used.
+		UE_INTERNAL PCG_API static TSharedPtr<FPCGGraphExecutor> CreateExecutor();
+		UE_INTERNAL PCG_API static void ExecuteExecutor(FPCGGraphExecutor& Executor);
+#if WITH_EDITOR
+		UE_INTERNAL PCG_API static void NotifyGraphChanged(FPCGGraphExecutor& Executor, UPCGGraph* Graph, EPCGChangeType ChangeType);
+#endif
+
 		/**
-		 * Temporary helper that runs the graph attached to the given component in a local executor, waiting
-		 *  in a loop for it to finish or to time out. Currently used in some prototype work that needs to
-		 *  run graphs in non-main-threads while game thread is blocked, but will either be replaced with
-		 *  a similar asynchronous call, or perhaps just removed if the prototypes using it are changed.
-		 * 
-		 * Note that pcg.GraphMultithreading must be set to true if this is run outside the game thread.
+		 * Schedules onto the given executor and waits in a loop for it to finish or to time out. Currently 
+		 *  used in some prototype work, but may be replaced with a similar asynchronous call, or perhaps
+		 *  just removed if the prototypes using it are changed.
 		 */
-		UE_INTERNAL PCG_API static bool RunGraphLocallyBlocking(UPCGComponent* PCGComponent, 
-			const FPCGDataCollection& InputData, FPCGDataCollection& OutputData, 
-			int32 CyclesToWait, double SleepDurationPerCycle);
+		UE_INTERNAL PCG_API static bool ScheduleOnExecutorAndWait(FPCGGraphExecutor& Executor, 
+			UPCGComponent* PCGComponent, const FPCGDataCollection& InputData, 
+			FPCGDataCollection& OutputData, int32 CyclesToWait, double SleepDurationPerCycle);
 	};
 }
\ No newline at end of file
