[PCG] Added new Generate Seed node for more complex seed generation

+ Added Hash Attribute metadata operation
+ Updated Random Choice node with parameterized seed for attribute sets

diff --git a/Engine/Plugins/PCG/Source/PCG/Private/Elements/Metadata/PCGHashAttribute.cpp b/Engine/Plugins/PCG/Source/PCG/Private/Elements/Metadata/PCGHashAttribute.cpp
new file mode 100644
--- /dev/null
+++ b/Engine/Plugins/PCG/Source/PCG/Private/Elements/Metadata/PCGHashAttribute.cpp
@@ -0,0 +1,29 @@
+// Copyright Epic Games, Inc. All Rights Reserved.
+
+#include "Elements/Metadata/PCGHashAttribute.h"
+
+#include "PCGParamData.h"
+
+#define LOCTEXT_NAMESPACE "PCGHashAttributeElement"
+
+FPCGElementPtr UPCGHashAttributeSettings::CreateElement() const
+{
+	return MakeShared<FPCGHashAttributeElement>();
+}
+
+bool FPCGHashAttributeElement::DoOperation(PCGMetadataOps::FOperationData& InOperationData) const
+{
+	TRACE_CPUPROFILER_EVENT_SCOPE(FPCGHashAttributeElement::Execute);
+
+	PCGMetadataAttribute::CallbackWithRightType(InOperationData.MostComplexInputType, [this, &InOperationData]<typename AttributeType>(AttributeType) -> void
+	{
+		DoUnaryOp<AttributeType>(InOperationData, [](const AttributeType& Value) -> int32
+		{
+			return static_cast<int32>(PCG::Private::MetadataTraits<AttributeType>::Hash(Value));
+		});
+	});
+
+	return true;
+}
+
+#undef LOCTEXT_NAMESPACE
diff --git a/Engine/Plugins/PCG/Source/PCG/Private/Elements/PCGGenerateSeedElement.cpp b/Engine/Plugins/PCG/Source/PCG/Private/Elements/PCGGenerateSeedElement.cpp
new file mode 100644
--- /dev/null
+++ b/Engine/Plugins/PCG/Source/PCG/Private/Elements/PCGGenerateSeedElement.cpp
@@ -0,0 +1,147 @@
+// Copyright Epic Games, Inc. All Rights Reserved.
+
+#include "Elements/PCGGenerateSeedElement.h"
+
+#include "PCGContext.h"
+#include "Data/PCGPointData.h"
+#include "Elements/Metadata/PCGMetadataElementCommon.h"
+#include "Helpers/PCGHelpers.h"
+#include "Metadata/PCGMetadataAttribute.h"
+#include "Metadata/PCGMetadataAttributeTpl.h"
+#include "Metadata/PCGMetadataAttributeTraits.h"
+#include "Metadata/Accessors/PCGAttributeAccessorHelpers.h"
+
+#include UE_INLINE_GENERATED_CPP_BY_NAME(PCGGenerateSeedElement)
+
+namespace PCGGeneratedSeedElement::Constants
+{
+	const FName DefaultGeneratedSeedAttributeName = TEXT("GeneratedSeed");
+}
+
+#define LOCTEXT_NAMESPACE "PCGGenerateSeedElement"
+
+UPCGGenerateSeedSettings::UPCGGenerateSeedSettings()
+{
+	OutputTarget.SetAttributeName(PCGGeneratedSeedElement::Constants::DefaultGeneratedSeedAttributeName);
+}
+
+TArray<FPCGPinProperties> UPCGGenerateSeedSettings::InputPinProperties() const
+{
+	TArray<FPCGPinProperties> PinProperties;
+	PinProperties.Emplace_GetRef(PCGPinConstants::DefaultInputLabel, EPCGDataType::Any).SetRequiredPin();
+	return PinProperties;
+}
+
+TArray<FPCGPinProperties> UPCGGenerateSeedSettings::OutputPinProperties() const
+{
+	TArray<FPCGPinProperties> PinProperties;
+	PinProperties.Emplace(PCGPinConstants::DefaultOutputLabel, EPCGDataType::Any);
+	return PinProperties;
+}
+
+FPCGElementPtr UPCGGenerateSeedSettings::CreateElement() const
+{
+	return MakeShared<FPCGGenerateSeedElement>();
+}
+
+bool FPCGGenerateSeedElement::ExecuteInternal(FPCGContext* Context) const
+{
+	TRACE_CPUPROFILER_EVENT_SCOPE(FPCGGenerateSeedElement::Execute);
+
+	const UPCGGenerateSeedSettings* Settings = Context->GetInputSettings<UPCGGenerateSeedSettings>();
+	check(Settings);
+
+	const TArray<FPCGTaggedData> Inputs = Context->InputData.GetInputsByPin(PCGPinConstants::DefaultInputLabel);
+
+	const int32 InitialSeed = Context->GetSeed();
+	FRandomStream RandomStream(InitialSeed);
+
+	for (const FPCGTaggedData& Input : Inputs)
+	{
+		const UPCGMetadata* InputMetadata = Input.Data->ConstMetadata();
+		check(InputMetadata);
+
+		UPCGData* OutputData = Input.Data->DuplicateData(Context);
+		UPCGMetadata* OutputMetadata = OutputData->MutableMetadata();
+		check(OutputMetadata);
+
+		FPCGAttributePropertyInputSelector SourceSelector = Settings->SeedSource.CopyAndFixLast(Input.Data);
+		FPCGAttributePropertyOutputSelector OutputSelector = Settings->OutputTarget.CopyAndFixSource(&SourceSelector, Input.Data);
+		// Validate that the attribute exists or create a new one. If it's not a basic attribute, it must exist already, so skip and let the accessor be validated later.
+		if (OutputSelector.GetSelection() == EPCGAttributePropertySelection::Attribute && OutputSelector.IsBasicAttribute())
+		{
+			const FPCGMetadataAttribute<int32>* SeedAttribute = OutputMetadata->FindOrCreateAttribute<int32>(OutputSelector, InitialSeed);
+			if (!SeedAttribute)
+			{
+				PCGLog::Metadata::LogFailToCreateAttributeError<int32>(OutputSelector.GetName(), Context);
+				continue;
+			}
+		}
+
+		const TUniquePtr<IPCGAttributeAccessor> OutputAccessor = PCGAttributeAccessorHelpers::CreateAccessor(OutputData, OutputSelector);
+		const TUniquePtr<IPCGAttributeAccessorKeys> OutputKeys = PCGAttributeAccessorHelpers::CreateKeys(OutputData, OutputSelector);
+		if (!OutputAccessor || !OutputKeys)
+		{
+			PCGLog::Metadata::LogFailToCreateAccessorError(OutputSelector, Context);
+			continue;
+		}
+
+		// For users that want deterministic seeds across multiple data.
+		if (Settings->bResetSeedPerInput)
+		{
+			RandomStream.Reset();
+		}
+
+		const int32 NumElements = OutputKeys->GetNum();
+		TArray<int32> Seeds;
+		Seeds.SetNumUninitialized(NumElements);
+		for (int32 I = 0; I < NumElements; ++I)
+		{
+			// Random unsigned int to signed int is fine. Note: Seed is mutated.
+			Seeds[I] = static_cast<int32>(RandomStream.GetUnsignedInt());
+		}
+
+		if (Settings->GenerationSource == EPCGGenerateSeedSource::HashStringConstant)
+		{
+			const uint32 Hash = PCG::Private::MetadataTraits<FString>::Hash(Settings->SourceString);
+			for (int32& Seed : Seeds)
+			{
+				Seed = HashCombineFast(Hash, Seed);
+			}
+		}
+		else if (Settings->GenerationSource == EPCGGenerateSeedSource::HashEachSourceAttribute)
+		{
+			const TUniquePtr<const IPCGAttributeAccessor> Accessor = PCGAttributeAccessorHelpers::CreateConstAccessor(Input.Data, SourceSelector);
+			const TUniquePtr<const IPCGAttributeAccessorKeys> Keys = PCGAttributeAccessorHelpers::CreateConstKeys(Input.Data, SourceSelector);
+			if (!Accessor || !Keys)
+			{
+				PCGLog::Metadata::LogFailToGetAttributeError(SourceSelector, Context);
+				break;
+			}
+
+			// Attribute hash is element-wise so apply on all seed attributes
+			PCGMetadataAttribute::CallbackWithRightType(Accessor->GetUnderlyingType(), [&Seeds, &Accessor, &Keys]<typename AttributeType>(AttributeType)
+			{
+				PCGMetadataElementCommon::ApplyOnAccessor<AttributeType>(*Keys, *Accessor, [&Seeds](const AttributeType& Value, const int32 Index)
+				{
+					const uint32 Hash = PCG::Private::MetadataTraits<AttributeType>::Hash(Value);
+					Seeds[Index] = HashCombineFast(Hash, Seeds[Index]);
+				});
+			});
+		}
+
+		if (!OutputAccessor->SetRange(MakeConstArrayView(Seeds), 0, *OutputKeys))
+		{
+			PCGLog::Metadata::LogFailToSetAttributeError<int32>(OutputSelector, OutputAccessor.Get(), Context);
+			continue;
+		}
+
+		FPCGTaggedData& Output = Context->OutputData.TaggedData.Emplace_GetRef();
+		Output.Pin = PCGPinConstants::DefaultOutputLabel;
+		Output.Data = OutputData;
+	}
+
+	return true;
+}
+
+#undef LOCTEXT_NAMESPACE
diff --git a/Engine/Plugins/PCG/Source/PCG/Private/Elements/PCGRandomChoice.cpp b/Engine/Plugins/PCG/Source/PCG/Private/Elements/PCGRandomChoice.cpp
--- a/Engine/Plugins/PCG/Source/PCG/Private/Elements/PCGRandomChoice.cpp
+++ b/Engine/Plugins/PCG/Source/PCG/Private/Elements/PCGRandomChoice.cpp
@@ -6,7 +6,9 @@
 #include "PCGParamData.h"
 #include "PCGPin.h"
 #include "Data/PCGPointData.h"
+#include "Elements/Metadata/PCGMetadataElementCommon.h"
 #include "Helpers/PCGHelpers.h"
+#include "Metadata/Accessors/PCGAttributeAccessorHelpers.h"
 
 #define LOCTEXT_NAMESPACE "PCGRandomChoiceElement"
 
@@ -124,8 +126,8 @@ bool FPCGRandomChoiceElement::ExecuteInternal(FPCGContext* Context) const
 			NumElements = InPoints.Num();
 			ChooseFunc = &PCGRandomChoice::ChoosePointData;
 
-			// Combine the seed with the first point so that multiple data produces different results.
-			if (!InPoints.IsEmpty())
+			// By default, combine the seed with the first point's seed so that multiple data produces different results.
+			if (!Settings->bHasCustomSeedSource && !InPoints.IsEmpty())
 			{
 				Seed = PCGHelpers::ComputeSeed(Seed, InPoints[0].Seed);
 			}
@@ -141,6 +143,30 @@ bool FPCGRandomChoiceElement::ExecuteInternal(FPCGContext* Context) const
 			continue;
 		}
 
+		if (NumElements > 0 && Settings->bHasCustomSeedSource)
+		{
+			FPCGAttributePropertyInputSelector SeedSelector = Settings->CustomSeedSource.CopyAndFixLast(CurrentInput.Data);
+			const TUniquePtr<const IPCGAttributeAccessor> Accessor = PCGAttributeAccessorHelpers::CreateConstAccessor(CurrentInput.Data, SeedSelector);
+			const TUniquePtr<const IPCGAttributeAccessorKeys> Keys = PCGAttributeAccessorHelpers::CreateConstKeys(CurrentInput.Data, SeedSelector);
+			if (!Accessor || !Keys)
+			{
+				PCGLog::Metadata::LogFailToCreateAccessorError(SeedSelector, Context);
+				continue;
+			}
+
+			PCGMetadataAttribute::CallbackWithRightType(Accessor->GetUnderlyingType(), [&Accessor, &Keys, &Seed, Settings]<typename T>(const T&)
+			{
+				PCGMetadataElementCommon::ApplyOnAccessor<T>(*Keys, *Accessor, [&Seed](const T& Value, int32)
+				{
+					const uint32 Hash = PCG::Private::MetadataTraits<T>::Hash(Value);
+					Seed = HashCombineFast(Seed, Hash);
+				},
+				EPCGAttributeAccessorFlags::StrictType,
+				PCGMetadataElementCommon::DefaultChunkSize,
+				/*Count=*/Settings->bUseFirstAttributeOnly ? 1 : /*UseAllKeys*/-1);
+			});
+		}
+
 		check(ChooseFunc);
 
 		int NumOfElementsToKeep = 0;
diff --git a/Engine/Plugins/PCG/Source/PCG/Private/PCGContext.cpp b/Engine/Plugins/PCG/Source/PCG/Private/PCGContext.cpp
--- a/Engine/Plugins/PCG/Source/PCG/Private/PCGContext.cpp
+++ b/Engine/Plugins/PCG/Source/PCG/Private/PCGContext.cpp
@@ -97,7 +97,7 @@ int FPCGContext::GetSeed() const
 	}
 	else
 	{
-		return 42;
+		return PCGValueConstants::DefaultSeed;
 	}
 }
 
diff --git a/Engine/Plugins/PCG/Source/PCG/Private/PCGSettings.cpp b/Engine/Plugins/PCG/Source/PCG/Private/PCGSettings.cpp
--- a/Engine/Plugins/PCG/Source/PCG/Private/PCGSettings.cpp
+++ b/Engine/Plugins/PCG/Source/PCG/Private/PCGSettings.cpp
@@ -637,7 +637,7 @@ UPCGNode* UPCGSettings::CreateNode() const
 
 int UPCGSettings::GetSeed(const UPCGComponent* InSourceComponent) const
 {
-	return !UseSeed() ? 42 : (InSourceComponent ? PCGHelpers::ComputeSeed(Seed, InSourceComponent->Seed) : Seed);
+	return !UseSeed() ? PCGValueConstants::DefaultSeed : (InSourceComponent ? PCGHelpers::ComputeSeed(Seed, InSourceComponent->Seed) : Seed);
 }
 
 #if WITH_EDITOR
diff --git a/Engine/Plugins/PCG/Source/PCG/Public/Elements/Metadata/PCGHashAttribute.h b/Engine/Plugins/PCG/Source/PCG/Public/Elements/Metadata/PCGHashAttribute.h
new file mode 100644
--- /dev/null
+++ b/Engine/Plugins/PCG/Source/PCG/Public/Elements/Metadata/PCGHashAttribute.h
@@ -0,0 +1,41 @@
+// Copyright Epic Games, Inc. All Rights Reserved.
+
+#pragma once
+
+#include "PCGMetadataOpElementBase.h"
+#include "PCGSettings.h"
+
+#include "PCGHashAttribute.generated.h"
+
+/** Hash an attribute into a 32-bit unsigned integer, expressed as a 64-bit signed integer. */
+UCLASS(MinimalAPI, BlueprintType, ClassGroup = (Procedural))
+class UPCGHashAttributeSettings : public UPCGMetadataSettingsBase
+{
+	GENERATED_BODY()
+
+public:
+	//~Begin UPCGMetadataSettingsBase
+	virtual FName GetInputPinLabel(uint32 Index) const override { return PCGPinConstants::DefaultInputLabel; }
+	virtual bool IsSupportedInputType(uint16 TypeId, uint32 InputIndex, bool& bHasSpecialRequirement) const override { return true; }
+	virtual uint32 GetOperandNum() const override { return 1; }
+	virtual uint16 GetOutputType(uint16 InputTypeId) const override { return static_cast<uint16>(EPCGMetadataTypes::Integer32); }
+	//~End UPCGMetadataSettingsBase
+
+	//~Begin UPCGSettings interface
+#if WITH_EDITOR
+	virtual FName GetDefaultNodeName() const override { return FName("HashAttribute"); }
+	virtual FText GetDefaultNodeTitle() const override { return NSLOCTEXT("PCGHashAttributeElement", "NodeTitle", "Hash Attribute"); }
+	virtual EPCGSettingsType GetType() const override { return EPCGSettingsType::Metadata; }
+#endif
+	virtual FPCGElementPtr CreateElement() const override;
+	//~End UPCGSettings interface
+
+	UPROPERTY(BlueprintReadWrite, EditAnywhere, Category = Settings, meta = (PCG_Overridable))
+	FPCGAttributePropertyInputSelector InputSource1;
+};
+
+class FPCGHashAttributeElement : public FPCGMetadataElementBase
+{
+protected:
+	virtual bool DoOperation(PCGMetadataOps::FOperationData& InOperationData) const override;
+};
diff --git a/Engine/Plugins/PCG/Source/PCG/Public/Elements/PCGGenerateSeedElement.h b/Engine/Plugins/PCG/Source/PCG/Public/Elements/PCGGenerateSeedElement.h
new file mode 100644
--- /dev/null
+++ b/Engine/Plugins/PCG/Source/PCG/Public/Elements/PCGGenerateSeedElement.h
@@ -0,0 +1,70 @@
+// Copyright Epic Games, Inc. All Rights Reserved.
+
+#pragma once
+
+#include "PCGSettings.h"
+
+#include "PCGGenerateSeedElement.generated.h"
+
+UENUM()
+enum class EPCGGenerateSeedSource : uint8
+{
+	RandomStream UMETA(Tooltip = "Creates a random stream using the node's seed, and generates new seeds from that. Quick solution for seed values, but may be identical across duplicate input data."),
+	HashEachSourceAttribute UMETA(Tooltip = "Hashes each value in a source attribute to determine the target seeds. Can be used to differentiate results across duplicate input data."),
+	HashStringConstant UMETA(Tooltip = "Creates a target seed by combining a string's hash value and the previous seed value. Quick, human-readable solution for seed values, but may be identical across duplicate input data."),
+};
+
+/** Generate a seed from either a random stream, a constant string, or a source attribute. */
+UCLASS(BlueprintType, ClassGroup = (Procedural), meta = (Keywords = "seed"))
+class UPCGGenerateSeedSettings : public UPCGSettings
+{
+	GENERATED_BODY()
+
+public:
+	UPCGGenerateSeedSettings();
+
+	//~Begin UPCGSettings interface
+#if WITH_EDITOR
+	virtual FName GetDefaultNodeName() const override { return FName(TEXT("GenerateSeed")); }
+	virtual FText GetDefaultNodeTitle() const override { return NSLOCTEXT("PCGGenerateSeed", "NodeTitle", "Generate Seed"); }
+	virtual FText GetNodeTooltipText() const override { return NSLOCTEXT("PCGGenerateSeed", "NodeTooltip", "Generate a seed attribute "); }
+	virtual EPCGSettingsType GetType() const override { return EPCGSettingsType::Metadata; }
+	virtual TArray<FText> GetNodeTitleAliases() const override { return {NSLOCTEXT("PCGGenerateSeed", "NodeTitleAlias", "Seed From Value")}; }
+#endif // WITH_EDITOR
+	virtual bool UseSeed() const override { return true; }
+	virtual bool HasDynamicPins() const override { return true; }
+
+protected:
+	virtual TArray<FPCGPinProperties> InputPinProperties() const override;
+	virtual TArray<FPCGPinProperties> OutputPinProperties() const override;
+	virtual FPCGElementPtr CreateElement() const override;
+	//~End UPCGSettings interface
+
+public:
+	/** The source method seed attribute. */
+	UPROPERTY(BlueprintReadWrite, EditAnywhere, Category = "Settings")
+	EPCGGenerateSeedSource GenerationSource = EPCGGenerateSeedSource::RandomStream;
+
+	/** This value will be hashed and applied to generate each new seed. */
+	UPROPERTY(BlueprintReadWrite, EditAnywhere, Category = "Settings", DisplayName = "String", meta = (EditCondition = "GenerationSource == EPCGGenerateSeedSource::HashStringConstant", EditConditionHides))
+	FString SourceString;
+
+	/** The source attribute to hash to generate each new seed. */
+	UPROPERTY(BlueprintReadWrite, EditAnywhere, Category = "Settings", meta = (EditCondition = "GenerationSource == EPCGGenerateSeedSource::HashEachSourceAttribute", EditConditionHides))
+	FPCGAttributePropertyInputSelector SeedSource;
+
+	/** Reset the seed at the beginning of each input's generation to stay order agnostic.  */
+	UPROPERTY(BlueprintReadWrite, EditAnywhere, Category = "Settings", meta = (PCG_Overridable, EditCondition = "GenerationSource != EPCGGenerateSeedSource::HashEachSourceAttribute", EditConditionHides))
+	bool bResetSeedPerInput = true;
+
+	/** The target attribute output of the generated seed. */
+	UPROPERTY(BlueprintReadWrite, EditAnywhere, Category = "Settings", meta = (PCG_Overridable))
+	FPCGAttributePropertyOutputSelector OutputTarget;
+};
+
+class FPCGGenerateSeedElement : public IPCGElement
+{
+protected:
+	virtual bool ExecuteInternal(FPCGContext* Context) const override;
+	virtual bool SupportsBasePointDataInputs(FPCGContext* InContext) const override { return true; }
+};
diff --git a/Engine/Plugins/PCG/Source/PCG/Public/Elements/PCGRandomChoice.h b/Engine/Plugins/PCG/Source/PCG/Public/Elements/PCGRandomChoice.h
--- a/Engine/Plugins/PCG/Source/PCG/Public/Elements/PCGRandomChoice.h
+++ b/Engine/Plugins/PCG/Source/PCG/Public/Elements/PCGRandomChoice.h
@@ -55,6 +55,18 @@ public:
 	// By default, we output discarded entries. If you don't need them, disable this option.
 	UPROPERTY(BlueprintReadWrite, EditAnywhere, Category = Settings)
 	bool bOutputDiscardedEntries = true;
+
+	/** Use an attribute as a source for generating the seed, i.e. similar to or replacing the $Seed property on points. Mostly useful for attribute sets as points have this unique seed by default. */
+	UPROPERTY(BlueprintReadWrite, EditAnywhere, Category = Settings, meta = (PCG_Overridable))
+	bool bHasCustomSeedSource = false;
+
+	/** Attribute to generate the selection seed from. */
+	UPROPERTY(BlueprintReadWrite, EditAnywhere, Category = Settings, meta = (EditCondition = "bHasCustomSeedSource", EditConditionHides, PCG_Overridable))
+	FPCGAttributePropertyInputSelector CustomSeedSource;
+
+	/** Controls whether only the first element's attribute will be used to generate the selection seed. Otherwise, all values will be used to compute it. */
+	UPROPERTY(BlueprintReadWrite, EditAnywhere, Category = Settings, meta = (EditCondition = "bHasCustomSeedSource", EditConditionHides, PCG_Overridable))
+	bool bUseFirstAttributeOnly = true;
 };
 
 class FPCGRandomChoiceElement : public IPCGElement
diff --git a/Engine/Plugins/PCG/Source/PCG/Public/Metadata/PCGMetadataAttributeTraits.h b/Engine/Plugins/PCG/Source/PCG/Public/Metadata/PCGMetadataAttributeTraits.h
--- a/Engine/Plugins/PCG/Source/PCG/Public/Metadata/PCGMetadataAttributeTraits.h
+++ b/Engine/Plugins/PCG/Source/PCG/Public/Metadata/PCGMetadataAttributeTraits.h
@@ -419,6 +419,17 @@ namespace PCG
 			}
 		};
 
+		template<typename T>
+		struct DefaultHashTraits
+		{
+			enum { CanHash = true };
+
+			static uint32 Hash(const T& A)
+			{
+				return GetTypeHash(A);
+			}
+		};
+
 		/** "Find Nearest" trait support driven by less/less or equal operation, used for types that do not support distance computation */
 		template<typename Traits, typename T>
 		struct CompareBasedFindNearestTraits
@@ -497,7 +508,7 @@ namespace PCG
 		};
 
 		template<typename T>
-		struct MetadataTraits : DefaultOperationTraits<T>, DefaultWeightedSumTraits<T>, DefaultMinMaxTraits<T>, DefaultCompareTraits<T>, LexToStringTraits<T>, DefaultScalarDistanceTraits<T>
+		struct MetadataTraits : DefaultOperationTraits<T>, DefaultWeightedSumTraits<T>, DefaultMinMaxTraits<T>, DefaultCompareTraits<T>, DefaultHashTraits<T>, LexToStringTraits<T>, DefaultScalarDistanceTraits<T>
 		{
 			enum { CompressData = false };
 			enum { CanSearchString = false };
@@ -505,7 +516,7 @@ namespace PCG
 		};
 
 		template<>
-		struct MetadataTraits<int32> : DefaultOperationTraits<int32>, DefaultWeightedSumTraits<int32>, DefaultMinMaxTraits<int32>, DefaultCompareTraits<int32>, LexToStringTraits<int32>, DefaultScalarDistanceTraits<int32>
+		struct MetadataTraits<int32> : DefaultOperationTraits<int32>, DefaultWeightedSumTraits<int32>, DefaultMinMaxTraits<int32>, DefaultCompareTraits<int32>, DefaultHashTraits<int32>, LexToStringTraits<int32>, DefaultScalarDistanceTraits<int32>
 		{
 			enum { CompressData = false };
 			enum { CanSearchString = false };
@@ -521,7 +532,7 @@ namespace PCG
 		};
 
 		template<>
-		struct MetadataTraits<int64> : DefaultOperationTraits<int64>, DefaultWeightedSumTraits<int64>, DefaultMinMaxTraits<int64>, DefaultCompareTraits<int64>, LexToStringTraits<int64>, DefaultScalarDistanceTraits<int64>
+		struct MetadataTraits<int64> : DefaultOperationTraits<int64>, DefaultWeightedSumTraits<int64>, DefaultMinMaxTraits<int64>, DefaultCompareTraits<int64>, DefaultHashTraits<int64>, LexToStringTraits<int64>, DefaultScalarDistanceTraits<int64>
 		{
 			enum { CompressData = false };
 			enum { CanSearchString = false };
@@ -536,7 +547,7 @@ namespace PCG
 		};
 
 		template<>
-		struct MetadataTraits<float> : DefaultOperationTraits<float>, DefaultWeightedSumTraits<float>, DefaultMinMaxTraits<float>, DefaultCompareTraits<float>, LexToStringTraits<float>, DefaultScalarDistanceTraits<float>
+		struct MetadataTraits<float> : DefaultOperationTraits<float>, DefaultWeightedSumTraits<float>, DefaultMinMaxTraits<float>, DefaultCompareTraits<float>, DefaultHashTraits<float>, LexToStringTraits<float>, DefaultScalarDistanceTraits<float>
 		{
 			enum { CompressData = false };
 			enum { CanSearchString = false };
@@ -557,7 +568,7 @@ namespace PCG
 		};
 
 		template<>
-		struct MetadataTraits<double> : DefaultOperationTraits<double>, DefaultWeightedSumTraits<double>, DefaultMinMaxTraits<double>, DefaultCompareTraits<double>, LexToStringTraits<double>, DefaultScalarDistanceTraits<double>
+		struct MetadataTraits<double> : DefaultOperationTraits<double>, DefaultWeightedSumTraits<double>, DefaultMinMaxTraits<double>, DefaultCompareTraits<double>, DefaultHashTraits<double>, LexToStringTraits<double>, DefaultScalarDistanceTraits<double>
 		{
 			enum { CompressData = false };
 			enum { CanSearchString = false };
@@ -578,7 +589,7 @@ namespace PCG
 		};
 
 		template<>
-		struct MetadataTraits<bool> : LexToStringTraits<bool>
+		struct MetadataTraits<bool> : LexToStringTraits<bool>, DefaultHashTraits<bool>
 		{
 			enum { CompressData = false };
 			enum { CanMinMax = true };
@@ -693,7 +704,7 @@ namespace PCG
 		};
 
 		template<>
-		struct MetadataTraits<FVector2D> : VectorTraits<FVector2D>
+		struct MetadataTraits<FVector2D> : VectorTraits<FVector2D>, DefaultHashTraits<FVector2D>
 		{
 			static bool Equal(const FVector2D& A, const FVector2D& B, const FVector2D::FReal Tolerance = UE_DOUBLE_SMALL_NUMBER)
 			{
@@ -740,7 +751,7 @@ namespace PCG
 		};
 
 		template<>
-		struct MetadataTraits<FVector> : VectorTraits<FVector>
+		struct MetadataTraits<FVector> : VectorTraits<FVector>, DefaultHashTraits<FVector>
 		{
 			static bool Equal(const FVector& A, const FVector& B, const FVector::FReal Tolerance = UE_DOUBLE_SMALL_NUMBER)
 			{
@@ -787,7 +798,7 @@ namespace PCG
 		};
 		
 		template<>
-		struct MetadataTraits<FVector4> : VectorTraits<FVector4>
+		struct MetadataTraits<FVector4> : VectorTraits<FVector4>, DefaultHashTraits<FVector4>
 		{
 			static bool Equal(const FVector4& A, const FVector4& B, const FVector4::FReal Tolerance = UE_DOUBLE_SMALL_NUMBER)
 			{
@@ -835,7 +846,7 @@ namespace PCG
 
 		// Quaternion
 		template<>
-		struct MetadataTraits<FQuat> : DefaultStringTraits<FQuat>, DistanceBasedFindNearestTraits<MetadataTraits<FQuat>, FQuat>
+		struct MetadataTraits<FQuat> : DefaultStringTraits<FQuat>, DistanceBasedFindNearestTraits<MetadataTraits<FQuat>, FQuat>, DefaultHashTraits<FQuat>
 		{
 			enum { CompressData = false };
 			enum { CanMinMax = false };
@@ -911,7 +922,7 @@ namespace PCG
 
 		// Rotator
 		template<>
-		struct MetadataTraits<FRotator> : DefaultStringTraits<FRotator>, DistanceBasedFindNearestTraits<MetadataTraits<FRotator>, FRotator>
+		struct MetadataTraits<FRotator> : DefaultStringTraits<FRotator>, DistanceBasedFindNearestTraits<MetadataTraits<FRotator>, FRotator>, DefaultHashTraits<FRotator>
 		{
 			enum { CompressData = false };
 			enum { CanMinMax = false };
@@ -978,11 +989,16 @@ namespace PCG
 				bOutSuccess = FDefaultValueHelper::ParseRotator(InValue, OutValue);
 				return bOutSuccess ? OutValue : FRotator::ZeroRotator;
 			}
+
+			static uint32 Hash(const FRotator& A)
+			{
+				return GetTypeHash(FVector(A.Pitch, A.Yaw, A.Roll));
+			}
 		};
 
 		// Transform
 		template<>
-		struct MetadataTraits<FTransform> : DefaultStringTraits<FTransform>
+		struct MetadataTraits<FTransform> : DefaultStringTraits<FTransform>, DefaultHashTraits<FTransform>
 		{
 			enum { CompressData = false };
 			enum { CanMinMax = false };
@@ -1050,7 +1066,7 @@ namespace PCG
 
 		// Strings
 		template<>
-		struct MetadataTraits<FString> : DefaultCompareTraits<FString>, DefaultStringTraits<FString>, CompareBasedFindNearestTraits<MetadataTraits<FString>, FString>
+		struct MetadataTraits<FString> : DefaultCompareTraits<FString>, DefaultStringTraits<FString>, CompareBasedFindNearestTraits<MetadataTraits<FString>, FString>, DefaultHashTraits<FString>
 		{
 			enum { CompressData = true };
 			enum { CanMinMax = false };
@@ -1087,7 +1103,7 @@ namespace PCG
 		};
 
 		template<>
-		struct MetadataTraits<FName> : DefaultStringTraits<FName>, CompareBasedFindNearestTraits<MetadataTraits<FName>, FName>
+		struct MetadataTraits<FName> : DefaultStringTraits<FName>, CompareBasedFindNearestTraits<MetadataTraits<FName>, FName>, DefaultHashTraits<FName>
 		{
 			enum { CompressData = false };
 			enum { CanMinMax = false };
@@ -1187,7 +1203,7 @@ namespace PCG
 				return A.ToString() <= B.ToString();
 			}
 
-			static bool GreaterOrEqual(const FSoftObjectPath& A, const FSoftObjectPath& B)
+			static bool GreaterOrEqual(const T& A, const T& B)
 			{
 				return A.ToString() >= B.ToString();
 			}
@@ -1195,13 +1211,13 @@ namespace PCG
 
 		// Soft object paths
 		template<>
-		struct MetadataTraits<FSoftObjectPath> : SoftObjectPathTraits<FSoftObjectPath>, DefaultStringTraits<FSoftObjectPath>
+		struct MetadataTraits<FSoftObjectPath> : SoftObjectPathTraits<FSoftObjectPath>, DefaultStringTraits<FSoftObjectPath>, DefaultHashTraits<FSoftObjectPath>
 		{
 		};
 
 		// Soft class paths
 		template<>
-		struct MetadataTraits<FSoftClassPath> : SoftObjectPathTraits<FSoftClassPath>, DefaultStringTraits<FSoftClassPath>
+		struct MetadataTraits<FSoftClassPath> : SoftObjectPathTraits<FSoftClassPath>, DefaultStringTraits<FSoftClassPath>, DefaultHashTraits<FSoftClassPath>
 		{
 		};
 
