Improved boolean parsing for Tag:Value

#rnx

diff --git a/Engine/Plugins/PCG/Source/PCG/Private/Helpers/PCGTagHelpers.cpp b/Engine/Plugins/PCG/Source/PCG/Private/Helpers/PCGTagHelpers.cpp
--- a/Engine/Plugins/PCG/Source/PCG/Private/Helpers/PCGTagHelpers.cpp
+++ b/Engine/Plugins/PCG/Source/PCG/Private/Helpers/PCGTagHelpers.cpp
@@ -37,6 +37,14 @@ namespace Private
 				{
 					Result.NumericValue = FCString::Atod(*RightSide);
 				}
+				else if (RightSide.Equals(TEXT("True")))
+				{
+					Result.BooleanValue = true;
+				}
+				else if (RightSide.Equals(TEXT("False")))
+				{
+					Result.BooleanValue = false;
+				}
 
 				Result.Attribute = MoveTemp(LeftSide);
 				Result.Value = MoveTemp(RightSide);
@@ -122,6 +130,14 @@ namespace Private
 					Attribute->SetValue(InEntryKey, TagData.NumericValue.GetValue());
 				}
 			}
+			else if (TagData.HasBooleanValue())
+			{
+				FPCGMetadataAttribute<bool>* Attribute = GetAttribute(false);
+				if (Attribute && InEntryKey != PCGInvalidEntryKey)
+				{
+					Attribute->SetValue(InEntryKey, TagData.BooleanValue.GetValue());
+				}
+			}
 			else if (TagData.HasValue())
 			{
 				FPCGMetadataAttribute<FString>* Attribute = GetAttribute(FString());
diff --git a/Engine/Plugins/PCG/Source/PCG/Public/Helpers/PCGTagHelpers.h b/Engine/Plugins/PCG/Source/PCG/Public/Helpers/PCGTagHelpers.h
--- a/Engine/Plugins/PCG/Source/PCG/Public/Helpers/PCGTagHelpers.h
+++ b/Engine/Plugins/PCG/Source/PCG/Public/Helpers/PCGTagHelpers.h
@@ -19,12 +19,14 @@ namespace Private
 		bool HasBeenSanitized() const { return OriginalAttribute.IsSet(); }
 		bool HasValue() const { return Value.IsSet(); }
 		bool HasNumericValue() const { return NumericValue.IsSet(); }
+		bool HasBooleanValue() const { return BooleanValue.IsSet(); }
 		const FString& GetOriginalAttribute() const { return HasBeenSanitized() ? OriginalAttribute.GetValue() : Attribute; }
 
 		FString Attribute;
 		TOptional<FString> OriginalAttribute;
 		TOptional<FString> Value;
 		TOptional<double> NumericValue;
+		TOptional<bool> BooleanValue;
 	};
 
 	// Builds the tag result structure from the provided tag.
