[PCG] 2/3 Inline Pin Constants (Default Values) Backend

+ FPCGDefaultValueContainer
* MetadataOps updated to support Default Values from inline constants rather than hardcoded.

diff --git a/Engine/Plugins/PCG/Source/PCG/Private/Elements/Metadata/PCGMetadataBreakTransform.cpp b/Engine/Plugins/PCG/Source/PCG/Private/Elements/Metadata/PCGMetadataBreakTransform.cpp
--- a/Engine/Plugins/PCG/Source/PCG/Private/Elements/Metadata/PCGMetadataBreakTransform.cpp
+++ b/Engine/Plugins/PCG/Source/PCG/Private/Elements/Metadata/PCGMetadataBreakTransform.cpp
@@ -37,6 +37,26 @@ void UPCGMetadataBreakTransformSettings::PostLoad()
 #endif // WITH_EDITOR
 }
 
+#if WITH_EDITOR
+  bool UPCGMetadataBreakTransformSettings::CanEditChange(const FProperty* InProperty) const
+{
+	if (!InProperty || !UPCGSettings::CanEditChange(InProperty))
+	{
+		return false;
+	}
+
+	if (InProperty->GetFName() == GET_MEMBER_NAME_CHECKED(UPCGMetadataBreakTransformSettings, InputSource))
+	{
+		if (const UPCGNode* Node = Cast<UPCGNode>(GetOuter()))
+		{
+			return Node->IsInputPinConnected(PCGPinConstants::DefaultInputLabel);
+		}
+	}
+
+	return false;
+}
+#endif // WITH_EDITOR
+
 FPCGAttributePropertyInputSelector UPCGMetadataBreakTransformSettings::GetInputSource(uint32 Index) const
 {
 	return InputSource;
diff --git a/Engine/Plugins/PCG/Source/PCG/Private/Elements/Metadata/PCGMetadataBreakVector.cpp b/Engine/Plugins/PCG/Source/PCG/Private/Elements/Metadata/PCGMetadataBreakVector.cpp
--- a/Engine/Plugins/PCG/Source/PCG/Private/Elements/Metadata/PCGMetadataBreakVector.cpp
+++ b/Engine/Plugins/PCG/Source/PCG/Private/Elements/Metadata/PCGMetadataBreakVector.cpp
@@ -87,6 +87,26 @@ void UPCGMetadataBreakVectorSettings::PostLoad()
 #endif // WITH_EDITOR
 }
 
+#if WITH_EDITOR
+  bool UPCGMetadataBreakVectorSettings::CanEditChange(const FProperty* InProperty) const
+{
+	if (!InProperty || !UPCGSettings::CanEditChange(InProperty))
+	{
+		return false;
+	}
+
+	if (InProperty->GetFName() == GET_MEMBER_NAME_CHECKED(UPCGMetadataBreakVectorSettings, InputSource))
+	{
+		if (const UPCGNode* Node = Cast<UPCGNode>(GetOuter()))
+		{
+			return Node->IsInputPinConnected(PCGPinConstants::DefaultInputLabel);
+		}
+	}
+
+	return false;
+}
+#endif // WITH_EDITOR
+
 FPCGAttributePropertyInputSelector UPCGMetadataBreakVectorSettings::GetInputSource(uint32 Index) const
 {
 	return InputSource;
diff --git a/Engine/Plugins/PCG/Source/PCG/Private/Elements/Metadata/PCGMetadataMakeRotator.cpp b/Engine/Plugins/PCG/Source/PCG/Private/Elements/Metadata/PCGMetadataMakeRotator.cpp
--- a/Engine/Plugins/PCG/Source/PCG/Private/Elements/Metadata/PCGMetadataMakeRotator.cpp
+++ b/Engine/Plugins/PCG/Source/PCG/Private/Elements/Metadata/PCGMetadataMakeRotator.cpp
@@ -5,6 +5,7 @@
 #include "PCGParamData.h"
 #include "Elements/Metadata/PCGMetadataElementCommon.h"
 #include "Metadata/PCGMetadata.h"
+#include "UObject/FortniteMainBranchObjectVersion.h"
 
 #define LOCTEXT_NAMESPACE "PCGMetadataMakeRotatorSettings"
 
@@ -95,6 +96,17 @@ uint16 UPCGMetadataMakeRotatorSettings::GetOutputType(uint16 InputTypeId) const
 }
 
 #if WITH_EDITOR
+void UPCGMetadataMakeRotatorSettings::PostEditChangeProperty(FPropertyChangedEvent& PropertyChangedEvent)
+{
+	Super::PostEditChangeProperty(PropertyChangedEvent);
+
+	if (PropertyChangedEvent.Property->GetName() == GET_MEMBER_NAME_CHECKED(UPCGMetadataMakeRotatorSettings, Operation))
+	{
+		// @todo_pcg: It doesn't need to be reset in most cases.
+		ResetDefaultValues();
+	}
+}
+
 FName UPCGMetadataMakeRotatorSettings::GetDefaultNodeName() const
 {
 	return TEXT("MakeRotatorAttribute");
@@ -109,47 +121,66 @@ TArray<FPCGPreConfiguredSettingsInfo> UPCGMetadataMakeRotatorSettings::GetPrecon
 {
 	return FPCGPreConfiguredSettingsInfo::PopulateFromEnum<EPCGMetadataMakeRotatorOp>();
 }
-#endif // WITH_EDITOR
 
-void UPCGMetadataMakeRotatorSettings::ApplyPreconfiguredSettings(const FPCGPreConfiguredSettingsInfo& PreconfiguredInfo)
+void UPCGMetadataMakeRotatorSettings::ApplyDeprecationBeforeUpdatePins(UPCGNode* InOutNode, TArray<TObjectPtr<UPCGPin>>& InputPins, TArray<TObjectPtr<UPCGPin>>& OutputPins)
 {
-	if (const UEnum* EnumPtr = StaticEnum<EPCGMetadataMakeRotatorOp>())
+	Super::ApplyDeprecationBeforeUpdatePins(InOutNode, InputPins, OutputPins);
+
+	// Supported default values on Make Rot From Angles
+	if (InOutNode
+		&& Operation == EPCGMetadataMakeRotatorOp::MakeRotFromAngles
+		&& GetLinkerCustomVersion(FFortniteMainBranchObjectVersion::GUID) < FFortniteMainBranchObjectVersion::PCGInlineConstantDefaultValues)
 	{
-		if (EnumPtr->IsValidEnumValue(PreconfiguredInfo.PreconfiguredIndex))
+		for (const UPCGPin* Pin : InputPins)
 		{
-			Operation = EPCGMetadataMakeRotatorOp(PreconfiguredInfo.PreconfiguredIndex);
+			if (IsPinDefaultValueEnabled(Pin->Properties.Label) && !Pin->IsConnected())
+			{
+				SetPinDefaultValueIsActivated(Pin->Properties.Label, /*bIsActivated=*/true, /*bDirtySettings=*/false);
+			}
 		}
 	}
 }
 
-FPCGElementPtr UPCGMetadataMakeRotatorSettings::CreateElement() const
-{
-	return MakeShared<FPCGMetadataMakeRotatorElement>();
-}
-
-bool UPCGMetadataMakeRotatorSettings::DoesInputSupportDefaultValue(uint32 Index) const
+FString UPCGMetadataMakeRotatorSettings::GetPinInitialDefaultValueString(FName PinLabel) const
 {
-	return Operation == EPCGMetadataMakeRotatorOp::MakeRotFromAngles;
+	switch (Operation)
+	{
+		case EPCGMetadataMakeRotatorOp::MakeRotFromX: // fall-through
+		case EPCGMetadataMakeRotatorOp::MakeRotFromY: // fall-through
+		case EPCGMetadataMakeRotatorOp::MakeRotFromZ: // fall-through
+		case EPCGMetadataMakeRotatorOp::MakeRotFromAngles:
+			return PCG::Private::MetadataTraits<double>::ZeroValueString();
+		case EPCGMetadataMakeRotatorOp::MakeRotFromXY: // fall-through
+		case EPCGMetadataMakeRotatorOp::MakeRotFromYX: // fall-through
+		case EPCGMetadataMakeRotatorOp::MakeRotFromXZ: // fall-through
+		case EPCGMetadataMakeRotatorOp::MakeRotFromZX: // fall-through
+		case EPCGMetadataMakeRotatorOp::MakeRotFromYZ: // fall-through
+		case EPCGMetadataMakeRotatorOp::MakeRotFromZY:
+			return PCG::Private::MetadataTraits<FVector2D>::ZeroValueString();
+		case EPCGMetadataMakeRotatorOp::MakeRotFromAxes:
+			return PCG::Private::MetadataTraits<FVector>::ZeroValueString();
+		default:
+			checkNoEntry();
+			return FString();
+	}
 }
+#endif // WITH_EDITOR
 
-UPCGParamData* UPCGMetadataMakeRotatorSettings::CreateDefaultValueParam(FPCGContext* Context, uint32 Index) const
+void UPCGMetadataMakeRotatorSettings::ApplyPreconfiguredSettings(const FPCGPreConfiguredSettingsInfo& PreconfiguredInfo)
 {
-	if (Operation != EPCGMetadataMakeRotatorOp::MakeRotFromAngles)
+	if (const UEnum* EnumPtr = StaticEnum<EPCGMetadataMakeRotatorOp>())
 	{
-		return nullptr;
+		if (EnumPtr->IsValidEnumValue(PreconfiguredInfo.PreconfiguredIndex))
+		{
+			Operation = EPCGMetadataMakeRotatorOp(PreconfiguredInfo.PreconfiguredIndex);
+		}
 	}
-
-	UPCGParamData* NewParamData = FPCGContext::NewObject_AnyThread<UPCGParamData>(Context);
-	NewParamData->Metadata->CreateAttribute<double>(NAME_None, 0, /*bAllowsInterpolation=*/ true, /*bOverrideParent=*/ false);
-	return NewParamData;
 }
 
-#if WITH_EDITOR
-FString UPCGMetadataMakeRotatorSettings::GetDefaultValueString(uint32 Index) const
+FPCGElementPtr UPCGMetadataMakeRotatorSettings::CreateElement() const
 {
-	return (Operation == EPCGMetadataMakeRotatorOp::MakeRotFromAngles ? FString::Printf(TEXT("%f"), 0.0) : FString());
+	return MakeShared<FPCGMetadataMakeRotatorElement>();
 }
-#endif // WITH_EDITOR
 
 bool FPCGMetadataMakeRotatorElement::DoOperation(PCGMetadataOps::FOperationData& OperationData) const
 {
diff --git a/Engine/Plugins/PCG/Source/PCG/Private/Elements/Metadata/PCGMetadataMakeTransform.cpp b/Engine/Plugins/PCG/Source/PCG/Private/Elements/Metadata/PCGMetadataMakeTransform.cpp
--- a/Engine/Plugins/PCG/Source/PCG/Private/Elements/Metadata/PCGMetadataMakeTransform.cpp
+++ b/Engine/Plugins/PCG/Source/PCG/Private/Elements/Metadata/PCGMetadataMakeTransform.cpp
@@ -6,6 +6,7 @@
 #include "Elements/Metadata/PCGMetadataElementCommon.h"
 #include "Metadata/PCGMetadata.h"
 #include "Metadata/PCGMetadataAttributeTpl.h"
+#include "UObject/FortniteMainBranchObjectVersion.h"
 
 #include UE_INLINE_GENERATED_CPP_BY_NAME(PCGMetadataMakeTransform)
 
@@ -119,68 +120,77 @@ FText UPCGMetadataMakeTransformSettings::GetDefaultNodeTitle() const
 {
 	return NSLOCTEXT("PCGMetadataMakeTransformSettings", "NodeTitle", "Make Transform Attribute");
 }
-#endif // WITH_EDITOR
 
-bool UPCGMetadataMakeTransformSettings::DoesInputSupportDefaultValue(uint32 Index) const
+void UPCGMetadataMakeTransformSettings::ApplyDeprecationBeforeUpdatePins(UPCGNode* InOutNode, TArray<TObjectPtr<UPCGPin>>& InputPins, TArray<TObjectPtr<UPCGPin>>& OutputPins)
 {
-	return true;
+	Super::ApplyDeprecationBeforeUpdatePins(InOutNode, InputPins, OutputPins);
+
+	// Supported default values on all pins
+	if (this != GetClass()->GetDefaultObject()
+		&& GetLinkerCustomVersion(FFortniteMainBranchObjectVersion::GUID) < FFortniteMainBranchObjectVersion::PCGInlineConstantDefaultValues)
+	{
+		for (const UPCGPin* Pin : InputPins)
+		{
+			if (IsPinDefaultValueEnabled(Pin->Properties.Label) && !Pin->IsConnected())
+			{
+				SetPinDefaultValueIsActivated(Pin->Properties.Label, /*bIsActivated=*/true, /*bDirtySettings=*/false);
+			}
+		}
+	}
 }
 
-UPCGParamData* UPCGMetadataMakeTransformSettings::CreateDefaultValueParam(FPCGContext* Context, uint32 Index) const
+FString UPCGMetadataMakeTransformSettings::GetPinInitialDefaultValueString(const FName PinLabel) const
 {
-	switch (Index)
-	{
-	case 0:
+	if (PinLabel == PCGMetadataTransformConstants::Translation)
 	{
 		// Location -> Default is Zero vector
-		UPCGParamData* NewParamData = FPCGContext::NewObject_AnyThread<UPCGParamData>(Context);
-		NewParamData->Metadata->CreateAttribute<FVector>(NAME_None, FVector::ZeroVector, /*bAllowsInterpolation=*/ true, /*bOverrideParent=*/ false);
-		return NewParamData;
+		return PCG::Private::MetadataTraits<FVector>::ZeroValueString();
 	}
-	case 1:
+	else if (PinLabel == PCGMetadataTransformConstants::Rotation)
 	{
 		// Rotation -> Default is Zero rotator
-		UPCGParamData* NewParamData = FPCGContext::NewObject_AnyThread<UPCGParamData>(Context);
-		NewParamData->Metadata->CreateAttribute<FRotator>(NAME_None, FRotator::ZeroRotator, /*bAllowsInterpolation=*/ true, /*bOverrideParent=*/ false);
-		return NewParamData;
+		return PCG::Private::MetadataTraits<FRotator>::ZeroValueString();
 	}
-	case 2:
+	else if (PinLabel == PCGMetadataTransformConstants::Scale)
 	{
 		// Scale -> Default is Vector (1, 1, 1)
-		UPCGParamData* NewParamData = FPCGContext::NewObject_AnyThread<UPCGParamData>(Context);
-		NewParamData->Metadata->CreateAttribute<FVector>(NAME_None, FVector::OneVector, /*bAllowsInterpolation=*/ true, /*bOverrideParent=*/ false);
-		return NewParamData;
+		return FVector::OneVector.ToString();
 	}
-	default:
-		return nullptr;
+	else
+	{
+		return FString();
 	}
 }
+#endif // WITH_EDITOR
 
-#if WITH_EDITOR
-FString UPCGMetadataMakeTransformSettings::GetDefaultValueString(uint32 Index) const
+EPCGMetadataTypes UPCGMetadataMakeTransformSettings::GetPinInitialDefaultValueType(const FName PinLabel) const
 {
-	switch (Index)
+	if (PinLabel == PCGMetadataTransformConstants::Translation
+		|| PinLabel == PCGMetadataTransformConstants::Scale)
 	{
-	case 0:
+		return EPCGMetadataTypes::Vector;
+	}
+	else if (PinLabel == PCGMetadataTransformConstants::Rotation)
 	{
-		// Location -> Default is Zero vector
-		return FVector::ZeroVector.ToString();
+		return EPCGMetadataTypes::Rotator;
 	}
-	case 1:
+	else
 	{
-		// Rotation -> Default is Zero rotator
-		return FRotator::ZeroRotator.ToString();
+		return EPCGMetadataTypes::Unknown;
 	}
-	case 2:
+}
+
+bool UPCGMetadataMakeTransformSettings::CreateInitialDefaultValueAttribute(const FName PinLabel, UPCGMetadata* OutMetadata) const
+{
+	if (PinLabel == PCGMetadataTransformConstants::Scale)
 	{
-		// Scale -> Default is Vector (1, 1, 1)
-		return FVector::OneVector.ToString();
+		return nullptr != OutMetadata->CreateAttribute(NAME_None, FVector::OneVector, /*bAllowsInterpolation=*/true, /*bOverrideParent=*/false);
 	}
-	default:
-		return FString();
+	else
+	{
+		return Super::CreateInitialDefaultValueAttribute(PinLabel, OutMetadata);
 	}
 }
-#endif // WITH_EDITOR
 
 FPCGElementPtr UPCGMetadataMakeTransformSettings::CreateElement() const
 {
diff --git a/Engine/Plugins/PCG/Source/PCG/Private/Elements/Metadata/PCGMetadataMakeVector.cpp b/Engine/Plugins/PCG/Source/PCG/Private/Elements/Metadata/PCGMetadataMakeVector.cpp
--- a/Engine/Plugins/PCG/Source/PCG/Private/Elements/Metadata/PCGMetadataMakeVector.cpp
+++ b/Engine/Plugins/PCG/Source/PCG/Private/Elements/Metadata/PCGMetadataMakeVector.cpp
@@ -5,6 +5,7 @@
 #include "PCGParamData.h"
 #include "Elements/Metadata/PCGMetadataElementCommon.h"
 #include "Metadata/PCGMetadata.h"
+#include "UObject/FortniteMainBranchObjectVersion.h"
 
 #include UE_INLINE_GENERATED_CPP_BY_NAME(PCGMetadataMakeVector)
 
@@ -77,6 +78,21 @@ void UPCGMetadataMakeVectorSettings::PostLoad()
 #endif // WITH_EDITOR
 }
 
+#if WITH_EDITOR
+void UPCGMetadataMakeVectorSettings::PostEditChangeProperty(FPropertyChangedEvent& PropertyChangedEvent)
+{
+	Super::PostEditChangeProperty(PropertyChangedEvent);
+
+	if (PropertyChangedEvent.Property->GetName() == GET_MEMBER_NAME_CHECKED(UPCGMetadataMakeVectorSettings, MakeVector3Op)
+		|| PropertyChangedEvent.Property->GetName() == GET_MEMBER_NAME_CHECKED(UPCGMetadataMakeVectorSettings, MakeVector4Op)
+		|| PropertyChangedEvent.Property->GetName() == GET_MEMBER_NAME_CHECKED(UPCGMetadataMakeVectorSettings, OutputType))
+	{
+		// @todo_pcg: It doesn't need to be reset in most cases.
+		ResetDefaultValues();
+	}
+}
+#endif // WITH_EDITOR
+
 FName UPCGMetadataMakeVectorSettings::GetInputPinLabel(uint32 Index) const
 {
 	static const FName DefaultNames[4] = {
@@ -194,56 +210,56 @@ FText UPCGMetadataMakeVectorSettings::GetDefaultNodeTitle() const
 {
 	return NSLOCTEXT("PCGMetadataMakeVectorSettings", "NodeTitle", "Make Vector Attribute");
 }
-#endif // WITH_EDITOR
 
-bool UPCGMetadataMakeVectorSettings::DoesInputSupportDefaultValue(uint32 Index) const
+void UPCGMetadataMakeVectorSettings::ApplyDeprecationBeforeUpdatePins(UPCGNode* InOutNode, TArray<TObjectPtr<UPCGPin>>& InputPins, TArray<TObjectPtr<UPCGPin>>& OutputPins)
 {
-	return true;
+	Super::ApplyDeprecationBeforeUpdatePins(InOutNode, InputPins, OutputPins);
+
+	// Supported default values on all pins
+	if (InOutNode && GetLinkerCustomVersion(FFortniteMainBranchObjectVersion::GUID) < FFortniteMainBranchObjectVersion::PCGInlineConstantDefaultValues)
+	{
+		for (const UPCGPin* Pin : InputPins)
+		{
+			if (IsPinDefaultValueEnabled(Pin->Properties.Label) && !Pin->IsConnected())
+			{
+				SetPinDefaultValueIsActivated(Pin->Properties.Label, /*bIsActivated=*/true, /*bDirtySettings=*/false);
+			}
+		}
+	}
 }
 
-UPCGParamData* UPCGMetadataMakeVectorSettings::CreateDefaultValueParam(FPCGContext* Context, uint32 Index) const
+FString UPCGMetadataMakeVectorSettings::GetPinInitialDefaultValueString(const FName PinLabel) const
 {
-	// Use labels since the logic is already done there.
-	FName Label = GetInputPinLabel(Index);
-
-	UPCGParamData* NewParamData = FPCGContext::NewObject_AnyThread<UPCGParamData>(Context);
-
-	if (Label == PCGMetadataMakeVectorConstants::XYZLabel)
+	if (PinLabel == PCGMetadataMakeVectorConstants::XYZLabel)
 	{
-		NewParamData->Metadata->CreateAttribute<FVector>(NAME_None, FVector::ZeroVector, /*bAllowsInterpolation=*/ true, /*bOverrideParent=*/ false);
+		return PCG::Private::MetadataTraits<FVector>::ZeroValueString();
 	}
-	else if (Label == PCGMetadataMakeVectorConstants::XYLabel || Label == PCGMetadataMakeVectorConstants::ZWLabel)
+	else if (PinLabel == PCGMetadataMakeVectorConstants::XYLabel || PinLabel == PCGMetadataMakeVectorConstants::ZWLabel)
 	{
-		NewParamData->Metadata->CreateAttribute<FVector2D>(NAME_None, FVector2D::ZeroVector, /*bAllowsInterpolation=*/ true, /*bOverrideParent=*/ false);
+		return PCG::Private::MetadataTraits<FVector2D>::ZeroValueString();
 	}
 	else
 	{
-		NewParamData->Metadata->CreateAttribute<double>(NAME_None, 0, /*bAllowsInterpolation=*/ true, /*bOverrideParent=*/ false);
+		return PCG::Private::MetadataTraits<double>::ZeroValueString();
 	}
-	
-	return NewParamData;
 }
+#endif // WITH_EDITOR
 
-#if WITH_EDITOR
-FString UPCGMetadataMakeVectorSettings::GetDefaultValueString(uint32 Index) const
+EPCGMetadataTypes UPCGMetadataMakeVectorSettings::GetPinInitialDefaultValueType(FName PinLabel) const
 {
-	// Use labels since the logic is already done there.
-	FName Label = GetInputPinLabel(Index);
-
-	if (Label == PCGMetadataMakeVectorConstants::XYZLabel)
+	if (PinLabel == PCGMetadataMakeVectorConstants::XYZLabel)
 	{
-		return FVector::ZeroVector.ToString();
+		return EPCGMetadataTypes::Vector;
 	}
-	else if (Label == PCGMetadataMakeVectorConstants::XYLabel || Label == PCGMetadataMakeVectorConstants::ZWLabel)
+	else if (PinLabel == PCGMetadataMakeVectorConstants::XYLabel || PinLabel == PCGMetadataMakeVectorConstants::ZWLabel)
 	{
-		return FVector2D::ZeroVector.ToString();
+		return EPCGMetadataTypes::Vector2;
 	}
 	else
 	{
-		return FString::Printf(TEXT("%f"), 0.0);
+		return EPCGMetadataTypes::Double;
 	}
 }
-#endif // WITH_EDITOR
 
 FPCGElementPtr UPCGMetadataMakeVectorSettings::CreateElement() const
 {
diff --git a/Engine/Plugins/PCG/Source/PCG/Private/Elements/Metadata/PCGMetadataOpElementBase.cpp b/Engine/Plugins/PCG/Source/PCG/Private/Elements/Metadata/PCGMetadataOpElementBase.cpp
--- a/Engine/Plugins/PCG/Source/PCG/Private/Elements/Metadata/PCGMetadataOpElementBase.cpp
+++ b/Engine/Plugins/PCG/Source/PCG/Private/Elements/Metadata/PCGMetadataOpElementBase.cpp
@@ -6,10 +6,14 @@
 #include "PCGCustomVersion.h"
 #include "PCGParamData.h"
 #include "PCGPin.h"
-#include "Data/PCGBasePointData.h"
+#include "Data/PCGPointData.h"
 #include "Elements/Metadata/PCGMetadataElementCommon.h"
+#include "Helpers/PCGPropertyHelpers.h"
+#include "Metadata/PCGMetadataTypesConstantStruct.h"
 #include "Metadata/Accessors/PCGAttributeAccessorHelpers.h"
 
+#include "Serialization/ArchiveCrc32.h"
+
 #include UE_INLINE_GENERATED_CPP_BY_NAME(PCGMetadataOpElementBase)
 
 #define LOCTEXT_NAMESPACE "PCGMetadataElementBaseElement"
@@ -24,12 +28,29 @@ namespace PCGMetadataBase
 	TAutoConsoleVariable<int> CVarMetadataOperationChunkSize(
 		TEXT("pcg.MetadataOperationChunkSize"),
 		256,
-		TEXT("Metadata operations chunk size."));	
-	
+		TEXT("Metadata operations chunk size."));
+
 	TAutoConsoleVariable<bool> CVarMetadataOperationReserveValues(
 		TEXT("pcg.MetadataOperationReserveValues"),
 		true,
 		TEXT("Metadata operations reserve values."));
+
+	namespace Helpers
+	{
+		FName GetInputSourceName(const int32 Index)
+		{
+			static const TStaticArray<FName, UPCGMetadataSettingsBase::MaxNumberOfInputs> InputSourceNames = {TEXT("InputSource1"), TEXT("InputSource2"), TEXT("InputSource3"), TEXT("InputSource4")};
+			const bool bValidIndex = Index > INDEX_NONE && Index < InputSourceNames.Num();
+			return ensure(bValidIndex) ? InputSourceNames[Index] : NAME_None;
+		}
+
+		FName GetDefaultValuePropertyName(const int32 Index)
+		{
+			static const TStaticArray<FName, UPCGMetadataSettingsBase::MaxNumberOfInputs> InputDefaultValuePropertyNames = {TEXT("DefaultValue1"), TEXT("DefaultValue2"), TEXT("DefaultValue3"), TEXT("DefaultValue4")};
+			const bool bValidIndex = Index > INDEX_NONE && Index < InputDefaultValuePropertyNames.Num();
+			return ensure(bValidIndex) ? InputDefaultValuePropertyNames[Index] : NAME_None;
+		}
+	}
 }
 
 void UPCGMetadataSettingsBase::PostLoad()
@@ -47,10 +68,11 @@ void UPCGMetadataSettingsBase::PostLoad()
 
 EPCGDataType UPCGMetadataSettingsBase::GetInputPinType(uint32 Index) const
 {
-	const EPCGDataType FirstPinTypeUnion = GetTypeUnionOfIncidentEdges(GetInputPinLabel(Index));
+	const FName PinLabel = GetInputPinLabel(Index);
+	const EPCGDataType FirstPinTypeUnion = GetTypeUnionOfIncidentEdges(PinLabel);
 
-	// If the pin is not connected but support default value, treat it as a param
-	if (FirstPinTypeUnion == EPCGDataType::None && DoesInputSupportDefaultValue(Index))
+	// If the pin is not connected but supports default values, treat it as a param
+	if (FirstPinTypeUnion == EPCGDataType::None && IsPinDefaultValueActivated(PinLabel))
 	{
 		return EPCGDataType::Param;
 	}
@@ -74,7 +96,7 @@ TArray<FName> UPCGMetadataSettingsBase::GetOutputDataFromPinOptions() const
 	return AllOptions;
 }
 
-uint32 UPCGMetadataSettingsBase::GetInputPinIndex(FName InPinLabel) const
+uint32 UPCGMetadataSettingsBase::GetInputPinIndex(const FName InPinLabel) const
 {
 	if (InPinLabel != PCGMetadataSettingsBaseConstants::DefaultOutputDataFromPinName)
 	{
@@ -90,6 +112,30 @@ uint32 UPCGMetadataSettingsBase::GetInputPinIndex(FName InPinLabel) const
 	return static_cast<uint32>(INDEX_NONE);
 }
 
+bool UPCGMetadataSettingsBase::IsPinDefaultValueEnabled(const FName PinLabel) const
+{
+	return GetInputPinIndex(PinLabel) != INDEX_NONE && PCGMetadataHelpers::MetadataTypeSupportsDefaultValues(GetPinDefaultValueType(PinLabel));
+}
+
+bool UPCGMetadataSettingsBase::IsPinDefaultValueActivated(const FName PinLabel) const
+{
+	if (!IsPinDefaultValueEnabled(PinLabel))
+	{
+		return false;
+	}
+
+	for (uint32 I = 0; I < GetOperandNum(); ++I)
+	{
+		if (PinLabel == GetInputPinLabel(I))
+		{
+			const FName PropertyName = PCGMetadataBase::Helpers::GetDefaultValuePropertyName(GetInputPinIndex(PinLabel));
+			return DefaultValues.IsPropertyActivated(PropertyName);
+		}
+	}
+
+	return false;
+}
+
 uint32 UPCGMetadataSettingsBase::GetInputPinToForward() const
 {
 	const uint32 OperandNum = GetOperandNum();
@@ -125,6 +171,25 @@ uint32 UPCGMetadataSettingsBase::GetInputPinToForward() const
 	return InputPinToForward;
 }
 
+const UPCGParamData* UPCGMetadataSettingsBase::CreateDefaultValueParamData(FPCGContext* Context, const FName PinLabel) const
+{
+	const FName PropertyName = PCGMetadataBase::Helpers::GetDefaultValuePropertyName(GetInputPinIndex(PinLabel));
+
+	const UPCGParamData* Data = DefaultValues.CreateParamData(Context, PropertyName);
+
+	// Default Value Container did not have a value. Try getting the node's 'reset' default value.
+	if (!Data)
+	{
+		const TObjectPtr<UPCGParamData> NewParamData = FPCGContext::NewObject_AnyThread<UPCGParamData>(Context);
+		if (CreateInitialDefaultValueAttribute(PinLabel, NewParamData->Metadata))
+		{
+			Data = NewParamData;
+		}
+	}
+
+	return Data;
+}
+
 EPCGDataType UPCGMetadataSettingsBase::GetCurrentPinTypes(const UPCGPin* InPin) const
 {
 	check(InPin);
@@ -153,12 +218,10 @@ void UPCGMetadataSettingsBase::PostEditChangeProperty(FPropertyChangedEvent& Pro
 
 bool UPCGMetadataSettingsBase::CanEditChange(const FProperty* InProperty) const
 {
-	if (!Super::CanEditChange(InProperty))
-	{
-		return false;
-	}
-
-	return InProperty && ((InProperty->GetFName() != GET_MEMBER_NAME_CHECKED(UPCGMetadataSettingsBase, OutputDataFromPin)) || (GetOperandNum() != 1));
+	return InProperty
+		&& Super::CanEditChange(InProperty)
+		&& CanEditInputSource(InProperty, GetOperandNum())
+		&& ((InProperty->GetFName() != GET_MEMBER_NAME_CHECKED(UPCGMetadataSettingsBase, OutputDataFromPin)) || (GetOperandNum() != 1));
 }
 
 bool UPCGMetadataSettingsBase::GetPinExtraIcon(const UPCGPin* InPin, FName& OutExtraIcon, FText& OutTooltip) const
@@ -204,6 +267,58 @@ void UPCGMetadataSettingsBase::ApplyDeprecation(UPCGNode* InOutNode)
 
 	Super::ApplyDeprecation(InOutNode);
 }
+
+void UPCGMetadataSettingsBase::SetPinDefaultValue(const FName PinLabel, const FString& DefaultValue, const bool bCreateIfNeeded)
+{
+	const FName PropertyName = PCGMetadataBase::Helpers::GetDefaultValuePropertyName(GetInputPinIndex(PinLabel));
+	if (PropertyName != NAME_None)
+	{
+		Modify();
+
+		if (!DefaultValues.FindProperty(PropertyName) && bCreateIfNeeded)
+		{
+			const EPCGMetadataTypes Type = GetPinInitialDefaultValueType(PinLabel);
+			DefaultValues.CreateNewProperty(PropertyName, Type);
+		}
+
+		if (DefaultValues.SetPropertyValueFromString(PropertyName, DefaultValue))
+		{
+			OnSettingsChangedDelegate.Broadcast(this, EPCGChangeType::Node | EPCGChangeType::Edge);
+		}
+	}
+}
+
+void UPCGMetadataSettingsBase::ConvertPinDefaultValueMetadataType(const FName PinLabel, const EPCGMetadataTypes DataType)
+{
+	if (ensure(IsPinDefaultValueActivated(PinLabel)))
+	{
+		const FName PropertyName = PCGMetadataBase::Helpers::GetDefaultValuePropertyName(GetInputPinIndex(PinLabel));
+		if (PropertyName != NAME_None && IsPinDefaultValueMetadataTypeValid(PinLabel, DataType))
+		{
+			Modify();
+			DefaultValues.ConvertPropertyType(PropertyName, DataType);
+			OnSettingsChangedDelegate.Broadcast(this, EPCGChangeType::Node | EPCGChangeType::Edge);
+		}
+	}
+}
+
+void UPCGMetadataSettingsBase::SetPinDefaultValueIsActivated(const FName PinLabel, const bool bIsActivated, const bool bDirtySettings)
+{
+	if (ensure(IsPinDefaultValueEnabled(PinLabel)))
+	{
+		if (bDirtySettings)
+		{
+			Modify();
+		}
+
+		const FName PropertyName = PCGMetadataBase::Helpers::GetDefaultValuePropertyName(GetInputPinIndex(PinLabel));
+		const bool bPropertyChanged = DefaultValues.SetPropertyActivated(PropertyName, bIsActivated);
+		if (bPropertyChanged && bDirtySettings)
+		{
+			OnSettingsChangedDelegate.Broadcast(this, EPCGChangeType::Node | EPCGChangeType::Edge);
+		}
+	}
+}
 #endif // WITH_EDITOR
 
 bool UPCGMetadataSettingsBase::DoesPinSupportPassThrough(UPCGPin* InPin) const
@@ -211,6 +326,77 @@ bool UPCGMetadataSettingsBase::DoesPinSupportPassThrough(UPCGPin* InPin) const
 	return InPin && !InPin->IsOutputPin() && GetInputPinIndex(InPin->Properties.Label) == GetInputPinToForward();
 }
 
+EPCGMetadataTypes UPCGMetadataSettingsBase::GetPinDefaultValueType(const FName PinLabel) const
+{
+	const FName PropertyName = PCGMetadataBase::Helpers::GetDefaultValuePropertyName(GetInputPinIndex(PinLabel));
+	if (PropertyName != NAME_None)
+	{
+		if (DefaultValues.FindProperty(PropertyName))
+		{
+			return DefaultValues.GetCurrentPropertyType(PropertyName);
+		}
+		else
+		{
+			return GetPinInitialDefaultValueType(PinLabel);
+		}
+	}
+
+	return EPCGMetadataTypes::Unknown;
+}
+
+bool UPCGMetadataSettingsBase::IsPinDefaultValueMetadataTypeValid(const FName PinLabel, EPCGMetadataTypes DataType) const
+{
+	bool bHasSpecialRequirement;
+	return IsSupportedInputType(static_cast<uint16>(DataType), GetInputPinIndex(PinLabel), bHasSpecialRequirement);
+}
+
+#if WITH_EDITOR
+void UPCGMetadataSettingsBase::ResetDefaultValues()
+{
+	DefaultValues.Reset();
+	OnSettingsChangedDelegate.Broadcast(this, EPCGChangeType::Settings | EPCGChangeType::Edge);
+}
+
+FString UPCGMetadataSettingsBase::GetPinDefaultValueAsString(const FName PinLabel) const
+{
+	if (ensure(IsPinDefaultValueActivated(PinLabel)))
+	{
+		const FName PropertyName = PCGMetadataBase::Helpers::GetDefaultValuePropertyName(GetInputPinIndex(PinLabel));
+		if (PropertyName != NAME_None)
+		{
+			if (DefaultValues.FindProperty(PropertyName))
+			{
+				return DefaultValues.GetPropertyValueAsString(PropertyName);
+			}
+			else
+			{
+				return GetPinInitialDefaultValueString(PinLabel);
+			}
+		}
+	}
+
+	return FString();
+}
+#endif // WITH_EDITOR
+
+EPCGMetadataTypes UPCGMetadataSettingsBase::GetPinInitialDefaultValueType(const FName PinLabel) const
+{
+#if WITH_EDITOR
+	ensure(IsPinDefaultValueMetadataTypeValid(PinLabel, EPCGMetadataTypes::Double));
+#endif // WITH_EDITOR
+	return EPCGMetadataTypes::Double;
+}
+
+bool UPCGMetadataSettingsBase::IsInputPinRequiredByExecution(const UPCGPin* InPin) const
+{
+	return InPin && !IsPinDefaultValueActivated(InPin->Properties.Label);
+}
+
+void UPCGMetadataSettingsBase::AddDefaultValuesToCrc(FArchiveCrc32& Crc32) const
+{
+	const_cast<FPCGDefaultValueContainer&>(DefaultValues).SerializeCrc(Crc32);
+}
+
 TArray<FPCGPinProperties> UPCGMetadataSettingsBase::InputPinProperties() const
 {
 	TArray<FPCGPinProperties> PinProperties;
@@ -222,7 +408,7 @@ TArray<FPCGPinProperties> UPCGMetadataSettingsBase::InputPinProperties() const
 		{
 			FPCGPinProperties& PinProperty = PinProperties.Emplace_GetRef(PinLabel, EPCGDataType::Any);
 
-			const bool bSupportDefaultValue = DoesInputSupportDefaultValue(InputPinIndex);
+			const bool bSupportDefaultValue = DefaultValuesAreEnabled() && IsPinDefaultValueEnabled(PinLabel);
 			if (!bSupportDefaultValue)
 			{
 				PinProperty.SetRequiredPin();
@@ -253,7 +439,10 @@ TArray<FPCGPinProperties> UPCGMetadataSettingsBase::InputPinProperties() const
 
 			if (bSupportDefaultValue)
 			{
-				AllTooltips.Add(FText::Format(LOCTEXT("PinTooltipDefaultValue", "Pin is optional, will use default value if not connected ({0})"), FText::FromString(GetDefaultValueString(InputPinIndex))));
+				EPCGMetadataTypes TypeValue = GetPinDefaultValueType(PinLabel);
+				const UEnum* TypeEnum = StaticEnum<EPCGMetadataTypes>();
+				const FText TypeString = TypeEnum ? TypeEnum->GetDisplayNameTextByValue(static_cast<int64>(TypeValue)) : FText();
+				AllTooltips.Add(FText::Format(LOCTEXT("PinTooltipDefaultValue", "Pin is using default value of type: {0}"), TypeString));
 			}
 
 			if (!AllTooltips.IsEmpty())
@@ -283,10 +472,30 @@ TArray<FPCGPinProperties> UPCGMetadataSettingsBase::OutputPinProperties() const
 	return PinProperties;
 }
 
+#if WITH_EDITOR
+bool UPCGMetadataSettingsBase::CanEditInputSource(const FProperty* InProperty, const int32 NumSources) const
+{
+	check(InProperty && NumSources <= MaxNumberOfInputs);
+
+	for (int32 SourceIndex = 0; SourceIndex < NumSources; ++SourceIndex)
+	{
+		if (InProperty->GetFName() == PCGMetadataBase::Helpers::GetInputSourceName(SourceIndex)
+			|| InProperty->GetFName() == TEXT("InputSource"))
+		{
+			if (const UPCGNode* Node = Cast<UPCGNode>(GetOuter()))
+			{
+				return Node->IsInputPinConnected(GetInputPinLabel(SourceIndex));
+			}
+		}
+	}
+
+	return true;
+}
+#endif // WITH_EDITOR
+
 void FPCGMetadataElementBase::PassthroughInput(FPCGContext* Context, TArray<FPCGTaggedData>& Outputs, const int32 Index) const
 {
 	check(Context);
-
 	const UPCGMetadataSettingsBase* Settings = Context->GetInputSettings<UPCGMetadataSettingsBase>();
 	check(Settings);
 
@@ -314,9 +523,9 @@ namespace PCGMetadataOpPrivate
 	using ContextType = FPCGMetadataElementBase::ContextType;
 	using ExecStateType = FPCGMetadataElementBase::ExecStateType;
 
-	void CreateAccessor(const UPCGMetadataSettingsBase* Settings, const FPCGTaggedData& InputData, PCGMetadataOps::FOperationData& OperationData, const int32 Index)
+	void CreateAccessor(const FPCGAttributePropertyInputSelector& Selector, const FPCGTaggedData& InputData, PCGMetadataOps::FOperationData& OperationData, const int32 Index)
 	{
-		OperationData.InputSources[Index] = Settings->GetInputSource(Index).CopyAndFixLast(InputData.Data);
+		OperationData.InputSources[Index] = Selector.CopyAndFixLast(InputData.Data);
 		const FPCGAttributePropertyInputSelector& InputSource = OperationData.InputSources[Index];
 
 		OperationData.InputAccessors[Index] = PCGAttributeAccessorHelpers::CreateConstAccessor(InputData.Data, InputSource);
@@ -407,7 +616,7 @@ bool FPCGMetadataElementBase::PrepareDataInternal(FPCGContext* Context) const
 	const TArray<FPCGTaggedData> PrimaryInputs = Context->InputData.GetInputsByPin(PrimaryPinLabel);
 
 	// There are no inputs on the primary pin, so pass-through inputs if the primary pin is required
-	if (Settings->IsPrimaryInputPinRequired() && PrimaryInputs.IsEmpty())
+	if (!Settings->IsPinDefaultValueActivated(PrimaryPinLabel) && PrimaryInputs.IsEmpty())
 	{
 		return true;
 	}
@@ -421,7 +630,7 @@ bool FPCGMetadataElementBase::PrepareDataInternal(FPCGContext* Context) const
 		{
 			const int32 CurrentInputNum = Context->InputData.GetInputsByPin(Settings->GetInputPinLabel(i)).Num();
 
-			OperandInputNumMax = FMath::Max(OperandInputNumMax, CurrentInputNum);
+			OperandInputNumMax = FMath::Max(OperandInputNumMax, FMath::Max(1, CurrentInputNum));
 
 			// For the current input, no input (0) could be default value and we support N:1 and 1:N
 			if (CurrentInputNum > 1 && CurrentInputNum != OperandInputNumMax)
@@ -456,6 +665,7 @@ bool FPCGMetadataElementBase::PrepareDataInternal(FPCGContext* Context) const
 		auto NoOperation = [this, Context, IterationIndex, &Outputs]() { PassthroughInput(Context, Outputs, IterationIndex); return EPCGTimeSliceInitResult::NoOperation; };
 
 		const uint32 PrimaryPinIndex = Settings->GetInputPinToForward();
+		OutState.DefaultValueOverriddenPins.AddZeroed(OperandNum);
 
 		// Iterate over the inputs and validate
 		for (uint32 OperandPinIndex = 0; OperandPinIndex < OperandNum; ++OperandPinIndex)
@@ -465,27 +675,35 @@ bool FPCGMetadataElementBase::PrepareDataInternal(FPCGContext* Context) const
 			TArray<FPCGTaggedData> CurrentPinInputData = Context->InputData.GetInputsByPin(CurrentPinLabel);
 
 			// This only needs to be checked once
-			if (!bIsInputConnected && CurrentPinInputData.IsEmpty())
+			if (Settings->DefaultValuesAreEnabled() && !bIsInputConnected && CurrentPinInputData.IsEmpty())
 			{
 				FPCGTaggedData& DefaultData = CurrentPinInputData.Emplace_GetRef();
 				DefaultData.Pin = CurrentPinLabel;
-				DefaultData.Data = Settings->CreateDefaultValueParam(Context, OperandPinIndex);
-
-				if (!DefaultData.Data)
-				{
-					PCGE_LOG(Error, GraphAndLog, FText::Format(LOCTEXT("CantCreateDefaultValue", "Pin '{0}' supports default value but we could not create it."), FText::FromName(CurrentPinLabel)));
-					return EPCGTimeSliceInitResult::AbortExecution;
-				}
-				else
+				if (Settings->IsPinDefaultValueActivated(CurrentPinLabel))
 				{
-					// Need to make sure the param data is properly tracked by the context to prevent garbage collection
-					TimeSlicedContext->TrackObject(DefaultData.Data);
+					// @todo_pcg: Future optimizations/refactors - cache the param data on the settings, or use accessors on the default value struct, etc.
+					// Create from the Default Value Container if it exists
+					DefaultData.Data = Settings->CreateDefaultValueParamData(Context, CurrentPinLabel);
 
-					// Need to make sure the param data has at least one entry
-					UPCGMetadata* DefaultParamMetadata = CastChecked<UPCGParamData>(DefaultData.Data)->Metadata;
-					if (DefaultParamMetadata->GetLocalItemCount() == 0)
+					// Couldn't create a default value
+					if (!DefaultData.Data)
 					{
-						DefaultParamMetadata->AddEntry();
+						PCGE_LOG(Error, GraphAndLog, FText::Format(LOCTEXT("CantCreateDefaultValue", "Pin '{0}' supports default value but we could not create it."), FText::FromName(CurrentPinLabel)));
+						return EPCGTimeSliceInitResult::AbortExecution;
+					}
+					else
+					{
+						// Need to make sure the param data is properly tracked by the context to prevent garbage collection
+						TimeSlicedContext->TrackObject(DefaultData.Data);
+
+						// Need to make sure the param data has at least one entry
+						UPCGMetadata* DefaultParamMetadata = CastChecked<UPCGParamData>(DefaultData.Data)->Metadata;
+						if (DefaultParamMetadata->GetLocalItemCount() == 0)
+						{
+							DefaultParamMetadata->AddEntry();
+						}
+
+						OutState.DefaultValueOverriddenPins[OperandPinIndex] = true;
 					}
 				}
 			}
@@ -539,7 +757,9 @@ bool FPCGMetadataElementBase::PrepareDataInternal(FPCGContext* Context) const
 		const FPCGTaggedData& PrimaryPinData = InputTaggedData[PrimaryPinIndex];
 
 		// First create an accessor for the input to forward (it's our control data)
-		PCGMetadataOpPrivate::CreateAccessor(Settings, PrimaryPinData, OperationData, PrimaryPinIndex);
+		const FPCGAttributePropertyInputSelector PrimarySelector = !OutState.DefaultValueOverriddenPins[PrimaryPinIndex] ? Settings->GetInputSource(PrimaryPinIndex) : FPCGAttributePropertyInputSelector{};
+
+		PCGMetadataOpPrivate::CreateAccessor(PrimarySelector, PrimaryPinData, OperationData, PrimaryPinIndex);
 		if (!PCGMetadataOpPrivate::ValidateAccessor(Context, Settings, PrimaryPinData, OperationData, PrimaryPinIndex))
 		{
 			return EPCGTimeSliceInitResult::AbortExecution;
@@ -565,8 +785,9 @@ bool FPCGMetadataElementBase::PrepareDataInternal(FPCGContext* Context) const
 					return EPCGTimeSliceInitResult::AbortExecution;
 				}
 
-				PCGMetadataOpPrivate::CreateAccessor(Settings, InputTaggedData[Index], OperationData, Index);
+				const FPCGAttributePropertyInputSelector Selector = !OutState.DefaultValueOverriddenPins[Index] ? Settings->GetInputSource(Index) : FPCGAttributePropertyInputSelector{};
 
+				PCGMetadataOpPrivate::CreateAccessor(Selector, InputTaggedData[Index], OperationData, Index);
 				if (!PCGMetadataOpPrivate::ValidateAccessor(Context, Settings, InputTaggedData[Index], OperationData, Index))
 				{
 					return EPCGTimeSliceInitResult::AbortExecution;
@@ -754,4 +975,25 @@ bool FPCGMetadataElementBase::ExecuteInternal(FPCGContext* Context) const
 	});
 }
 
+void FPCGMetadataElementBase::GetDependenciesCrc(const FPCGDataCollection& InInput, const UPCGSettings* InSettings, UPCGComponent* InComponent, FPCGCrc& OutCrc) const
+{
+	FPCGCrc Crc;
+	IPCGElement::GetDependenciesCrc(InInput, InSettings, InComponent, Crc);
+
+	// Add the default values to the crc
+	if (const UPCGMetadataSettingsBase* Settings = Cast<UPCGMetadataSettingsBase>(InSettings))
+	{
+		FArchiveCrc32 Crc32;
+		/*
+		 * Implementation note: In theory, the default value behind connected pins should not factor into the Crc, but in practice,
+		 * it makes the code more obscure for a gain that might not even be meaningful.
+		 */
+		Crc32.SetIsSaving(true);
+		Settings->AddDefaultValuesToCrc(Crc32);
+		Crc.Combine(Crc32.GetCrc());
+	}
+
+	OutCrc = Crc;
+}
+
 #undef LOCTEXT_NAMESPACE
diff --git a/Engine/Plugins/PCG/Source/PCG/Private/Elements/Metadata/PCGMetadataRotatorOpElement.cpp b/Engine/Plugins/PCG/Source/PCG/Private/Elements/Metadata/PCGMetadataRotatorOpElement.cpp
--- a/Engine/Plugins/PCG/Source/PCG/Private/Elements/Metadata/PCGMetadataRotatorOpElement.cpp
+++ b/Engine/Plugins/PCG/Source/PCG/Private/Elements/Metadata/PCGMetadataRotatorOpElement.cpp
@@ -38,7 +38,7 @@ namespace PCGMetadataRotatorSettings
 		return Operation == EPCGMetadataRotatorOperation::Lerp;
 	}
 
-	inline constexpr bool IsTransfromOp(EPCGMetadataRotatorOperation Operation)
+	inline constexpr bool IsTransformOp(EPCGMetadataRotatorOperation Operation)
 	{
 		return (uint16)Operation >= (uint16)EPCGMetadataRotatorOperation::TransformOp;
 	}
@@ -127,7 +127,7 @@ FName UPCGMetadataRotatorSettings::GetInputPinLabel(uint32 Index) const
 	switch (Index)
 	{
 	case 0:
-		if (!PCGMetadataRotatorSettings::IsUnaryOp(Operation) && !PCGMetadataRotatorSettings::IsTransfromOp(Operation))
+		if (!PCGMetadataRotatorSettings::IsUnaryOp(Operation) && !PCGMetadataRotatorSettings::IsTransformOp(Operation))
 		{
 			return PCGMetadataSettingsBaseConstants::DoubleInputFirstLabel;
 		}
@@ -136,7 +136,7 @@ FName UPCGMetadataRotatorSettings::GetInputPinLabel(uint32 Index) const
 			return PCGPinConstants::DefaultInputLabel;
 		}
 	case 1:
-		return PCGMetadataRotatorSettings::IsTransfromOp(Operation) ? PCGMetadataSettingsBaseConstants::TransformLabel : PCGMetadataSettingsBaseConstants::DoubleInputSecondLabel;
+		return PCGMetadataRotatorSettings::IsTransformOp(Operation) ? PCGMetadataSettingsBaseConstants::TransformLabel : PCGMetadataSettingsBaseConstants::DoubleInputSecondLabel;
 	case 2:
 		return PCGMetadataSettingsBaseConstants::LerpRatioLabel;
 	default:
@@ -167,7 +167,7 @@ bool UPCGMetadataRotatorSettings::IsSupportedInputType(uint16 TypeId, uint32 Inp
 		bHasSpecialRequirement = true;
 		return PCG::Private::IsOfTypes<float, double>(TypeId);
 	}
-	else if (InputIndex == 1 && PCGMetadataRotatorSettings::IsTransfromOp(Operation))
+	else if (InputIndex == 1 && PCGMetadataRotatorSettings::IsTransformOp(Operation))
 	{
 		bHasSpecialRequirement = true;
 		return PCG::Private::IsOfTypes<FTransform>(TypeId);
@@ -234,6 +234,34 @@ void UPCGMetadataRotatorSettings::ApplyPreconfiguredSettings(const FPCGPreConfig
 	}
 }
 
+bool UPCGMetadataRotatorSettings::IsPinDefaultValueMetadataTypeValid(FName PinLabel, const EPCGMetadataTypes DataType) const
+{
+	return DataType == EPCGMetadataTypes::Rotator
+		|| DataType == EPCGMetadataTypes::Vector;
+}
+
+EPCGMetadataTypes UPCGMetadataRotatorSettings::GetPinInitialDefaultValueType(const FName PinLabel) const
+{
+	if (PinLabel == PCGPinConstants::DefaultInputLabel
+		|| PinLabel == PCGMetadataSettingsBaseConstants::DoubleInputFirstLabel
+		|| PinLabel == PCGMetadataSettingsBaseConstants::DoubleInputSecondLabel)
+	{
+		return EPCGMetadataTypes::Rotator;
+	}
+	else if (PinLabel == PCGMetadataSettingsBaseConstants::TransformLabel)
+	{
+		return EPCGMetadataTypes::Transform;
+	}
+	else if (PinLabel == PCGMetadataSettingsBaseConstants::LerpRatioLabel)
+	{
+		return EPCGMetadataTypes::Double;
+	}
+	else
+	{
+		return EPCGMetadataTypes::Unknown;
+	}
+}
+
 FPCGElementPtr UPCGMetadataRotatorSettings::CreateElement() const
 {
 	return MakeShared<FPCGMetadataRotatorElement>();
@@ -254,7 +282,7 @@ bool FPCGMetadataRotatorElement::DoOperation(PCGMetadataOps::FOperationData& Ope
 		}
 		else
 		{
-			if (PCGMetadataRotatorSettings::IsTransfromOp(Operation))
+			if (PCGMetadataRotatorSettings::IsTransformOp(Operation))
 			{
 				return DoBinaryOp<AttributeType, FTransform>(OperationData, [Operation](const AttributeType& Value, const FTransform& Transform)->AttributeType {
 					return PCGMetadataRotatorSettings::ApplyTransformOperation(Value, Transform, Operation);
diff --git a/Engine/Plugins/PCG/Source/PCG/Private/Elements/Metadata/PCGMetadataTransformOpElement.cpp b/Engine/Plugins/PCG/Source/PCG/Private/Elements/Metadata/PCGMetadataTransformOpElement.cpp
--- a/Engine/Plugins/PCG/Source/PCG/Private/Elements/Metadata/PCGMetadataTransformOpElement.cpp
+++ b/Engine/Plugins/PCG/Source/PCG/Private/Elements/Metadata/PCGMetadataTransformOpElement.cpp
@@ -197,6 +197,24 @@ void UPCGMetadataTransformSettings::ApplyPreconfiguredSettings(const FPCGPreConf
 	}
 }
 
+bool UPCGMetadataTransformSettings::IsPinDefaultValueEnabled(const FName PinLabel) const
+{
+	// @todo_pcg: Add transforms when they are supported.
+	return PinLabel == PCGMetadataSettingsBaseConstants::LerpRatioLabel && Super::IsPinDefaultValueEnabled(PinLabel);
+}
+
+EPCGMetadataTypes UPCGMetadataTransformSettings::GetPinDefaultValueType(const FName PinLabel) const
+{
+	if (PinLabel == PCGMetadataSettingsBaseConstants::LerpRatioLabel)
+	{
+		return EPCGMetadataTypes::Double;
+	}
+	else
+	{
+		return EPCGMetadataTypes::Transform;
+	}
+}
+
 FPCGElementPtr UPCGMetadataTransformSettings::CreateElement() const
 {
 	return MakeShared<FPCGMetadataTransformElement>();
diff --git a/Engine/Plugins/PCG/Source/PCG/Private/Elements/Metadata/PCGMetadataVectorOpElement.cpp b/Engine/Plugins/PCG/Source/PCG/Private/Elements/Metadata/PCGMetadataVectorOpElement.cpp
--- a/Engine/Plugins/PCG/Source/PCG/Private/Elements/Metadata/PCGMetadataVectorOpElement.cpp
+++ b/Engine/Plugins/PCG/Source/PCG/Private/Elements/Metadata/PCGMetadataVectorOpElement.cpp
@@ -359,6 +359,17 @@ void UPCGMetadataVectorSettings::ApplyPreconfiguredSettings(const FPCGPreConfigu
 	}
 }
 
+bool UPCGMetadataVectorSettings::IsPinDefaultValueEnabled(const FName PinLabel) const
+{
+	// @todo_pcg: Remove this when transforms are enabled.
+	if (PCGMetadataVectorSettings::IsTransformOp(Operation) && PinLabel == GetInputPinLabel(1))
+	{
+		return false;
+	}
+
+	return Super::IsPinDefaultValueEnabled(PinLabel);
+}
+
 FPCGElementPtr UPCGMetadataVectorSettings::CreateElement() const
 {
 	return MakeShared<FPCGMetadataVectorElement>();
diff --git a/Engine/Plugins/PCG/Source/PCG/Private/Elements/PCGCreateAttribute.cpp b/Engine/Plugins/PCG/Source/PCG/Private/Elements/PCGCreateAttribute.cpp
--- a/Engine/Plugins/PCG/Source/PCG/Private/Elements/PCGCreateAttribute.cpp
+++ b/Engine/Plugins/PCG/Source/PCG/Private/Elements/PCGCreateAttribute.cpp
@@ -340,7 +340,7 @@ bool UPCGCreateAttributeSetSettings::ConvertNode(const FPCGPreconfiguredInfo& Co
 					PropertyName.SetNumber(NameIteration);
 				}
 
-				FPropertyBagPropertyDesc PropertyDesc = PCGPropertyHelpers::GetPropertyBagDescWithMetadataType(PropertyName, AttributeTypes.Type);
+				FPropertyBagPropertyDesc PropertyDesc = PCGPropertyHelpers::CreatePropertyBagDescWithMetadataType(PropertyName, AttributeTypes.Type);
 				Graph->AddUserParameters({PropertyDesc});
 				const FPropertyBagPropertyDesc* Desc = UserParametersStruct->FindPropertyDescByName(PropertyName);
 				check(Desc);
diff --git a/Engine/Plugins/PCG/Source/PCG/Private/Helpers/PCGDefaultValueContainer.cpp b/Engine/Plugins/PCG/Source/PCG/Private/Helpers/PCGDefaultValueContainer.cpp
new file mode 100644
--- /dev/null
+++ b/Engine/Plugins/PCG/Source/PCG/Private/Helpers/PCGDefaultValueContainer.cpp
@@ -0,0 +1,121 @@
+// Copyright Epic Games, Inc. All Rights Reserved.
+
+#include "Helpers/PCGDefaultValueContainer.h"
+
+#include "PCGContext.h"
+#include "PCGParamData.h"
+#include "Helpers/PCGGraphParameterExtension.h"
+#include "Helpers/PCGPropertyHelpers.h"
+#include "Metadata/PCGMetadataAttributeTraits.h"
+
+const FProperty* FPCGDefaultValueContainer::CreateNewProperty(const FName PropertyName, const EPCGMetadataTypes Type)
+{
+	if (PropertyName == NAME_None || !PCGMetadataHelpers::MetadataTypeSupportsDefaultValues(Type))
+	{
+		return nullptr;
+	}
+
+	if (PropertyBag.FindPropertyDescByName(PropertyName))
+	{
+		PropertyBag.RemovePropertyByName(PropertyName);
+	}
+
+	const FPropertyBagPropertyDesc PropertyDesc = PCGPropertyHelpers::CreatePropertyBagDescWithMetadataType(PropertyName, Type);
+	PropertyBag.AddProperties({PropertyDesc});
+	if (const FPropertyBagPropertyDesc* OutPropertyDesc = PropertyBag.FindPropertyDescByName(PropertyName))
+	{
+		return OutPropertyDesc->CachedProperty;
+	}
+	else
+	{
+		return nullptr;
+	}
+}
+
+const FProperty* FPCGDefaultValueContainer::FindProperty(const FName PropertyName) const
+{
+	const FPropertyBagPropertyDesc* PropertyDesc = PropertyBag.FindPropertyDescByName(PropertyName);
+	return PropertyDesc ? PropertyDesc->CachedProperty : nullptr;
+}
+
+EPCGMetadataTypes FPCGDefaultValueContainer::GetCurrentPropertyType(const FName PropertyName) const
+{
+	const FProperty* Property = FindProperty(PropertyName);
+	return Property ? PCGPropertyHelpers::GetMetadataTypeFromProperty(Property) : EPCGMetadataTypes::Unknown;
+}
+
+FString FPCGDefaultValueContainer::GetPropertyValueAsString(const FName PropertyName) const
+{
+	TValueOrError<FString, EPropertyBagResult> Result = PropertyBag.GetValueSerializedString(PropertyName);
+
+	return Result.HasValue() ? Result.GetValue() : FString(TEXT("Error"));
+}
+
+const UPCGParamData* FPCGDefaultValueContainer::CreateParamData(FPCGContext* Context, const FName PropertyName) const
+{
+	if (const FProperty* PropertyPtr = FindProperty(PropertyName))
+	{
+		const TObjectPtr<UPCGParamData> NewParamData = FPCGContext::NewObject_AnyThread<UPCGParamData>(Context);
+		if (NewParamData->Metadata->CreateAttributeFromDataProperty(NAME_None, PropertyBag.GetValue().GetMemory(), PropertyPtr))
+		{
+			return NewParamData;
+		}
+	}
+
+	return nullptr;
+}
+
+bool FPCGDefaultValueContainer::IsPropertyActivated(const FName PropertyName) const
+{
+	return ActivatedProperties.Contains(PropertyName);
+}
+
+#if WITH_EDITOR
+const FProperty* FPCGDefaultValueContainer::ConvertPropertyType(const FName PropertyName, const EPCGMetadataTypes Type)
+{
+	if (!PCGMetadataHelpers::MetadataTypeSupportsDefaultValues(Type) || Type == GetCurrentPropertyType(PropertyName))
+	{
+		return nullptr;
+	}
+
+	return CreateNewProperty(PropertyName, Type);
+}
+
+bool FPCGDefaultValueContainer::SetPropertyValueFromString(const FName PropertyName, const FString& ValueString)
+{
+	if (PropertyName == NAME_None)
+	{
+		return false;
+	}
+
+	SetPropertyActivated(PropertyName, /*bIsActive=*/true);
+	return PropertyBag.SetValueSerializedString(PropertyName, ValueString) == EPropertyBagResult::Success;
+}
+
+bool FPCGDefaultValueContainer::SetPropertyActivated(const FName PropertyName, const bool bIsActivated)
+{
+	if ((PropertyName != NAME_None) && (bIsActivated != ActivatedProperties.Contains(PropertyName)))
+	{
+		if (bIsActivated)
+		{
+			ActivatedProperties.Add(PropertyName);
+			return true;
+		}
+		else
+		{
+			ActivatedProperties.Remove(PropertyName);
+			return true;
+		}
+	}
+	else
+	{
+		return false;
+	}
+}
+
+void FPCGDefaultValueContainer::Reset()
+{
+	ActivatedProperties.Reset();
+	PropertyBag.Reset();
+}
+#endif // WITH_EDITOR
diff --git a/Engine/Plugins/PCG/Source/PCG/Private/Helpers/PCGPropertyHelpers.cpp b/Engine/Plugins/PCG/Source/PCG/Private/Helpers/PCGPropertyHelpers.cpp
--- a/Engine/Plugins/PCG/Source/PCG/Private/Helpers/PCGPropertyHelpers.cpp
+++ b/Engine/Plugins/PCG/Source/PCG/Private/Helpers/PCGPropertyHelpers.cpp
@@ -232,7 +232,7 @@ EPCGMetadataTypes PCGPropertyHelpers::GetMetadataTypeFromProperty(const FPropert
 	return PropertyAccessor.IsValid() ? EPCGMetadataTypes(PropertyAccessor->GetUnderlyingType()) : EPCGMetadataTypes::Unknown;
 }
 
-FPropertyBagPropertyDesc PCGPropertyHelpers::GetPropertyBagDescWithMetadataType(const FName InPropertyName, const EPCGMetadataTypes Type)
+FPropertyBagPropertyDesc PCGPropertyHelpers::CreatePropertyBagDescWithMetadataType(const FName InPropertyName, const EPCGMetadataTypes Type)
 {
 	EPropertyBagPropertyType ValueType = EPropertyBagPropertyType::Struct;
 	const UObject* ValueTypeObject = nullptr;
diff --git a/Engine/Plugins/PCG/Source/PCG/Public/Elements/Metadata/PCGMetadataBitwiseOpElement.h b/Engine/Plugins/PCG/Source/PCG/Public/Elements/Metadata/PCGMetadataBitwiseOpElement.h
--- a/Engine/Plugins/PCG/Source/PCG/Public/Elements/Metadata/PCGMetadataBitwiseOpElement.h
+++ b/Engine/Plugins/PCG/Source/PCG/Public/Elements/Metadata/PCGMetadataBitwiseOpElement.h
@@ -54,6 +54,13 @@ protected:
 	virtual FPCGElementPtr CreateElement() const override;
 	//~End UPCGSettings interface
 
+	//~Begin IPCGSettingsDefaultValueProvider interface
+#if WITH_EDITOR
+	virtual FString GetPinInitialDefaultValueString(FName PinLabel) const override { return PCG::Private::MetadataTraits<int32>::ZeroValueString(); }
+#endif // WITH_EDITOR
+	virtual EPCGMetadataTypes GetPinInitialDefaultValueType(FName PinLabel) const override { return EPCGMetadataTypes::Integer32; }
+	//~End IPCGSettingsDefaultValueProvider interface
+
 public:
 	UPROPERTY(BlueprintReadWrite, EditAnywhere, Category = Settings)
 	EPCGMetadataBitwiseOperation Operation = EPCGMetadataBitwiseOperation::And;
diff --git a/Engine/Plugins/PCG/Source/PCG/Public/Elements/Metadata/PCGMetadataBooleanOpElement.h b/Engine/Plugins/PCG/Source/PCG/Public/Elements/Metadata/PCGMetadataBooleanOpElement.h
--- a/Engine/Plugins/PCG/Source/PCG/Public/Elements/Metadata/PCGMetadataBooleanOpElement.h
+++ b/Engine/Plugins/PCG/Source/PCG/Public/Elements/Metadata/PCGMetadataBooleanOpElement.h
@@ -46,6 +46,13 @@ public:
 	virtual uint16 GetOutputType(uint16 InputTypeId) const override;
 	//~End UPCGMetadataSettingsBase interface
 
+	//~Begin IPCGSettingsDefaultValueProvider interface
+#if WITH_EDITOR
+	virtual FString GetPinInitialDefaultValueString(FName PinLabel) const override { return PCG::Private::MetadataTraits<bool>::ZeroValueString(); }
+#endif // WITH_EDITOR
+	virtual EPCGMetadataTypes GetPinInitialDefaultValueType(FName PinLabel) const override { return EPCGMetadataTypes::Boolean; }
+	//~End IPCGSettingsDefaultValueProvider interface
+
 protected:
 	//~Begin UPCGSettings interface
 #if WITH_EDITOR
diff --git a/Engine/Plugins/PCG/Source/PCG/Public/Elements/Metadata/PCGMetadataBreakTransform.h b/Engine/Plugins/PCG/Source/PCG/Public/Elements/Metadata/PCGMetadataBreakTransform.h
--- a/Engine/Plugins/PCG/Source/PCG/Public/Elements/Metadata/PCGMetadataBreakTransform.h
+++ b/Engine/Plugins/PCG/Source/PCG/Public/Elements/Metadata/PCGMetadataBreakTransform.h
@@ -15,6 +15,9 @@ class UPCGMetadataBreakTransformSettings : public UPCGMetadataSettingsBase
 public:
 	// ~Begin UObject interface
 	virtual void PostLoad() override;
+#if WITH_EDITOR
+	virtual bool CanEditChange(const FProperty* InProperty) const override;
+#endif // WITH_EDITOR
 	// ~End UObject interface
 	// 
 	//~Begin UPCGSettings interface
@@ -22,7 +25,6 @@ public:
 	virtual FName GetDefaultNodeName() const override;
 	virtual FText GetDefaultNodeTitle() const override;
 #endif
-	
 
 	virtual FPCGAttributePropertyInputSelector GetInputSource(uint32 Index) const override;
 
@@ -39,6 +41,13 @@ protected:
 	virtual FPCGElementPtr CreateElement() const override;
 	//~End UPCGSettings interface
 
+	//~Begin IPCGSettingsDefaultValueProvider interface
+#if WITH_EDITOR
+	virtual FString GetPinInitialDefaultValueString(FName PinLabel) const override { return PCG::Private::MetadataTraits<FTransform>::ZeroValueString(); }
+#endif // WITH_EDITOR
+	virtual EPCGMetadataTypes GetPinInitialDefaultValueType(FName PinLabel) const override { return EPCGMetadataTypes::Transform; }
+	//~End IPCGSettingsDefaultValueProvider interface
+
 public:
 	UPROPERTY(BlueprintReadWrite, EditAnywhere, Category = Input, meta = (PCG_Overridable))
 	FPCGAttributePropertyInputSelector InputSource;
diff --git a/Engine/Plugins/PCG/Source/PCG/Public/Elements/Metadata/PCGMetadataBreakVector.h b/Engine/Plugins/PCG/Source/PCG/Public/Elements/Metadata/PCGMetadataBreakVector.h
--- a/Engine/Plugins/PCG/Source/PCG/Public/Elements/Metadata/PCGMetadataBreakVector.h
+++ b/Engine/Plugins/PCG/Source/PCG/Public/Elements/Metadata/PCGMetadataBreakVector.h
@@ -23,6 +23,9 @@ class UPCGMetadataBreakVectorSettings : public UPCGMetadataSettingsBase
 public:
 	// ~Begin UObject interface
 	virtual void PostLoad() override;
+#if WITH_EDITOR
+	virtual bool CanEditChange(const FProperty* InProperty) const override;
+#endif // WITH_EDITOR
 	// ~End UObject interface
 
 	//~Begin UPCGSettings interface
@@ -30,7 +33,7 @@ public:
 	virtual FName GetDefaultNodeName() const override;
 	virtual FText GetDefaultNodeTitle() const override;
 #endif
-	
+
 
 	virtual FPCGAttributePropertyInputSelector GetInputSource(uint32 Index) const override;
 
@@ -45,6 +48,13 @@ protected:
 	virtual FPCGElementPtr CreateElement() const override;
 	//~End UPCGSettings interface
 
+	//~Begin IPCGSettingsDefaultValueProvider interface
+#if WITH_EDITOR
+	virtual FString GetPinInitialDefaultValueString(FName PinLabel) const override { return PCG::Private::MetadataTraits<FVector>::ZeroValueString(); }
+#endif // WITH_EDITOR
+	virtual EPCGMetadataTypes GetPinInitialDefaultValueType(FName PinLabel) const override { return EPCGMetadataTypes::Vector; }
+	//~End IPCGSettingsDefaultValueProvider interface
+
 public:
 	UPROPERTY(BlueprintReadWrite, EditAnywhere, Category = Input, meta = (PCG_Overridable))
 	FPCGAttributePropertyInputSelector InputSource;
diff --git a/Engine/Plugins/PCG/Source/PCG/Public/Elements/Metadata/PCGMetadataMakeRotator.h b/Engine/Plugins/PCG/Source/PCG/Public/Elements/Metadata/PCGMetadataMakeRotator.h
--- a/Engine/Plugins/PCG/Source/PCG/Public/Elements/Metadata/PCGMetadataMakeRotator.h
+++ b/Engine/Plugins/PCG/Source/PCG/Public/Elements/Metadata/PCGMetadataMakeRotator.h
@@ -43,18 +43,15 @@ class UPCGMetadataMakeRotatorSettings : public UPCGMetadataSettingsBase
 
 public:
 #if WITH_EDITOR
+	virtual void PostEditChangeProperty(FPropertyChangedEvent& PropertyChangedEvent) override;
 	//~Begin UPCGSettings interface
 	virtual FName GetDefaultNodeName() const override;
 	virtual FText GetDefaultNodeTitle() const override;
 	virtual TArray<FPCGPreConfiguredSettingsInfo> GetPreconfiguredInfo() const override;
 	virtual bool OnlyExposePreconfiguredSettings() const override { return true; }
+	virtual void ApplyDeprecationBeforeUpdatePins(UPCGNode* InOutNode, TArray<TObjectPtr<UPCGPin>>& InputPins, TArray<TObjectPtr<UPCGPin>>& OutputPins) override;
 #endif
 	virtual void ApplyPreconfiguredSettings(const FPCGPreConfiguredSettingsInfo& PreconfiguredInfo) override;
-	virtual bool DoesInputSupportDefaultValue(uint32 Index) const override;
-	virtual UPCGParamData* CreateDefaultValueParam(FPCGContext* Context, uint32 Index) const override;
-#if WITH_EDITOR
-	virtual FString GetDefaultValueString(uint32 Index) const override;
-#endif // WITH_EDITOR
 	//~End UPCGSettings interface
 
 	virtual FPCGAttributePropertyInputSelector GetInputSource(uint32 Index) const override;
@@ -70,6 +67,12 @@ protected:
 	virtual FPCGElementPtr CreateElement() const override;
 	//~End UPCGSettings interface
 
+	//~Begin IPCGSettingsDefaultValueProvider interface
+#if WITH_EDITOR
+	virtual FString GetPinInitialDefaultValueString(FName PinLabel) const override;
+#endif // WITH_EDITOR
+	//~End IPCGSettingsDefaultValueProvider interface
+
 public:
 	UPROPERTY(BlueprintReadWrite, EditAnywhere, Category = Input, meta = (PCG_Overridable))
 	FPCGAttributePropertyInputSelector InputSource1;
diff --git a/Engine/Plugins/PCG/Source/PCG/Public/Elements/Metadata/PCGMetadataMakeTransform.h b/Engine/Plugins/PCG/Source/PCG/Public/Elements/Metadata/PCGMetadataMakeTransform.h
--- a/Engine/Plugins/PCG/Source/PCG/Public/Elements/Metadata/PCGMetadataMakeTransform.h
+++ b/Engine/Plugins/PCG/Source/PCG/Public/Elements/Metadata/PCGMetadataMakeTransform.h
@@ -28,12 +28,8 @@ public:
 	//~Begin UPCGSettings interface
 	virtual FName GetDefaultNodeName() const override;
 	virtual FText GetDefaultNodeTitle() const override;
+	virtual void ApplyDeprecationBeforeUpdatePins(UPCGNode* InOutNode, TArray<TObjectPtr<UPCGPin>>& InputPins, TArray<TObjectPtr<UPCGPin>>& OutputPins) override;
 #endif
-	virtual bool DoesInputSupportDefaultValue(uint32 Index) const override;
-	virtual UPCGParamData* CreateDefaultValueParam(FPCGContext* Context, uint32 Index) const override;
-#if WITH_EDITOR
-	virtual FString GetDefaultValueString(uint32 Index) const override;
-#endif // WITH_EDITORs
 	//~End UPCGSettings interface
 
 	virtual FPCGAttributePropertyInputSelector GetInputSource(uint32 Index) const override;
@@ -49,6 +45,14 @@ protected:
 	virtual FPCGElementPtr CreateElement() const override;
 	//~End UPCGSettings interface
 
+	//~Begin IPCGSettingsDefaultValueProvider interface
+#if WITH_EDITOR
+	virtual FString GetPinInitialDefaultValueString(FName PinLabel) const override;
+#endif // WITH_EDITOR
+	virtual EPCGMetadataTypes GetPinInitialDefaultValueType(FName PinLabel) const override;
+	virtual bool CreateInitialDefaultValueAttribute(FName PinLabel, UPCGMetadata* OutMetadata) const override;
+	//~End IPCGSettingsDefaultValueProvider interface
+
 public:
 	UPROPERTY(BlueprintReadWrite, EditAnywhere, Category = Input, meta = (PCG_Overridable))
 	FPCGAttributePropertyInputSelector InputSource1;
diff --git a/Engine/Plugins/PCG/Source/PCG/Public/Elements/Metadata/PCGMetadataMakeVector.h b/Engine/Plugins/PCG/Source/PCG/Public/Elements/Metadata/PCGMetadataMakeVector.h
--- a/Engine/Plugins/PCG/Source/PCG/Public/Elements/Metadata/PCGMetadataMakeVector.h
+++ b/Engine/Plugins/PCG/Source/PCG/Public/Elements/Metadata/PCGMetadataMakeVector.h
@@ -44,18 +44,17 @@ class UPCGMetadataMakeVectorSettings : public UPCGMetadataSettingsBase
 public:
 	// ~Begin UObject interface
 	virtual void PostLoad() override;
+#if WITH_EDITOR
+	virtual void PostEditChangeProperty(FPropertyChangedEvent& PropertyChangedEvent) override;
+#endif // WITH_EDITOR
 	// ~End UObject interface
 
 #if WITH_EDITOR
 	//~Begin UPCGSettings interface
 	virtual FName GetDefaultNodeName() const override;
 	virtual FText GetDefaultNodeTitle() const override;
+	virtual void ApplyDeprecationBeforeUpdatePins(UPCGNode* InOutNode, TArray<TObjectPtr<UPCGPin>>& InputPins, TArray<TObjectPtr<UPCGPin>>& OutputPins) override;
 #endif
-	virtual bool DoesInputSupportDefaultValue(uint32 Index) const override;
-	virtual UPCGParamData* CreateDefaultValueParam(FPCGContext* Context, uint32 Index) const override;
-#if WITH_EDITOR
-	virtual FString GetDefaultValueString(uint32 Index) const override;
-#endif // WITH_EDITOR
 	//~End UPCGSettings interface
 
 	//~Begin UPCGMetadataSettingsBase interface
@@ -66,9 +65,6 @@ public:
 
 	virtual bool IsSupportedInputType(uint16 TypeId, uint32 InputIndex, bool& bHasSpecialRequirement) const override;
 	virtual uint16 GetOutputType(uint16 InputTypeId) const override;
-
-	/** If the primary pin is required and missing, some operations can be bypassed. Some, like Make Vector, can replace missing pins with default values. */
-	virtual bool IsPrimaryInputPinRequired() const override { return false; }
 	//~End UPCGMetadataSettingsBase interface
 
 protected:
@@ -76,6 +72,13 @@ protected:
 	virtual FPCGElementPtr CreateElement() const override;
 	//~End UPCGSettings interface
 
+	//~Begin IPCGSettingsDefaultValueProvider interface
+#if WITH_EDITOR
+	virtual FString GetPinInitialDefaultValueString(FName PinLabel) const override;
+#endif // WITH_EDITOR
+	virtual EPCGMetadataTypes GetPinInitialDefaultValueType(FName PinLabel) const override;
+	//~End IPCGSettingsDefaultValueProvider interface
+
 public:
 	UPROPERTY(BlueprintReadWrite, EditAnywhere, Category = Input, meta = (PCG_Overridable))
 	FPCGAttributePropertyInputSelector InputSource1;
diff --git a/Engine/Plugins/PCG/Source/PCG/Public/Elements/Metadata/PCGMetadataOpElementBase.h b/Engine/Plugins/PCG/Source/PCG/Public/Elements/Metadata/PCGMetadataOpElementBase.h
--- a/Engine/Plugins/PCG/Source/PCG/Public/Elements/Metadata/PCGMetadataOpElementBase.h
+++ b/Engine/Plugins/PCG/Source/PCG/Public/Elements/Metadata/PCGMetadataOpElementBase.h
@@ -5,13 +5,15 @@
 #include "PCGSettings.h"
 
 #include "Data/PCGPointData.h"
+#include "Elements/PCGTimeSlicedElementBase.h"
 #include "Helpers/PCGAsync.h"
+#include "Helpers/PCGDefaultValueContainer.h"
 #include "Metadata/PCGAttributePropertySelector.h"
+#include "Metadata/PCGDefaultValueInterface.h"
 #include "Metadata/Accessors/IPCGAttributeAccessor.h"
 #include "Metadata/Accessors/PCGAttributeAccessorKeys.h"
 
 #include "Containers/StaticArray.h"
-#include "Elements/PCGTimeSlicedElementBase.h"
 
 #include "PCGMetadataOpElementBase.generated.h"
 
@@ -82,7 +84,7 @@ enum class EPCGMetadataSettingsBaseTypes
  * If the node doesn't provide an output, check the logs to know why it failed.
  */
 UCLASS(BlueprintType, Abstract, ClassGroup = (Procedural))
-class PCG_API UPCGMetadataSettingsBase : public UPCGSettings
+class PCG_API UPCGMetadataSettingsBase : public UPCGSettings, public IPCGSettingsDefaultValueProvider
 {
 	GENERATED_BODY()
 
@@ -95,6 +97,29 @@ public:
 #endif // WITH_EDITOR
 	// ~End UObject interface
 
+	//~Begin IPCGSettingsDefaultValueProvider interface
+	virtual bool DefaultValuesAreEnabled() const override { return true; }
+	virtual bool IsPinDefaultValueEnabled(FName PinLabel) const override;
+	virtual bool IsPinDefaultValueActivated(FName PinLabel) const override;
+	virtual EPCGMetadataTypes GetPinDefaultValueType(FName PinLabel) const override;
+	virtual bool IsPinDefaultValueMetadataTypeValid(FName PinLabel, EPCGMetadataTypes DataType) const override;
+#if WITH_EDITOR
+	virtual void SetPinDefaultValue(FName PinLabel, const FString& DefaultValue, bool bCreateIfNeeded = false) override;
+	virtual void ConvertPinDefaultValueMetadataType(FName PinLabel, EPCGMetadataTypes DataType) override;
+	virtual void SetPinDefaultValueIsActivated(FName PinLabel, bool bIsActivated, bool bDirtySettings = true) override;
+	virtual void ResetDefaultValues() override;
+	virtual FString GetPinInitialDefaultValueString(FName PinLabel) const override { return PCG::Private::MetadataTraits<double>::ZeroValueString(); }
+	virtual FString GetPinDefaultValueAsString(FName PinLabel) const override;
+#endif // WITH_EDITOR
+
+protected:
+	virtual EPCGMetadataTypes GetPinInitialDefaultValueType(FName PinLabel) const override;
+	//~End IPCGSettingsDefaultValueProvider interface
+
+public:
+	/** Creates a Param Data with the inline constant default value properties inserted as metadata. */
+	const UPCGParamData* CreateDefaultValueParamData(FPCGContext* Context, FName PinLabel) const;
+
 	//~Begin UPCGSettings interface
 #if WITH_EDITOR
 	virtual EPCGSettingsType GetType() const override { return EPCGSettingsType::Metadata; }
@@ -107,9 +132,23 @@ public:
 	virtual EPCGDataType GetCurrentPinTypes(const UPCGPin* InPin) const override;
 	virtual bool DoesPinSupportPassThrough(UPCGPin* InPin) const override;
 
+	virtual bool CanCullTaskIfUnwired() const override { return false; }
+	virtual bool IsInputPinRequiredByExecution(const UPCGPin* InPin) const override;
+
+	/** Adds the default values to the Crc for caching. */
+	void AddDefaultValuesToCrc(FArchiveCrc32& Crc32) const;
+
 protected:
 	virtual TArray<FPCGPinProperties> InputPinProperties() const override;
 	virtual TArray<FPCGPinProperties> OutputPinProperties() const override;
+
+#if WITH_EDITOR
+	/**
+	 * Helper to check if the InputSource property should be hidden to the user.
+	 * NumSources should match the number of InputSource's the node has in its current configuration.
+	 */
+	bool CanEditInputSource(const FProperty* InProperty, int32 NumSources = 1) const;
+#endif // WITH_EDITOR
 	//~End UPCGSettings interface
 
 public:
@@ -129,22 +168,18 @@ public:
 	virtual TArray<uint16> GetAllOutputTypes() const { return TArray<uint16>(); }
 
 	/* Can be overriden by child class to support default values on unplugged pins. */
+	UE_DEPRECATED(5.6, "Override 'UPCGSettings::IsPinDefaultValueEnabled' and 'UPCGSettings::IsPinDefaultValueActivated' instead.")
 	virtual bool DoesInputSupportDefaultValue(uint32 Index) const { return false; }
 
 	UE_DEPRECATED(5.5, "Call/Implement version with FPCGContext parameter")
 	virtual UPCGParamData* CreateDefaultValueParam(uint32 Index) const { return nullptr; }
 
+	UE_DEPRECATED(5.6, "Replaced by inline DefaultValue system. Use CreateDefaultValueParamData instead.")
 	virtual UPCGParamData* CreateDefaultValueParam(FPCGContext* Context, uint32 Index) const { return nullptr; }
-#if WITH_EDITOR
-	virtual FString GetDefaultValueString(uint32 Index) const { return FString(); }
-#endif // WITH_EDITOR
 
 	/* Return the current input pin to forward to the output. */
 	uint32 GetInputPinToForward() const;
 
-	/** If the primary input is needed and empty, it will pass through data. */
-	virtual bool IsPrimaryInputPinRequired() const { return true; }
-
 	UPROPERTY(BlueprintReadWrite, EditAnywhere, Category = "Output", meta = (PCG_Overridable))
 	FPCGAttributePropertyOutputSelector OutputTarget;
 
@@ -178,6 +213,11 @@ protected:
 	/* Return the list of all the Input pins */
 	UFUNCTION()
 	TArray<FName> GetOutputDataFromPinOptions() const;
+
+private:
+	/** Stores the default values for the pins to be used as inline constants. */
+	UPROPERTY()
+	FPCGDefaultValueContainer DefaultValues;
 };
 
 namespace PCGMetadataOps
@@ -199,6 +239,8 @@ namespace PCGMetadataOps
 		TArray<TUniquePtr<const IPCGAttributeAccessor>> InputAccessors;
 		TArray<TUniquePtr<IPCGAttributeAccessor>> OutputAccessors;
 
+		TArray<bool, TFixedAllocator<UPCGMetadataSettingsBase::MaxNumberOfInputs>> DefaultValueOverriddenPins;
+
 		template <int32 NbInputs, int32 NbOutputs>
 		void Validate();
 	};
@@ -211,6 +253,7 @@ protected:
 	virtual bool ExecuteInternal(FPCGContext* Context) const override;
 	virtual EPCGElementExecutionLoopMode ExecutionLoopMode(const UPCGSettings* Settings) const override { return EPCGElementExecutionLoopMode::SinglePrimaryPin; }
 	virtual bool SupportsBasePointDataInputs(FPCGContext* InContext) const override { return true; }
+	virtual void GetDependenciesCrc(const FPCGDataCollection& InInput, const UPCGSettings* InSettings, UPCGComponent* InComponent, FPCGCrc& OutCrc) const override;
 
 	virtual bool DoOperation(PCGMetadataOps::FOperationData& InOperationData) const = 0;
 
@@ -510,4 +553,4 @@ template <typename InType1, typename InType2, typename InType3, typename InType4
 inline bool FPCGMetadataElementBase::DoQuaternaryOp(PCGMetadataOps::FOperationData& InOperationData, Callbacks&& ...InCallbacks) const
 {
 	return DoNAryOp<InType1, InType2, InType3, InType4>(InOperationData, ForwardAsTuple(std::forward<Callbacks>(InCallbacks)...));
-}
+}
\ No newline at end of file
diff --git a/Engine/Plugins/PCG/Source/PCG/Public/Elements/Metadata/PCGMetadataRotatorOpElement.h b/Engine/Plugins/PCG/Source/PCG/Public/Elements/Metadata/PCGMetadataRotatorOpElement.h
--- a/Engine/Plugins/PCG/Source/PCG/Public/Elements/Metadata/PCGMetadataRotatorOpElement.h
+++ b/Engine/Plugins/PCG/Source/PCG/Public/Elements/Metadata/PCGMetadataRotatorOpElement.h
@@ -64,6 +64,11 @@ protected:
 	virtual FPCGElementPtr CreateElement() const override;
 	//~End UPCGSettings interface
 
+	//~Begin IPCGSettingsDefaultValueProvider interface
+	virtual bool IsPinDefaultValueMetadataTypeValid(FName PinLabel, EPCGMetadataTypes DataType) const override;
+	virtual EPCGMetadataTypes GetPinInitialDefaultValueType(FName PinLabel) const override;
+	//~End IPCGSettingsDefaultValueProvider interface
+
 public:
 	UPROPERTY(BlueprintReadWrite, EditAnywhere, Category = Settings)
 	EPCGMetadataRotatorOperation Operation = EPCGMetadataRotatorOperation::Combine;
diff --git a/Engine/Plugins/PCG/Source/PCG/Public/Elements/Metadata/PCGMetadataStringOpElement.h b/Engine/Plugins/PCG/Source/PCG/Public/Elements/Metadata/PCGMetadataStringOpElement.h
--- a/Engine/Plugins/PCG/Source/PCG/Public/Elements/Metadata/PCGMetadataStringOpElement.h
+++ b/Engine/Plugins/PCG/Source/PCG/Public/Elements/Metadata/PCGMetadataStringOpElement.h
@@ -52,6 +52,13 @@ public:
 	virtual uint16 GetOutputType(uint16 InputTypeId) const override;
 	//~End UPCGMetadataSettingsBase interface
 
+	//~Begin IPCGSettingsDefaultValueProvider interface
+#if WITH_EDITOR
+	virtual FString GetPinInitialDefaultValueString(FName PinLabel) const override { return PCG::Private::MetadataTraits<FString>::ZeroValueString(); }
+#endif // WITH_EDITOR
+	virtual EPCGMetadataTypes GetPinInitialDefaultValueType(FName PinLabel) const override { return EPCGMetadataTypes::String; }
+	//~End IPCGSettingsDefaultValueProvider interface
+
 public:
 	UPROPERTY(BlueprintReadWrite, EditAnywhere, Category = Settings)
 	EPCGMetadataStringOperation Operation = EPCGMetadataStringOperation::Append;
diff --git a/Engine/Plugins/PCG/Source/PCG/Public/Elements/Metadata/PCGMetadataTransformOpElement.h b/Engine/Plugins/PCG/Source/PCG/Public/Elements/Metadata/PCGMetadataTransformOpElement.h
--- a/Engine/Plugins/PCG/Source/PCG/Public/Elements/Metadata/PCGMetadataTransformOpElement.h
+++ b/Engine/Plugins/PCG/Source/PCG/Public/Elements/Metadata/PCGMetadataTransformOpElement.h
@@ -58,6 +58,10 @@ public:
 	virtual bool IsSupportedInputType(uint16 TypeId, uint32 InputIndex, bool& bHasSpecialRequirement) const override;
 	//~End UPCGMetadataSettingsBase interface
 
+	//~Begin IPCGSettingsDefaultValueProvider interface
+	virtual bool IsPinDefaultValueEnabled(FName PinLabel) const override;
+	virtual EPCGMetadataTypes GetPinDefaultValueType(FName PinLabel) const override;
+	//~End IPCGSettingsDefaultValueProvider interface
 protected:
 	//~Begin UPCGSettings interface
 #if WITH_EDITOR
diff --git a/Engine/Plugins/PCG/Source/PCG/Public/Elements/Metadata/PCGMetadataVectorOpElement.h b/Engine/Plugins/PCG/Source/PCG/Public/Elements/Metadata/PCGMetadataVectorOpElement.h
--- a/Engine/Plugins/PCG/Source/PCG/Public/Elements/Metadata/PCGMetadataVectorOpElement.h
+++ b/Engine/Plugins/PCG/Source/PCG/Public/Elements/Metadata/PCGMetadataVectorOpElement.h
@@ -63,6 +63,11 @@ protected:
 	virtual FPCGElementPtr CreateElement() const override;
 	//~End UPCGSettings interface
 
+	//~Begin IPCGSettingsDefaultValueProvider interface
+	virtual bool IsPinDefaultValueEnabled(FName PinLabel) const override;
+	virtual EPCGMetadataTypes GetPinInitialDefaultValueType(FName PinLabel) const override { return EPCGMetadataTypes::Vector; }
+	//~End IPCGSettingsDefaultValueProvider interface
+
 public:
 	UPROPERTY(BlueprintReadWrite, EditAnywhere, Category = Settings)
 	EPCGMetadataVectorOperation Operation = EPCGMetadataVectorOperation::Cross;
diff --git a/Engine/Plugins/PCG/Source/PCG/Public/Helpers/PCGDefaultValueContainer.h b/Engine/Plugins/PCG/Source/PCG/Public/Helpers/PCGDefaultValueContainer.h
new file mode 100644
--- /dev/null
+++ b/Engine/Plugins/PCG/Source/PCG/Public/Helpers/PCGDefaultValueContainer.h
@@ -0,0 +1,63 @@
+// Copyright Epic Games, Inc. All Rights Reserved.
+
+#pragma once
+
+#include "Serialization/ArchiveCrc32.h"
+#include "StructUtils/PropertyBag.h"
+
+#include "PCGDefaultValueContainer.generated.h"
+
+enum class EPCGMetadataTypes : uint8;
+struct FPCGContext;
+class UPCGParamData;
+
+/** A struct used to store a default value locally for use with inline constant values. */
+USTRUCT(BlueprintType)
+struct FPCGDefaultValueContainer
+{
+	GENERATED_BODY()
+
+	/** Creates a new FProperty with the specified type. Will overwrite if it already exists. Returns nullptr if the type is invalid. */
+	const FProperty* CreateNewProperty(FName PropertyName, EPCGMetadataTypes Type);
+
+	/** Finds an FProperty in the property bag. Returns nullptr if it does not exist. */
+	const FProperty* FindProperty(FName PropertyName) const;
+
+	/** Find the current metadata subtype of the specified FProperty. */
+	EPCGMetadataTypes GetCurrentPropertyType(FName PropertyName) const;
+
+	/** Get the current value of the property as a string. */
+	FString GetPropertyValueAsString(FName PropertyName) const;
+
+	/** Serialize the properties into a Crc32 to be used for caching purposes. */
+	void SerializeCrc(FArchiveCrc32& Crc32) { PropertyBag.Serialize(Crc32); }
+
+	/** Helper that creates a Param Data with the typed property as an attribute attached to the metadata. */
+	const UPCGParamData* CreateParamData(FPCGContext* Context, FName PropertyName) const;
+
+	/** The property exists and is set to active. */
+	bool IsPropertyActivated(FName PropertyName) const;
+
+#if WITH_EDITOR
+	/** Change the metadata subtype for the specified FProperty. Will overwrite. */
+	const FProperty* ConvertPropertyType(FName PropertyName, EPCGMetadataTypes Type);
+
+	/** Sets the FProperty's value directly with a string input. Returns true if successful. */
+	bool SetPropertyValueFromString(FName PropertyName, const FString& ValueString);
+
+	/** Sets the property as currently active or inactive. Returns true if a change was made. */
+	bool SetPropertyActivated(FName PropertyName, bool bIsActivated);
+
+	/** Reset the container to empty. */
+	void Reset();
+#endif // WITH_EDITOR
+
+private:
+	/** A mapping of which properties are currently active. */
+	UPROPERTY()
+	TSet<FName> ActivatedProperties;
+
+	/** The handling and storage of dynamic properties. */
+	UPROPERTY()
+	FInstancedPropertyBag PropertyBag;
+};
diff --git a/Engine/Plugins/PCG/Source/PCG/Public/Helpers/PCGMetadataHelpers.h b/Engine/Plugins/PCG/Source/PCG/Public/Helpers/PCGMetadataHelpers.h
--- a/Engine/Plugins/PCG/Source/PCG/Public/Helpers/PCGMetadataHelpers.h
+++ b/Engine/Plugins/PCG/Source/PCG/Public/Helpers/PCGMetadataHelpers.h
@@ -69,4 +69,35 @@ namespace PCGMetadataHelpers
 	PCG_API void SetPointAttributes(FPCGMetadataDomain* InOutMetadata, const TArrayView<const FPCGPoint>& InPoints, const FPCGMetadataDomain* InMetadata, const TArrayView<FPCGPoint>& OutPoints, FPCGContext* OptionalContext = nullptr);
 
 	PCG_API void ComputePointWeightedAttribute(FPCGMetadataDomain* InOutMetadata, FPCGPoint& OutPoint, const TArrayView<TPair<const FPCGPoint*, float>>& InWeightedPoints, const FPCGMetadataDomain* InMetadata);
+
+	// [EXPERIMENTAL] - This function may be renamed or deprecated in the future once it is decoupled with default values.
+	/** Helper that checks that Type is supported by default values. */
+	inline bool MetadataTypeSupportsDefaultValues(const EPCGMetadataTypes Type)
+	{
+		switch (Type)
+		{
+			case EPCGMetadataTypes::Double:    // fall-through
+			case EPCGMetadataTypes::Integer32: // fall-through
+			case EPCGMetadataTypes::Integer64: // fall-through
+			case EPCGMetadataTypes::Vector:    // fall-through
+			case EPCGMetadataTypes::Vector2:   // fall-through
+			case EPCGMetadataTypes::Vector4:   // fall-through
+			case EPCGMetadataTypes::String:    // fall-through
+			case EPCGMetadataTypes::Name:      // fall-through
+			case EPCGMetadataTypes::Boolean:   // fall-through
+			case EPCGMetadataTypes::Rotator:
+				return true;
+			// @todo_pcg: Enable the rest once they're supported in the UI
+			case EPCGMetadataTypes::Quaternion:     // fall-through
+			case EPCGMetadataTypes::Transform:      // fall-through
+			case EPCGMetadataTypes::SoftObjectPath: // fall-through
+			case EPCGMetadataTypes::SoftClassPath:  // fall-through
+			// This is automatically converted in the accessor to a Double
+			case EPCGMetadataTypes::Float: // fall-through
+			// Anything else is not valid.
+			case EPCGMetadataTypes::Unknown: // fall-through
+			default:
+				return false;
+		}
+	}
 }
diff --git a/Engine/Plugins/PCG/Source/PCG/Public/Helpers/PCGPropertyHelpers.h b/Engine/Plugins/PCG/Source/PCG/Public/Helpers/PCGPropertyHelpers.h
--- a/Engine/Plugins/PCG/Source/PCG/Public/Helpers/PCGPropertyHelpers.h
+++ b/Engine/Plugins/PCG/Source/PCG/Public/Helpers/PCGPropertyHelpers.h
@@ -52,7 +52,7 @@ namespace PCGPropertyHelpers
 	 * @param Type - The PCG metadata type.
 	 * @returns The property bag property type.
 	 */
-	PCG_API FPropertyBagPropertyDesc GetPropertyBagDescWithMetadataType(FName InPropertyName, EPCGMetadataTypes Type);
+	PCG_API FPropertyBagPropertyDesc CreatePropertyBagDescWithMetadataType(FName InPropertyName, EPCGMetadataTypes Type);
 
 	struct FExtractorParameters
 	{
@@ -180,7 +180,7 @@ inline bool PCGPropertyHelpers::SetPropertyValueFromCallback(ObjectType* InObjec
 	{
 		using Type = decltype(Dummy);
 		Type Value{};
-		FPCGAttributeAccessorKeysSingleObjectPtr<ObjectType> Key(InObject);
+		FPCGAttributeAccessorKeysSingleObjectPtr<std::remove_const_t<ObjectType>> Key(InObject);
 
 		if (InFunc(Value))
 		{
diff --git a/Engine/Plugins/PCG/Source/PCG/Public/Metadata/PCGDefaultValueInterface.h b/Engine/Plugins/PCG/Source/PCG/Public/Metadata/PCGDefaultValueInterface.h
new file mode 100644
--- /dev/null
+++ b/Engine/Plugins/PCG/Source/PCG/Public/Metadata/PCGDefaultValueInterface.h
@@ -0,0 +1,72 @@
+// Copyright Epic Games, Inc. All Rights Reserved.
+
+#pragma once
+
+#include "Metadata/PCGMetadata.h"
+#include "Metadata/PCGMetadataAttributeTraits.h"
+
+#include "UObject/Interface.h"
+
+#include "PCGDefaultValueInterface.generated.h"
+
+UINTERFACE(MinimalAPI)
+class UPCGSettingsDefaultValueProvider : public UInterface
+{
+	GENERATED_BODY()
+};
+
+class IPCGSettingsDefaultValueProvider
+{
+	GENERATED_BODY()
+
+public:
+	/** One or more pins on this node has a 'default value' and can be adjusted via an inline constant. */
+	virtual bool DefaultValuesAreEnabled() const { return false; }
+
+	/** The specified pin can accomodate 'default value' inline constants. */
+	virtual bool IsPinDefaultValueEnabled(FName PinLabel) const { return false; }
+
+	/** The specified pin has a 'default value' currently activated. */
+	virtual bool IsPinDefaultValueActivated(FName PinLabel) const { return false; }
+
+	/** Gets the current 'default value' type, if supported, for the pin. */
+	virtual EPCGMetadataTypes GetPinDefaultValueType(FName PinLabel) const { return EPCGMetadataTypes::Unknown; }
+
+	/** Whether the pin supports the provided metadata type. */
+	virtual bool IsPinDefaultValueMetadataTypeValid(FName PinLabel, EPCGMetadataTypes DataType) const { return false; }
+
+	/** Add an attribute to a given metadata for the initial default value. In most cases, this will be the Zero Value. */
+	virtual bool CreateInitialDefaultValueAttribute(FName PinLabel, UPCGMetadata* OutMetadata) const
+	{
+		EPCGMetadataTypes AttributeType = GetPinDefaultValueType(PinLabel);
+		return OutMetadata && PCGMetadataAttribute::CallbackWithRightType(static_cast<uint16>(AttributeType), [OutMetadata]<typename T>(T) -> bool
+		{
+			T Value = PCG::Private::MetadataTraits<T>::ZeroValue();
+			return nullptr != OutMetadata->CreateAttribute(NAME_None, Value, /*bAllowsInterpolation=*/true, /*bOverrideParent=*/false);
+		});
+	}
+
+#if WITH_EDITOR
+	/** Resets all default values to their 'reset' value and deactivates them. */
+	virtual void ResetDefaultValues() {}
+
+	/** Sets the pin's default value string directly. */
+	virtual void SetPinDefaultValue(FName PinLabel, const FString& DefaultValue, bool bCreateIfNeeded = false) {}
+
+	/** Attempts a metadata type conversion of the pin's default value inline constant. */
+	virtual void ConvertPinDefaultValueMetadataType(FName PinLabel, EPCGMetadataTypes DataType) {}
+
+	/** Sets the default value to active. Must be overridden by the subclass. */
+	virtual void SetPinDefaultValueIsActivated(FName PinLabel, bool bIsActivated, bool bDirtySettings = true) PURE_VIRTUAL(IPCGSettingsDefaultValueProvider::SetPinDefaultValueIsActivated, return;);
+
+	/** Gets the 'default value', if supported, for the pin. */
+	virtual FString GetPinDefaultValueAsString(FName PinLabel) const { return FString(); }
+
+	/** For the initial 'default value' of the pin. */
+	virtual FString GetPinInitialDefaultValueString(FName PinLabel) const { return FString(); }
+#endif // WITH_EDITOR
+
+protected:
+	/** For the initial 'default value' type of the pin. */
+	virtual EPCGMetadataTypes GetPinInitialDefaultValueType(FName PinLabel) const { return EPCGMetadataTypes::Unknown; }
+};
\ No newline at end of file
diff --git a/Engine/Plugins/PCG/Source/PCG/Public/Metadata/PCGMetadataAttributeTraits.h b/Engine/Plugins/PCG/Source/PCG/Public/Metadata/PCGMetadataAttributeTraits.h
--- a/Engine/Plugins/PCG/Source/PCG/Public/Metadata/PCGMetadataAttributeTraits.h
+++ b/Engine/Plugins/PCG/Source/PCG/Public/Metadata/PCGMetadataAttributeTraits.h
@@ -349,6 +349,8 @@ namespace PCG
 				bOutSuccess = true;
 				return T(InString);
 			}
+
+			static FString ZeroValueString();
 		};
 
 		template<typename T>
@@ -358,6 +360,8 @@ namespace PCG
 			{
 				return LexToString(A);
 			}
+
+			static FString ZeroValueString();
 		};
 
 		template<typename T>
@@ -1321,5 +1325,24 @@ namespace PCG
 				}
 			}
 		}
+
+		template <typename T>
+		inline FString DefaultStringTraits<T>::ZeroValueString()
+		{
+			if constexpr (std::is_same_v<T, FString>)
+			{
+				return FString();
+			}
+			else
+			{
+				return MetadataTraits<T>::ZeroValue().ToString();
+			}
+		}
+
+		template <typename T>
+		inline FString LexToStringTraits<T>::ZeroValueString()
+		{
+			return LexToString(MetadataTraits<T>::ZeroValue());
+		}
 	}
 }
diff --git a/Engine/Plugins/PCG/Source/PCG/Public/PCGSettings.h b/Engine/Plugins/PCG/Source/PCG/Public/PCGSettings.h
--- a/Engine/Plugins/PCG/Source/PCG/Public/PCGSettings.h
+++ b/Engine/Plugins/PCG/Source/PCG/Public/PCGSettings.h
@@ -372,9 +372,7 @@ public:
 
 	UPROPERTY(EditAnywhere, BlueprintReadWrite, Category = Debug, AdvancedDisplay, meta = (DisplayName="Determinism", NoResetToDefault))
 	FPCGDeterminismSettings DeterminismSettings;
-#endif
 
-#if WITH_EDITORONLY_DATA
 	UPROPERTY(BlueprintReadWrite, EditAnywhere, Category = AssetInfo, AssetRegistrySearchable)
 	bool bExposeToLibrary = true;
 
diff --git a/Engine/Source/Runtime/Core/Public/UObject/FortniteMainBranchObjectVersions.inl b/Engine/Source/Runtime/Core/Public/UObject/FortniteMainBranchObjectVersions.inl
--- a/Engine/Source/Runtime/Core/Public/UObject/FortniteMainBranchObjectVersions.inl
+++ b/Engine/Source/Runtime/Core/Public/UObject/FortniteMainBranchObjectVersions.inl
@@ -571,5 +571,8 @@ FORTNITE_MAIN_VERSION(SceneCaptureDefaultSettings, 185)
 // Add Cloth Asset Base class serialization
 FORTNITE_MAIN_VERSION(AddClothAssetBase, 186)
 
+// Add inline constant default values to the PCG graph nodes.
+FORTNITE_MAIN_VERSION(PCGInlineConstantDefaultValues, 187)
+
 // -----<new versions can be added above this line>-------------------------------------------------
 #undef FORTNITE_MAIN_VERSION
\ No newline at end of file
