[PCG] 3/3 Inline Pin Constants (Default Values) Frontend

+ New Inline Pin Widgets for editing values directly on node pins

Still left TODO:
* SoftObjectPath and SoftClassPath widgets still require some polish for stability
* Quaternion and Transform still need discussion and/or widgets (they would take up a lot of space, or could be broken into multiple widgets/pins, etc)
* Iterate on UX feedback

diff --git a/Engine/Plugins/PCG/Source/PCGEditor/Private/Nodes/PCGEditorGraphNodeBase.cpp b/Engine/Plugins/PCG/Source/PCGEditor/Private/Nodes/PCGEditorGraphNodeBase.cpp
--- a/Engine/Plugins/PCG/Source/PCGEditor/Private/Nodes/PCGEditorGraphNodeBase.cpp
+++ b/Engine/Plugins/PCG/Source/PCGEditor/Private/Nodes/PCGEditorGraphNodeBase.cpp
@@ -12,6 +12,8 @@
 #include "PCGWorldActor.h"
 #include "Elements/PCGHiGenGridSize.h"
 #include "Elements/PCGReroute.h"
+#include "Helpers/PCGMetadataHelpers.h"
+#include "Metadata/PCGDefaultValueInterface.h"
 
 #include "PCGEditor.h"
 #include "PCGEditorCommands.h"
@@ -43,6 +45,11 @@ namespace PCGEditorGraphSwitches
 
 namespace PCGEditorGraphNodeBase
 {
+	IPCGSettingsDefaultValueProvider* GetDefaultValueInterface(UPCGSettings* InSettings)
+	{
+		return (InSettings && InSettings->Implements<UPCGSettingsDefaultValueProvider>()) ? CastChecked<IPCGSettingsDefaultValueProvider>(InSettings) : nullptr;
+	}
+
 	/** Whether this node was culled during graph compilation or during graph execution. */
 	bool ShouldDisplayAsActive(const UPCGEditorGraphNodeBase* InNode, const UPCGComponent* InComponentBeingDebugged, const FPCGStack* InStackBeingInspected)
 	{
@@ -193,24 +200,71 @@ void UPCGEditorGraphNodeBase::GetNodeContextMenuActions(UToolMenu* Menu, class U
 	}
 
 	// Local pin special actions should come first.
+	if (Context->Pin)
 	{
-		FToolMenuSection& Section = Menu->AddSection("EdGraphSchemaPinActions", LOCTEXT("PinActionsMenuHeader", "Pin Actions"));
-		if (Context->Pin && CanUserAddRemoveDynamicInputPins())
+		const bool bDynamicPinSubMenu = CanUserAddRemoveDynamicInputPins();
+		const bool bDefaultValueSubMenu = !Context->Pin->HasAnyConnections() && IsSettingsDefaultValuesEnabled();
+
+		if (bDynamicPinSubMenu || bDefaultValueSubMenu)
 		{
-			Section.AddMenuEntry(FPCGEditorCommands::Get().AddSourcePin);
-			Section.AddMenuEntry("RemovePin",
-				LOCTEXT("RemovePin", "Remove Source Pin"),
-				LOCTEXT("RemovePinTooltip", "Remove this source pin from the current node"),
-				FSlateIcon(),
-				FUIAction(
-					FExecuteAction::CreateLambda([Pin = Context->Pin, this]
+			FToolMenuSection& Section = Menu->AddSection("EdGraphSchemaPinActions", LOCTEXT("PinActionsMenuHeader", "Pin Actions"));
+
+			if (bDynamicPinSubMenu)
+			{
+				Section.AddMenuEntry(FPCGEditorCommands::Get().AddSourcePin);
+				Section.AddMenuEntry("RemovePin",
+					LOCTEXT("RemovePin", "Remove Source Pin"),
+					LOCTEXT("RemovePinTooltip", "Remove this source pin from the current node"),
+					FSlateIcon(),
+					FUIAction(
+						FExecuteAction::CreateLambda([Pin = Context->Pin, this]
+						{
+							const_cast<UPCGEditorGraphNodeBase*>(this)->OnUserRemoveDynamicInputPin(const_cast<UEdGraphPin*>(Pin));
+						}),
+						FCanExecuteAction::CreateLambda([Pin = Context->Pin, this]
+						{
+							return this->CanUserRemoveDynamicInputPin(const_cast<UEdGraphPin*>(Pin));
+						})));
+			}
+
+			// Default value inline constants
+			if (bDefaultValueSubMenu)
+			{
+				const FName PinLabel = Context->Pin->PinName;
+				const bool bIsActive = IsPinDefaultValueActivated(PinLabel);
+
+				Section.AddMenuEntry(TEXT("ActivateInlineConstantPin"),
+					bIsActive
+						? LOCTEXT("DeactivateInlineConstantPin", "Deactivate Inline Constant")
+						: LOCTEXT("ActivateInlineConstantPin", "Activate Inline Constant"),
+					FText(),
+					FSlateIcon(),
+					FUIAction(
+							FExecuteAction::CreateLambda([this, PinLabel, bIsActive] { OnUserSetPinDefaultValueActivated(PinLabel, !bIsActive); }),
+							FCanExecuteAction::CreateLambda([this, Pin = Context->Pin] { return IsPinDefaultValueEnabled(Pin->PinName) && !Pin->HasAnyConnections(); })));
+
+				const UEnum* EnumPtr = StaticEnum<EPCGMetadataTypes>();
+				for (int32 I = 0; I < EnumPtr->NumEnums() - 1; ++I)
+				{
+					if (!EnumPtr->GetMetaData(TEXT("Hidden"), I).IsEmpty())
 					{
-						const_cast<UPCGEditorGraphNodeBase*>(this)->OnUserRemoveDynamicInputPin(const_cast<UEdGraphPin*>(Pin));
-					}),
-					FCanExecuteAction::CreateLambda([Pin = Context->Pin, this]
+						continue;
+					}
+
+					EPCGMetadataTypes DataType = static_cast<EPCGMetadataTypes>(EnumPtr->GetValueByIndex(I));
+					if (bIsActive && CanConvertToDefaultValueMetadataType(PinLabel, DataType))
 					{
-						return const_cast<UPCGEditorGraphNodeBase*>(this)->CanUserRemoveDynamicInputPin(const_cast<UEdGraphPin*>(Pin));
-					})));
+						const FName EntryName = FName(FString(TEXT("ConvertPinType")) + EnumPtr->GetAuthoredNameStringByIndex(I));
+						Section.AddMenuEntry(EntryName,
+							FText::Format(ConvertToFormat, EnumPtr->GetDisplayNameTextByIndex(I)),
+							LOCTEXT("ConvertPinTypeTooltip", "Convert this pin's inline constant to a different type."),
+							FSlateIcon(),
+							FUIAction(
+									FExecuteAction::CreateLambda([this, PinLabel, DataType] { ConvertPinDefaultValueMetadataType(PinLabel, DataType); }),
+									FCanExecuteAction::CreateLambda([this, Pin = Context->Pin] { return IsPinDefaultValueEnabled(Pin->PinName) && !Pin->HasAnyConnections(); })));
+					}
+				}
+			}
 		}
 	}
 
@@ -744,6 +798,65 @@ bool UPCGEditorGraphNodeBase::CanUserAddRemoveDynamicInputPins() const
 	return Settings && Settings->IsA<UPCGSettingsWithDynamicInputs>();
 }
 
+bool UPCGEditorGraphNodeBase::IsSettingsDefaultValuesEnabled() const
+{
+	const IPCGSettingsDefaultValueProvider* DefaultValueInterface = GetDefaultValueInterface();
+	return DefaultValueInterface && DefaultValueInterface->DefaultValuesAreEnabled();
+}
+
+bool UPCGEditorGraphNodeBase::IsPinDefaultValueEnabled(const FName PinLabel) const
+{
+	const IPCGSettingsDefaultValueProvider* DefaultValueInterface = GetDefaultValueInterface();
+	return DefaultValueInterface && DefaultValueInterface->IsPinDefaultValueEnabled(PinLabel);
+}
+
+bool UPCGEditorGraphNodeBase::IsPinDefaultValueActivated(const FName PinLabel) const
+{
+	const IPCGSettingsDefaultValueProvider* DefaultValueInterface = GetDefaultValueInterface();
+	return DefaultValueInterface && DefaultValueInterface->IsPinDefaultValueActivated(PinLabel);
+}
+
+void UPCGEditorGraphNodeBase::OnUserSetPinDefaultValueActivated(const FName PinLabel, const bool bIsActivated) const
+{
+	UPCGSettings* Settings = PCGNode ? PCGNode->GetSettings() : nullptr;
+	if (IPCGSettingsDefaultValueProvider* DefaultValueInterface = GetDefaultValueInterface())
+	{
+		FText TransactionDescription = bIsActivated
+			? LOCTEXT("PCGEditorSetPinInlineConstantDeactivated", "Deactivate Pin Inline Constant")
+			: LOCTEXT("PCGEditorSetPinInlineConstantActivated", "Activate Pin Inline Constant");
+		FScopedTransaction Transaction(*FPCGEditorCommon::ContextIdentifier, std::move(TransactionDescription), Settings);
+		DefaultValueInterface->SetPinDefaultValueIsActivated(PinLabel, bIsActivated);
+	}
+}
+
+void UPCGEditorGraphNodeBase::ConvertPinDefaultValueMetadataType(const FName PinLabel, const EPCGMetadataTypes DataType) const
+{
+	UPCGSettings* Settings = PCGNode ? PCGNode->GetSettings() : nullptr;
+	if (IPCGSettingsDefaultValueProvider* Interface = GetDefaultValueInterface())
+	{
+		FScopedTransaction Transaction(*FPCGEditorCommon::ContextIdentifier, LOCTEXT("PCGEditorConvertPinInlineConstantMetadataType", "Convert Pin Inline Constant Type"), Settings);
+		Interface->ConvertPinDefaultValueMetadataType(PinLabel, DataType);
+	}
+}
+
+bool UPCGEditorGraphNodeBase::CanConvertToDefaultValueMetadataType(const FName PinLabel, const EPCGMetadataTypes DataType) const
+{
+	// Early out for invalid types and if the pin is already this type.
+	if (!PCGMetadataHelpers::MetadataTypeSupportsDefaultValues(DataType) || DataType == GetPinDefaultValueType(PinLabel))
+	{
+		return false;
+	}
+
+	const IPCGSettingsDefaultValueProvider* Interface = GetDefaultValueInterface();
+	return Interface && Interface->IsPinDefaultValueMetadataTypeValid(PinLabel, DataType);
+}
+
+EPCGMetadataTypes UPCGEditorGraphNodeBase::GetPinDefaultValueType(const FName PinLabel) const
+{
+	const IPCGSettingsDefaultValueProvider* Interface = GetDefaultValueInterface();
+	return Interface ? Interface->GetPinDefaultValueType(PinLabel) : EPCGMetadataTypes::Unknown;
+}
+
 bool UPCGEditorGraphNodeBase::HasFlippedTitleLines() const
 {
 	return PCGNode ? PCGNode->HasFlippedTitleLines() : false;
@@ -864,6 +977,16 @@ bool UPCGEditorGraphNodeBase::IsCompatible(const UPCGPin* InputPin, const UPCGPi
 	return true;
 }
 
+UPCGSettings* UPCGEditorGraphNodeBase::GetSettings() const
+{
+	return PCGNode ? PCGNode->GetSettings() : nullptr;
+}
+
+IPCGSettingsDefaultValueProvider* UPCGEditorGraphNodeBase::GetDefaultValueInterface() const
+{
+	return PCGEditorGraphNodeBase::GetDefaultValueInterface(GetSettings());
+}
+
 void UPCGEditorGraphNodeBase::OnPickColor()
 {
 	FColorPickerArgs PickerArgs;
@@ -998,9 +1121,6 @@ FEdGraphPinType UPCGEditorGraphNodeBase::GetPinType(const UPCGPin* InPin)
 {
 	FEdGraphPinType EdPinType;
 	EdPinType.ResetToDefaults();
-	EdPinType.PinCategory = NAME_None;
-	EdPinType.PinSubCategory = NAME_None;
-	EdPinType.ContainerType = EPinContainerType::None;
 
 	check(InPin);
 	const EPCGDataType PinType = InPin->GetCurrentTypes();
@@ -1204,6 +1324,52 @@ void UPCGEditorGraphNodeBase::OnCommentBubbleToggled(bool bInCommentBubbleVisibl
 	}
 }
 
+void UPCGEditorGraphNodeBase::PinDefaultValueChanged(UEdGraphPin* Pin)
+{
+	check(PCGNode);
+
+	if (Pin && !Pin->IsPendingKill())
+	{
+		Super::PinDefaultValueChanged(Pin);
+
+		if (IPCGSettingsDefaultValueProvider* DefaultValueInterface = GetDefaultValueInterface())
+		{
+			if (ensure(DefaultValueInterface->IsPinDefaultValueActivated(Pin->PinName)))
+			{
+				DefaultValueInterface->SetPinDefaultValue(Pin->PinName, Pin->DefaultValue, /*bCreateIfNeeded=*/true);
+			}
+		}
+	}
+}
+
+void UPCGEditorGraphNodeBase::PinConnectionListChanged(UEdGraphPin* Pin)
+{
+	check(PCGNode);
+
+	Super::PinConnectionListChanged(Pin);
+
+	// One-time UX functionality, where when a user connects a node with default values, it will activate them for all other pins
+	if (!bHasEverBeenConnected && Pin && !Pin->IsPendingKill())
+	{
+		if (IPCGSettingsDefaultValueProvider* Interface = GetDefaultValueInterface())
+		{
+			if (Pin->Direction == EGPD_Input && Interface->DefaultValuesAreEnabled())
+			{
+				for (const UPCGPin* InputPin : PCGNode->GetInputPins())
+				{
+					if (InputPin && InputPin->Properties.Label != Pin->PinName && Interface->IsPinDefaultValueEnabled(InputPin->Properties.Label))
+					{
+						// Connecting a pin will dirty anyway, so no need to dirty again
+						Interface->SetPinDefaultValueIsActivated(InputPin->Properties.Label, /*bIsActive=*/true);
+					}
+				}
+
+				bHasEverBeenConnected = true;
+			}
+		}
+	}
+}
+
 void UPCGEditorGraphNodeBase::OnUserAddDynamicInputPin()
 {
 	check(PCGNode);
diff --git a/Engine/Plugins/PCG/Source/PCGEditor/Private/Nodes/PCGEditorGraphNodeBase.h b/Engine/Plugins/PCG/Source/PCGEditor/Private/Nodes/PCGEditorGraphNodeBase.h
--- a/Engine/Plugins/PCG/Source/PCGEditor/Private/Nodes/PCGEditorGraphNodeBase.h
+++ b/Engine/Plugins/PCG/Source/PCGEditor/Private/Nodes/PCGEditorGraphNodeBase.h
@@ -10,15 +10,19 @@
 
 #include "PCGEditorGraphNodeBase.generated.h"
 
-struct FPCGPreconfiguredInfo;
 enum class EPCGChangeType : uint32;
+enum class EPCGMetadataTypes : uint8;
+
+struct FPCGPreconfiguredInfo;
+struct FPCGStack;
 
+class IPCGSettingsDefaultValueProvider;
 class UPCGComponent;
 class UPCGGraph;
 class UPCGNode;
 class UPCGPin;
+class UPCGSettings;
 class UToolMenu;
-struct FPCGStack;
 
 UCLASS()
 class UPCGEditorGraphNodeBase : public UEdGraphNode
@@ -48,6 +52,8 @@ public:
 	virtual void OnUpdateCommentText(const FString& NewComment) override;
 	virtual void OnCommentBubbleToggled(bool bInCommentBubbleVisible) override;
 	virtual bool GetCanRenameNode() const override { return bCanRenameNode; }
+	virtual void PinDefaultValueChanged(UEdGraphPin* Pin) override;
+	virtual void PinConnectionListChanged(UEdGraphPin* Pin) override;
 	// ~End UEdGraphNode interface
 
 	void OnUserAddDynamicInputPin();
@@ -88,6 +94,27 @@ public:
 
 	bool CanUserAddRemoveDynamicInputPins() const;
 
+	/** The settings support default value inline constants. */
+	bool IsSettingsDefaultValuesEnabled() const;
+
+	/** The pin has default value inline constants enabled. */
+	bool IsPinDefaultValueEnabled(FName PinLabel) const;
+
+	/** The pin has a default value and is currently active. */
+	bool IsPinDefaultValueActivated(FName PinLabel) const;
+
+	/** User activated or deactivated the inline constant value. */
+	void OnUserSetPinDefaultValueActivated(FName PinLabel, bool bIsActivated) const;
+
+	/** User is converting the default value type on the pin. */
+	void ConvertPinDefaultValueMetadataType(FName PinLabel, EPCGMetadataTypes DataType) const;
+
+	/** The target type is valid for this pin's default value and this pin is not already this type. */
+	bool CanConvertToDefaultValueMetadataType(FName PinLabel, EPCGMetadataTypes DataType) const;
+
+	/** Get the default value type on the pin. */
+	EPCGMetadataTypes GetPinDefaultValueType(FName PinLabel) const;
+
 	/** Whether to flip the order of the title lines - display generated title first and authored second. */
 	bool HasFlippedTitleLines() const;
 
@@ -118,6 +145,9 @@ public:
 	/** Whether this OutputPin can be connected to this InputPin */
 	virtual bool IsCompatible(const UPCGPin* InputPin, const UPCGPin* OutputPin, FText& OutReason) const;
 
+	UPCGSettings* GetSettings() const;
+	IPCGSettingsDefaultValueProvider* GetDefaultValueInterface() const;
+
 	DECLARE_DELEGATE(FOnPCGEditorGraphNodeChanged);
 	FOnPCGEditorGraphNodeChanged OnNodeChangedDelegate;
 
@@ -153,6 +183,10 @@ protected:
 	UPROPERTY()
 	TObjectPtr<UPCGNode> PCGNode = nullptr;
 
+	/** A flag if the node has ever been connected with an edge.*/
+	UPROPERTY()
+	bool bHasEverBeenConnected = false;
+
 	int32 DeferredReconstructCounter = 0;
 	bool bDeferredReconstruct = false;
 	bool bDisableReconstructFromNode = false;
diff --git a/Engine/Plugins/PCG/Source/PCGEditor/Private/Nodes/SPCGEditorGraphNode.cpp b/Engine/Plugins/PCG/Source/PCGEditor/Private/Nodes/SPCGEditorGraphNode.cpp
--- a/Engine/Plugins/PCG/Source/PCGEditor/Private/Nodes/SPCGEditorGraphNode.cpp
+++ b/Engine/Plugins/PCG/Source/PCGEditor/Private/Nodes/SPCGEditorGraphNode.cpp
@@ -6,430 +6,33 @@
 #include "PCGPin.h"
 #include "PCGSettings.h"
 #include "PCGSettingsWithDynamicInputs.h"
+#include "Metadata/PCGDefaultValueInterface.h"
 
 #include "PCGEditorGraphNodeBase.h"
 #include "PCGEditorStyle.h"
+#include "Metadata/PCGMetadataAttributeTraits.h"
+#include "Pins/SPCGEditorGraphPinBool.h"
+#include "Pins/SPCGEditorGraphPinNumSlider.h"
+#include "Pins/SPCGEditorGraphPinString.h"
+#include "Pins/SPCGEditorGraphPinVectorSlider.h"
 
 #include "GraphEditorSettings.h"
 #include "IDocumentation.h"
 #include "SCommentBubble.h"
 #include "SGraphPin.h"
-#include "SLevelOfDetailBranchNode.h"
 #include "SPinTypeSelector.h"
 #include "TutorialMetaData.h"
 #include "Algo/AnyOf.h"
-#include "Widgets/Images/SImage.h"
+
+#include "Widgets/SBoxPanel.h"
 #include "Widgets/SToolTip.h"
+#include "Widgets/Images/SImage.h"
 #include "Widgets/Input/SButton.h"
 #include "Widgets/Layout/SBox.h"
 #include "Widgets/Text/SInlineEditableTextBlock.h"
-#include "Widgets/SBoxPanel.h"
 
 #define LOCTEXT_NAMESPACE "SPCGEditorGraphNode"
 
-/** PCG pin primarily to give more control over pin coloring. */
-class SPCGEditorGraphNodePin : public SGraphPin
-{
-public:
-	SLATE_BEGIN_ARGS(SPCGEditorGraphNodePin)
-		: _PinLabelStyle(NAME_DefaultPinLabelStyle)
-		, _UsePinColorForText(false)
-		, _SideToSideMargin(5.0f)
-		{}
-		SLATE_ARGUMENT(FName, PinLabelStyle)
-		SLATE_ARGUMENT(bool, UsePinColorForText)
-		SLATE_ARGUMENT(float, SideToSideMargin)
-	SLATE_END_ARGS()
-
-	void Construct(const FArguments& InArgs, UEdGraphPin* InPin);
-
-	virtual FSlateColor GetPinColor() const override;
-	virtual FSlateColor GetPinTextColor() const override;
-	FName GetLabelStyle(FName DefaultLabelStyle) const;
-	bool GetExtraIcon(FName& OutExtraIcon, FText& OutTooltip) const;
-
-	/** Whether pin is required to be connected for execution. */
-	bool ShouldDisplayAsRequiredForExecution() const;
-
-private:
-	void ApplyUnusedPinStyle(FSlateColor& InOutColor) const;
-	void GetPCGNodeAndPin(const UPCGNode*& OutNode, const UPCGPin*& OutPin) const;
-};
-
-void SPCGEditorGraphNodePin::Construct(const FArguments& InArgs, UEdGraphPin* InPin)
-{
-	// IMPLEMENTATION NOTE: this is the code from SGraphPin::Construct with additional padding exposed,
-	// with an optional extra icon shown before the pin label, and with a marker icon to show pins
-	// that are required for execution.
-
-	bUsePinColorForText = InArgs._UsePinColorForText;
-	this->SetCursor(EMouseCursor::Default);
-
-	SetVisibility(MakeAttributeSP(this, &SPCGEditorGraphNodePin::GetPinVisiblity));
-
-	GraphPinObj = InPin;
-	check(GraphPinObj != NULL);
-
-	const UEdGraphSchema* Schema = GraphPinObj->GetSchema();
-	checkf(
-		Schema, 
-		TEXT("Missing schema for pin: %s with outer: %s of type %s"), 
-		*(GraphPinObj->GetName()),
-		GraphPinObj->GetOuter() ? *(GraphPinObj->GetOuter()->GetName()) : TEXT("NULL OUTER"), 
-		GraphPinObj->GetOuter() ? *(GraphPinObj->GetOuter()->GetClass()->GetName()) : TEXT("NULL OUTER")
-	);
-
-	const bool bIsInput = (GetDirection() == EGPD_Input);
-
-	// A small marker to indicate the pin is required for the node to be executed.
-	TSharedPtr<SImage> RequiredPinIconWidget;
-	float RequiredPinMarkerWidth = 0.0f;
-	bool bDisplayPinMarker = false;
-	if (bIsInput)
-	{
-		const FSlateBrush* RequiredPinMarkerIcon = FPCGEditorStyle::Get().GetBrush(PCGEditorStyleConstants::Pin_Required);
-		RequiredPinMarkerWidth = RequiredPinMarkerIcon ? RequiredPinMarkerIcon->GetImageSize().X : 8.0f;
-		bDisplayPinMarker = ShouldDisplayAsRequiredForExecution();
-
-		RequiredPinIconWidget =
-			SNew(SImage)
-			.Image(bDisplayPinMarker ? RequiredPinMarkerIcon : FAppStyle::GetNoBrush())
-			.ColorAndOpacity(this, &SPCGEditorGraphNodePin::GetPinColor);
-	}
-
-	// Create the pin icon widget
-	TSharedRef<SWidget> PinWidgetRef = SPinTypeSelector::ConstructPinTypeImage(
-		MakeAttributeSP(this, &SPCGEditorGraphNodePin::GetPinIcon ),
-		MakeAttributeSP(this, &SPCGEditorGraphNodePin::GetPinColor),
-		MakeAttributeSP(this, &SPCGEditorGraphNodePin::GetSecondaryPinIcon),
-		MakeAttributeSP(this, &SPCGEditorGraphNodePin::GetSecondaryPinColor));
-	PinImage = PinWidgetRef;
-
-	PinWidgetRef->SetCursor( 
-		TAttribute<TOptional<EMouseCursor::Type> >::Create (
-			TAttribute<TOptional<EMouseCursor::Type> >::FGetter::CreateRaw( this, &SPCGEditorGraphNodePin::GetPinCursor )
-		)
-	);
-
-	// Create the pin indicator widget (used for watched values)
-	static const FName NAME_NoBorder("NoBorder");
-	TSharedRef<SWidget> PinStatusIndicator =
-		SNew(SButton)
-		.ButtonStyle(FAppStyle::Get(), NAME_NoBorder)
-		.Visibility(this, &SPCGEditorGraphNodePin::GetPinStatusIconVisibility)
-		.ContentPadding(0)
-		.OnClicked(this, &SPCGEditorGraphNodePin::ClickedOnPinStatusIcon)
-		[
-			SNew(SImage)
-			.Image(this, &SPCGEditorGraphNodePin::GetPinStatusIcon)
-		];
-
-	TSharedRef<SWidget> LabelWidget = GetLabelWidget(GetLabelStyle(InArgs._PinLabelStyle));
-
-	// Create the widget used for the pin body (status indicator, label, and value)
-	LabelAndValue =
-		SNew(SWrapBox)
-		.PreferredSize(150.f);
-
-	TSharedPtr<SImage> ExtraPinIconWidget;
-	FName ExtraPinIcon;
-	FText ExtraPinIconTooltip;
-	if (GetExtraIcon(ExtraPinIcon, ExtraPinIconTooltip))
-	{
-		ExtraPinIconWidget = SNew(SImage)
-			.Image(FAppStyle::GetBrush(ExtraPinIcon))
-			.ColorAndOpacity(this, &SPCGEditorGraphNodePin::GetPinTextColor);
-
-		if (!ExtraPinIconTooltip.IsEmpty())
-		{
-			ExtraPinIconWidget->SetToolTipText(ExtraPinIconTooltip);
-		}
-	}
-
-	if (!bIsInput)
-	{
-		LabelAndValue->AddSlot()
-			.VAlign(VAlign_Center)
-			[
-				PinStatusIndicator
-			];
-
-		LabelAndValue->AddSlot()
-			.VAlign(VAlign_Center)
-			[
-				LabelWidget
-			];
-
-		if (ExtraPinIconWidget.IsValid())
-		{
-			LabelAndValue->AddSlot()
-				.Padding(5, 0, 0, 0)
-				.VAlign(VAlign_Center)
-				[
-					ExtraPinIconWidget.ToSharedRef()
-				];
-		}
-	}
-	else
-	{
-		if (ExtraPinIconWidget.IsValid())
-		{
-			LabelAndValue->AddSlot()
-				.Padding(0, 0, 5, 0)
-				.VAlign(VAlign_Center)
-				[
-					ExtraPinIconWidget.ToSharedRef()
-				];
-		}
-
-		LabelAndValue->AddSlot()
-			.VAlign(VAlign_Center)
-			[
-				LabelWidget
-			];
-
-		ValueWidget = GetDefaultValueWidget();
-
-		if (ValueWidget != SNullWidget::NullWidget)
-		{
-			TSharedPtr<SBox> ValueBox;
-			LabelAndValue->AddSlot()
-				.Padding(bIsInput ? FMargin(InArgs._SideToSideMargin, 0, 0, 0) : FMargin(0, 0, InArgs._SideToSideMargin, 0))
-				.VAlign(VAlign_Center)
-				[
-					SAssignNew(ValueBox, SBox)
-					.Padding(0.0f)
-					[
-						ValueWidget.ToSharedRef()
-					]
-				];
-
-			if (!DoesWidgetHandleSettingEditingEnabled())
-			{
-				ValueBox->SetEnabled(TAttribute<bool>(this, &SPCGEditorGraphNodePin::IsEditingEnabled));
-			}
-		}
-
-		LabelAndValue->AddSlot()
-			.VAlign(VAlign_Center)
-			[
-				PinStatusIndicator
-			];
-	}
-
-	TSharedPtr<SHorizontalBox> PinContent;
-	if (bIsInput)
-	{
-		// Input pin
-		FullPinHorizontalRowWidget = PinContent = 
-			SNew(SHorizontalBox)
-			+SHorizontalBox::Slot()
-			.AutoWidth()
-			.VAlign(VAlign_Center)
-			[
-				RequiredPinIconWidget.ToSharedRef()
-			]
-			+SHorizontalBox::Slot()
-			.AutoWidth()
-			.VAlign(VAlign_Center)
-			.Padding(bDisplayPinMarker ? 0.0f : RequiredPinMarkerWidth, 0, InArgs._SideToSideMargin, 0)
-			[
-				PinWidgetRef
-			]
-			+SHorizontalBox::Slot()
-			.AutoWidth()
-			.VAlign(VAlign_Center)
-			[
-				LabelAndValue.ToSharedRef()
-			];
-	}
-	else
-	{
-		// Output pin
-		FullPinHorizontalRowWidget = PinContent = SNew(SHorizontalBox)
-			+SHorizontalBox::Slot()
-			.AutoWidth()
-			.VAlign(VAlign_Center)
-			[
-				LabelAndValue.ToSharedRef()
-			]
-			+SHorizontalBox::Slot()
-			.AutoWidth()
-			.VAlign(VAlign_Center)
-			.Padding(InArgs._SideToSideMargin, 0, 0, 0)
-			[
-				PinWidgetRef
-			];
-	}
-
-	// Set up a hover for pins that is tinted the color of the pin.
-	
-	SBorder::Construct(SBorder::FArguments()
-		.BorderImage(this, &SPCGEditorGraphNodePin::GetPinBorder)
-		.BorderBackgroundColor(this, &SPCGEditorGraphNodePin::GetHighlightColor)
-		.OnMouseButtonDown(this, &SPCGEditorGraphNodePin::OnPinNameMouseDown)
-		.Padding(0) // NOTE: This is different from base class implementation
-		[
-			SNew(SBorder)
-			.BorderImage(CachedImg_Pin_DiffOutline)
-			.BorderBackgroundColor(this, &SPCGEditorGraphNodePin::GetPinDiffColor)
-			.Padding(0) // NOTE: This is different from base class implementation
-			[
-				SNew(SLevelOfDetailBranchNode)
-				.UseLowDetailSlot(this, &SPCGEditorGraphNodePin::UseLowDetailPinNames)
-				.LowDetail()
-				[
-					//@TODO: Try creating a pin-colored line replacement that doesn't measure text / call delegates but still renders
-					PinWidgetRef
-				]
-				.HighDetail()
-				[
-					PinContent.ToSharedRef()
-				]
-			]
-		]
-	);
-
-	TSharedPtr<IToolTip> TooltipWidget = SNew(SToolTip)
-		.Text(this, &SPCGEditorGraphNodePin::GetTooltipText);
-
-	SetToolTip(TooltipWidget);
-
-}
-
-void SPCGEditorGraphNodePin::GetPCGNodeAndPin(const UPCGNode*& OutNode, const UPCGPin*& OutPin) const
-{
-	UEdGraphPin* GraphPin = GetPinObj();
-	if (GraphPin && !GraphPin->IsPendingKill())
-	{
-		const UPCGEditorGraphNodeBase* EditorNode = CastChecked<const UPCGEditorGraphNodeBase>(GraphPinObj->GetOwningNode());
-		
-		OutNode = EditorNode ? EditorNode->GetPCGNode() : nullptr;
-
-		OutPin = OutNode ? OutNode->GetInputPin(GraphPin->GetFName()) : nullptr;
-		if (!OutPin)
-		{
-			OutPin = OutNode ? OutNode->GetOutputPin(GraphPin->GetFName()) : nullptr;
-		}
-	}
-	else
-	{
-		OutNode = nullptr;
-		OutPin = nullptr;
-	}
-}
-
-void SPCGEditorGraphNodePin::ApplyUnusedPinStyle(FSlateColor& InOutColor) const
-{
-	const UPCGPin* PCGPin = nullptr;
-	const UPCGNode* PCGNode = nullptr;
-	GetPCGNodeAndPin(PCGNode, PCGPin);
-
-	bool bPinDisabled = false;
-
-	// Check if the pin was deactivated in the previous execution.
-	const UEdGraphPin* Pin = GetPinObj();
-	if (Pin && Pin->Direction == EEdGraphPinDirection::EGPD_Output)
-	{
-		if (const UPCGEditorGraphNodeBase* Node = Cast<UPCGEditorGraphNodeBase>(Pin ? Pin->GetOwningNode() : nullptr))
-		{
-			// If node is already disabled, don't bother disabling pin on top of that, does not look nice to disable both and may
-			// not be meaningful to do so in any case.
-			if (!Node->IsDisplayAsDisabledForced())
-			{
-				bPinDisabled = !Node->IsOutputPinActive(Pin);
-			}
-		}
-	}
-
-	// Halve opacity if pin is unused - intended to happen whether disabled or not
-	if (bPinDisabled || (PCGPin && PCGNode && !PCGNode->IsPinUsedByNodeExecution(PCGPin)))
-	{
-		FLinearColor Color = InOutColor.GetSpecifiedColor();
-		Color.A *= 0.5;
-		InOutColor = Color;
-	}
-}
-
-// Adapted from SGraphPin::GetPinColor
-FSlateColor SPCGEditorGraphNodePin::GetPinColor() const
-{
-	FSlateColor Color = SGraphPin::GetPinColor();
-
-	ApplyUnusedPinStyle(Color);
-
-	return Color;
-}
-
-FSlateColor SPCGEditorGraphNodePin::GetPinTextColor() const
-{
-	FSlateColor Color = SGraphPin::GetPinTextColor();
-
-	ApplyUnusedPinStyle(Color);
-
-	return Color;
-}
-
-FName SPCGEditorGraphNodePin::GetLabelStyle(FName DefaultLabelStyle) const
-{
-	const UPCGPin* PCGPin = nullptr;
-	const UPCGNode* PCGNode = nullptr;
-	FName LabelStyle = NAME_None;
-
-	GetPCGNodeAndPin(PCGNode, PCGPin);
-
-	const UPCGSettings* Settings = PCGNode ? PCGNode->GetSettings() : nullptr;
-
-	if (!PCGPin || !Settings || !Settings->GetPinLabelStyle(PCGPin, LabelStyle))
-	{
-		LabelStyle = DefaultLabelStyle;
-	}
-
-	return LabelStyle;
-}
-
-bool SPCGEditorGraphNodePin::GetExtraIcon(FName& OutExtraIcon, FText& OutTooltip) const
-{
-	const UPCGPin* PCGPin = nullptr;
-	const UPCGNode* PCGNode = nullptr;
-
-	GetPCGNodeAndPin(PCGNode, PCGPin);
-
-	const UPCGSettings* Settings = PCGNode ? PCGNode->GetSettings() : nullptr;
-	return (PCGPin && Settings) ? Settings->GetPinExtraIcon(PCGPin, OutExtraIcon, OutTooltip) : false;
-}
-
-bool SPCGEditorGraphNodePin::ShouldDisplayAsRequiredForExecution() const
-{
-	const UPCGPin* PCGPin = nullptr;
-	const UPCGNode* PCGNode = nullptr;
-	GetPCGNodeAndPin(PCGNode, PCGPin);
-
-	// Trivial early out tests, and advanced pins should never display as required.
-	if (!PCGPin || !PCGNode || PCGPin->Properties.IsAdvancedPin())
-	{
-		return false;
-	}
-
-	if (PCGNode->IsInputPinRequiredByExecution(PCGPin))
-	{
-		return true;
-	}
-
-	const UPCGSettings* Settings = PCGNode->GetSettings();
-	if (Settings && Settings->CanCullTaskIfUnwired())
-	{
-		// If the node will cull if unwired, and if it only has a single normal pin (and no required pins), then display the pin
-		// as required, because it effectively is. So return false if there are other pins which are not advanced.
-		return !Algo::AnyOf(PCGNode->GetInputPins(), [PCGPin](UPCGPin* InOtherPin)
-		{
-			return InOtherPin && InOtherPin != PCGPin && !InOtherPin->Properties.IsAdvancedPin();
-		});
-	}
-
-	return false;
-}
-
 void SPCGEditorGraphNode::Construct(const FArguments& InArgs, UPCGEditorGraphNodeBase* InNode)
 {
 	GraphNode = InNode;
@@ -512,9 +115,65 @@ TSharedRef<SWidget> SPCGEditorGraphNode::CreateTitleWidget(TSharedPtr<SNodeTitle
 	return InlineEditableText.ToSharedRef();
 }
 
-TSharedPtr<SGraphPin> SPCGEditorGraphNode::CreatePinWidget(UEdGraphPin* Pin) const
+TSharedPtr<SGraphPin> SPCGEditorGraphNode::CreatePinWidget(UEdGraphPin* InPin) const
 {
-	return SNew(SPCGEditorGraphNodePin, Pin);
+	const UPCGNode* PCGNode = PCGEditorGraphNode ? PCGEditorGraphNode->GetPCGNode() : nullptr;
+	UPCGSettings* NodeSettings = PCGNode ? PCGNode->GetSettings() : nullptr;
+
+	if (InPin && NodeSettings && NodeSettings->Implements<UPCGSettingsDefaultValueProvider>())
+	{
+		const IPCGSettingsDefaultValueProvider* DefaultValueInterface = CastChecked<IPCGSettingsDefaultValueProvider>(NodeSettings);
+		if (DefaultValueInterface->DefaultValuesAreEnabled() && DefaultValueInterface->IsPinDefaultValueActivated(InPin->PinName))
+		{
+			// Set the string default value to match the settings' source of truth.
+			InPin->DefaultValue = DefaultValueInterface->GetPinDefaultValueAsString(InPin->PinName);
+
+			// To link the transaction to the settings for Undo/Redo.
+			TDelegate<void()> OnModify = FSimpleDelegate::CreateLambda([NodeSettingsPtr = TWeakObjectPtr(NodeSettings)]
+			{
+				if (NodeSettingsPtr.IsValid())
+				{
+					NodeSettingsPtr.Pin()->Modify();
+				}
+			});
+
+			switch (DefaultValueInterface->GetPinDefaultValueType(InPin->PinName))
+			{
+				case EPCGMetadataTypes::Name: // fall-through
+				case EPCGMetadataTypes::String:
+					return SNew(SPCGEditorGraphPinString, InPin, std::move(OnModify));
+				// Float is converted to double by the property accessor under the hood
+				case EPCGMetadataTypes::Float: // fall-through
+				case EPCGMetadataTypes::Double:
+					return SNew(SPCGEditorGraphPinNumSlider<double>, InPin, std::move(OnModify));
+				case EPCGMetadataTypes::Integer32:
+					return SNew(SPCGEditorGraphPinNumSlider<int32>, InPin, std::move(OnModify))
+						.MinDesiredBoxWidth(40.f);
+				case EPCGMetadataTypes::Integer64:
+					return SNew(SPCGEditorGraphPinNumSlider<int64>, InPin, std::move(OnModify))
+						.MinDesiredBoxWidth(40.f);
+				case EPCGMetadataTypes::Vector:
+					return SNew(SPCGEditorGraphPinVectorSlider<FVector>, InPin, std::move(OnModify));
+				case EPCGMetadataTypes::Vector2:
+					return SNew(SPCGEditorGraphPinVectorSlider<FVector2D>, InPin, std::move(OnModify));
+				case EPCGMetadataTypes::Vector4:
+					return SNew(SPCGEditorGraphPinVectorSlider<FVector4>, InPin, std::move(OnModify));
+				case EPCGMetadataTypes::Rotator:
+					return SNew(SPCGEditorGraphPinVectorSlider<FRotator>, InPin, std::move(OnModify));
+				case EPCGMetadataTypes::Boolean:
+					return SNew(SPCGEditorGraphPinBool, InPin, std::move(OnModify));
+				// @todo_pcg: Will be added once widgets are created.
+				case EPCGMetadataTypes::SoftObjectPath: // fall-through
+				case EPCGMetadataTypes::SoftClassPath:  // fall-through
+				case EPCGMetadataTypes::Quaternion:     // fall-through
+				case EPCGMetadataTypes::Transform:      // fall-through
+				default:
+					break;
+			}
+		}
+	}
+
+	return SNew(SPCGEditorGraphNodePin, InPin);
 }
 
 EVisibility SPCGEditorGraphNode::IsAddPinButtonVisible() const
diff --git a/Engine/Plugins/PCG/Source/PCGEditor/Private/Pins/SPCGEditorGraphNodePin.cpp b/Engine/Plugins/PCG/Source/PCGEditor/Private/Pins/SPCGEditorGraphNodePin.cpp
new file mode 100644
--- /dev/null
+++ b/Engine/Plugins/PCG/Source/PCGEditor/Private/Pins/SPCGEditorGraphNodePin.cpp
@@ -0,0 +1,388 @@
+// Copyright Epic Games, Inc. All Rights Reserved.
+
+#pragma once
+
+#include "SPCGEditorGraphNodePin.h"
+
+#include "PCGEditorStyle.h"
+#include "PCGNode.h"
+#include "PCGPin.h"
+#include "PCGSettings.h"
+#include "Nodes/PCGEditorGraphNodeBase.h"
+
+#include "SLevelOfDetailBranchNode.h"
+#include "SPinTypeSelector.h"
+#include "Algo/AnyOf.h"
+#include "Widgets/SToolTip.h"
+#include "Widgets/Input/SButton.h"
+
+void SPCGEditorGraphNodePin::Construct(const FArguments& InArgs, UEdGraphPin* InPin)
+{
+	// IMPLEMENTATION NOTE: this is the code from SGraphPin::Construct with additional padding exposed,
+	// with an optional extra icon shown before the pin label, and with a marker icon to show pins
+	// that are required for execution.
+
+	bUsePinColorForText = InArgs._UsePinColorForText;
+	this->SetCursor(EMouseCursor::Default);
+
+	SetVisibility(MakeAttributeSP(this, &SPCGEditorGraphNodePin::GetPinVisiblity));
+
+	GraphPinObj = InPin;
+	check(GraphPinObj != NULL);
+
+	const UEdGraphSchema* Schema = GraphPinObj->GetSchema();
+	checkf(
+		Schema,
+		TEXT("Missing schema for pin: %s with outer: %s of type %s"),
+		*(GraphPinObj->GetName()),
+		GraphPinObj->GetOuter() ? *(GraphPinObj->GetOuter()->GetName()) : TEXT("NULL OUTER"),
+		GraphPinObj->GetOuter() ? *(GraphPinObj->GetOuter()->GetClass()->GetName()) : TEXT("NULL OUTER")
+		);
+
+	const bool bIsInput = (GetDirection() == EGPD_Input);
+
+	// A small marker to indicate the pin is required for the node to be executed.
+	TSharedPtr<SImage> RequiredPinIconWidget;
+	float RequiredPinMarkerWidth = 0.0f;
+	bool bDisplayPinMarker = false;
+	if (bIsInput)
+	{
+		const FSlateBrush* RequiredPinMarkerIcon = FPCGEditorStyle::Get().GetBrush(PCGEditorStyleConstants::Pin_Required);
+		RequiredPinMarkerWidth = RequiredPinMarkerIcon ? RequiredPinMarkerIcon->GetImageSize().X : 8.0f;
+		bDisplayPinMarker = ShouldDisplayAsRequiredForExecution();
+
+		RequiredPinIconWidget =
+			SNew(SImage)
+			.Image(bDisplayPinMarker ? RequiredPinMarkerIcon : FAppStyle::GetNoBrush())
+			.ColorAndOpacity(this, &SPCGEditorGraphNodePin::GetPinColor);
+	}
+
+	// Create the pin icon widget
+	TSharedRef<SWidget> PinWidgetRef = SPinTypeSelector::ConstructPinTypeImage(
+		MakeAttributeSP(this, &SPCGEditorGraphNodePin::GetPinIcon),
+		MakeAttributeSP(this, &SPCGEditorGraphNodePin::GetPinColor),
+		MakeAttributeSP(this, &SPCGEditorGraphNodePin::GetSecondaryPinIcon),
+		MakeAttributeSP(this, &SPCGEditorGraphNodePin::GetSecondaryPinColor));
+	PinImage = PinWidgetRef;
+
+	PinWidgetRef->SetCursor(
+		TAttribute<TOptional<EMouseCursor::Type>>::Create(
+			TAttribute<TOptional<EMouseCursor::Type>>::FGetter::CreateRaw(this, &SPCGEditorGraphNodePin::GetPinCursor)
+			)
+		);
+
+	// Create the pin indicator widget (used for watched values)
+	static const FName NAME_NoBorder("NoBorder");
+	TSharedRef<SWidget> PinStatusIndicator =
+		SNew(SButton)
+		.ButtonStyle(FAppStyle::Get(), NAME_NoBorder)
+		.Visibility(this, &SPCGEditorGraphNodePin::GetPinStatusIconVisibility)
+		.ContentPadding(0)
+		.OnClicked(this, &SPCGEditorGraphNodePin::ClickedOnPinStatusIcon)
+		[
+			SNew(SImage)
+			.Image(this, &SPCGEditorGraphNodePin::GetPinStatusIcon)
+		];
+
+	TSharedRef<SWidget> LabelWidget = GetLabelWidget(GetLabelStyle(InArgs._PinLabelStyle));
+
+	// Create the widget used for the pin body (status indicator, label, and value)
+	LabelAndValue =
+		SNew(SWrapBox)
+		.PreferredSize(150.f);
+
+	TSharedPtr<SImage> ExtraPinIconWidget;
+	FName ExtraPinIcon;
+	FText ExtraPinIconTooltip;
+	if (GetExtraIcon(ExtraPinIcon, ExtraPinIconTooltip))
+	{
+		ExtraPinIconWidget = SNew(SImage)
+			.Image(FAppStyle::GetBrush(ExtraPinIcon))
+			.ColorAndOpacity(this, &SPCGEditorGraphNodePin::GetPinTextColor);
+
+		if (!ExtraPinIconTooltip.IsEmpty())
+		{
+			ExtraPinIconWidget->SetToolTipText(ExtraPinIconTooltip);
+		}
+	}
+
+	if (!bIsInput)
+	{
+		LabelAndValue->AddSlot()
+			.VAlign(VAlign_Center)
+			[
+				PinStatusIndicator
+			];
+
+		LabelAndValue->AddSlot()
+			.VAlign(VAlign_Center)
+			[
+				LabelWidget
+			];
+
+		if (ExtraPinIconWidget.IsValid())
+		{
+			LabelAndValue->AddSlot()
+				.Padding(5, 0, 0, 0)
+				.VAlign(VAlign_Center)
+				[
+					ExtraPinIconWidget.ToSharedRef()
+				];
+		}
+	}
+	else
+	{
+		if (ExtraPinIconWidget.IsValid())
+		{
+			LabelAndValue->AddSlot()
+				.Padding(0, 0, 5, 0)
+				.VAlign(VAlign_Center)
+				[
+					ExtraPinIconWidget.ToSharedRef()
+				];
+		}
+
+		LabelAndValue->AddSlot()
+			.VAlign(VAlign_Center)
+			[
+				LabelWidget
+			];
+
+		ValueWidget = GetDefaultValueWidget();
+
+		if (ValueWidget != SNullWidget::NullWidget)
+		{
+			TSharedPtr<SBox> ValueBox;
+			LabelAndValue->AddSlot()
+				.Padding(FMargin(InArgs._SideToSideMargin, 0, 0, 0))
+				.VAlign(VAlign_Center)
+				[
+					SAssignNew(ValueBox, SBox)
+					[
+						ValueWidget.ToSharedRef()
+					]
+				];
+
+			if (!DoesWidgetHandleSettingEditingEnabled())
+			{
+				ValueBox->SetEnabled(TAttribute<bool>(this, &SPCGEditorGraphNodePin::IsEditingEnabled));
+			}
+		}
+
+		LabelAndValue->AddSlot()
+			.VAlign(VAlign_Center)
+			[
+				PinStatusIndicator
+			];
+	}
+
+	TSharedPtr<SHorizontalBox> PinContent;
+	if (bIsInput)
+	{
+		// Input pin
+		FullPinHorizontalRowWidget = PinContent =
+			SNew(SHorizontalBox)
+			+ SHorizontalBox::Slot()
+			.AutoWidth()
+			.VAlign(VAlign_Center)
+			[
+				RequiredPinIconWidget.ToSharedRef()
+			]
+			+ SHorizontalBox::Slot()
+			.AutoWidth()
+			.VAlign(VAlign_Center)
+			.Padding(bDisplayPinMarker ? 0.0f : RequiredPinMarkerWidth, 0, InArgs._SideToSideMargin, 0)
+			[
+				PinWidgetRef
+			]
+			+ SHorizontalBox::Slot()
+			.AutoWidth()
+			.VAlign(VAlign_Center)
+			[
+				LabelAndValue.ToSharedRef()
+			];
+	}
+	else
+	{
+		// Output pin
+		FullPinHorizontalRowWidget = PinContent = SNew(SHorizontalBox)
+			+ SHorizontalBox::Slot()
+			.AutoWidth()
+			.VAlign(VAlign_Center)
+			[
+				LabelAndValue.ToSharedRef()
+			]
+			+ SHorizontalBox::Slot()
+			.AutoWidth()
+			.VAlign(VAlign_Center)
+			.Padding(InArgs._SideToSideMargin, 0, 0, 0)
+			[
+				PinWidgetRef
+			];
+	}
+
+	// Set up a hover for pins that is tinted the color of the pin.
+
+	SBorder::Construct(SBorder::FArguments()
+		.BorderImage(this, &SPCGEditorGraphNodePin::GetPinBorder)
+		.BorderBackgroundColor(this, &SPCGEditorGraphNodePin::GetHighlightColor)
+		.OnMouseButtonDown(this, &SPCGEditorGraphNodePin::OnPinNameMouseDown)
+		.Padding(0) // NOTE: This is different from base class implementation
+		[
+			SNew(SBorder)
+			.BorderImage(CachedImg_Pin_DiffOutline)
+			.BorderBackgroundColor(this, &SPCGEditorGraphNodePin::GetPinDiffColor)
+			.Padding(0) // NOTE: This is different from base class implementation
+			[
+				SNew(SLevelOfDetailBranchNode)
+				.UseLowDetailSlot(this, &SPCGEditorGraphNodePin::UseLowDetailPinNames)
+				.LowDetail()
+				[
+					//@TODO: Try creating a pin-colored line replacement that doesn't measure text / call delegates but still renders
+					PinWidgetRef
+				]
+				.HighDetail()
+				[
+					PinContent.ToSharedRef()
+				]
+			]
+		]);
+
+	TSharedPtr<IToolTip> TooltipWidget = SNew(SToolTip)
+		.Text(this, &SPCGEditorGraphNodePin::GetTooltipText);
+
+	SetToolTip(TooltipWidget);
+}
+
+void SPCGEditorGraphNodePin::GetPCGNodeAndPin(const UPCGNode*& OutNode, const UPCGPin*& OutPin) const
+{
+	UEdGraphPin* GraphPin = GetPinObj();
+	if (GraphPin && !GraphPin->IsPendingKill())
+	{
+		const UPCGEditorGraphNodeBase* EditorNode = CastChecked<const UPCGEditorGraphNodeBase>(GraphPinObj->GetOwningNode());
+
+		OutNode = EditorNode ? EditorNode->GetPCGNode() : nullptr;
+
+		OutPin = OutNode ? OutNode->GetInputPin(GraphPin->GetFName()) : nullptr;
+		if (!OutPin)
+		{
+			OutPin = OutNode ? OutNode->GetOutputPin(GraphPin->GetFName()) : nullptr;
+		}
+	}
+	else
+	{
+		OutNode = nullptr;
+		OutPin = nullptr;
+	}
+}
+
+void SPCGEditorGraphNodePin::ApplyUnusedPinStyle(FSlateColor& InOutColor) const
+{
+	const UPCGPin* PCGPin = nullptr;
+	const UPCGNode* PCGNode = nullptr;
+	GetPCGNodeAndPin(PCGNode, PCGPin);
+
+	bool bPinDisabled = false;
+
+	// Check if the pin was deactivated in the previous execution.
+	const UEdGraphPin* Pin = GetPinObj();
+	if (Pin && Pin->Direction == EEdGraphPinDirection::EGPD_Output)
+	{
+		if (const UPCGEditorGraphNodeBase* Node = Cast<UPCGEditorGraphNodeBase>(Pin ? Pin->GetOwningNode() : nullptr))
+		{
+			// If node is already disabled, don't bother disabling pin on top of that, does not look nice to disable both and may
+			// not be meaningful to do so in any case.
+			if (!Node->IsDisplayAsDisabledForced())
+			{
+				bPinDisabled = !Node->IsOutputPinActive(Pin);
+			}
+		}
+	}
+
+	// Halve opacity if pin is unused - intended to happen whether disabled or not
+	if (bPinDisabled || (PCGPin && PCGNode && !PCGNode->IsPinUsedByNodeExecution(PCGPin)))
+	{
+		FLinearColor Color = InOutColor.GetSpecifiedColor();
+		Color.A *= 0.5;
+		InOutColor = Color;
+	}
+}
+
+// Adapted from SGraphPin::GetPinColor
+FSlateColor SPCGEditorGraphNodePin::GetPinColor() const
+{
+	FSlateColor Color = SGraphPin::GetPinColor();
+
+	ApplyUnusedPinStyle(Color);
+
+	return Color;
+}
+
+FSlateColor SPCGEditorGraphNodePin::GetPinTextColor() const
+{
+	FSlateColor Color = SGraphPin::GetPinTextColor();
+
+	ApplyUnusedPinStyle(Color);
+
+	return Color;
+}
+
+FName SPCGEditorGraphNodePin::GetLabelStyle(FName DefaultLabelStyle) const
+{
+	const UPCGPin* PCGPin = nullptr;
+	const UPCGNode* PCGNode = nullptr;
+	FName LabelStyle = NAME_None;
+
+	GetPCGNodeAndPin(PCGNode, PCGPin);
+
+	const UPCGSettings* Settings = PCGNode ? PCGNode->GetSettings() : nullptr;
+
+	if (!PCGPin || !Settings || !Settings->GetPinLabelStyle(PCGPin, LabelStyle))
+	{
+		LabelStyle = DefaultLabelStyle;
+	}
+
+	return LabelStyle;
+}
+
+bool SPCGEditorGraphNodePin::GetExtraIcon(FName& OutExtraIcon, FText& OutTooltip) const
+{
+	const UPCGPin* PCGPin = nullptr;
+	const UPCGNode* PCGNode = nullptr;
+
+	GetPCGNodeAndPin(PCGNode, PCGPin);
+
+	const UPCGSettings* Settings = PCGNode ? PCGNode->GetSettings() : nullptr;
+	return (PCGPin && Settings) ? Settings->GetPinExtraIcon(PCGPin, OutExtraIcon, OutTooltip) : false;
+}
+
+bool SPCGEditorGraphNodePin::ShouldDisplayAsRequiredForExecution() const
+{
+	const UPCGPin* PCGPin = nullptr;
+	const UPCGNode* PCGNode = nullptr;
+	GetPCGNodeAndPin(PCGNode, PCGPin);
+
+	// Trivial early out tests, and advanced pins should never display as required.
+	if (!PCGPin || !PCGNode || PCGPin->Properties.IsAdvancedPin())
+	{
+		return false;
+	}
+
+	if (PCGNode->IsInputPinRequiredByExecution(PCGPin))
+	{
+		return true;
+	}
+
+	const UPCGSettings* Settings = PCGNode->GetSettings();
+	if (Settings && Settings->CanCullTaskIfUnwired())
+	{
+		// If the node will cull if unwired, and if it only has a single normal pin (and no required pins), then display the pin
+		// as required, because it effectively is. So return false if there are other pins which are not advanced.
+		return !Algo::AnyOf(PCGNode->GetInputPins(), [PCGPin](UPCGPin* InOtherPin)
+		{
+			return InOtherPin && InOtherPin != PCGPin && !InOtherPin->Properties.IsAdvancedPin();
+		});
+	}
+
+	return false;
+}
\ No newline at end of file
diff --git a/Engine/Plugins/PCG/Source/PCGEditor/Private/Pins/SPCGEditorGraphNodePin.h b/Engine/Plugins/PCG/Source/PCGEditor/Private/Pins/SPCGEditorGraphNodePin.h
new file mode 100644
--- /dev/null
+++ b/Engine/Plugins/PCG/Source/PCGEditor/Private/Pins/SPCGEditorGraphNodePin.h
@@ -0,0 +1,39 @@
+// Copyright Epic Games, Inc. All Rights Reserved.
+
+#pragma once
+
+#include "SGraphPin.h"
+
+class UPCGPin;
+class UPCGNode;
+
+/** PCG pin primarily to give more control over pin coloring. */
+class SPCGEditorGraphNodePin : public SGraphPin
+{
+public:
+	SLATE_BEGIN_ARGS(SPCGEditorGraphNodePin)
+		: _PinLabelStyle(NAME_DefaultPinLabelStyle)
+		, _SideToSideMargin(5.0f)
+		, _UsePinColorForText(false)
+	{}
+
+	SLATE_ARGUMENT(FName, PinLabelStyle)
+	SLATE_ARGUMENT(float, SideToSideMargin)
+	SLATE_ARGUMENT(bool, UsePinColorForText)
+
+	SLATE_END_ARGS()
+
+	void Construct(const FArguments& InArgs, UEdGraphPin* InPin);
+
+	virtual FSlateColor GetPinColor() const override;
+	virtual FSlateColor GetPinTextColor() const override;
+	FName GetLabelStyle(FName DefaultLabelStyle) const;
+	bool GetExtraIcon(FName& OutExtraIcon, FText& OutTooltip) const;
+
+	/** Whether pin is required to be connected for execution. */
+	bool ShouldDisplayAsRequiredForExecution() const;
+
+private:
+	void ApplyUnusedPinStyle(FSlateColor& InOutColor) const;
+	void GetPCGNodeAndPin(const UPCGNode*& OutNode, const UPCGPin*& OutPin) const;
+};
\ No newline at end of file
diff --git a/Engine/Plugins/PCG/Source/PCGEditor/Private/Pins/SPCGEditorGraphPinBool.h b/Engine/Plugins/PCG/Source/PCGEditor/Private/Pins/SPCGEditorGraphPinBool.h
new file mode 100644
--- /dev/null
+++ b/Engine/Plugins/PCG/Source/PCGEditor/Private/Pins/SPCGEditorGraphPinBool.h
@@ -0,0 +1,65 @@
+// Copyright Epic Games, Inc. All Rights Reserved.
+
+#pragma once
+
+#include "SPCGEditorGraphNodePin.h"
+
+#include "ScopedTransaction.h"
+
+#include "Widgets/Input/SCheckBox.h"
+
+/** Note: Derived from Engine/Source/Editor/GraphEditor/Public/KismetPins/SGraphPinBool.h */
+
+class SPCGEditorGraphPinBool final : public SPCGEditorGraphNodePin
+{
+public:
+	SLATE_BEGIN_ARGS(SPCGEditorGraphPinBool) {}
+	SLATE_END_ARGS()
+
+	void Construct(const FArguments& InArgs, UEdGraphPin* InPin, TDelegate<void()>&& OnModify)
+	{
+		OnModifyDelegate = std::move(OnModify);
+		SGraphPin::Construct(SGraphPin::FArguments(), InPin);
+	}
+
+protected:
+	//~ Begin SGraphPin Interface
+	virtual TSharedRef<SWidget> GetDefaultValueWidget() override
+	{
+		return SNew(SCheckBox)
+			.IsChecked(this, &SPCGEditorGraphPinBool::IsDefaultValueChecked)
+			.IsEnabled(this, &SGraphPin::GetDefaultValueIsEditable)
+			.OnCheckStateChanged(this, &SPCGEditorGraphPinBool::OnDefaultValueCheckBoxChanged)
+			.Visibility(this, &SGraphPin::GetDefaultValueVisibility);
+	}
+
+	//~ End SGraphPin Interface
+
+	/** Determine if the checkbox should be checked or not */
+	ECheckBoxState IsDefaultValueChecked() const
+	{
+		FString CurrentValue = GraphPinObj->GetDefaultAsString();
+		return CurrentValue.ToBool() ? ECheckBoxState::Checked : ECheckBoxState::Unchecked;
+	}
+
+	/** Called when check box is changed */
+	void OnDefaultValueCheckBoxChanged(const ECheckBoxState InIsChecked) const
+	{
+		if (GraphPinObj->IsPendingKill())
+		{
+			return;
+		}
+
+		const FString BoolString = (InIsChecked == ECheckBoxState::Checked) ? TEXT("true") : TEXT("false");
+		if (GraphPinObj->GetDefaultAsString() != BoolString)
+		{
+			const FScopedTransaction Transaction(NSLOCTEXT("PCGGraphEditor", "ChangeBoolPinValue", "Change Bool Pin Value"));
+			GraphPinObj->Modify();
+			OnModifyDelegate.ExecuteIfBound();
+			GraphPinObj->GetSchema()->TrySetDefaultValue(*GraphPinObj, BoolString);
+		}
+	}
+
+private:
+	FSimpleDelegate OnModifyDelegate;
+};
diff --git a/Engine/Plugins/PCG/Source/PCGEditor/Private/Pins/SPCGEditorGraphPinNumSlider.h b/Engine/Plugins/PCG/Source/PCGEditor/Private/Pins/SPCGEditorGraphPinNumSlider.h
new file mode 100644
--- /dev/null
+++ b/Engine/Plugins/PCG/Source/PCGEditor/Private/Pins/SPCGEditorGraphPinNumSlider.h
@@ -0,0 +1,172 @@
+// Copyright Epic Games, Inc. All Rights Reserved.
+
+#pragma once
+
+#include "SPCGEditorGraphNodePin.h"
+
+#include "ScopedTransaction.h"
+#include "Editor.h"
+#include "NumericPropertyParams.h"
+
+#include "Widgets/Layout/SBox.h"
+#include "Widgets/Input/SNumericEntryBox.h"
+
+/** Note: Derived from Engine/Source/Editor/GraphEditor/Public/KismetPins/SGraphPinNumSlider.h */
+
+template <typename NumericType>
+class SPCGEditorGraphPinNumSlider final : public SPCGEditorGraphNodePin
+{
+public:
+	SLATE_BEGIN_ARGS(SPCGEditorGraphPinNumSlider)
+		: _PinProperty(nullptr)
+		, _MinDesiredBoxWidth(60.0f)
+		, _MaxDesiredBoxWidth(400.f)
+		, _ShouldShowDisabledWhenConnected(true)
+	{}
+
+	SLATE_ARGUMENT(FProperty*, PinProperty)
+	SLATE_ARGUMENT(float, MinDesiredBoxWidth)
+	SLATE_ARGUMENT(float, MaxDesiredBoxWidth)
+	SLATE_ARGUMENT(bool, ShouldShowDisabledWhenConnected)
+
+	SLATE_END_ARGS()
+
+	void Construct(const FArguments& InArgs, UEdGraphPin* InPin, TDelegate<void()>&& OnModify)
+	{
+		MinDesiredBoxWidth = InArgs._MinDesiredBoxWidth;
+		MaxDesiredBoxWidth = InArgs._MaxDesiredBoxWidth;
+		bShouldHideWhenConnected = InArgs._ShouldShowDisabledWhenConnected;
+		OnModifyDelegate = std::move(OnModify);
+
+		SPCGEditorGraphNodePin::Construct(SPCGEditorGraphNodePin::FArguments(), InPin);
+	}
+
+	//~ Begin SGraphPin Interface
+	virtual TSharedRef<SWidget> GetDefaultValueWidget() override
+	{
+		// Use generic defaults for now.
+		TNumericPropertyParams<NumericType> NumericPropertyParams(/*Property=*/nullptr, /*MetaDataGetter=*/nullptr);
+
+		// Save last committed value to compare when value changes
+		CachedValue = GetNumericValue().GetValue();
+
+		return SNew(SBox)
+			.MinDesiredWidth(MinDesiredBoxWidth)
+			.MaxDesiredWidth(400)
+			[
+				SNew(SNumericEntryBox<NumericType>)
+				.EditableTextBoxStyle(FAppStyle::Get(), "Graph.EditableTextBox")
+				.BorderForegroundColor(FSlateColor::UseForeground())
+				.Visibility(this, &SPCGEditorGraphPinNumSlider::GetDefaultValueVisibility)
+				.IsEnabled(this, &SPCGEditorGraphPinNumSlider::GetDefaultValueIsEditable)
+				.Value(this, &SPCGEditorGraphPinNumSlider::GetNumericValue)
+				.MinValue(NumericPropertyParams.MinValue)
+				.MaxValue(NumericPropertyParams.MaxValue)
+				.MinSliderValue(NumericPropertyParams.MinSliderValue)
+				.MaxSliderValue(NumericPropertyParams.MaxSliderValue)
+				.SliderExponent(NumericPropertyParams.SliderExponent)
+				.Delta(NumericPropertyParams.Delta)
+				.LinearDeltaSensitivity(NumericPropertyParams.GetLinearDeltaSensitivityAttribute())
+				.AllowWheel(true)
+				.WheelStep(NumericPropertyParams.WheelStep)
+				.AllowSpin(true)
+				.OnValueCommitted(this, &SPCGEditorGraphPinNumSlider::OnValueCommitted)
+				.OnValueChanged(this, &SPCGEditorGraphPinNumSlider::OnValueChanged)
+				.OnBeginSliderMovement(this, &SPCGEditorGraphPinNumSlider::OnBeginSliderMovement)
+				.OnEndSliderMovement(this, &SPCGEditorGraphPinNumSlider::OnEndSliderMovement)
+			];
+	}
+
+private:
+	//~ End SGraphPin Interface
+
+	void OnValueChanged(NumericType NewValue)
+	{
+		SliderValue = NewValue;
+	}
+
+	void OnValueCommitted(NumericType NewValue, ETextCommit::Type CommitInfo)
+	{
+		if (GraphPinObj->IsPendingKill())
+		{
+			return;
+		}
+
+		if (CachedValue != NewValue)
+		{
+			CachedValue = NewValue;
+			const FScopedTransaction Transaction(NSLOCTEXT("PCGGraphEditor", "ChangePinNumberValue", "Change Pin Number Value"));
+			GraphPinObj->Modify();
+			OnModifyDelegate.ExecuteIfBound();
+			GraphPinObj->GetSchema()->TrySetDefaultValue(*GraphPinObj, *LexToSanitizedString(NewValue));
+		}
+	}
+
+	/**
+	* Called when the slider begins to move.
+	*/
+	void OnBeginSliderMovement()
+	{
+		GEditor->BeginTransaction(NSLOCTEXT("PCGGraphEditor", "ChangeNumberPinValueSlider", "Change Number Pin Value Slider"));
+		GraphPinObj->Modify();
+		SliderValue = GetNumericValue().GetValue();
+		bIsUsingSlider = true;
+	}
+
+	/**
+	* Called when the slider stops moving.
+	*/
+	void OnEndSliderMovement(NumericType NewValue)
+	{
+		bIsUsingSlider = false;
+		GEditor->EndTransaction();
+	}
+
+	TOptional<NumericType> GetNumericValue() const
+	{
+		NumericType Num = NumericType();
+		LexFromString(Num, *GraphPinObj->GetDefaultAsString());
+		return bIsUsingSlider ? SliderValue : Num;
+	}
+
+	virtual EVisibility GetDefaultValueVisibility() const override
+	{
+		// If this is only for showing default value, always show
+		if (bOnlyShowDefaultValue)
+		{
+			return EVisibility::Visible;
+		}
+
+		// First ask schema
+		UEdGraphPin* GraphPin = GetPinObj();
+
+	    if (!GraphPin)
+	    {
+	        return EVisibility::Hidden;
+	    }
+
+		const UEdGraphSchema* Schema = (!GraphPin->IsPendingKill()) ? GraphPin->GetSchema() : nullptr;
+
+		const bool bIsInputPin = (GraphPin->Direction == EGPD_Input);
+		const bool bIsHiddenBySchema = (!Schema || Schema->ShouldHidePinDefaultValue(GraphPin));
+		const bool bHiddenWhenConnected = (IsConnected() && bShouldHideWhenConnected);
+		const bool bIsNotConnectableOrOrphaned = (GraphPin->bNotConnectable && !GraphPin->bOrphanedPin);
+
+		if ((bIsInputPin && !bIsHiddenBySchema && !bHiddenWhenConnected) || bIsNotConnectableOrOrphaned)
+		{
+			return EVisibility::Visible;
+		}
+		else
+		{
+			return EVisibility::Hidden;
+		}
+	}
+
+	NumericType CachedValue = 0;
+	NumericType SliderValue = 0;
+	float MinDesiredBoxWidth = 0;
+	float MaxDesiredBoxWidth = 0;
+	bool bIsUsingSlider = false;
+	bool bShouldHideWhenConnected = false;
+	FSimpleDelegate OnModifyDelegate;
+};
diff --git a/Engine/Plugins/PCG/Source/PCGEditor/Private/Pins/SPCGEditorGraphPinString.h b/Engine/Plugins/PCG/Source/PCGEditor/Private/Pins/SPCGEditorGraphPinString.h
new file mode 100644
--- /dev/null
+++ b/Engine/Plugins/PCG/Source/PCGEditor/Private/Pins/SPCGEditorGraphPinString.h
@@ -0,0 +1,87 @@
+// Copyright Epic Games, Inc. All Rights Reserved.
+
+#pragma once
+
+#include "SPCGEditorGraphNodePin.h"
+
+#include "ScopedTransaction.h"
+
+#include "SSearchableComboBox.h"
+#include "Widgets/Input/SEditableTextBox.h"
+
+/** Note: Derived from Engine/Source/Editor/GraphEditor/Public/KismetPins/SGraphPinString.h */
+
+class SPCGEditorGraphPinString final : public SPCGEditorGraphNodePin
+{
+public:
+	SLATE_BEGIN_ARGS(SPCGEditorGraphPinString)
+		: _MinDesiredBoxWidth(60.f)
+		, _MaxDesiredBoxWidth(400.f)
+	{}
+
+	SLATE_ARGUMENT(float, MinDesiredBoxWidth)
+	SLATE_ARGUMENT(float, MaxDesiredBoxWidth)
+
+	SLATE_END_ARGS()
+
+	void Construct(const FArguments& InArgs, UEdGraphPin* InPin, TDelegate<void()>&& OnModify)
+	{
+		OnModifyDelegate = std::move(OnModify);
+		MinDesiredBoxWidth = InArgs._MinDesiredBoxWidth;
+		MaxDesiredBoxWidth = InArgs._MaxDesiredBoxWidth;
+
+		SPCGEditorGraphNodePin::Construct(SPCGEditorGraphNodePin::FArguments(), InPin);
+	}
+
+	//~ Begin SGraphPin Interface
+	virtual TSharedRef<SWidget> GetDefaultValueWidget() override
+	{
+		return SNew(SBox)
+			.MinDesiredWidth(MinDesiredBoxWidth)
+			.MaxDesiredWidth(MaxDesiredBoxWidth)
+			[
+				SNew(SEditableTextBox)
+				.Style(FAppStyle::Get(), "Graph.EditableTextBox")
+				.Text(this, &SPCGEditorGraphPinString::GetTypeInValue)
+				.SelectAllTextWhenFocused(true)
+				.Visibility(this, &SGraphPin::GetDefaultValueVisibility)
+				.IsReadOnly(this, &SPCGEditorGraphPinString::GetDefaultValueIsReadOnly)
+				.OnTextCommitted(this, &SPCGEditorGraphPinString::SetTypeInValue)
+				.ForegroundColor(FSlateColor::UseForeground())
+			];
+	}
+
+	//~ End SGraphPin Interface
+
+	FText GetTypeInValue() const
+	{
+		return FText::FromString(GraphPinObj->GetDefaultAsString());
+	}
+
+	void SetTypeInValue(const FText& NewTypeInValue, ETextCommit::Type CommitInfo) const
+	{
+		if (GraphPinObj->IsPendingKill())
+		{
+			return;
+		}
+
+		if (!GraphPinObj->GetDefaultAsString().Equals(NewTypeInValue.ToString()))
+		{
+			const FScopedTransaction Transaction(NSLOCTEXT("PCGGraphEditor", "ChangeStringPinValue", "Change String Pin Value"));
+			GraphPinObj->Modify();
+			OnModifyDelegate.ExecuteIfBound();
+			GraphPinObj->GetSchema()->TrySetDefaultValue(*GraphPinObj, NewTypeInValue.ToString());
+		}
+	}
+
+	bool GetDefaultValueIsReadOnly() const
+	{
+		return GraphPinObj->bDefaultValueIsReadOnly;
+	}
+
+private:
+	TArray<TSharedPtr<FString>> ComboBoxOptions;
+	float MinDesiredBoxWidth = 0;
+	float MaxDesiredBoxWidth = 0;
+	FSimpleDelegate OnModifyDelegate;
+};
diff --git a/Engine/Plugins/PCG/Source/PCGEditor/Private/Pins/SPCGEditorGraphPinVectorSlider.h b/Engine/Plugins/PCG/Source/PCGEditor/Private/Pins/SPCGEditorGraphPinVectorSlider.h
new file mode 100644
--- /dev/null
+++ b/Engine/Plugins/PCG/Source/PCGEditor/Private/Pins/SPCGEditorGraphPinVectorSlider.h
@@ -0,0 +1,334 @@
+// Copyright Epic Games, Inc. All Rights Reserved.
+
+#pragma once
+
+#include "SPCGEditorGraphNodePin.h"
+
+#include "Metadata/PCGMetadataAttributeTraits.h"
+
+#include "ScopedTransaction.h"
+
+#include "KismetPins/SVectorSlider.h"
+#include "KismetPins/SVector2DSlider.h"
+#include "KismetPins/SVector4Slider.h"
+
+/** Note: Derived from Engine/Source/Editor/GraphEditor/Public/KismetPins/SVectorSlider.h */
+
+namespace PCGEditorGraphPinVectorSlider
+{
+	template <typename T>
+	struct TVectorTrait;
+
+	template <typename T>
+	struct TVectorTrait<UE::Math::TRotator<T>>;
+
+	template <typename T>
+	struct TVectorTrait<UE::Math::TVector<T>>;
+
+	template <typename T>
+	struct TVectorTrait<UE::Math::TVector2<T>>;
+
+	template <typename T>
+	struct TVectorTrait<UE::Math::TVector4<T>>;
+}
+
+template <typename T>
+class SPCGEditorGraphPinVectorSlider final : public SPCGEditorGraphNodePin
+{
+	using TVectorTrait = PCGEditorGraphPinVectorSlider::TVectorTrait<T>;
+	static_assert(TVectorTrait::value);
+	using DataType = typename TVectorTrait::DataType;
+
+public:
+	SLATE_BEGIN_ARGS(SPCGEditorGraphPinVectorSlider)
+	{}
+
+		// SLATE_EVENT(FSimpleDelegate, OnModify)
+
+	SLATE_END_ARGS()
+
+	void Construct(const FArguments& InArgs, UEdGraphPin* InPin, TDelegate<void()>&& OnModify)
+	{
+		OnModifyDelegate = std::move(OnModify);
+		SGraphPin::Construct(SGraphPin::FArguments(), InPin);
+	}
+
+protected:
+	TSharedRef<SWidget> GetDefaultValueWidget() override
+	{
+		if constexpr (std::is_same_v<T, UE::Math::TVector<DataType>> || std::is_same_v<T, UE::Math::TRotator<DataType>>)
+		{
+			return SNew(SVectorSlider<DataType>, TVectorTrait::bIsRotator, nullptr)
+				.VisibleText_0(this, &SPCGEditorGraphPinVectorSlider::GetCurrentStringValue_0)
+				.VisibleText_1(this, &SPCGEditorGraphPinVectorSlider::GetCurrentStringValue_1)
+				.VisibleText_2(this, &SPCGEditorGraphPinVectorSlider::GetCurrentStringValue_2)
+				.OnNumericCommitted_Box_0(this, &SPCGEditorGraphPinVectorSlider::OnCommittedValueTextBox_0)
+				.OnNumericCommitted_Box_1(this, &SPCGEditorGraphPinVectorSlider::OnCommittedValueTextBox_1)
+				.OnNumericCommitted_Box_2(this, &SPCGEditorGraphPinVectorSlider::OnCommittedValueTextBox_2)
+				.Visibility(this, &SGraphPin::GetDefaultValueVisibility)
+				.IsEnabled(this, &SGraphPin::GetDefaultValueIsEditable);
+		}
+		else if constexpr (std::is_same_v<T, UE::Math::TVector2<DataType>>)
+		{
+			return SNew(SVector2DSlider<DataType>, nullptr)
+				.VisibleText_X(this, &SPCGEditorGraphPinVectorSlider::GetCurrentStringValue_0)
+				.VisibleText_Y(this, &SPCGEditorGraphPinVectorSlider::GetCurrentStringValue_1)
+				.OnNumericCommitted_Box_X(this, &SPCGEditorGraphPinVectorSlider::OnCommittedValueTextBox_0)
+				.OnNumericCommitted_Box_Y(this, &SPCGEditorGraphPinVectorSlider::OnCommittedValueTextBox_1)
+				.Visibility(this, &SGraphPin::GetDefaultValueVisibility)
+				.IsEnabled(this, &SGraphPin::GetDefaultValueIsEditable);
+		}
+		else if constexpr (std::is_same_v<T, UE::Math::TVector4<DataType>>)
+		{
+			return SNew(SVector4Slider<DataType>, nullptr)
+				.VisibleText_0(this, &SPCGEditorGraphPinVectorSlider::GetCurrentStringValue_0)
+				.VisibleText_1(this, &SPCGEditorGraphPinVectorSlider::GetCurrentStringValue_1)
+				.VisibleText_2(this, &SPCGEditorGraphPinVectorSlider::GetCurrentStringValue_2)
+				.VisibleText_3(this, &SPCGEditorGraphPinVectorSlider::GetCurrentStringValue_3)
+				.OnNumericCommitted_Box_0(this, &SPCGEditorGraphPinVectorSlider::OnCommittedValueTextBox_0)
+				.OnNumericCommitted_Box_1(this, &SPCGEditorGraphPinVectorSlider::OnCommittedValueTextBox_1)
+				.OnNumericCommitted_Box_2(this, &SPCGEditorGraphPinVectorSlider::OnCommittedValueTextBox_2)
+				.OnNumericCommitted_Box_3(this, &SPCGEditorGraphPinVectorSlider::OnCommittedValueTextBox_3)
+				.Visibility(this, &SGraphPin::GetDefaultValueVisibility)
+				.IsEnabled(this, &SGraphPin::GetDefaultValueIsEditable);
+		}
+		else
+		{
+			checkNoEntry();
+			return SNullWidget::NullWidget;
+		}
+	}
+
+private:
+	enum EVectorAxis
+	{
+		Axis_0,
+		Axis_1,
+		Axis_2,
+		Axis_3
+	};
+
+	// Text Box 0: Rotator->Roll, Vector->X
+	FString GetCurrentStringValue_0() const
+	{
+		// Update from default value if needed: ex. Undo/Redo.
+		UpdateFromDefaultValue();
+		return GetComponentStringAlongAxis(Axis_0);
+	}
+
+	// Text Box 1: Rotator->Pitch, Vector->Y
+	FString GetCurrentStringValue_1() const
+	{
+		return GetComponentStringAlongAxis(Axis_1);
+	}
+
+	// Text Box 2: Rotator->Yaw, Vector->Z
+	FString GetCurrentStringValue_2() const
+	{
+		return GetComponentStringAlongAxis(Axis_2);
+	}
+
+	// Text Box 3: FVector4 only
+	FString GetCurrentStringValue_3() const
+	{
+		static_assert(std::is_same_v<FVector4, T>);
+		return GetComponentStringAlongAxis(Axis_3);
+	}
+
+	void OnCommittedValueTextBox_0(DataType NewValue, const ETextCommit::Type CommitInfo)
+	{
+		// TODO: Should work in tandem with OnLoseFocus and (CommitInfo != ETextCommit::OnUserMovedFocus)
+		SetNewComponentValue(Axis_0, NewValue, true);
+	}
+
+	void OnCommittedValueTextBox_1(DataType NewValue, const ETextCommit::Type CommitInfo)
+	{
+		SetNewComponentValue(Axis_1, NewValue, true);
+	}
+
+	void OnCommittedValueTextBox_2(DataType NewValue, const ETextCommit::Type CommitInfo)
+	{
+		SetNewComponentValue(Axis_2, NewValue, true);
+	}
+
+	void OnCommittedValueTextBox_3(DataType NewValue, const ETextCommit::Type CommitInfo)
+	{
+		static_assert(std::is_same_v<FVector4, T>);
+		SetNewComponentValue(Axis_3, NewValue, true);
+	}
+
+	void UpdateFromDefaultValue() const
+	{
+		if (GraphPinObj->IsPendingKill())
+		{
+			return;
+		}
+
+		const uint32 Hash = GetTypeHash(GraphPinObj->DefaultValue);
+		if (Hash == DefaultValueHash)
+		{
+			return;
+		}
+
+		DefaultValueHash = Hash;
+		if constexpr (TVectorTrait::bIsRotator)
+		{
+			CachedValue = ConvertDefaultValueStringToRotator(GetPinObj()->DefaultValue);
+		}
+		else
+		{
+			CachedValue = ConvertDefaultValueStringToVector(GetPinObj()->DefaultValue);
+		}
+	}
+
+	void SetDefaultValue() const
+	{
+		const FScopedTransaction Transaction(NSLOCTEXT("PCGGraphEditor", "ChangeVectorPinValue", "Change Vector Pin Value"));
+		GraphPinObj->Modify();
+		OnModifyDelegate.ExecuteIfBound();
+
+		//Set new default value
+		GraphPinObj->GetSchema()->TrySetDefaultValue(*GraphPinObj, GetSerializedString());
+	}
+
+	void SetNewComponentValue(const EVectorAxis Axis, const DataType NewComponentValue, const bool bUpdate)
+	{
+		if (GraphPinObj->IsPendingKill())
+		{
+			return;
+		}
+
+		DataType& CurrentComponentValue = GetComponentRefAlongAxis(CachedValue, Axis);
+		if (NewComponentValue == CurrentComponentValue)
+		{
+			return;
+		}
+
+		CurrentComponentValue = NewComponentValue;
+
+		if (bUpdate)
+		{
+			SetDefaultValue();
+		}
+	}
+
+	static double& GetComponentRefAlongAxis(T& Value, const EVectorAxis Axis)
+	{
+		if constexpr (TVectorTrait::bIsRotator)
+		{
+			switch (Axis)
+			{
+				case Axis_0:
+					return Value.Roll;
+				case Axis_1:
+					return Value.Pitch;
+				case Axis_2:
+					return Value.Yaw;
+				default:
+					checkNoEntry();
+					return Value.Roll;
+			}
+		}
+		else
+		{
+			switch (Axis)
+			{
+				case Axis_0:
+					return Value.X;
+				case Axis_1:
+					return Value.Y;
+				case Axis_2: // fall-through
+					if constexpr (!std::is_same_v<T, UE::Math::TVector2<DataType>>)
+					{
+						return Value.Z;
+					}
+				case Axis_3: // fall-through
+					if constexpr (std::is_same_v<T, UE::Math::TVector4<DataType>>)
+					{
+						return Value.W;
+					}
+				default:
+					checkNoEntry();
+					return Value.X;
+			}
+		}
+	}
+
+	static const double& GetComponentRefAlongAxis(const T& Value, const EVectorAxis Axis)
+	{
+		// No side effects on Value which is returned as a const ref, so safe to const cast.
+		return GetComponentRefAlongAxis(const_cast<T&>(Value), Axis);
+	}
+
+	FString GetComponentStringAlongAxis(const EVectorAxis Axis) const
+	{
+		const DataType& Value = GetComponentRefAlongAxis(CachedValue, Axis);
+		return FString::SanitizeFloat(Value);
+	}
+
+	FString GetSerializedString() const
+	{
+		FString Result;
+		TBaseStructure<T>::Get()->ExportText(Result, &CachedValue, /*Defaults=*/nullptr, /*OwnerObject=*/nullptr, PPF_None, /*ExportRootScope=*/nullptr);
+		return Result;
+	}
+
+	T ConvertDefaultValueStringToRotator(const FString& InDefaultValueString) const
+	{
+		static_assert(TVectorTrait::bIsRotator, "Should only be called if T = TRotator");
+
+		FString ModifiedString = InDefaultValueString;
+		// TODO: Format required for InitFromString conflicts with PropertyBag/ExportText. Investigate alternatives.
+		ModifiedString.ReplaceInline(TEXT("oll"), TEXT(""), ESearchCase::CaseSensitive);
+		ModifiedString.ReplaceInline(TEXT("aw"), TEXT(""), ESearchCase::CaseSensitive);
+		ModifiedString.ReplaceInline(TEXT("itch"), TEXT(""), ESearchCase::CaseSensitive);
+
+		T Rotator;
+		Rotator.InitFromString(ModifiedString);
+		return Rotator;
+	}
+
+	T ConvertDefaultValueStringToVector(const FString& InDefaultValueString) const
+	{
+		static_assert(!TVectorTrait::bIsRotator, "Should only be called if T = TVector, TVector2, or TVector4");
+
+		T Vector;
+		Vector.InitFromString(InDefaultValueString);
+		return Vector;
+	}
+
+	mutable T CachedValue = PCG::Private::MetadataTraits<T>::ZeroValue();
+	mutable uint32 DefaultValueHash = 0;
+	FSimpleDelegate OnModifyDelegate;
+};
+
+namespace PCGEditorGraphPinVectorSlider
+{
+	template <typename T>
+	struct TVectorTrait<UE::Math::TVector<T>> : std::true_type
+	{
+		using DataType = T;
+		static constexpr bool bIsRotator = false;
+	};
+
+	template <typename T>
+	struct TVectorTrait<UE::Math::TVector2<T>> : TVectorTrait<UE::Math::TVector<T>>
+	{};
+
+	template <typename T>
+	struct TVectorTrait<UE::Math::TVector4<T>> : TVectorTrait<UE::Math::TVector<T>>
+	{};
+
+	template <typename T>
+	struct TVectorTrait<UE::Math::TRotator<T>> : TVectorTrait<UE::Math::TVector<T>>
+	{
+		static constexpr bool bIsRotator = true;
+	};
+
+	template <typename>
+	struct TVectorTrait : std::false_type
+	{
+		using DataType = void;
+		static constexpr bool bIsRotator = false;
+	};
+}
