[PCG] Fix Increment CI issue with inline constants

#rnx

diff --git a/Engine/Plugins/PCG/Source/PCG/Private/Elements/Metadata/PCGMetadataRotatorOpElement.cpp b/Engine/Plugins/PCG/Source/PCG/Private/Elements/Metadata/PCGMetadataRotatorOpElement.cpp
--- a/Engine/Plugins/PCG/Source/PCG/Private/Elements/Metadata/PCGMetadataRotatorOpElement.cpp
+++ b/Engine/Plugins/PCG/Source/PCG/Private/Elements/Metadata/PCGMetadataRotatorOpElement.cpp
@@ -237,7 +237,8 @@ void UPCGMetadataRotatorSettings::ApplyPreconfiguredSettings(const FPCGPreConfig
 bool UPCGMetadataRotatorSettings::IsPinDefaultValueMetadataTypeValid(FName PinLabel, const EPCGMetadataTypes DataType) const
 {
 	return DataType == EPCGMetadataTypes::Rotator
-		|| DataType == EPCGMetadataTypes::Vector;
+		|| DataType == EPCGMetadataTypes::Vector
+		|| DataType == EPCGMetadataTypes::Double;
 }
 
 EPCGMetadataTypes UPCGMetadataRotatorSettings::GetPinInitialDefaultValueType(const FName PinLabel) const
