[PCG] Remove ensure from inline constants preventing CI

#rnx

diff --git a/Engine/Plugins/PCG/Source/PCG/Private/Elements/Metadata/PCGMetadataOpElementBase.cpp b/Engine/Plugins/PCG/Source/PCG/Private/Elements/Metadata/PCGMetadataOpElementBase.cpp
--- a/Engine/Plugins/PCG/Source/PCG/Private/Elements/Metadata/PCGMetadataOpElementBase.cpp
+++ b/Engine/Plugins/PCG/Source/PCG/Private/Elements/Metadata/PCGMetadataOpElementBase.cpp
@@ -381,9 +381,6 @@ FString UPCGMetadataSettingsBase::GetPinDefaultValueAsString(const FName PinLabe
 
 EPCGMetadataTypes UPCGMetadataSettingsBase::GetPinInitialDefaultValueType(const FName PinLabel) const
 {
-#if WITH_EDITOR
-	ensure(IsPinDefaultValueMetadataTypeValid(PinLabel, EPCGMetadataTypes::Double));
-#endif // WITH_EDITOR
 	return EPCGMetadataTypes::Double;
 }
 
