[PCG] Fix starving of PCGAsync tasks when using low core limit (-corelimit=6)
- Split thread budget so that executing threads is just a subset and doesn't end up starving PCGAsync because of non timesliced nodes

#rnx

diff --git a/Engine/Plugins/PCG/Source/PCG/Private/Graph/PCGGraphExecutor.cpp b/Engine/Plugins/PCG/Source/PCG/Private/Graph/PCGGraphExecutor.cpp
--- a/Engine/Plugins/PCG/Source/PCG/Private/Graph/PCGGraphExecutor.cpp
+++ b/Engine/Plugins/PCG/Source/PCG/Private/Graph/PCGGraphExecutor.cpp
@@ -52,6 +52,11 @@ static TAutoConsoleVariable<float> CVarMaxPercentageOfThreadsToUse(
 	0.9f,
 	TEXT("Maximum percentage of number of threads for concurrent PCG processing"));
 
+static TAutoConsoleVariable<float> CVarMaxPercentageOfExecutingThreads(
+	TEXT("pcg.MaxPercentageOfExecutingThreads"),
+	0.3f,
+	TEXT("Maximum percentage of number of threads for processing scheduled pcg elements (percentage of pcg.MaxPercentageOfThreadsToUse)"));
+
 namespace PCGGraphExecutor
 {
 	TAutoConsoleVariable<float> CVarTimePerFrame(
@@ -86,11 +91,6 @@ namespace PCGGraphExecutor
 		true,
 		TEXT("Controls whether tasks are culled at execution time, for example in response to an deactivated dynamic branch pin"));
 
-	TAutoConsoleVariable<int32> CVarMaxWorkerTasks(
-		TEXT("pcg.Graph.MaxWorkerTasks"),
-		32,
-		TEXT("Max in flight PCG Element tasks allowed at the same time. Note that PCG Element execution can spawn async tasks themselves"));
-
 	TAutoConsoleVariable<bool> CVarPassGPUDataThroughGridLinks(
 		TEXT("pcg.Graph.GPU.PassGPUDataThroughGridLinks"),
 		true,
@@ -1353,7 +1353,18 @@ bool FPCGGraphExecutor::ExecuteScheduling(double EndTime, TSharedPtr<FPCGGraphAc
 	TRACE_CPUPROFILER_EVENT_SCOPE(FPCGGraphExecutor::ExecuteScheduling);
 	
 	const float MaxPercentageOfThreadsToUse = FMath::Clamp(CVarMaxPercentageOfThreadsToUse.GetValueOnAnyThread(), 0.0f, 1.0f);
-	const int32 MaxNumThreads = FMath::Max(0, FMath::Min((int32)(FPlatformMisc::NumberOfCoresIncludingHyperthreads() * MaxPercentageOfThreadsToUse), CVarMaxNumTasks.GetValueOnAnyThread() - 1));
+
+	// Clamp to 0.8 to never allow starving of PCGAsync Tasks
+	const float MaxPercentageOfExecutingThreads = FMath::Clamp(CVarMaxPercentageOfExecutingThreads.GetValueOnAnyThread(), 0.0f, 0.8f);
+
+	// Number of threads to use across PCG
+	const int32 MaxNumThreads = FMath::Max(0, (int32)(FPlatformMisc::NumberOfCoresIncludingHyperthreads() * MaxPercentageOfThreadsToUse));
+
+	// Number of threads to use for PCG Element execution (outside of main thread tick)
+	const int32 MaxExecutingThreads = PCGGraphExecutor::CVarGraphMultithreading.GetValueOnAnyThread() ? FMath::Max(0, (int32)(MaxNumThreads * MaxPercentageOfExecutingThreads)) : 0;
+
+	// Number of threads to use for PCG Async operations 
+	const int32 MaxPCGAsyncThreads = FMath::Min(MaxNumThreads - MaxExecutingThreads, CVarMaxNumTasks.GetValueOnAnyThread());
 
 	TArray<FCachedResult*> CachedResults;
 	{
@@ -1473,7 +1484,7 @@ bool FPCGGraphExecutor::ExecuteScheduling(double EndTime, TSharedPtr<FPCGGraphAc
 						check (!ActiveTask->bIsExecutingTask && !ActiveTask->Context->bIsPaused)
 						
 						ActiveTask->StartExecuting();
-						ActiveTask->Context->AsyncState.NumAvailableTasks = FMath::Max(1, MaxNumThreads);
+						ActiveTask->Context->AsyncState.NumAvailableTasks = FMath::Max(1, MaxPCGAsyncThreads);
 						ActiveTask->Context->AsyncState.bIsRunningOnMainThread = true;
 						ActiveTask->Context->AsyncState.bIsRunningOutOfTick = false;
 						ActiveTask->Context->AsyncState.EndTime = EndTime;
@@ -1483,15 +1494,15 @@ bool FPCGGraphExecutor::ExecuteScheduling(double EndTime, TSharedPtr<FPCGGraphAc
 				}
 			}
 
-			// PCGGraphExecutor::CVarMaxWorkerTasks.GetValueOnAnyThread()+1 because the mainthread doesn't count
-			if (PCGGraphExecutor::CVarGraphMultithreading.GetValueOnAnyThread() && PCGGraphExecutor::CVarMaxWorkerTasks.GetValueOnAnyThread()+1 > FPCGGraphActiveTask::NumExecuting && MaxNumThreads > 0)
+			// Consider the Main Thread task (+1)
+			if (MaxExecutingThreads > 0 && (MaxExecutingThreads+1) > FPCGGraphActiveTask::NumExecuting)
 			{
 				TRACE_CPUPROFILER_EVENT_SCOPE(FPCGGraphExecutor::ExecuteScheduling::LaunchTasks);
 				for (int32 ExecutionIndex = 0; ExecutionIndex < ActiveTasks.Num(); ++ExecutionIndex)
 				{
 					TSharedPtr<FPCGGraphActiveTask>& ActiveTask = ActiveTasks[ExecutionIndex];
 
-					if (PCGGraphExecutor::CVarMaxWorkerTasks.GetValueOnAnyThread()+1 <= FPCGGraphActiveTask::NumExecuting)
+					if ((MaxExecutingThreads+1) <= FPCGGraphActiveTask::NumExecuting)
 					{
 						break;
 					}
@@ -1513,7 +1524,7 @@ bool FPCGGraphExecutor::ExecuteScheduling(double EndTime, TSharedPtr<FPCGGraphAc
 					{
 						check(ActiveTask->Context->CurrentPhase != EPCGExecutionPhase::Done);
 						check(MaxNumThreads > 0);
-						ActiveTask->Context->AsyncState.NumAvailableTasks = MaxNumThreads;
+						ActiveTask->Context->AsyncState.NumAvailableTasks = FMath::Max(1, MaxPCGAsyncThreads);
 						ActiveTask->Context->AsyncState.EndTime = EndTime;
 						ActiveTask->Context->AsyncState.bIsRunningOnMainThread = false;
 
