[PCG] Trivial Sleeping -> Paused rename to match FPCGContext::bPaused terminology

#rnx

diff --git a/Engine/Plugins/PCG/Source/PCG/Private/Graph/PCGGraphExecutor.cpp b/Engine/Plugins/PCG/Source/PCG/Private/Graph/PCGGraphExecutor.cpp
--- a/Engine/Plugins/PCG/Source/PCG/Private/Graph/PCGGraphExecutor.cpp
+++ b/Engine/Plugins/PCG/Source/PCG/Private/Graph/PCGGraphExecutor.cpp
@@ -402,7 +402,7 @@ TSet<UPCGComponent*> FPCGGraphExecutor::Cancel(TFunctionRef<bool(TWeakObjectPtr<
 
 		GatherCancelledComponentsForActiveTasks(ActiveTasks);
 		GatherCancelledComponentsForActiveTasks(ActiveTasksGameThreadOnly);
-		GatherCancelledComponentsForActiveTasks(SleepingTasks);
+		GatherCancelledComponentsForActiveTasks(PausedTasks);
 	}
 
 	// In one instance this function was observed to return nullptr in the CancelledComponents set.
@@ -560,31 +560,31 @@ TSet<UPCGComponent*> FPCGGraphExecutor::Cancel(TFunctionRef<bool(TWeakObjectPtr<
 			}
 		}
 
-		// Remove from sleeping tasks
+		// Remove from paused tasks
 		{
-			TArray<TSharedPtr<FPCGGraphActiveTask>> CancelledSleepingTasks;
+			TArray<TSharedPtr<FPCGGraphActiveTask>> CancelledPausedTasks;
 			{
 				PCGGraphExecutor::TScopeLock ScopeLock(LiveTasksLock);
-				for (int32 SleepingTaskIndex = SleepingTasks.Num() - 1; SleepingTaskIndex >= 0; --SleepingTaskIndex)
+				for (int32 PausedTaskIndex = PausedTasks.Num() - 1; PausedTaskIndex >= 0; --PausedTaskIndex)
 				{
-					TSharedPtr<FPCGGraphActiveTask> Task = SleepingTasks[SleepingTaskIndex];
+					TSharedPtr<FPCGGraphActiveTask> Task = PausedTasks[PausedTaskIndex];
 					if (Task->Context && CancelledComponents.Contains(Task->Context->SourceComponent.Get()))
 					{
-						check(Task->TaskIndex == SleepingTaskIndex);
-						CancelledSleepingTasks.Add(Task);
-						PCGGraphExecutor::RemoveAtFromActiveTaskArrayNoLock(SleepingTasks, SleepingTaskIndex);
+						check(Task->TaskIndex == PausedTaskIndex);
+						CancelledPausedTasks.Add(Task);
+						PCGGraphExecutor::RemoveAtFromActiveTaskArrayNoLock(PausedTasks, PausedTaskIndex);
 						Task->TaskIndex = INDEX_NONE;
 					}
 				}
 			}
 
-			for (TSharedPtr<FPCGGraphActiveTask>& CancelledSleepingTask : CancelledSleepingTasks)
+			for (TSharedPtr<FPCGGraphActiveTask>& CancelledPausedTask : CancelledPausedTasks)
 			{
-				check(CancelledSleepingTask->Element);
-				CancelledSleepingTask->Element->Abort(CancelledSleepingTask->Context.Get());
+				check(CancelledPausedTask->Element);
+				CancelledPausedTask->Element->Abort(CancelledPausedTask->Context.Get());
 
-				FPCGTaskId CancelledTaskId = CancelledSleepingTask->NodeId;
-				RemoveTaskFromInputSuccessors(CancelledTaskId, CancelledSleepingTask->Inputs);
+				FPCGTaskId CancelledTaskId = CancelledPausedTask->NodeId;
+				RemoveTaskFromInputSuccessors(CancelledTaskId, CancelledPausedTask->Inputs);
 								
 				bStableCancellationSet &= !CancelNextTasks(CancelledTaskId, CancelledComponents);
 			}
@@ -627,7 +627,7 @@ bool FPCGGraphExecutor::IsAnyGraphCurrentlyExecuting() const
 
 int32 FPCGGraphExecutor::GetNonScheduledRemainingTaskCount() const
 {
-	return NewReadyTaskCount + Tasks.Num() + ReadyTasks.Num() + ActiveTasks.Num() + ActiveTasksGameThreadOnly.Num() + SleepingTasks.Num();
+	return NewReadyTaskCount + Tasks.Num() + ReadyTasks.Num() + ActiveTasks.Num() + ActiveTasksGameThreadOnly.Num() + PausedTasks.Num();
 }
 
 UPCGData* FPCGGraphExecutor::GetExecutionCacheData(FPCGTaskId InGraphExecutionTaskId, EExecutionCacheDataType InExecutionCacheDataType)
@@ -876,7 +876,7 @@ bool FPCGGraphExecutor::ProcessScheduledTasks()
 			PCGGraphExecutor::TScopeLock ChildScopeLock(LiveTasksLock);
 			PCGGraphExecutor::TScopeLock CachingResultsScopeLock(CachingResultsLock);
 			// This is a safeguard to check if we're in a stuck state
-			if (CollectGCCachingResults.Num() == 0 && ReadyTasks.Num() == 0 && ActiveTasks.Num() == 0 && ActiveTasksGameThreadOnly.Num() == 0 && SleepingTasks.Num() == 0 && Tasks.Num() > 0)
+			if (CollectGCCachingResults.Num() == 0 && ReadyTasks.Num() == 0 && ActiveTasks.Num() == 0 && ActiveTasksGameThreadOnly.Num() == 0 && PausedTasks.Num() == 0 && Tasks.Num() > 0)
 			{
 				UE_LOG(LogPCG, Error, TEXT("PCG Graph executor error: tasks are in a deadlocked state. Will drop all tasks."));
 				ClearAllTasks();
@@ -971,7 +971,7 @@ void FPCGGraphExecutor::PostTaskExecute(TSharedPtr<FPCGGraphActiveTask> ActiveTa
 
 			if (ActiveTask.Context->bIsPaused)
 			{
-				PCGGraphExecutor::AddToActiveTaskArrayNoLock(SleepingTasks, ActiveTaskPtr);
+				PCGGraphExecutor::AddToActiveTaskArrayNoLock(PausedTasks, ActiveTaskPtr);
 				PCGGraphExecutor::RemoveAtFromActiveTaskArrayNoLock(ActiveTaskArrayToRemove, TaskRemoveIndex);
 			}
 			else if (ActiveTask.bIsGameThreadOnly != ActiveTask.CanExecuteOnlyOnMainThread())
@@ -983,7 +983,7 @@ void FPCGGraphExecutor::PostTaskExecute(TSharedPtr<FPCGGraphActiveTask> ActiveTa
 				PCGGraphExecutor::RemoveAtFromActiveTaskArrayNoLock(ActiveTaskArrayToRemove, TaskRemoveIndex);
 			}
 
-			// Task might have been moved to SleepingTasks or changed from ActiveTasks to ActiveTasksGameThreadOnly (or vice-versa) 
+			// Task might have been moved to PausedTasks or changed from ActiveTasks to ActiveTasksGameThreadOnly (or vice-versa) 
 			// but it isn't done so we return
 			return;
 		}
@@ -1084,8 +1084,8 @@ void FPCGGraphExecutor::PostTaskExecute(TSharedPtr<FPCGGraphActiveTask> ActiveTa
 		// Make sure we set this to false inside lock
 		ActiveTask.StopExecuting();
 
-		// Next Scheduling call needs to check if this task completion unblocked some sleeping task(s)
-		bNeedToCheckSleepingTasks = true;
+		// Next Scheduling call needs to check if this task completion unblocked some paused task(s)
+		bNeedToCheckPausedTasks = true;
 	}
 }
 
@@ -1340,7 +1340,7 @@ double FPCGGraphExecutor::GetTickBudgetInSeconds() const
 	return PCGGraphExecutor::CVarTimePerFrame.GetValueOnAnyThread() / 1000.0;
 }
 
-bool FPCGGraphExecutor::ExecuteScheduling(double EndTime, TSharedPtr<FPCGGraphActiveTask>* OutMainThreadTask, bool bForceCheckSleepingTasks)
+bool FPCGGraphExecutor::ExecuteScheduling(double EndTime, TSharedPtr<FPCGGraphActiveTask>* OutMainThreadTask, bool bForceCheckPausedTasks)
 {
 	// Can currently be disabled by cancellation
 	if (SchedulingDisabled)
@@ -1372,27 +1372,27 @@ bool FPCGGraphExecutor::ExecuteScheduling(double EndTime, TSharedPtr<FPCGGraphAc
 		PCGGraphExecutor::TScopeLock ScopeLock(LiveTasksLock, !bIsInGameThread);
 		if (ScopeLock.IsLocked())
 		{
-			if (bNeedToCheckSleepingTasks || bForceCheckSleepingTasks)
+			if (bNeedToCheckPausedTasks || bForceCheckPausedTasks)
 			{
-				bNeedToCheckSleepingTasks = false;
+				bNeedToCheckPausedTasks = false;
 
-				// First, wake up any sleeping tasks that can be reactivated
-				for (int32 SleepingTaskIndex = SleepingTasks.Num() - 1; SleepingTaskIndex >= 0; --SleepingTaskIndex)
+				// First, wake up any paused tasks that can be reactivated
+				for (int32 PausedTaskIndex = PausedTasks.Num() - 1; PausedTaskIndex >= 0; --PausedTaskIndex)
 				{
-					TSharedPtr<FPCGGraphActiveTask>& SleepingTask = SleepingTasks[SleepingTaskIndex];
+					TSharedPtr<FPCGGraphActiveTask>& PausedTask = PausedTasks[PausedTaskIndex];
 
-					if (SleepingTask->Context->bIsPaused)
+					if (PausedTask->Context->bIsPaused)
 					{
-						continue; // still sleeping
+						continue;
 					}
 
-					const bool bIsGameThreadOnly = SleepingTask->CanExecuteOnlyOnMainThread();
+					const bool bIsGameThreadOnly = PausedTask->CanExecuteOnlyOnMainThread();
 
 					TArray<TSharedPtr<FPCGGraphActiveTask>>& ActiveTaskArray = bIsGameThreadOnly ? ActiveTasksGameThreadOnly : ActiveTasks;
-					SleepingTask->bIsGameThreadOnly = bIsGameThreadOnly;
+					PausedTask->bIsGameThreadOnly = bIsGameThreadOnly;
 
-					PCGGraphExecutor::AddToActiveTaskArrayNoLock(ActiveTaskArray, SleepingTask);
-					PCGGraphExecutor::RemoveAtFromActiveTaskArrayNoLock(SleepingTasks, SleepingTaskIndex);
+					PCGGraphExecutor::AddToActiveTaskArrayNoLock(ActiveTaskArray, PausedTask);
+					PCGGraphExecutor::RemoveAtFromActiveTaskArrayNoLock(PausedTasks, PausedTaskIndex);
 					bStateChanged = true;
 				}
 			}
@@ -1591,8 +1591,8 @@ bool FPCGGraphExecutor::ExecuteScheduling(double EndTime, TSharedPtr<FPCGGraphAc
  * - When an [ActiveTasks/ActiveTasksGameThreadOnly] Execute call returns and it returns true, it means that the task is done. In this case we call [FPCGGraphExecutor::StoreResults] followed by [FPCGGraphExecutor::QueueNextTasks] this will potentially move some [Tasks] to [ReadyTasks]
  *   so that new tasks can get executed now that their inputs have been filled
  * 
- * - After a call to Execute, if a [ActiveTasks/ActiveTasksGameThreadOnly] bPaused flag is true, it gets moved to [SleepingTasks]
- * - [SleepingTasks] gets processed at least once a frame to see if some tasks are no longer paused and need to be moved back to [ActiveTasks/ActiveTasksGameThreadOnly]
+ * - After a call to Execute, if a [ActiveTasks/ActiveTasksGameThreadOnly] bPaused flag is true, it gets moved to [PausedTasks]
+ * - [PausedTasks] gets processed at least once a frame to see if some tasks are no longer paused and need to be moved back to [ActiveTasks/ActiveTasksGameThreadOnly]
  * - It will also be processed again if some other [ActiveTasks/ActiveTasksGameThreadOnly] completes execution
  * 
  * - Cache information: When we call [FPCGGraphExecutor::PrepareForExecute] is called on a task, if a Cache result is found for this task then we can directly call [FPCGGraphExecutor::StoreResults] and [FPCGGraphExecutor::QueueNextTasks] on it skipping execution [FPCGGraphExecutor::ProcessCachedResults]
@@ -1636,7 +1636,7 @@ void FPCGGraphExecutor::Execute()
 	while ((bFirstLoop || (FPlatformTime::Seconds() < EndTime)) && bContinueExecute)
 	{
 		TSharedPtr<FPCGGraphActiveTask> MainThreadTask;
-		bContinueExecute = ExecuteScheduling(EndTime, &MainThreadTask, /*bForceCheckSleepingTasks=*/bFirstLoop);
+		bContinueExecute = ExecuteScheduling(EndTime, &MainThreadTask, /*bForceCheckPausedTasks=*/bFirstLoop);
 		bFirstLoop = false;
 
 		if (MainThreadTask)
@@ -1668,7 +1668,7 @@ void FPCGGraphExecutor::ClearAllTasks()
 	ReadyTasks.Reset();
 	ActiveTasks.Reset();
 	ActiveTasksGameThreadOnly.Reset();
-	SleepingTasks.Reset();
+	PausedTasks.Reset();
 }
 
 TArray<FPCGGraphExecutor::FCachedResult*> FPCGGraphExecutor::QueueNextTasksInternal(FPCGTaskId FinishedTask)
@@ -1774,8 +1774,8 @@ TArray<FPCGTaskId> FPCGGraphExecutor::ProcessCachedResultsInternal(TArray<FCache
 			CollectGCCachingResults.Remove(CachedResult->TaskId);
 		}
 
-		// Next Scheduling call needs to check if this task completion unblocked some sleeping task(s)
-		bNeedToCheckSleepingTasks = true;
+		// Next Scheduling call needs to check if this task completion unblocked some paused task(s)
+		bNeedToCheckPausedTasks = true;
 	}
 
 	return NextTasks;
@@ -2316,7 +2316,7 @@ void FPCGGraphExecutor::AddReferencedObjects(FReferenceCollector& Collector)
 
 		Algo::ForEach(ActiveTasks, AddReferencesActiveTask);
 		Algo::ForEach(ActiveTasksGameThreadOnly, AddReferencesActiveTask);
-		Algo::ForEach(SleepingTasks, AddReferencesActiveTask);
+		Algo::ForEach(PausedTasks, AddReferencesActiveTask);
 	}
 
 	{
@@ -2327,7 +2327,7 @@ void FPCGGraphExecutor::AddReferencedObjects(FReferenceCollector& Collector)
 		});
 	}
 
-	// Some task might have been delayed to the game thread and we are no longer referencing them through (Ready/Sleep/Active)
+	// Some task might have been delayed to the game thread and we are no longer referencing them through (Ready/Paused/Active)
 	// but we still want their UObject references to stay valid
 	{
 		PCGGraphExecutor::TScopeLock ScopeLock(CollectGCReferenceTasksLock);
@@ -2874,8 +2874,8 @@ namespace PCGGraphExecutor
 			{
 				PCGGraphExecutionLogging::LogGridLinkageTaskExecuteRetrieveWaitOnScheduledGraph(InContext, ComponentWithData, InResourceKey);
 
-				// The component was already generating, but we were not asleep. Not really clear what's happening here,
-				// but in any case go to sleep and wake up when it's done.
+				// The component was already generating, but we were not paused. Not really clear what's happening here,
+				// but in any case go to paused and unpause when it's done.
 				InContext->bIsPaused = true;
 
 				// Wake up this task after graph has generated.
@@ -2901,7 +2901,7 @@ namespace PCGGraphExecutor
 				const FPCGTaskId GraphTaskId = ComponentWithData->GenerateLocalGetTaskId(GenTrigger, /*bForce=*/true);
 				Subsystem->ScheduleGeneric(WakeUpLambda, InContext->SourceComponent.Get(), { GraphTaskId });
 
-				// Update state and go to sleep.
+				// Update state and go to paused.
 				InContext->bScheduledGraph = true;
 				InContext->bIsPaused = true;
 				return false;
diff --git a/Engine/Plugins/PCG/Source/PCG/Private/Graph/PCGGraphExecutor.h b/Engine/Plugins/PCG/Source/PCG/Private/Graph/PCGGraphExecutor.h
--- a/Engine/Plugins/PCG/Source/PCG/Private/Graph/PCGGraphExecutor.h
+++ b/Engine/Plugins/PCG/Source/PCG/Private/Graph/PCGGraphExecutor.h
@@ -212,7 +212,7 @@ private:
 	void PostTaskExecute(TSharedPtr<FPCGGraphActiveTask> ActiveTask, bool bIsDone);
 	bool ProcessScheduledTasks();
 	void ExecuteTasksEnded();
-	bool ExecuteScheduling(double EndTime, TSharedPtr<FPCGGraphActiveTask>* OutMainThreadTask = nullptr, bool bForceCheckSleepingTasks = false);
+	bool ExecuteScheduling(double EndTime, TSharedPtr<FPCGGraphActiveTask>* OutMainThreadTask = nullptr, bool bForceCheckPausedTasks = false);
 
 	TSet<UPCGComponent*> Cancel(TFunctionRef<bool(TWeakObjectPtr<UPCGComponent>)> CancelFilter);
 	void ClearAllTasks();
@@ -302,8 +302,8 @@ private:
 	TArray<FPCGGraphTask> ReadyTasks;
 	TArray<TSharedPtr<FPCGGraphActiveTask>> ActiveTasks;
 	TArray<TSharedPtr<FPCGGraphActiveTask>> ActiveTasksGameThreadOnly;
-	TArray<TSharedPtr<FPCGGraphActiveTask>> SleepingTasks;
-	bool bNeedToCheckSleepingTasks = false;
+	TArray<TSharedPtr<FPCGGraphActiveTask>> PausedTasks;
+	bool bNeedToCheckPausedTasks = false;
 
 	/** Lock level 3 */
 	FPCGLock CollectGCReferenceTasksLock;
