Fix crash if FPCGInternalHelpers::ScheduleOnExecutorAndWait times out.

diff --git a/Engine/Plugins/PCG/Source/PCG/Private/Helpers/PCGInternalHelpers.cpp b/Engine/Plugins/PCG/Source/PCG/Private/Helpers/PCGInternalHelpers.cpp
--- a/Engine/Plugins/PCG/Source/PCG/Private/Helpers/PCGInternalHelpers.cpp
+++ b/Engine/Plugins/PCG/Source/PCG/Private/Helpers/PCGInternalHelpers.cpp
@@ -61,12 +61,24 @@ bool UE::FPCGInternalHelpers::ScheduleOnExecutorAndWait(FPCGGraphExecutor& Execu
 		/*bAllowHierarchicalGeneration*/ false);
 	PCGComponent->CurrentGenerationTask = TaskId;
 
-	bool bDone = false;
-	const FPCGTaskId FinalTaskId = FPCGGraphExecutorInternalAccess::ScheduleSingleElement(Executor,
-		MakeShared<FPCGGenericElementAnyThread>([&Executor, &OutputData, TaskId, &bDone](FPCGContext*) -> bool
+	// Schedule a task to get the data out and mark as completed when done. Pass a weak pointer
+	//  that we'll fill out if we didn't time out.
+	struct FOutputState
+	{
+		FPCGDataCollection Output;
+		bool bDone = false;
+	};
+	TSharedPtr<FOutputState> OutputState = MakeShared<FOutputState>();
+
+	FPCGGraphExecutorInternalAccess::ScheduleSingleElement(Executor,
+		MakeShared<FPCGGenericElementAnyThread>([&Executor, OutputStateWeak = OutputState.ToWeakPtr(), TaskId](FPCGContext*) -> bool
 		{
-			Executor.GetOutputData(TaskId, OutputData);
-			bDone = true;
+			TSharedPtr<FOutputState> OutputStatePinned = OutputStateWeak.Pin();
+			if (OutputStatePinned.IsValid())
+			{
+				Executor.GetOutputData(TaskId, OutputStatePinned->Output);
+				OutputStatePinned->bDone = true;
+			}
 			return true;
 		}), 
 		false, PCGComponent, { TaskId }, {});
@@ -75,7 +87,7 @@ bool UE::FPCGInternalHelpers::ScheduleOnExecutorAndWait(FPCGGraphExecutor& Execu
 	CyclesToWait = FMath::Max(0, CyclesToWait);
 	SleepDurationPerCycle = FMath::Max(0, SleepDurationPerCycle);
 	
-	while (!bDone && CycleCount < CyclesToWait)
+	while (!OutputState->bDone && CycleCount < CyclesToWait)
 	{
 		if (SleepDurationPerCycle > 0)
 		{
@@ -83,5 +95,14 @@ bool UE::FPCGInternalHelpers::ScheduleOnExecutorAndWait(FPCGGraphExecutor& Execu
 		}
 		++CycleCount;
 	}
-	return bDone;
+	if (OutputState->bDone)
+	{
+		OutputData = OutputState->Output;
+		return true;
+	}
+	// TODO: If we fail, we should go ahead and cancel this execution (unless we decide to not time out at all,
+	//  once we are able to use the proper PCG executor?). However this would require a threadsafe way to cancel,
+	//  since this function gets called from worker threads.
+
+	return false;
 }
\ No newline at end of file
