[PCG] Fix InitializeAsCopy parenting

#rnx

diff --git a/Engine/Plugins/PCG/Source/PCG/Private/Metadata/PCGMetadata.cpp b/Engine/Plugins/PCG/Source/PCG/Private/Metadata/PCGMetadata.cpp
--- a/Engine/Plugins/PCG/Source/PCG/Private/Metadata/PCGMetadata.cpp
+++ b/Engine/Plugins/PCG/Source/PCG/Private/Metadata/PCGMetadata.cpp
@@ -385,15 +385,31 @@ void UPCGMetadata::InitializeAsCopyWithAttributeFilter(const UPCGMetadata* InMet
 		return;
 	}
 
+	// If any metadata domain sets its parent, we need to parent this metadata to the metadata to copy, to preserve the hierarchy.
+	// But in case of partial copy, the parenting is not necessary.
+	bool bShouldParent = false;
+
 	SetupDomainsFromOtherMetadataIfNeeded(InMetadataToCopy);
 		
 	TMap<FPCGMetadataDomainID, TSet<FName>> FilteredAttributeSplit = FPCGAttributeIdentifier::SplitByDomain(InFilteredAttributes);
-	InMetadataToCopy->ForEachValidUniqueConstMetadataDomain([this, &FilteredAttributeSplit, InFilterMode, &EntriesToCopy](const FPCGMetadataDomainID& DomainID, const FPCGMetadataDomain& OtherMetadataDomain)
+	InMetadataToCopy->ForEachValidUniqueConstMetadataDomain([this, &FilteredAttributeSplit, InFilterMode, &EntriesToCopy, &bShouldParent, InMetadataToCopy](const FPCGMetadataDomainID& DomainID, const FPCGMetadataDomain& OtherMetadataDomain)
 	{
 		TSet<FName>* FilteredAttributesPtr = FilteredAttributeSplit.Find(DomainID);
 		FPCGMetadataDomain* CurrentMetadataDomain = FindOrCreateMetadataDomain(DomainID);
 		CurrentMetadataDomain->InitializeAsCopyWithAttributeFilter(&OtherMetadataDomain, (FilteredAttributesPtr ? *FilteredAttributesPtr : TSet<FName>{}), InFilterMode, EntriesToCopy);
+
+		// We still validate that the parent was set correctly and it matches the metadata to copy.
+		if (CurrentMetadataDomain->Parent && CurrentMetadataDomain->Parent->TopMetadata == InMetadataToCopy)
+		{
+			bShouldParent = true;
+		}
 	});
+
+	if (bShouldParent)
+	{
+		Parent = InMetadataToCopy;
+		OtherParents = InMetadataToCopy->OtherParents;
+	}
 }
 
 void UPCGMetadata::AddAttributesFiltered(const UPCGMetadata* InOther, const TSet<FName>& InFilteredAttributes, EPCGMetadataFilterMode InFilterMode, EPCGStringMatchingOperator InMatchOperator)
