[PCG] GPU conditional copy points

diff --git a/Engine/Plugins/PCG/Shaders/Private/Elements/PCGCopyPoints.usf b/Engine/Plugins/PCG/Shaders/Private/Elements/PCGCopyPoints.usf
--- a/Engine/Plugins/PCG/Shaders/Private/Elements/PCGCopyPoints.usf
+++ b/Engine/Plugins/PCG/Shaders/Private/Elements/PCGCopyPoints.usf
@@ -37,7 +37,7 @@ void Main(uint3 GroupId : SV_GroupID, uint GroupIndex : SV_GroupIndex)
 	const uint SourceNumData = Source_GetNumData();
 	const uint TargetNumData = Target_GetNumData();
 	const uint TargetPointsCount = Target_GetNumElements();
-	const uint bCopyEachSourceOnEveryTarget = GetCopyEachSourceOnEveryTarget();
+	const uint bCopyEachSourceOnEveryTarget = CopyPoints_GetCopyEachSourceOnEveryTarget();
 
 	// Data must be non-zero, and NxM (when bCopyEachSourceOnEveryTarget == true), N:N, 1:N, or N:1
 	if (SourceNumData == 0 || TargetNumData == 0 || (!bCopyEachSourceOnEveryTarget && SourceNumData != TargetNumData && SourceNumData != 1 && TargetNumData != 1))
@@ -135,13 +135,13 @@ void Main(uint3 GroupId : SV_GroupID, uint GroupIndex : SV_GroupIndex)
 	float4 TargetColor = Target_GetColor(TargetDataIndex, TargetPointIndexInData);
 	uint TargetSeed = Target_GetSeed(TargetDataIndex, TargetPointIndexInData);
 
-	const uint RotationInheritance = GetRotationInheritance();
-	const uint ApplyTargetRotationToPositions = GetApplyTargetRotationToPositions();
-	const uint ScaleInheritance = GetScaleInheritance();
-	const uint ApplyTargetScaleToPositions = GetApplyTargetScaleToPositions();
-	const uint ColorInheritance = GetColorInheritance();
-	const uint SeedInheritance = GetSeedInheritance();
-	const uint AttributeInheritance = GetAttributeInheritance();
+	const uint RotationInheritance = CopyPoints_GetRotationInheritance();
+	const uint ApplyTargetRotationToPositions = CopyPoints_GetApplyTargetRotationToPositions();
+	const uint ScaleInheritance = CopyPoints_GetScaleInheritance();
+	const uint ApplyTargetScaleToPositions = CopyPoints_GetApplyTargetScaleToPositions();
+	const uint ColorInheritance = CopyPoints_GetColorInheritance();
+	const uint SeedInheritance = CopyPoints_GetSeedInheritance();
+	const uint AttributeInheritance = CopyPoints_GetAttributeInheritance();
 
 	float3 ScaleForPosition = ApplyTargetScaleToPositions > 0 ? TargetScale : float3(1, 1, 1);
 
diff --git a/Engine/Plugins/PCG/Shaders/Private/Elements/PCGCopyPointsAnalysis.usf b/Engine/Plugins/PCG/Shaders/Private/Elements/PCGCopyPointsAnalysis.usf
new file mode 100644
--- /dev/null
+++ b/Engine/Plugins/PCG/Shaders/Private/Elements/PCGCopyPointsAnalysis.usf
@@ -0,0 +1,68 @@
+// Copyright Epic Games, Inc. All Rights Reserved.
+
+#include "/Engine/Private/ComputeShaderUtils.ush"
+#include "/Engine/Private/Quaternion.ush"
+#include "/Plugin/PCG/Private/PCGShaderUtils.ush"
+#include "/Plugin/PCG/Private/PCGShaderUtilsInternal.ush"
+
+// TODO: should be code-gen'd from PCGComputeCommon.h
+#define MAX_NUM_ATTRS 128
+#define NUM_RESERVED_ATTRS 32
+
+[numthreads(64, 1, 1)]
+void Main(uint3 GroupId : SV_GroupID, uint GroupIndex : SV_GroupIndex)
+{
+	// Mark the kernel as having executed. Must run before we early out via thread index, because the kernel is still 'executed' even if the number of
+	// threads to iterate on is zero. Even if GetNumThreads() returns 0, the kernel will still have been dispatched on a single thread to set this flag.
+	if (all(GroupId == 0) && GroupIndex == 0)
+	{
+		Out_SetAsExecutedInternal();
+	}
+
+	const uint ThreadIndex = GetUnWrappedDispatchThreadId(GroupId, GroupIndex, 64);
+
+	const uint SourceNumData = Source_GetNumData();
+	const uint NumTargetPointsTotal = Target_GetNumElements();
+
+	if (ThreadIndex >= SourceNumData * NumTargetPointsTotal)
+	{
+		// Discard threads outside of data range.
+		return;
+	}
+
+	const int MatchAttributeId = CopyPoints_GetMatchAttributeId();
+	if (MatchAttributeId == -1)
+	{
+		// No match attribute, early out with no copy.
+		Out_SetBool(0, ThreadIndex, CopyPoints_GetSelectedFlagAttributeId(), false);
+		return;
+	}
+
+	const uint NumDataTarget = Target_GetNumData();
+
+	// todo_pcg: Needs both modes of copy points.
+
+	// Output is attribute set with SourceNumData * TargetNumData bools. Kernel executes one thread per output.
+	// copypoints(S = ((src 0), (src 1)), T = ((tgt 0), (tgt 1))) = (src 0 x tgt 0), (src 0 x tgt 1), (src 1 x tgt 0), (src 1 tgt 1))
+	const uint SourceDataIndex = ThreadIndex / NumDataTarget;
+	const uint TargetDataIndex = ThreadIndex % NumDataTarget;
+
+	if (Source_GetNumElements(SourceDataIndex) > 0 && Target_GetNumElements(TargetDataIndex) > 0)
+	{
+		// If first point is removed, assume no copy.
+		// todo_pcg: This is temporary and will use a data domain attribute once support has been added.
+		if (!Source_IsPointRemoved(SourceDataIndex, /*ElementIndex=*/0) && !Target_IsPointRemoved(TargetDataIndex, /*ElementIndex=*/0))
+		{
+			const int SourceValue = Source_GetInt(SourceDataIndex, /*ElementIndex=*/0, MatchAttributeId);
+			const int TargetValue = Target_GetInt(TargetDataIndex, /*ElementIndex=*/0, MatchAttributeId);
+
+			if (SourceValue == TargetValue)
+			{
+				Out_SetBool(/*DataIndex=*/0, ThreadIndex, CopyPoints_GetSelectedFlagAttributeId(), true);
+				return;
+			}
+		}
+	}
+
+	Out_SetBool(0, ThreadIndex, CopyPoints_GetSelectedFlagAttributeId(), false);
+}
diff --git a/Engine/Plugins/PCG/Shaders/Private/Elements/PCGCopyPointsConditional.usf b/Engine/Plugins/PCG/Shaders/Private/Elements/PCGCopyPointsConditional.usf
new file mode 100644
--- /dev/null
+++ b/Engine/Plugins/PCG/Shaders/Private/Elements/PCGCopyPointsConditional.usf
@@ -0,0 +1,234 @@
+// Copyright Epic Games, Inc. All Rights Reserved.
+
+// todo_pcg: Unify the implementation below with PCGCopyPoints.usf and remove this file.
+
+#include "/Engine/Private/ComputeShaderUtils.ush"
+#include "/Engine/Private/Quaternion.ush"
+#include "/Plugin/PCG/Private/PCGShaderUtils.ush"
+#include "/Plugin/PCG/Private/PCGShaderUtilsInternal.ush"
+
+// Mirrors EPCGCopyPointsInheritanceMode
+#define COPY_POINTS_INHERITANCE_MODE_Relative 0
+#define COPY_POINTS_INHERITANCE_MODE_Source 1
+#define COPY_POINTS_INHERITANCE_MODE_Target 2
+
+// Mirrors EPCGCopyPointsMetadataInheritanceMode
+#define COPY_POINTS_METADATA_INHERITANCE_MODE_SourceFirst 0
+#define COPY_POINTS_METADATA_INHERITANCE_MODE_TargetFirst 1
+#define COPY_POINTS_METADATA_INHERITANCE_MODE_SourceOnly 2
+#define COPY_POINTS_METADATA_INHERITANCE_MODE_TargetOnly 3
+#define COPY_POINTS_METADATA_INHERITANCE_MODE_None 4
+
+// TODO: should be code-gen'd from PCGComputeCommon.h
+#define MAX_NUM_ATTRS 128
+#define NUM_RESERVED_ATTRS 32
+
+[numthreads(64, 1, 1)]
+void Main(uint3 GroupId : SV_GroupID, uint GroupIndex : SV_GroupIndex)
+{
+	// Mark the kernel as having executed. Must run before we early out via thread index, because the kernel is still 'executed' even if the number of
+	// threads to iterate on is zero. Even if GetNumThreads() returns 0, the kernel will still have been dispatched on a single thread to set this flag.
+	if (all(GroupId == 0) && GroupIndex == 0)
+	{
+		Out_SetAsExecutedInternal();
+	}
+
+	const uint ThreadIndex = GetUnWrappedDispatchThreadId(GroupId, GroupIndex, 64);
+	if (ThreadIndex >= GetNumThreads().x) return;
+
+	uint OutDataIndex, OutElemIndex;
+	if (!Out_GetThreadData(ThreadIndex, OutDataIndex, OutElemIndex))
+	{
+		return;
+	}
+
+	// Get source and target data indices for this output data index.
+	uint SourceDataIndex, TargetDataIndex;
+	CopyPoints_GetSourceAndTargetDataIndices(OutDataIndex, SourceDataIndex, TargetDataIndex);
+
+	const uint SourceNumData = Source_GetNumData();
+	const uint TargetNumData = Target_GetNumData();
+
+	// @todo_pcg: Until we support pairwise copy, bCopyEachSourceOnEveryTarget is a tautology.
+	//const uint bCopyEachSourceOnEveryTarget = CopyPoints_GetCopyEachSourceOnEveryTarget();
+
+	// Data must be non-zero, and NxM (when bCopyEachSourceOnEveryTarget == true), N:N, 1:N, or N:1
+	if (SourceNumData == 0 || TargetNumData == 0)
+		// || (!bCopyEachSourceOnEveryTarget && SourceNumData != TargetNumData && SourceNumData != 1 && TargetNumData != 1))
+	{
+		// TODO: Runtime error diagnostics
+		return;
+	}
+	
+	uint SourcePointIndexInData = 0;
+	uint TargetPointIndexInData = 0;
+
+	// @todo_pcg: Support both cardinalities/modes. Probably best to just have CopyPoints_GetSourceAndTargetDataIndices provide the indices for both modes, i.e. for
+	// upload either (0 0, 0 1, 1 0, 1 1) or (0 0, 1 1) depending on the mode.
+	//if (bCopyEachSourceOnEveryTarget)
+	{
+		const uint NumTargetPoints = Target_GetNumElements(TargetDataIndex);
+		SourcePointIndexInData = OutElemIndex / NumTargetPoints;
+		TargetPointIndexInData = OutElemIndex % NumTargetPoints;
+	}
+	/*
+	else
+	{
+		const uint NumIterations = max(SourceNumData, TargetNumData);
+
+		bool bFoundInData = false;
+		uint TotalElementCount = 0;
+
+		for (uint DataIndex = 0; DataIndex < NumIterations; ++DataIndex)
+		{
+			SourceDataIndex = min(DataIndex, SourceNumData - 1);
+			TargetDataIndex = min(DataIndex, TargetNumData - 1);
+
+			const uint SourceElemCount = Source_GetNumElements(SourceDataIndex);
+			const uint TargetElemCount = Target_GetNumElements(TargetDataIndex);
+
+			const uint ElemCount = SourceElemCount * TargetElemCount;
+
+			if (TotalElementCount + ElemCount > ThreadIndex)
+			{
+				const uint ElementIndex = ThreadIndex - TotalElementCount;
+				SourcePointIndexInData = ElementIndex / TargetElemCount;
+				TargetPointIndexInData = ElementIndex % TargetElemCount;
+				bFoundInData = true;
+				break;
+			}
+
+			TotalElementCount += ElemCount;
+		}
+
+		if (!bFoundInData)
+		{
+			return;
+		}
+	}
+	*/
+
+	// Propagate the 'Removed' status to the output point if either source or target point is removed. Otherwise, this point will not be culled.
+	if (Source_IsPointRemoved(SourceDataIndex, SourcePointIndexInData) || Target_IsPointRemoved(TargetDataIndex, TargetPointIndexInData))
+	{
+		Out_RemovePoint(OutDataIndex, OutElemIndex);
+		return;
+	}
+
+	float3 SourcePosition = Source_GetPosition(SourceDataIndex, SourcePointIndexInData);
+	float4 SourceRotation = Source_GetRotation(SourceDataIndex, SourcePointIndexInData);
+	float3 SourceScale = Source_GetScale(SourceDataIndex, SourcePointIndexInData);
+	float4 SourceColor = Source_GetColor(SourceDataIndex, SourcePointIndexInData);
+	uint SourceSeed = Source_GetSeed(SourceDataIndex, SourcePointIndexInData);
+	
+	float3 TargetPosition = Target_GetPosition(TargetDataIndex, TargetPointIndexInData);
+	float4 TargetRotation = Target_GetRotation(TargetDataIndex, TargetPointIndexInData);
+	float3 TargetScale = Target_GetScale(TargetDataIndex, TargetPointIndexInData);
+	float4 TargetColor = Target_GetColor(TargetDataIndex, TargetPointIndexInData);
+	uint TargetSeed = Target_GetSeed(TargetDataIndex, TargetPointIndexInData);
+
+	const uint RotationInheritance = CopyPoints_GetRotationInheritance();
+	const uint ApplyTargetRotationToPositions = CopyPoints_GetApplyTargetRotationToPositions();
+	const uint ScaleInheritance = CopyPoints_GetScaleInheritance();
+	const uint ApplyTargetScaleToPositions = CopyPoints_GetApplyTargetScaleToPositions();
+	const uint ColorInheritance = CopyPoints_GetColorInheritance();
+	const uint SeedInheritance = CopyPoints_GetSeedInheritance();
+	const uint AttributeInheritance = CopyPoints_GetAttributeInheritance();
+
+	float3 ScaleForPosition = ApplyTargetScaleToPositions > 0 ? TargetScale : float3(1, 1, 1);
+
+	float3 OutPosition = TargetPosition + (ApplyTargetRotationToPositions > 0 ? QuatRotateVector(TargetRotation, (SourcePosition * ScaleForPosition)) : (SourcePosition * ScaleForPosition));
+	float4 OutRotation = float4(0, 0, 0, 1);
+	float3 OutScale = (float3)0;
+	float4 OutColor = (float4)0;
+	uint OutSeed = 0;
+
+	if (RotationInheritance == COPY_POINTS_INHERITANCE_MODE_Relative)
+	{
+		OutRotation = QuatMultiply(TargetRotation, SourceRotation);
+	}
+	else if (RotationInheritance == COPY_POINTS_INHERITANCE_MODE_Source)
+	{
+		OutRotation = SourceRotation;
+	}
+	else if (RotationInheritance == COPY_POINTS_INHERITANCE_MODE_Target)
+	{
+		OutRotation = TargetRotation;
+	}
+
+	if (ScaleInheritance == COPY_POINTS_INHERITANCE_MODE_Relative)
+	{
+		OutScale = SourceScale * TargetScale;
+	}
+	else if (ScaleInheritance == COPY_POINTS_INHERITANCE_MODE_Source)
+	{
+		OutScale = SourceScale;
+	}
+	else if (ScaleInheritance == COPY_POINTS_INHERITANCE_MODE_Target)
+	{
+		OutScale = TargetScale;
+	}
+
+	if (ColorInheritance == COPY_POINTS_INHERITANCE_MODE_Relative)
+	{
+		OutColor = SourceColor * TargetColor;
+	}
+	else if (ColorInheritance == COPY_POINTS_INHERITANCE_MODE_Source)
+	{
+		OutColor = SourceColor;
+	}
+	else if (ColorInheritance == COPY_POINTS_INHERITANCE_MODE_Target)
+	{
+		OutColor = TargetColor;
+	}
+
+	if (SeedInheritance == COPY_POINTS_INHERITANCE_MODE_Relative)
+	{
+		OutSeed = ComputeSeed(SourceSeed, TargetSeed);
+	}
+	else if (SeedInheritance == COPY_POINTS_INHERITANCE_MODE_Source)
+	{
+		OutSeed = SourceSeed;
+	}
+	else if (SeedInheritance == COPY_POINTS_INHERITANCE_MODE_Target)
+	{
+		OutSeed = TargetSeed;
+	}
+	
+	// Write output
+	Out_SetPosition(OutDataIndex, OutElemIndex, OutPosition);
+	Out_SetRotation(OutDataIndex, OutElemIndex, OutRotation);
+	Out_SetScale(OutDataIndex, OutElemIndex, OutScale);
+	Out_SetColor(OutDataIndex, OutElemIndex, OutColor);
+	Out_SetSeed(OutDataIndex, OutElemIndex, OutSeed);
+
+	// These properties do not have an inheritance mode, simply copy from the source data.
+	Out_SetBoundsMin(OutDataIndex, OutElemIndex, Source_GetBoundsMin(SourceDataIndex, SourcePointIndexInData));
+	Out_SetBoundsMax(OutDataIndex, OutElemIndex, Source_GetBoundsMax(SourceDataIndex, SourcePointIndexInData));
+	Out_SetDensity(OutDataIndex, OutElemIndex, Source_GetDensity(SourceDataIndex, SourcePointIndexInData));
+	Out_SetSteepness(OutDataIndex, OutElemIndex, Source_GetSteepness(SourceDataIndex, SourcePointIndexInData));
+
+	// Copy attributes
+	if (AttributeInheritance == COPY_POINTS_METADATA_INHERITANCE_MODE_SourceOnly)
+	{
+		PCG_COPY_ATTRIBUTES_TO_OUTPUT(Out, Source, OutDataIndex, OutElemIndex, SourceDataIndex, SourcePointIndexInData, /*bMetadataOnly=*/true, /*bInitNonCopiedAttributes=*/true);
+	}
+	else if (AttributeInheritance == COPY_POINTS_METADATA_INHERITANCE_MODE_TargetOnly)
+	{
+		PCG_COPY_ATTRIBUTES_TO_OUTPUT(Out, Target, OutDataIndex, OutElemIndex, TargetDataIndex, TargetPointIndexInData, /*bMetadataOnly=*/true, /*bInitNonCopiedAttributes=*/true);
+	}
+	else if (AttributeInheritance == COPY_POINTS_METADATA_INHERITANCE_MODE_SourceFirst)
+	{
+		PCG_COPY_ATTRIBUTES_TO_OUTPUT(Out, Target, OutDataIndex, OutElemIndex, TargetDataIndex, TargetPointIndexInData, /*bMetadataOnly=*/true, /*bInitNonCopiedAttributes=*/true);
+
+		// Don't allow Source to initialize non-copied attributes in the output, otherwise it will overwrite attributes from Target.
+		PCG_COPY_ATTRIBUTES_TO_OUTPUT(Out, Source, OutDataIndex, OutElemIndex, SourceDataIndex, SourcePointIndexInData, /*bMetadataOnly=*/true, /*bInitNonCopiedAttributes=*/false);
+	}
+	else if (AttributeInheritance == COPY_POINTS_METADATA_INHERITANCE_MODE_TargetFirst)
+	{
+		PCG_COPY_ATTRIBUTES_TO_OUTPUT(Out, Source, OutDataIndex, OutElemIndex, SourceDataIndex, SourcePointIndexInData, /*bMetadataOnly=*/true, /*bInitNonCopiedAttributes=*/true);
+
+		// Don't allow Target to initialize non-copied attributes in the output, otherwise it will overwrite attributes from Source.
+		PCG_COPY_ATTRIBUTES_TO_OUTPUT(Out, Target, OutDataIndex, OutElemIndex, TargetDataIndex, TargetPointIndexInData, /*bMetadataOnly=*/true, /*bInitNonCopiedAttributes=*/false);
+	}
+}
diff --git a/Engine/Plugins/PCG/Source/PCG/Internal/Compute/PCGDataForGPU.h b/Engine/Plugins/PCG/Source/PCG/Internal/Compute/PCGDataForGPU.h
--- a/Engine/Plugins/PCG/Source/PCG/Internal/Compute/PCGDataForGPU.h
+++ b/Engine/Plugins/PCG/Source/PCG/Internal/Compute/PCGDataForGPU.h
@@ -155,6 +155,8 @@ public:
 
 	bool ContainsAttribute(FName InAttributeName) const;
 
+	bool ContainsAttribute(FName InAttributeName, EPCGKernelAttributeType InAttributeType) const;
+
 	void AddAttribute(FPCGKernelAttributeKey InAttribute, const UPCGDataBinding* InBinding, const TArray<int32>* InOptionalUniqueStringKeys = nullptr);
 
 public:
diff --git a/Engine/Plugins/PCG/Source/PCG/Private/Compute/DataInterfaces/Elements/PCGCopyPointsAnalysisDataInterface.cpp b/Engine/Plugins/PCG/Source/PCG/Private/Compute/DataInterfaces/Elements/PCGCopyPointsAnalysisDataInterface.cpp
new file mode 100644
--- /dev/null
+++ b/Engine/Plugins/PCG/Source/PCG/Private/Compute/DataInterfaces/Elements/PCGCopyPointsAnalysisDataInterface.cpp
@@ -0,0 +1,103 @@
+// Copyright Epic Games, Inc. All Rights Reserved.
+
+#include "PCGCopyPointsAnalysisDataInterface.h"
+
+#include "PCGContext.h"
+#include "Compute/PCGComputeKernel.h"
+#include "Compute/PCGDataBinding.h"
+#include "Elements/PCGCopyPoints.h"
+
+#include "GlobalRenderResources.h"
+#include "RenderGraphBuilder.h"
+#include "RenderGraphResources.h"
+#include "RHIResources.h"
+#include "ShaderCompilerCore.h"
+#include "ShaderCore.h"
+#include "ShaderParameterMetadataBuilder.h"
+#include "ComputeFramework/ShaderParamTypeDefinition.h"
+
+#include UE_INLINE_GENERATED_CPP_BY_NAME(PCGCopyPointsAnalysisDataInterface)
+
+void UPCGCopyPointsAnalysisDataInterface::GetSupportedInputs(TArray<FShaderFunctionDefinition>& OutFunctions) const
+{
+	OutFunctions.AddDefaulted_GetRef()
+		.SetName(TEXT("CopyPoints_GetMatchAttributeId"))
+		.AddReturnType(FShaderValueType::Get(EShaderFundamentalType::Int));
+
+	OutFunctions.AddDefaulted_GetRef()
+		.SetName(TEXT("CopyPoints_GetSelectedFlagAttributeId"))
+		.AddReturnType(FShaderValueType::Get(EShaderFundamentalType::Int));
+}
+
+BEGIN_SHADER_PARAMETER_STRUCT(FPCGCopyPointsAnalysisDataInterfaceParameters,)
+	SHADER_PARAMETER(int32, MatchAttributeId)
+	SHADER_PARAMETER(int32, SelectedFlagAttributeId)
+END_SHADER_PARAMETER_STRUCT()
+
+void UPCGCopyPointsAnalysisDataInterface::GetShaderParameters(TCHAR const* UID, FShaderParametersMetadataBuilder& InOutBuilder, FShaderParametersMetadataAllocations& InOutAllocations) const
+{
+	InOutBuilder.AddNestedStruct<FPCGCopyPointsAnalysisDataInterfaceParameters>(UID);
+}
+
+void UPCGCopyPointsAnalysisDataInterface::GetHLSL(FString& OutHLSL, FString const& InDataInterfaceName) const
+{
+	TMap<FString, FStringFormatArg> TemplateArgs =
+	{
+		{ TEXT("DataInterfaceName"), InDataInterfaceName },
+	};
+
+	OutHLSL += FString::Format(TEXT(
+		"int {DataInterfaceName}_MatchAttributeId;\n"
+		"int {DataInterfaceName}_SelectedFlagAttributeId;\n"
+		"\n"
+		"int CopyPoints_GetMatchAttributeId_{DataInterfaceName}() { return {DataInterfaceName}_MatchAttributeId; }\n"
+		"int CopyPoints_GetSelectedFlagAttributeId_{DataInterfaceName}() { return {DataInterfaceName}_SelectedFlagAttributeId; }\n"
+		), TemplateArgs);
+}
+
+UComputeDataProvider* UPCGCopyPointsAnalysisDataInterface::CreateDataProvider(TObjectPtr<UObject> InBinding, uint64 InInputMask, uint64 InOutputMask) const
+{
+	TRACE_CPUPROFILER_EVENT_SCOPE(UPCGCopyPointsAnalysisDataInterface::CreateDataProvider);
+	UPCGDataBinding* Binding = CastChecked<UPCGDataBinding>(InBinding);
+
+	TObjectPtr<UPCGCopyPointsAnalysisDataProvider> DataProvider = NewObject<UPCGCopyPointsAnalysisDataProvider>();
+	DataProvider->Settings = CastChecked<UPCGCopyPointsSettings>(ProducerKernel->GetSettings());
+
+	// Resolve attribute IDs from incoming data.
+	const FName MatchAttributeName = DataProvider->Settings->MatchAttribute;
+	DataProvider->MatchAttributeId = Binding->GetAttributeId(MatchAttributeName, EPCGKernelAttributeType::Int);
+	ensure(DataProvider->MatchAttributeId != INDEX_NONE);
+
+	DataProvider->SelectedFlagAttributeId = Binding->GetAttributeId(PCGCopyPointsConstants::SelectedFlagAttributeName, EPCGKernelAttributeType::Bool);
+	ensure(DataProvider->SelectedFlagAttributeId != INDEX_NONE);
+
+	return DataProvider;
+}
+
+FComputeDataProviderRenderProxy* UPCGCopyPointsAnalysisDataProvider::GetRenderProxy()
+{
+	const UPCGCopyPointsSettings* SettingsPtr = Settings;
+	check(SettingsPtr);
+
+	FPCGCopyPointsAnalysisDataProviderProxy::FCopyPointsAnalysisData_RenderThread ProxyData;
+	ProxyData.MatchAttributeId = MatchAttributeId;
+	ProxyData.SelectedFlagAttributeId = SelectedFlagAttributeId;
+
+	return new FPCGCopyPointsAnalysisDataProviderProxy(MoveTemp(ProxyData));
+}
+
+bool FPCGCopyPointsAnalysisDataProviderProxy::IsValid(FValidationData const& InValidationData) const
+{
+	return InValidationData.ParameterStructSize == sizeof(FParameters);
+}
+
+void FPCGCopyPointsAnalysisDataProviderProxy::GatherDispatchData(FDispatchData const& InDispatchData)
+{
+	const TStridedView<FParameters> ParameterArray = MakeStridedParameterView<FParameters>(InDispatchData);
+	for (int32 InvocationIndex = 0; InvocationIndex < ParameterArray.Num(); ++InvocationIndex)
+	{
+		FParameters& Parameters = ParameterArray[InvocationIndex];
+		Parameters.MatchAttributeId = Data.MatchAttributeId;
+		Parameters.SelectedFlagAttributeId = Data.SelectedFlagAttributeId;
+	}
+}
diff --git a/Engine/Plugins/PCG/Source/PCG/Private/Compute/DataInterfaces/Elements/PCGCopyPointsAnalysisDataInterface.h b/Engine/Plugins/PCG/Source/PCG/Private/Compute/DataInterfaces/Elements/PCGCopyPointsAnalysisDataInterface.h
new file mode 100644
--- /dev/null
+++ b/Engine/Plugins/PCG/Source/PCG/Private/Compute/DataInterfaces/Elements/PCGCopyPointsAnalysisDataInterface.h
@@ -0,0 +1,72 @@
+// Copyright Epic Games, Inc. All Rights Reserved.
+
+#pragma once
+
+#include "Compute/DataInterfaces/PCGComputeDataInterface.h"
+#include "ComputeFramework/ComputeDataProvider.h"
+
+#include "PCGCopyPointsAnalysisDataInterface.generated.h"
+
+class FPCGCopyPointsAnalysisDataInterfaceParameters;
+class UPCGCopyPointsSettings;
+
+/** Data Interface to marshal Copy Points settings to the GPU. */
+UCLASS(ClassGroup = (Procedural))
+class UPCGCopyPointsAnalysisDataInterface : public UPCGComputeDataInterface
+{
+	GENERATED_BODY()
+
+public:
+	//~ Begin UComputeDataInterface Interface
+	TCHAR const* GetClassName() const override { return TEXT("PCGCopyPointsAnalysis"); }
+	void GetSupportedInputs(TArray<FShaderFunctionDefinition>& OutFunctions) const override;
+	void GetShaderParameters(TCHAR const* UID, FShaderParametersMetadataBuilder& InOutBuilder, FShaderParametersMetadataAllocations& InOutAllocations) const override;
+	void GetHLSL(FString& OutHLSL, FString const& InDataInterfaceName) const override;
+	UComputeDataProvider* CreateDataProvider(TObjectPtr<UObject> InBinding, uint64 InInputMask, uint64 InOutputMask) const override;
+	//~ End UComputeDataInterface Interface
+};
+
+UCLASS()
+class UPCGCopyPointsAnalysisDataProvider : public UPCGComputeDataProvider
+{
+	GENERATED_BODY()
+
+public:
+	//~ Begin UComputeDataProvider Interface
+	FComputeDataProviderRenderProxy* GetRenderProxy() override;
+	//~ End UComputeDataProvider Interface
+
+public:
+	UPROPERTY()
+	TObjectPtr<const UPCGCopyPointsSettings> Settings;
+
+	UPROPERTY()
+	int32 MatchAttributeId = INDEX_NONE;
+
+	UPROPERTY()
+	int32 SelectedFlagAttributeId = INDEX_NONE;
+};
+
+class FPCGCopyPointsAnalysisDataProviderProxy : public FComputeDataProviderRenderProxy
+{
+public:
+	struct FCopyPointsAnalysisData_RenderThread
+	{
+		int32 MatchAttributeId = 0;
+		int32 SelectedFlagAttributeId = 0;
+	};
+
+	FPCGCopyPointsAnalysisDataProviderProxy(FCopyPointsAnalysisData_RenderThread InData)
+		: Data(MoveTemp(InData))
+	{}
+
+	//~ Begin FComputeDataProviderRenderProxy Interface
+	bool IsValid(FValidationData const& InValidationData) const override;
+	void GatherDispatchData(FDispatchData const& InDispatchData) override;
+	//~ End FComputeDataProviderRenderProxy Interface
+
+protected:
+	using FParameters = FPCGCopyPointsAnalysisDataInterfaceParameters;
+
+	FCopyPointsAnalysisData_RenderThread Data;
+};
diff --git a/Engine/Plugins/PCG/Source/PCG/Private/Compute/DataInterfaces/Elements/PCGCopyPointsDataInterface.cpp b/Engine/Plugins/PCG/Source/PCG/Private/Compute/DataInterfaces/Elements/PCGCopyPointsDataInterface.cpp
--- a/Engine/Plugins/PCG/Source/PCG/Private/Compute/DataInterfaces/Elements/PCGCopyPointsDataInterface.cpp
+++ b/Engine/Plugins/PCG/Source/PCG/Private/Compute/DataInterfaces/Elements/PCGCopyPointsDataInterface.cpp
@@ -2,10 +2,17 @@
 
 #include "PCGCopyPointsDataInterface.h"
 
+#include "PCGContext.h"
+#include "PCGParamData.h"
+#include "Compute/PCGComputeGraph.h"
+#include "Compute/PCGDataBinding.h"
+#include "Compute/Data/PCGProxyForGPUData.h"
 #include "Elements/PCGCopyPoints.h"
 #include "Elements/PCGCopyPointsKernel.h"
 
 #include "GlobalRenderResources.h"
+#include "RenderGraphBuilder.h"
+#include "RenderGraphResources.h"
 #include "RHIResources.h"
 #include "ShaderCompilerCore.h"
 #include "ShaderCore.h"
@@ -17,36 +24,42 @@
 void UPCGCopyPointsDataInterface::GetSupportedInputs(TArray<FShaderFunctionDefinition>& OutFunctions) const
 {
 	OutFunctions.AddDefaulted_GetRef()
-		.SetName(TEXT("GetRotationInheritance"))
+		.SetName(TEXT("CopyPoints_GetRotationInheritance"))
 		.AddReturnType(FShaderValueType::Get(EShaderFundamentalType::Uint));
 
 	OutFunctions.AddDefaulted_GetRef()
-		.SetName(TEXT("GetApplyTargetRotationToPositions"))
+		.SetName(TEXT("CopyPoints_GetApplyTargetRotationToPositions"))
 		.AddReturnType(FShaderValueType::Get(EShaderFundamentalType::Uint));
 
 	OutFunctions.AddDefaulted_GetRef()
-		.SetName(TEXT("GetScaleInheritance"))
+		.SetName(TEXT("CopyPoints_GetScaleInheritance"))
 		.AddReturnType(FShaderValueType::Get(EShaderFundamentalType::Uint));
 
 	OutFunctions.AddDefaulted_GetRef()
-		.SetName(TEXT("GetApplyTargetScaleToPositions"))
+		.SetName(TEXT("CopyPoints_GetApplyTargetScaleToPositions"))
 		.AddReturnType(FShaderValueType::Get(EShaderFundamentalType::Uint));
 
 	OutFunctions.AddDefaulted_GetRef()
-		.SetName(TEXT("GetColorInheritance"))
+		.SetName(TEXT("CopyPoints_GetColorInheritance"))
 		.AddReturnType(FShaderValueType::Get(EShaderFundamentalType::Uint));
 
 	OutFunctions.AddDefaulted_GetRef()
-		.SetName(TEXT("GetSeedInheritance"))
+		.SetName(TEXT("CopyPoints_GetSeedInheritance"))
 		.AddReturnType(FShaderValueType::Get(EShaderFundamentalType::Uint));
 
 	OutFunctions.AddDefaulted_GetRef()
-		.SetName(TEXT("GetAttributeInheritance"))
+		.SetName(TEXT("CopyPoints_GetAttributeInheritance"))
 		.AddReturnType(FShaderValueType::Get(EShaderFundamentalType::Uint));
 
 	OutFunctions.AddDefaulted_GetRef()
-		.SetName(TEXT("GetCopyEachSourceOnEveryTarget"))
+		.SetName(TEXT("CopyPoints_GetCopyEachSourceOnEveryTarget"))
 		.AddReturnType(FShaderValueType::Get(EShaderFundamentalType::Bool));
+
+	OutFunctions.AddDefaulted_GetRef()
+		.SetName(TEXT("CopyPoints_GetSourceAndTargetDataIndices"))
+		.AddParam(EShaderFundamentalType::Uint) // InOutputIndex
+		.AddParam(EShaderFundamentalType::Uint, 0, 0, EShaderParamModifier::Out) // OutSourceIndex
+		.AddParam(EShaderFundamentalType::Uint, 0, 0, EShaderParamModifier::Out); // OutTargetIndex
 }
 
 BEGIN_SHADER_PARAMETER_STRUCT(FPCGCopyPointsDataInterfaceParameters,)
@@ -58,6 +71,7 @@ BEGIN_SHADER_PARAMETER_STRUCT(FPCGCopyPointsDataInterfaceParameters,)
 	SHADER_PARAMETER(uint32, SeedInheritance)
 	SHADER_PARAMETER(uint32, AttributeInheritance)
 	SHADER_PARAMETER(uint32, bCopyEachSourceOnEveryTarget)
+	SHADER_PARAMETER_RDG_BUFFER_SRV(StructuredBuffer<FUintVector2>, SourceAndTargetDataIndices)
 END_SHADER_PARAMETER_STRUCT()
 
 void UPCGCopyPointsDataInterface::GetShaderParameters(TCHAR const* UID, FShaderParametersMetadataBuilder& InOutBuilder, FShaderParametersMetadataAllocations& InOutAllocations) const
@@ -81,41 +95,185 @@ void UPCGCopyPointsDataInterface::GetHLSL(FString& OutHLSL, FString const& InDat
 		"uint {DataInterfaceName}_SeedInheritance;\n"
 		"uint {DataInterfaceName}_AttributeInheritance;\n"
 		"uint {DataInterfaceName}_bCopyEachSourceOnEveryTarget;\n"
+		"StructuredBuffer<uint2> {DataInterfaceName}_SourceAndTargetDataIndices;\n"
+		"\n"
+		"uint CopyPoints_GetRotationInheritance_{DataInterfaceName}() { return {DataInterfaceName}_RotationInheritance; }\n"
+		"uint CopyPoints_GetApplyTargetRotationToPositions_{DataInterfaceName}() { return {DataInterfaceName}_ApplyTargetRotationToPositions; }\n"
+		"uint CopyPoints_GetScaleInheritance_{DataInterfaceName}() { return {DataInterfaceName}_ScaleInheritance; }\n"
+		"uint CopyPoints_GetApplyTargetScaleToPositions_{DataInterfaceName}() { return {DataInterfaceName}_ApplyTargetScaleToPositions; }\n"
+		"uint CopyPoints_GetColorInheritance_{DataInterfaceName}() { return {DataInterfaceName}_ColorInheritance; }\n"
+		"uint CopyPoints_GetSeedInheritance_{DataInterfaceName}() { return {DataInterfaceName}_SeedInheritance; }\n"
+		"uint CopyPoints_GetAttributeInheritance_{DataInterfaceName}() { return {DataInterfaceName}_AttributeInheritance; }\n"
+		"uint CopyPoints_GetCopyEachSourceOnEveryTarget_{DataInterfaceName}() { return {DataInterfaceName}_bCopyEachSourceOnEveryTarget; }\n"
 		"\n"
-		"uint GetRotationInheritance_{DataInterfaceName}() { return {DataInterfaceName}_RotationInheritance; }\n"
-		"uint GetApplyTargetRotationToPositions_{DataInterfaceName}() { return {DataInterfaceName}_ApplyTargetRotationToPositions; }\n"
-		"uint GetScaleInheritance_{DataInterfaceName}() { return {DataInterfaceName}_ScaleInheritance; }\n"
-		"uint GetApplyTargetScaleToPositions_{DataInterfaceName}() { return {DataInterfaceName}_ApplyTargetScaleToPositions; }\n"
-		"uint GetColorInheritance_{DataInterfaceName}() { return {DataInterfaceName}_ColorInheritance; }\n"
-		"uint GetSeedInheritance_{DataInterfaceName}() { return {DataInterfaceName}_SeedInheritance; }\n"
-		"uint GetAttributeInheritance_{DataInterfaceName}() { return {DataInterfaceName}_AttributeInheritance; }\n"
-		"uint GetCopyEachSourceOnEveryTarget_{DataInterfaceName}() { return {DataInterfaceName}_bCopyEachSourceOnEveryTarget; }\n"
+		"void CopyPoints_GetSourceAndTargetDataIndices_{DataInterfaceName}(uint InOutputIndex, out uint OutSourceIndex, out uint OutTargetIndex)\n"
+		"{\n"
+		"	const uint2 Indices = {DataInterfaceName}_SourceAndTargetDataIndices[InOutputIndex];\n"
+		"	OutSourceIndex = Indices[0];\n"
+		"	OutTargetIndex = Indices[1];\n"
+		"}\n"
 		), TemplateArgs);
 }
 
 UComputeDataProvider* UPCGCopyPointsDataInterface::CreateDataProvider(TObjectPtr<UObject> InBinding, uint64 InInputMask, uint64 InOutputMask) const
 {
+	TRACE_CPUPROFILER_EVENT_SCOPE(UPCGCopyPointsDataInterface::CreateDataProvider);
+
 	TObjectPtr<UPCGCopyPointsDataProvider> DataProvider = NewObject<UPCGCopyPointsDataProvider>();
-	DataProvider->ProducerKernel = CastChecked<UPCGCopyPointsKernel>(ProducerKernel);
+	DataProvider->ProducerKernel = ProducerKernel;
 
 	return DataProvider;
 }
 
+bool UPCGCopyPointsDataProvider::PerformPreExecuteReadbacks_GameThread(UPCGDataBinding* InBinding)
+{
+	TRACE_CPUPROFILER_EVENT_SCOPE(UPCGCopyPointsDataProvider::PerformPreExecuteReadbacks_GameThread);
+	check(InBinding);
+
+	if (!Super::PerformPreExecuteReadbacks_GameThread(InBinding))
+	{
+		return false;
+	}
+
+	const UPCGCopyPointsSettings* Settings = CastChecked<UPCGCopyPointsSettings>(ProducerKernel->GetSettings());
+	if (!Settings->bMatchBasedOnAttribute)
+	{
+		// If we're not matching based on attribute then we don't need to do any readbacks.
+		return true;
+	}
+
+	int32 FirstAnalysisDataIndex = INDEX_NONE;
+
+	if (InBinding->Graph && ProducerKernel)
+	{
+		const int32 BindingIndex = InBinding->Graph->GetBindingIndex(ProducerKernel->GetKernelIndex(), PCGCopyPointsConstants::SelectedFlagsPinLabel, /*bIsInput=*/true);
+		if (ensure(BindingIndex != INDEX_NONE))
+		{
+			if (const FName* VirtualLabel = InBinding->Graph->GetVirtualPinLabelFromBindingIndex(BindingIndex))
+			{
+				FirstAnalysisDataIndex = InBinding->InputDataCollection.TaggedData.IndexOfByPredicate([VirtualLabel](const FPCGTaggedData& InData)
+				{
+					return InData.Pin == *VirtualLabel;
+				});
+			}
+		}
+	}
+
+	if (FirstAnalysisDataIndex == INDEX_NONE)
+	{
+		ensure(false);
+		return true;
+	}
+
+	const UPCGProxyForGPUData* DataWithGPUSupport = Cast<UPCGProxyForGPUData>(InBinding->InputDataCollection.TaggedData[FirstAnalysisDataIndex].Data);
+	if (DataWithGPUSupport)
+	{
+		TSharedPtr<FPCGContextHandle> Context = InBinding->ContextHandle.Pin();
+
+		// Poll until readback is done.
+		UPCGProxyForGPUData::FReadbackResult Result = DataWithGPUSupport->GetCPUData(Context ? Context->GetContext() : nullptr);
+
+		if (!Result.bComplete)
+		{
+			// Read back pending - wait for it to complete.
+			return false;
+		}
+
+		if (Result.TaggedData.Data)
+		{
+			InBinding->InputDataCollection.TaggedData[FirstAnalysisDataIndex].Data = Result.TaggedData.Data;
+			InBinding->InputDataCollection.TaggedData[FirstAnalysisDataIndex].Tags = MoveTemp(Result.TaggedData.Tags);
+		}
+		else
+		{
+			ensure(false);
+			return true;
+		}
+	}
+
+	return true;
+}
+
+bool UPCGCopyPointsDataProvider::PrepareForExecute_GameThread(UPCGDataBinding* InBinding)
+{
+	TRACE_CPUPROFILER_EVENT_SCOPE(UPCGCopyPointsDataProvider::PrepareForExecute_GameThread);
+	check(InBinding);
+
+	if (!Super::PrepareForExecute_GameThread(InBinding))
+	{
+		return false;
+	}
+
+	const UPCGCopyPointsSettings* Settings = CastChecked<UPCGCopyPointsSettings>(ProducerKernel->GetSettings());
+	if (!Settings->bMatchBasedOnAttribute)
+	{
+		// If we're not matching based on attribute then we don't need to create flags.
+		return true;
+	}
+
+	int32 FirstAnalysisDataIndex = INDEX_NONE;
+
+	const int32 BindingIndex = InBinding->Graph ? InBinding->Graph->GetBindingIndex(ProducerKernel->GetKernelIndex(), PCGCopyPointsConstants::SelectedFlagsPinLabel, /*bIsInput=*/true) : INDEX_NONE;
+	if (ensure(BindingIndex != INDEX_NONE))
+	{
+		if (const FName* VirtualLabel = InBinding->Graph->GetVirtualPinLabelFromBindingIndex(BindingIndex))
+		{
+			FirstAnalysisDataIndex = InBinding->InputDataCollection.TaggedData.IndexOfByPredicate([VirtualLabel](const FPCGTaggedData& InData)
+			{
+				return InData.Pin == *VirtualLabel;
+			});
+		}
+	}
+
+	if (FirstAnalysisDataIndex == INDEX_NONE)
+	{
+		return true;
+	}
+
+	FPCGDataCollectionDesc SourcePinDesc;
+	FPCGDataCollectionDesc TargetPinDesc;
+	ensure(InBinding->GetKernelPinDataDesc_GameThread(ProducerKernel, PCGCopyPointsConstants::SourcePointsLabel, /*bIsInput=*/true, SourcePinDesc));
+	ensure(InBinding->GetKernelPinDataDesc_GameThread(ProducerKernel, PCGCopyPointsConstants::TargetPointsLabel, /*bIsInput=*/true, TargetPinDesc));
+
+	const uint32 NumOutputs = SourcePinDesc.DataDescs.Num()* TargetPinDesc.DataDescs.Num();
+	SourceAndTargetDataIndices.Reserve(NumOutputs);
+
+	const UPCGParamData* AnalysisResultsData = Cast<UPCGParamData>(InBinding->InputDataCollection.TaggedData[FirstAnalysisDataIndex].Data);
+	const UPCGMetadata* AnalysisMetadata = AnalysisResultsData ? AnalysisResultsData->ConstMetadata() : nullptr;
+	const FPCGMetadataAttributeBase* AnalysisAttributeBase = AnalysisMetadata ? AnalysisMetadata->GetConstAttribute(PCGCopyPointsConstants::SelectedFlagAttributeName) : nullptr;
+	const int32 NumElements = AnalysisAttributeBase ? AnalysisMetadata->GetItemCountForChild() : INDEX_NONE;
+	const FPCGMetadataAttribute<bool>* SelectedFlagsAttribute = ((NumElements == NumOutputs) && AnalysisAttributeBase && (AnalysisAttributeBase->GetTypeId() == PCG::Private::MetadataTypes<bool>::Id)) ? static_cast<const FPCGMetadataAttribute<bool>*>(AnalysisAttributeBase) : nullptr;
+
+	for (int32 S = 0; S < SourcePinDesc.DataDescs.Num(); ++S)
+	{
+		for (int32 T = 0; T < TargetPinDesc.DataDescs.Num(); ++T)
+		{
+			if (!SelectedFlagsAttribute || SelectedFlagsAttribute->GetValue(S * TargetPinDesc.DataDescs.Num() + T))
+			{
+				SourceAndTargetDataIndices.Emplace(S, T);
+			}
+		}
+	}
+
+	return true;
+}
+
 FComputeDataProviderRenderProxy* UPCGCopyPointsDataProvider::GetRenderProxy()
 {
-	const UPCGCopyPointsSettings* Settings = ProducerKernel->Settings;
-	check(Settings);
-
-	return new FPCGCopyPointsDataProviderProxy(
-		static_cast<uint32>(Settings->RotationInheritance),
-		static_cast<uint32>((Settings->RotationInheritance != EPCGCopyPointsInheritanceMode::Source || Settings->bApplyTargetRotationToPositions) ? 1 : 0),
-		static_cast<uint32>(Settings->ScaleInheritance),
-		static_cast<uint32>((Settings->ScaleInheritance != EPCGCopyPointsInheritanceMode::Source || Settings->bApplyTargetScaleToPositions) ? 1 : 0),
-		static_cast<uint32>(Settings->ColorInheritance),
-		static_cast<uint32>(Settings->SeedInheritance),
-		static_cast<uint32>(Settings->AttributeInheritance),
-		static_cast<uint32>(Settings->bCopyEachSourceOnEveryTarget)
-	);
+	const UPCGCopyPointsSettings* Settings = CastChecked<UPCGCopyPointsSettings>(ProducerKernel->GetSettings());
+
+	FPCGCopyPointsDataProviderProxy::FCopyPointsData_RenderThread ProxyData;
+	ProxyData.RotationInheritance = static_cast<uint32>(Settings->RotationInheritance);
+	ProxyData.ApplyTargetRotationToPositions = static_cast<uint32>((Settings->RotationInheritance != EPCGCopyPointsInheritanceMode::Source || Settings->bApplyTargetRotationToPositions) ? 1 : 0);
+	ProxyData.ScaleInheritance = static_cast<uint32>(Settings->ScaleInheritance);
+	ProxyData.ApplyTargetScaleToPositions = static_cast<uint32>((Settings->ScaleInheritance != EPCGCopyPointsInheritanceMode::Source || Settings->bApplyTargetScaleToPositions) ? 1 : 0);
+	ProxyData.ColorInheritance = static_cast<uint32>(Settings->ColorInheritance);
+	ProxyData.SeedInheritance = static_cast<uint32>(Settings->SeedInheritance);
+	ProxyData.AttributeInheritance = static_cast<uint32>(Settings->AttributeInheritance);
+	ProxyData.bCopyEachSourceOnEveryTarget = static_cast<uint32>(Settings->bCopyEachSourceOnEveryTarget);
+	ProxyData.SourceAndTargetDataIndices = SourceAndTargetDataIndices;
+
+	return new FPCGCopyPointsDataProviderProxy(MoveTemp(ProxyData));
 }
 
 bool FPCGCopyPointsDataProviderProxy::IsValid(FValidationData const& InValidationData) const
@@ -128,19 +286,41 @@ bool FPCGCopyPointsDataProviderProxy::IsValid(FValidationData const& InValidatio
 	return true;
 }
 
+void FPCGCopyPointsDataProviderProxy::AllocateResources(FRDGBuilder& GraphBuilder, FAllocationData const& InAllocationData)
+{
+	TRACE_CPUPROFILER_EVENT_SCOPE(FPCGCopyPointsDataProviderProxy::AllocateResources);
+
+	// Ensure we always allocate at least 1 element so we have something to bind.
+	const uint32 NumIndexPairs = FMath::Max(1, Data.SourceAndTargetDataIndices.Num());
+	FRDGBufferDesc IndexPairsDesc = FRDGBufferDesc::CreateStructuredDesc(sizeof(FUintVector2), NumIndexPairs);
+
+	SourceAndTargetDataIndicesBuffer = GraphBuilder.CreateBuffer(IndexPairsDesc, TEXT("PCGCopyPoints_SourceAndTargetDataIndices"));
+	SourceAndTargetDataIndicesBufferSRV = GraphBuilder.CreateSRV(SourceAndTargetDataIndicesBuffer);
+
+	if (Data.SourceAndTargetDataIndices.IsEmpty())
+	{
+		GraphBuilder.QueueBufferUpload(SourceAndTargetDataIndicesBuffer, &FUintVector2::ZeroValue, 1);
+	}
+	else
+	{
+		GraphBuilder.QueueBufferUpload(SourceAndTargetDataIndicesBuffer, MakeArrayView(Data.SourceAndTargetDataIndices));
+	}
+}
+
 void FPCGCopyPointsDataProviderProxy::GatherDispatchData(FDispatchData const& InDispatchData)
 {
 	const TStridedView<FParameters> ParameterArray = MakeStridedParameterView<FParameters>(InDispatchData);
 	for (int32 InvocationIndex = 0; InvocationIndex < ParameterArray.Num(); ++InvocationIndex)
 	{
 		FParameters& Parameters = ParameterArray[InvocationIndex];
-		Parameters.RotationInheritance = RotationInheritance;
-		Parameters.ApplyTargetRotationToPositions = ApplyTargetRotationToPositions;
-		Parameters.ScaleInheritance = ScaleInheritance;
-		Parameters.ApplyTargetScaleToPositions = ApplyTargetScaleToPositions;
-		Parameters.ColorInheritance = ColorInheritance;
-		Parameters.SeedInheritance = SeedInheritance;
-		Parameters.AttributeInheritance = AttributeInheritance;
-		Parameters.bCopyEachSourceOnEveryTarget = bCopyEachSourceOnEveryTarget;
+		Parameters.RotationInheritance = Data.RotationInheritance;
+		Parameters.ApplyTargetRotationToPositions = Data.ApplyTargetRotationToPositions;
+		Parameters.ScaleInheritance = Data.ScaleInheritance;
+		Parameters.ApplyTargetScaleToPositions = Data.ApplyTargetScaleToPositions;
+		Parameters.ColorInheritance = Data.ColorInheritance;
+		Parameters.SeedInheritance = Data.SeedInheritance;
+		Parameters.AttributeInheritance = Data.AttributeInheritance;
+		Parameters.bCopyEachSourceOnEveryTarget = Data.bCopyEachSourceOnEveryTarget;
+		Parameters.SourceAndTargetDataIndices = SourceAndTargetDataIndicesBufferSRV;
 	}
 }
diff --git a/Engine/Plugins/PCG/Source/PCG/Private/Compute/DataInterfaces/Elements/PCGCopyPointsDataInterface.h b/Engine/Plugins/PCG/Source/PCG/Private/Compute/DataInterfaces/Elements/PCGCopyPointsDataInterface.h
--- a/Engine/Plugins/PCG/Source/PCG/Private/Compute/DataInterfaces/Elements/PCGCopyPointsDataInterface.h
+++ b/Engine/Plugins/PCG/Source/PCG/Private/Compute/DataInterfaces/Elements/PCGCopyPointsDataInterface.h
@@ -8,7 +8,7 @@
 #include "PCGCopyPointsDataInterface.generated.h"
 
 class FPCGCopyPointsDataInterfaceParameters;
-class UPCGCopyPointsKernel;
+class UPCGCopyPointsSettings;
 
 /** Data Interface to marshal Copy Points settings to the GPU. */
 UCLASS(ClassGroup = (Procedural))
@@ -27,47 +27,59 @@ public:
 };
 
 UCLASS()
-class UPCGCopyPointsDataProvider : public UComputeDataProvider
+class UPCGCopyPointsDataProvider : public UPCGComputeDataProvider
 {
 	GENERATED_BODY()
 
 public:
+	//~ Begin UPCGComputeDataProvider Interface
+	virtual bool PerformPreExecuteReadbacks_GameThread(UPCGDataBinding* InBinding) override;
+	virtual bool PrepareForExecute_GameThread(UPCGDataBinding* InBinding) override;
+	//~ End UPCGComputeDataProvider Interface
+
 	//~ Begin UComputeDataProvider Interface
 	FComputeDataProviderRenderProxy* GetRenderProxy() override;
 	//~ End UComputeDataProvider Interface
 
+public:
+	UPROPERTY()
+	TObjectPtr<const UPCGComputeKernel> ProducerKernel;
+
 	UPROPERTY()
-	TObjectPtr<const UPCGCopyPointsKernel> ProducerKernel;
+	TArray<FUintVector2> SourceAndTargetDataIndices;
 };
 
 class FPCGCopyPointsDataProviderProxy : public FComputeDataProviderRenderProxy
 {
 public:
-	FPCGCopyPointsDataProviderProxy(uint32 InRotationInheritance, uint32 InApplyTargetRotationToPositions, uint32 InScaleInheritance, uint32 InApplyTargetScaleToPositions, uint32 InColorInheritance, uint32 InSeedInheritance, uint32 InAttributeInheritance, uint32 bInCopyEachSourceOnEveryTarget)
-		: RotationInheritance(InRotationInheritance)
-		, ApplyTargetRotationToPositions(InApplyTargetRotationToPositions)
-		, ScaleInheritance(InScaleInheritance)
-		, ApplyTargetScaleToPositions(InApplyTargetScaleToPositions)
-		, ColorInheritance(InColorInheritance)
-		, SeedInheritance(InSeedInheritance)
-		, AttributeInheritance(InAttributeInheritance)
-		, bCopyEachSourceOnEveryTarget(bInCopyEachSourceOnEveryTarget)
+	struct FCopyPointsData_RenderThread
+	{
+		uint32 RotationInheritance = 0;
+		uint32 ApplyTargetRotationToPositions = 0;
+		uint32 ScaleInheritance = 0;
+		uint32 ApplyTargetScaleToPositions = 0;
+		uint32 ColorInheritance = 0;
+		uint32 SeedInheritance = 0;
+		uint32 AttributeInheritance = 0;
+		uint32 bCopyEachSourceOnEveryTarget = 0;
+		TArray<FUintVector2> SourceAndTargetDataIndices;
+	};
+
+	FPCGCopyPointsDataProviderProxy(FCopyPointsData_RenderThread InData)
+		: Data(MoveTemp(InData))
 	{}
 
 	//~ Begin FComputeDataProviderRenderProxy Interface
 	bool IsValid(FValidationData const& InValidationData) const override;
+	void AllocateResources(FRDGBuilder& GraphBuilder, FAllocationData const& InAllocationData) override;
 	void GatherDispatchData(FDispatchData const& InDispatchData) override;
 	//~ End FComputeDataProviderRenderProxy Interface
 
 protected:
 	using FParameters = FPCGCopyPointsDataInterfaceParameters;
 
-	uint32 RotationInheritance = 0;
-	uint32 ApplyTargetRotationToPositions = 0;
-	uint32 ScaleInheritance = 0;
-	uint32 ApplyTargetScaleToPositions = 0;
-	uint32 ColorInheritance = 0;
-	uint32 SeedInheritance = 0;
-	uint32 AttributeInheritance = 0;
-	uint32 bCopyEachSourceOnEveryTarget = 0;
+	FCopyPointsData_RenderThread Data;
+
+	FRDGBufferRef SourceAndTargetDataIndicesBuffer = nullptr;
+	FRDGBufferSRVRef SourceAndTargetDataIndicesBufferSRV = nullptr;
 };
diff --git a/Engine/Plugins/PCG/Source/PCG/Private/Compute/PCGComputeGraph.h b/Engine/Plugins/PCG/Source/PCG/Private/Compute/PCGComputeGraph.h
--- a/Engine/Plugins/PCG/Source/PCG/Private/Compute/PCGComputeGraph.h
+++ b/Engine/Plugins/PCG/Source/PCG/Private/Compute/PCGComputeGraph.h
@@ -98,6 +98,8 @@ public:
 	/** Bindings in compute graphs roughly map to a single API on a Data Interface, like GetNumData(). There can be multiple bindings per PCG edge. */
 	int32 GetBindingIndex(int32 InKernelIndex, FName InPinLabel, bool bInIsInput) const;
 
+	const FName* GetVirtualPinLabelFromBindingIndex(int32 InBindingIndex) const { return CPUDataBindingToVirtualPinLabel.Find(InBindingIndex); }
+
 	bool IsGraphValid(FPCGContext* InContext = nullptr, bool bQuiet = true) const;
 
 	void DebugLogDataDescriptions(const UPCGDataBinding* InBinding) const;
diff --git a/Engine/Plugins/PCG/Source/PCG/Private/Compute/PCGDataForGPU.cpp b/Engine/Plugins/PCG/Source/PCG/Private/Compute/PCGDataForGPU.cpp
--- a/Engine/Plugins/PCG/Source/PCG/Private/Compute/PCGDataForGPU.cpp
+++ b/Engine/Plugins/PCG/Source/PCG/Private/Compute/PCGDataForGPU.cpp
@@ -720,6 +720,14 @@ bool FPCGDataDesc::ContainsAttribute(FName InAttributeName) const
 	});
 }
 
+bool FPCGDataDesc::ContainsAttribute(FName InAttributeName, EPCGKernelAttributeType InAttributeType) const
+{
+	return AttributeDescs.ContainsByPredicate([InAttributeName, InAttributeType](const FPCGKernelAttributeDesc& AttributeDesc)
+	{
+		return AttributeDesc.Name == InAttributeName && AttributeDesc.Type == InAttributeType;
+	});
+}
+
 void FPCGDataDesc::AddAttribute(FPCGKernelAttributeKey InAttribute, const UPCGDataBinding* InBinding, const TArray<int32>* InOptionalUniqueStringKeys)
 {
 	const int32 AttributeId = InBinding->GetAttributeId(InAttribute);
diff --git a/Engine/Plugins/PCG/Source/PCG/Private/Elements/PCGCopyPoints.cpp b/Engine/Plugins/PCG/Source/PCG/Private/Elements/PCGCopyPoints.cpp
--- a/Engine/Plugins/PCG/Source/PCG/Private/Elements/PCGCopyPoints.cpp
+++ b/Engine/Plugins/PCG/Source/PCG/Private/Elements/PCGCopyPoints.cpp
@@ -5,6 +5,7 @@
 #include "PCGContext.h"
 #include "Data/PCGBasePointData.h"
 #include "Data/PCGSpatialData.h"
+#include "Elements/PCGCopyPointsAnalysisKernel.h"
 #include "Elements/PCGCopyPointsKernel.h"
 #include "Graph/PCGGPUGraphCompilationContext.h"
 #include "Helpers/PCGAsync.h"
@@ -24,14 +25,25 @@ FText UPCGCopyPointsSettings::GetNodeTooltipText() const
 
 void UPCGCopyPointsSettings::CreateKernels(FPCGGPUCompilationContext& InOutContext, UObject* InObjectOuter, TArray<UPCGComputeKernel*>& OutKernels, TArray<FPCGKernelEdge>& OutEdges) const
 {
-	UPCGCopyPointsKernel* Kernel = InOutContext.NewObject_AnyThread<UPCGCopyPointsKernel>(InObjectOuter);
-	Kernel->Initialize(this, bDumpDataDescriptions);
-	OutKernels.Add(Kernel);
-
-	// Connect node pins to kernel pins
-	OutEdges.Emplace(FPCGPinReference(PCGCopyPointsConstants::SourcePointsLabel), FPCGPinReference(Kernel, PCGCopyPointsConstants::SourcePointsLabel));
-	OutEdges.Emplace(FPCGPinReference(PCGCopyPointsConstants::TargetPointsLabel), FPCGPinReference(Kernel, PCGCopyPointsConstants::TargetPointsLabel));
-	OutEdges.Emplace(FPCGPinReference(Kernel, PCGPinConstants::DefaultOutputLabel), FPCGPinReference(PCGPinConstants::DefaultOutputLabel));
+	UPCGCopyPointsKernel* CopyKernel = InOutContext.NewObject_AnyThread<UPCGCopyPointsKernel>(InObjectOuter);
+	CopyKernel->Initialize(this, bDumpDataDescriptions);
+	OutKernels.Add(CopyKernel);
+
+	OutEdges.Emplace(FPCGPinReference(PCGCopyPointsConstants::SourcePointsLabel), FPCGPinReference(CopyKernel, PCGCopyPointsConstants::SourcePointsLabel));
+	OutEdges.Emplace(FPCGPinReference(PCGCopyPointsConstants::TargetPointsLabel), FPCGPinReference(CopyKernel, PCGCopyPointsConstants::TargetPointsLabel));
+	OutEdges.Emplace(FPCGPinReference(CopyKernel, PCGPinConstants::DefaultOutputLabel), FPCGPinReference(PCGPinConstants::DefaultOutputLabel));
+
+	if (bMatchBasedOnAttribute)
+	{
+		UPCGCopyPointsAnalysisKernel* AnalysisKernel = InOutContext.NewObject_AnyThread<UPCGCopyPointsAnalysisKernel>(InObjectOuter);
+		AnalysisKernel->Initialize(this, bDumpDataDescriptions);
+		OutKernels.Add(AnalysisKernel);
+
+		// Connect to node pins and the copy kernel.
+		OutEdges.Emplace(FPCGPinReference(PCGCopyPointsConstants::SourcePointsLabel), FPCGPinReference(AnalysisKernel, PCGCopyPointsConstants::SourcePointsLabel));
+		OutEdges.Emplace(FPCGPinReference(PCGCopyPointsConstants::TargetPointsLabel), FPCGPinReference(AnalysisKernel, PCGCopyPointsConstants::TargetPointsLabel));
+		OutEdges.Emplace(FPCGPinReference(AnalysisKernel, PCGPinConstants::DefaultOutputLabel), FPCGPinReference(CopyKernel, PCGCopyPointsConstants::SelectedFlagsPinLabel));
+	}
 }
 #endif
 
diff --git a/Engine/Plugins/PCG/Source/PCG/Private/Elements/PCGCopyPointsAnalysisKernel.cpp b/Engine/Plugins/PCG/Source/PCG/Private/Elements/PCGCopyPointsAnalysisKernel.cpp
new file mode 100644
--- /dev/null
+++ b/Engine/Plugins/PCG/Source/PCG/Private/Elements/PCGCopyPointsAnalysisKernel.cpp
@@ -0,0 +1,146 @@
+// Copyright Epic Games, Inc. All Rights Reserved.
+
+#include "Elements/PCGCopyPointsAnalysisKernel.h"
+
+#include "PCGComponent.h"
+#include "PCGContext.h"
+#include "PCGModule.h"
+#include "Compute/PCGComputeGraph.h"
+#include "Compute/PCGDataBinding.h"
+#include "Compute/PCGPinPropertiesGPU.h"
+#include "Compute/DataInterfaces/Elements/PCGCopyPointsAnalysisDataInterface.h"
+#include "Compute/Elements/PCGComputeGraphElement.h"
+#include "Elements/PCGCopyPoints.h"
+#include "Elements/PCGCopyPointsKernelShared.h"
+#include "Graph/PCGGPUGraphCompilationContext.h"
+
+#include "ShaderCompilerCore.h"
+#include "Engine/StaticMesh.h"
+
+#include UE_INLINE_GENERATED_CPP_BY_NAME(PCGCopyPointsAnalysisKernel)
+
+#define LOCTEXT_NAMESPACE "PCGCopyPointsAnalysisKernel"
+
+void UPCGCopyPointsAnalysisKernel::Initialize(const UPCGCopyPointsSettings* InSettings, bool bLogDescriptions)
+{
+	Super::Initialize(bLogDescriptions);
+
+	Settings = InSettings;
+}
+
+const UPCGSettings* UPCGCopyPointsAnalysisKernel::GetSettings() const
+{
+	return Settings;
+}
+
+bool UPCGCopyPointsAnalysisKernel::IsKernelValid(FPCGContext* InContext, bool bQuiet) const
+{
+	TRACE_CPUPROFILER_EVENT_SCOPE(UPCGCopyPointsAnalysisKernel::IsKernelValid);
+	const UPCGCopyPointsSettings* CopyPointSettings = Cast<UPCGCopyPointsSettings>(Settings);
+
+	if (!Super::IsKernelValid(InContext, bQuiet) || !ensure(CopyPointSettings))
+	{
+		return false;
+	}
+
+	if (!ValidatePCGNode(InContext, bQuiet, CopyPointSettings))
+	{
+		return false;
+	}
+
+	if (!InContext || !InContext->IsComputeContext())
+	{
+		return true;
+	}
+
+	const FPCGComputeGraphContext* Context = static_cast<FPCGComputeGraphContext*>(InContext);
+
+	return PCGCopyPointsKernel::IsKernelValid(this, CopyPointSettings, Context, bQuiet);
+}
+
+bool UPCGCopyPointsAnalysisKernel::ComputeOutputBindingDataDesc(const UPCGComputeGraph* InGraph, FName InOutputPinLabel, const UPCGDataBinding* InBinding, FPCGDataCollectionDesc& OutDataDesc) const
+{
+	check(InGraph);
+	check(InBinding);
+
+	// Code assumes single output pin.
+	if (!ensure(InOutputPinLabel == PCGPinConstants::DefaultOutputLabel))
+	{
+		return false;
+	}
+
+	FPCGDataCollectionDesc SourcesDesc, TargetsDesc;
+	ensure(InBinding->GetKernelPinDataDesc_GameThread(this, PCGCopyPointsConstants::SourcePointsLabel, /*bIsInput=*/true, SourcesDesc));
+	ensure(InBinding->GetKernelPinDataDesc_GameThread(this, PCGCopyPointsConstants::TargetPointsLabel, /*bIsInput=*/true, TargetsDesc));
+
+	const uint32 NumSourceDatas = SourcesDesc.DataDescs.Num();
+	const uint32 NumTargetDatas = TargetsDesc.DataDescs.Num();
+
+	// Output: a single attribute set with a single attribute true/false value per NumSource*NumTarget possible outputs.
+	FPCGDataCollectionDesc OutDesc;
+
+	FPCGDataDesc& DataDesc = OutDesc.DataDescs.Emplace_GetRef();
+	DataDesc.Type = EPCGDataType::Param;
+	DataDesc.ElementCount = NumSourceDatas * NumTargetDatas;
+
+	FPCGKernelAttributeKey SelectedFlagAttributeKey = FPCGKernelAttributeKey(PCGCopyPointsConstants::SelectedFlagAttributeName, EPCGKernelAttributeType::Bool);
+	DataDesc.AddAttribute(SelectedFlagAttributeKey, InBinding);
+
+	OutDataDesc = MoveTemp(OutDesc);
+
+	return true;
+}
+
+int UPCGCopyPointsAnalysisKernel::ComputeThreadCount(const UPCGDataBinding* InBinding) const
+{
+	int ThreadCount = 0;
+
+	FPCGDataCollectionDesc OutputPinDesc;
+	if (ensure(InBinding->GetKernelPinDataDesc_GameThread(this, PCGPinConstants::DefaultOutputLabel, /*bIsInput=*/false, OutputPinDesc)))
+	{
+		for (const FPCGDataDesc& DataDesc : OutputPinDesc.DataDescs)
+		{
+			ThreadCount += DataDesc.ElementCount;
+		}
+	}
+
+	return ThreadCount;
+}
+
+#if WITH_EDITOR
+FString UPCGCopyPointsAnalysisKernel::GetCookedSource(const UPCGComputeGraph* InComputeGraph) const
+{
+	FString TemplateFile;
+	ensure(LoadShaderSourceFile(TEXT("/Plugin/PCG/Private/Elements/PCGCopyPointsAnalysis.usf"), EShaderPlatform::SP_PCD3D_SM5, &TemplateFile, nullptr));
+	return TemplateFile;
+}
+
+void UPCGCopyPointsAnalysisKernel::CreateAdditionalInputDataInterfaces(FPCGGPUCompilationContext& InOutContext, UObject* InObjectOuter, TArray<TObjectPtr<UComputeDataInterface>>& OutDataInterfaces) const
+{
+	Super::CreateAdditionalInputDataInterfaces(InOutContext, InObjectOuter, OutDataInterfaces);
+
+	TObjectPtr<UPCGCopyPointsAnalysisDataInterface> NodeDI = InOutContext.NewObject_AnyThread<UPCGCopyPointsAnalysisDataInterface>(InObjectOuter);
+	NodeDI->ProducerKernel = this;
+
+	OutDataInterfaces.Add(NodeDI);
+}
+#endif
+
+void UPCGCopyPointsAnalysisKernel::GetKernelAttributeKeys(TArray<FPCGKernelAttributeKey>& OutKeys) const
+{
+	// Register the attribute this node creates.
+	OutKeys.AddUnique(FPCGKernelAttributeKey(PCGCopyPointsConstants::SelectedFlagAttributeName, EPCGKernelAttributeType::Bool));
+}
+
+void UPCGCopyPointsAnalysisKernel::GetInputPins(TArray<FPCGPinProperties>& OutPins) const
+{
+	OutPins.Emplace(PCGCopyPointsConstants::SourcePointsLabel, EPCGDataType::Point);
+	OutPins.Emplace(PCGCopyPointsConstants::TargetPointsLabel, EPCGDataType::Point);
+}
+
+void UPCGCopyPointsAnalysisKernel::GetOutputPins(TArray<FPCGPinPropertiesGPU>& OutPins) const
+{
+	OutPins.Emplace(PCGPinConstants::DefaultOutputLabel, EPCGDataType::Point);
+}
+
+#undef LOCTEXT_NAMESPACE
diff --git a/Engine/Plugins/PCG/Source/PCG/Private/Elements/PCGCopyPointsKernel.cpp b/Engine/Plugins/PCG/Source/PCG/Private/Elements/PCGCopyPointsKernel.cpp
--- a/Engine/Plugins/PCG/Source/PCG/Private/Elements/PCGCopyPointsKernel.cpp
+++ b/Engine/Plugins/PCG/Source/PCG/Private/Elements/PCGCopyPointsKernel.cpp
@@ -4,10 +4,14 @@
 
 #include "PCGContext.h"
 #include "PCGModule.h"
+#include "PCGParamData.h"
 #include "Compute/PCGComputeGraph.h"
 #include "Compute/PCGDataBinding.h"
 #include "Compute/DataInterfaces/Elements/PCGCopyPointsDataInterface.h"
+#include "Compute/Data/PCGProxyForGPUData.h"
+#include "Compute/Elements/PCGComputeGraphElement.h"
 #include "Elements/PCGCopyPoints.h"
+#include "Elements/PCGCopyPointsKernelShared.h"
 #include "Graph/PCGGPUGraphCompilationContext.h"
 
 #include "ShaderCompilerCore.h"
@@ -33,13 +37,26 @@ const UPCGSettings* UPCGCopyPointsKernel::GetSettings() const
 bool UPCGCopyPointsKernel::IsKernelValid(FPCGContext* InContext, bool bQuiet) const
 {
 	TRACE_CPUPROFILER_EVENT_SCOPE(UPCGCopyPointsKernel::IsKernelValid);
+	const UPCGCopyPointsSettings* CopyPointSettings = Cast<UPCGCopyPointsSettings>(Settings);
 
-	if (!Super::IsKernelValid(InContext, bQuiet) || !ensure(Settings))
+	if (!Super::IsKernelValid(InContext, bQuiet) || !ensure(CopyPointSettings))
 	{
 		return false;
 	}
 
-	return ValidatePCGNode(InContext, bQuiet, Cast<UPCGSettings>(Settings));
+	if (!ValidatePCGNode(InContext, bQuiet, CopyPointSettings))
+	{
+		return false;
+	}
+
+	if (!InContext || !InContext->IsComputeContext())
+	{
+		return true;
+	}
+
+	const FPCGComputeGraphContext* Context = static_cast<FPCGComputeGraphContext*>(InContext);
+
+	return PCGCopyPointsKernel::IsKernelValid(this, CopyPointSettings, Context, bQuiet);
 }
 
 bool UPCGCopyPointsKernel::ComputeOutputBindingDataDesc(const UPCGComputeGraph* InGraph, FName InOutputPinLabel, const UPCGDataBinding* InBinding, FPCGDataCollectionDesc& OutDataDesc) const
@@ -53,6 +70,64 @@ bool UPCGCopyPointsKernel::ComputeOutputBindingDataDesc(const UPCGComputeGraph*
 		return false;
 	}
 
+	// A graph split was injected before this kernel. We expect to find selected flags
+	int NumSelected = INDEX_NONE;
+	TArray<bool> Selected;
+	
+	if (Settings->bMatchBasedOnAttribute)
+	{
+		int32 FirstAnalysisDataIndex = INDEX_NONE;
+
+		if (InBinding->Graph)
+		{
+			const int32 BindingIndex = InBinding->Graph->GetBindingIndex(KernelIndex, PCGCopyPointsConstants::SelectedFlagsPinLabel, /*bIsInput=*/true);
+			if (ensure(BindingIndex != INDEX_NONE))
+			{
+				if (const FName* VirtualLabel = InGraph->GetVirtualPinLabelFromBindingIndex(BindingIndex))
+				{
+					FirstAnalysisDataIndex = InBinding->InputDataCollection.TaggedData.IndexOfByPredicate([VirtualLabel](const FPCGTaggedData& InData)
+					{
+						return InData.Pin == *VirtualLabel;
+					});
+				}
+			}
+		}
+
+		if (FirstAnalysisDataIndex != INDEX_NONE)
+		{
+			const UPCGParamData* AnalysisResultsData = Cast<UPCGParamData>(InBinding->InputDataCollection.TaggedData[FirstAnalysisDataIndex].Data);
+			const UPCGMetadata* AnalysisMetadata = AnalysisResultsData ? AnalysisResultsData->ConstMetadata() : nullptr;
+			const FPCGMetadataAttributeBase* AnalysisAttributeBase = AnalysisMetadata ? AnalysisMetadata->GetConstAttribute(PCGCopyPointsConstants::SelectedFlagAttributeName) : nullptr;
+
+			if (AnalysisAttributeBase && AnalysisAttributeBase->GetTypeId() == PCG::Private::MetadataTypes<bool>::Id)
+			{
+				const FPCGMetadataAttribute<bool>* AnalysisAttribute = static_cast<const FPCGMetadataAttribute<bool>*>(AnalysisAttributeBase);
+				const int32 NumElements = AnalysisMetadata->GetItemCountForChild();
+
+				Selected.SetNumUninitialized(NumElements);
+
+				NumSelected = 0;
+
+				for (int64 MetadataKey = 0; MetadataKey < NumElements; ++MetadataKey)
+				{
+					const bool bSelected = AnalysisAttribute->GetValue(MetadataKey);
+					Selected[MetadataKey] = bSelected;
+					NumSelected += bSelected ? 1 : 0;
+				}
+			}
+			else
+			{
+				UE_LOG(LogPCG, Warning, TEXT("No analysis data received by copy points kernel, no points will be copied."));
+				if (InBinding->InputDataCollection.TaggedData[FirstAnalysisDataIndex].Data->IsA<UPCGProxyForGPUData>())
+				{
+					UE_LOG(LogPCG, Error, TEXT("Data was not read back."));
+				}
+				OutDataDesc = {};
+				return true;
+			}
+		}
+	}
+
 	FPCGDataCollectionDesc SourcePinDesc;
 	FPCGDataCollectionDesc TargetPinDesc;
 	ensure(InBinding->GetKernelPinDataDesc_GameThread(this, PCGCopyPointsConstants::SourcePointsLabel, /*bIsInput=*/true, SourcePinDesc));
@@ -68,6 +143,11 @@ bool UPCGCopyPointsKernel::ComputeOutputBindingDataDesc(const UPCGComputeGraph*
 	{
 		for (int32 I = 0; I < NumIterations; ++I)
 		{
+			if (!Selected.IsEmpty() && !Selected[I])
+			{
+				continue;
+			}
+
 			const int32 SourceIndex = Settings->bCopyEachSourceOnEveryTarget ? (I / NumTargets) : FMath::Min(I, NumSources - 1);
 			const int32 TargetIndex = Settings->bCopyEachSourceOnEveryTarget ? (I % NumTargets) : FMath::Min(I, NumTargets - 1);
 
@@ -145,7 +225,17 @@ int UPCGCopyPointsKernel::ComputeThreadCount(const UPCGDataBinding* InBinding) c
 FString UPCGCopyPointsKernel::GetCookedSource(const UPCGComputeGraph* InComputeGraph) const
 {
 	FString TemplateFile;
-	ensure(LoadShaderSourceFile(TEXT("/Plugin/PCG/Private/Elements/PCGCopyPoints.usf"), EShaderPlatform::SP_PCD3D_SM5, &TemplateFile, nullptr));
+
+	// todo_pcg: Figure out how permutations work (GetPermutations/GatherPermutations) so we can unify the two shader files into one.
+	if (Settings && Settings->bMatchBasedOnAttribute)
+	{
+		ensure(LoadShaderSourceFile(TEXT("/Plugin/PCG/Private/Elements/PCGCopyPointsConditional.usf"), EShaderPlatform::SP_PCD3D_SM5, &TemplateFile, nullptr));
+	}
+	else
+	{
+		ensure(LoadShaderSourceFile(TEXT("/Plugin/PCG/Private/Elements/PCGCopyPoints.usf"), EShaderPlatform::SP_PCD3D_SM5, &TemplateFile, nullptr));
+	}
+
 	return TemplateFile;
 }
 
@@ -164,6 +254,12 @@ void UPCGCopyPointsKernel::GetInputPins(TArray<FPCGPinProperties>& OutPins) cons
 {
 	OutPins.Emplace(PCGCopyPointsConstants::SourcePointsLabel, EPCGDataType::Point);
 	OutPins.Emplace(PCGCopyPointsConstants::TargetPointsLabel, EPCGDataType::Point);
+
+	if (Settings->bMatchBasedOnAttribute)
+	{
+		// One bool flag per output data that signals if the output data should be computed or not.
+		OutPins.Emplace(PCGCopyPointsConstants::SelectedFlagsPinLabel, EPCGDataType::Param);
+	}
 }
 
 void UPCGCopyPointsKernel::GetOutputPins(TArray<FPCGPinPropertiesGPU>& OutPins) const
diff --git a/Engine/Plugins/PCG/Source/PCG/Private/Elements/PCGCopyPointsKernelShared.cpp b/Engine/Plugins/PCG/Source/PCG/Private/Elements/PCGCopyPointsKernelShared.cpp
new file mode 100644
--- /dev/null
+++ b/Engine/Plugins/PCG/Source/PCG/Private/Elements/PCGCopyPointsKernelShared.cpp
@@ -0,0 +1,60 @@
+// Copyright Epic Games, Inc. All Rights Reserved.
+
+#include "Elements/PCGCopyPointsKernelShared.h"
+
+#include "PCGContext.h"
+#include "Compute/PCGDataBinding.h"
+#include "Compute/Elements/PCGComputeGraphElement.h"
+#include "Elements/PCGCopyPoints.h"
+#include "Graph/PCGGPUGraphCompilationContext.h"
+
+#define LOCTEXT_NAMESPACE "PCGCopyPointsKernel"
+
+bool PCGCopyPointsKernel::IsKernelValid(const UPCGComputeKernel* InKernel, const UPCGCopyPointsSettings* InCopyPointSettings, const FPCGComputeGraphContext* InContext, bool bQuiet)
+{
+	if (InCopyPointSettings->bMatchBasedOnAttribute)
+	{
+		if (!InCopyPointSettings->bCopyEachSourceOnEveryTarget)
+		{
+			PCG_KERNEL_VALIDATION_ERR(InContext, InCopyPointSettings, bQuiet, LOCTEXT("CopyEachSourceOnEveryTargetNotSupported", "Copy Each Source On Every Target must be enabled when Match Based On Attribute is enabled."));
+			return false;
+		}
+
+		if (const UPCGDataBinding* DataBinding = InContext->DataBinding.Get(); ensure(DataBinding))
+		{
+			auto ValidateAttributeExists = [InContext, DataBinding, InKernel, InCopyPointSettings, bQuiet](FName InputPin) -> bool
+			{
+				FPCGDataCollectionDesc PinDataDesc;
+				ensure(DataBinding->GetKernelPinDataDesc_GameThread(InKernel, InputPin, /*bIsInputPin=*/true, PinDataDesc));
+
+				const FName MatchAttributeName = InCopyPointSettings->MatchAttribute;
+
+				// todo_pcg: Can generalize this to any type?
+				constexpr EPCGKernelAttributeType MatchAttributeType = EPCGKernelAttributeType::Int;
+
+				for (const FPCGDataDesc& DataDesc : PinDataDesc.DataDescs)
+				{
+					if (!DataDesc.ContainsAttribute(MatchAttributeName, MatchAttributeType))
+					{
+						PCG_KERNEL_VALIDATION_ERR(InContext, InCopyPointSettings, bQuiet, FText::Format(
+							LOCTEXT("MatchAttributeMissing", "Match attribute '{0}' not found, this attribute must be present on all input data, and be of type Integer."),
+							FText::FromName(MatchAttributeName)));
+
+						return false;
+					}
+				}
+
+				return true;
+			};
+
+			if (!ValidateAttributeExists(PCGCopyPointsConstants::SourcePointsLabel) || !ValidateAttributeExists(PCGCopyPointsConstants::TargetPointsLabel))
+			{
+				return false;
+			}
+		}
+	}
+
+	return true;
+}
+
+#undef LOCTEXT_NAMESPACE
diff --git a/Engine/Plugins/PCG/Source/PCG/Public/Elements/Metadata/PCGMetadataPartitionKernel.h b/Engine/Plugins/PCG/Source/PCG/Public/Elements/Metadata/PCGMetadataPartitionKernel.h
--- a/Engine/Plugins/PCG/Source/PCG/Public/Elements/Metadata/PCGMetadataPartitionKernel.h
+++ b/Engine/Plugins/PCG/Source/PCG/Public/Elements/Metadata/PCGMetadataPartitionKernel.h
@@ -23,6 +23,7 @@ public:
 #if WITH_EDITOR
 	virtual FString GetCookedSource(const UPCGComputeGraph* InComputeGraph) const override;
 	void CreateAdditionalInputDataInterfaces(FPCGGPUCompilationContext& InOutContext, UObject* InObjectOuter, TArray<TObjectPtr<UComputeDataInterface>>& OutDataInterfaces) const override;
+	// Split graph to update element counts on CPU before proceeding. Compacts data for downstream kernels.
 	virtual bool SplitGraphAtOutput() const override { return true; }
 #endif
 	virtual void GetKernelAttributeKeys(TArray<FPCGKernelAttributeKey>& OutKeys) const override;
diff --git a/Engine/Plugins/PCG/Source/PCG/Public/Elements/PCGCopyPoints.h b/Engine/Plugins/PCG/Source/PCG/Public/Elements/PCGCopyPoints.h
--- a/Engine/Plugins/PCG/Source/PCG/Public/Elements/PCGCopyPoints.h
+++ b/Engine/Plugins/PCG/Source/PCG/Public/Elements/PCGCopyPoints.h
@@ -10,6 +10,8 @@ namespace PCGCopyPointsConstants
 {
 	const FName SourcePointsLabel = TEXT("Source");
 	const FName TargetPointsLabel = TEXT("Target");
+	const FName SelectedFlagsPinLabel = TEXT("SelectedFlags");
+	const FName SelectedFlagAttributeName = TEXT("SelectForCopy");
 }
 
 UENUM()
@@ -96,6 +98,20 @@ public:
 	/** If this option is set, each source point data will be copied to every target point data (cartesian product), producing N * M point data. Otherwise, will do a N:N (or N:1 or 1:N) operation, producing N point data. */
 	UPROPERTY(BlueprintReadWrite, EditAnywhere, Category = Settings, meta = (PCG_Overridable))
 	bool bCopyEachSourceOnEveryTarget = true;
+
+	// Note: Not overridable currently. Used to make compile time decision about which kernels to emit. If this needs to be overridable then we
+	// need to investigate how to deal with this at execution time.
+	/** Perform a conditional copy point where data pairs must have a matching attribute in order to participate in the copy operation.
+	* EXPERIMENTAL: This feature is work in progress and is currently GPU-only. It may not be stable, and may be removed in the future.
+	*/
+	UPROPERTY(BlueprintReadWrite, EditAnywhere, Category = Settings, AdvancedDisplay, meta = (EditCondition = "bExecuteOnGPU", EditConditionHides))
+	bool bMatchBasedOnAttribute = false;
+
+	/** Attribute that must be present on both the source data and the target point, and have the same value.
+	* EXPERIMENTAL: This feature is work in progress and is currently GPU-only. It may not be stable, and may be removed in the future.
+	*/
+	UPROPERTY(BlueprintReadWrite, EditAnywhere, Category = Settings, AdvancedDisplay, meta = (EditCondition = "bExecuteOnGPU && bMatchBasedOnAttribute", EditConditionHides))
+	FName MatchAttribute = NAME_None;
 };
 
 class FPCGCopyPointsElement : public IPCGElement
diff --git a/Engine/Plugins/PCG/Source/PCG/Public/Elements/PCGCopyPointsAnalysisKernel.h b/Engine/Plugins/PCG/Source/PCG/Public/Elements/PCGCopyPointsAnalysisKernel.h
new file mode 100644
--- /dev/null
+++ b/Engine/Plugins/PCG/Source/PCG/Public/Elements/PCGCopyPointsAnalysisKernel.h
@@ -0,0 +1,37 @@
+// Copyright Epic Games, Inc. All Rights Reserved.
+
+#pragma once
+
+#include "Compute/PCGComputeKernel.h"
+
+#include "PCGCopyPointsAnalysisKernel.generated.h"
+
+class UPCGCopyPointsSettings;
+
+UCLASS()
+class UPCGCopyPointsAnalysisKernel : public UPCGComputeKernel
+{
+	GENERATED_BODY()
+
+public:
+	void Initialize(const UPCGCopyPointsSettings* InSettings, bool bLogDescriptions);
+
+	virtual const UPCGSettings* GetSettings() const override;
+	virtual bool IsKernelValid(FPCGContext* InContext = nullptr, bool bQuiet = true) const override;
+	virtual bool ComputeOutputBindingDataDesc(const UPCGComputeGraph* InGraph, FName InOutputPinLabel, const UPCGDataBinding* InBinding, FPCGDataCollectionDesc& OutDataDesc) const override;
+	virtual int ComputeThreadCount(const UPCGDataBinding* Binding) const override;
+#if WITH_EDITOR
+	virtual FString GetCookedSource(const UPCGComputeGraph* InComputeGraph) const override;
+	void CreateAdditionalInputDataInterfaces(FPCGGPUCompilationContext& InOutContext, UObject* InObjectOuter, TArray<TObjectPtr<UComputeDataInterface>>& OutDataInterfaces) const override;
+	// Split graph to read back analysis results.
+	virtual bool SplitGraphAtOutput() const override { return true; }
+#endif
+	virtual void GetKernelAttributeKeys(TArray<FPCGKernelAttributeKey>& OutKeys) const override;
+
+	virtual void GetInputPins(TArray<FPCGPinProperties>& OutPins) const override;
+	virtual void GetOutputPins(TArray<FPCGPinPropertiesGPU>& OutPins) const override;
+
+public:
+	UPROPERTY()
+	TObjectPtr<const UPCGCopyPointsSettings> Settings;
+};
diff --git a/Engine/Plugins/PCG/Source/PCG/Public/Elements/PCGCopyPointsKernelShared.h b/Engine/Plugins/PCG/Source/PCG/Public/Elements/PCGCopyPointsKernelShared.h
new file mode 100644
--- /dev/null
+++ b/Engine/Plugins/PCG/Source/PCG/Public/Elements/PCGCopyPointsKernelShared.h
@@ -0,0 +1,14 @@
+// Copyright Epic Games, Inc. All Rights Reserved.
+
+#pragma once
+
+class UPCGComputeKernel;
+class UPCGCopyPointsSettings;
+struct FPCGComputeGraphContext;
+struct FPCGContext;
+
+namespace PCGCopyPointsKernel
+{
+	/** Performs validation common to all copy points kernels. */
+	bool IsKernelValid(const UPCGComputeKernel* InKernel, const UPCGCopyPointsSettings* InCopyPointSettings, const FPCGComputeGraphContext* InContext, bool bQuiet);
+}
