[PCG] Fix unallocated properties not being properly initialized with some filtering elements (causing a different results when pcg.Experimental.EnablePointArrayData == 1)
- Add a new CopyUnallocatedProperties method that needs to be called for filtering algos in the initialization code
- Convert some elements to use OneToOne range async since they don't do any filtering

#rnx

diff --git a/Engine/Plugins/PCG/Source/PCG/Private/Data/PCGPointArrayData.cpp b/Engine/Plugins/PCG/Source/PCG/Private/Data/PCGPointArrayData.cpp
--- a/Engine/Plugins/PCG/Source/PCG/Private/Data/PCGPointArrayData.cpp
+++ b/Engine/Plugins/PCG/Source/PCG/Private/Data/PCGPointArrayData.cpp
@@ -224,6 +224,21 @@ void UPCGPointArrayData::MoveRange(int32 RangeStartIndex, int32 MoveToIndex, int
 	PointArray.MoveRange(RangeStartIndex, MoveToIndex, NumElements);
 }
 
+void UPCGPointArrayData::CopyUnallocatedPropertiesTo(UPCGBasePointData* OutPointData) const
+{
+	if (UPCGPointArrayData* OutPointArrayData = Cast<UPCGPointArrayData>(OutPointData); OutPointArrayData && !OutPointArrayData->HasSpatialDataParent())
+	{
+		GetProperty<FTransform>(EPCGPointNativeProperties::Transform)->CopyUnallocatedProperty(*OutPointArrayData->GetProperty<FTransform>(EPCGPointNativeProperties::Transform, /*bWithInheritance=*/false));
+		GetProperty<float>(EPCGPointNativeProperties::Density)->CopyUnallocatedProperty(*OutPointArrayData->GetProperty<float>(EPCGPointNativeProperties::Density, /*bWithInheritance=*/false));
+		GetProperty<FVector>(EPCGPointNativeProperties::BoundsMin)->CopyUnallocatedProperty(*OutPointArrayData->GetProperty<FVector>(EPCGPointNativeProperties::BoundsMin, /*bWithInheritance=*/false));
+		GetProperty<FVector>(EPCGPointNativeProperties::BoundsMax)->CopyUnallocatedProperty(*OutPointArrayData->GetProperty<FVector>(EPCGPointNativeProperties::BoundsMax, /*bWithInheritance=*/false));
+		GetProperty<FVector4>(EPCGPointNativeProperties::Color)->CopyUnallocatedProperty(*OutPointArrayData->GetProperty<FVector4>(EPCGPointNativeProperties::Color, /*bWithInheritance=*/false));
+		GetProperty<float>(EPCGPointNativeProperties::Steepness)->CopyUnallocatedProperty(*OutPointArrayData->GetProperty<float>(EPCGPointNativeProperties::Steepness, /*bWithInheritance=*/false));
+		GetProperty<int32>(EPCGPointNativeProperties::Seed)->CopyUnallocatedProperty(*OutPointArrayData->GetProperty<int32>(EPCGPointNativeProperties::Seed, /*bWithInheritance=*/false));
+		GetProperty<int64>(EPCGPointNativeProperties::MetadataEntry)->CopyUnallocatedProperty(*OutPointArrayData->GetProperty<int64>(EPCGPointNativeProperties::MetadataEntry, /*bWithInheritance=*/false));
+	}
+}
+
 TPCGValueRange<FTransform> UPCGPointArrayData::GetTransformValueRange(bool bAllocate)
 { 
 	FlattenPropertiesIfNeeded(EPCGPointNativeProperties::Transform);
diff --git a/Engine/Plugins/PCG/Source/PCG/Private/Data/PCGProjectionData.cpp b/Engine/Plugins/PCG/Source/PCG/Private/Data/PCGProjectionData.cpp
--- a/Engine/Plugins/PCG/Source/PCG/Private/Data/PCGProjectionData.cpp
+++ b/Engine/Plugins/PCG/Source/PCG/Private/Data/PCGProjectionData.cpp
@@ -312,6 +312,7 @@ const UPCGBasePointData* UPCGProjectionData::CreateBasePointData(FPCGContext* Co
 
 		// @todo_pcg: try to optimize allocations here
 		PointData->AllocateProperties(SourcePointData->GetAllocatedProperties() | EPCGPointNativeProperties::Transform | EPCGPointNativeProperties::Color | EPCGPointNativeProperties::Density);
+		SourcePointData->CopyUnallocatedPropertiesTo(PointData);
 	};
 
 	auto MoveDataRangeFunc = [PointData](int32 RangeStartIndex, int32 MoveToIndex, int32 NumElements)
diff --git a/Engine/Plugins/PCG/Source/PCG/Private/Elements/PCGApplyHierarchy.cpp b/Engine/Plugins/PCG/Source/PCG/Private/Elements/PCGApplyHierarchy.cpp
--- a/Engine/Plugins/PCG/Source/PCG/Private/Elements/PCGApplyHierarchy.cpp
+++ b/Engine/Plugins/PCG/Source/PCG/Private/Elements/PCGApplyHierarchy.cpp
@@ -567,6 +567,7 @@ bool FPCGApplyHierarchyElement::ExecuteInternal(FPCGContext* InContext) const
 				{
 					FilteredData->SetNumPoints(OriginalData->GetNumPoints());
 					FilteredData->AllocateProperties(OriginalData->GetAllocatedProperties());
+					OriginalData->CopyUnallocatedPropertiesTo(FilteredData);
 				};
 
 				auto AsyncProcessRangeFunc = [FilteredData, OriginalData, &IterState](int32 StartReadIndex, int32 StartWriteIndex, int32 Count)
diff --git a/Engine/Plugins/PCG/Source/PCG/Private/Elements/PCGCreatePointsGrid.cpp b/Engine/Plugins/PCG/Source/PCG/Private/Elements/PCGCreatePointsGrid.cpp
--- a/Engine/Plugins/PCG/Source/PCG/Private/Elements/PCGCreatePointsGrid.cpp
+++ b/Engine/Plugins/PCG/Source/PCG/Private/Elements/PCGCreatePointsGrid.cpp
@@ -211,10 +211,8 @@ bool FPCGCreatePointsGridElement::ExecuteInternal(FPCGContext* Context) const
 	const int32 NumIterations = static_cast<int32>(NumIterations64);
 		
 	UPCGBasePointData* PtData = nullptr;
-	TPCGValueRange<FTransform> TransformRange;
-	TPCGValueRange<int32> SeedRange;
 
-	auto InitFunc = [Context, Settings, CellSize, NumIterations, &PtData, &TransformRange, &SeedRange]()
+	auto InitFunc = [Context, Settings, CellSize, NumIterations, &PtData]()
 	{ 
 		PtData = FPCGContext::NewPointData_AnyThread(Context);
 		
@@ -229,12 +227,14 @@ bool FPCGCreatePointsGridElement::ExecuteInternal(FPCGContext* Context) const
 			PtData->SetExtents(CellSize * 0.5);
 		}
 						
-		TransformRange = PtData->GetTransformValueRange();
-		SeedRange = PtData->GetSeedValueRange();
+		PtData->AllocateProperties(EPCGPointNativeProperties::Transform | EPCGPointNativeProperties::Seed);
 	};
 
-	auto ProcessFunc = [Settings, CullingShape, &PointCountX, &PointCountY, &PointOffsetX, &PointOffsetY, &PointOffsetZ, &CellSize, &GridExtents, &LocalTransform, &TransformRange, &SeedRange](int32 StartIndex, int32 StartWriteIndex, int32 Count)
+	auto ProcessFunc = [Settings, CullingShape, &PointCountX, &PointCountY, &PointOffsetX, &PointOffsetY, &PointOffsetZ, &CellSize, &GridExtents, &LocalTransform, &PtData](int32 StartIndex, int32 StartWriteIndex, int32 Count)
 	{
+		TPCGValueRange<FTransform> TransformRange = PtData->GetTransformValueRange(/*bAllocate=*/false);
+		TPCGValueRange<int32> SeedRange = PtData->GetSeedValueRange(/*bAllocate=*/false);
+
 		ensure(TransformRange.Num() == SeedRange.Num() && (StartWriteIndex + Count) <= TransformRange.Num());
 
 		int32 NumWritten = 0;
diff --git a/Engine/Plugins/PCG/Source/PCG/Private/Elements/PCGDensityFilter.cpp b/Engine/Plugins/PCG/Source/PCG/Private/Elements/PCGDensityFilter.cpp
--- a/Engine/Plugins/PCG/Source/PCG/Private/Elements/PCGDensityFilter.cpp
+++ b/Engine/Plugins/PCG/Source/PCG/Private/Elements/PCGDensityFilter.cpp
@@ -88,6 +88,7 @@ bool FPCGDensityFilterElement::ExecuteInternal(FPCGContext* Context) const
 		{
 			FilteredData->SetNumPoints(OriginalData->GetNumPoints());
 			FilteredData->AllocateProperties(OriginalData->GetAllocatedProperties());
+			OriginalData->CopyUnallocatedPropertiesTo(FilteredData);
 		};
 
 		auto AsyncProcessRangeFunc = [FilteredData, OriginalData, MinBound, MaxBound, bInvertFilter, bKeepZeroDensityPoints](int32 StartReadIndex, int32 StartWriteIndex, int32 Count)
diff --git a/Engine/Plugins/PCG/Source/PCG/Private/Elements/PCGPointOperationElementBase.cpp b/Engine/Plugins/PCG/Source/PCG/Private/Elements/PCGPointOperationElementBase.cpp
--- a/Engine/Plugins/PCG/Source/PCG/Private/Elements/PCGPointOperationElementBase.cpp
+++ b/Engine/Plugins/PCG/Source/PCG/Private/Elements/PCGPointOperationElementBase.cpp
@@ -130,13 +130,11 @@ PRAGMA_ENABLE_DEPRECATION_WARNINGS
 			return Count;
 		};
 
-		return FPCGAsync::AsyncProcessingRangeEx(
+		return FPCGAsync::AsyncProcessingOneToOneRangeEx(
 			&Context->AsyncState,
 			IterState.NumPoints,
 			/*InitializeFunc=*/[] {}, // Not useful for this context, since its preferred to initialize in PrepareDataInternal, so empty lambda
 			std::move(ProcessRangeFunc),
-			/*MoveRangeFunc=*/[](int32 ReadIndex, int32 WriteIndex, int32 Count) {}, // This is a one to one element so no point get discarded, so empty lambda
-			/*FinalizeFunc=*/[](int32) {}, // Not useful here, so empty lambda
 			Context->TimeSliceIsEnabled(),
 			PointsPerChunk);
 	}
diff --git a/Engine/Plugins/PCG/Source/PCG/Private/Elements/PCGSelectPoints.cpp b/Engine/Plugins/PCG/Source/PCG/Private/Elements/PCGSelectPoints.cpp
--- a/Engine/Plugins/PCG/Source/PCG/Private/Elements/PCGSelectPoints.cpp
+++ b/Engine/Plugins/PCG/Source/PCG/Private/Elements/PCGSelectPoints.cpp
@@ -114,6 +114,7 @@ bool FPCGSelectPointsElement::ExecuteInternal(FPCGContext* Context) const
 			{
 				SampledData->SetNumPoints(OriginalData->GetNumPoints());
 				SampledData->AllocateProperties(OriginalData->GetAllocatedProperties());
+				OriginalData->CopyUnallocatedPropertiesTo(SampledData);
 			};
 
 			auto AsyncProcessRangeFunc = [SampledData, OriginalData, Seed, Ratio, bKeepZeroDensityPoints](int32 StartReadIndex, int32 StartWriteIndex, int32 Count)
diff --git a/Engine/Plugins/PCG/Source/PCG/Private/Elements/PCGTransformPoints.cpp b/Engine/Plugins/PCG/Source/PCG/Private/Elements/PCGTransformPoints.cpp
--- a/Engine/Plugins/PCG/Source/PCG/Private/Elements/PCGTransformPoints.cpp
+++ b/Engine/Plugins/PCG/Source/PCG/Private/Elements/PCGTransformPoints.cpp
@@ -233,13 +233,11 @@ bool FPCGTransformPointsElement::ExecuteInternal(FPCGContext* Context) const
 			return NumWritten;
 		};
 				
-		FPCGAsync::AsyncProcessingRangeEx(
+		FPCGAsync::AsyncProcessingOneToOneRangeEx(
 			&Context->AsyncState,
 			PointData->GetNumPoints(),
 			/*InitializeFunc=*/[](){},
 			ProcessRangeFunc,
-			/*MoveRangeFunc=*/[](int32 ReadIndex, int32 WriteIndex, int32 Count) {}, // This is a one to one element so no point get discarded, so empty lambda
-			/*FinalizeFunc=*/[](int32) {}, // Not useful here, so empty lambda
 			/*bTimeSliceEnabled=*/false);
 				
 		if (TargetAttribute)
diff --git a/Engine/Plugins/PCG/Source/PCG/Public/Data/PCGBasePointData.h b/Engine/Plugins/PCG/Source/PCG/Public/Data/PCGBasePointData.h
--- a/Engine/Plugins/PCG/Source/PCG/Public/Data/PCGBasePointData.h
+++ b/Engine/Plugins/PCG/Source/PCG/Public/Data/PCGBasePointData.h
@@ -63,6 +63,7 @@ public:
 	virtual void AllocateProperties(EPCGPointNativeProperties Properties) {}
 	virtual void FreeProperties(EPCGPointNativeProperties Properties) {}
 	virtual void MoveRange(int32 RangeStartIndex, int32 MoveToIndex, int32 NumElements) PURE_VIRTUAL(UPCGBasePointData::MoveRange, return;);
+	virtual void CopyUnallocatedPropertiesTo(UPCGBasePointData* OutPointData) const {}
 
 	virtual TArray<FTransform> GetTransformsCopy() const PURE_VIRTUAL(UPCGBasePointData::GetTransformsCopy, return TArray<FTransform>(););
 	
diff --git a/Engine/Plugins/PCG/Source/PCG/Public/Data/PCGPointArrayData.h b/Engine/Plugins/PCG/Source/PCG/Public/Data/PCGPointArrayData.h
--- a/Engine/Plugins/PCG/Source/PCG/Public/Data/PCGPointArrayData.h
+++ b/Engine/Plugins/PCG/Source/PCG/Public/Data/PCGPointArrayData.h
@@ -38,6 +38,7 @@ public:
 	virtual void AllocateProperties(EPCGPointNativeProperties Properties) override;
 	virtual void FreeProperties(EPCGPointNativeProperties Properties) override;
 	virtual void MoveRange(int32 RangeStartIndex, int32 MoveToIndex, int32 NumElements) override;
+	virtual void CopyUnallocatedPropertiesTo(UPCGBasePointData* OutPointData) const override;
 
 	virtual TArray<FTransform> GetTransformsCopy() const { return PointArray.GetTransformCopy(); }
 
diff --git a/Engine/Plugins/PCG/Source/PCG/Public/PCGPointArray.h b/Engine/Plugins/PCG/Source/PCG/Public/PCGPointArray.h
--- a/Engine/Plugins/PCG/Source/PCG/Public/PCGPointArray.h
+++ b/Engine/Plugins/PCG/Source/PCG/Public/PCGPointArray.h
@@ -211,6 +211,14 @@ public:
 		}
 	}
 
+	FORCEINLINE void CopyUnallocatedProperty(FPCGPointArrayProperty<T>& OutProperty) const
+	{
+		if (!IsAllocated() && !OutProperty.IsAllocated())
+		{
+			OutProperty.Value = Value;
+		}
+	}
+
 private:
 	// Array containing values if allocated
 	TArray<T> Values;
